[Uba]
* Fixed path in Deploy-Windows.bat

diff --git a/Engine/Source/Programs/UnrealBuildAccelerator/Deploy-Windows.bat b/Engine/Source/Programs/UnrealBuildAccelerator/Deploy-Windows.bat
--- a/Engine/Source/Programs/UnrealBuildAccelerator/Deploy-Windows.bat
+++ b/Engine/Source/Programs/UnrealBuildAccelerator/Deploy-Windows.bat
@@ -13,7 +13,7 @@ set config=%1
 
 echo.
 echo === Building Targets ===
-pushd "../../../../../.."
+pushd "../../../.."
 call Engine/Build/BatchFiles/RunUBT.bat -NoBox -NoBoxLocal -NoSNDBS -NoXGE -Architectures=x64+arm64 ^
 	-Target="UbaAgent Win64 %config%" ^
 	-Target="UbaCli Win64 %config%" ^
