[Uba]
* Added so FileAccessor are using temp file for write on non-windows (since DeleteOnClose does not exist)
* Reverted reuse of filemappings (shm_open).. it caused problems when running agent on same machine

diff --git a/Engine/Source/Programs/UnrealBuildAccelerator/Common/Private/UbaFile.cpp b/Engine/Source/Programs/UnrealBuildAccelerator/Common/Private/UbaFile.cpp
--- a/Engine/Source/Programs/UnrealBuildAccelerator/Common/Private/UbaFile.cpp
+++ b/Engine/Source/Programs/UnrealBuildAccelerator/Common/Private/UbaFile.cpp
@@ -437,12 +437,6 @@ namespace uba
 		if (h == InvalidFileHandle)
 			return true;
 
-		if (h & DeleteOnCloseFlag)
-		{
-			UBA_ASSERT(fileName);
-			unlink(fileName);
-		}
-
 		if (close(asFileDescriptor(h)) == 0)
 			return true;
 		UBA_ASSERTF(false, TC("CloseFile error handling not implemented"));
diff --git a/Engine/Source/Programs/UnrealBuildAccelerator/Common/Private/UbaFileAccessor.cpp b/Engine/Source/Programs/UnrealBuildAccelerator/Common/Private/UbaFileAccessor.cpp
--- a/Engine/Source/Programs/UnrealBuildAccelerator/Common/Private/UbaFileAccessor.cpp
+++ b/Engine/Source/Programs/UnrealBuildAccelerator/Common/Private/UbaFileAccessor.cpp
@@ -11,25 +11,20 @@ namespace uba
 	HANDLE asHANDLE(FileHandle fh);
 	#else
 	int asFileDescriptor(FileHandle fh);
+	Atomic<u32> g_tempFileCounter;
 	#endif
 		
+#if PLATFORM_WINDOWS
 	bool SetDeleteOnClose(Logger& logger, const tchar* fileName, FileHandle& handle, bool value)
 	{
-#if PLATFORM_WINDOWS
 		ExtendedTimerScope ts(SystemStats::GetCurrent().setFileInfo);
 		FILE_DISPOSITION_INFO info;
 		info.DeleteFile = value;
 		if (!::SetFileInformationByHandle(asHANDLE(handle), FileDispositionInfo, &info, sizeof(info)))
 			return logger.Error(TC("SetFileInformationByHandle (FileDispositionInfo) failed on %llu %s (%s)"), uintptr_t(handle), fileName, LastErrorToText().data);
 		return true;
-#else
-		if (value)
-			(u64&)handle = handle | DeleteOnCloseFlag;
-		else
-			(u64&)handle = handle & ~DeleteOnCloseFlag;
-		return true;
-#endif
 	}
+#endif
 
 	FileAccessor::FileAccessor(Logger& logger, const tchar* fileName)
 	:	m_logger(logger)
@@ -42,17 +37,29 @@ namespace uba
 		InternalClose(false, nullptr);
 	}
 
-	bool FileAccessor::CreateWrite(bool allowRead, u32 flagsAndAttributes, u64 fileSize)
+	bool FileAccessor::CreateWrite(bool allowRead, u32 flagsAndAttributes, u64 fileSize, const tchar* tempPath)
 	{
 		UBA_ASSERT(flagsAndAttributes != 0);
 		m_size = fileSize;
 
+		const tchar* realFileName = m_fileName;
+
+		#if !PLATFORM_WINDOWS
+		m_tempPath = tempPath;
+		StringBuffer<> tempFile;
+		if (tempPath)
+		{
+			m_tempFileIndex = g_tempFileCounter++;
+			realFileName = tempFile.Append(tempPath).Append("Temp_").AppendValue(m_tempFileIndex).data;
+		}
+		#endif
+
 		u32 createDisp = CREATE_ALWAYS;
 		u32 dwDesiredAccess = GENERIC_WRITE | DELETE;
 		if (allowRead)
 			dwDesiredAccess |= GENERIC_READ;
 		u32 dwShareMode = 0;// FILE_SHARE_READ | FILE_SHARE_WRITE;
-		m_fileHandle = uba::CreateFileW(m_fileName, dwDesiredAccess, dwShareMode, createDisp, flagsAndAttributes);
+		m_fileHandle = uba::CreateFileW(realFileName, dwDesiredAccess, dwShareMode, createDisp, flagsAndAttributes);
 		if (m_fileHandle == InvalidFileHandle)
 		{
 			u32 lastError = GetLastError();
@@ -61,13 +68,13 @@ namespace uba
 			if (lastError == ERROR_SHARING_VIOLATION)
 				GetProcessHoldingFile(additionalInfo, m_fileName);
 #endif
-			return m_logger.Error(TC("ERROR opening file %s for write (%s%s)"), m_fileName, LastErrorToText(lastError).data, additionalInfo.data);
+			return m_logger.Error(TC("ERROR opening file %s for write (%s%s)"), realFileName, LastErrorToText(lastError).data, additionalInfo.data);
 		}
 
+#if PLATFORM_WINDOWS
 		if (!SetDeleteOnClose(m_logger, m_fileName, m_fileHandle, true))
 			return false;
 
-#if PLATFORM_WINDOWS
 		if (flagsAndAttributes & FILE_FLAG_OVERLAPPED)
 		{
 			if (fileSize != ~u64(0))
@@ -77,9 +84,9 @@ namespace uba
 				{
 					DWORD lastError = GetLastError();
 					if (lastError != ERROR_INVALID_FUNCTION) // Some file systems don't support this
-						return m_logger.Error(TC("Failed to make file %s sparse (%s)"), m_fileName, LastErrorToText(lastError).data);
+						return m_logger.Error(TC("Failed to make file %s sparse (%s)"), realFileName, LastErrorToText(lastError).data);
 				}
-				SetEndOfFile(m_logger, m_fileName, m_fileHandle, fileSize);
+				SetEndOfFile(m_logger, realFileName, m_fileHandle, fileSize);
 			}
 			(u64&)m_fileHandle |= OverlappedIoFlag;
 		}
@@ -88,12 +95,12 @@ namespace uba
 		return true;
 	}
 
-	bool FileAccessor::CreateMemoryWrite(bool allowRead, u32 flagsAndAttributes, u64 size)
+	bool FileAccessor::CreateMemoryWrite(bool allowRead, u32 flagsAndAttributes, u64 size, const tchar* tempPath)
 	{
 		m_size = size;
 
 		UBA_ASSERT(flagsAndAttributes != 0);
-		if (!CreateWrite(allowRead, flagsAndAttributes, size))
+		if (!CreateWrite(allowRead, flagsAndAttributes, size, tempPath))
 			return false;
 
 		m_mappingHandle = uba::CreateFileMappingW(m_fileHandle, PAGE_READWRITE, size);
@@ -310,19 +317,43 @@ namespace uba
 
 		if (m_fileHandle != InvalidFileHandle)
 		{
-			if (success && m_isWrite)
+			const tchar* realFileName = m_fileName;
+			StringBuffer<> tempFile;
+
+			if (m_isWrite)
 			{
-				if (!SetDeleteOnClose(m_logger, m_fileName, m_fileHandle, false))
-					return m_logger.Error(TC("Failed to remove delete on close for file %s (%s)"), m_fileName, LastErrorToText().data);
-				if (lastWriteTime)
+				#if !PLATFORM_WINDOWS
+				if (m_tempPath)
+					realFileName = tempFile.Append(m_tempPath).Append("Temp_").AppendValue(m_tempFileIndex).data;
+				#endif
+
+				if (success)
+				{
+					#if PLATFORM_WINDOWS
+					if (!SetDeleteOnClose(m_logger, realFileName, m_fileHandle, false))
+						return m_logger.Error(TC("Failed to remove delete on close for file %s (%s)"), realFileName, LastErrorToText().data);
+					#else
+					if (m_tempPath && rename(realFileName, m_fileName) == -1)
+						return m_logger.Error(TC("Failed to rename temporary file %s to %s (%s)"), realFileName, m_fileName, strerror(errno));
+					#endif
+
+					if (lastWriteTime)
+					{
+						*lastWriteTime = 0;
+						if (!GetFileLastWriteTime(*lastWriteTime, m_fileHandle))
+							m_logger.Warning(TC("Failed to get file time for %s (%s)"), realFileName, LastErrorToText().data);
+					}
+				}
+				else
 				{
-					*lastWriteTime = 0;
-					if (!GetFileLastWriteTime(*lastWriteTime, m_fileHandle))
-						m_logger.Warning(TC("Failed to get file time for %s (%s)"), m_fileName, LastErrorToText().data);
+					#if !PLATFORM_WINDOWS
+					if (m_tempPath && remove(realFileName) == -1)
+						return m_logger.Error(TC("Failed to remove temporary file %s (%s)"), realFileName, strerror(errno));
+					#endif
 				}
 			}
-			if (!CloseFile(m_fileName, m_fileHandle))
-				return m_logger.Error(TC("Failed to close file %s (%s)"), m_fileName, LastErrorToText().data);
+			if (!CloseFile(realFileName, m_fileHandle))
+				return m_logger.Error(TC("Failed to close file %s (%s)"), realFileName, LastErrorToText().data);
 			m_fileHandle = InvalidFileHandle;
 		}
 		return true;
diff --git a/Engine/Source/Programs/UnrealBuildAccelerator/Common/Private/UbaFileMapping.cpp b/Engine/Source/Programs/UnrealBuildAccelerator/Common/Private/UbaFileMapping.cpp
--- a/Engine/Source/Programs/UnrealBuildAccelerator/Common/Private/UbaFileMapping.cpp
+++ b/Engine/Source/Programs/UnrealBuildAccelerator/Common/Private/UbaFileMapping.cpp
@@ -52,11 +52,12 @@ namespace uba
 				break;
 			if (errno == EEXIST)
 			{
-				if (shm_unlink(uidName.data) == -1)
-				{
-					printf("Failed to shm_unlink %s (%s)\n", uidName.data, strerror(errno));
-					uid = ++g_mappingUidCounter;
-				}
+				uid = ++g_mappingUidCounter;
+				//if (shm_unlink(uidName.data) == -1)
+				//{
+				//	printf("Failed to shm_unlink %s (%s)\n", uidName.data, strerror(errno));
+				//	uid = ++g_mappingUidCounter;
+				//}
 				continue;
 			}
 			UBA_ASSERTF(false, "Failed to create filemapping with name %hs - %hs", uidName.data, strerror(errno));
@@ -147,7 +148,8 @@ namespace uba
 			GetMappingHandleName(uidName, h.uid);
 			if (shm_unlink(uidName.data) != 0)
 			{
-				UBA_ASSERTF(false, "Failed to unlink %s (%s)", uidName.data, strerror(errno));
+				SetLastError(errno);
+				//UBA_ASSERTF(false, "Failed to unlink %s (%s)", uidName.data, strerror(errno));
 				return false;
 			}
 		}
diff --git a/Engine/Source/Programs/UnrealBuildAccelerator/Common/Private/UbaSession.cpp b/Engine/Source/Programs/UnrealBuildAccelerator/Common/Private/UbaSession.cpp
--- a/Engine/Source/Programs/UnrealBuildAccelerator/Common/Private/UbaSession.cpp
+++ b/Engine/Source/Programs/UnrealBuildAccelerator/Common/Private/UbaSession.cpp
@@ -1008,16 +1008,10 @@ namespace uba
 			m_storage.CreateDirectory(m_sessionOutputDir.data);
 		}
 
-		#if PLATFORM_WINDOWS
 		m_tempPath.Append(m_sessionDir).Append(PathSeparator).Append(TC("temp"));
 		m_storage.CreateDirectory(m_tempPath.data);
-		m_sessionDir.MakeLower();
-		m_sessionBinDir.MakeLower();
-		m_sessionOutputDir.MakeLower();
-		m_tempPath.MakeLower();
-		#else
-		m_tempPath.Append("/tmp/");
-		#endif
+		m_tempPath.EnsureEndsWithSlash();
+
 		m_sessionDir.EnsureEndsWithSlash();
 		m_sessionBinDir.EnsureEndsWithSlash();
 		m_sessionOutputDir.EnsureEndsWithSlash();
@@ -1027,7 +1021,6 @@ namespace uba
 
 		#if PLATFORM_WINDOWS
 		m_systemPath.count = GetEnvironmentVariableW(TC("SystemRoot"), m_systemPath.data, m_systemPath.capacity);
-		m_systemPath.MakeLower();
 		#else
 		m_systemPath.Append(TC("/nonexistingpath"));
 		#endif
@@ -1913,13 +1906,13 @@ namespace uba
 
 			if (useFileMapForWrite)
 			{
-				if (!destinationFile.CreateMemoryWrite(false, attributes, fileSize))
+				if (!destinationFile.CreateMemoryWrite(false, attributes, fileSize, m_tempPath.data))
 					return false;
 				memcpy(destinationFile.GetData(), mem, fileSize);
 			}
 			else
 			{
-				if (!destinationFile.CreateWrite(false, attributes, fileSize))
+				if (!destinationFile.CreateWrite(false, attributes, fileSize, m_tempPath.data))
 					return false;
 				if (!destinationFile.Write(mem, fileSize))
 					return false;
diff --git a/Engine/Source/Programs/UnrealBuildAccelerator/Common/Private/UbaStorage.cpp b/Engine/Source/Programs/UnrealBuildAccelerator/Common/Private/UbaStorage.cpp
--- a/Engine/Source/Programs/UnrealBuildAccelerator/Common/Private/UbaStorage.cpp
+++ b/Engine/Source/Programs/UnrealBuildAccelerator/Common/Private/UbaStorage.cpp
@@ -219,7 +219,7 @@ namespace uba
 
 #if !UBA_USE_SPARSEFILE
 		FileAccessor destinationFile(m_logger, to);
-		if (!destinationFile.CreateWrite(to))
+		if (!destinationFile.CreateWrite(false, DefaultAttributes(), 0, m_tempPath.data))
 			return false;
 		if (!destinationFile.Write(&fileSize, sizeof(u64))) // Store file size first in compressed file
 			return false;
@@ -450,7 +450,7 @@ namespace uba
 				return m_logger.Error(TC("GetFileSize failed for %s (%s)"), fileName, LastErrorToText().data);
 
 			FileAccessor destinationFile(m_logger, casFile);
-			if (!destinationFile.CreateWrite(casFile))
+			if (!destinationFile.CreateWrite(false, DefaultAttributes(), 0, m_tempPath.data))
 				return false;
 
 			u8* slot= PopBufferSlot();
@@ -970,8 +970,14 @@ namespace uba
 		m_casCapacityBytes = info.casCapacityBytes;
 		m_storeCompressed = info.storeCompressed;
 		m_rootDir.count = GetFullPathNameW(info.rootDir, m_rootDir.capacity, m_rootDir.data, NULL);
-		m_rootDir.EnsureEndsWithSlash();
-		m_rootDir.Replace('/', PathSeparator).EnsureEndsWithSlash().Append(TC("cas")).Append(PathSeparator);
+		m_rootDir.Replace('/', PathSeparator).EnsureEndsWithSlash();
+
+		m_tempPath.Append(m_rootDir).Append(TC("castemp"));
+		CreateDirectory(m_tempPath.data);
+		DeleteAllFiles(m_logger, m_tempPath.data, false);
+		m_tempPath.EnsureEndsWithSlash();
+
+		m_rootDir.Append(TC("cas")).EnsureEndsWithSlash();
 
 		m_casDataBuffer.AddTransient(TC("CasData"));
 
@@ -1018,7 +1024,7 @@ namespace uba
 		if (!wasTerminated)
 		{
 			FileAccessor isRunningFile(m_logger, isRunningName);
-			if (!isRunningFile.CreateWrite() || !isRunningFile.Close())
+			if (!isRunningFile.CreateWrite(false, DefaultAttributes(), 0, m_tempPath.data) || !isRunningFile.Close())
 				return m_logger.Error(TC("Failed to create temporary \".isRunning\" file"));
 		}
 
@@ -1225,7 +1231,7 @@ namespace uba
 
 		{
 			FileAccessor tempFile(m_logger, tempFileName.data);
-			if (!tempFile.CreateWrite())
+			if (!tempFile.CreateWrite(false, DefaultAttributes(), 0, m_tempPath.data))
 				return false;
 			
 			ScopedReadLock fileTableLock(m_fileTableLookupLock);
@@ -1885,7 +1891,7 @@ namespace uba
 
 				if (writeDirectlyToFile || !decompressedSize)
 				{
-					if (!destinationFile.CreateWrite(allowRead, writeFlags | fileAttributes, decompressedSize))
+					if (!destinationFile.CreateWrite(allowRead, writeFlags | fileAttributes, decompressedSize, m_tempPath.data))
 						return false;
 					if (decompressedSize)
 						if (!DecompressMemoryToFile(readData, destinationFile, decompressedSize, useNoBuffering))
@@ -1893,7 +1899,7 @@ namespace uba
 				}
 				else
 				{
-					if (!destinationFile.CreateMemoryWrite(allowRead, writeFlags | fileAttributes, decompressedSize))
+					if (!destinationFile.CreateMemoryWrite(allowRead, writeFlags | fileAttributes, decompressedSize, m_tempPath.data))
 						return false;
 
 					if (casEntry->mappingHandle.IsValid())
diff --git a/Engine/Source/Programs/UnrealBuildAccelerator/Common/Private/UbaStorageClient.cpp b/Engine/Source/Programs/UnrealBuildAccelerator/Common/Private/UbaStorageClient.cpp
--- a/Engine/Source/Programs/UnrealBuildAccelerator/Common/Private/UbaStorageClient.cpp
+++ b/Engine/Source/Programs/UnrealBuildAccelerator/Common/Private/UbaStorageClient.cpp
@@ -363,7 +363,7 @@ namespace uba
 				bool useOverlap = !IsRunningWine() && isCompressed == IsCompressed(casKey) && actualSize > 1024 * 1024;
 				if (useOverlap)
 					extraFlags |= FILE_FLAG_OVERLAPPED;
-				if (!destinationFile.CreateWrite(false, extraFlags, actualSize))
+				if (!destinationFile.CreateWrite(false, extraFlags, actualSize, m_tempPath.data))
 					return false;
 			}
 #endif
diff --git a/Engine/Source/Programs/UnrealBuildAccelerator/Common/Private/UbaTrace.cpp b/Engine/Source/Programs/UnrealBuildAccelerator/Common/Private/UbaTrace.cpp
--- a/Engine/Source/Programs/UnrealBuildAccelerator/Common/Private/UbaTrace.cpp
+++ b/Engine/Source/Programs/UnrealBuildAccelerator/Common/Private/UbaTrace.cpp
@@ -95,7 +95,7 @@ namespace uba
 		if (!writeFileName || !*writeFileName)
 			return true;
 		FileAccessor traceFile(m_logger, writeFileName);
-		if (!traceFile.CreateWrite(writeFileName))
+		if (!traceFile.CreateWrite(false, DefaultAttributes(), 0, nullptr))
 			return false;
 		u64 fileSize = m_memoryPos;
 		if (!traceFile.Write(m_memoryBegin, fileSize))
diff --git a/Engine/Source/Programs/UnrealBuildAccelerator/Common/Public/UbaFile.h b/Engine/Source/Programs/UnrealBuildAccelerator/Common/Public/UbaFile.h
--- a/Engine/Source/Programs/UnrealBuildAccelerator/Common/Public/UbaFile.h
+++ b/Engine/Source/Programs/UnrealBuildAccelerator/Common/Public/UbaFile.h
@@ -29,11 +29,10 @@ namespace uba
 	#endif
 
 	inline constexpr u64 FileHandleFlagMask = 0x0000'0000'ffff'ffff;
-#if PLATFORM_WINDOWS
+
+	#if PLATFORM_WINDOWS
 	inline constexpr u64 OverlappedIoFlag = 0x0000'0001'0000'0000; // Only used by windows
-#else
-	inline constexpr u64 DeleteOnCloseFlag = 0x0000'0010'0000'0000; // Only used by non-windows
-#endif
+	#endif
 
 	struct FileInformation
 	{
diff --git a/Engine/Source/Programs/UnrealBuildAccelerator/Common/Public/UbaFileAccessor.h b/Engine/Source/Programs/UnrealBuildAccelerator/Common/Public/UbaFileAccessor.h
--- a/Engine/Source/Programs/UnrealBuildAccelerator/Common/Public/UbaFileAccessor.h
+++ b/Engine/Source/Programs/UnrealBuildAccelerator/Common/Public/UbaFileAccessor.h
@@ -13,8 +13,9 @@ namespace uba
 		FileAccessor(Logger& logger, const tchar* fileName);
 		~FileAccessor();
 
-		bool CreateWrite(bool allowRead = false, u32 flagsAndAttributes = DefaultAttributes(), u64 size = 0);
-		bool CreateMemoryWrite(bool allowRead, u32 flagsAndAttributes, u64 size);
+		// tempPath is only used by posix and will create a temp file and then move it to place when done. (Since DeleteOnClose does not exist)
+		bool CreateWrite(bool allowRead, u32 flagsAndAttributes, u64 size, const tchar* tempPath);
+		bool CreateMemoryWrite(bool allowRead, u32 flagsAndAttributes, u64 size, const tchar* tempPath);
 		bool Close(u64* lastWriteTime = nullptr);
 
 		bool Write(const void* data, u64 dataLen, u64 offset = 0);
@@ -37,5 +38,10 @@ namespace uba
 		u64 m_size = 0;
 		u8* m_data = nullptr;
 		bool m_isWrite = false;
+
+		#if !PLATFORM_WINDOWS
+		const tchar* m_tempPath = nullptr;
+		u32 m_tempFileIndex = 0;
+		#endif
 	};
 }
\ No newline at end of file
diff --git a/Engine/Source/Programs/UnrealBuildAccelerator/Common/Public/UbaStorage.h b/Engine/Source/Programs/UnrealBuildAccelerator/Common/Public/UbaStorage.h
--- a/Engine/Source/Programs/UnrealBuildAccelerator/Common/Public/UbaStorage.h
+++ b/Engine/Source/Programs/UnrealBuildAccelerator/Common/Public/UbaStorage.h
@@ -150,6 +150,7 @@ namespace uba
 		LoggerWithWriter m_logger;
 
 		StringBuffer<> m_rootDir;
+		StringBuffer<> m_tempPath;
 
 		ReaderWriterLock m_fileTableLookupLock;
 		struct FileEntry { ReaderWriterLock lock; CasKey casKey = CasKeyZero; u64 size = 0; u64 lastWritten = 0; bool verified = false; };
diff --git a/Engine/Source/Programs/UnrealBuildAccelerator/Core/Private/UbaPlatform.cpp b/Engine/Source/Programs/UnrealBuildAccelerator/Core/Private/UbaPlatform.cpp
--- a/Engine/Source/Programs/UnrealBuildAccelerator/Core/Private/UbaPlatform.cpp
+++ b/Engine/Source/Programs/UnrealBuildAccelerator/Core/Private/UbaPlatform.cpp
@@ -277,6 +277,7 @@ namespace uba
 		return (MutexHandle)(u64)::CreateMutexW(NULL, bInitialOwner, lpName);
 		#else
 		// TODO: This is used to check for exclusivity and also for trace streams (only created by host and read by visualizer)
+		SetLastError(ERROR_SUCCESS);
 		return ((MutexHandle)(u64)1337); // Just some random value
 		#endif
 	}
diff --git a/Engine/Source/Programs/UnrealBuildAccelerator/Core/Public/UbaDirectoryTable.h b/Engine/Source/Programs/UnrealBuildAccelerator/Core/Public/UbaDirectoryTable.h
--- a/Engine/Source/Programs/UnrealBuildAccelerator/Core/Public/UbaDirectoryTable.h
+++ b/Engine/Source/Programs/UnrealBuildAccelerator/Core/Public/UbaDirectoryTable.h
@@ -38,7 +38,7 @@ namespace uba
 				u64 pos = reader.GetPosition();
 				if (pos == size)
 					break;
-				UBA_ASSERT(pos < size);
+				UBA_ASSERTF(pos < size, TC("Should never read past size (pos: %u, size: %u)"), pos, size);
 				u64 storageSize = reader.Read7BitEncoded();
 				StringKey dirKey = reader.ReadStringKey();
 				auto insres = m_lookup.try_emplace(dirKey, m_memoryBlock); // Note that this is allowed to overwrite
diff --git a/Engine/Source/Programs/UnrealBuildAccelerator/Detours/Private/PosixOS/UbaDetoursFunctionsPosix.cpp b/Engine/Source/Programs/UnrealBuildAccelerator/Detours/Private/PosixOS/UbaDetoursFunctionsPosix.cpp
--- a/Engine/Source/Programs/UnrealBuildAccelerator/Detours/Private/PosixOS/UbaDetoursFunctionsPosix.cpp
+++ b/Engine/Source/Programs/UnrealBuildAccelerator/Detours/Private/PosixOS/UbaDetoursFunctionsPosix.cpp
@@ -14,6 +14,10 @@
 #include <dlfcn.h>
 #include <sys/stat.h>
 
+#if PLATFORM_LINUX
+#include <sys/prctl.h>
+#endif
+
 namespace uba
 {
 	constexpr bool g_logToScreen = false;
@@ -1447,9 +1451,11 @@ int Internal_execve(const char* pathname, char* const _Nullable argv[], char* co
 	DEBUG_LOG_TRUE("execve", "");
 
 	bool inVfork = t_inVfork != 0;
-
 	pid_t pid;
 	int res = posix_spawn(&pid, pathname, nullptr, nullptr, argv, envp);
+	if (inVfork)
+		t_inVfork = pid;
+
 	if (res != 0)
 	{
 		UBA_ASSERTF(false, "Failed to spawn %s", pathname);
@@ -1464,9 +1470,6 @@ int Internal_execve(const char* pathname, char* const _Nullable argv[], char* co
 		UBA_ASSERT(WIFEXITED(status));
 	}
 
-	if (inVfork)
-		t_inVfork = pid;
-
 	{
 		TimerScope ts(g_stats.createProcess);
 		ScopedWriteLock pcs(g_communicationLock);
@@ -1509,7 +1512,10 @@ UBA_EXPORT pid_t UBA_WRAPPER(vfork)(void)
 	DEBUG_LOG_TRUE("vfork", "");
 	pid_t pid = fork();
 	if (pid == 0)
+	{
+		prctl(PR_SET_PDEATHSIG, SIGHUP, 0, 0, 0); // We want the process to die if the parent die
 		t_inVfork = 1;
+	}
 	return pid;
 }
 #endif
@@ -1621,7 +1627,10 @@ namespace uba
 	int GetProcessExecutablePath(tchar* Path, u32 PathSize)
 	{
 #if PLATFORM_LINUX
-		return TRUE_WRAPPER(readlink)("/proc/self/exe", Path, PathSize);
+		auto res = TRUE_WRAPPER(readlink)("/proc/self/exe", Path, PathSize);
+		if (res != -1)
+			Path[res] = 0;
+		return res;
 #elif PLATFORM_MAC
 		if (_NSGetExecutablePath(Path, &PathSize) == 0)
 			return strlen(Path);
@@ -1634,8 +1643,8 @@ namespace uba
 	{
 		SuppressDetourScope s;
 
-		//g_systemTemp.Append(getenv("TMPDIR"));
-		g_systemTemp.Append("/tmp/");
+		g_systemTemp.Append(getenv("TMPDIR"));
+		//g_systemTemp.Append("/tmp/");
 
 		#if PLATFORM_LINUX
 		#define DETOURED_FUNCTION(func) \
