[Uba]
* Fixed test linux compile errors

diff --git a/Engine/Source/Programs/UnrealBuildAccelerator/Common/Public/UbaFileAccessor.h b/Engine/Source/Programs/UnrealBuildAccelerator/Common/Public/UbaFileAccessor.h
--- a/Engine/Source/Programs/UnrealBuildAccelerator/Common/Public/UbaFileAccessor.h
+++ b/Engine/Source/Programs/UnrealBuildAccelerator/Common/Public/UbaFileAccessor.h
@@ -14,8 +14,8 @@ namespace uba
 		~FileAccessor();
 
 		// tempPath is only used by posix and will create a temp file and then move it to place when done. (Since DeleteOnClose does not exist)
-		bool CreateWrite(bool allowRead, u32 flagsAndAttributes, u64 size, const tchar* tempPath);
-		bool CreateMemoryWrite(bool allowRead, u32 flagsAndAttributes, u64 size, const tchar* tempPath);
+		bool CreateWrite(bool allowRead = false, u32 flagsAndAttributes = DefaultAttributes(), u64 size = 0, const tchar* tempPath = nullptr);
+		bool CreateMemoryWrite(bool allowRead = false, u32 flagsAndAttributes = DefaultAttributes(), u64 size = 0, const tchar* tempPath = nullptr);
 		bool Close(u64* lastWriteTime = nullptr);
 
 		bool Write(const void* data, u64 dataLen, u64 offset = 0);
diff --git a/Engine/Source/Programs/UnrealBuildAccelerator/Test/Private/UbaTestBasics.h b/Engine/Source/Programs/UnrealBuildAccelerator/Test/Private/UbaTestBasics.h
--- a/Engine/Source/Programs/UnrealBuildAccelerator/Test/Private/UbaTestBasics.h
+++ b/Engine/Source/Programs/UnrealBuildAccelerator/Test/Private/UbaTestBasics.h
@@ -79,7 +79,7 @@ namespace uba
 	bool TestFiles(Logger& logger, const StringBufferBase& rootDir)
 	{
 		FileAccessor fileHandle(logger, TC("UbaTestFile"));
-		if (!fileHandle.CreateWrite(TC("UbaTestFile")))
+		if (!fileHandle.CreateWrite())
 			return logger.Error(TC("Failed to create file for write"));
 
 		u8 byte = 'H';
