[Uba]
* Changed critical section to use pthread mutex

diff --git a/Engine/Source/Programs/UnrealBuildAccelerator/Core/Private/UbaSynchronization.cpp b/Engine/Source/Programs/UnrealBuildAccelerator/Core/Private/UbaSynchronization.cpp
--- a/Engine/Source/Programs/UnrealBuildAccelerator/Core/Private/UbaSynchronization.cpp
+++ b/Engine/Source/Programs/UnrealBuildAccelerator/Core/Private/UbaSynchronization.cpp
@@ -3,116 +3,101 @@
 #include "UbaSynchronization.h"
 #include "UbaPlatform.h"
 
-#if PLATFORM_WINDOWS
-#define UBA_USE_WIN 1
-#else
-#define UBA_USE_WIN 0
-#endif
-
-#if !UBA_USE_WIN
-#include <shared_mutex>
-#include <mutex>
-#endif
-
 namespace uba
 {
 
 	CriticalSection::CriticalSection()
 	{
-		#if UBA_USE_WIN
+		#if PLATFORM_WINDOWS
 		static_assert(sizeof(data) >= sizeof(CRITICAL_SECTION));
 		InitializeCriticalSection((CRITICAL_SECTION*)&data);
 		#else
-		static_assert(sizeof(data) >= sizeof(std::recursive_mutex));
-		new (data) std::recursive_mutex();
+		static_assert(sizeof(data) >= sizeof(pthread_mutex_t));
+		pthread_mutexattr_t attr;
+		pthread_mutexattr_init(&attr);
+		pthread_mutexattr_settype(&attr, PTHREAD_MUTEX_RECURSIVE);
+		pthread_mutex_init((pthread_mutex_t*)data, &attr);
 		#endif
 	}
 
 	CriticalSection::~CriticalSection()
 	{
-		#if UBA_USE_WIN
+		#if PLATFORM_WINDOWS
 		DeleteCriticalSection((CRITICAL_SECTION*)&data);
 		#else
-		((std::recursive_mutex&)data).std::recursive_mutex::~recursive_mutex();
+		pthread_mutex_destroy((pthread_mutex_t*)data);
 		#endif
 	}
 
 	void CriticalSection::Enter()
 	{
-		#if UBA_USE_WIN
+		#if PLATFORM_WINDOWS
 		EnterCriticalSection((CRITICAL_SECTION*)&data);
 		#else
-		((std::recursive_mutex&)data).lock();
+		pthread_mutex_lock((pthread_mutex_t*)data);
 		#endif
 	}
 
 	void CriticalSection::Leave()
 	{
-		#if UBA_USE_WIN
+		#if PLATFORM_WINDOWS
 		LeaveCriticalSection((CRITICAL_SECTION*)&data);
 		#else
-		((std::recursive_mutex&)data).unlock();
+		pthread_mutex_unlock((pthread_mutex_t*)data);
 		#endif
 	}
 
 	ReaderWriterLock::ReaderWriterLock()
 	{
-		#if UBA_USE_WIN
+		#if PLATFORM_WINDOWS
 		static_assert(sizeof(data) >= sizeof(SRWLOCK));
 		InitializeSRWLock((SRWLOCK*)&data);
 		#else
 		static_assert(sizeof(data) >= sizeof(pthread_rwlock_t));
 		pthread_rwlock_init((pthread_rwlock_t*)data, NULL);
-		//static_assert(sizeof(data) >= sizeof(std::shared_mutex));
-		//new (data) std::shared_mutex();
 		#endif
 	}
 
 	ReaderWriterLock::~ReaderWriterLock()
 	{
-		#if !UBA_USE_WIN
+		#if !PLATFORM_WINDOWS
 		pthread_rwlock_destroy((pthread_rwlock_t*)data);
-		//((std::shared_mutex&)data).std::shared_mutex::~shared_mutex();
 		#endif
 	}
 
 	void ReaderWriterLock::EnterRead()
 	{
-		#if UBA_USE_WIN
+		#if PLATFORM_WINDOWS
 		AcquireSRWLockShared((SRWLOCK*)&data);
 		#else
 		pthread_rwlock_rdlock((pthread_rwlock_t*)data);
-		//((std::shared_mutex&)data).lock_shared();
 		#endif
 	}
 
 	void ReaderWriterLock::LeaveRead()
 	{
-		#if UBA_USE_WIN
+		#if PLATFORM_WINDOWS
 		ReleaseSRWLockShared((SRWLOCK*)&data);
 		#else
 		pthread_rwlock_unlock((pthread_rwlock_t*)data);
-		//((std::shared_mutex&)data).unlock_shared();
 		#endif
 	}
 
 	void ReaderWriterLock::EnterWrite()
 	{
-		#if UBA_USE_WIN
+		#if PLATFORM_WINDOWS
 		AcquireSRWLockExclusive((SRWLOCK*)&data);
 		#else
 		pthread_rwlock_wrlock((pthread_rwlock_t*)data);
-		//((std::shared_mutex&)data).lock();
 		#endif
 	}
 
 	void ReaderWriterLock::LeaveWrite()
 	{
-		#if UBA_USE_WIN
+		#if PLATFORM_WINDOWS
 		ReleaseSRWLockExclusive((SRWLOCK*)&data);
 		#else
 		pthread_rwlock_unlock((pthread_rwlock_t*)data);
-		//((std::shared_mutex&)data).unlock();
 		#endif
 	}
 }
