[Uba]
* Improved assert descriptions

diff --git a/Engine/Source/Programs/UnrealBuildAccelerator/Common/Private/UbaStorageClient.cpp b/Engine/Source/Programs/UnrealBuildAccelerator/Common/Private/UbaStorageClient.cpp
--- a/Engine/Source/Programs/UnrealBuildAccelerator/Common/Private/UbaStorageClient.cpp
+++ b/Engine/Source/Programs/UnrealBuildAccelerator/Common/Private/UbaStorageClient.cpp
@@ -553,8 +553,8 @@ namespace uba
 
 						// Move overflow back to the beginning of the buffer and start the next block (if there is one)
 						readBuffer = slot;
-						UBA_ASSERTF(readPosition - overflow >= readBuffer, TC("ReadPosition - overflow is before beginning of buffer (overflow: %u)"), overflow);
-						UBA_ASSERTF(readPosition <= readBuffer + BufferSlotHalfSize, TC("readPosition is outside readBuffer size (pos: %llu, overflow: %u)"), readPosition - readBuffer, overflow);
+						UBA_ASSERTF(readPosition - overflow >= readBuffer, TC("ReadPosition - overflow is before beginning of buffer (overflow: %u) for file %s"), overflow, hint);
+						UBA_ASSERTF(readPosition <= readBuffer + BufferSlotHalfSize, TC("ReadPosition is outside readBuffer size (pos: %llu, overflow: %u) for file %s"), readPosition - readBuffer, overflow, hint);
 						memmove(readBuffer, readPosition - overflow, overflow);
 						readPosition = readBuffer + overflow;
 						if (overflow)
