[UBA]
* Fixed macos compile errors

diff --git a/Engine/Source/Programs/UnrealBuildAccelerator/Core/Public/UbaSynchronization.h b/Engine/Source/Programs/UnrealBuildAccelerator/Core/Public/UbaSynchronization.h
--- a/Engine/Source/Programs/UnrealBuildAccelerator/Core/Public/UbaSynchronization.h
+++ b/Engine/Source/Programs/UnrealBuildAccelerator/Core/Public/UbaSynchronization.h
@@ -72,8 +72,10 @@ namespace uba
 
 		#if PLATFORM_WINDOWS
 		u64 data[1];
-		#else
+		#elif PLATFORM_LINUX
 		u64 data[7];
+		#else
+		u64 data[25];
 		#endif
 
 		ReaderWriterLock(const ReaderWriterLock&) = delete;
diff --git a/Engine/Source/Programs/UnrealBuildAccelerator/Detours/Private/UbaDetoursShared.cpp b/Engine/Source/Programs/UnrealBuildAccelerator/Detours/Private/UbaDetoursShared.cpp
--- a/Engine/Source/Programs/UnrealBuildAccelerator/Detours/Private/UbaDetoursShared.cpp
+++ b/Engine/Source/Programs/UnrealBuildAccelerator/Detours/Private/UbaDetoursShared.cpp
@@ -352,7 +352,7 @@ namespace uba
 						(u64&)outAttr.data.ftLastWriteTime = info.lastWrite;
 #else
 						outAttr.data.st_mtime = info.lastWrite;
-						outAttr.data.st_mode = info.attributes;
+						outAttr.data.st_mode = (mode_t)info.attributes;
 						outAttr.data.st_dev = info.volumeSerial;
 						outAttr.data.st_ino = info.fileIndex;
 						outAttr.data.st_size = info.size;
