[Uba]
* Fixed so child processes wait for acknowledgement that ExitMessage has been handled on session side. This to prevent race where parent process might move on after waiting for child and then read data produced by the child (and the data is written in the Exit message by the session process)

diff --git a/Engine/Source/Programs/UnrealBuildAccelerator/Common/Private/UbaProcess.cpp b/Engine/Source/Programs/UnrealBuildAccelerator/Common/Private/UbaProcess.cpp
--- a/Engine/Source/Programs/UnrealBuildAccelerator/Common/Private/UbaProcess.cpp
+++ b/Engine/Source/Programs/UnrealBuildAccelerator/Common/Private/UbaProcess.cpp
@@ -1159,6 +1159,7 @@ namespace uba
 		payload.communicationOffset = communicationOffset;
 		payload.rulesIndex = m_rulesIndex;
 		payload.runningRemote = runningRemote;
+		payload.isChild = m_parentProcess != nullptr;
 		payload.trackInputs = m_startInfo.trackInputs;
 		payload.useCustomAllocator = m_startInfo.useCustomAllocator && g_applicationRules[m_rulesIndex].rules->AllowMiMalloc();
 		payload.isRunningWine = IsRunningWine();
diff --git a/Engine/Source/Programs/UnrealBuildAccelerator/Core/Public/UbaDetoursPayload.h b/Engine/Source/Programs/UnrealBuildAccelerator/Core/Public/UbaDetoursPayload.h
--- a/Engine/Source/Programs/UnrealBuildAccelerator/Core/Public/UbaDetoursPayload.h
+++ b/Engine/Source/Programs/UnrealBuildAccelerator/Core/Public/UbaDetoursPayload.h
@@ -21,6 +21,7 @@ namespace uba
 		bool trackInputs = false;
 		bool useCustomAllocator = true;
 		bool isRunningWine = false;
+		bool isChild = false;
 		int uiLanguage = 0;
 		StringBuffer<256> logFile;
 	};
diff --git a/Engine/Source/Programs/UnrealBuildAccelerator/Detours/Private/UbaDetoursShared.cpp b/Engine/Source/Programs/UnrealBuildAccelerator/Detours/Private/UbaDetoursShared.cpp
--- a/Engine/Source/Programs/UnrealBuildAccelerator/Detours/Private/UbaDetoursShared.cpp
+++ b/Engine/Source/Programs/UnrealBuildAccelerator/Detours/Private/UbaDetoursShared.cpp
@@ -50,6 +50,7 @@ namespace uba
 	u32 g_rulesIndex;
 	ApplicationRules* g_rules;
 	bool g_runningRemote;
+	bool g_isChild;
 
 #if UBA_DEBUG_LOG_ENABLED
 	FileHandle g_debugFile = InvalidFileHandle;
diff --git a/Engine/Source/Programs/UnrealBuildAccelerator/Detours/Private/UbaDetoursShared.h b/Engine/Source/Programs/UnrealBuildAccelerator/Detours/Private/UbaDetoursShared.h
--- a/Engine/Source/Programs/UnrealBuildAccelerator/Detours/Private/UbaDetoursShared.h
+++ b/Engine/Source/Programs/UnrealBuildAccelerator/Detours/Private/UbaDetoursShared.h
@@ -102,6 +102,7 @@ namespace uba
 	extern u32 g_rulesIndex;
 	extern ApplicationRules* g_rules;
 	extern bool g_runningRemote;
+	extern bool g_isChild;
 
 	#if PLATFORM_WINDOWS
 	constexpr u32 ErrorSuccess = 0;
diff --git a/Engine/Source/Programs/UnrealBuildAccelerator/Detours/Private/Windows/UbaDetoursFunctionsNtDll.inl b/Engine/Source/Programs/UnrealBuildAccelerator/Detours/Private/Windows/UbaDetoursFunctionsNtDll.inl
--- a/Engine/Source/Programs/UnrealBuildAccelerator/Detours/Private/Windows/UbaDetoursFunctionsNtDll.inl
+++ b/Engine/Source/Programs/UnrealBuildAccelerator/Detours/Private/Windows/UbaDetoursFunctionsNtDll.inl
@@ -1245,6 +1245,7 @@ NTSTATUS NTAPI Detoured_NtClose(HANDLE handle)
 				newInfo = fi;
 				newInfo.originalName = g_memoryBlock.Strdup(fo->newName.c_str());
 				newInfo.name = newInfo.originalName;
+				newInfo.fileNameKey = fileNameKey;
 				UBA_ASSERT(!fo->deleteOnClose);
 				fi = {};
 				fi.deleted = true;
diff --git a/Engine/Source/Programs/UnrealBuildAccelerator/Detours/Private/Windows/UbaDetoursFunctionsWin.cpp b/Engine/Source/Programs/UnrealBuildAccelerator/Detours/Private/Windows/UbaDetoursFunctionsWin.cpp
--- a/Engine/Source/Programs/UnrealBuildAccelerator/Detours/Private/Windows/UbaDetoursFunctionsWin.cpp
+++ b/Engine/Source/Programs/UnrealBuildAccelerator/Detours/Private/Windows/UbaDetoursFunctionsWin.cpp
@@ -558,7 +558,10 @@ void SendExitMessage(DWORD exitCode, u64 startTime)
 
 	g_stats.Write(writer);
 
-	writer.Flush(false); // This can't wait for response since the session process might move on and reuse shared memory with someone else
+	// We must flush here if this is a child because,
+	// if there is a parent process waiting for this to finish,
+	// the parent might move on before Exit message has been processed on session side
+	writer.Flush(g_isChild);
 }
 
 // Variables used to communicate state from kernelbase functions to ntdll functions
@@ -683,6 +686,7 @@ void PreInit(const DetoursPayload& payload)
 	g_rules = g_applicationRules[payload.rulesIndex].rules;
 	g_useMiMalloc = payload.useCustomAllocator;
 	g_runningRemote = payload.runningRemote;
+	g_isChild = payload.isChild;
 	g_isDetachedProcess = g_rules->AllowDetach();
 	g_isRunningWine = payload.isRunningWine;
 	g_uiLanguage = payload.uiLanguage;
