[Uba] - Added comment on Flush(false) for exit message on posix .

diff --git a/Engine/Source/Programs/UnrealBuildAccelerator/Detours/Private/PosixOS/UbaDetoursFunctionsPosix.cpp b/Engine/Source/Programs/UnrealBuildAccelerator/Detours/Private/PosixOS/UbaDetoursFunctionsPosix.cpp
--- a/Engine/Source/Programs/UnrealBuildAccelerator/Detours/Private/PosixOS/UbaDetoursFunctionsPosix.cpp
+++ b/Engine/Source/Programs/UnrealBuildAccelerator/Detours/Private/PosixOS/UbaDetoursFunctionsPosix.cpp
@@ -1779,7 +1779,11 @@ namespace uba
 		writer.WriteU32(0); // Exit code
 		writer.WriteString(""); // Log name
 		g_stats.Write(writer);
-		writer.Flush(false); // This can't wait for response since the session process might move on and reuse shared memory with someone else
+
+		// This can't wait for response since the session process might move on and reuse shared memory with someone else
+		// Note, if we start using memory mapped files we need to change this to true for child processes since Exit message is writing files to disk..
+		// .. and if we don't wait to exit this process until those files are written we might end up in a race condition with the parent using those files
+		writer.Flush(false);
 
 		#if UBA_DEBUG_LOG_ENABLED
 		if (isLogging())
