[Uba]
* Fixed so DeleteAllCas is running in parallel to speed up cleanups since we often have 100k+ files
* Added -resetstore as cmd line option to UbaAgent to be able to wipe all cas storage

diff --git a/Engine/Source/Programs/UnrealBuildAccelerator/Agent/Private/UbaAgent.cpp b/Engine/Source/Programs/UnrealBuildAccelerator/Agent/Private/UbaAgent.cpp
--- a/Engine/Source/Programs/UnrealBuildAccelerator/Agent/Private/UbaAgent.cpp
+++ b/Engine/Source/Programs/UnrealBuildAccelerator/Agent/Private/UbaAgent.cpp
@@ -85,6 +85,7 @@ namespace uba
 		logger.Info(TC("  -named=<name>           Use named events and file mappings by providing the base name in this option"));
 		logger.Info(TC("  -nopoll                 Does not keep polling for work; attempts to connect once then exits"));
 		logger.Info(TC("  -nostore                Does not use storage to store files (with a few exceptions such as binaries)"));
+		logger.Info(TC("  -resetstore             Delete all cas"));
 		logger.Info(TC("  -quiet                  Does not output any logging in console"));
 		logger.Info(TC("  -maxidle=<seconds>      Max time agent will idle before disconnecting. Ignored if -nopoll is not set"));
 		logger.Info(TC("  -binasversion           Will use binaries as version. This will cause updates everytime binaries change on host side"));
@@ -399,6 +400,7 @@ namespace uba
 		bool useQuic = false;
 		bool poll = true;
 		bool useStorage = true;
+		bool resetStore = false;
 		bool quiet = false;
 		bool verbose = false;
 		bool printSummary = false;
@@ -541,6 +543,10 @@ namespace uba
 			{
 				useStorage = false;
 			}
+			else if (name.Equals(TC("-resetstore")))
+			{
+				resetStore = true;
+			}
 			else if (name.Equals(TC("-binasversion")))
 			{
 				useBinariesAsVersion = true;
@@ -711,7 +717,12 @@ namespace uba
 			info.casCapacityBytes = storageCapacity;
 			info.storeCompressed = storeCompressed;
 			StorageImpl storage(info);
-			if (!storage.LoadCasTable(false))
+			if (resetStore)
+			{
+				if (!storage.Reset())
+					return -1;
+			}
+			else if (!storage.LoadCasTable(false))
 				return -1;
 		}
 
diff --git a/Engine/Source/Programs/UnrealBuildAccelerator/Common/Private/UbaFile.cpp b/Engine/Source/Programs/UnrealBuildAccelerator/Common/Private/UbaFile.cpp
--- a/Engine/Source/Programs/UnrealBuildAccelerator/Common/Private/UbaFile.cpp
+++ b/Engine/Source/Programs/UnrealBuildAccelerator/Common/Private/UbaFile.cpp
@@ -302,11 +302,11 @@ namespace uba
 				{
 					if (!DeleteFileW(fullPath.data))
 					{
-						if (count)
-							++(*count);
 						logger.Warning(TC("Failed to delete file %s (%s)"), fullPath.data, LastErrorToText().data);
 						success = false;
 					}
+					else if (count)
+						++(*count);
 				}
 			});
 
diff --git a/Engine/Source/Programs/UnrealBuildAccelerator/Common/Private/UbaStorage.cpp b/Engine/Source/Programs/UnrealBuildAccelerator/Common/Private/UbaStorage.cpp
--- a/Engine/Source/Programs/UnrealBuildAccelerator/Common/Private/UbaStorage.cpp
+++ b/Engine/Source/Programs/UnrealBuildAccelerator/Common/Private/UbaStorage.cpp
@@ -1476,6 +1476,27 @@ namespace uba
 	bool StorageImpl::DeleteAllCas()
 	{
 		u32 deleteCount = 0;
+
+		WorkManagerImpl workManager(GetLogicalProcessorCount());
+		{
+			Atomic<u32> atomicDeleteCount;
+			TraverseDir(m_logger, m_rootDir.data, [&](const DirectoryEntry& e)
+				{
+					if (!IsDirectory(e.attributes))
+						return;
+					workManager.AddWork([&, name = TString(e.name)]()
+						{
+							StringBuffer<> fullPath;
+							fullPath.Append(m_rootDir).Append(name);
+							u32 deleteCountTemp = 0;
+							DeleteAllFiles(m_logger, fullPath.data, true, &deleteCountTemp);
+							atomicDeleteCount += deleteCountTemp;
+						}, 1, TC(""));
+				});
+			workManager.Wait();
+			deleteCount += atomicDeleteCount;
+		}
+
 		bool res = DeleteAllFiles(m_logger, m_rootDir.data, false, &deleteCount);
 		m_logger.Info(TC("Deleted %u cas files"), deleteCount);
 		m_createdDirs.clear();
