[Uba]
* Added unit test for inverted connect path (client listen and server connects)
* Changed so AddClient waits for connect to succeed/fail before returning

diff --git a/Engine/Source/Programs/UnrealBuildAccelerator/Common/Private/UbaNetworkServer.cpp b/Engine/Source/Programs/UnrealBuildAccelerator/Common/Private/UbaNetworkServer.cpp
--- a/Engine/Source/Programs/UnrealBuildAccelerator/Common/Private/UbaNetworkServer.cpp
+++ b/Engine/Source/Programs/UnrealBuildAccelerator/Common/Private/UbaNetworkServer.cpp
@@ -628,17 +628,24 @@ namespace uba
 				return false;
 		}
 
-		m_addConnections.emplace_back([this, &backend, ip2 = TString(ip), port, cryptoKey]()
+		Event done(true);
+		bool success = false;
+
+		m_addConnections.emplace_back([this, &success , &done, &backend, ip2 = TString(ip), port, cryptoKey]()
 			{
-				bool success = backend.Connect(m_logger, ip2.c_str(), [this, &backend, cryptoKey](void* connection, const sockaddr& remoteSocketAddr, bool* timedOut)
+				// TODO: Should this retry?
+				success = backend.Connect(m_logger, ip2.c_str(), [this, &backend, cryptoKey](void* connection, const sockaddr& remoteSocketAddr, bool* timedOut)
 					{
 						return AddConnection(backend, connection, remoteSocketAddr, cryptoKey);
 					}, port, nullptr);
 				if (!success)
 					Crypto::DestroyKey(cryptoKey);
+				done.Set();
 				return 0;
 			});
-		return true;
+
+		done.IsSet();
+		return success;
 	}
 
 	void NetworkServer::PrintSummary(Logger& logger)
diff --git a/Engine/Source/Programs/UnrealBuildAccelerator/Common/Public/UbaNetworkClient.h b/Engine/Source/Programs/UnrealBuildAccelerator/Common/Public/UbaNetworkClient.h
--- a/Engine/Source/Programs/UnrealBuildAccelerator/Common/Public/UbaNetworkClient.h
+++ b/Engine/Source/Programs/UnrealBuildAccelerator/Common/Public/UbaNetworkClient.h
@@ -61,6 +61,7 @@ namespace uba
 		const Guid& GetUid() { return m_uid; }
 		LogWriter& GetLogWriter() { return m_logWriter; }
 		NetworkBackendTcp& GetTcpBackend() { return *m_tcpBackend; }
+		u32 GetConnectionCount() { return m_connectionCount; }
 
 	private:
 		struct Connection
@@ -80,7 +81,6 @@ namespace uba
 		void ConnectedCallback(NetworkBackend& backend, void* backendConnection);
 		static bool ReceiveResponseHeader(void* context, u8* headerData, void*& outBodyContext, u8*& outBodyData, u32& outBodySize);
 		static bool ReceiveResponseBody(void* context, bool recvError, u8* headerData, void* bodyContext, u8* bodyData, u32 bodySize);
-		void ThreadReceive(Connection& connection);
 		void OnDisconnected(Connection& connection, bool calledFromReceive);
 		bool Send(NetworkMessage& message, void* response, u32 responseCapacity, bool async);
 		void ReturnMessageId(u16 id);
diff --git a/Engine/Source/Programs/UnrealBuildAccelerator/Common/Public/UbaNetworkServer.h b/Engine/Source/Programs/UnrealBuildAccelerator/Common/Public/UbaNetworkServer.h
--- a/Engine/Source/Programs/UnrealBuildAccelerator/Common/Public/UbaNetworkServer.h
+++ b/Engine/Source/Programs/UnrealBuildAccelerator/Common/Public/UbaNetworkServer.h
@@ -44,7 +44,7 @@ namespace uba
 		void DisallowNewClients();	// Disallow new clients to connect but old clients can still create more connections
 		void StopAll();				// Stops all listen and disconnect all active connections
 
-		bool AddClient(NetworkBackend& backend, const tchar* ip, u16 port = DefaultPort, const u8* cryptoKey128 = nullptr); // Adds a client that server will create one or more connections to
+		bool AddClient(NetworkBackend& backend, const tchar* ip, u16 port = DefaultPort, const u8* cryptoKey128 = nullptr); // Adds a client that server will create one or more connections to (note this will return before we know if it was a success or not)
 
 		void PrintSummary(Logger& logger);
 
diff --git a/Engine/Source/Programs/UnrealBuildAccelerator/Test/Private/UbaTestNetwork.h b/Engine/Source/Programs/UnrealBuildAccelerator/Test/Private/UbaTestNetwork.h
--- a/Engine/Source/Programs/UnrealBuildAccelerator/Test/Private/UbaTestNetwork.h
+++ b/Engine/Source/Programs/UnrealBuildAccelerator/Test/Private/UbaTestNetwork.h
@@ -80,6 +80,50 @@ namespace uba
 			return logger.Error(TC("Failed to connect"));
 
 
+		StackBinaryWriter<128> writer;
+		NetworkMessage msg(client, 1, SessionMessageType_Ping, writer);
+		writer.WriteU32(32);
+		StackBinaryReader<32> reader;
+		if (!msg.Send(reader))
+			return logger.Error(TC("Failed to get message"));
+
+		u8 value = reader.ReadByte();
+		logger.Info(TC("Got value %u"), value);
+		return true;
+	}
+
+	bool TestClientServer2(LoggerWithWriter& logger, const StringBufferBase& testRootDir)
+	{
+		auto& logWriter = logger.m_writer;
+		NetworkBackendTcp serverTcp(logWriter, TC("ServerTcp"));
+		NetworkBackendTcp clientTcp(logWriter, TC("ClientTcp"));
+
+		bool ctorSuccess = true;
+		NetworkServer server(ctorSuccess, logWriter);
+		NetworkClient client(ctorSuccess, logWriter);
+
+		server.RegisterService(1, [&](const ConnectionInfo& connectionInfo, u8 messageType, BinaryReader& reader, BinaryWriter& writer)
+			{
+				logger.Info(TC("Got ping!"));
+				UBA_ASSERT(messageType == SessionMessageType_Ping);
+				writer.WriteByte(42);
+				return true;
+			});
+
+		if (!client.StartListen(clientTcp, 1239))
+			return logger.Error(TC("Client failed to listen"));
+		Sleep(100);
+		if (!server.AddClient(serverTcp, TC("127.0.0.1"), 1239))
+			return logger.Error(TC("Server failed to connect"));
+
+		u64 time = GetTime();
+		while (!client.GetConnectionCount())
+		{
+			if (TimeToMs(GetTime() - time) > 4000)
+				return logger.Error(TC("Client failed to establish connection"));
+			Sleep(100);
+		}
+
 		StackBinaryWriter<128> writer;
 		NetworkMessage msg(client, 1, SessionMessageType_Ping, writer);
 		writer.WriteU32(32);
