[Uba] * Added code to handle leaked shm_open. This is made by creating "lock files" which only purpose is to keep track of which shm_open that is used and not. When a shm_open is created a lock-file is also created and locked. When process dies, all files gets unlocked. When a new session starts we can traverse the lock-files and see which ones are not locked, those belonged to a process that does not exist anymore, and we can call shm_unlink on those and delete them. This change solves the ever leaking shm_open entries causing machine to eventually go oom

diff --git a/Engine/Source/Programs/UnrealBuildAccelerator/Common/Private/UbaFileMapping.cpp b/Engine/Source/Programs/UnrealBuildAccelerator/Common/Private/UbaFileMapping.cpp
--- a/Engine/Source/Programs/UnrealBuildAccelerator/Common/Private/UbaFileMapping.cpp
+++ b/Engine/Source/Programs/UnrealBuildAccelerator/Common/Private/UbaFileMapping.cpp
@@ -1,11 +1,16 @@
 // Copyright Epic Games, Inc. All Rights Reserved.
 
 #include "UbaFileMapping.h"
+#include "UbaDirectoryIterator.h"
 #include "UbaPlatform.h"
 #include "UbaProcessStats.h"
 #include "UbaTimer.h"
 #include "UbaFile.h"
 
+#if !PLATFORM_WINDOWS
+#include <sys/file.h>
+#endif
+
 namespace uba
 {
 #if PLATFORM_WINDOWS
@@ -26,9 +31,10 @@ namespace uba
 	int asFileDescriptor(FileHandle fh);
 #endif
 
-	Atomic<u64> g_mappingUidCounter = 12345;
+	ReaderWriterLock g_mappingUidCounterLock;
+	Atomic<u64> g_mappingUidCounter;
 
-	FileMappingHandle CreateMemoryMappingW(u32 flProtect, u64 maxSize, const tchar* name)
+	FileMappingHandle CreateMemoryMappingW(Logger& logger, u32 flProtect, u64 maxSize, const tchar* name)
 	{
 		ExtendedTimerScope ts(SystemStats::GetCurrent().createFileMapping);
 #if PLATFORM_WINDOWS
@@ -37,39 +43,114 @@ namespace uba
 		UBA_ASSERT(!name);
 		UBA_ASSERT((flProtect & (~u32(PAGE_READWRITE | SEC_RESERVE))) == 0);
 
-		int oflags = O_CREAT | O_RDWR | O_NOFOLLOW | O_EXCL;
-		//		int oflags = O_CREAT | O_RDWR | O_EXCL;
+		// Since we need to not leak file mappings we use files as a trick to know which ones are used and not
+		StringBuffer<64> lockDir;
+		lockDir.Append("/tmp/uba_shm_locks");
+
+		ScopedWriteLock lock(g_mappingUidCounterLock);
+		if (!g_mappingUidCounter)
+		{
+			// Create dir
+			if (mkdir(lockDir.data, 0777) == -1)
+				if (errno != EEXIST)
+				{
+					UBA_ASSERTF(false, "Failed to create %s (%s)", lockDir.data, strerror(errno));
+					return {};
+				}
+
+			// Clear out all orphaned shm_open
+			TraverseDir(logger, lockDir.data,
+				[&](const DirectoryEntry& e)
+				{
+					u32 uid = strtoul(e.name, nullptr, 10);
+
+					StringBuffer<128> lockFile;
+					lockFile.Append(lockDir).EnsureEndsWithSlash().Append(e.name);
+					int lockFd = open(lockFile.data, O_RDWR, S_IRUSR | S_IWUSR);
+					if (lockFd == -1)
+					{
+						if (errno == EPERM)
+						{
+							g_mappingUidCounter = uid;
+							return;
+						}
+						UBA_ASSERTF(false, "Failed to open %s (%s)", lockFile.data, strerror(errno));
+					}
+
+					if (flock(lockFd, LOCK_EX | LOCK_NB) == 0)
+					{
+						StringBuffer<64> uidName;
+						GetMappingHandleName(uidName, uid);
+						if (shm_unlink(uidName.data) == 0)
+							logger.Info("Removed old shared memory %s", uidName.data);
+						remove(lockFile.data);
+					}
+					else
+					{
+						g_mappingUidCounter = uid;
+					}
+					close(lockFd);
+				});
+
+			if (g_mappingUidCounter)
+				logger.Info("Starting shared memory files at %u", g_mappingUidCounter.load());
+		}
+		lock.Leave();
+
+		// Let's find a free shm name
+		StringBuffer<128> lockFile;
+		u64 uid;
+		int shmFd;
+		int lockFd;
 
-		u64 uid = ++g_mappingUidCounter;
-		int fd;
 		while (true)
 		{
-			StringBuffer<64> uidName;
-			GetMappingHandleName(uidName, uid);
+			uid = ++g_mappingUidCounter;
 
-			fd = shm_open(uidName.data, oflags, S_IRUSR | S_IWUSR);
-			if (fd != -1)
-				break;
-			if (errno == EEXIST)
+			lockFile.Clear().Append(lockDir).EnsureEndsWithSlash().AppendValue(uid);
+
+			lockFd = open(lockFile.data, O_CREAT | O_RDWR | O_NOFOLLOW | O_EXCL, S_IRUSR | S_IWUSR);
+			if (lockFd == -1)
 			{
-				uid = ++g_mappingUidCounter;
-				//if (shm_unlink(uidName.data) == -1)
-				//{
-				//	printf("Failed to shm_unlink %s (%s)\n", uidName.data, strerror(errno));
-				//	uid = ++g_mappingUidCounter;
-				//}
+				if (errno == EEXIST)
+					continue;
+				UBA_ASSERTF(false, "Failed to open/create %s (%s)", lockFile.data, strerror(errno));
 				continue;
 			}
-			UBA_ASSERTF(false, "Failed to create filemapping with name %hs - %hs", uidName.data, strerror(errno));
+
+			if (flock(lockFd, LOCK_EX | LOCK_NB) != 0) // Some other process is using this one
+			{
+				close(lockFd);
+				continue;
+			}
+
+			StringBuffer<64> uidName;
+			GetMappingHandleName(uidName, uid);
+
+			int oflags = O_CREAT | O_RDWR | O_NOFOLLOW | O_EXCL;
+			shmFd = shm_open(uidName.data, oflags, S_IRUSR | S_IWUSR);
+			if (shmFd != -1)
+				break;
+			remove(lockFile.data);
+			close(lockFd);
+			SetLastError(errno);
+			UBA_ASSERTF(false, "This should not happen.. someone created shm %s without lock-file %s (%s)", uidName.data, lockFile.data, strerror(errno));
 			return {};
 		}
-		//shm_unlink("/myregionTest");
 
 		if (maxSize != 0)
-			if (ftruncate(fd, (s64)maxSize) == -1)
-				UBA_ASSERTF(false, "Failed to truncate file mapping '%s' to size %llu (%s)" , name, maxSize, strerror(errno));
-
-		return { fd, uid };
+		{
+			if (ftruncate(shmFd, (s64)maxSize) == -1)
+			{
+				SetLastError(errno);
+				close(shmFd);
+				remove(lockFile.data);
+				close(lockFd);
+				//UBA_ASSERTF(false, "Failed to truncate file mapping '%s' to size %llu (%s)", name, maxSize, strerror(errno));
+				return {};
+			}
+		}
+		return { shmFd, lockFd, uid };
 #endif
 	}
 
@@ -80,9 +161,9 @@ namespace uba
 		return { InternalCreateFileMappingW(asHANDLE(hFile), flProtect, (DWORD)ToHigh(maxSize), ToLow(maxSize), NULL) };
 #else
 		FileMappingHandle h;
-		h.fd = asFileDescriptor(hFile);
-		lseek(h.fd, maxSize-1, SEEK_SET);
-		write(h.fd, "", 1);
+		h.shmFd = asFileDescriptor(hFile);
+		lseek(h.shmFd, maxSize-1, SEEK_SET);
+		write(h.shmFd, "", 1);
 		return h;
 #endif
 	}
@@ -99,8 +180,8 @@ namespace uba
 		if (desiredAccess & FILE_MAP_WRITE)
 			prot |= PROT_WRITE;
 		UBA_ASSERT(hFileMappingObject.IsValid());
-		int fd = hFileMappingObject.fd;
-		void* rptr = mmap(NULL, bytesToMap, prot, MAP_SHARED, fd, s64(offset));
+		int shmFd = hFileMappingObject.shmFd;
+		void* rptr = mmap(NULL, bytesToMap, prot, MAP_SHARED, shmFd, s64(offset));
 		if (rptr != MAP_FAILED)
 			return (u8*)rptr;
 		//UBA_ASSERTF(false, "Failed to map file with fd %i, desiredAccess %u offset %llu, bytesToMap %llu (%s)", fd, desiredAccess, offset, bytesToMap, strerror(errno));
@@ -137,22 +218,26 @@ namespace uba
 #if PLATFORM_WINDOWS
 		return CloseHandle(h.handle);
 #else
-		if (h.fd == -1)
+		if (h.shmFd == -1)
 			return true;
-		if (h.uid != ~u64(0))
-		{
-			if (close(h.fd) != 0)
-				UBA_ASSERT(false);
+		if (h.uid == ~u64(0))
+			return true;
+		if (close(h.shmFd) != 0)
+			UBA_ASSERT(false);
 
-			StringBuffer<64> uidName;
-			GetMappingHandleName(uidName, h.uid);
-			if (shm_unlink(uidName.data) != 0)
-			{
-				SetLastError(errno);
-				//UBA_ASSERTF(false, "Failed to unlink %s (%s)", uidName.data, strerror(errno));
-				return false;
-			}
+		StringBuffer<64> uidName;
+		GetMappingHandleName(uidName, h.uid);
+		if (shm_unlink(uidName.data) != 0)
+		{
+			SetLastError(errno);
+			UBA_ASSERTF(false, "Failed to unlink %s (%s)", uidName.data, strerror(errno));
+			return false;
 		}
+
+		StringBuffer<128> lockFile;
+		lockFile.Append("/tmp/uba_shm_locks").EnsureEndsWithSlash().AppendValue(h.uid);
+		remove(lockFile.data);
+		close(h.lockFd);
 		return true;
 #endif
 	}
@@ -312,7 +397,7 @@ namespace uba
 				{
 					UBA_ASSERT(&storage == &m_storage[MappedView_Transient]);
 					u64 capacity = (IsWindows ? 32ull : 8ull) * 1024 * 1024 * 1024; // Linux can't have larger than 8gb
-					file-> handle = uba::CreateMemoryMappingW(PAGE_READWRITE | SEC_RESERVE, capacity);
+					file-> handle = uba::CreateMemoryMappingW(m_logger, PAGE_READWRITE | SEC_RESERVE, capacity);
 					if (!file->handle.IsValid())
 					{
 						m_logger.Error(TC("%s - Failed to create memory map (%s)"), file->name, LastErrorToText().data);
@@ -465,7 +550,7 @@ namespace uba
 
 	bool FileMappingAllocator::Init(u64 blockSize, u64 capacity)
 	{
-		m_mappingHandle = uba::CreateMemoryMappingW(PAGE_READWRITE|SEC_RESERVE, capacity);
+		m_mappingHandle = uba::CreateMemoryMappingW(m_logger, PAGE_READWRITE|SEC_RESERVE, capacity);
 		if (!m_mappingHandle.IsValid())
 		{
 			m_logger.Error(TC("%s - Failed to create memory map (%s)"), m_name, LastErrorToText().data);
diff --git a/Engine/Source/Programs/UnrealBuildAccelerator/Common/Private/UbaSession.cpp b/Engine/Source/Programs/UnrealBuildAccelerator/Common/Private/UbaSession.cpp
--- a/Engine/Source/Programs/UnrealBuildAccelerator/Common/Private/UbaSession.cpp
+++ b/Engine/Source/Programs/UnrealBuildAccelerator/Common/Private/UbaSession.cpp
@@ -1025,12 +1025,12 @@ namespace uba
 		m_systemPath.Append(TC("/nonexistingpath"));
 		#endif
 
-		m_fileMappingTableHandle = uba::CreateMemoryMappingW(PAGE_READWRITE, FileMappingTableMemSize);
+		m_fileMappingTableHandle = uba::CreateMemoryMappingW(m_logger, PAGE_READWRITE, FileMappingTableMemSize);
 		UBA_ASSERT(m_fileMappingTableHandle.IsValid());
 		m_fileMappingTableMem = MapViewOfFile(m_fileMappingTableHandle, FILE_MAP_WRITE, 0, FileMappingTableMemSize);
 		UBA_ASSERT(m_fileMappingTableMem);
 
-		m_directoryTableHandle = uba::CreateMemoryMappingW(PAGE_READWRITE, DirTableMemSize);
+		m_directoryTableHandle = uba::CreateMemoryMappingW(m_logger, PAGE_READWRITE, DirTableMemSize);
 		UBA_ASSERT(m_directoryTableHandle.IsValid());
 		m_directoryTableMem = MapViewOfFile(m_directoryTableHandle, FILE_MAP_WRITE, 0, DirTableMemSize);
 		UBA_ASSERT(m_directoryTableMem);
diff --git a/Engine/Source/Programs/UnrealBuildAccelerator/Common/Private/UbaTrace.cpp b/Engine/Source/Programs/UnrealBuildAccelerator/Common/Private/UbaTrace.cpp
--- a/Engine/Source/Programs/UnrealBuildAccelerator/Common/Private/UbaTrace.cpp
+++ b/Engine/Source/Programs/UnrealBuildAccelerator/Common/Private/UbaTrace.cpp
@@ -47,7 +47,7 @@ namespace uba
 	bool Trace::StartWrite(const tchar* namedTrace, u64 traceMemCapacity)
 	{
 		m_memoryCapacity = traceMemCapacity;
-		m_memoryHandle = uba::CreateMemoryMappingW(PAGE_READWRITE|SEC_RESERVE, m_memoryCapacity, namedTrace);
+		m_memoryHandle = uba::CreateMemoryMappingW(m_logger, PAGE_READWRITE|SEC_RESERVE, m_memoryCapacity, namedTrace);
 		if (!m_memoryHandle.IsValid())
 			return false;
 
@@ -289,7 +289,7 @@ namespace uba
 		#if PLATFORM_WINDOWS
 		StringBuffer<245> channelMutex;
 		channelMutex.Append(TC("Uba")).Append(channelName).Append(TC("Channel"));
-		m_memHandle = uba::CreateMemoryMappingW(PAGE_READWRITE, 256, channelName);
+		m_memHandle = uba::CreateMemoryMappingW(m_logger, PAGE_READWRITE, 256, channelName);
 		if (!m_memHandle.IsValid())
 		{
 			MessageBox(NULL, TC("Failed to create file mapping for trace channel"), TC("UbaVisualizer"), 0);
diff --git a/Engine/Source/Programs/UnrealBuildAccelerator/Common/Public/UbaFileMapping.h b/Engine/Source/Programs/UnrealBuildAccelerator/Common/Public/UbaFileMapping.h
--- a/Engine/Source/Programs/UnrealBuildAccelerator/Common/Public/UbaFileMapping.h
+++ b/Engine/Source/Programs/UnrealBuildAccelerator/Common/Public/UbaFileMapping.h
@@ -33,18 +33,19 @@ namespace uba
 
 	struct FileMappingHandle
 	{
-		FileMappingHandle(int fd_ = -1, u64 uid_ = ~u64(0)) : fd(fd_), uid(uid_) { };
-		FileMappingHandle(const FileMappingHandle& o) { fd = o.fd; uid = o.uid; }
-		int fd;
+		FileMappingHandle(int shmFd_ = -1, int lockFd_ = -1, u64 uid_ = ~u64(0)) : shmFd(shmFd_), lockFd(lockFd_), uid(uid_) { };
+		FileMappingHandle(const FileMappingHandle& o) { shmFd = o.shmFd; uid = o.uid; }
+		int shmFd;
+		int lockFd;
 		u64 uid;
-		bool operator==(const FileMappingHandle& o) const { return fd == o.fd; }
-		bool IsValid() const { return fd != -1; }
+		bool operator==(const FileMappingHandle& o) const { return shmFd == o.shmFd; }
+		bool IsValid() const { return shmFd != -1; }
 		u64 ToU64() const { return uid; }
 		void FromU64(u64 v) { UBA_ASSERT(false); uid = v; }
 	};
 #endif
 
-	FileMappingHandle CreateMemoryMappingW(u32 flProtect, u64 maxSize, const tchar* name = nullptr);
+	FileMappingHandle CreateMemoryMappingW(Logger& logger, u32 flProtect, u64 maxSize, const tchar* name = nullptr);
 	FileMappingHandle CreateFileMappingW(FileHandle hFile, u32 flProtect, u64 maxSize = 0);
 	u8* MapViewOfFile(FileMappingHandle hFileMappingObject, u32 dwDesiredAccess, u64 offset, u64 dwNumberOfBytesToMap);
 	bool MapViewCommit(void* address, u64 size);
diff --git a/Engine/Source/Programs/UnrealBuildAccelerator/Visualizer/Private/UbaTraceReader.cpp b/Engine/Source/Programs/UnrealBuildAccelerator/Visualizer/Private/UbaTraceReader.cpp
--- a/Engine/Source/Programs/UnrealBuildAccelerator/Visualizer/Private/UbaTraceReader.cpp
+++ b/Engine/Source/Programs/UnrealBuildAccelerator/Visualizer/Private/UbaTraceReader.cpp
@@ -121,7 +121,7 @@ namespace uba
 	bool TraceReader::StartReadClient(TraceView& out, NetworkClient& client)
 	{
 		u32 traceMemSize = 128 * 1024 * 1024;
-		m_memoryHandle = uba::CreateMemoryMappingW(PAGE_READWRITE, traceMemSize);
+		m_memoryHandle = uba::CreateMemoryMappingW(m_logger, PAGE_READWRITE, traceMemSize);
 		if (!m_memoryHandle.IsValid())
 			return false;
 		m_memoryBegin = MapViewOfFile(m_memoryHandle, FILE_MAP_ALL_ACCESS, 0, traceMemSize);
