[Uba]
* Added missing posix_spawnattr_setgroup call

diff --git a/Engine/Source/Programs/UnrealBuildAccelerator/Common/Private/UbaProcess.cpp b/Engine/Source/Programs/UnrealBuildAccelerator/Common/Private/UbaProcess.cpp
--- a/Engine/Source/Programs/UnrealBuildAccelerator/Common/Private/UbaProcess.cpp
+++ b/Engine/Source/Programs/UnrealBuildAccelerator/Common/Private/UbaProcess.cpp
@@ -1291,6 +1291,8 @@ namespace uba
 			UBA_ASSERTF(res == 0, TC("posix_spawnattr_init"));
 			res = posix_spawnattr_setflags(&attr, flags);
 			UBA_ASSERTF(res == 0, TC("posix_spawnattr_setflags"));
+			res = posix_spawnattr_setpgroup(&attr, getpgrp());
+			UBA_ASSERTF(res == 0, TC("posix_spawnattr_setpgroup"));
 
 			posix_spawn_file_actions_t fileActions;
 			posix_spawn_file_actions_init(&fileActions);
