[Uba]
* Fixed link deadlock on build farm because we have linked some of our linux tools with jemalloc which does a readlink call during initialization. Since we have detoured readlink we call dlsym first time readlink is called. And dlsym makes a calloc call when resolving the next/true function causing jemalloc to initialize. Solution was to just always return "not found" for /etc/je_malloc.conf"

diff --git a/Engine/Source/Programs/UnrealBuildAccelerator/Detours/Private/PosixOS/UbaDetoursFunctionsPosix.cpp b/Engine/Source/Programs/UnrealBuildAccelerator/Detours/Private/PosixOS/UbaDetoursFunctionsPosix.cpp
--- a/Engine/Source/Programs/UnrealBuildAccelerator/Detours/Private/PosixOS/UbaDetoursFunctionsPosix.cpp
+++ b/Engine/Source/Programs/UnrealBuildAccelerator/Detours/Private/PosixOS/UbaDetoursFunctionsPosix.cpp
@@ -687,13 +687,13 @@ UBA_EXPORT char* UBA_WRAPPER(getcwd)(char* buf, size_t size)
 	return buf;
 }
 
-UBA_EXPORT char* UBA_WRAPPER(getenv)(const char* name)
-{
-	UBA_INIT_DETOUR(getenv, name);
-	auto res = TRUE_WRAPPER(getenv)(name);
-	DEBUG_LOG_TRUE("getenv", "(%s) -> %s", name, res ? res : "<null>");
-	return res;
-}
+//UBA_EXPORT char* UBA_WRAPPER(getenv)(const char* name)
+//{
+//	UBA_INIT_DETOUR(getenv, name);
+//	auto res = TRUE_WRAPPER(getenv)(name);
+//	DEBUG_LOG_TRUE("getenv", "(%s) -> %s", name, res ? res : "<null>");
+//	return res;
+//}
 
 UBA_EXPORT int UBA_WRAPPER(setenv)(const char* name, const char* value, int replace)
 {
@@ -728,6 +728,13 @@ UBA_EXPORT char* UBA_WRAPPER(realpath)(const char* path, char* resolved_path)
 
 UBA_EXPORT ssize_t UBA_WRAPPER(readlink)(const char* pathname, char* buf, size_t bufsiz)
 {
+	// Beautiful hack. Some of our tools use je_malloc and dlsym do memory allocation so we end up in a deadlock when initializing detour (since detour use dlsym to figure out true function)
+	if (Equals(pathname, "/etc/je_malloc.conf"))
+	{
+		errno = ENOENT;
+		return -1;
+	}
+
 	UBA_INIT_DETOUR(readlink, pathname, buf, bufsiz);
 
 	if (Equals(pathname, "/proc/self/exe"))
