[Uba]
* Enabled socket reuse on linux because it seems farm turns around faster than sockets gets freed by the os

diff --git a/Engine/Source/Programs/UnrealBuildAccelerator/Common/Private/UbaNetworkBackendTcp.cpp b/Engine/Source/Programs/UnrealBuildAccelerator/Common/Private/UbaNetworkBackendTcp.cpp
--- a/Engine/Source/Programs/UnrealBuildAccelerator/Common/Private/UbaNetworkBackendTcp.cpp
+++ b/Engine/Source/Programs/UnrealBuildAccelerator/Common/Private/UbaNetworkBackendTcp.cpp
@@ -309,7 +309,7 @@ namespace uba
 		auto listenSocketCleanup = MakeGuard([&]() { closesocket(listenSocket); listenSocket = INVALID_SOCKET; });
 
 		// This is here to be able to iterate fast when doing development.. seems like socket ends up in TIME_WAIT state after close and it takes some time to be able to use socket again
-		#if !PLATFORM_WINDOWS && UBA_DEBUG
+		#if !PLATFORM_WINDOWS// && UBA_DEBUG
 		int optval = 1;
 		::setsockopt(listenSocket, SOL_SOCKET, SO_REUSEADDR, (const char*)&optval, sizeof optval);
 		#endif
