[Uba]
* Bumped backtrace entry count from 10 to 16 on windows for asserts

diff --git a/Engine/Source/Programs/UnrealBuildAccelerator/Core/Private/UbaPlatform.cpp b/Engine/Source/Programs/UnrealBuildAccelerator/Core/Private/UbaPlatform.cpp
--- a/Engine/Source/Programs/UnrealBuildAccelerator/Core/Private/UbaPlatform.cpp
+++ b/Engine/Source/Programs/UnrealBuildAccelerator/Core/Private/UbaPlatform.cpp
@@ -156,7 +156,7 @@ namespace uba
 				}
 			}
 
-			const int kMaxCallers = 10;// 62;
+			const int kMaxCallers = 16;// 62;
 			void* callers[kMaxCallers];
 			if (u32 count = (func)(0, kMaxCallers, callers, NULL))
 			{
