[UBA]
* Minor optimization around super hack to resolve dump_syms hang on linux target when dump_syms is compiled with jemalloc

diff --git a/Engine/Source/Programs/UnrealBuildAccelerator/Detours/Private/PosixOS/UbaDetoursFunctionsPosix.cpp b/Engine/Source/Programs/UnrealBuildAccelerator/Detours/Private/PosixOS/UbaDetoursFunctionsPosix.cpp
--- a/Engine/Source/Programs/UnrealBuildAccelerator/Detours/Private/PosixOS/UbaDetoursFunctionsPosix.cpp
+++ b/Engine/Source/Programs/UnrealBuildAccelerator/Detours/Private/PosixOS/UbaDetoursFunctionsPosix.cpp
@@ -729,7 +729,7 @@ UBA_EXPORT char* UBA_WRAPPER(realpath)(const char* path, char* resolved_path)
 UBA_EXPORT ssize_t UBA_WRAPPER(readlink)(const char* pathname, char* buf, size_t bufsiz)
 {
 	// Beautiful hack. Some of our tools use je_malloc and dlsym do memory allocation so we end up in a deadlock when initializing detour (since detour use dlsym to figure out true function)
-	if (Equals(pathname, "/etc/je_malloc.conf"))
+	if (!g_isDetouring && Equals(pathname, "/etc/je_malloc.conf"))
 	{
 		errno = ENOENT;
 		return -1;
