[UBA]
* Fixed macos compile error

diff --git a/Engine/Source/Programs/UnrealBuildAccelerator/Detours/Private/PosixOS/UbaDetoursFunctionsPosix.cpp b/Engine/Source/Programs/UnrealBuildAccelerator/Detours/Private/PosixOS/UbaDetoursFunctionsPosix.cpp
--- a/Engine/Source/Programs/UnrealBuildAccelerator/Detours/Private/PosixOS/UbaDetoursFunctionsPosix.cpp
+++ b/Engine/Source/Programs/UnrealBuildAccelerator/Detours/Private/PosixOS/UbaDetoursFunctionsPosix.cpp
@@ -37,7 +37,6 @@ using namespace uba;
 	DETOURED_FUNCTION(fchdir) \
 	DETOURED_FUNCTION(chroot) \
 	DETOURED_FUNCTION(getcwd) \
-	DETOURED_FUNCTION(getenv) \
 	DETOURED_FUNCTION(setenv) \
 	DETOURED_FUNCTION(realpath) \
 	DETOURED_FUNCTION(readlink) \
