[Uba]
* Implemented simple Http get using tcp directly in order to work on all platforms and also be lightweight
* Changed so aws queries are using above implementation and enabled aws checks on linux
* Added UbaCli "-checkaws" option to be able to test this functionality

diff --git a/Engine/Source/Programs/UnrealBuildAccelerator/Cli/Private/UbaCli.cpp b/Engine/Source/Programs/UnrealBuildAccelerator/Cli/Private/UbaCli.cpp
--- a/Engine/Source/Programs/UnrealBuildAccelerator/Cli/Private/UbaCli.cpp
+++ b/Engine/Source/Programs/UnrealBuildAccelerator/Cli/Private/UbaCli.cpp
@@ -7,6 +7,8 @@
 #include "UbaPlatform.h"
 #include "UbaVersion.h"
 
+#include "UbaAWS.h"
+
 #if PLATFORM_WINDOWS
 #include <dbghelp.h>
 #include <io.h>
@@ -55,6 +57,7 @@ namespace uba
 		logger.Info(TC("   -loop=<count>           Loop the commandline <count> number of times. Will exit when/if it fails"));
 		logger.Info(TC("   -workdir=<dir>          Working directory"));
 		logger.Info(TC("   -checkcas               Check so all cas entries are correct"));
+		logger.Info(TC("   -checkaws               Check if we are inside aws and output information about aws"));
 		logger.Info(TC("   -getcas                 Will print hash of application"));
 		logger.Info(TC("   -summary                Print summary at the end of a session"));
 		logger.Info(TC("   -nocustomalloc          Disable custom allocator for processes. If you see odd crashes this can be tested"));
@@ -128,6 +131,7 @@ namespace uba
 		bool disableCustomAllocator = false;
 		bool quiet = false;
 		bool checkCas = false;
+		bool checkAws = false;
 		bool getCas = false;
 		bool printSummary = false;
 		u32 loopCount = 1;
@@ -236,6 +240,10 @@ namespace uba
 			{
 				checkCas = true;
 			}
+			else if (name.Equals(TC("-checkaws")))
+			{
+				checkAws = true;
+			}
 			else if (name.Equals(TC("-getcas")))
 			{
 				getCas = true;
@@ -279,6 +287,24 @@ namespace uba
 			return success ? 0 : -1;
 		}
 
+		if (checkAws)
+		{
+			AWS aws;
+			StringBuffer<> info;
+			if (aws.Init(logger, info, TC("UbaCli")))
+			{
+				logger.Info(TC("We are inside AWS: %s (%s)"), info.data, aws.GetAvailabilityZone());
+				
+				StringBuffer<> reason;
+				u64 terminateTime;
+				if (aws.IsTerminating(logger, reason, terminateTime))
+					logger.Info(TC(".. and are being terminated: %s"), reason.data);
+			}
+			else
+				logger.Info(TC("Seems like we are not running inside aws."));
+			return 0;
+		}
+
 		if (commandType == CommandType_NotSet)
 		{
 			const tchar* errorMsg = argc == 1 ? TC("") : TC("\nERROR: First argument must be command type. Options are 'local,remote or native'");
diff --git a/Engine/Source/Programs/UnrealBuildAccelerator/Common/Private/UbaNetworkBackendTcp.cpp b/Engine/Source/Programs/UnrealBuildAccelerator/Common/Private/UbaNetworkBackendTcp.cpp
--- a/Engine/Source/Programs/UnrealBuildAccelerator/Common/Private/UbaNetworkBackendTcp.cpp
+++ b/Engine/Source/Programs/UnrealBuildAccelerator/Common/Private/UbaNetworkBackendTcp.cpp
@@ -9,6 +9,7 @@
 
 #if PLATFORM_LINUX
 #include <netinet/tcp.h>
+#define strcpy_s(a, b, c) strcpy(a, c)
 #endif
 
 #if PLATFORM_WINDOWS
@@ -852,4 +853,132 @@ namespace uba
 				return true;
 		return true;
 	}
+
+	HttpConnection::HttpConnection()
+	{
+		m_socket = INVALID_SOCKET;
+		*m_host = 0;
+	}
+
+	HttpConnection::~HttpConnection()
+	{
+		#if PLATFORM_WINDOWS
+		if (m_wsaInitDone)
+			WSACleanup();
+		#endif
+		if (m_socket != INVALID_SOCKET)
+			closesocket(m_socket);
+	}
+
+	bool HttpConnection::Connect(Logger& logger, const char* host)
+	{
+		#if PLATFORM_WINDOWS
+		WSADATA wsaData;
+		if (!m_wsaInitDone)
+			if (int res = WSAStartup(MAKEWORD(2, 2), &wsaData))
+				return logger.Error(TC("WSAStartup failed (%d)"), res);
+		m_wsaInitDone = true;
+		#endif
+
+		protoent* protoent = getprotobyname("tcp");
+		if (protoent == NULL)
+			return logger.Error(TC("HttpRequest: socket error"));
+
+		SOCKET sock = socket(AF_INET, SOCK_STREAM, protoent->p_proto);
+		if (sock == -1)
+			return logger.Error(TC("HttpRequest: socket error"));
+		auto socketClose = MakeGuard([sock]() { closesocket(sock); });
+
+		hostent* hostent = gethostbyname(host);
+		if (hostent == NULL)
+			return logger.Error(TC("HttpRequest: gethostbyname error"));
+
+		unsigned long in_addr = inet_addr(inet_ntoa(*(struct in_addr*)*(hostent->h_addr_list)));
+		if (in_addr == INADDR_NONE)
+			return logger.Error(TC("HttpRequest: inet_addr error"));
+
+		sockaddr_in sockaddr_in;
+		sockaddr_in.sin_addr.s_addr = in_addr;
+		sockaddr_in.sin_family = AF_INET;
+		sockaddr_in.sin_port = htons(80);
+
+		if (connect(sock, (struct sockaddr*)&sockaddr_in, sizeof(sockaddr_in)) == -1)
+			return false;// logger.Error(TC("HttpRequest: connect error"));
+
+		socketClose.Cancel();
+
+		strcpy_s(m_host, sizeof_array(m_host), host);
+		m_socket = sock;
+		return true;
+	}
+
+	bool HttpConnection::Get(Logger& logger, StringBufferBase& outResponse, u32& outStatusCode, const char* host, const char* path)
+	{
+		// TODO: Fix so we reuse socket connection for multiple queries
+		if (*m_host)// && _stricmp(m_host, host) != 0)
+		{
+			closesocket(m_socket);
+			m_socket = INVALID_SOCKET;
+			*m_host = 0;
+		}
+
+		if (m_socket == INVALID_SOCKET)
+			if (!Connect(logger, host))
+				return false;
+
+		char request[512];
+		int requestLen = snprintf(request, 512, "GET /%s HTTP/1.1\r\nHost: %s\r\n\r\n", path, m_host);
+
+		int totalBytesSent = 0;
+		while (totalBytesSent < requestLen) {
+			int bytesSent = send(m_socket, request + totalBytesSent, requestLen - totalBytesSent, 0);
+			if (bytesSent == -1)
+				return logger.Error(TC("HttpRequest: send error"));
+			totalBytesSent += bytesSent;
+		}
+
+#if PLATFORM_WINDOWS
+#pragma warning(push)
+#pragma warning(disable:6386) // analyzer claims that buf can have buffer overrun.. but can't see how that can happen
+#endif
+
+		u32 readPos = 0;
+		char buf[4*1024];
+		int bytesRead = 0;
+		while ((bytesRead = recv(m_socket, buf + readPos, sizeof(buf) - readPos, 0)) > 0)
+			readPos += bytesRead;
+
+		if (bytesRead == -1)
+			return logger.Error(TC("HttpRequest: recv error"));
+
+		if (readPos == sizeof(buf))
+			return logger.Error(TC("HttpRequest: buffer overflow"));
+
+		buf[readPos] = 0;
+
+#if PLATFORM_WINDOWS
+#pragma warning(pop)
+#endif
+
+		char* firstSpace = strchr(buf, ' '); // After version (where status code starts)
+		if (!firstSpace)
+			return logger.Error(TC("HttpRequest: first space not found (read %u)"), readPos);
+		char* secondSpace = strchr(firstSpace + 1, ' '); // after status code
+		if (!secondSpace)
+			return logger.Error(TC("HttpRequest: second space not found"));
+		*secondSpace = 0;
+		outStatusCode = strtoul(firstSpace + 1, nullptr, 10);
+
+		char* bodyStart = strstr(secondSpace + 1, "\r\n\r\n");
+		if (!bodyStart)
+			return logger.Error(TC("HttpRequest: no body found"));
+
+		#if PLATFORM_WINDOWS
+		outResponse.Appendf(L"%S", bodyStart + 4);
+		#else
+		outResponse.Append(bodyStart + 4);
+		#endif
+
+		return true;
+	}
 }
\ No newline at end of file
diff --git a/Engine/Source/Programs/UnrealBuildAccelerator/Common/Public/UbaAWS.h b/Engine/Source/Programs/UnrealBuildAccelerator/Common/Public/UbaAWS.h
--- a/Engine/Source/Programs/UnrealBuildAccelerator/Common/Public/UbaAWS.h
+++ b/Engine/Source/Programs/UnrealBuildAccelerator/Common/Public/UbaAWS.h
@@ -2,186 +2,56 @@
 
 #pragma once
 
-#if PLATFORM_WINDOWS
-
-#include "UbaPlatform.h"
+#include "UbaNetworkBackend.h"
 #include "UbaStringBuffer.h"
-#include "UbaLogger.h"
-
-#include <winhttp.h> 
-#pragma comment (lib, "Winhttp.lib")
 
 #define UBA_USE_AWS
 
 namespace uba
 {
-	struct HttpRequest
-	{
-		bool Init(Logger& logger, HINTERNET connection, const wchar_t* verb, const wchar_t* objectName)
-		{
-			m_connection = connection;
-			m_verb = verb;
-			m_objectName = objectName;
-			m_request = WinHttpOpenRequest(m_connection, verb, objectName, NULL, WINHTTP_NO_REFERER, WINHTTP_DEFAULT_ACCEPT_TYPES, WINHTTP_FLAG_REFRESH);
-
-			return m_request != NULL;
-		}
-
-		~HttpRequest()
-		{
-			if (m_request)
-				WinHttpCloseHandle(m_request);
-		}
-
-		bool Send(Logger& logger, StringBufferBase& outContent, const wchar_t* headers = nullptr, u32* outStatusCode = nullptr)
-		{
-			if (!m_request)
-				return false;
-
-			if (!WinHttpSendRequest(m_request, headers ? headers : WINHTTP_NO_ADDITIONAL_HEADERS, headers ? u32(wcslen(headers)) : 0, WINHTTP_NO_REQUEST_DATA, 0, 0, 0))
-			{
-				DWORD err = GetLastError();
-
-				WinHttpCloseHandle(m_request);
-				m_request = NULL;
-
-				if (err == ERROR_WINHTTP_CANNOT_CONNECT || err == ERROR_WINHTTP_TIMEOUT)
-					return false;
-
-				if (!Init(logger, m_connection, m_verb, m_objectName))
-					return logger.Error(L"HttpOpenRequestW reinit error: %ls", LastErrorToText().data);
-
-				if (!WinHttpSendRequest(m_request, headers ? headers : WINHTTP_NO_ADDITIONAL_HEADERS, headers ? u32(wcslen(headers)) : 0, WINHTTP_NO_REQUEST_DATA, 0, 0, 0))
-				{
-					err = GetLastError();
-					if (err != ERROR_WINHTTP_CANNOT_CONNECT && err != ERROR_WINHTTP_TIMEOUT)
-						if (headers) // hack to prevent log for non-aws machines
-							logger.Info(L"WinHttpSendRequest error: %ls", LastErrorToText().data);
-
-					WinHttpCloseHandle(m_request);
-					m_request = NULL;
-					return false;
-				}
-			}
-
-			if (!WinHttpReceiveResponse(m_request, NULL))
-			{
-				WinHttpCloseHandle(m_request);
-				m_request = NULL;
-				return false;
-			}
-
-			if (outStatusCode)
-			{
-				DWORD bufferLen = sizeof(u32);
-				if (!WinHttpQueryHeaders(m_request, WINHTTP_QUERY_STATUS_CODE|WINHTTP_QUERY_FLAG_NUMBER, WINHTTP_HEADER_NAME_BY_INDEX, outStatusCode, &bufferLen, WINHTTP_NO_HEADER_INDEX))
-					return logger.Error(L"Error %u in WinHttpQueryHeaders", GetLastError());
-			}
-
-			for (u32 i=0; i!=100; ++i)
-			{
-				// Check for available data.
-				DWORD dwSize = 0;
-				if (!WinHttpQueryDataAvailable(m_request, &dwSize))
-					return logger.Error(L"Error %u in WinHttpQueryDataAvailable", GetLastError());
-
-				if (dwSize == 0)
-					return true;
-
-				CHAR szBuffer[4097];
-				DWORD dwDownloaded = 0;
-				if (!WinHttpReadData(m_request, (LPVOID)szBuffer, dwSize, &dwDownloaded))
-					return logger.Error(L"Error %u in WinHttpReadData", GetLastError());
-				szBuffer[dwDownloaded] = 0;
-				outContent.Appendf(L"%hs", szBuffer);
-			}
-			return logger.Error(L"Unknown error reading http query data", GetLastError());
-		}
-
-		HINTERNET m_connection = NULL;
-		HINTERNET m_request = NULL;
-		const wchar_t* m_verb = nullptr;
-		const wchar_t* m_objectName = nullptr;
-	};
-
 	class AWS
 	{
 	public:
-		~AWS()
-		{
-			if (m_connection)
-				WinHttpCloseHandle(m_connection);
-			if (m_session)
-				WinHttpCloseHandle(m_session);
-		}
 
-		bool InitCore(Logger& logger, const wchar_t* application)
-		{
-			m_session = WinHttpOpen(L"AWS", WINHTTP_ACCESS_TYPE_DEFAULT_PROXY, WINHTTP_NO_PROXY_NAME, WINHTTP_NO_PROXY_BYPASS, 0);
-			if (!m_session)
-				return logger.Error(L"WinHttpOpen failed: %ls", LastErrorToText().data);
-			auto internetGuard = MakeGuard([&]() { WinHttpCloseHandle(m_session); m_session = NULL; });
+		static constexpr char g_imdsHost[]						= "169.254.169.254";
+		static constexpr char g_imdsInstanceId[]				= "latest/meta-data/instance-id";
+		static constexpr char g_imdsInstanceLifeCycle[]			= "latest/meta-data/instance-life-cycle";
+		static constexpr char g_imdsAutoScalingLifeCycleState[] = "latest/meta-data/autoscaling/target-lifecycle-state";
+		static constexpr char g_imdsInstanceAvailabilityZone[]	= "latest/meta-data/placement/availability-zone";
+		static constexpr char g_imdsSpotInstanceAction[]		= "latest/meta-data/spot/instance-action";
 
-			if (!WinHttpSetTimeouts(m_session, 100, 100, 100, 100))
-				return logger.Error(L"WinHttpSetTimeouts failed: %ls", LastErrorToText().data);
 
-			m_connection = WinHttpConnect(m_session, L"169.254.169.254", INTERNET_DEFAULT_HTTP_PORT, 0);
-			if (!m_connection)
-				return logger.Error(L"WinHttpConnect failed: %ls", LastErrorToText().data);
-			auto connectionGuard = MakeGuard([&]() { WinHttpCloseHandle(m_connection); m_connection = NULL; });
-
-			HttpRequest tokenPut;
-			if (!tokenPut.Init(logger, m_connection, L"PUT", L"latest/api/token"))
-				return false;
-			StringBuffer<512> awsToken;
-			awsToken.Append(L"X-aws-ec2-metadata-token: ");
-			if (!tokenPut.Send(logger, awsToken, L"X-aws-ec2-metadata-token-ttl-seconds: 21600")) // Expires after 6 hours
-				return false;
-			if (awsToken.IsEmpty())
-				return false;
-
-			connectionGuard.Cancel();
-			internetGuard.Cancel();
-
-			WinHttpSetTimeouts(m_session, 1000, 1000, 1000, 1000);
-
-			m_awsToken = awsToken.data;
+		bool InitCore(Logger& logger, const tchar* application)
+		{
 			return true;
 		}
 
-		bool Init(Logger& logger, StringBufferBase& outExtraInfo, const wchar_t* application)
+		bool Init(Logger& logger, StringBufferBase& outExtraInfo, const tchar* application)
 		{
 			if (!InitCore(logger, application))
 				return false;
 
-			HttpRequest instanceIdReq;
-			if (instanceIdReq.Init(logger, m_connection, L"GET", L"latest/meta-data/instance-id"))
-			{
-				outExtraInfo.Append(L", AWS: ");
-				if (!instanceIdReq.Send(logger, outExtraInfo, m_awsToken.c_str()))
-					outExtraInfo.Clear();
-			}
+			HttpConnection http;
+
+			u32 statusCode = 0;
+
+			StringBuffer<128> instanceId;
+			if (!http.Get(logger, instanceId, statusCode, g_imdsHost, g_imdsInstanceId))
+				return false;
+			outExtraInfo.Append(TC(", AWS: ")).Append(instanceId);
 
-			HttpRequest lifeCycleReq;
-			if (lifeCycleReq.Init(logger, m_connection, L"GET", L"latest/meta-data/instance-life-cycle"))
+			StringBuffer<32> instanceLifeCycle;
+			if (http.Get(logger, instanceLifeCycle, statusCode, g_imdsHost, g_imdsInstanceLifeCycle))
 			{
-				StringBuffer<32> instanceLifeCycle;
-				if (lifeCycleReq.Send(logger, instanceLifeCycle, m_awsToken.c_str()))
-				{
-					outExtraInfo.Append(' ').Append(instanceLifeCycle);
-					m_isSpot = instanceLifeCycle.Contains(L"spot");
-				}
+				outExtraInfo.Append(' ').Append(instanceLifeCycle);
+				m_isSpot = instanceLifeCycle.Contains(TC("spot"));
 			}
 
-			if (m_autoscalingRequest.Init(logger, m_connection, L"GET", L"latest/meta-data/autoscaling/target-lifecycle-state"))
+			StringBuffer<32> autoscaling;
+			if (http.Get(logger, autoscaling, statusCode, g_imdsHost, g_imdsAutoScalingLifeCycleState) && statusCode == 200)
 			{
-				StringBuffer<32> instanceLifeCycle;
-				u32 statusCode = 0;
-				if (m_autoscalingRequest.Send(logger, instanceLifeCycle, m_awsToken.c_str(), &statusCode) && statusCode == 200)
-				{
-					outExtraInfo.Append(m_isSpot ? '/' : ' ').Append(L"autoscale");
-					m_isAutoscaling = true;
-				}
+				outExtraInfo.Append(m_isSpot ? '/' : ' ').Append(TC("autoscale"));
+				m_isAutoscaling = true;
 			}
 
 			if (!InitAvailabilityZone(logger))
@@ -192,11 +62,11 @@ namespace uba
 
 		bool InitAvailabilityZone(Logger& logger)
 		{
-			HttpRequest availabilityZoneReq;
-			if (!availabilityZoneReq.Init(logger, m_connection, L"GET", L"latest/meta-data/placement/availability-zone"))
-				return false;
-			StringBuffer<256> availabilityZone;
-			if (!availabilityZoneReq.Send(logger, availabilityZone, m_awsToken.c_str()))
+			HttpConnection http;
+
+			StringBuffer<128> availabilityZone;
+			u32 statusCode = 0;
+			if (!http.Get(logger, availabilityZone, statusCode, g_imdsHost, g_imdsInstanceAvailabilityZone))
 				return false;
 			m_availabilityZone = availabilityZone.data;
 			return true;
@@ -204,21 +74,21 @@ namespace uba
 
 		bool InitPolling(Logger& logger)
 		{
-			if (m_isSpot)
-				m_instanceActionRequest.Init(logger, m_connection, L"GET", L"latest/meta-data/spot/instance-action");
 			return true;
 		}
 
 		bool IsTerminating(Logger& logger, StringBufferBase& outReason, u64& outTerminationTimeMs)
 		{
+			HttpConnection http;
+
 			outTerminationTimeMs = 0;
 			if (m_isSpot)
 			{
 				StringBuffer<1024> content;
 				u32 statusCode = 0;
-				if (m_instanceActionRequest.Send(logger, content, m_awsToken.c_str(), &statusCode) && statusCode == 200)
+				if (http.Get(logger, content, statusCode, g_imdsHost, g_imdsSpotInstanceAction) && statusCode == 200)
 				{
-					outReason.Append(L"AWS spot instance interruption");
+					outReason.Append(TC("AWS spot instance interruption"));
 					return true;
 				}
 			}
@@ -227,17 +97,17 @@ namespace uba
 			{
 				StringBuffer<1024> content;
 				u32 statusCode = 0;
-				if (m_autoscalingRequest.Send(logger, content, m_awsToken.c_str(), &statusCode) && statusCode == 200)
+				if (http.Get(logger, content, statusCode, g_imdsHost, g_imdsAutoScalingLifeCycleState) && statusCode == 200)
 				{
 					//if (!content.Equals(L"InService"))
 					//{
 					//	wprintf(L"AWSACTION: AUTOSCALE: %ls\n", content.data);
 					//}
 
-					if (!content.Contains(L"InService")) // AWS can return "InServiceI" as well?
+					if (!content.Contains(TC("InService"))) // AWS can return "InServiceI" as well?
 					{
 						//wprintf(L"AWSACTION: AUTOSCALE REBALANCING!!!! (%ls)\n", content.data);
-						outReason.Append(L"AWS autoscale rebalancing");
+						outReason.Append(TC("AWS autoscale rebalancing"));
 						return true;
 					}
 				}
@@ -245,22 +115,13 @@ namespace uba
 			return false;
 		}
 
-		const wchar_t* GetAvailabilityZone()
+		const tchar* GetAvailabilityZone()
 		{
 			return m_availabilityZone.c_str();
 		}
 
-	private:
-		HINTERNET m_session = NULL;
-		HINTERNET m_connection = NULL;
-		TString m_awsToken;
 		TString m_availabilityZone;
-		HttpRequest m_autoscalingRequest;
-		HttpRequest m_instanceActionRequest;
-
 		bool m_isSpot = false;
 		bool m_isAutoscaling = false;
 	};
 }
-
-#endif // PLATFORM_WINDOWS
diff --git a/Engine/Source/Programs/UnrealBuildAccelerator/Common/Public/UbaNetworkBackend.h b/Engine/Source/Programs/UnrealBuildAccelerator/Common/Public/UbaNetworkBackend.h
--- a/Engine/Source/Programs/UnrealBuildAccelerator/Common/Public/UbaNetworkBackend.h
+++ b/Engine/Source/Programs/UnrealBuildAccelerator/Common/Public/UbaNetworkBackend.h
@@ -75,4 +75,23 @@ namespace uba
 		SendContext(u8 sendFlags = 0) : flags(sendFlags) { *data = 0; }
 		~SendContext() { }//UBA_ASSERT(isFinished); }
 	};
+
+	class HttpConnection
+	{
+	public:
+		HttpConnection();
+		~HttpConnection();
+		bool Get(Logger& logger, StringBufferBase& outResponse, u32& outStatusCode, const char* host, const char* path);
+
+	private:
+		bool Connect(Logger& logger, const char* host);
+
+		char m_host[256];
+		#if PLATFORM_WINDOWS
+		u64 m_socket;
+		bool m_wsaInitDone = false;
+		#else
+		int m_socket;
+		#endif
+	};
 }
