[UBA]
* Fixed possible memory alignment issue for instances allocated with placement new

diff --git a/Engine/Source/Programs/UnrealBuildAccelerator/Detours/Private/UbaDetoursShared.cpp b/Engine/Source/Programs/UnrealBuildAccelerator/Detours/Private/UbaDetoursShared.cpp
--- a/Engine/Source/Programs/UnrealBuildAccelerator/Detours/Private/UbaDetoursShared.cpp
+++ b/Engine/Source/Programs/UnrealBuildAccelerator/Detours/Private/UbaDetoursShared.cpp
@@ -23,13 +23,13 @@ namespace uba
 	// make sure they are initialized before detoured functions are called.
 	// (If this doesn't work we will need to resort to Functions returning a static variable)
 
-	static u8 g_memoryBlockMem[sizeof(MemoryBlock)];
+	static u64 g_memoryBlockMem[sizeof(MemoryBlock)/sizeof(u64)];
 	MemoryBlock& g_memoryBlock = *(MemoryBlock*)g_memoryBlockMem;
 
-	static u8 g_directoryTableMem[sizeof(DirectoryTable)];
+	static u64 g_directoryTableMem[sizeof(DirectoryTable)/sizeof(u64)];
 	DirectoryTable& g_directoryTable = *(DirectoryTable*)g_directoryTableMem;
 
-	static u8 g_mappedFileTableMem[sizeof(MappedFileTable)];
+	static u64 g_mappedFileTableMem[sizeof(MappedFileTable)/sizeof(u64)];
 	MappedFileTable& g_mappedFileTable = *(MappedFileTable*)g_mappedFileTableMem;
 
 	void InitStatics()
