[UBA]
* Fixed Build-StaticAnalysis.bat not working after move

diff --git a/Engine/Source/Programs/UnrealBuildAccelerator/Build-StaticAnalysis.bat b/Engine/Source/Programs/UnrealBuildAccelerator/Build-StaticAnalysis.bat
--- a/Engine/Source/Programs/UnrealBuildAccelerator/Build-StaticAnalysis.bat
+++ b/Engine/Source/Programs/UnrealBuildAccelerator/Build-StaticAnalysis.bat
@@ -3,7 +3,7 @@ setlocal
 
 pushd "%~dp0"
 
-set config=Shipping
+set config=Development
 
 if "%~1"=="" goto blank
 
@@ -13,7 +13,7 @@ set config=%1
 
 echo.
 echo === Building Targets ===
-pushd "../../../../../.."
+pushd "../../../.."
 call Engine/Build/BatchFiles/RunUBT.bat -NoUba -NoUbaLocal -NoSNDBS -NoXGE -Architectures=x64 ^
 	-Target="UbaAgent Win64 %config%" ^
 	-Target="UbaCli Win64 %config%" ^
