[Uba]
* Fixed so opened written files that were never explicitly closed is closed in exit if they have a close id

diff --git a/Engine/Source/Programs/UnrealBuildAccelerator/Detours/Private/PosixOS/UbaDetoursFunctionsPosix.cpp b/Engine/Source/Programs/UnrealBuildAccelerator/Detours/Private/PosixOS/UbaDetoursFunctionsPosix.cpp
--- a/Engine/Source/Programs/UnrealBuildAccelerator/Detours/Private/PosixOS/UbaDetoursFunctionsPosix.cpp
+++ b/Engine/Source/Programs/UnrealBuildAccelerator/Detours/Private/PosixOS/UbaDetoursFunctionsPosix.cpp
@@ -62,6 +62,7 @@ using namespace uba;
 	DETOURED_FUNCTION(telldir) \
 	DETOURED_FUNCTION(closedir) \
 	DETOURED_FUNCTION(stat) \
+	DETOURED_FUNCTION(truncate) \
 	DETOURED_FUNCTION(glob) \
 	DETOURED_FUNCTION(chmod) \
 	DETOURED_FUNCTION(rename) \
@@ -1020,6 +1021,13 @@ UBA_EXPORT int UBA_WRAPPER(stat)(const char* file, struct stat* attr)
 	return Shared_stat("stat", file, attr, [](const char* file, struct stat* attr) { return TRUE_WRAPPER(stat)(file, attr); });
 }
 
+UBA_EXPORT int UBA_WRAPPER(truncate)(const char* path, off_t length)
+{
+	UBA_INIT_DETOUR(truncate, path, length);
+	UBA_ASSERT(false); // TODO: Implement this if it is ever called
+	return TRUE_WRAPPER(truncate)(path, length);
+}
+
 UBA_EXPORT int UBA_WRAPPER(access)(const char* pathname, int mode)
 {
 	UBA_INIT_DETOUR(access, pathname, mode);
@@ -1776,10 +1784,27 @@ namespace uba
 	{
 		if (!g_isInitialized)
 			return;
+		
 		g_isInitialized = false;
-
 		g_isDetouring = false;
 
+		{
+			ScopedWriteLock lock(g_fileHandlesLock);
+			for (auto& kv : g_fileHandles)
+			{
+				DetouredHandle& h = kv.second;
+				FileObject* fo = h.fileObject;
+				if (!fo->closeId)
+					continue;
+				TRUE_WRAPPER(close)(kv.first);
+				u64 mappingHandle = 0;
+				u64 mappingWritten = 0;
+				FileInfo& fi = *fo->fileInfo;
+				const tchar* path = fi.name;
+				Rpc_UpdateCloseHandle(path, fo->closeId, fo->deleteOnClose, fo->newName.c_str(), mappingHandle, mappingWritten, true);
+			}
+		}
+		
 		BinaryWriter writer;
 		writer.WriteByte(MessageType_Exit);
 		writer.WriteU32(0); // Exit code
