[UBA]
* Fixed so visualizer is showing correct ping numbers when showing linux builds

diff --git a/Engine/Source/Programs/UnrealBuildAccelerator/Core/Private/UbaTimer.cpp b/Engine/Source/Programs/UnrealBuildAccelerator/Core/Private/UbaTimer.cpp
--- a/Engine/Source/Programs/UnrealBuildAccelerator/Core/Private/UbaTimer.cpp
+++ b/Engine/Source/Programs/UnrealBuildAccelerator/Core/Private/UbaTimer.cpp
@@ -50,9 +50,10 @@ namespace uba
 		#endif
 	}
 
-	TimeToText::TimeToText(u64 time, bool allowMinutes)
+	TimeToText::TimeToText(u64 time, bool allowMinutes) : TimeToText(time, allowMinutes, GetFrequency()) {}
+	TimeToText::TimeToText(u64 time, bool allowMinutes, u64 frequency)
 	{
-		u64 ms = TimeToMs(time);
+		u64 ms = TimeToMs(time, frequency);
 		if (ms == 0 && time != 0)
 			TStrcpy_s(str, 32, TC("<1ms"));
 		else if (ms < 1000)
diff --git a/Engine/Source/Programs/UnrealBuildAccelerator/Core/Public/UbaTimer.h b/Engine/Source/Programs/UnrealBuildAccelerator/Core/Public/UbaTimer.h
--- a/Engine/Source/Programs/UnrealBuildAccelerator/Core/Public/UbaTimer.h
+++ b/Engine/Source/Programs/UnrealBuildAccelerator/Core/Public/UbaTimer.h
@@ -63,6 +63,7 @@ namespace uba
 	struct TimeToText
 	{
 		TimeToText(u64 time, bool allowMinutes = false);
+		TimeToText(u64 time, bool allowMinutes, u64 frequency);
 		operator const tchar*() const { return str; };
 		tchar str[32];
 	};
diff --git a/Engine/Source/Programs/UnrealBuildAccelerator/Test/Private/UbaTestAll.h b/Engine/Source/Programs/UnrealBuildAccelerator/Test/Private/UbaTestAll.h
--- a/Engine/Source/Programs/UnrealBuildAccelerator/Test/Private/UbaTestAll.h
+++ b/Engine/Source/Programs/UnrealBuildAccelerator/Test/Private/UbaTestAll.h
@@ -53,8 +53,8 @@ namespace uba
 
 		logger.Info(TC("Running tests (Test rootdir: %s)"), testRootDir.data);
 
-		//UBA_TEST(TestRemoteDetouredTestApp)
-		UBA_TESTS
+		UBA_TEST(TestDetouredTestApp)
+		//UBA_TESTS
 
 		logger.Info(TC("Tests finished successfully!"));
 		Sleep(3000);
diff --git a/Engine/Source/Programs/UnrealBuildAccelerator/Visualizer/Private/UbaVisualizer.cpp b/Engine/Source/Programs/UnrealBuildAccelerator/Visualizer/Private/UbaVisualizer.cpp
--- a/Engine/Source/Programs/UnrealBuildAccelerator/Visualizer/Private/UbaVisualizer.cpp
+++ b/Engine/Source/Programs/UnrealBuildAccelerator/Visualizer/Private/UbaVisualizer.cpp
@@ -505,7 +505,7 @@ namespace uba
 					if (memAvail)
 						text.Appendf(L" Mem: %ls/%ls", BytesToText(session.memTotal - memAvail).str, BytesToText(session.memTotal).str);
 					if (ping)
-						text.Appendf(L" Ping: %ls", TimeToText(ping).str);
+						text.Appendf(L" Ping: %ls", TimeToText(ping, false, m_traceView.frequency).str);
 				}
 				else
 					text.Append(L" - Disconnected");
@@ -1175,7 +1175,7 @@ namespace uba
 			logger.SetColor(m_sendColor).Info(L"  Recv: %ls/s", BytesToText(m_stats.recvBytesPerSecond).str);
 			logger.SetColor(m_recvColor).Info(L"  Send: %ls/s", BytesToText(m_stats.sendBytesPerSecond).str);
 			if (m_stats.ping)
-				logger.Info(L"  Ping: %ls", TimeToText(m_stats.ping).str);
+				logger.Info(L"  Ping: %ls", TimeToText(m_stats.ping, false, m_traceView.frequency).str);
 		}
 		else if (m_buttonSelected != ~0u)
 		{
