[UBA]
* Fixed so redirected stdout/err is not writing to screen

diff --git a/Engine/Source/Programs/UnrealBuildAccelerator/Detours/Private/PosixOS/UbaDetoursFunctionsPosix.cpp b/Engine/Source/Programs/UnrealBuildAccelerator/Detours/Private/PosixOS/UbaDetoursFunctionsPosix.cpp
--- a/Engine/Source/Programs/UnrealBuildAccelerator/Detours/Private/PosixOS/UbaDetoursFunctionsPosix.cpp
+++ b/Engine/Source/Programs/UnrealBuildAccelerator/Detours/Private/PosixOS/UbaDetoursFunctionsPosix.cpp
@@ -821,7 +821,7 @@ UBA_EXPORT int UBA_WRAPPER(open)(const char* file, int flags, ...)
 UBA_EXPORT ssize_t UBA_WRAPPER(write)(int fd, const void* buf, size_t count)
 {
 	UBA_INIT_DETOUR(write, fd, buf, count);
-	if (fd == 1 || fd == 2) // stdout and stderr
+	if ((fd == 1 || fd == 2) && isatty(fd)) // stdout and stderr
 	{
 		Shared_WriteConsole((const char*)buf, count);
 		return count;
