[Uba]
* Fixed so failing CopyFileExW is handled (in a bit hacky way) in detoured code

diff --git a/Engine/Source/Programs/UnrealBuildAccelerator/Detours/Private/Windows/UbaDetoursFunctionsKernelBase.inl b/Engine/Source/Programs/UnrealBuildAccelerator/Detours/Private/Windows/UbaDetoursFunctionsKernelBase.inl
--- a/Engine/Source/Programs/UnrealBuildAccelerator/Detours/Private/Windows/UbaDetoursFunctionsKernelBase.inl
+++ b/Engine/Source/Programs/UnrealBuildAccelerator/Detours/Private/Windows/UbaDetoursFunctionsKernelBase.inl
@@ -1162,6 +1162,8 @@ BOOL Detoured_CopyFileExW(LPCWSTR lpExistingFileName, LPCWSTR lpNewFileName, LPP
 		return lastError == ERROR_SUCCESS;
 	}
 
+	// TODO: This copy should probably be moved to session process instead.. to handle failing to copy better
+
 	bool res;
 	{
 		SuppressCreateFileDetourScope cfs;
@@ -1170,8 +1172,11 @@ BOOL Detoured_CopyFileExW(LPCWSTR lpExistingFileName, LPCWSTR lpNewFileName, LPP
 	DEBUG_LOG_TRUE(L"CopyFileExW", L"%ls to %ls  (%ls to %ls) -> %ls", lpExistingFileName, lpNewFileName, newFromName.data, newToName.data, ToString(res));
 
 	// We need to report the new file that has been added (and we must do it _after_ it has been copied
-	if (closeId)
-		Rpc_UpdateCloseHandle(newToName.data, closeId, false, L"", 0, 0, true);
+	if (!closeId)
+		return res;
+
+	bool deleteOnClose = res == false; // If failing to copy we set deleteOnClose
+	Rpc_UpdateCloseHandle(newToName.data, closeId, deleteOnClose, L"", 0, 0, true);
 
 	return res;
 }
