[UBA]
* Fixed bug with time stamps on linux.. was using seconds accuracy
* Fixed a few bugs on linux paths related to not updating directory caches properly when files changed
* Fixed so asserts are logging properly on linux
* Added debug logger for session side to be able to track all tracking of file caches (disabled by default)

diff --git a/Engine/Source/Programs/UnrealBuildAccelerator/Common/Private/UbaFile.cpp b/Engine/Source/Programs/UnrealBuildAccelerator/Common/Private/UbaFile.cpp
--- a/Engine/Source/Programs/UnrealBuildAccelerator/Common/Private/UbaFile.cpp
+++ b/Engine/Source/Programs/UnrealBuildAccelerator/Common/Private/UbaFile.cpp
@@ -61,7 +61,7 @@ namespace uba
 			if (!QueryFullProcessImageNameW(hProcess, 0, sz, &cch))
 				continue;
 			if (cch <= MaxPath)
-				out.Append(TC(" - ")).Append(sz);
+				out.Appendf(TC(" - %s"), sz); // Using Appendf to have capacity check
 		}
 	}
 
@@ -133,7 +133,7 @@ namespace uba
 				if (IsDirectory(attr))
 					return false;
 		}
-		StringBuffer<256> additionalInfo;
+		StringBuffer<4096> additionalInfo;
 		#if PLATFORM_WINDOWS
 		if (lastError == ERROR_SHARING_VIOLATION)
 			GetProcessHoldingFile(additionalInfo, fileName);
@@ -161,7 +161,7 @@ namespace uba
 			UBA_ASSERTF(false, TC("GetFileInformationByHandle (fstat) error handling not implemented"));
 			return false;
 		}
-		out.lastWriteTime = attr.st_mtime;
+		out.lastWriteTime = FromTimeSpec(attr.st_mtim);
 		out.attributes = attr.st_mode;
 		out.volumeSerialNumber = attr.st_dev;
 		out.index = attr.st_ino;
@@ -185,7 +185,7 @@ namespace uba
 		int res = stat(fileName, &attr);
 		if (res != 0)
 			return false;// logger.Error(TC("GetFileInformation: CreateFile failed for file %s (%s)"), fileName, strerror(errno));
-		out.lastWriteTime = attr.st_mtime;
+		out.lastWriteTime = FromTimeSpec(attr.st_mtim);
 		out.attributes = attr.st_mode;
 		out.volumeSerialNumber = attr.st_dev;
 		out.index = attr.st_ino;
@@ -546,7 +546,7 @@ namespace uba
 			UBA_ASSERTF(false, TC("GetFileLastWriteTime (fstat) error handling not implemented: %s"), strerror(errno));
 			return false;
 		}
-		outTime = attr.st_mtime;
+		outTime = FromTimeSpec(attr.st_mtim);
 		return true;
 #endif
 	}
@@ -707,7 +707,7 @@ namespace uba
 #else
 		timeval tv;
 		gettimeofday(&tv, NULL);
-		return u64(tv.tv_sec);// +(tv.tv_usec / 1000ull);
+		return u64(tv.tv_sec) * 10'000'000ull + u64(tv.tv_usec)*10ull;
 #endif
 	}
 
@@ -716,7 +716,7 @@ namespace uba
 #if PLATFORM_WINDOWS
 		return TimeToMs(fileTime)/1000;
 #else
-		return fileTime;
+		return fileTime / 10'000'000ull;
 #endif
 	}
 
diff --git a/Engine/Source/Programs/UnrealBuildAccelerator/Common/Private/UbaFileMapping.cpp b/Engine/Source/Programs/UnrealBuildAccelerator/Common/Private/UbaFileMapping.cpp
--- a/Engine/Source/Programs/UnrealBuildAccelerator/Common/Private/UbaFileMapping.cpp
+++ b/Engine/Source/Programs/UnrealBuildAccelerator/Common/Private/UbaFileMapping.cpp
@@ -82,7 +82,7 @@ namespace uba
 						StringBuffer<64> uidName;
 						GetMappingHandleName(uidName, uid);
 						if (shm_unlink(uidName.data) == 0)
-							logger.Info("Removed old shared memory %s", uidName.data);
+							;// logger.Info("Removed old shared memory %s", uidName.data);
 						remove(lockFile.data);
 					}
 					else
diff --git a/Engine/Source/Programs/UnrealBuildAccelerator/Common/Private/UbaLogger.cpp b/Engine/Source/Programs/UnrealBuildAccelerator/Common/Private/UbaLogger.cpp
--- a/Engine/Source/Programs/UnrealBuildAccelerator/Common/Private/UbaLogger.cpp
+++ b/Engine/Source/Programs/UnrealBuildAccelerator/Common/Private/UbaLogger.cpp
@@ -1,6 +1,8 @@
 // Copyright Epic Games, Inc. All Rights Reserved.
 
 #include "UbaLogger.h"
+#include "UbaBinaryReaderWriter.h"
+#include "UbaFileAccessor.h"
 #include "UbaPlatform.h"
 #include "UbaStringBuffer.h"
 
@@ -163,9 +165,9 @@ namespace uba
 		HANDLE m_stdout = 0;
 		u32 m_defaultAttributes = 0;
 #endif
-		u32 m_scopeCount = 0;
 	} g_consoleLogWriterImpl;
 	LogWriter& g_consoleLogWriter = g_consoleLogWriterImpl;
+	thread_local u32 t_consoleLogScopeCount = 0;
 
 	class NullLogWriter : public LogWriter
 	{
@@ -192,13 +194,13 @@ namespace uba
 
 	void ConsoleLogWriter::BeginScope()
 	{
-		if (!m_scopeCount++)
+		if (!t_consoleLogScopeCount++)
 			m_lock.EnterWrite();
 	}
 
 	void ConsoleLogWriter::EndScope()
 	{
-		if (--m_scopeCount)
+		if (--t_consoleLogScopeCount)
 			return;
 #if PLATFORM_WINDOWS
 		if (!m_stdout)
@@ -209,7 +211,7 @@ namespace uba
 
 	void ConsoleLogWriter::Log(LogEntryType type, const tchar* str, u32 strLen, const tchar* prefix, u32 prefixLen)
 	{
-		if (m_scopeCount)
+		if (t_consoleLogScopeCount)
 			return LogNoLock(type, str, strLen, prefix, prefixLen);
 		ScopedWriteLock lock(m_lock);
 		LogNoLock(type, str, strLen, prefix, prefixLen);
@@ -294,4 +296,90 @@ namespace uba
 		else
 			TSprintf_s(str, 32, TC("%.1fgb"), double(bytes) / (1000ull * 1000 * 1000));
 	}
+
+#if UBA_DEBUG_LOGGER
+	thread_local u32 t_debugLogScopeCount = 0;
+
+	class DebugLogWriter : public LogWriter
+	{
+	public:
+		virtual void BeginScope() override
+		{
+			if (!m_file)
+				return;
+			if (!t_debugLogScopeCount++)
+				m_logLock.EnterWrite();
+		}
+
+		virtual void EndScope() override
+		{
+			if (!m_file)
+				return;
+			if (--t_debugLogScopeCount)
+				return;
+			//
+
+			m_logLock.LeaveWrite();
+		}
+
+		virtual void Log(LogEntryType type, const tchar* str, u32 strLen, const tchar* prefix = nullptr, u32 prefixLen = 0) override
+		{
+			if (!m_file)
+				return;
+			if (t_debugLogScopeCount)
+				return LogNoLock(type, str, strLen, prefix, prefixLen);
+			ScopedWriteLock lock(m_logLock);
+			LogNoLock(type, str, strLen, prefix, prefixLen);
+		}
+
+		void LogNoLock(LogEntryType type, const tchar* str, u32 strLen, const tchar* prefix, u32 prefixLen)
+		{
+			#if PLATFORM_WINDOWS
+			u8 buffer[2048];
+			BinaryWriter writer(buffer);
+			writer.WriteString(str, strLen);
+			BinaryReader reader(buffer);
+			u64 charLen = reader.Read7BitEncoded();
+			m_file->Write(reader.GetPositionData(), charLen);
+			#else
+			m_file->Write(str, strLen);
+			#endif
+		}
+
+		TString m_fileName;
+		FileAccessor* m_file = nullptr;
+		ReaderWriterLock m_logLock;
+
+	} g_debugLogWriter;
+
+	bool StartDebugLogger(Logger& outerLogger, const tchar* fileName)
+	{
+		g_debugLogWriter.m_fileName = fileName;
+		auto fa = new FileAccessor(outerLogger, g_debugLogWriter.m_fileName.c_str());
+		if (!fa->CreateWrite())
+		{
+			delete fa;
+			return false;
+		}
+
+		#if PLATFORM_WINDOWS
+		unsigned char utf8BOM[] = { 0xef,0xbb,0xbf }; 
+		fa->Write(utf8BOM, sizeof(utf8BOM));
+		#endif
+
+		g_debugLogWriter.m_file = fa;
+		return true;
+	}
+
+	void StopDebugLogger()
+	{
+		if (!g_debugLogWriter.m_file)
+			return;
+		g_debugLogWriter.m_file->Close();
+		delete g_debugLogWriter.m_file;
+		g_debugLogWriter.m_file = nullptr;
+	}
+
+	LoggerWithWriter g_debugLogger(g_debugLogWriter);
+#endif
 }
diff --git a/Engine/Source/Programs/UnrealBuildAccelerator/Common/Private/UbaProcess.cpp b/Engine/Source/Programs/UnrealBuildAccelerator/Common/Private/UbaProcess.cpp
--- a/Engine/Source/Programs/UnrealBuildAccelerator/Common/Private/UbaProcess.cpp
+++ b/Engine/Source/Programs/UnrealBuildAccelerator/Common/Private/UbaProcess.cpp
@@ -1167,12 +1167,7 @@ namespace uba
 		if (*m_startInfo.logFile)
 		{
 			#if !UBA_DEBUG_LOG_ENABLED
-			static bool runOnce = true;
-			if (runOnce)
-			{
-				runOnce = false;
-				logger.Warning(TC("Build has log files disabled so no logs will be produced"));
-			}
+			static bool runOnce = [&]() { logger.Warning(TC("Build has log files disabled so no logs will be produced")); return false; }();
 			#endif
 			payload.logFile.Append(m_startInfo.logFile);
 		}
@@ -1232,12 +1227,7 @@ namespace uba
 			if (*m_startInfo.logFile)
 			{
 				#if !UBA_DEBUG_LOG_ENABLED
-				static bool runOnce = true;
-				if (runOnce)
-				{
-					runOnce = false;
-					logger.Warning(TC("Build has log files disabled so no logs will be produced"));
-				}
+				static bool runOnce = [&]() { logger.Warning(TC("Build has log files disabled so no logs will be produced")); return false; }();
 				#endif
 				logFile.Append("UBA_LOGFILE=").Append(m_startInfo.logFile);
 			}
diff --git a/Engine/Source/Programs/UnrealBuildAccelerator/Common/Private/UbaSession.cpp b/Engine/Source/Programs/UnrealBuildAccelerator/Common/Private/UbaSession.cpp
--- a/Engine/Source/Programs/UnrealBuildAccelerator/Common/Private/UbaSession.cpp
+++ b/Engine/Source/Programs/UnrealBuildAccelerator/Common/Private/UbaSession.cpp
@@ -209,6 +209,13 @@ namespace uba
 			hasher.Update(forHash.data, forHash.count);
 		}
 
+		#if UBA_DEBUG_LOGGER
+		g_debugLogger.BeginScope();
+		auto dg = MakeGuard([]() { g_debugLogger.EndScope(); });
+		g_debugLogger.Info(TC("TRACKDIR %s\n"), dirPath);
+		#endif
+
+
 		BinaryWriter memoryWriter(memoryBlock.data(), 0, memoryBlock.size());
 
 		bool res = TraverseDir(m_logger, dirPathLen ? dirPath : TC("/"),
@@ -226,6 +233,10 @@ namespace uba
 				UBA_ASSERT(e.attributes);
 				memoryWriter.WriteString(e.name, e.nameLen);
 
+				#if UBA_DEBUG_LOGGER
+				g_debugLogger.Info(TC("    %s (Size: %llu, Key: %s, Id: %llu)\n"), e.name, e.size, KeyToString(fileKey).data, e.id);
+				#endif
+
 				res.first->second = u32(memoryWriter.GetPosition()); // Temporary offset that will be used further down to calculate the real offset
 				memoryWriter.WriteU64(e.lastWritten);
 				memoryWriter.WriteU32(e.attributes);
@@ -728,8 +739,16 @@ namespace uba
 		{
 			BinaryReader reader(dirTable.m_memory + insres.first->second);
 			u64 oldLastWriteTime = reader.ReadU64();
+
 			if (lastWriteTime == oldLastWriteTime)
+			{
+				#if !PLATFORM_WINDOWS
+				reader.Skip(sizeof(u32) * 2);
+				u64 oldFileIndex = reader.ReadU64();
+				UBA_ASSERT(oldFileIndex == fileIndex); // Checking so it is really the same file
+				#endif
 				return true;
+			}
 		}
 
 		FileEntryAdded(fileNameKey, lastWriteTime, fileSize);
@@ -751,6 +770,12 @@ namespace uba
 			writer.WriteU64(fileSize);
 			written = writer.GetPosition();
 		}
+
+		#if UBA_DEBUG_LOGGER
+		g_debugLogger.Info(TC("TRACKADD    %s (Size: %llu, Key: %s, Id: %llu)\n"), fileName, fileSize, KeyToString(fileNameKey).data, fileIndex);
+		#endif
+
+
 		ScopedWriteLock memoryLock(dirTable.m_memoryLock);
 		u8* startPos = dirTable.m_memory + dirTable.m_memorySize;
 		BinaryWriter writer(startPos);
@@ -822,6 +847,11 @@ namespace uba
 			writer.WriteU64(0);
 			written = writer.GetPosition();
 		}
+
+		#if UBA_DEBUG_LOGGER
+		g_debugLogger.Info(TC("TRACKDEL    %s (Key: %s)\n"), fileName, KeyToString(fileNameKey).data);
+		#endif
+
 		ScopedWriteLock memoryLock(dirTable.m_memoryLock);
 		u8* startPos = dirTable.m_memory + dirTable.m_memorySize;
 		BinaryWriter writer(startPos);
@@ -997,10 +1027,10 @@ namespace uba
 				});
 		}
 
-		m_sessionDir.Append(m_rootDir).Append(TC("sessions")).Append(PathSeparator).Append(m_id);// + TC("\\");
-		m_sessionBinDir.Append(m_sessionDir).Append(PathSeparator).Append(TC("bin"));
-		m_sessionOutputDir.Append(m_sessionDir).Append(PathSeparator).Append(TC("output"));
-		m_sessionLogDir.Append(m_sessionDir).Append(PathSeparator).Append(TC("log"));
+		m_sessionDir.Append(m_rootDir).Append(TC("sessions")).Append(PathSeparator).Append(m_id).Append(PathSeparator);
+		m_sessionBinDir.Append(m_sessionDir).Append(TC("bin"));
+		m_sessionOutputDir.Append(m_sessionDir).Append(TC("output"));
+		m_sessionLogDir.Append(m_sessionDir).Append(TC("log"));
 
 		if (m_runningRemote)
 		{
@@ -1008,17 +1038,21 @@ namespace uba
 			m_storage.CreateDirectory(m_sessionOutputDir.data);
 		}
 
-		m_tempPath.Append(m_sessionDir).Append(PathSeparator).Append(TC("temp"));
+		m_tempPath.Append(m_sessionDir).Append(TC("temp"));
 		m_storage.CreateDirectory(m_tempPath.data);
 		m_tempPath.EnsureEndsWithSlash();
 
-		m_sessionDir.EnsureEndsWithSlash();
 		m_sessionBinDir.EnsureEndsWithSlash();
 		m_sessionOutputDir.EnsureEndsWithSlash();
 
 		m_storage.CreateDirectory(m_sessionLogDir.data);
 		m_sessionLogDir.EnsureEndsWithSlash();
 
+
+		#if UBA_DEBUG_LOGGER
+		StartDebugLogger(m_logger, StringBuffer<512>().Append(m_sessionDir).Append(TC("SessionDebug.log")).data);
+		#endif
+
 		#if PLATFORM_WINDOWS
 		m_systemPath.count = GetEnvironmentVariableW(TC("SystemRoot"), m_systemPath.data, m_systemPath.capacity);
 		#else
@@ -1130,6 +1164,10 @@ namespace uba
 		//u32 count;
 		//DeleteAllFiles(m_logger, m_sessionDir.data, count);
 		//#endif
+
+		#if UBA_DEBUG_LOGGER
+		StopDebugLogger();
+		#endif
 	}
 
 	void Session::CancelAllProcessesAndWait(bool terminate)
@@ -1854,7 +1892,10 @@ namespace uba
 		#else
 		out.errorCode = 0;
 		if (chmod(msg.fileName.data, (mode_t)msg.fileMode) == 0)
+		{
+			RegisterCreateFileForWrite(msg.fileNameKey, msg.fileName.data, msg.fileName.count, true);
 			return true;
+		}
 		out.errorCode = errno;
 		#endif
 		return true;
diff --git a/Engine/Source/Programs/UnrealBuildAccelerator/Common/Public/UbaDirectoryIterator.h b/Engine/Source/Programs/UnrealBuildAccelerator/Common/Public/UbaDirectoryIterator.h
--- a/Engine/Source/Programs/UnrealBuildAccelerator/Common/Public/UbaDirectoryIterator.h
+++ b/Engine/Source/Programs/UnrealBuildAccelerator/Common/Public/UbaDirectoryIterator.h
@@ -230,7 +230,7 @@ namespace uba
 			DirectoryEntry entry;
 			entry.name = fileName;
 			entry.nameLen = strlen(fileName);
-			entry.lastWritten = attr.st_mtime;
+			entry.lastWritten = FromTimeSpec(attr.st_mtim);
 			entry.attributes = attr.st_mode;
 			entry.volumeSerial = attr.st_dev;
 			entry.id = attr.st_ino;
diff --git a/Engine/Source/Programs/UnrealBuildAccelerator/Common/Public/UbaLogger.h b/Engine/Source/Programs/UnrealBuildAccelerator/Common/Public/UbaLogger.h
--- a/Engine/Source/Programs/UnrealBuildAccelerator/Common/Public/UbaLogger.h
+++ b/Engine/Source/Programs/UnrealBuildAccelerator/Common/Public/UbaLogger.h
@@ -5,6 +5,9 @@
 #include "UbaBase.h"
 #include <stdarg.h>
 
+
+#define UBA_DEBUG_LOGGER 0
+
 namespace uba
 {
 	enum LogEntryType : u8
@@ -92,4 +95,11 @@ namespace uba
 		operator const tchar* () const { return str; };
 		tchar str[32];
 	};
+
+
+	#if UBA_DEBUG_LOGGER
+	bool StartDebugLogger(Logger& outerLogger, const tchar* fileName);
+	void StopDebugLogger();
+	extern LoggerWithWriter g_debugLogger;
+	#endif
 }
diff --git a/Engine/Source/Programs/UnrealBuildAccelerator/Core/Private/UbaStringBuffer.cpp b/Engine/Source/Programs/UnrealBuildAccelerator/Core/Private/UbaStringBuffer.cpp
--- a/Engine/Source/Programs/UnrealBuildAccelerator/Core/Private/UbaStringBuffer.cpp
+++ b/Engine/Source/Programs/UnrealBuildAccelerator/Core/Private/UbaStringBuffer.cpp
@@ -203,6 +203,9 @@ namespace uba
 
 		if (u32(len) >= capacity - count)
 		{
+			// 16 chars needed to fit buffer overflow string
+			if (count > capacity - 16)
+				count = capacity - 16;
 			Append(TC("BUFFEROVERFLOW!"));
 			return *this;
 		}
diff --git a/Engine/Source/Programs/UnrealBuildAccelerator/Core/Public/UbaPlatform.h b/Engine/Source/Programs/UnrealBuildAccelerator/Core/Public/UbaPlatform.h
--- a/Engine/Source/Programs/UnrealBuildAccelerator/Core/Public/UbaPlatform.h
+++ b/Engine/Source/Programs/UnrealBuildAccelerator/Core/Public/UbaPlatform.h
@@ -129,5 +129,7 @@ namespace uba
 	#define TStrdup strdup
 	#define localtime_s(a,b) localtime_r(b,a)
 	void GetMappingHandleName(StringBufferBase& out, u64 uid);
+	inline u64 FromTimeSpec(const timespec& ts) { return u64(ts.tv_sec) * 10'000'000ull + u64(ts.tv_nsec/100); }
+	inline timespec ToTimeSpec(u64 time) { timespec ts; ts.tv_sec = time / 10'000'000ull; ts.tv_nsec = (time - (u64(ts.tv_sec) * 10'000'000ull)) * 100; return ts; }
 #endif // PLATFORM_WINDOWS
 }
diff --git a/Engine/Source/Programs/UnrealBuildAccelerator/Detours/Private/PosixOS/UbaDetoursFunctionsPosix.cpp b/Engine/Source/Programs/UnrealBuildAccelerator/Detours/Private/PosixOS/UbaDetoursFunctionsPosix.cpp
--- a/Engine/Source/Programs/UnrealBuildAccelerator/Detours/Private/PosixOS/UbaDetoursFunctionsPosix.cpp
+++ b/Engine/Source/Programs/UnrealBuildAccelerator/Detours/Private/PosixOS/UbaDetoursFunctionsPosix.cpp
@@ -43,6 +43,7 @@ using namespace uba;
 	DETOURED_FUNCTION(readlinkat) \
 	DETOURED_FUNCTION(open) \
 	DETOURED_FUNCTION(write) \
+	DETOURED_FUNCTION(dup) \
 	DETOURED_FUNCTION(dup2) \
 	DETOURED_FUNCTION(close) \
 	DETOURED_FUNCTION(fopen) \
@@ -51,6 +52,7 @@ using namespace uba;
 	DETOURED_FUNCTION(fchmodat) \
 	DETOURED_FUNCTION(fwrite) \
 	DETOURED_FUNCTION(fstat) \
+	DETOURED_FUNCTION(futimens) \
 	DETOURED_FUNCTION(fclose) \
 	DETOURED_FUNCTION(opendir) \
 	DETOURED_FUNCTION(fdopendir) \
@@ -67,6 +69,7 @@ using namespace uba;
 	DETOURED_FUNCTION(chmod) \
 	DETOURED_FUNCTION(rename) \
 	DETOURED_FUNCTION(renameat) \
+	DETOURED_FUNCTION(utimensat) \
 	DETOURED_FUNCTION(remove) \
 	DETOURED_FUNCTION(symlink) \
 	DETOURED_FUNCTION(access) \
@@ -514,21 +517,26 @@ void Shared_close(int fd, const TrueClose& trueClose)
 
 	DetouredHandle& h = findIt->second;
 	FileObject* fo = h.fileObject;
-	UBA_ASSERT(fo->refCount == 1);
-
+	UBA_ASSERT(fo->refCount >= 1);
 	g_fileHandles.erase(findIt);
 	lock.Leave();
 
 	trueClose();
 
+	if (--fo->refCount)
+		return;
+
 	u64 mappingHandle = 0;
 	u64 mappingWritten = 0;
 	FileInfo& fi = *fo->fileInfo;
 	const tchar* path = fi.name;
 
 	if (fo->closeId)
+	{
+		//for (auto& kv : g_fileHandles)
+		//	UBA_ASSERT(&fi != kv.second.fileObject->fileInfo);
 		Rpc_UpdateCloseHandle(path, fo->closeId, fo->deleteOnClose, fo->newName.c_str(), mappingHandle, mappingWritten, true);
-
+	}
 	delete fo;
 }
 
@@ -576,16 +584,28 @@ int Shared_fstat(const char* funcName, int fd, struct stat* attr, const True_fst
 	errno = fileAttr.lastError;
 	memcpy(attr, &fileAttr.data, sizeof(struct stat));
 
-#if 0
-	struct stat attr2;
-	int res2 = TRUE_WRAPPER(__fxstat)(ver, fd, &attr2);
-	UBA_ASSERT(res == res2);
-	UBA_ASSERT(attr->st_mtime == attr2.st_mtime);
-	//UBA_ASSERT(attr->st_mode == attr2.st_mode);
-	//UBA_ASSERT(attr->st_dev == attr2.st_dev)
-	UBA_ASSERT(attr->st_ino == attr2.st_ino);
-	UBA_ASSERT(attr->st_size == attr2.st_size);
-#endif
+	#if UBA_DEBUG_VALIDATE
+	if (!g_runningRemote)
+	{
+		struct stat attr2;
+		int res2 = trueFstat(fd, &attr2);
+		UBA_ASSERTF(res == res2, "fstat: return value differs for %s (%i vs %i)", fi.originalName, res, res2);
+		if (res != -1)
+		{
+			bool isDir = S_ISDIR(attr->st_mode);
+			UBA_ASSERT(isDir == S_ISDIR(attr2.st_mode));
+			//UBA_ASSERT(attr->st_mode == attr2.st_mode);
+			//UBA_ASSERT(attr->st_dev == attr2.st_dev)
+			UBA_ASSERTF(attr->st_ino == attr2.st_ino, "fstat: st_ino mismatch for %s (%llu vs %llu)", fi.originalName, attr->st_ino, attr2.st_ino);
+			UBA_ASSERT(isDir || attr->st_size == attr2.st_size);
+			UBA_ASSERTF(isDir || FromTimeSpec(attr->st_mtim) == FromTimeSpec(attr2.st_mtim), "fstat: st_mtim mismatch for %s (%llu vs %llu)", fi.originalName, FromTimeSpec(attr->st_mtim), FromTimeSpec(attr2.st_mtim));
+		}
+		else
+		{
+			UBA_ASSERT(fileAttr.lastError == errno);
+		}
+	}
+	#endif
 
 	return res;
 }
@@ -624,27 +644,25 @@ int Shared_stat(const char* funcName, const char* file, struct stat* attr, const
 	errno = fileAttr.lastError;
 	memcpy(attr, &fileAttr.data, sizeof(fileAttr.data));
 
-	#if 0
+	#if UBA_DEBUG_VALIDATE
 	if (!g_runningRemote)
 	{
 		struct stat attr2;
-		int res2 = trueStat(realName, &attr2);
-		if (res != res2)
-			printf("%s: Exists %i vs %i  (%s)\n", realName, res, res2, KeyToString(ToStringKey(fixedFile)).data);
-		if (res == 0)
+		int res2 = trueStat(file, &attr2);
+		UBA_ASSERTF(res == res2, "stat: return value differs for %s (%i vs %i)", file, res, res2);
+		if (res != -1)
 		{
-			//UBA_ASSERT(attr->st_mtime == attr2.st_mtime);
+			bool isDir = S_ISDIR(attr->st_mode);
+			UBA_ASSERT(isDir == S_ISDIR(attr2.st_mode));
 			//UBA_ASSERT(attr->st_mode == attr2.st_mode);
-			//UBA_ASSERTF(attr->st_dev == attr2.st_dev, "%s: %u vs %u", realName, attr->st_dev, attr2.st_dev);
-			//UBA_ASSERTF(attr->st_ino == attr2.st_ino, "%s: %llu vs %llu", realName, attr->st_ino, attr2.st_ino);
-			//UBA_ASSERTF(attr->st_size == attr2.st_size, "%s: %llu vs %llu", realName, attr->st_size, attr2.st_size);
-
-			if (attr->st_dev != attr2.st_dev)
-				printf("%s: Dev %lu vs %lu\n", realName, attr->st_dev, attr2.st_dev);
-			if (attr->st_ino != attr2.st_ino)
-				printf("%s: Ino %lu vs %lu\n", realName, attr->st_ino, attr2.st_ino);
-			if (attr->st_size != attr2.st_size)
-				printf("%s: Size %lu vs %lu\n", realName, attr->st_size, attr2.st_size);
+			//UBA_ASSERT(attr->st_dev == attr2.st_dev)
+			UBA_ASSERTF(attr->st_ino == attr2.st_ino, "stat: st_ino mismatch for %s (%llu vs %llu)", file, attr->st_ino, attr2.st_ino);
+			UBA_ASSERT(isDir || attr->st_size == attr2.st_size);
+			UBA_ASSERTF(isDir || FromTimeSpec(attr->st_mtim) == FromTimeSpec(attr2.st_mtim), "stat: st_mtim mismatch for %s (%llu vs %llu)", file, FromTimeSpec(attr->st_mtim), FromTimeSpec(attr2.st_mtim));
+		}
+		else
+		{
+			UBA_ASSERT(fileAttr.lastError == errno);
 		}
 	}
 	#endif
@@ -821,7 +839,7 @@ UBA_EXPORT int UBA_WRAPPER(open)(const char* file, int flags, ...)
 UBA_EXPORT ssize_t UBA_WRAPPER(write)(int fd, const void* buf, size_t count)
 {
 	UBA_INIT_DETOUR(write, fd, buf, count);
-	if ((fd == 1 || fd == 2) && isatty(fd)) // stdout and stderr
+	if (isatty(fd)) // stdout and stderr
 	{
 		Shared_WriteConsole((const char*)buf, count);
 		return count;
@@ -836,10 +854,32 @@ UBA_EXPORT ssize_t UBA_WRAPPER(write)(int fd, const void* buf, size_t count)
 //	return -1;
 //}
 
+UBA_EXPORT int UBA_WRAPPER(dup)(int oldfd)
+{
+	UBA_INIT_DETOUR(dup, oldfd);
+	auto res = TRUE_WRAPPER(dup)(oldfd);
+	DEBUG_LOG_TRUE("dup", "(%i) -> %i", oldfd, res);
+	return res;
+}
+
 UBA_EXPORT int UBA_WRAPPER(dup2)(int oldfd, int newfd)
 {
 	UBA_INIT_DETOUR(dup2, oldfd, newfd);
 	auto res = TRUE_WRAPPER(dup2)(oldfd, newfd);
+
+	if (res != -1)
+	{
+		ScopedWriteLock lock(g_fileHandlesLock);
+		auto findIt = g_fileHandles.find(oldfd);
+		if (findIt != g_fileHandles.end())
+		{
+			DetouredHandle& h = findIt->second;
+			FileObject* fo = h.fileObject;
+			++fo->refCount;
+			g_fileHandles[newfd].fileObject = fo;
+		}
+	}
+
 	DEBUG_LOG_TRUE("dup2", "(%i, %i) -> %i", oldfd, newfd, res);
 	return res;
 }
@@ -906,6 +946,13 @@ UBA_EXPORT int UBA_WRAPPER(fstat)(int fd, struct stat* buf)
 	return Shared_fstat("fstat", fd, buf, [](int fd, struct stat* buf) { return TRUE_WRAPPER(fstat)(fd, buf); });
 }
 
+UBA_EXPORT int UBA_WRAPPER(futimens)(int fd, const struct timespec* times)
+{
+	UBA_INIT_DETOUR(futimens, fd, times);
+	DEBUG_LOG_TRUE("futimens", "(%i)", fd);
+	return TRUE_WRAPPER(futimens)(fd, times);
+}
+
 UBA_EXPORT int UBA_WRAPPER(fclose)(FILE* stream)
 {
 	UBA_INIT_DETOUR(fclose, stream);
@@ -1229,6 +1276,13 @@ UBA_EXPORT int UBA_WRAPPER(renameat)(int olddirfd, const char* oldpath, int newd
 	return TRUE_WRAPPER(renameat)(olddirfd, oldpath, newdirfd, newpath);
 }
 
+UBA_EXPORT int UBA_WRAPPER(utimensat)(int dirfd, const char* pathname, const struct timespec* times, int flags)
+{
+	UBA_INIT_DETOUR(utimensat, dirfd, pathname, times, flags);
+	DEBUG_LOG_TRUE("utimensat", "(%s)", pathname);
+	return TRUE_WRAPPER(utimensat)(dirfd, pathname, times, flags);
+}
+
 UBA_EXPORT int UBA_WRAPPER(symlink)(const char* path1, const char* path2)
 {
 	UBA_INIT_DETOUR(symlink, path1, path2);
@@ -1372,7 +1426,7 @@ UBA_EXPORT int UBA_WRAPPER(posix_spawn)(pid_t* pid, const char* path, const posi
 	}
 
 	#if UBA_DEBUG_LOG_ENABLED
-	DEBUG_LOG_TRUE("posix_spawn", "%s (%s)", path, ldpreload.data);
+	DEBUG_LOG_TRUE("posix_spawn", "%s", path);
 	for (u32 i = 0; argv[i]; ++i)
 		DEBUG_LOG("            %s", argv[i]);
 	#endif
@@ -1703,7 +1757,7 @@ namespace uba
 		}
 		#endif
 
-		g_exeDir.count = GetProcessExecutablePath(g_exeDir.data, g_exeDir.capacity); //TRUE_WRAPPER(readlink)("/proc/self/exe", g_exeDir.data, g_exeDir.capacity);
+		g_exeDir.count = GetProcessExecutablePath(g_exeDir.data, g_exeDir.capacity);
 		if (g_exeDir.count == -1)
 		{
 			UBA_ASSERT(false);
@@ -1748,7 +1802,9 @@ namespace uba
 
 		UBA_ASSERT(g_virtualApplicationDir.capacity > 0);
 
-		g_virtualApplicationDir.Append(g_virtualApplication.data, strrchr(g_virtualApplication.data, '/') - g_virtualApplication.data + 1);
+		const char* lastSlash = strrchr(g_virtualApplication.data, '/');
+		UBA_ASSERTF(lastSlash, "Need fullpath for application (%s)", g_virtualApplication.data);
+		g_virtualApplicationDir.Append(g_virtualApplication.data, lastSlash - g_virtualApplication.data + 1);
 
 		setenv("PWD", g_virtualWorkingDir.data, 1);
 		g_virtualWorkingDir.EnsureEndsWithSlash();
@@ -1792,6 +1848,7 @@ namespace uba
 			ScopedWriteLock lock(g_fileHandlesLock);
 			for (auto& kv : g_fileHandles)
 			{
+				TRUE_WRAPPER(close)(kv.first);
 				DetouredHandle& h = kv.second;
 				FileObject* fo = h.fileObject;
 				if (!fo->closeId)
@@ -1835,4 +1892,22 @@ namespace uba
 		errno = t;
 	}
 	#endif
+
+	ANALYSIS_NORETURN void UbaAssert(const tchar* text, const char* file, u32 line, const char* expr, u32 terminateCode)
+	{
+		SuppressDetourScope s;
+		StringBuffer<4096> b;
+		WriteAssertInfo(b, text, file, line, expr, 1);
+		Rpc_WriteLog(b.data, b.count, true);
+
+		BinaryWriter writer;
+		writer.WriteByte(MessageType_Exit);
+		writer.WriteU32(terminateCode); // Exit code
+		writer.WriteString(""); // Log name
+		g_stats.Write(writer);
+		writer.Flush(false);
+
+		CloseCom();
+		TRUE_WRAPPER(_exit)(int(terminateCode));
+	}
 }
diff --git a/Engine/Source/Programs/UnrealBuildAccelerator/Detours/Private/PosixOS/UbaDetoursMainPosix.cpp b/Engine/Source/Programs/UnrealBuildAccelerator/Detours/Private/PosixOS/UbaDetoursMainPosix.cpp
--- a/Engine/Source/Programs/UnrealBuildAccelerator/Detours/Private/PosixOS/UbaDetoursMainPosix.cpp
+++ b/Engine/Source/Programs/UnrealBuildAccelerator/Detours/Private/PosixOS/UbaDetoursMainPosix.cpp
@@ -196,16 +196,4 @@ namespace uba
 		m_pos = m_begin;
 		m_end = m_begin + CommunicationMemSize/2;
 	}
-
-	ANALYSIS_NORETURN void UbaAssert(const tchar* text, const char* file, u32 line, const char* expr, u32 terminateCode)
-	{
-		SuppressDetourScope s;
-		StringBuffer<4096> b;
-		WriteAssertInfo(b, text, file, line, expr, 1);
-		Rpc_WriteLog(b.data, b.count, true);
-		fputs(b.data, stdout);
-		fputs("\n", stdout);
-		fflush(stdout);
-		_exit(int(terminateCode));
-	}
 }
diff --git a/Engine/Source/Programs/UnrealBuildAccelerator/Detours/Private/UbaDetoursFileMappingTable.cpp b/Engine/Source/Programs/UnrealBuildAccelerator/Detours/Private/UbaDetoursFileMappingTable.cpp
--- a/Engine/Source/Programs/UnrealBuildAccelerator/Detours/Private/UbaDetoursFileMappingTable.cpp
+++ b/Engine/Source/Programs/UnrealBuildAccelerator/Detours/Private/UbaDetoursFileMappingTable.cpp
@@ -171,6 +171,9 @@ namespace uba
 		entryNameForKey.Append(entryName, entryNameLen);
 		if (CaseInsensitiveFs)
 			entryNameForKey.MakeLower();
+		else if (entryNameForKey.count == 1 && entryNameForKey[0] == '/')
+			checkIfDir = true;
+
 		CHECK_PATH(entryNameForKey.data);
 		DirectoryTable::Exists exists = g_directoryTable.EntryExists(entryNameKey, entryNameForKey.data, entryNameLen, checkIfDir, &dirTableOffset);
 		if (exists != DirectoryTable::Exists_Maybe)
@@ -184,12 +187,6 @@ namespace uba
 		#endif
 
 		u64 dirNameLen = lastPathSeparator - entryName;
-		// There's a potential issue with the way paths with a single leading slash are handled
-		// e.g.: /usr ends up with a lastPathSeparator == 0 which throws off the logic
-		// However if we do the right thing tests start failing, so there's some ghost in the machine
-		// that needs to be worked out still. 
-		//dirNameLen = dirNameLen == 0 ? entryNameLen : dirNameLen;
-		//printf("RPC_GetEntryOffset: %s len: %llu entry: %s lastPath:%s\n", entryNameLower.data, dirNameLen, entryName, lastPathSeparator);
 		DirHash hash(entryNameForKey.data, dirNameLen);
 
 		if (Rpc_UpdateDirectory(hash.key, entryName, dirNameLen) == ~u32(0))
diff --git a/Engine/Source/Programs/UnrealBuildAccelerator/Detours/Private/UbaDetoursShared.cpp b/Engine/Source/Programs/UnrealBuildAccelerator/Detours/Private/UbaDetoursShared.cpp
--- a/Engine/Source/Programs/UnrealBuildAccelerator/Detours/Private/UbaDetoursShared.cpp
+++ b/Engine/Source/Programs/UnrealBuildAccelerator/Detours/Private/UbaDetoursShared.cpp
@@ -162,9 +162,7 @@ namespace uba
 
 	void Rpc_WriteLog(const tchar* text, u64 textCharLength, bool printInSession)
 	{
-		#if PLATFORM_WINDOWS
-		DEBUG_LOG(TC("LOG  %.*ls"), u32(textCharLength), text); // TODO: Investigate, deadlocks on non-windows
-		#endif	
+		DEBUG_LOG(TC("LOG  %.*s"), u32(textCharLength), text); // TODO: Investigate, deadlocks on non-windows
 		TimerScope ts(g_stats.log);
 		ScopedWriteLock pcs(g_communicationLock);
 		BinaryWriter writer;
@@ -352,7 +350,7 @@ namespace uba
 						(u64&)outAttr.data.ftLastAccessTime = info.lastWrite;
 						(u64&)outAttr.data.ftLastWriteTime = info.lastWrite;
 #else
-						outAttr.data.st_mtime = info.lastWrite;
+						outAttr.data.st_mtim = ToTimeSpec(info.lastWrite);
 						outAttr.data.st_mode = (mode_t)info.attributes;
 						outAttr.data.st_dev = info.volumeSerial;
 						outAttr.data.st_ino = info.fileIndex;
diff --git a/Engine/Source/Programs/UnrealBuildAccelerator/Test/Private/UbaTestAll.h b/Engine/Source/Programs/UnrealBuildAccelerator/Test/Private/UbaTestAll.h
--- a/Engine/Source/Programs/UnrealBuildAccelerator/Test/Private/UbaTestAll.h
+++ b/Engine/Source/Programs/UnrealBuildAccelerator/Test/Private/UbaTestAll.h
@@ -21,6 +21,7 @@ namespace uba
 		UBA_TEST(TestStorage) \
 		UBA_TEST(TestDetouredTestApp) \
 		UBA_TEST(TestDetouredClang) \
+		UBA_TEST(TestDetouredTouch) \
 		UBA_TEST(TestRemoteDetouredTestApp) \
 		UBA_TEST(TestCustomService) \
 
@@ -53,8 +54,8 @@ namespace uba
 
 		logger.Info(TC("Running tests (Test rootdir: %s)"), testRootDir.data);
 
-		UBA_TEST(TestDetouredTestApp)
-		//UBA_TESTS
+		//UBA_TEST(TestDetouredTestApp)
+		UBA_TESTS
 
 		logger.Info(TC("Tests finished successfully!"));
 		Sleep(3000);
diff --git a/Engine/Source/Programs/UnrealBuildAccelerator/Test/Private/UbaTestBasics.h b/Engine/Source/Programs/UnrealBuildAccelerator/Test/Private/UbaTestBasics.h
--- a/Engine/Source/Programs/UnrealBuildAccelerator/Test/Private/UbaTestBasics.h
+++ b/Engine/Source/Programs/UnrealBuildAccelerator/Test/Private/UbaTestBasics.h
@@ -100,8 +100,10 @@ namespace uba
 			return logger.Error(TC("Failed to get last written time"));
 
 		u64 systemTime = GetSystemTimeAsFileTime();
-		if (GetFileTimeAsSeconds(systemTime)+1 - GetFileTimeAsSeconds(writeTime) > 3)
-			return logger.Error(TC("system time or last written time is wrong (system: %llu, write: %llu)"), systemTime, writeTime);
+		if (systemTime < writeTime)
+			return logger.Error(TC("System time is lower than last written time"));
+		if (GetFileTimeAsSeconds(systemTime) - GetFileTimeAsSeconds(writeTime) > 3)
+			return logger.Error(TC("System time or last written time is wrong (system: %llu, write: %llu, diffInSec: %llu)"), systemTime, writeTime, GetFileTimeAsSeconds(systemTime) - GetFileTimeAsSeconds(writeTime));
 
 
 		u8 byte2 = 0;
diff --git a/Engine/Source/Programs/UnrealBuildAccelerator/Test/Private/UbaTestSession.h b/Engine/Source/Programs/UnrealBuildAccelerator/Test/Private/UbaTestSession.h
--- a/Engine/Source/Programs/UnrealBuildAccelerator/Test/Private/UbaTestSession.h
+++ b/Engine/Source/Programs/UnrealBuildAccelerator/Test/Private/UbaTestSession.h
@@ -256,4 +256,47 @@ namespace uba
 			return true;
 		return RunRemote(logger, testRootDir, RunClang);
 	}
+
+	bool TestDetouredTouch(LoggerWithWriter& logger, const StringBufferBase& testRootDir)
+	{
+		if (IsWindows)
+			return true;
+		return RunLocal(logger, testRootDir, [](LoggerWithWriter& logger, SessionServer& session, const tchar* workingDir, const RunProcessFunction& runProcess)
+			{
+				StringBuffer<> file;
+				file.Append(workingDir).Append(TC("TouchFile.h"));
+				FileAccessor fr(logger, file.data);
+
+				if (!fr.CreateWrite())
+					return false;
+				if (!fr.Write("Foo", 4))
+					return false;
+				if (!fr.Close())
+					return false;
+				FileInformation oldInfo;
+				if (!GetFileInformation(oldInfo, logger, file.data))
+					return false;
+
+				Sleep(100);
+
+				ProcessStartInfo processInfo;
+				processInfo.application = TC("/usr/bin/touch");
+				processInfo.workingDir = workingDir;
+				processInfo.arguments = file.data;
+				processInfo.logFile = TC("[REDACTED_USER_HOME]/Touch.log");
+				ProcessHandle process = runProcess(processInfo);
+				if (!process.WaitForExit(10000))
+					return logger.Error(TC("UbaTestApp did not exit in 10 seconds"));
+				u32 exitCode = process.GetExitCode();
+				if (exitCode != 0)
+					return false;
+
+				FileInformation newInfo;
+				if (!GetFileInformation(newInfo, logger, file.data))
+					return false;
+				if (newInfo.lastWriteTime == oldInfo.lastWriteTime)
+					return logger.Error(TC("File time not changed after touch"));
+				return true;
+			});
+	}
 }
diff --git a/Engine/Source/Programs/UnrealBuildAccelerator/TestApp/Private/PosixOS/UbaTestAppPosix.cpp b/Engine/Source/Programs/UnrealBuildAccelerator/TestApp/Private/PosixOS/UbaTestAppPosix.cpp
--- a/Engine/Source/Programs/UnrealBuildAccelerator/TestApp/Private/PosixOS/UbaTestAppPosix.cpp
+++ b/Engine/Source/Programs/UnrealBuildAccelerator/TestApp/Private/PosixOS/UbaTestAppPosix.cpp
@@ -136,5 +136,9 @@ int main()
 	if (stat("FileWF", &attrWF) == -1)
 		return LogError("stat for FileW failed");
 
+	struct stat attrRoot;
+	if (stat("/", &attrRoot) != -1)
+		return LogError("stat for '/' failed");
+
 	return 0;
 }
