[Uba]
* Fixed macos compile errors

diff --git a/Engine/Source/Programs/UnrealBuildAccelerator/Common/Private/UbaFile.cpp b/Engine/Source/Programs/UnrealBuildAccelerator/Common/Private/UbaFile.cpp
--- a/Engine/Source/Programs/UnrealBuildAccelerator/Common/Private/UbaFile.cpp
+++ b/Engine/Source/Programs/UnrealBuildAccelerator/Common/Private/UbaFile.cpp
@@ -161,7 +161,7 @@ namespace uba
 			UBA_ASSERTF(false, TC("GetFileInformationByHandle (fstat) error handling not implemented"));
 			return false;
 		}
-		out.lastWriteTime = FromTimeSpec(attr.st_mtim);
+		out.lastWriteTime = FromTimeSpec(attr.st_mtimespec);
 		out.attributes = attr.st_mode;
 		out.volumeSerialNumber = attr.st_dev;
 		out.index = attr.st_ino;
@@ -185,7 +185,7 @@ namespace uba
 		int res = stat(fileName, &attr);
 		if (res != 0)
 			return false;// logger.Error(TC("GetFileInformation: CreateFile failed for file %s (%s)"), fileName, strerror(errno));
-		out.lastWriteTime = FromTimeSpec(attr.st_mtim);
+		out.lastWriteTime = FromTimeSpec(attr.st_mtimespec);
 		out.attributes = attr.st_mode;
 		out.volumeSerialNumber = attr.st_dev;
 		out.index = attr.st_ino;
@@ -546,7 +546,7 @@ namespace uba
 			UBA_ASSERTF(false, TC("GetFileLastWriteTime (fstat) error handling not implemented: %s"), strerror(errno));
 			return false;
 		}
-		outTime = FromTimeSpec(attr.st_mtim);
+		outTime = FromTimeSpec(attr.st_mtimespec);
 		return true;
 #endif
 	}
diff --git a/Engine/Source/Programs/UnrealBuildAccelerator/Common/Public/UbaDirectoryIterator.h b/Engine/Source/Programs/UnrealBuildAccelerator/Common/Public/UbaDirectoryIterator.h
--- a/Engine/Source/Programs/UnrealBuildAccelerator/Common/Public/UbaDirectoryIterator.h
+++ b/Engine/Source/Programs/UnrealBuildAccelerator/Common/Public/UbaDirectoryIterator.h
@@ -230,7 +230,7 @@ namespace uba
 			DirectoryEntry entry;
 			entry.name = fileName;
 			entry.nameLen = strlen(fileName);
-			entry.lastWritten = FromTimeSpec(attr.st_mtim);
+			entry.lastWritten = FromTimeSpec(attr.st_mtimespec);
 			entry.attributes = attr.st_mode;
 			entry.volumeSerial = attr.st_dev;
 			entry.id = attr.st_ino;
diff --git a/Engine/Source/Programs/UnrealBuildAccelerator/Core/Public/UbaPlatform.h b/Engine/Source/Programs/UnrealBuildAccelerator/Core/Public/UbaPlatform.h
--- a/Engine/Source/Programs/UnrealBuildAccelerator/Core/Public/UbaPlatform.h
+++ b/Engine/Source/Programs/UnrealBuildAccelerator/Core/Public/UbaPlatform.h
@@ -131,5 +131,8 @@ namespace uba
 	void GetMappingHandleName(StringBufferBase& out, u64 uid);
 	inline u64 FromTimeSpec(const timespec& ts) { return u64(ts.tv_sec) * 10'000'000ull + u64(ts.tv_nsec/100); }
 	inline timespec ToTimeSpec(u64 time) { timespec ts; ts.tv_sec = time / 10'000'000ull; ts.tv_nsec = (time - (u64(ts.tv_sec) * 10'000'000ull)) * 100; return ts; }
+	#if PLATFORM_LINUX
+	#define st_mtimespec st_mtim
+	#endif
 #endif // PLATFORM_WINDOWS
 }
diff --git a/Engine/Source/Programs/UnrealBuildAccelerator/Detours/Private/PosixOS/UbaDetoursFunctionsPosix.cpp b/Engine/Source/Programs/UnrealBuildAccelerator/Detours/Private/PosixOS/UbaDetoursFunctionsPosix.cpp
--- a/Engine/Source/Programs/UnrealBuildAccelerator/Detours/Private/PosixOS/UbaDetoursFunctionsPosix.cpp
+++ b/Engine/Source/Programs/UnrealBuildAccelerator/Detours/Private/PosixOS/UbaDetoursFunctionsPosix.cpp
@@ -598,7 +598,7 @@ int Shared_fstat(const char* funcName, int fd, struct stat* attr, const True_fst
 			//UBA_ASSERT(attr->st_dev == attr2.st_dev)
 			UBA_ASSERTF(attr->st_ino == attr2.st_ino, "fstat: st_ino mismatch for %s (%llu vs %llu)", fi.originalName, attr->st_ino, attr2.st_ino);
 			UBA_ASSERT(isDir || attr->st_size == attr2.st_size);
-			UBA_ASSERTF(isDir || FromTimeSpec(attr->st_mtim) == FromTimeSpec(attr2.st_mtim), "fstat: st_mtim mismatch for %s (%llu vs %llu)", fi.originalName, FromTimeSpec(attr->st_mtim), FromTimeSpec(attr2.st_mtim));
+			UBA_ASSERTF(isDir || FromTimeSpec(attr->st_mtimespec) == FromTimeSpec(attr2.st_mtimespec), "fstat: st_mtim mismatch for %s (%llu vs %llu)", fi.originalName, FromTimeSpec(attr->st_mtimespec), FromTimeSpec(attr2.st_mtim));
 		}
 		else
 		{
@@ -658,7 +658,7 @@ int Shared_stat(const char* funcName, const char* file, struct stat* attr, const
 			//UBA_ASSERT(attr->st_dev == attr2.st_dev)
 			UBA_ASSERTF(attr->st_ino == attr2.st_ino, "stat: st_ino mismatch for %s (%llu vs %llu)", file, attr->st_ino, attr2.st_ino);
 			UBA_ASSERT(isDir || attr->st_size == attr2.st_size);
-			UBA_ASSERTF(isDir || FromTimeSpec(attr->st_mtim) == FromTimeSpec(attr2.st_mtim), "stat: st_mtim mismatch for %s (%llu vs %llu)", file, FromTimeSpec(attr->st_mtim), FromTimeSpec(attr2.st_mtim));
+			UBA_ASSERTF(isDir || FromTimeSpec(attr->st_mtimespec) == FromTimeSpec(attr2.st_mtimespec), "stat: st_mtim mismatch for %s (%llu vs %llu)", file, FromTimeSpec(attr->st_mtimespec), FromTimeSpec(attr2.st_mtimespec));
 		}
 		else
 		{
diff --git a/Engine/Source/Programs/UnrealBuildAccelerator/Detours/Private/UbaDetoursShared.cpp b/Engine/Source/Programs/UnrealBuildAccelerator/Detours/Private/UbaDetoursShared.cpp
--- a/Engine/Source/Programs/UnrealBuildAccelerator/Detours/Private/UbaDetoursShared.cpp
+++ b/Engine/Source/Programs/UnrealBuildAccelerator/Detours/Private/UbaDetoursShared.cpp
@@ -350,7 +350,7 @@ namespace uba
 						(u64&)outAttr.data.ftLastAccessTime = info.lastWrite;
 						(u64&)outAttr.data.ftLastWriteTime = info.lastWrite;
 #else
-						outAttr.data.st_mtim = ToTimeSpec(info.lastWrite);
+						outAttr.data.st_mtimespec = ToTimeSpec(info.lastWrite);
 						outAttr.data.st_mode = (mode_t)info.attributes;
 						outAttr.data.st_dev = info.volumeSerial;
 						outAttr.data.st_ino = info.fileIndex;
