[Uba]
* Fixed so temp files stored in file mappings are cleared directly when parent process dies

diff --git a/Engine/Source/Programs/UnrealBuildAccelerator/Common/Private/UbaProcess.cpp b/Engine/Source/Programs/UnrealBuildAccelerator/Common/Private/UbaProcess.cpp
--- a/Engine/Source/Programs/UnrealBuildAccelerator/Common/Private/UbaProcess.cpp
+++ b/Engine/Source/Programs/UnrealBuildAccelerator/Common/Private/UbaProcess.cpp
@@ -102,9 +102,7 @@ namespace uba
 			delete &m_writtenFiles;
 			delete &m_writtenFilesLock;
 
-			for (auto& pair : m_tempFiles)
-				if (pair.second.mappingHandle.IsValid())
-					CloseFileMapping(pair.second.mappingHandle);
+			ClearTempFiles();
 			delete& m_tempFiles;
 			delete& m_tempFilesLock;
 		}
@@ -382,7 +380,7 @@ namespace uba
 
 			UBA_ASSERT(!m_parentProcess);
 			m_writtenFiles.clear();
-			m_tempFiles.clear();
+			ClearTempFiles();
 		}
 
 		m_processStats.wallTime = GetTime() - m_startTime;
@@ -423,6 +421,9 @@ namespace uba
 		m_session.m_processCommunicationAllocator.Free(m_comMemory);
 		m_comMemory = {};
 
+		if (!m_parentProcess)
+			ClearTempFiles();
+
 		if (m_startInfo.exitedFunc)
 		{
 			ProcessHandle h;
@@ -1454,4 +1455,12 @@ namespace uba
 		return m_nativeProcessExitCode;
 #endif
 	}
+
+	void ProcessImpl::ClearTempFiles()
+	{
+		for (auto& pair : m_tempFiles)
+			if (pair.second.mappingHandle.IsValid())
+				CloseFileMapping(pair.second.mappingHandle);
+		m_tempFiles.clear();
+	}
 }
diff --git a/Engine/Source/Programs/UnrealBuildAccelerator/Common/Private/UbaSession.cpp b/Engine/Source/Programs/UnrealBuildAccelerator/Common/Private/UbaSession.cpp
--- a/Engine/Source/Programs/UnrealBuildAccelerator/Common/Private/UbaSession.cpp
+++ b/Engine/Source/Programs/UnrealBuildAccelerator/Common/Private/UbaSession.cpp
@@ -1149,7 +1149,7 @@ namespace uba
 		m_trace.StopWrite(m_traceOutputFile.data);
 
 		CancelAllProcessesAndWait();
-		m_deadProcesses.clear();
+		FlushDeadProcesses();
 
 		//for (auto& i : m_fileMappingTableLookup)
 		//	CloseHandle(i.second.mapping);
@@ -1215,6 +1215,7 @@ namespace uba
 
 	ProcessHandle Session::RunProcess(const ProcessStartInfo& startInfo, bool async)
 	{
+		FlushDeadProcesses();
 		ValidateStartInfo(startInfo);
 		return InternalRunProcess(startInfo, async, nullptr);
 	}
@@ -1376,6 +1377,12 @@ namespace uba
 		m_processes.erase(id);
 	}
 
+	void Session::FlushDeadProcesses()
+	{
+		ScopedWriteLock lock(m_processesLock);
+		m_deadProcesses.clear();
+	}
+
 	bool Session::GetInitResponse(InitResponse& out, const InitMessage& msg)
 	{
 		out.directoryTableHandle = m_directoryTableHandle.ToU64();
diff --git a/Engine/Source/Programs/UnrealBuildAccelerator/Common/Private/UbaSessionClient.cpp b/Engine/Source/Programs/UnrealBuildAccelerator/Common/Private/UbaSessionClient.cpp
--- a/Engine/Source/Programs/UnrealBuildAccelerator/Common/Private/UbaSessionClient.cpp
+++ b/Engine/Source/Programs/UnrealBuildAccelerator/Common/Private/UbaSessionClient.cpp
@@ -1239,6 +1239,8 @@ namespace uba
 		{
 			u32 waitTimeoutMs = 3000;
 
+			FlushDeadProcesses();
+
 			GetMemoryInfo(memAvail, memTotal);
 			if (memAvail < memRequiredFree)
 			{
diff --git a/Engine/Source/Programs/UnrealBuildAccelerator/Common/Private/UbaSessionServer.cpp b/Engine/Source/Programs/UnrealBuildAccelerator/Common/Private/UbaSessionServer.cpp
--- a/Engine/Source/Programs/UnrealBuildAccelerator/Common/Private/UbaSessionServer.cpp
+++ b/Engine/Source/Programs/UnrealBuildAccelerator/Common/Private/UbaSessionServer.cpp
@@ -254,6 +254,7 @@ namespace uba
 
 	ProcessHandle SessionServer::RunProcessRemote(const ProcessStartInfo& startInfo, float weight)
 	{
+		FlushDeadProcesses();
 		ValidateStartInfo(startInfo);
 		u32 processId = ++m_processIdCounter;
 		RemoteProcess* remoteProcess = new RemoteProcess(this, startInfo, processId, weight);
diff --git a/Engine/Source/Programs/UnrealBuildAccelerator/Common/Public/UbaProcess.h b/Engine/Source/Programs/UnrealBuildAccelerator/Common/Public/UbaProcess.h
--- a/Engine/Source/Programs/UnrealBuildAccelerator/Common/Public/UbaProcess.h
+++ b/Engine/Source/Programs/UnrealBuildAccelerator/Common/Public/UbaProcess.h
@@ -71,6 +71,7 @@ namespace uba
 		const tchar* InternalGetChildLogFile(StringBufferBase& temp);
 		u32 InternalCreateProcess(bool runningRemote, void* environment, FileMappingHandle communicationHandle, u64 communicationOffset);
 		u32	InternalExitProcess(bool cancel);
+		void ClearTempFiles();
 
 		ProcessStartInfo m_startInfo;
 		Session& m_session;
diff --git a/Engine/Source/Programs/UnrealBuildAccelerator/Common/Public/UbaSession.h b/Engine/Source/Programs/UnrealBuildAccelerator/Common/Public/UbaSession.h
--- a/Engine/Source/Programs/UnrealBuildAccelerator/Common/Public/UbaSession.h
+++ b/Engine/Source/Programs/UnrealBuildAccelerator/Common/Public/UbaSession.h
@@ -115,6 +115,7 @@ namespace uba
 		ProcessHandle InternalRunProcess(const ProcessStartInfo& startInfo, bool async, ProcessImpl* parent);
 		void ProcessAdded(Process& process, u32 sessionId);
 		void ProcessExited(ProcessImpl& process, u64 executionTime);
+		void FlushDeadProcesses();
 		void PrintProcessStats(ProcessStats& stats, const tchar* logName);
 		void StopTraceThread();
 		u32 GetDirectoryTableSize();
