[UBA]
* Fixed duplicated map/unmap for memory file

diff --git a/Engine/Source/Programs/UnrealBuildAccelerator/Detours/Private/Windows/UbaDetoursFunctionsWin.cpp b/Engine/Source/Programs/UnrealBuildAccelerator/Detours/Private/Windows/UbaDetoursFunctionsWin.cpp
--- a/Engine/Source/Programs/UnrealBuildAccelerator/Detours/Private/Windows/UbaDetoursFunctionsWin.cpp
+++ b/Engine/Source/Programs/UnrealBuildAccelerator/Detours/Private/Windows/UbaDetoursFunctionsWin.cpp
@@ -143,7 +143,11 @@ struct MemoryFile
 		{
 			mappedSize = 32 * 1024 * 1024;
 			mappingHandle = True_CreateFileMappingW(INVALID_HANDLE_VALUE, NULL, PAGE_READWRITE | SEC_RESERVE, ToHigh(reserveSize), ToLow(reserveSize), NULL);
+			if (!mappingHandle)
+				FatalError(1348, L"CreateFileMappingW failed trying to reserve %llu. (Error code: %u)", reserveSize, GetLastError());
 			baseAddress = (u8*)True_MapViewOfFile(mappingHandle, FILE_MAP_READ | FILE_MAP_WRITE, 0, 0, mappedSize);
+			if (!baseAddress)
+				FatalError(1353, L"MapViewOfFile failed trying to map %llu. ReservedSize: %llu (Error code: %u)", mappedSize, reserveSize, GetLastError());
 		}
 	}
 
@@ -320,6 +324,7 @@ void MemoryFile::EnsureCommited(DetouredHandle& handle, u64 size)
 		return;
 	if (size > mappedSize)
 	{
+		bool shouldRemap = true;
 		if (size > reserveSize)
 		{
 			if (writtenSize == 0 && !isReported)
@@ -328,15 +333,20 @@ void MemoryFile::EnsureCommited(DetouredHandle& handle, u64 size)
 				Rpc_WriteLogf(L"TODO: RE-RESERVING MemoryFile. Initial reserve: %llu, New reserve: %llu. Please fix application rules", reserveSize, newReserve);
 				Unreserve();
 				Reserve(newReserve);
+				shouldRemap = false;
 			}
 			else
 				FatalError(1347, L"Reserved size of %ls is smaller than what is requested to be. ReserveSize: %llu Written: %llu Requested: %llu", HandleToName(handle), reserveSize, writtenSize, size);
 		}
-		True_UnmapViewOfFile(baseAddress);
-		mappedSize = Min(reserveSize, AlignUp(Max(size, mappedSize*4), g_pageSize));
-		baseAddress = (u8*)True_MapViewOfFile(mappingHandle, FILE_MAP_READ|FILE_MAP_WRITE, 0, 0, mappedSize);
-		if (!baseAddress)
-			FatalError(1347, L"MapViewOfFile failed trying to map %llu for %ls. ReservedSize: %llu (Error code: %u)", mappedSize, HandleToName(handle), reserveSize, GetLastError());
+
+		if (shouldRemap)
+		{
+			True_UnmapViewOfFile(baseAddress);
+			mappedSize = Min(reserveSize, AlignUp(Max(size, mappedSize * 4), g_pageSize));
+			baseAddress = (u8*)True_MapViewOfFile(mappingHandle, FILE_MAP_READ | FILE_MAP_WRITE, 0, 0, mappedSize);
+			if (!baseAddress)
+				FatalError(1347, L"MapViewOfFile failed trying to map %llu for %ls. ReservedSize: %llu (Error code: %u)", mappedSize, HandleToName(handle), reserveSize, GetLastError());
+		}
 	}
 
 	u64 toCommit = Min(reserveSize, AlignUp(size - committedSize, g_pageSize));
