[UBA]
* Fixed bug where mapped memory were released while still used

diff --git a/Engine/Source/Programs/UnrealBuildAccelerator/Common/Private/UbaProcess.cpp b/Engine/Source/Programs/UnrealBuildAccelerator/Common/Private/UbaProcess.cpp
--- a/Engine/Source/Programs/UnrealBuildAccelerator/Common/Private/UbaProcess.cpp
+++ b/Engine/Source/Programs/UnrealBuildAccelerator/Common/Private/UbaProcess.cpp
@@ -410,6 +410,8 @@ namespace uba
 			while (!((ProcessImpl*)child.m_process)->m_hasExited)
 				Sleep(100);
 
+		UBA_ASSERT(!m_parentProcess || !m_parentProcess->m_hasExited);
+
 		m_hasExited = true;
 
 		#if !PLATFORM_WINDOWS
