[UBA]
* Changed http get append to copy char-by-char instead of using sprintf

diff --git a/Engine/Source/Programs/UnrealBuildAccelerator/Common/Private/UbaNetworkBackendTcp.cpp b/Engine/Source/Programs/UnrealBuildAccelerator/Common/Private/UbaNetworkBackendTcp.cpp
--- a/Engine/Source/Programs/UnrealBuildAccelerator/Common/Private/UbaNetworkBackendTcp.cpp
+++ b/Engine/Source/Programs/UnrealBuildAccelerator/Common/Private/UbaNetworkBackendTcp.cpp
@@ -972,13 +972,7 @@ namespace uba
 		char* bodyStart = strstr(secondSpace + 1, "\r\n\r\n");
 		if (!bodyStart)
 			return logger.Error(TC("HttpRequest: no body found"));
-
-		#if PLATFORM_WINDOWS
-		outResponse.Appendf(L"%S", bodyStart + 4);
-		#else
 		outResponse.Append(bodyStart + 4);
-		#endif
-
 		return true;
 	}
 }
\ No newline at end of file
diff --git a/Engine/Source/Programs/UnrealBuildAccelerator/Core/Private/UbaStringBuffer.cpp b/Engine/Source/Programs/UnrealBuildAccelerator/Core/Private/UbaStringBuffer.cpp
--- a/Engine/Source/Programs/UnrealBuildAccelerator/Core/Private/UbaStringBuffer.cpp
+++ b/Engine/Source/Programs/UnrealBuildAccelerator/Core/Private/UbaStringBuffer.cpp
@@ -219,6 +219,18 @@ namespace uba
 		return *this;
 	}
 
+	#if PLATFORM_WINDOWS
+	StringBufferBase& StringBufferBase::Append(const char* str)
+	{
+		u32 capacityEnd = capacity - 1;
+		for (const char* i = str; *i; ++i)
+			if (count < capacityEnd)
+				data[count++] = *i;
+		data[count] = 0;
+		return *this;
+	}
+	#endif
+
 	StringBufferBase& StringBufferBase::Resize(u64 newSize)
 	{
 		UBA_ASSERT(newSize < capacity);
diff --git a/Engine/Source/Programs/UnrealBuildAccelerator/Core/Public/UbaStringBuffer.h b/Engine/Source/Programs/UnrealBuildAccelerator/Core/Public/UbaStringBuffer.h
--- a/Engine/Source/Programs/UnrealBuildAccelerator/Core/Public/UbaStringBuffer.h
+++ b/Engine/Source/Programs/UnrealBuildAccelerator/Core/Public/UbaStringBuffer.h
@@ -34,6 +34,10 @@ namespace uba
 		StringBufferBase& Resize(u64 newSize);
 		StringBufferBase& Clear();
 
+		#if PLATFORM_WINDOWS
+		StringBufferBase& Append(const char* str);
+		#endif
+
 		tchar operator[](u64 i) const { return data[i]; }
 		tchar& operator[](u64 i) { return data[i]; }
 
