UBA: Various fixes for POSIX APIs.

* Checking returns a bit better
* Handling "/" in a friendly way since realpath calls stat with "/"
* Hooks up lockFd so we don't accidently close fd1 and fd2 (aka stdout/stderr)
* MacOS doesn't have subreaper so implement another way using sysctl

diff --git a/Engine/Source/Programs/UnrealBuildAccelerator/Common/Private/UbaProcess.cpp b/Engine/Source/Programs/UnrealBuildAccelerator/Common/Private/UbaProcess.cpp
--- a/Engine/Source/Programs/UnrealBuildAccelerator/Common/Private/UbaProcess.cpp
+++ b/Engine/Source/Programs/UnrealBuildAccelerator/Common/Private/UbaProcess.cpp
@@ -15,10 +15,19 @@
 #include <wchar.h>
 #include <stdio.h>
 #include <spawn.h>
+#include <wordexp.h>
+
+// These headers are used for tracking child and beyond
+// processes and making sure they clean up properly
+// Linux uses PR_SET_CHILD_SUBREAPER
+// Mac has to roll it's own solution
 #if PLATFORM_LINUX
 #include <sys/prctl.h>
+#elif PLATFORM_MAC
+#include <sys/types.h>
+#include <sys/sysctl.h>
 #endif
-#include <wordexp.h>
+
 extern char **environ;
 #endif
 
@@ -405,10 +414,19 @@ namespace uba
 
 		m_session.ProcessExited(*this, m_processStats.wallTime);
 
+		#if PLATFORM_LINUX
 		// This should not really ever happen.. but just in case.. since children use memory from parent
 		for (auto& child : m_childProcesses)
+		{
 			while (!((ProcessImpl*)child.m_process)->m_hasExited)
+			{
 				Sleep(100);
+			}
+		}
+		#elif PLATFORM_MAC
+		int res = WaitForProcessGroup(getpgrp());
+		UBA_ASSERT(res == 0);
+		#endif
 
 		UBA_ASSERT(!m_parentProcess || !m_parentProcess->m_hasExited);
 
@@ -992,6 +1010,57 @@ namespace uba
 		return temp.data;
 	}
 
+#if PLATFORM_MAC
+	int ProcessImpl::WaitForProcessGroup(pid_t pgid)
+	{
+		int name[] = {CTL_KERN, KERN_PROC, KERN_PROC_PGRP, pgid};
+
+		for (;;) {
+			// Query the list of processes in the group by using sysctl(3).
+			// This is "hard" because we don't know how big that list is, so we
+			// have to first query the size of the output data and then account for
+			// the fact that the size might change by the time we actually issue
+			// the query.
+			struct kinfo_proc *procs = NULL;
+			size_t nprocs = 0;
+			do {
+				size_t len;
+				if (sysctl(name, 4, 0, &len, NULL, 0) == -1) {
+					printf("Something went wrong\n");
+					return -1;
+				}
+				procs = (struct kinfo_proc *)malloc(len);
+				if (sysctl(name, 4, procs, &len, NULL, 0) == -1) {
+					UBA_ASSERT(errno == ENOMEM);
+					free(procs);
+					procs = NULL;
+				} else {
+					nprocs = len / sizeof(struct kinfo_proc);
+				}
+			} while (procs == NULL);
+			UBA_ASSERT(nprocs >= 1);  // Must have found the group leader at least.
+
+			if (nprocs == 1) {
+				// Found only one process, which must be the leader because we have
+				// purposely expect it as a zombie.
+				UBA_ASSERT(procs->kp_proc.p_pid == pgid);
+				free(procs);
+				return 0;
+			}
+
+			// More than one process left in the process group.  Pause a little bit
+			// before retrying to avoid burning CPU.
+			struct timespec ts;
+			ts.tv_sec = 0;
+			ts.tv_nsec = 1000000;
+			if (nanosleep(&ts, NULL) == -1) {
+				UBA_ASSERT(FALSE);
+				return -1;
+			}
+		}
+	}
+#endif
+
 	u32 ProcessImpl::InternalCreateProcess(bool runningRemote, void* environment, FileMappingHandle communicationHandle, u64 communicationOffset)
 	{
 		ScopedWriteLock initLock(m_initLock);
diff --git a/Engine/Source/Programs/UnrealBuildAccelerator/Common/Public/UbaFileMapping.h b/Engine/Source/Programs/UnrealBuildAccelerator/Common/Public/UbaFileMapping.h
--- a/Engine/Source/Programs/UnrealBuildAccelerator/Common/Public/UbaFileMapping.h
+++ b/Engine/Source/Programs/UnrealBuildAccelerator/Common/Public/UbaFileMapping.h
@@ -33,8 +33,8 @@ namespace uba
 
 	struct FileMappingHandle
 	{
-		FileMappingHandle(int shmFd_ = -1, int lockFd_ = -1, u64 uid_ = ~u64(0)) : shmFd(shmFd_), lockFd(lockFd_), uid(uid_) { };
-		FileMappingHandle(const FileMappingHandle& o) { shmFd = o.shmFd; uid = o.uid; }
+		FileMappingHandle(int shmFd_ = -1, int lockFd_ = -1, u64 uid_ = ~u64(0)) : shmFd(shmFd_), lockFd(lockFd_), uid(uid_) {};
+		FileMappingHandle(const FileMappingHandle& o) { shmFd = o.shmFd; uid = o.uid; lockFd = o.lockFd; }
 		int shmFd;
 		int lockFd;
 		u64 uid;
diff --git a/Engine/Source/Programs/UnrealBuildAccelerator/Common/Public/UbaProcess.h b/Engine/Source/Programs/UnrealBuildAccelerator/Common/Public/UbaProcess.h
--- a/Engine/Source/Programs/UnrealBuildAccelerator/Common/Public/UbaProcess.h
+++ b/Engine/Source/Programs/UnrealBuildAccelerator/Common/Public/UbaProcess.h
@@ -69,6 +69,11 @@ namespace uba
 
 		void SetRulesIndex(const ProcessStartInfo& si);
 		const tchar* InternalGetChildLogFile(StringBufferBase& temp);
+		
+		#if PLATFORM_MAC
+		int WaitForProcessGroup(pid_t pgid);
+		#endif
+
 		u32 InternalCreateProcess(bool runningRemote, void* environment, FileMappingHandle communicationHandle, u64 communicationOffset);
 		u32	InternalExitProcess(bool cancel);
 		void ClearTempFiles();
diff --git a/Engine/Source/Programs/UnrealBuildAccelerator/Core/Public/UbaPathUtils.h b/Engine/Source/Programs/UnrealBuildAccelerator/Core/Public/UbaPathUtils.h
--- a/Engine/Source/Programs/UnrealBuildAccelerator/Core/Public/UbaPathUtils.h
+++ b/Engine/Source/Programs/UnrealBuildAccelerator/Core/Public/UbaPathUtils.h
@@ -201,6 +201,18 @@ namespace uba
 			*bufferCharLen = u32(charLen - 1); // Remove terminator
 		return true;
 #else
+		char tmp[MaxPath] = {0};
+		if (fileName[0] == '~')
+		{
+			if (realpath(fileName, tmp) == 0)
+			{
+				// If we can't  get a real path to whatever
+				// fileName we're looking for, then we need to
+				// assume it's not a real file, errno should be set
+				// so just return false
+				return false; 
+			}
+		}
 		UBA_ASSERT(fileName[0] != '~');
 
 		u64 memPos = 0;
diff --git a/Engine/Source/Programs/UnrealBuildAccelerator/Detours/Private/PosixOS/UbaDetoursFunctionsPosix.cpp b/Engine/Source/Programs/UnrealBuildAccelerator/Detours/Private/PosixOS/UbaDetoursFunctionsPosix.cpp
--- a/Engine/Source/Programs/UnrealBuildAccelerator/Detours/Private/PosixOS/UbaDetoursFunctionsPosix.cpp
+++ b/Engine/Source/Programs/UnrealBuildAccelerator/Detours/Private/PosixOS/UbaDetoursFunctionsPosix.cpp
@@ -37,10 +37,13 @@ using namespace uba;
 	DETOURED_FUNCTION(fchdir) \
 	DETOURED_FUNCTION(chroot) \
 	DETOURED_FUNCTION(getcwd) \
+	DETOURED_FUNCTION(getenv) \
 	DETOURED_FUNCTION(setenv) \
 	DETOURED_FUNCTION(realpath) \
 	DETOURED_FUNCTION(readlink) \
 	DETOURED_FUNCTION(readlinkat) \
+	DETOURED_FUNCTION(read) \
+	DETOURED_FUNCTION(pread) \
 	DETOURED_FUNCTION(open) \
 	DETOURED_FUNCTION(write) \
 	DETOURED_FUNCTION(dup) \
@@ -65,6 +68,7 @@ using namespace uba;
 	DETOURED_FUNCTION(closedir) \
 	DETOURED_FUNCTION(stat) \
 	DETOURED_FUNCTION(truncate) \
+	DETOURED_FUNCTION(lstat) \
 	DETOURED_FUNCTION(glob) \
 	DETOURED_FUNCTION(chmod) \
 	DETOURED_FUNCTION(rename) \
@@ -598,7 +602,7 @@ int Shared_fstat(const char* funcName, int fd, struct stat* attr, const True_fst
 			//UBA_ASSERT(attr->st_dev == attr2.st_dev)
 			UBA_ASSERTF(attr->st_ino == attr2.st_ino, "fstat: st_ino mismatch for %s (%llu vs %llu)", fi.originalName, attr->st_ino, attr2.st_ino);
 			UBA_ASSERT(isDir || attr->st_size == attr2.st_size);
-			UBA_ASSERTF(isDir || FromTimeSpec(attr->st_mtimespec) == FromTimeSpec(attr2.st_mtimespec), "fstat: st_mtim mismatch for %s (%llu vs %llu)", fi.originalName, FromTimeSpec(attr->st_mtimespec), FromTimeSpec(attr2.st_mtim));
+			UBA_ASSERTF(isDir || FromTimeSpec(attr->st_mtimespec) == FromTimeSpec(attr2.st_mtimespec), "fstat: st_mtim mismatch for %s (%llu vs %llu)", fi.originalName, FromTimeSpec(attr->st_mtimespec), FromTimeSpec(attr2.st_mtimespec));
 		}
 		else
 		{
@@ -617,7 +621,12 @@ int Shared_stat(const char* funcName, const char* file, struct stat* attr, const
 		return trueStat(file, attr);
 
 	StringBuffer<> fixedFile;
-	FixPath(fixedFile, file);
+	if ((!FixPath(fixedFile, file)) || (fixedFile.Equals("/")) || (access(file, F_OK) != 0))
+	{
+		int res =  trueStat(file, attr);
+		return res;
+	}
+
 	UBA_ASSERTF(fixedFile.count, "FixPath failed with %s", file);
 
 	StringBuffer<> temp;
@@ -673,6 +682,31 @@ int Shared_stat(const char* funcName, const char* file, struct stat* attr, const
 
 
 // Detoured functions
+// #if PLATFORM_MAC
+// extern "C" {
+// char* realpath$DARWIN_EXTSN(const char* path, char* resolved_path);
+// }
+// UBA_EXPORT int UBA_WRAPPER(_NSGetExecutablePath)(char* buf, uint32_t* bufsize)
+// {
+// 	int res = TRUE_WRAPPER(_NSGetExecutablePath)(buf, bufsize);
+// 	return TRUE_WRAPPER(_NSGetExecutablePath)(buf, bufsize);
+// }
+
+// UBA_EXPORT int UBA_WRAPPER(_NSGetEnviron)(char* buf, uint32_t* bufsize)
+// {
+// 	int res = TRUE_WRAPPER(_NSGetEnviron)(buf, bufsize);
+// 	printf("%s for %s\n", __func__, buf);
+// 	return TRUE_WRAPPER(_NSGetEnviron)(buf, bufsize);
+// }
+// extern "C" char* UBA_WRAPPER(realpath$DARWIN_EXTSN)(const char* path, char* resolved_path)
+// {
+// 	printf(">>>>>> uba_realpathDARWIN: %s\n", path);
+// 	char* res = TRUE_WRAPPER(realpath$DARWIN_EXTSN)(path, resolved_path);
+// 	printf("<<<<< uba_realpathDARWIN: %s\n", strlen(resolved_path) > 0 ? resolved_path : "(NULL)");
+
+// 	return res;
+// }
+// #endif
 
 
 UBA_EXPORT int UBA_WRAPPER(chdir)(const char* path)
@@ -705,13 +739,13 @@ UBA_EXPORT char* UBA_WRAPPER(getcwd)(char* buf, size_t size)
 	return buf;
 }
 
-//UBA_EXPORT char* UBA_WRAPPER(getenv)(const char* name)
-//{
-//	UBA_INIT_DETOUR(getenv, name);
-//	auto res = TRUE_WRAPPER(getenv)(name);
-//	DEBUG_LOG_TRUE("getenv", "(%s) -> %s", name, res ? res : "<null>");
-//	return res;
-//}
+UBA_EXPORT char* UBA_WRAPPER(getenv)(const char* name)
+{
+	UBA_INIT_DETOUR(getenv, name);
+	auto res = TRUE_WRAPPER(getenv)(name);
+	DEBUG_LOG_TRUE("getenv", "(%s) -> %s", name, res ? res : "<null>");
+	return res;
+}
 
 UBA_EXPORT int UBA_WRAPPER(setenv)(const char* name, const char* value, int replace)
 {
@@ -1062,6 +1096,11 @@ UBA_EXPORT int UBA_WRAPPER(fstatat64)(int dirfd, const char* pathname, struct st
 }
 #endif
 
+UBA_EXPORT int UBA_WRAPPER(lstat)(const char *path, struct stat *buf)
+{
+	return TRUE_WRAPPER(lstat)(path, buf);
+}
+
 UBA_EXPORT int UBA_WRAPPER(stat)(const char* file, struct stat* attr)
 {
 	UBA_INIT_DETOUR(stat, file, attr);
@@ -1077,10 +1116,14 @@ UBA_EXPORT int UBA_WRAPPER(truncate)(const char* path, off_t length)
 
 UBA_EXPORT int UBA_WRAPPER(access)(const char* pathname, int mode)
 {
+	if (strlen(pathname) == 0) UBA_ASSERT(false);
 	UBA_INIT_DETOUR(access, pathname, mode);
 
 	StringBuffer<> fixedPath;
-	FixPath(fixedPath, pathname);
+	if (!FixPath(fixedPath, pathname))
+	{
+		return TRUE_WRAPPER(access)(pathname, mode);
+	}
 
 	bool checkIfDir = false;
 	StringBuffer<> temp;
@@ -1290,6 +1333,29 @@ UBA_EXPORT int UBA_WRAPPER(symlink)(const char* path1, const char* path2)
 	return TRUE_WRAPPER(symlink)(path1, path2);
 }
 
+UBA_EXPORT ssize_t	 UBA_WRAPPER(pread)(int __fd, void * __buf, size_t __nbyte, off_t __offset)
+{
+	// char filePath[PATH_MAX];
+	// if (fcntl(__fd, F_GETPATH, filePath) != -1)
+	// {
+	// 	// printf("***** PREAD: %d %s bytes: 0x%lx\n",__fd, filePath, __nbyte);
+	// 	// do something with the file path
+	// }
+	return TRUE_WRAPPER(pread)(__fd, __buf, __nbyte, __offset);
+}
+
+UBA_EXPORT ssize_t UBA_WRAPPER(read)(int fd, void *buf, size_t nbyte)
+{
+	// char filePath[PATH_MAX];
+	// if (fcntl(fd, F_GETPATH, filePath) != -1)
+	// {
+	// 	// printf("***** READ: %d %s bytes: %lu\n",fd, filePath, nbyte);
+	// 	// do something with the file path
+	// }
+	return TRUE_WRAPPER(read)(fd, buf, nbyte);
+}
+
+
 UBA_EXPORT int UBA_WRAPPER(remove)(const char* pathname)
 {
 	UBA_INIT_DETOUR(remove, pathname);
@@ -1396,7 +1462,11 @@ UBA_EXPORT int UBA_WRAPPER(posix_spawn)(pid_t* pid, const char* path, const posi
 		rulesStr.Append("UBA_RULES=").AppendValue(reader.ReadU32());
 
 		u32 dllNameSize = reader.ReadU32();
+	#if PLATFORM_LINUX
 		ldpreload.Append("LD_PRELOAD=");
+	#else
+		ldpreload.Append("DYLD_INSERT_LIBRARIES=");
+	#endif
 		reader.ReadBytes(ldpreload.data + ldpreload.count, dllNameSize);
 		ldpreload.Resize(ldpreload.count + dllNameSize);
 
@@ -1501,7 +1571,6 @@ UBA_EXPORT pid_t UBA_WRAPPER(wait4)(pid_t pid, int* status, int options, struct
 	UBA_ASSERT(!t_inVfork);
 	pid_t res = TRUE_WRAPPER(wait4)(pid, status, options, rusage);
 	DEBUG_LOG_TRUE("wait4", "");
-	UBA_ASSERT(false);
 	return res;
 }
 
@@ -1581,7 +1650,9 @@ UBA_EXPORT pid_t UBA_WRAPPER(vfork)(void)
 	pid_t pid = fork();
 	if (pid == 0)
 	{
+	#if PLATFORM_LINUX
 		prctl(PR_SET_PDEATHSIG, SIGHUP, 0, 0, 0); // We want the process to die if the parent die
+	#endif
 		t_inVfork = 1;
 	}
 	return pid;
diff --git a/Engine/Source/Programs/UnrealBuildAccelerator/Detours/Private/UbaDetoursFileMappingTable.cpp b/Engine/Source/Programs/UnrealBuildAccelerator/Detours/Private/UbaDetoursFileMappingTable.cpp
--- a/Engine/Source/Programs/UnrealBuildAccelerator/Detours/Private/UbaDetoursFileMappingTable.cpp
+++ b/Engine/Source/Programs/UnrealBuildAccelerator/Detours/Private/UbaDetoursFileMappingTable.cpp
@@ -180,7 +180,7 @@ namespace uba
 			return dirTableOffset;
 
 		const tchar* lastPathSeparator = TStrrchr(entryName, PathSeparator);
-		UBA_ASSERTF(lastPathSeparator, TC("No path separator found in %s"), entryName);
+		UBA_ASSERTF(lastPathSeparator, TC("No path separator found in %s"), strlen(entryName) > 0 ? entryName : "(NULL)");
 
 		#if PLATFORM_WINDOWS
 		UBA_ASSERT(wcsncmp(entryName, g_systemTemp.data, g_systemTemp.count) != 0);
diff --git a/Engine/Source/Programs/UnrealBuildAccelerator/TestApp/Private/PosixOS/UbaTestAppPosix.cpp b/Engine/Source/Programs/UnrealBuildAccelerator/TestApp/Private/PosixOS/UbaTestAppPosix.cpp
--- a/Engine/Source/Programs/UnrealBuildAccelerator/TestApp/Private/PosixOS/UbaTestAppPosix.cpp
+++ b/Engine/Source/Programs/UnrealBuildAccelerator/TestApp/Private/PosixOS/UbaTestAppPosix.cpp
@@ -136,8 +136,12 @@ int main()
 	if (stat("FileWF", &attrWF) == -1)
 		return LogError("stat for FileW failed");
 
+	char fullPath[PATH_MAX];
+	if (realpath("/usr/bin/clang", fullPath) == nullptr)
+		return LogError("realpath for 'clang' failed");
+
 	struct stat attrRoot;
-	if (stat("/", &attrRoot) != -1)
+	if (stat("/", &attrRoot) != 0)
 		return LogError("stat for '/' failed");
 
 	return 0;
