[UBA]
* Fixed windows compile errors

diff --git a/Engine/Source/Programs/UnrealBuildAccelerator/Detours/Private/UbaDetoursFileMappingTable.cpp b/Engine/Source/Programs/UnrealBuildAccelerator/Detours/Private/UbaDetoursFileMappingTable.cpp
--- a/Engine/Source/Programs/UnrealBuildAccelerator/Detours/Private/UbaDetoursFileMappingTable.cpp
+++ b/Engine/Source/Programs/UnrealBuildAccelerator/Detours/Private/UbaDetoursFileMappingTable.cpp
@@ -180,7 +180,7 @@ namespace uba
 			return dirTableOffset;
 
 		const tchar* lastPathSeparator = TStrrchr(entryName, PathSeparator);
-		UBA_ASSERTF(lastPathSeparator, TC("No path separator found in %s"), strlen(entryName) > 0 ? entryName : "(NULL)");
+		UBA_ASSERTF(lastPathSeparator, TC("No path separator found in %s"), TStrlen(entryName) > 0 ? entryName : TC("(NULL)"));
 
 		#if PLATFORM_WINDOWS
 		UBA_ASSERT(wcsncmp(entryName, g_systemTemp.data, g_systemTemp.count) != 0);
