[UBA]
* Removed hungarian notation on platform agnostic code

diff --git a/Engine/Source/Programs/UnrealBuildAccelerator/Common/Private/UbaFileMapping.cpp b/Engine/Source/Programs/UnrealBuildAccelerator/Common/Private/UbaFileMapping.cpp
--- a/Engine/Source/Programs/UnrealBuildAccelerator/Common/Private/UbaFileMapping.cpp
+++ b/Engine/Source/Programs/UnrealBuildAccelerator/Common/Private/UbaFileMapping.cpp
@@ -154,33 +154,33 @@ namespace uba
 #endif
 	}
 
-	FileMappingHandle CreateFileMappingW(FileHandle hFile, u32 flProtect, u64 maxSize)
+	FileMappingHandle CreateFileMappingW(FileHandle file, u32 flProtect, u64 maxSize)
 	{
 		ExtendedTimerScope ts(SystemStats::GetCurrent().createFileMapping);
 #if PLATFORM_WINDOWS
-		return { InternalCreateFileMappingW(asHANDLE(hFile), flProtect, (DWORD)ToHigh(maxSize), ToLow(maxSize), NULL) };
+		return { InternalCreateFileMappingW(asHANDLE(file), flProtect, (DWORD)ToHigh(maxSize), ToLow(maxSize), NULL) };
 #else
 		FileMappingHandle h;
-		h.shmFd = asFileDescriptor(hFile);
+		h.shmFd = asFileDescriptor(file);
 		lseek(h.shmFd, maxSize-1, SEEK_SET);
 		write(h.shmFd, "", 1);
 		return h;
 #endif
 	}
 
-	u8* MapViewOfFile(FileMappingHandle hFileMappingObject, u32 desiredAccess, u64 offset, u64 bytesToMap)
+	u8* MapViewOfFile(FileMappingHandle fileMappingObject, u32 desiredAccess, u64 offset, u64 bytesToMap)
 	{
 		ExtendedTimerScope ts(SystemStats::GetCurrent().mapViewOfFile);
 #if PLATFORM_WINDOWS
-		return (u8*)::MapViewOfFile(hFileMappingObject.handle, desiredAccess, (DWORD)ToHigh(offset), ToLow(offset), bytesToMap);
+		return (u8*)::MapViewOfFile(fileMappingObject.handle, desiredAccess, (DWORD)ToHigh(offset), ToLow(offset), bytesToMap);
 #else
 		int prot = 0;
 		if (desiredAccess & FILE_MAP_READ)
 			prot |= PROT_READ;
 		if (desiredAccess & FILE_MAP_WRITE)
 			prot |= PROT_WRITE;
-		UBA_ASSERT(hFileMappingObject.IsValid());
-		int shmFd = hFileMappingObject.shmFd;
+		UBA_ASSERT(fileMappingObject.IsValid());
+		int shmFd = fileMappingObject.shmFd;
 		void* rptr = mmap(NULL, bytesToMap, prot, MAP_SHARED, shmFd, s64(offset));
 		if (rptr != MAP_FAILED)
 			return (u8*)rptr;
diff --git a/Engine/Source/Programs/UnrealBuildAccelerator/Common/Public/UbaFileMapping.h b/Engine/Source/Programs/UnrealBuildAccelerator/Common/Public/UbaFileMapping.h
--- a/Engine/Source/Programs/UnrealBuildAccelerator/Common/Public/UbaFileMapping.h
+++ b/Engine/Source/Programs/UnrealBuildAccelerator/Common/Public/UbaFileMapping.h
@@ -46,8 +46,8 @@ namespace uba
 #endif
 
 	FileMappingHandle CreateMemoryMappingW(Logger& logger, u32 flProtect, u64 maxSize, const tchar* name = nullptr);
-	FileMappingHandle CreateFileMappingW(FileHandle hFile, u32 flProtect, u64 maxSize = 0);
-	u8* MapViewOfFile(FileMappingHandle hFileMappingObject, u32 dwDesiredAccess, u64 offset, u64 dwNumberOfBytesToMap);
+	FileMappingHandle CreateFileMappingW(FileHandle file, u32 flProtect, u64 maxSize = 0);
+	u8* MapViewOfFile(FileMappingHandle fileMappingObject, u32 dwDesiredAccess, u64 offset, u64 dwNumberOfBytesToMap);
 	bool MapViewCommit(void* address, u64 size);
 	bool UnmapViewOfFile(const void* lpBaseAddress, u64 bytesToUnmap);
 	bool CloseFileMapping(FileMappingHandle h);
