[Uba]
* Added PrepopulateCas to StorageClient. This can be used to use local files as source for cas entries. It doesn't really matter where files are located as long as their cas key matches whatever files the client needs
* Added -populateCas=<dir> to UbaAgent. This can be used to add one or more directories for pre population

diff --git a/Engine/Source/Programs/UnrealBuildAccelerator/Agent/Private/UbaAgent.cpp b/Engine/Source/Programs/UnrealBuildAccelerator/Agent/Private/UbaAgent.cpp
--- a/Engine/Source/Programs/UnrealBuildAccelerator/Agent/Private/UbaAgent.cpp
+++ b/Engine/Source/Programs/UnrealBuildAccelerator/Agent/Private/UbaAgent.cpp
@@ -97,6 +97,7 @@ namespace uba
 		logger.Info(TC("  -memwait=<percent>      The amount of memory needed to spawn a process. Set this to 100 to disable. Defaults to 80%%"));
 		logger.Info(TC("  -memkill=<percent>      The amount of memory needed before processes starts to be killed. Set this to 100 to disable. Defaults to 90%%"));
 		logger.Info(TC("  -crypto=<key>           16 bytes crypto key used for secure network transfer"));
+		logger.Info(TC("  -populateCas=<dir>      Prepopulate cas database with files in dir. If files needed exists on machine this can be an optimization"));
 		logger.Info(TC(""));
 		return -1;
 	}
@@ -414,6 +415,7 @@ namespace uba
 		u32 memKillLoadPercent = 90;
 		u8 crypto[16];
 		bool hasCrypto = false;
+		Vector<TString> populateCasDirs;
 
 		for (int i=1; i!=argc; ++i)
 		{
@@ -602,6 +604,12 @@ namespace uba
 				((u64*)crypto)[1] = StringToValue(value.data + 16, 16);
 				hasCrypto = true;
 			}
+			else if (name.Equals(TC("-populateCas")))
+			{
+				if (value.IsEmpty())
+					return PrintHelp(TC("-populateCas needs a dir"));
+				populateCasDirs.push_back(value.data);
+			}
 			else if (name.Equals(TC("-sentry")))
 			{
 				if (value.IsEmpty())
@@ -679,8 +687,13 @@ namespace uba
 		#if defined(UBA_USE_AWS)
 		AWS aws;
 		aws.Init(logger, extraInfo, TC("Agent"));
+		if (zone.IsEmpty())
+			zone.Append(aws.GetAvailabilityZone());
 		#endif
 
+		if (zone.count)
+			extraInfo.Append(TC(", ")).Append(zone);
+
 		if (IsRunningWine())
 			extraInfo.Append(TC(", Linux/WINE"));
 
@@ -920,18 +933,15 @@ namespace uba
 			storageInfo.zone = zone.data;
 			storageInfo.proxyPort = proxyPort;
 
-			#if defined(UBA_USE_AWS)
-			if (zone.IsEmpty())
-				storageInfo.zone = aws.GetAvailabilityZone();
-			#endif
-
-
 			auto storageClient = new StorageClient(storageInfo);
 			auto bscsg = MakeGuard([&]() { g_storageClient = nullptr; delete storageClient; });
 
 			if (!storageClient->LoadCasTable(true))
 				return -1;
 
+			if (!storageClient->PopulateCasFromDirs(populateCasDirs, maxProcessCount))
+				return -1;
+
 			SessionClient* sessionClient = nullptr;
 
 			CasKey keys[2];
diff --git a/Engine/Source/Programs/UnrealBuildAccelerator/Common/Private/UbaStorage.cpp b/Engine/Source/Programs/UnrealBuildAccelerator/Common/Private/UbaStorage.cpp
--- a/Engine/Source/Programs/UnrealBuildAccelerator/Common/Private/UbaStorage.cpp
+++ b/Engine/Source/Programs/UnrealBuildAccelerator/Common/Private/UbaStorage.cpp
@@ -1951,7 +1951,7 @@ namespace uba
 
 			StringBuffer<> casFile;
 			#if !UBA_USE_SPARSEFILE
-			if (!StorageImpl::GetCasFileName(casFile, actualKey))
+			if (!GetCasFileName(casFile, actualKey))
 				return false;
 			#else
 			UBA_ASSERT(false);
diff --git a/Engine/Source/Programs/UnrealBuildAccelerator/Common/Private/UbaStorageClient.cpp b/Engine/Source/Programs/UnrealBuildAccelerator/Common/Private/UbaStorageClient.cpp
--- a/Engine/Source/Programs/UnrealBuildAccelerator/Common/Private/UbaStorageClient.cpp
+++ b/Engine/Source/Programs/UnrealBuildAccelerator/Common/Private/UbaStorageClient.cpp
@@ -1,6 +1,7 @@
 // Copyright Epic Games, Inc. All Rights Reserved.
 
 #include "UbaStorageClient.h"
+#include "UbaDirectoryIterator.h"
 #include "UbaFileAccessor.h"
 #include "UbaNetworkClient.h"
 #include "UbaNetworkMessage.h"
@@ -53,8 +54,8 @@ namespace uba
 	StorageClient::~StorageClient()
 	{
 		delete m_proxyClient;
-		for (auto& pair : m_temporaryStorageFiles)
-			CloseFileMapping(pair.second.handle);
+		for (auto& pair : m_localStorageFiles)
+			CloseFileMapping(pair.second.casEntry.mappingHandle);
 	}
 
 	bool StorageClient::IsUsingProxy()
@@ -69,15 +70,33 @@ namespace uba
 			m_proxyClient->client.Disconnect();
 	}
 
+	bool StorageClient::PopulateCasFromDirs(const DirVector& directories, u32 workerCount)
+	{
+		if (directories.empty())
+			return true;
+
+		WorkManagerImpl workManager(workerCount);
+		ReaderWriterLock lock;
+		bool success = true;
+		for (auto& dir : directories)
+			success = PopulateCasFromDirsRecursive(dir.c_str(), workManager) && success;
+		workManager.Wait();
+
+		if (u32 fileCount = u32(m_localStorageFiles.size()))
+			m_logger.Info(TC("Prepopulated %u files to cas"), fileCount);
+
+		return success;
+	}
+
 #if !UBA_USE_SPARSEFILE
 	bool StorageClient::GetCasFileName(StringBufferBase& out, const CasKey& casKey)
 	{
-		ScopedReadLock tempLock(m_temporaryStorageFilesLock);
-		auto findIt = m_temporaryStorageFiles.find(casKey);
-		if (findIt != m_temporaryStorageFiles.end())
+		ScopedReadLock tempLock(m_localStorageFilesLock);
+		auto findIt = m_localStorageFiles.find(casKey);
+		if (findIt != m_localStorageFiles.end())
 		{
-			if (findIt->second.handle.IsValid())
-				GetMappingString(out, findIt->second.handle, 0);
+			if (findIt->second.casEntry.mappingHandle.IsValid())
+				GetMappingString(out, findIt->second.casEntry.mappingHandle, 0);
 			else
 				out.Append(findIt->second.fileName);
 			return true;
@@ -90,17 +109,17 @@ namespace uba
 
 	MappedView StorageClient::MapView(const CasKey& casKey, const tchar* hint)
 	{
-		ScopedReadLock tempLock(m_temporaryStorageFilesLock);
-		auto findIt = m_temporaryStorageFiles.find(AsCompressed(casKey, false));
-		bool isValid = findIt != m_temporaryStorageFiles.end();
+		ScopedReadLock tempLock(m_localStorageFilesLock);
+		auto findIt = m_localStorageFiles.find(AsCompressed(casKey, false));
+		bool isValid = findIt != m_localStorageFiles.end();
 		tempLock.Leave();
 		if (!isValid)
 			return StorageImpl::MapView(casKey, hint);
-		MappedFile& file = findIt->second;
+		LocalFile& file = findIt->second;
 
 		MappedView view;
-		view.handle = file.handle;
-		view.size = file.size;
+		view.handle = file.casEntry.mappingHandle;
+		view.size = file.casEntry.size;
 		view.offset = 0;
 		view.isCompressed = false;
 		return view;
@@ -133,15 +152,15 @@ namespace uba
 
 		// Cas file might have been created by this client and in that case we can just reuse the file we just wrote. No need to fetch from server
 		// This needs to be first so it doesnt end up in the cas table even though it is not in there. (otherwise it might be garbage collected)
-		ScopedReadLock tempLock(m_temporaryStorageFilesLock);
-		auto findIt = m_temporaryStorageFiles.find(AsCompressed(casKey, false));
-		if (findIt != m_temporaryStorageFiles.end())
+		ScopedReadLock tempLock(m_localStorageFilesLock);
+		auto findIt = m_localStorageFiles.find(AsCompressed(casKey, false));
+		if (findIt != m_localStorageFiles.end())
 		{
 			out.casKey = findIt->first;
-			MappedFile& mf = findIt->second;
-			out.size = mf.size;
-			out.view.handle = mf.handle;
-			out.view.size = mf.size;
+			LocalFile& mf = findIt->second;
+			out.size = mf.casEntry.size;
+			out.view.handle = mf.casEntry.mappingHandle;
+			out.view.size = mf.casEntry.size;
 			out.view.isCompressed = false;
 			return true;
 		}
@@ -615,6 +634,20 @@ namespace uba
 		return true;
 	}
 
+	bool StorageClient::HasCasFile(const CasKey& casKey, CasEntry** out)
+	{
+		CasKey localKey = AsCompressed(casKey, false);
+		ScopedReadLock tempLock(m_localStorageFilesLock);
+		auto findIt = m_localStorageFiles.find(localKey);
+		if (findIt != m_localStorageFiles.end())
+		{
+			if (out)
+				*out = &findIt->second.casEntry;
+			return true;
+		}
+		return StorageImpl::HasCasFile(casKey, out);
+	}
+
 	bool StorageClient::StoreCasFile(CasKey& out, StringKey fileNameKey, const tchar* fileName, FileMappingHandle mappingHandle, u64 mappingOffset, u64 fileSize, const tchar* hint, bool deferCreation, bool keepMappingInMemory)
 	{
 		NetworkClient& client = m_client; // Don't use proxy
@@ -655,24 +688,30 @@ namespace uba
 
 			if (keepMappingInMemory)
 			{
+				ScopedWriteLock lock(m_localStorageFilesLock);
+				auto insres = m_localStorageFiles.try_emplace(AsCompressed(casKey, false));
+				LocalFile& localFile = insres.first->second;
+				if (!insres.second && localFile.casEntry.mappingHandle.IsValid())
+					return true;
+
 				if (isPersistentMapping)
 				{
 					FileMappingHandle mappingHandle2;
 					if (DuplicateFileMapping(GetCurrentProcessHandle(), mappingHandle, GetCurrentProcessHandle(), &mappingHandle2, FILE_MAP_READ, false, 0))
 					{
-						ScopedWriteLock lock(m_temporaryStorageFilesLock);
-						m_temporaryStorageFiles.try_emplace(AsCompressed(casKey, false), MappedFile{ mappingHandle2, fileSize });
+						localFile.casEntry.mappingHandle = mappingHandle2;
+						localFile.casEntry.size = fileSize;
 					}
 					else
 						m_logger.Warning(TC("Failed to duplicate handle for file mapping %s (%s)"), fileName, LastErrorToText().data);
 				}
 				else
 				{
-					ScopedWriteLock lock(m_temporaryStorageFilesLock);
+					localFile.casEntry.size = fileSize;
 #if !UBA_USE_SPARSEFILE
-					m_temporaryStorageFiles.try_emplace(AsCompressed(casKey, false), MappedFile{ {0}, fileSize, fileName });
+					localFile.fileName = fileName;
 #else
-					m_temporaryStorageFiles.try_emplace(AsCompressed(casKey, false), MappedFile{ mappingHandle, fileSize });
+					localFile.casEntry.mappingHandle = mappingHandle2;
 					lock.Leave();
 					mappingClose.Cancel();
 #endif
@@ -994,4 +1033,51 @@ namespace uba
 
 		return true;
 	}
+
+	bool StorageClient::PopulateCasFromDirsRecursive(const tchar* dir, WorkManager& workManager)
+	{
+		StringBuffer<> fullPath;
+		fullPath.Append(dir).EnsureEndsWithSlash();
+		u32 dirLen = fullPath.count;
+		TraverseDir(m_logger, dir, [&](const DirectoryEntry& e)
+			{
+				fullPath.Resize(dirLen).Append(e.name);
+				if (IsDirectory(e.attributes))
+				{
+					PopulateCasFromDirsRecursive(fullPath.data, workManager);
+					return;
+				}
+
+				workManager.AddWork([&, filePath = TString(fullPath.data), name = TString(e.name)]()
+					{
+						FileInformation info;
+						if (!GetFileInformation(info, m_logger, filePath.c_str()))
+						{
+							m_logger.Error(TC("Failed to get information for file %s"), filePath.c_str());
+							return;
+						}
+
+						CasKey casKey;
+						if (!CalculateCasKey(casKey, filePath.c_str()))
+						{
+							m_logger.Error(TC("Failed to calculate cas key for %s"), filePath.c_str());
+							return;
+						}
+
+						ScopedWriteLock lookupLock(m_localStorageFilesLock);
+						auto insres = m_localStorageFiles.try_emplace(AsCompressed(casKey, false));
+						if (!insres.second)
+							return;
+						LocalFile& localFile = insres.first->second;
+						lookupLock.Leave();
+
+						localFile.casEntry.size = info.size;
+						localFile.fileName = filePath;
+
+					}, 1, TC(""));
+			});
+		return true;
+	}
+
+
 }
diff --git a/Engine/Source/Programs/UnrealBuildAccelerator/Common/Public/UbaStorage.h b/Engine/Source/Programs/UnrealBuildAccelerator/Common/Public/UbaStorage.h
--- a/Engine/Source/Programs/UnrealBuildAccelerator/Common/Public/UbaStorage.h
+++ b/Engine/Source/Programs/UnrealBuildAccelerator/Common/Public/UbaStorage.h
@@ -117,7 +117,7 @@ namespace uba
 
 		struct CasEntry;
 
-		bool HasCasFile(const CasKey& casKey, CasEntry** out = nullptr);
+		virtual bool HasCasFile(const CasKey& casKey, CasEntry** out = nullptr);
 		bool EnsureCasFile(const CasKey& casKey, const tchar* fileName);
 		CasKey CalculateCasKey(const tchar* fileName, FileHandle fileHandle, u64 fileSize, bool storeCompressed);
 		CasKey CalculateCasKey(u8* fileMem, u64 fileSize, bool storeCompressed);
diff --git a/Engine/Source/Programs/UnrealBuildAccelerator/Common/Public/UbaStorageClient.h b/Engine/Source/Programs/UnrealBuildAccelerator/Common/Public/UbaStorageClient.h
--- a/Engine/Source/Programs/UnrealBuildAccelerator/Common/Public/UbaStorageClient.h
+++ b/Engine/Source/Programs/UnrealBuildAccelerator/Common/Public/UbaStorageClient.h
@@ -33,20 +33,28 @@ namespace uba
 		bool IsUsingProxy();
 		void StopProxy();
 
-#if !UBA_USE_SPARSEFILE
+		using DirVector = Vector<TString>;
+		bool PopulateCasFromDirs(const DirVector& directories, u32 workerCount);
+
+		#if !UBA_USE_SPARSEFILE
 		virtual bool GetCasFileName(StringBufferBase& out, const CasKey& casKey) override;
-#endif
+		#endif
+
 		virtual MappedView MapView(const CasKey& casKey, const tchar* hint) override;
 
 		virtual bool GetZone(StringBufferBase& out) override;
 		virtual bool RetrieveCasFile(RetrieveResult& out, const CasKey& casKey, const tchar* hint, FileMappingBuffer* mappingBuffer = nullptr, u64 memoryMapAlignment = 1, bool allowProxy = true) override;
 		virtual bool StoreCasFile(CasKey& out, const tchar* fileName, const CasKey& casKeyOverride = CasKeyZero, bool deferCreation = false) override;
 		virtual bool StoreCasFile(CasKey& out, StringKey fileNameKey, const tchar* fileName, FileMappingHandle mappingFile, u64 mappingOffset, u64 fileSize, const tchar* hint, bool deferCreation = false, bool keepMappingInMemory = false) override;
+		virtual bool HasCasFile(const CasKey& casKey, CasEntry** out = nullptr) override;
 		virtual void Ping() override;
 
 		static bool SendBatchMessages(NetworkClient& client, u16 fetchId, u8* slot, u64 capacity, u64 left, u32 messageMaxSize, u32& readIndex, u32& responseSize);
 		static bool SendAllSegments(NetworkClient& client, u16 fetchId, u8* readBuffer, u64 left, u32 messageMaxSize);
+
+	private:
 		bool SendFile(const CasKey& casKey, const tchar* fileName, u8* sourceMem, u64 sourceSize, const tchar* hint);
+		bool PopulateCasFromDirsRecursive(const tchar* dir, WorkManager& workManager);
 
 		NetworkClient& m_client;
 		bool m_sendCompressed;
@@ -55,9 +63,9 @@ namespace uba
 
 		TString m_zone;
 
-		ReaderWriterLock m_temporaryStorageFilesLock;
-		struct MappedFile { FileMappingHandle handle; u64 size; TString fileName; };
-		UnorderedMap<CasKey, MappedFile> m_temporaryStorageFiles;
+		ReaderWriterLock m_localStorageFilesLock;
+		struct LocalFile { CasEntry casEntry; TString fileName; };
+		UnorderedMap<CasKey, LocalFile> m_localStorageFiles;
 
 		ReaderWriterLock m_sendOneAtTheTimeLock;
 		ReaderWriterLock m_retrieveOneBatchAtTheTimeLock;
