[Uba]
* Added disallowed path support. It is not possible to register paths that are disallowed for serving to clients

diff --git a/Engine/Source/Programs/UnrealBuildAccelerator/Common/Private/UbaSessionClient.cpp b/Engine/Source/Programs/UnrealBuildAccelerator/Common/Private/UbaSessionClient.cpp
--- a/Engine/Source/Programs/UnrealBuildAccelerator/Common/Private/UbaSessionClient.cpp
+++ b/Engine/Source/Programs/UnrealBuildAccelerator/Common/Private/UbaSessionClient.cpp
@@ -216,6 +216,8 @@ namespace uba
 		auto insres = m_handledApplicationEnvironments.insert(application);
 		if (insres.second)
 		{
+			auto failGuard = MakeGuard([&]() { m_handledApplicationEnvironments.erase(insres.first); });
+
 			StackBinaryReader<SendMaxSize> reader;
 			{
 				StackBinaryWriter<1024> writer;
@@ -264,6 +266,8 @@ namespace uba
 				if (!WriteBinFile(temp, moduleName, newCasKey, keyStr, fileAttributes))
 					return false;
 			}
+
+			failGuard.Cancel();
 		}
 		out.Append(m_sessionBinDir).Append(keyStr).Append(PathSeparator).AppendFileName(application);
 		return true;
@@ -1347,7 +1351,8 @@ namespace uba
 					{
 						m_logger.Error(TC("Failed to ensure application environment for %s"), startInfo.application);
 						SendReturnProcess(startInfo.processId, TC("Failed to ensure application environment"));
-						continue;
+						m_loop = false;
+						break;
 					}
 
 					void* env = GetProcessEnvironmentVariables();
diff --git a/Engine/Source/Programs/UnrealBuildAccelerator/Common/Private/UbaStorage.cpp b/Engine/Source/Programs/UnrealBuildAccelerator/Common/Private/UbaStorage.cpp
--- a/Engine/Source/Programs/UnrealBuildAccelerator/Common/Private/UbaStorage.cpp
+++ b/Engine/Source/Programs/UnrealBuildAccelerator/Common/Private/UbaStorage.cpp
@@ -503,6 +503,11 @@ namespace uba
 		CasEntryAccessed(casEntry);
 	}
 
+	bool StorageImpl::IsDisallowedPath(const tchar* fileName)
+	{
+		return false;
+	}
+
 	bool StorageImpl::DecompressMemoryToMemory(u8* compressedData, u8* writeData, u64 decompressedSize)
 	{
 		StorageStats& stats = Stats();
@@ -730,6 +735,7 @@ namespace uba
 				return true;
 
 		casEntry.key = casKey;
+		casEntry.disallowed = IsDisallowedPath(fileName);
 
 #if !UBA_USE_SPARSEFILE
 		StringBuffer<> casFile;
diff --git a/Engine/Source/Programs/UnrealBuildAccelerator/Common/Private/UbaStorageServer.cpp b/Engine/Source/Programs/UnrealBuildAccelerator/Common/Private/UbaStorageServer.cpp
--- a/Engine/Source/Programs/UnrealBuildAccelerator/Common/Private/UbaStorageServer.cpp
+++ b/Engine/Source/Programs/UnrealBuildAccelerator/Common/Private/UbaStorageServer.cpp
@@ -47,6 +47,12 @@ namespace uba
 		m_server.UnregisterService(ServiceId);
 	}
 
+	bool StorageServer::RegisterDisallowedPath(const tchar* path)
+	{
+		m_disallowedPaths.push_back(path);
+		return true;
+	}
+
 	bool StorageServer::GetZone(StringBufferBase& out)
 	{
 		if (m_zone.empty())
@@ -209,6 +215,14 @@ namespace uba
 		return StorageImpl::WriteCompressed(out, from, InvalidFileHandle, fileMem, mapping.fileSize, toFile);
 	}
 
+	bool StorageServer::IsDisallowedPath(const tchar* fileName)
+	{
+		for (auto& path : m_disallowedPaths)
+			if (StartsWith(fileName, path.c_str()))
+				return true;
+		return false;
+	}
+
 	void StorageServer::SetTrace(Trace* trace, bool detailed)
 	{
 		m_trace = trace;
@@ -388,6 +402,13 @@ namespace uba
 					casEntry = &findIt->second;
 				}
 
+				if (casEntry->disallowed)
+				{
+					writer.WriteU16(0);
+					m_logger.Error(TC("Client is asking for cas content of file that is not allowed to be transferred. (%s)"), hint.data);
+					return true;
+				}
+				
 				FileHandle readFileHandle = InvalidFileHandle;
 				u64 fileSize;
 				u8* memoryBegin = nullptr;
diff --git a/Engine/Source/Programs/UnrealBuildAccelerator/Common/Public/UbaStorage.h b/Engine/Source/Programs/UnrealBuildAccelerator/Common/Public/UbaStorage.h
--- a/Engine/Source/Programs/UnrealBuildAccelerator/Common/Public/UbaStorage.h
+++ b/Engine/Source/Programs/UnrealBuildAccelerator/Common/Public/UbaStorage.h
@@ -131,6 +131,7 @@ namespace uba
 		bool VerifyExisting(bool& outReturnValue, ScopedWriteLock& entryLock, const CasKey& casKey, CasEntry& casEntry, StringBufferBase& casFile);
 		bool AddCasFile(const tchar* fileName, const CasKey& casKey, bool deferCreation);
 		void CasEntryAccessed(const CasKey& casKey);
+		virtual bool IsDisallowedPath(const tchar* fileName);
 		virtual bool DecompressMemoryToMemory(u8* compressedData, u8* writeData, u64 decompressedSize) override;
 		bool DecompressMemoryToFile(u8* compressedData, FileAccessor& destination, u64 decompressedSize, bool useNoBuffering);
 
@@ -168,6 +169,7 @@ namespace uba
 			bool exists = false; // File exists on disk
 			bool dropped = false; // This file is not seen anymore. will be deleted during shutdown
 			bool beingWritten = false; // This is set while file is being written (when coming from network)..
+			bool disallowed = false; // This is set if cas is created from disallowed file
 
 			FileMappingHandle mappingHandle;
 			u64 mappingOffset = 0;
diff --git a/Engine/Source/Programs/UnrealBuildAccelerator/Common/Public/UbaStorageServer.h b/Engine/Source/Programs/UnrealBuildAccelerator/Common/Public/UbaStorageServer.h
--- a/Engine/Source/Programs/UnrealBuildAccelerator/Common/Public/UbaStorageServer.h
+++ b/Engine/Source/Programs/UnrealBuildAccelerator/Common/Public/UbaStorageServer.h
@@ -20,10 +20,13 @@ namespace uba
 		StorageServer(const StorageServerCreateInfo& info);
 		~StorageServer();
 
+		bool RegisterDisallowedPath(const tchar* path);
+
 		virtual bool GetZone(StringBufferBase& out) override;
 		virtual bool RetrieveCasFile(RetrieveResult& out, const CasKey& casKey, const tchar* hint, FileMappingBuffer* mappingBuffer = nullptr, u64 memoryMapAlignment = 1, bool allowProxy = true) override;
 		virtual bool StoreCasFile(CasKey& out, StringKey fileNameKey, const tchar* fileName, FileMappingHandle mappingHandle, u64 mappingOffset, u64 fileSize, const tchar* hint, bool deferCreation = false, bool keepMappingInMemory = false) override;
 		virtual bool WriteCompressed(WriteResult& out, const tchar* from, const tchar* toFile) override;
+		virtual bool IsDisallowedPath(const tchar* fileName) override;
 		virtual void SetTrace(Trace* trace, bool detailed) override;
 		void OnDisconnected(u32 clientId);
 		bool HandleMessage(const ConnectionInfo& connectionInfo, u8 messageType, BinaryReader& reader, BinaryWriter& writer);
@@ -136,5 +139,7 @@ namespace uba
 		ReaderWriterLock m_loadCasTableLock;
 
 		Trace* m_trace = nullptr;
+
+		Vector<TString> m_disallowedPaths;
 	};
 }
diff --git a/Engine/Source/Programs/UnrealBuildAccelerator/Host/Private/UbaExports.cpp b/Engine/Source/Programs/UnrealBuildAccelerator/Host/Private/UbaExports.cpp
--- a/Engine/Source/Programs/UnrealBuildAccelerator/Host/Private/UbaExports.cpp
+++ b/Engine/Source/Programs/UnrealBuildAccelerator/Host/Private/UbaExports.cpp
@@ -198,6 +198,11 @@ extern "C"
 		storageServer->SaveCasTable(true);
 	}
 
+	void StorageServer_RegisterDisallowedPath(uba::StorageServer* storageServer, const uba::tchar* path)
+	{
+		storageServer->RegisterDisallowedPath(path);
+	}
+
 	/*
 uba::StorageClient* CreateStorageClient(uba::NetworkClient& client, const uba::tchar* rootDir, uba::u64 casCapacityBytes, bool storeCompressed, bool sendCompressed)
 	{
diff --git a/Engine/Source/Programs/UnrealBuildAccelerator/Host/Public/UbaExports.h b/Engine/Source/Programs/UnrealBuildAccelerator/Host/Public/UbaExports.h
--- a/Engine/Source/Programs/UnrealBuildAccelerator/Host/Public/UbaExports.h
+++ b/Engine/Source/Programs/UnrealBuildAccelerator/Host/Public/UbaExports.h
@@ -58,6 +58,7 @@ extern "C"
 	UBA_API uba::Storage* CreateStorageServer(uba::NetworkServer& server, const uba::tchar* rootDir, uba::u64 casCapacityBytes, bool storeCompressed, uba::LogWriter& writer = uba::g_consoleLogWriter, const uba::tchar* zone = TC(""));
 	UBA_API void DestroyStorageServer(uba::Storage* storageServer);
 	UBA_API void StorageServer_SaveCasTable(uba::Storage* storageServer);
+	UBA_API void StorageServer_RegisterDisallowedPath(uba::StorageServer* storageServer, const uba::tchar* path);
 
 	//UBA_API uba::StorageClient* CreateStorageClient(uba::NetworkClient& client, const uba::tchar* rootDir, uba::u64 casCapacityBytes, bool storeCompressed, bool sendCompressed);
 
