[Uba]
* Don't write output file if process is cancelled

diff --git a/Engine/Source/Programs/UnrealBuildAccelerator/Common/Private/UbaProcess.cpp b/Engine/Source/Programs/UnrealBuildAccelerator/Common/Private/UbaProcess.cpp
--- a/Engine/Source/Programs/UnrealBuildAccelerator/Common/Private/UbaProcess.cpp
+++ b/Engine/Source/Programs/UnrealBuildAccelerator/Common/Private/UbaProcess.cpp
@@ -831,6 +831,7 @@ namespace uba
 						m_processStats.Print(m_session.m_logger);
 					}
 
+					if (g_applicationRules[m_rulesIndex].rules->IsExitCodeSuccess(m_nativeProcessExitCode) && !IsCancelled())
 					{
 						ScopedWriteLock lock(m_writtenFilesLock);
 						for (auto& kv : m_writtenFiles)
@@ -843,12 +844,6 @@ namespace uba
 									m_messageSuccess = false;
 						}
 					}
-					//PROCESS_MEMORY_COUNTERS mem;
-					//mem.cb = sizeof(mem);
-					//GetProcessMemoryInfo(m_nativeProcessHandle, &mem, sizeof(mem));
-					//mem.PeakPagefileUsage;
-					//m_session.m_logger.Debug(TC("PROCESS MEM: WorkingSet: %s Page: %s"), BytesToText(mem.PeakWorkingSetSize), BytesToText(mem.PeakPagefileUsage));
-					
 					return false;
 				}
 			case MessageType_Custom:
diff --git a/Engine/Source/Programs/UnrealBuildAccelerator/Core/Private/UbaApplicationRules.cpp b/Engine/Source/Programs/UnrealBuildAccelerator/Core/Private/UbaApplicationRules.cpp
--- a/Engine/Source/Programs/UnrealBuildAccelerator/Core/Private/UbaApplicationRules.cpp
+++ b/Engine/Source/Programs/UnrealBuildAccelerator/Core/Private/UbaApplicationRules.cpp
@@ -41,6 +41,11 @@ namespace uba
 				return true;
 			return Super::KeepInMemory(fileName, fileNameLen, systemTemp);
 		}
+
+		virtual bool IsExitCodeSuccess(u32 exitCode) override
+		{
+			return exitCode == 0;
+		}
 	};
 
 	class ApplicationRulesClExe : public ApplicationRulesVC
@@ -131,6 +136,11 @@ namespace uba
 		{
 			return true;
 		}
+
+		virtual bool IsExitCodeSuccess(u32 exitCode) override
+		{
+			return exitCode == 0;
+		}
 	};
 
 	class ApplicationRulesClangPlusPlusExe : public ApplicationRulesClang
diff --git a/Engine/Source/Programs/UnrealBuildAccelerator/Core/Public/UbaApplicationRules.h b/Engine/Source/Programs/UnrealBuildAccelerator/Core/Public/UbaApplicationRules.h
--- a/Engine/Source/Programs/UnrealBuildAccelerator/Core/Public/UbaApplicationRules.h
+++ b/Engine/Source/Programs/UnrealBuildAccelerator/Core/Public/UbaApplicationRules.h
@@ -108,6 +108,11 @@ namespace uba
 		{
 			return false;
 		}
+
+		virtual bool IsExitCodeSuccess(u32 exitCode)
+		{
+			return true;
+		}
 	};
 
 
