[Uba]
* Fixed so memory block takes alignment for allocations

diff --git a/Engine/Source/Programs/UnrealBuildAccelerator/Common/Private/UbaSessionClient.cpp b/Engine/Source/Programs/UnrealBuildAccelerator/Common/Private/UbaSessionClient.cpp
--- a/Engine/Source/Programs/UnrealBuildAccelerator/Common/Private/UbaSessionClient.cpp
+++ b/Engine/Source/Programs/UnrealBuildAccelerator/Common/Private/UbaSessionClient.cpp
@@ -885,7 +885,7 @@ namespace uba
 			if (needMore)
 				toRead = left;
 
-			m_nameToHashTableMem.AllocateNoLock(toRead, TC("NameToHashTable"));
+			m_nameToHashTableMem.AllocateNoLock(toRead, 1, TC("NameToHashTable"));
 			reader.ReadBytes(pos, toRead);
 			localTableSize += toRead;
 
diff --git a/Engine/Source/Programs/UnrealBuildAccelerator/Common/Private/UbaSessionServer.cpp b/Engine/Source/Programs/UnrealBuildAccelerator/Common/Private/UbaSessionServer.cpp
--- a/Engine/Source/Programs/UnrealBuildAccelerator/Common/Private/UbaSessionServer.cpp
+++ b/Engine/Source/Programs/UnrealBuildAccelerator/Common/Private/UbaSessionServer.cpp
@@ -697,7 +697,7 @@ namespace uba
 					{
 						lookupCasKey = casKey;
 						BinaryWriter w(m_nameToHashTableMem.memory, m_nameToHashTableMem.writtenSize, NameToHashMemSize);
-						m_nameToHashTableMem.AllocateNoLock(sizeof(StringKey) + sizeof(CasKey), TC("NameToHashTable"));
+						m_nameToHashTableMem.AllocateNoLock(sizeof(StringKey) + sizeof(CasKey), 1, TC("NameToHashTable"));
 						w.WriteStringKey(fileNameKey);
 						w.WriteCasKey(casKey);
 					}
@@ -1445,7 +1445,7 @@ namespace uba
 			return;
 		lookupCasKey = cachedInfo.casKey;
 		BinaryWriter w(m_nameToHashTableMem.memory, m_nameToHashTableMem.writtenSize, NameToHashMemSize);
-		m_nameToHashTableMem.AllocateNoLock(sizeof(StringKey) + sizeof(CasKey), TC("NameToHashTable"));
+		m_nameToHashTableMem.AllocateNoLock(sizeof(StringKey) + sizeof(CasKey), 1, TC("NameToHashTable"));
 		w.WriteStringKey(fileNameKey);
 		w.WriteCasKey(lookupCasKey);
 	}
diff --git a/Engine/Source/Programs/UnrealBuildAccelerator/Common/Private/UbaStorageClient.cpp b/Engine/Source/Programs/UnrealBuildAccelerator/Common/Private/UbaStorageClient.cpp
--- a/Engine/Source/Programs/UnrealBuildAccelerator/Common/Private/UbaStorageClient.cpp
+++ b/Engine/Source/Programs/UnrealBuildAccelerator/Common/Private/UbaStorageClient.cpp
@@ -900,7 +900,7 @@ namespace uba
 
 			compressBuffer += 8;
 			totalWritten += 8;
-			memoryBlock.Allocate(totalWritten, hint);
+			memoryBlock.Allocate(totalWritten, 1, hint);
 
 			u64 diff = u64(OodleLZ_GetCompressedBufferSizeNeeded(m_sendCasCompressor, BufferSlotHalfSize)) - BufferSlotHalfSize;
 			u64 maxUncompressedBlock = BufferSlotHalfSize - diff - totalWritten - 8; // 8 bytes block header
@@ -914,7 +914,7 @@ namespace uba
 				if (reserveSize > memoryBlock.mappedSize)
 				{
 					u64 toAllocate = reserveSize - memoryBlock.writtenSize;
-					memoryBlock.Allocate(toAllocate, hint);
+					memoryBlock.Allocate(toAllocate, 1, hint);
 				}
 
 				u8* destBuf = compressBuffer;
diff --git a/Engine/Source/Programs/UnrealBuildAccelerator/Core/Private/UbaMemory.cpp b/Engine/Source/Programs/UnrealBuildAccelerator/Core/Private/UbaMemory.cpp
--- a/Engine/Source/Programs/UnrealBuildAccelerator/Core/Private/UbaMemory.cpp
+++ b/Engine/Source/Programs/UnrealBuildAccelerator/Core/Private/UbaMemory.cpp
@@ -55,18 +55,19 @@ namespace uba
 		memory = nullptr;
 	}
 
-	void* MemoryBlock::Allocate(u64 bytes, const tchar* hint)
+	void* MemoryBlock::Allocate(u64 bytes, u64 alignment, const tchar* hint)
 	{
 		if (!memory)
-			return malloc(bytes);
+			return malloc(bytes); // TODO: alignment?
 
 		ScopedWriteLock l(lock);
-		return AllocateNoLock(bytes, hint);
+		return AllocateNoLock(bytes, alignment, hint);
 	}
 
-	void* MemoryBlock::AllocateNoLock(u64 bytes, const tchar* hint)
+	void* MemoryBlock::AllocateNoLock(u64 bytes, u64 alignment, const tchar* hint)
 	{
-		u64 newPos = writtenSize + bytes;
+		u64 startPos = AlignUp(writtenSize, alignment);
+		u64 newPos = startPos + bytes;
 		
 		#if PLATFORM_WINDOWS
 		if (newPos > mappedSize)
@@ -82,8 +83,8 @@ namespace uba
 		}
 		#endif
 
-		void* ret = memory + writtenSize;
-		writtenSize += bytes;
+		void* ret = memory + startPos;
+		writtenSize = newPos;
 		return ret;
 	}
 
@@ -97,7 +98,7 @@ namespace uba
 	{
 		u64 len = TStrlen(str);
 		u64 memSize = (len + 1) * sizeof(tchar);
-		void* mem = Allocate(memSize, TC("Strdup"));
+		void* mem = Allocate(memSize, sizeof(tchar), TC("Strdup"));
 		const void* src = str;
 		memcpy(mem, src, memSize);
 		return (tchar*)mem;
diff --git a/Engine/Source/Programs/UnrealBuildAccelerator/Core/Public/UbaMemory.h b/Engine/Source/Programs/UnrealBuildAccelerator/Core/Public/UbaMemory.h
--- a/Engine/Source/Programs/UnrealBuildAccelerator/Core/Public/UbaMemory.h
+++ b/Engine/Source/Programs/UnrealBuildAccelerator/Core/Public/UbaMemory.h
@@ -72,8 +72,8 @@ namespace uba
 		~MemoryBlock();
 		void Init(u64 reserveSize_, void* baseAddress_ = nullptr);
 		void Deinit();
-		void* Allocate(u64 bytes, const tchar* hint);
-		void* AllocateNoLock(u64 bytes, const tchar* hint);
+		void* Allocate(u64 bytes, u64 alignment, const tchar* hint);
+		void* AllocateNoLock(u64 bytes, u64 alignment, const tchar* hint);
 		void Free(void* p);
 		tchar* Strdup(const tchar* str);
 		
@@ -100,7 +100,7 @@ namespace uba
 
 		value_type* allocate(u64 n)
 		{
-			return (value_type*)m_block->Allocate(sizeof(value_type)*n, TC("GrowingAllocator"));
+			return (value_type*)m_block->Allocate(sizeof(value_type)*n, alignof(value_type), TC("GrowingAllocator"));
 		}
 
 		/// @warning Naive implementation, assumes `p` is valid.
@@ -141,7 +141,7 @@ namespace uba
 				m_nextFree = *(u64*)m_nextFree;
 				return ptr;
 			}
-			void* mem = m_memory.Allocate(sizeof(Type), TC("BlockAllocator"));
+			void* mem = m_memory.Allocate(sizeof(Type), alignof(Type), TC("BlockAllocator"));
 			return mem;
 		}
 
diff --git a/Engine/Source/Programs/UnrealBuildAccelerator/Detours/Private/Windows/UbaDetoursFunctionsWin.cpp b/Engine/Source/Programs/UnrealBuildAccelerator/Detours/Private/Windows/UbaDetoursFunctionsWin.cpp
--- a/Engine/Source/Programs/UnrealBuildAccelerator/Detours/Private/Windows/UbaDetoursFunctionsWin.cpp
+++ b/Engine/Source/Programs/UnrealBuildAccelerator/Detours/Private/Windows/UbaDetoursFunctionsWin.cpp
@@ -815,7 +815,7 @@ void Init(const DetoursPayload& payload, u64 startTime)
 	}
 
 	if (payload.trackInputs)
-		g_trackInputsMem = (u8*)g_memoryBlock.Allocate(TrackInputsMemCapacity, L"TrackInputs");
+		g_trackInputsMem = (u8*)g_memoryBlock.Allocate(TrackInputsMemCapacity, 1, L"TrackInputs");
 	
 	g_systemRoot.count = GetEnvironmentVariableW(L"SystemRoot", g_systemRoot.data, g_systemRoot.capacity);
 	g_systemRoot.MakeLower();
