[UBA]
* Fixed uba unit test compile error

diff --git a/Engine/Source/Programs/UnrealBuildAccelerator/Test/Private/UbaTestBasics.h b/Engine/Source/Programs/UnrealBuildAccelerator/Test/Private/UbaTestBasics.h
--- a/Engine/Source/Programs/UnrealBuildAccelerator/Test/Private/UbaTestBasics.h
+++ b/Engine/Source/Programs/UnrealBuildAccelerator/Test/Private/UbaTestBasics.h
@@ -181,7 +181,7 @@ namespace uba
 	bool TestMemoryBlock(Logger& logger, const StringBufferBase& rootDir)
 	{
 		MemoryBlock block(1024 * 1024);
-		u64* mem = (u64*)block.Allocate(8, TC("Foo"));
+		u64* mem = (u64*)block.Allocate(8, 1, TC("Foo"));
 		*mem = 0x1234;
 		block.Free(mem);
 		return true;
