[Uba]
* Changed error message for GetLastError to always be english since rest of text is english

diff --git a/Engine/Source/Programs/UnrealBuildAccelerator/Common/Private/UbaLogger.cpp b/Engine/Source/Programs/UnrealBuildAccelerator/Common/Private/UbaLogger.cpp
--- a/Engine/Source/Programs/UnrealBuildAccelerator/Common/Private/UbaLogger.cpp
+++ b/Engine/Source/Programs/UnrealBuildAccelerator/Common/Private/UbaLogger.cpp
@@ -276,8 +276,7 @@ namespace uba
 	LastErrorToText::LastErrorToText(u32 lastError)
 	{
 #if PLATFORM_WINDOWS
-		size_t size = ::FormatMessageW(FORMAT_MESSAGE_FROM_SYSTEM | FORMAT_MESSAGE_IGNORE_INSERTS,
-			NULL, lastError, MAKELANGID(LANG_NEUTRAL, SUBLANG_DEFAULT), data, capacity, NULL);
+		size_t size = ::FormatMessageW(FORMAT_MESSAGE_FROM_SYSTEM | FORMAT_MESSAGE_IGNORE_INSERTS, NULL, lastError, MAKELANGID(LANG_ENGLISH, SUBLANG_ENGLISH_US), data, capacity, NULL);
 		if (!size)
 			AppendValue(lastError);
 		else
