[Uba]
* Brought back code where parent waits for children's m_hasExited because it can actually happen!! (Need to investigate)

diff --git a/Engine/Source/Programs/UnrealBuildAccelerator/Common/Private/UbaProcess.cpp b/Engine/Source/Programs/UnrealBuildAccelerator/Common/Private/UbaProcess.cpp
--- a/Engine/Source/Programs/UnrealBuildAccelerator/Common/Private/UbaProcess.cpp
+++ b/Engine/Source/Programs/UnrealBuildAccelerator/Common/Private/UbaProcess.cpp
@@ -414,20 +414,16 @@ namespace uba
 
 		m_session.ProcessExited(*this, m_processStats.wallTime);
 
-		#if PLATFORM_LINUX
-		// This should not really ever happen.. but just in case.. since children use memory from parent
-		for (auto& child : m_childProcesses)
-		{
-			while (!((ProcessImpl*)child.m_process)->m_hasExited)
-			{
-				Sleep(100);
-			}
-		}
-		#elif PLATFORM_MAC
+		#if PLATFORM_MAC
 		int res = WaitForProcessGroup(getpgrp());
 		UBA_ASSERT(res == 0);
 		#endif
 
+		// For some reason a parent can exit before a child. Need to figure out repro for this but I've seen it happen on ClangEditor win64
+		for (auto& child : m_childProcesses)
+			while (!((ProcessImpl*)child.m_process)->m_hasExited)
+				Sleep(10);
+
 		UBA_ASSERT(!m_parentProcess || !m_parentProcess->m_hasExited);
 
 		m_hasExited = true;
