[Uba]
* Fixed bug caused by populatecas changelist causing StoreCas to return true without populating out casKey

diff --git a/Engine/Source/Programs/UnrealBuildAccelerator/Common/Private/UbaStorageClient.cpp b/Engine/Source/Programs/UnrealBuildAccelerator/Common/Private/UbaStorageClient.cpp
--- a/Engine/Source/Programs/UnrealBuildAccelerator/Common/Private/UbaStorageClient.cpp
+++ b/Engine/Source/Programs/UnrealBuildAccelerator/Common/Private/UbaStorageClient.cpp
@@ -691,30 +691,30 @@ namespace uba
 				ScopedWriteLock lock(m_localStorageFilesLock);
 				auto insres = m_localStorageFiles.try_emplace(AsCompressed(casKey, false));
 				LocalFile& localFile = insres.first->second;
-				if (!insres.second && localFile.casEntry.mappingHandle.IsValid())
-					return true;
-
-				if (isPersistentMapping)
+				if (insres.second || !localFile.casEntry.mappingHandle.IsValid())
 				{
-					FileMappingHandle mappingHandle2;
-					if (DuplicateFileMapping(GetCurrentProcessHandle(), mappingHandle, GetCurrentProcessHandle(), &mappingHandle2, FILE_MAP_READ, false, 0))
+					if (isPersistentMapping)
 					{
-						localFile.casEntry.mappingHandle = mappingHandle2;
-						localFile.casEntry.size = fileSize;
+						FileMappingHandle mappingHandle2;
+						if (DuplicateFileMapping(GetCurrentProcessHandle(), mappingHandle, GetCurrentProcessHandle(), &mappingHandle2, FILE_MAP_READ, false, 0))
+						{
+							localFile.casEntry.mappingHandle = mappingHandle2;
+							localFile.casEntry.size = fileSize;
+						}
+						else
+							m_logger.Warning(TC("Failed to duplicate handle for file mapping %s (%s)"), fileName, LastErrorToText().data);
 					}
 					else
-						m_logger.Warning(TC("Failed to duplicate handle for file mapping %s (%s)"), fileName, LastErrorToText().data);
-				}
-				else
-				{
-					localFile.casEntry.size = fileSize;
+					{
+						localFile.casEntry.size = fileSize;
 #if !UBA_USE_SPARSEFILE
-					localFile.fileName = fileName;
+						localFile.fileName = fileName;
 #else
-					localFile.casEntry.mappingHandle = mappingHandle2;
-					lock.Leave();
-					mappingClose.Cancel();
+						localFile.casEntry.mappingHandle = mappingHandle2;
+						lock.Leave();
+						mappingClose.Cancel();
 #endif
+					}
 				}
 			}
 		}
@@ -880,6 +880,8 @@ namespace uba
 
 	bool StorageClient::SendFile(const CasKey& casKey, const tchar* fileName, u8* sourceMem, u64 sourceSize, const tchar* hint)
 	{
+		UBA_ASSERT(casKey != CasKeyZero);
+
 		NetworkClient& client = m_client; // Don't use proxy
 
 		StorageStats& stats = Stats();
diff --git a/Engine/Source/Programs/UnrealBuildAccelerator/Common/Private/UbaStorageServer.cpp b/Engine/Source/Programs/UnrealBuildAccelerator/Common/Private/UbaStorageServer.cpp
--- a/Engine/Source/Programs/UnrealBuildAccelerator/Common/Private/UbaStorageServer.cpp
+++ b/Engine/Source/Programs/UnrealBuildAccelerator/Common/Private/UbaStorageServer.cpp
@@ -64,6 +64,7 @@ namespace uba
 	bool StorageServer::RetrieveCasFile(RetrieveResult& out, const CasKey& casKey, const tchar* hint, FileMappingBuffer* mappingBuffer, u64 memoryMapAlignment, bool allowProxy)
 	{
 		UBA_ASSERT(!mappingBuffer);
+		UBA_ASSERT(casKey != CasKeyZero);
 		out.casKey = casKey;
 		out.size = InvalidValue;
 
