[Uba]
* Fixed race condition where NetorkClient can be disconnected at the exact same time a full proper message is received by client. If receive code manage to get passed check for m_error when disconnect happens but the disconnect code sets message size to 0 before setting binary reader size we get a situation where Send returns success but size is 0. Outside code might read some data and then use GetLeft() which will now be a wrapped number (negative position) and all kinds of weird things will happen

diff --git a/Engine/Source/Programs/UnrealBuildAccelerator/Common/Private/UbaNetworkClient.cpp b/Engine/Source/Programs/UnrealBuildAccelerator/Common/Private/UbaNetworkClient.cpp
--- a/Engine/Source/Programs/UnrealBuildAccelerator/Common/Private/UbaNetworkClient.cpp
+++ b/Engine/Source/Programs/UnrealBuildAccelerator/Common/Private/UbaNetworkClient.cpp
@@ -376,8 +376,8 @@ namespace uba
 		{
 			if (m && m->m_connection == &connection)
 			{
-				m->m_responseSize = 0;
 				m->m_error = true;
+				//m->m_responseSize = 0; // There is a race here where a message could just have fully arrived when disconnected. Can't set this to zero because we might be passed the m_error check but before setting reader size
 				m->m_event.Set();
 			}
 			++messageId;
