[Uba]
* Fixed so TMPDIR is local-only
* Fixed so HOME envvar is used when resolving paths with ~

diff --git a/Engine/Source/Programs/UnrealBuildAccelerator/Common/Private/UbaSessionClient.cpp b/Engine/Source/Programs/UnrealBuildAccelerator/Common/Private/UbaSessionClient.cpp
--- a/Engine/Source/Programs/UnrealBuildAccelerator/Common/Private/UbaSessionClient.cpp
+++ b/Engine/Source/Programs/UnrealBuildAccelerator/Common/Private/UbaSessionClient.cpp
@@ -1009,6 +1009,8 @@ namespace uba
 		AddEnvironmentVariableNoLock(TC("Path"), TC("c:\\noenvironment"));
 		AddEnvironmentVariableNoLock(TC("TEMP"), m_tempPath.data);
 		AddEnvironmentVariableNoLock(TC("TMP"), m_tempPath.data);
+		#else
+		AddEnvironmentVariableNoLock(TC("TMPDIR"), m_tempPath.data);
 		#endif
 
 		StringBuffer<> v;
diff --git a/Engine/Source/Programs/UnrealBuildAccelerator/Common/Private/UbaSessionServer.cpp b/Engine/Source/Programs/UnrealBuildAccelerator/Common/Private/UbaSessionServer.cpp
--- a/Engine/Source/Programs/UnrealBuildAccelerator/Common/Private/UbaSessionServer.cpp
+++ b/Engine/Source/Programs/UnrealBuildAccelerator/Common/Private/UbaSessionServer.cpp
@@ -173,6 +173,8 @@ namespace uba
 		m_localEnvironmentVariables.insert(TC("Path"));
 		m_localEnvironmentVariables.insert(TC("TMP"));
 		m_localEnvironmentVariables.insert(TC("TEMP"));
+		#else
+		m_localEnvironmentVariables.insert(TC("TMPDIR"));
 		#endif
 
 		StringBuffer<> detoursFile;
diff --git a/Engine/Source/Programs/UnrealBuildAccelerator/Core/Public/UbaPathUtils.h b/Engine/Source/Programs/UnrealBuildAccelerator/Core/Public/UbaPathUtils.h
--- a/Engine/Source/Programs/UnrealBuildAccelerator/Core/Public/UbaPathUtils.h
+++ b/Engine/Source/Programs/UnrealBuildAccelerator/Core/Public/UbaPathUtils.h
@@ -201,19 +201,13 @@ namespace uba
 			*bufferCharLen = u32(charLen - 1); // Remove terminator
 		return true;
 #else
-		char tmp[MaxPath] = {0};
+		StringBuffer<MaxPath> tmp;
 		if (fileName[0] == '~')
 		{
-			if (realpath(fileName, tmp) == 0)
-			{
-				// If we can't  get a real path to whatever
-				// fileName we're looking for, then we need to
-				// assume it's not a real file, errno should be set
-				// so just return false
-				return false; 
-			}
+			const char* homeDir = getenv("HOME");
+			tmp.Append(homeDir).EnsureEndsWithSlash().Append(fileName + 1);
+			fileName = tmp.data;
 		}
-		UBA_ASSERT(fileName[0] != '~');
 
 		u64 memPos = 0;
 		if (fileName[0] != '/')
diff --git a/Engine/Source/Programs/UnrealBuildAccelerator/Detours/Private/PosixOS/UbaDetoursFunctionsPosix.cpp b/Engine/Source/Programs/UnrealBuildAccelerator/Detours/Private/PosixOS/UbaDetoursFunctionsPosix.cpp
--- a/Engine/Source/Programs/UnrealBuildAccelerator/Detours/Private/PosixOS/UbaDetoursFunctionsPosix.cpp
+++ b/Engine/Source/Programs/UnrealBuildAccelerator/Detours/Private/PosixOS/UbaDetoursFunctionsPosix.cpp
@@ -149,48 +149,12 @@ using namespace uba;
 
 #endif
 
-// Look into these
-/*
-35054:     binding file / lib / x86_64 - linux - gnu / libbfd - 2.38 - system.so[0] to / lib / x86_64 - linux - gnu / libc.so.6[0] : normal symbol `lstat'[GLIBC_2.33]
-35054:     binding file / lib / x86_64 - linux - gnu / libbfd - 2.38 - system.so[0] to / lib / x86_64 - linux - gnu / libc.so.6[0] : normal symbol `unlink'[GLIBC_2.2.5]
-35054:     binding file / lib / x86_64 - linux - gnu / libbfd - 2.38 - system.so[0] to / lib / x86_64 - linux - gnu / libc.so.6[0] : normal symbol `fopen64'[GLIBC_2.2.5]
-35054:     binding file / lib / x86_64 - linux - gnu / libbfd - 2.38 - system.so[0] to / lib / x86_64 - linux - gnu / libc.so.6[0] : normal symbol `fileno'[GLIBC_2.2.5]
-35054:     binding file / lib / x86_64 - linux - gnu / libbfd - 2.38 - system.so[0] to / lib / x86_64 - linux - gnu / libc.so.6[0] : normal symbol `fcntl'[GLIBC_2.2.5]
-35054:     binding file /usr/bin/ld [0] to /lib/x86_64-linux-gnu/libbfd-2.38-system.so [0]: normal symbol `filename_cmp'
-35054:     binding file /usr/bin/ld [0] to /lib/x86_64-linux-gnu/libbfd-2.38-system.so [0]: normal symbol `lrealpath'
-35054:     binding file /usr/bin/ld [0] to /lib/x86_64-linux-gnu/libc.so.6 [0]: normal symbol `fclose' [GLIBC_2.2.5]
-35054:     binding file /lib/x86_64-linux-gnu/libbfd-2.38-system.so [0] to /lib/x86_64-linux-gnu/libc.so.6 [0]: normal symbol `fnmatch' [GLIBC_2.2.5]
-35054:     binding file /lib/x86_64-linux-gnu/libbfd-2.38-system.so [0] to /lib/x86_64-linux-gnu/libc.so.6 [0]: normal symbol `__realpath_chk' [GLIBC_2.4]
-35054:     binding file /lib/x86_64-linux-gnu/libbfd-2.38-system.so [0] to /lib/x86_64-linux-gnu/libc.so.6 [0]: normal symbol `sbrk' [GLIBC_2.2.5]
-35054:     binding file /usr/bin/ld [0] to /lib/x86_64-linux-gnu/libc.so.6 [0]: normal symbol `glob' [GLIBC_2.27]
-35054:     binding file /usr/bin/ld [0] to /lib/x86_64-linux-gnu/libc.so.6 [0]: normal symbol `exit' [GLIBC_2.2.5]
-35054:     binding file /usr/bin/ld [0] to /lib/x86_64-linux-gnu/libc.so.6 [0]: normal symbol `fwrite' [GLIBC_2.2.5]
-35054:     binding file /usr/bin/ld [0] to /lib/x86_64-linux-gnu/libc.so.6 [0]: normal symbol `execv' [GLIBC_2.2.5]
-35054:     binding file /usr/bin/ld [0] to /lib/x86_64-linux-gnu/libc.so.6 [0]: normal symbol `access' [GLIBC_2.2.5]
-35054:     binding file /usr/bin/ld [0] to /lib/x86_64-linux-gnu/libc.so.6 [0]: normal symbol `wait4' [GLIBC_2.2.5]
-35054:     binding file /usr/bin/ld [0] to /lib/x86_64-linux-gnu/libc.so.6 [0]: normal symbol `dup2' [GLIBC_2.2.5]
-35054:     binding file /usr/bin/ld [0] to /lib/x86_64-linux-gnu/libc.so.6 [0]: normal symbol `abort' [GLIBC_2.2.5]
-35054:     binding file /usr/bin/ld [0] to /lib/x86_64-linux-gnu/libc.so.6 [0]: normal symbol `getenv' [GLIBC_2.2.5]
-35054:     binding file /usr/bin/ld [0] to /lib/x86_64-linux-gnu/libc.so.6 [0]: normal symbol `mkstemps' [GLIBC_2.11]
-35054:     binding file /usr/bin/ld [0] to /lib/x86_64-linux-gnu/libc.so.6 [0]: normal symbol `environ' [GLIBC_2.2.5]
-35054:     binding file /usr/bin/ld [0] to /lib/x86_64-linux-gnu/libc.so.6 [0]: normal symbol `optarg' [GLIBC_2.2.5]
-35054:     binding file /usr/bin/ld [0] to /lib/x86_64-linux-gnu/libc.so.6 [0]: normal symbol `optind' [GLIBC_2.2.5]
-35054:     binding file /usr/bin/ld [0] to /lib/x86_64-linux-gnu/libc.so.6 [0]: normal symbol `opterr' [GLIBC_2.2.5]
-35054:     binding file /usr/bin/ld [0] to /lib/x86_64-linux-gnu/libc.so.6 [0]: normal symbol `globfree' [GLIBC_2.2.5]
-35054:     binding file /usr/bin/ld [0] to /lib/x86_64-linux-gnu/libc.so.6 [0]: normal symbol `getopt_long_only' [GLIBC_2.2.5]
-35054:     binding file /usr/bin/ld [0] to /lib/x86_64-linux-gnu/libc.so.6 [0]: normal symbol `strpbrk' [GLIBC_2.2.5]
-35054:     binding file /usr/bin/ld [0] to /lib/x86_64-linux-gnu/libc.so.6 [0]: normal symbol `dlsym' [GLIBC_2.34]
-35054:     binding file /usr/bin/ld [0] to /lib/x86_64-linux-gnu/libc.so.6 [0]: normal symbol `execvp' [GLIBC_2.2.5]
-35054:     binding file /usr/bin/ld [0] to /lib/x86_64-linux-gnu/libc.so.6 [0]: normal symbol `pipe2' [GLIBC_2.9]
-35054:     binding file /usr/bin/ld [0] to /lib/x86_64-linux-gnu/libc.so.6 [0]: normal symbol `glob' [GLIBC_2.27]
-35054:     binding file /usr/bin/ld [0] to /lib/x86_64-linux-gnu/libc.so.6 [0]: normal symbol `exit' [GLIBC_2.2.5]
-*/
 #if 0
 #define	EPERM		 1	/* Operation not permitted */
 #define	ENOENT		 2	/* No such file or directory */
 #define	ESRCH		 3	/* No such process */
 #define	EINTR		 4	/* Interrupted system call */
-#define	EIO		 5	/* I/O error */
+#define	EIO			 5	/* I/O error */
 #define	ENXIO		 6	/* No such device or address */
 #define	E2BIG		 7	/* Argument list too long */
 #define	ENOEXEC		 8	/* Exec format error */
@@ -617,11 +581,8 @@ int Shared_fstat(const char* funcName, int fd, struct stat* attr, const True_fst
 template<typename True_stat>
 int Shared_stat(const char* funcName, const char* file, struct stat* attr, const True_stat& trueStat)
 {
-	if (!g_isDetouring || t_disallowDetour || StartsWith(file, "/etc/"))
-		return trueStat(file, attr);
-
 	StringBuffer<> fixedFile;
-	if ((!FixPath(fixedFile, file)) || (fixedFile.Equals("/")) || (access(file, F_OK) != 0))
+	if (!FixPath(fixedFile, file) || fixedFile.Equals("/") || fixedFile.StartsWith("/etc/") || fixedFile.StartsWith(g_systemTemp.data))//(access(file, F_OK) != 0))
 	{
 		int res =  trueStat(file, attr);
 		return res;
@@ -1098,6 +1059,8 @@ UBA_EXPORT int UBA_WRAPPER(fstatat64)(int dirfd, const char* pathname, struct st
 
 UBA_EXPORT int UBA_WRAPPER(lstat)(const char *path, struct stat *buf)
 {
+	UBA_INIT_DETOUR(lstat, path, buf);
+	DEBUG_LOG_TRUE("lstat", "");
 	return TRUE_WRAPPER(lstat)(path, buf);
 }
 
@@ -1116,11 +1079,10 @@ UBA_EXPORT int UBA_WRAPPER(truncate)(const char* path, off_t length)
 
 UBA_EXPORT int UBA_WRAPPER(access)(const char* pathname, int mode)
 {
-	if (strlen(pathname) == 0) UBA_ASSERT(false);
 	UBA_INIT_DETOUR(access, pathname, mode);
 
 	StringBuffer<> fixedPath;
-	if (!FixPath(fixedPath, pathname))
+	if (!FixPath(fixedPath, pathname) || fixedPath.StartsWith(g_systemTemp.data))
 	{
 		return TRUE_WRAPPER(access)(pathname, mode);
 	}
