[UBA]
* Fixed so libDetours.so can be in a path containing spaces (by adding quotes around path)

diff --git a/Engine/Source/Programs/UnrealBuildAccelerator/Common/Private/UbaProcess.cpp b/Engine/Source/Programs/UnrealBuildAccelerator/Common/Private/UbaProcess.cpp
--- a/Engine/Source/Programs/UnrealBuildAccelerator/Common/Private/UbaProcess.cpp
+++ b/Engine/Source/Programs/UnrealBuildAccelerator/Common/Private/UbaProcess.cpp
@@ -667,15 +667,6 @@ namespace uba
 					const char* detoursLib = m_session.m_detoursLibrary.c_str();
 					u32 detoursLibLen = u32(m_session.m_detoursLibrary.size());
 
-					#if !PLATFORM_WINDOWS
-					if (!*detoursLib)
-					{
-						detoursLib = "/mnt/e/dev/fn/Engine/Binaries/Linux/UnrealBuildAccelerator/" UBA_DETOURS_LIBRARY;
-						detoursLibLen = u32(strlen(detoursLib));
-					}
-					#endif
-
-
 					writer.WriteU32(childProcessId);
 					writer.WriteU32(process.m_rulesIndex);
 					writer.WriteU32(detoursLibLen);
@@ -1261,22 +1252,6 @@ namespace uba
 
 		if (!m_parentProcess)
 		{
-
-			StringBuffer<> ldPreload;
-			#if PLATFORM_LINUX
-				ldPreload.Append("LD_PRELOAD=");
-			#else
-				ldPreload.Append("DYLD_INSERT_LIBRARIES=");
-			#endif
-			
-			//if (m_session.m_detoursLibrary[0] != '/')
-			//	ldPreload.Append("./");
-			const char* detoursLib = m_session.m_detoursLibrary.c_str();
-			if (*detoursLib)
-				ldPreload.Append(detoursLib);
-			else
-				ldPreload.Append("./" UBA_DETOURS_LIBRARY);
-
 			StringBuffer<128> comIdVar;
 			comIdVar.Append("UBA_COMID=").AppendValue(communicationHandle.uid).Append('+').AppendValue(communicationOffset);
 
@@ -1306,7 +1281,6 @@ namespace uba
 				it += TStrlen(s) + 1;
 			}
 
-			envvars.push_back(ldPreload.data);
 			envvars.push_back(comIdVar.data);
 			envvars.push_back(workingDir.data);
 			envvars.push_back(rulesStr.data);
diff --git a/Engine/Source/Programs/UnrealBuildAccelerator/Common/Private/UbaSession.cpp b/Engine/Source/Programs/UnrealBuildAccelerator/Common/Private/UbaSession.cpp
--- a/Engine/Source/Programs/UnrealBuildAccelerator/Common/Private/UbaSession.cpp
+++ b/Engine/Source/Programs/UnrealBuildAccelerator/Common/Private/UbaSession.cpp
@@ -179,6 +179,35 @@ namespace uba
 		m_environmentVariables.push_back(0);
 	}
 
+	void Session::AddDetoursEnvironmentVariable()
+	{
+#if !PLATFORM_WINDOWS
+
+		const char* detoursLib = m_detoursLibrary.c_str();
+		if (*detoursLib)
+		{
+			#if PLATFORM_LINUX
+			//if (strchr(detoursLib, ' '))
+			{
+				const char* lastSlash = strrchr(detoursLib, '/');
+				UBA_ASSERT(lastSlash);
+				StringBuffer<> ldLibPath;
+				ldLibPath.Append(detoursLib, lastSlash - detoursLib);
+				AddEnvironmentVariableNoLock("LD_LIBRARY_PATH", ldLibPath.data);
+				detoursLib = lastSlash + 1;
+			}
+			#endif
+		}
+		else
+			detoursLib = "./" UBA_DETOURS_LIBRARY;
+
+		#if PLATFORM_LINUX
+		AddEnvironmentVariableNoLock("LD_PRELOAD", detoursLib);
+		#else
+		AddEnvironmentVariableNoLock("DYLD_INSERT_LIBRARIES", detoursLib);
+		#endif
+#endif
+	}
 
 	bool Session::WriteDirectoryEntriesInternal(DirectoryTable::Directory& dir, const StringKey& dirKey, const tchar* dirPath, bool isRefresh, u32& outTableOffset)
 	{
@@ -1609,6 +1638,7 @@ namespace uba
 			AddEnvironmentVariableNoLock("PATH", paths.c_str());
 		}
 		AddEnvironmentVariableNoLock("TMPDIR", m_tempPath.data);
+		AddDetoursEnvironmentVariable();
 #endif
 		m_environmentVariables.push_back(0);
 		return m_environmentVariables.data();
diff --git a/Engine/Source/Programs/UnrealBuildAccelerator/Common/Private/UbaSessionClient.cpp b/Engine/Source/Programs/UnrealBuildAccelerator/Common/Private/UbaSessionClient.cpp
--- a/Engine/Source/Programs/UnrealBuildAccelerator/Common/Private/UbaSessionClient.cpp
+++ b/Engine/Source/Programs/UnrealBuildAccelerator/Common/Private/UbaSessionClient.cpp
@@ -1011,6 +1011,7 @@ namespace uba
 		AddEnvironmentVariableNoLock(TC("TMP"), m_tempPath.data);
 		#else
 		AddEnvironmentVariableNoLock(TC("TMPDIR"), m_tempPath.data);
+		AddDetoursEnvironmentVariable();
 		#endif
 
 		StringBuffer<> v;
diff --git a/Engine/Source/Programs/UnrealBuildAccelerator/Common/Private/UbaSessionServer.cpp b/Engine/Source/Programs/UnrealBuildAccelerator/Common/Private/UbaSessionServer.cpp
--- a/Engine/Source/Programs/UnrealBuildAccelerator/Common/Private/UbaSessionServer.cpp
+++ b/Engine/Source/Programs/UnrealBuildAccelerator/Common/Private/UbaSessionServer.cpp
@@ -175,6 +175,9 @@ namespace uba
 		m_localEnvironmentVariables.insert(TC("TEMP"));
 		#else
 		m_localEnvironmentVariables.insert(TC("TMPDIR"));
+		m_localEnvironmentVariables.insert(TC("LD_PRELOAD"));
+		m_localEnvironmentVariables.insert(TC("LD_LIBRARY_PATH"));
+		m_localEnvironmentVariables.insert(TC("DYLD_INSERT_LIBRARIES"));
 		#endif
 
 		StringBuffer<> detoursFile;
diff --git a/Engine/Source/Programs/UnrealBuildAccelerator/Common/Public/UbaSession.h b/Engine/Source/Programs/UnrealBuildAccelerator/Common/Public/UbaSession.h
--- a/Engine/Source/Programs/UnrealBuildAccelerator/Common/Public/UbaSession.h
+++ b/Engine/Source/Programs/UnrealBuildAccelerator/Common/Public/UbaSession.h
@@ -164,6 +164,7 @@ namespace uba
 		static constexpr CasKey CasKeyIsDirectory = { ~u64(0), ~u64(0), ~u32(0) };
 
 		void AddEnvironmentVariableNoLock(const tchar* key, const tchar* value);
+		void AddDetoursEnvironmentVariable();
 		bool WriteDirectoryEntriesInternal(DirectoryTable::Directory& dir, const StringKey& dirKey, const tchar* dirPath, bool isRefresh, u32& outTableOffset);
 		void WriteDirectoryEntriesRecursive(const StringKey& dirKey, tchar* dirPath, u32& outTableOffset);
 		bool CopyImports(Vector<BinaryModule>& out, const tchar* library, tchar* applicationDir, tchar* applicationDirEnd, UnorderedSet<TString>& handledImports);
diff --git a/Engine/Source/Programs/UnrealBuildAccelerator/Detours/Private/PosixOS/UbaDetoursFunctionsPosix.cpp b/Engine/Source/Programs/UnrealBuildAccelerator/Detours/Private/PosixOS/UbaDetoursFunctionsPosix.cpp
--- a/Engine/Source/Programs/UnrealBuildAccelerator/Detours/Private/PosixOS/UbaDetoursFunctionsPosix.cpp
+++ b/Engine/Source/Programs/UnrealBuildAccelerator/Detours/Private/PosixOS/UbaDetoursFunctionsPosix.cpp
@@ -1391,7 +1391,6 @@ UBA_EXPORT int UBA_WRAPPER(posix_spawn)(pid_t* pid, const char* path, const posi
 	u32 processId = 0;
 	StringBuffer<512> currentDir;
 	StringBuffer<256> comIdVar;
-	StringBuffer<256> ldpreload;
 	StringBuffer<32> rulesStr;
 	StringBuffer<256> logFile;
 
@@ -1424,13 +1423,7 @@ UBA_EXPORT int UBA_WRAPPER(posix_spawn)(pid_t* pid, const char* path, const posi
 		rulesStr.Append("UBA_RULES=").AppendValue(reader.ReadU32());
 
 		u32 dllNameSize = reader.ReadU32();
-	#if PLATFORM_LINUX
-		ldpreload.Append("LD_PRELOAD=");
-	#else
-		ldpreload.Append("DYLD_INSERT_LIBRARIES=");
-	#endif
-		reader.ReadBytes(ldpreload.data + ldpreload.count, dllNameSize);
-		ldpreload.Resize(ldpreload.count + dllNameSize);
+		reader.Skip(dllNameSize);
 
 		commandLine = reader.ReadString();
 
@@ -1447,7 +1440,6 @@ UBA_EXPORT int UBA_WRAPPER(posix_spawn)(pid_t* pid, const char* path, const posi
 	{
 		for (u32 i = 0; envp[i]; ++i)
 			envvars.push_back(envp[i]);
-		envvars.push_back(ldpreload.data);
 		envvars.push_back(comIdVar.data);
 		envvars.push_back(currentDir.data);
 		envvars.push_back(rulesStr.data);
diff --git a/Engine/Source/Programs/UnrealBuildAccelerator/Detours/Private/PosixOS/UbaDetoursMainPosix.cpp b/Engine/Source/Programs/UnrealBuildAccelerator/Detours/Private/PosixOS/UbaDetoursMainPosix.cpp
--- a/Engine/Source/Programs/UnrealBuildAccelerator/Detours/Private/PosixOS/UbaDetoursMainPosix.cpp
+++ b/Engine/Source/Programs/UnrealBuildAccelerator/Detours/Private/PosixOS/UbaDetoursMainPosix.cpp
@@ -68,11 +68,6 @@ static void __attribute__((constructor (65535))) InitCtor()
 	using namespace uba;
 
 	SuppressDetourScope s;
-	#if PLATFORM_LINUX
-		unsetenv("LD_PRELOAD");
-	#else
-		unsetenv("DYLD_INSERT_LIBRARIES=");
-	#endif
 
 	const char* comId = getenv("UBA_COMID");
 	//printf("Starting up %s... (comid: %s) %u\n", __progname, (comId && *comId) ? comId : "NOTSET", getpid());
