[Uba]
* Fixed race condition where process could be returned and then picked up before the callback in process return message is invoked.

diff --git a/Engine/Source/Programs/UnrealBuildAccelerator/Common/Private/UbaSessionServer.cpp b/Engine/Source/Programs/UnrealBuildAccelerator/Common/Private/UbaSessionServer.cpp
--- a/Engine/Source/Programs/UnrealBuildAccelerator/Common/Private/UbaSessionServer.cpp
+++ b/Engine/Source/Programs/UnrealBuildAccelerator/Common/Private/UbaSessionServer.cpp
@@ -1023,8 +1023,6 @@ namespace uba
 				m_trace.ProcessReturned(process->m_processId);
 				m_queuedRemoteProcesses.push_front(h);
 
-				cs2.Leave();
-
 				if (m_remoteProcessReturnedEvent)
 					m_remoteProcessReturnedEvent(*process);
 				return true;
