UnrealBuildAccelerator: Fix mac lib name

#rnx

diff --git a/Engine/Source/Programs/Shared/EpicGames.UBA/Utils.cs b/Engine/Source/Programs/Shared/EpicGames.UBA/Utils.cs
--- a/Engine/Source/Programs/Shared/EpicGames.UBA/Utils.cs
+++ b/Engine/Source/Programs/Shared/EpicGames.UBA/Utils.cs
@@ -33,7 +33,7 @@ namespace EpicGames.UBA
 			}
 			else if (OperatingSystem.IsMacOS())
 			{
-				return Path.Combine(assemblyFolder, "runtimes", $"osx", "native", "libUbaHost.dylib");
+				return Path.Combine(assemblyFolder, "runtimes", $"osx-{arch}", "native", "libUbaHost.dylib");
 			}
 			throw new PlatformNotSupportedException();
 		}
diff --git a/Engine/Source/Programs/UnrealBuildAccelerator/UbaAgent.Target.cs b/Engine/Source/Programs/UnrealBuildAccelerator/UbaAgent.Target.cs
--- a/Engine/Source/Programs/UnrealBuildAccelerator/UbaAgent.Target.cs
+++ b/Engine/Source/Programs/UnrealBuildAccelerator/UbaAgent.Target.cs
@@ -83,10 +83,11 @@ public class UbaAgentTarget : TargetRules
 
 		if (Target.Platform.IsInGroup(UnrealPlatformGroup.Apple))
 		{
+			string BinaryPrefix = Rules.bShouldCompileAsDLL ? "lib" : string.Empty;
 			string BinaryExt = Rules.bShouldCompileAsDLL ? ".dylib" : string.Empty;
-			Rules.OutputFile = Path.Combine(BinariesFolder, $"{Rules.LaunchModuleName}{BinaryExt}");
+			Rules.OutputFile = Path.Combine(BinariesFolder, $"{BinaryPrefix}{Rules.LaunchModuleName}{BinaryExt}");
 			Rules.GlobalDefinitions.AddRange(new string[] {
-				"UBA_DETOURS_LIBRARY=\"UbaDetours.dylib\"",
+				"UBA_DETOURS_LIBRARY=\"libUbaDetours.dylib\"",
 				"UBA_AGENT_EXECUTABLE=\"UbaAgent\"",
 			});
 		}
