[UBA]
* Added .isNotAWS file that is written if uba notices it is not an aws machine. This file is then used to skip http request all together. This helps startup times in some network environments where a connect on aws meta data ip can cause a long stall
* Cleaned up AWS logic a bit
* Added so UbaCli -workdir=x can handle relative paths

diff --git a/Engine/Source/Programs/UnrealBuildAccelerator/Agent/Private/UbaAgent.cpp b/Engine/Source/Programs/UnrealBuildAccelerator/Agent/Private/UbaAgent.cpp
--- a/Engine/Source/Programs/UnrealBuildAccelerator/Agent/Private/UbaAgent.cpp
+++ b/Engine/Source/Programs/UnrealBuildAccelerator/Agent/Private/UbaAgent.cpp
@@ -686,7 +686,7 @@ namespace uba
 		// Check if AWS
 		#if defined(UBA_USE_AWS)
 		AWS aws;
-		aws.Init(logger, extraInfo, TC("Agent"));
+		aws.QueryInformation(logger, extraInfo, g_rootDir.data);
 		if (zone.IsEmpty())
 			zone.Append(aws.GetAvailabilityZone());
 		#endif
@@ -1029,10 +1029,6 @@ namespace uba
 
 			u64 lastLogTime = GetTime();
 
-			#if defined(UBA_USE_AWS)
-			aws.InitPolling(logger);
-			#endif
-
 			u32 tcpConnectionCount = 1;
 
 			while (true)
diff --git a/Engine/Source/Programs/UnrealBuildAccelerator/Cli/Private/UbaCli.cpp b/Engine/Source/Programs/UnrealBuildAccelerator/Cli/Private/UbaCli.cpp
--- a/Engine/Source/Programs/UnrealBuildAccelerator/Cli/Private/UbaCli.cpp
+++ b/Engine/Source/Programs/UnrealBuildAccelerator/Cli/Private/UbaCli.cpp
@@ -201,7 +201,8 @@ namespace uba
 			{
 				if (value.IsEmpty())
 					return PrintHelp(TC("-workdir needs a value"));
-				workDir.Append(value);
+				if ((workDir.count = GetFullPathNameW(value.data, workDir.capacity, workDir.data, nullptr)) == 0)
+					return PrintHelp(StringBuffer<>().Appendf(TC("-workdir has invalid path %s"), value.data).data);
 			}
 			else if (name.Equals(TC("-capacity")))
 			{
@@ -291,7 +292,7 @@ namespace uba
 		{
 			AWS aws;
 			StringBuffer<> info;
-			if (aws.Init(logger, info, TC("UbaCli")))
+			if (aws.QueryInformation(logger, info, TC("UbaCli")))
 			{
 				logger.Info(TC("We are inside AWS: %s (%s)"), info.data, aws.GetAvailabilityZone());
 				
@@ -364,6 +365,9 @@ namespace uba
 		if (workDir.IsEmpty())
 			workDir.Append(currentDir);
 
+		// TODO: Change workdir to make it full
+
+
 		StringBuffer<> logFile;
 		logFile.count = GetFullPathNameW(g_rootDir.data, logFile.capacity, logFile.data, nullptr);
 		logFile.EnsureEndsWithSlash().Append(TC("DebugLog.log"));
diff --git a/Engine/Source/Programs/UnrealBuildAccelerator/Common/Public/UbaAWS.h b/Engine/Source/Programs/UnrealBuildAccelerator/Common/Public/UbaAWS.h
--- a/Engine/Source/Programs/UnrealBuildAccelerator/Common/Public/UbaAWS.h
+++ b/Engine/Source/Programs/UnrealBuildAccelerator/Common/Public/UbaAWS.h
@@ -2,6 +2,7 @@
 
 #pragma once
 
+#include "UbaFileAccessor.h"
 #include "UbaNetworkBackend.h"
 #include "UbaStringBuffer.h"
 
@@ -21,14 +22,9 @@ namespace uba
 		static constexpr char g_imdsSpotInstanceAction[]		= "latest/meta-data/spot/instance-action";
 
 
-		bool InitCore(Logger& logger, const tchar* application)
+		bool QueryInformation(Logger& logger, StringBufferBase& outExtraInfo, const tchar* rootDir)
 		{
-			return true;
-		}
-
-		bool Init(Logger& logger, StringBufferBase& outExtraInfo, const tchar* application)
-		{
-			if (!InitCore(logger, application))
+			if (IsNotAWS(logger, rootDir))
 				return false;
 
 			HttpConnection http;
@@ -37,7 +33,8 @@ namespace uba
 
 			StringBuffer<128> instanceId;
 			if (!http.Get(logger, instanceId, statusCode, g_imdsHost, g_imdsInstanceId))
-				return false;
+				return WriteIsNotAws(logger, rootDir);
+
 			outExtraInfo.Append(TC(", AWS: ")).Append(instanceId);
 
 			StringBuffer<32> instanceLifeCycle;
@@ -54,27 +51,51 @@ namespace uba
 				m_isAutoscaling = true;
 			}
 
-			if (!InitAvailabilityZone(logger))
+			if (!QueryAvailabilityZone(logger, nullptr))
 				return false;
 
 			return true;
 		}
 
-		bool InitAvailabilityZone(Logger& logger)
+		bool QueryAvailabilityZone(Logger& logger, const tchar* rootDir)
 		{
+			if (rootDir && IsNotAWS(logger, rootDir))
+				return false;
+
 			HttpConnection http;
 
 			StringBuffer<128> availabilityZone;
 			u32 statusCode = 0;
 			if (!http.Get(logger, availabilityZone, statusCode, g_imdsHost, g_imdsInstanceAvailabilityZone))
+			{
+				if (rootDir)
+					WriteIsNotAws(logger, rootDir);
 				return false;
+			}
 			m_availabilityZone = availabilityZone.data;
 			return true;
 		}
 
-		bool InitPolling(Logger& logger)
+		// Returns true if we _know_ we are not in AWS
+		bool IsNotAWS(Logger& logger, const tchar* rootDir)
 		{
-			return true;
+			StringBuffer<512> file;
+			file.Append(rootDir).EnsureEndsWithSlash().Append(".isNotAWS");
+			if (FileExists(logger, file.data))
+				return true;
+			return false;
+		}
+
+		bool WriteIsNotAws(Logger& logger, const tchar* rootDir)
+		{
+			StringBuffer<512> file;
+			file.Append(rootDir).EnsureEndsWithSlash().Append(".isNotAWS");
+			FileAccessor f(logger, file.data);
+			if (!f.CreateWrite())
+				return false;
+			if (!f.Close())
+				return false;
+			return false;
 		}
 
 		bool IsTerminating(Logger& logger, StringBufferBase& outReason, u64& outTerminationTimeMs)
diff --git a/Engine/Source/Programs/UnrealBuildAccelerator/Host/Private/UbaExports.cpp b/Engine/Source/Programs/UnrealBuildAccelerator/Host/Private/UbaExports.cpp
--- a/Engine/Source/Programs/UnrealBuildAccelerator/Host/Private/UbaExports.cpp
+++ b/Engine/Source/Programs/UnrealBuildAccelerator/Host/Private/UbaExports.cpp
@@ -175,9 +175,8 @@ extern "C"
 		if (!zone || !*zone)
 		{
 			LoggerWithWriter logger(writer, TC(""));
-			if (aws.InitCore(logger, TC("UbaHost")))
-				if (aws.InitAvailabilityZone(logger))
-					zone = aws.GetAvailabilityZone();
+			if (aws.QueryAvailabilityZone(logger, rootDir))
+				zone = aws.GetAvailabilityZone();
 		}
 #endif
 
