UBA: Additional changes for Mac support.

* Fixes an issue with socket closing
* Adds support for getting number of procs.

diff --git a/Engine/Source/Programs/UnrealBuildAccelerator/Core/Private/UbaPlatform.cpp b/Engine/Source/Programs/UnrealBuildAccelerator/Core/Private/UbaPlatform.cpp
--- a/Engine/Source/Programs/UnrealBuildAccelerator/Core/Private/UbaPlatform.cpp
+++ b/Engine/Source/Programs/UnrealBuildAccelerator/Core/Private/UbaPlatform.cpp
@@ -327,13 +327,7 @@ namespace uba
 		#if PLATFORM_WINDOWS
 		return ::GetActiveProcessorCount(ALL_PROCESSOR_GROUPS);
 		#else
-		#if PLATFORM_MAC
-		UBA_ASSERTF(false, TC("GetLogicalProcessorCount not implemented"));
-		return 0;
-		#else
 		return (u32)sysconf(_SC_NPROCESSORS_ONLN);
-		//return std::thread::hardware_concurrency();
-		#endif
 		#endif
 	}
 
