[Uba]
* Fixed so puts is properly detoured so output logging going through there is captured by uba
* Fixed so linux application paths can contain spaces

diff --git a/Engine/Source/Programs/UnrealBuildAccelerator/Common/Private/UbaProcess.cpp b/Engine/Source/Programs/UnrealBuildAccelerator/Common/Private/UbaProcess.cpp
--- a/Engine/Source/Programs/UnrealBuildAccelerator/Common/Private/UbaProcess.cpp
+++ b/Engine/Source/Programs/UnrealBuildAccelerator/Common/Private/UbaProcess.cpp
@@ -770,8 +770,9 @@ namespace uba
 			case MessageType_Log:
 				{
 					bool printInSession = reader.ReadBool();
+					bool isError = reader.ReadBool();
 					TString line = reader.ReadString();
-					LogLine(printInSession, std::move(line), LogEntryType_Info);
+					LogLine(printInSession, std::move(line), isError ? LogEntryType_Error : LogEntryType_Info);
 					return true;
 				}
 				
@@ -1293,10 +1294,17 @@ namespace uba
 
 			wordexp_t  w;
 
-			auto expRes = wordexp(m_realApplication.c_str(), &w, 0);
+			const char* realApplication = m_realApplication.c_str();
+			StringBuffer<> tempApplication;
+			if (m_realApplication.find(' ') != TString::npos)
+			{
+				tempApplication.Append('\"').Append(m_realApplication).Append('\"');
+				realApplication = tempApplication.data;
+			}
+			auto expRes = wordexp(realApplication, &w, 0);
 			if (expRes != 0)
 			{
-				logger.Error("wordexp failed (%i) parsing application name: %s", expRes, m_realApplication.c_str());
+				logger.Error("wordexp failed (%i) parsing application name: %s", expRes, realApplication);
 				return UBA_EXIT_CODE(16);
 			}
 			const char* args = m_startInfo.arguments;
diff --git a/Engine/Source/Programs/UnrealBuildAccelerator/Detours/Private/PosixOS/UbaDetoursFunctionsPosix.cpp b/Engine/Source/Programs/UnrealBuildAccelerator/Detours/Private/PosixOS/UbaDetoursFunctionsPosix.cpp
--- a/Engine/Source/Programs/UnrealBuildAccelerator/Detours/Private/PosixOS/UbaDetoursFunctionsPosix.cpp
+++ b/Engine/Source/Programs/UnrealBuildAccelerator/Detours/Private/PosixOS/UbaDetoursFunctionsPosix.cpp
@@ -836,7 +836,7 @@ UBA_EXPORT ssize_t UBA_WRAPPER(write)(int fd, const void* buf, size_t count)
 	UBA_INIT_DETOUR(write, fd, buf, count);
 	if (isatty(fd)) // stdout and stderr
 	{
-		Shared_WriteConsole((const char*)buf, count);
+		Shared_WriteConsole((const char*)buf, count, fd == fileno(stderr));
 		return count;
 	}
 	//DEBUG_LOG_TRUE("write", "(%i size: %llu)", fd, count);
@@ -1923,7 +1923,7 @@ namespace uba
 		SuppressDetourScope s;
 		StringBuffer<4096> b;
 		WriteAssertInfo(b, text, file, line, expr, 1);
-		Rpc_WriteLog(b.data, b.count, true);
+		Rpc_WriteLog(b.data, b.count, true, true);
 
 		BinaryWriter writer;
 		writer.WriteByte(MessageType_Exit);
diff --git a/Engine/Source/Programs/UnrealBuildAccelerator/Detours/Private/UbaDetoursShared.cpp b/Engine/Source/Programs/UnrealBuildAccelerator/Detours/Private/UbaDetoursShared.cpp
--- a/Engine/Source/Programs/UnrealBuildAccelerator/Detours/Private/UbaDetoursShared.cpp
+++ b/Engine/Source/Programs/UnrealBuildAccelerator/Detours/Private/UbaDetoursShared.cpp
@@ -151,7 +151,7 @@ namespace uba
 		va_end(arg);
 		StringBuffer<2048> sb;
 		sb.Append(GetApplicationShortName()).Append(TC(" ERROR: ")).Append(buffer);
-		Rpc_WriteLog(sb.data, sb.count, true);
+		Rpc_WriteLog(sb.data, sb.count, true, true);
 
 		#if PLATFORM_WINDOWS // Maybe all platforms should call exit()?
 		ExitProcess(code);
@@ -160,7 +160,7 @@ namespace uba
 		#endif
 	}
 
-	void Rpc_WriteLog(const tchar* text, u64 textCharLength, bool printInSession)
+	void Rpc_WriteLog(const tchar* text, u64 textCharLength, bool printInSession, bool isError)
 	{
 		DEBUG_LOG(TC("LOG  %.*s"), u32(textCharLength), text); // TODO: Investigate, deadlocks on non-windows
 		TimerScope ts(g_stats.log);
@@ -168,7 +168,7 @@ namespace uba
 		BinaryWriter writer;
 		writer.WriteByte(MessageType_Log);
 		writer.WriteBool(printInSession);
-		//writer.WriteBool(hConsoleOutput == g_stderrHandle);
+		writer.WriteBool(isError);
 		writer.WriteString(text, textCharLength);
 		writer.Flush();
 	}
@@ -185,7 +185,7 @@ namespace uba
 			count = int(TStrlen(buffer));
 		}
 		va_end(arg);
-		Rpc_WriteLog(buffer, u32(count), false);
+		Rpc_WriteLog(buffer, u32(count), false, false);
 	}
 
 	//TODO: Implement SetConsoleTextAttribute.. clang is using it to color errors
@@ -195,7 +195,7 @@ namespace uba
 	EARLY_INIT ReaderWriterLock g_consoleStringCs;
 
 	template<typename CharType>
-	void Shared_WriteConsoleT(const CharType* chars, u32 charCount)
+	void Shared_WriteConsoleT(const CharType* chars, u32 charCount, bool isError)
 	{
 		if (!g_echoOn)
 			return;
@@ -212,7 +212,7 @@ namespace uba
 				*write = 0;
 				u32 strLen = u32(write - g_consoleString);
 				if (!g_rules->SuppressLogLine(g_consoleString, strLen))
-					Rpc_WriteLog(g_consoleString, strLen, false);
+					Rpc_WriteLog(g_consoleString, strLen, false, isError);
 				write = g_consoleString;
 				left = sizeof_array(g_consoleString) - 1;
 			}
@@ -228,10 +228,10 @@ namespace uba
 		g_consoleStringIndex = u32(write - g_consoleString);
 	}
 
-	void Shared_WriteConsole(const char* chars, u32 charCount) { Shared_WriteConsoleT(chars, charCount); }
+	void Shared_WriteConsole(const char* chars, u32 charCount, bool isError) { Shared_WriteConsoleT(chars, charCount, isError); }
 
 	#if PLATFORM_WINDOWS
-	void Shared_WriteConsole(const wchar_t* chars, u32 charCount) { Shared_WriteConsoleT(chars, charCount); }
+	void Shared_WriteConsole(const wchar_t* chars, u32 charCount, bool isError) { Shared_WriteConsoleT(chars, charCount, isError); }
 	#endif
 
 
diff --git a/Engine/Source/Programs/UnrealBuildAccelerator/Detours/Private/UbaDetoursShared.h b/Engine/Source/Programs/UnrealBuildAccelerator/Detours/Private/UbaDetoursShared.h
--- a/Engine/Source/Programs/UnrealBuildAccelerator/Detours/Private/UbaDetoursShared.h
+++ b/Engine/Source/Programs/UnrealBuildAccelerator/Detours/Private/UbaDetoursShared.h
@@ -19,7 +19,7 @@ namespace uba
 
 	ANALYSIS_NORETURN void FatalError(u32 code, const tchar* format, ...);
 
-	void Rpc_WriteLog(const tchar* text, u64 textCharLength, bool printInSession);
+	void Rpc_WriteLog(const tchar* text, u64 textCharLength, bool printInSession, bool isError);
 	void Rpc_WriteLogf(const tchar* format, ...);
 
 	const tchar* GetApplicationShortName();
@@ -128,8 +128,8 @@ namespace uba
 	inline bool CanDetour(const tchar* file) { return !t_disallowDetour && g_rules->CanDetour(file); }
 	inline bool KeepInMemory(const tchar* fileName, u32 fileNameLen) { return g_allowKeepFilesInMemory && g_rules->KeepInMemory(fileName, fileNameLen, g_systemTemp.data); }
 
-	void Shared_WriteConsole(const char* chars, u32 charCount);
-	void Shared_WriteConsole(const wchar_t* chars, u32 charCount);
+	void Shared_WriteConsole(const char* chars, u32 charCount, bool isError);
+	void Shared_WriteConsole(const wchar_t* chars, u32 charCount, bool isError);
 
 	const tchar* Shared_GetFileAttributes(FileAttributes& outAttr, const tchar* fileName, bool checkIfDir = false);
 }
\ No newline at end of file
diff --git a/Engine/Source/Programs/UnrealBuildAccelerator/Detours/Private/Windows/UbaDetoursFunctionsKernelBase.inl b/Engine/Source/Programs/UnrealBuildAccelerator/Detours/Private/Windows/UbaDetoursFunctionsKernelBase.inl
--- a/Engine/Source/Programs/UnrealBuildAccelerator/Detours/Private/Windows/UbaDetoursFunctionsKernelBase.inl
+++ b/Engine/Source/Programs/UnrealBuildAccelerator/Detours/Private/Windows/UbaDetoursFunctionsKernelBase.inl
@@ -315,7 +315,7 @@ BOOL Detoured_WriteConsoleA(HANDLE hConsoleOutput, const VOID* lpBuffer, DWORD n
 {
 	DETOURED_CALL(WriteConsoleA);
 	//DEBUG_LOG_TRUE_AND_DETOURED(L"WriteConsoleA (%hs)", (char*)lpBuffer);
-	Shared_WriteConsole((const char*)lpBuffer, nNumberOfCharsToWrite);
+	Shared_WriteConsole((const char*)lpBuffer, nNumberOfCharsToWrite, false);
 	if (lpNumberOfCharsWritten)
 		*lpNumberOfCharsWritten = nNumberOfCharsToWrite;
 	return TRUE;//True_WriteConsoleA(hConsoleOutput, lpBuffer, nNumberOfCharsToWrite, lpNumberOfCharsWritten, lpReserved);
@@ -325,7 +325,7 @@ BOOL Detoured_WriteConsoleW(HANDLE hConsoleOutput, const VOID* lpBuffer, DWORD n
 {
 	DETOURED_CALL(WriteConsoleW);
 	//DEBUG_LOG_DETOURED(L"WriteConsoleW"", L""); // Too much spam
-	Shared_WriteConsole((const wchar_t*)lpBuffer, nNumberOfCharsToWrite);
+	Shared_WriteConsole((const wchar_t*)lpBuffer, nNumberOfCharsToWrite, false);
 	if (lpNumberOfCharsWritten)
 		*lpNumberOfCharsWritten = nNumberOfCharsToWrite;
 	return TRUE;//True_WriteConsoleW(hConsoleOutput, lpBuffer, nNumberOfCharsToWrite, lpNumberOfCharsWritten, lpReserved);
@@ -498,7 +498,7 @@ DWORD Detoured_GetSecurityInfo(HANDLE handle, SE_OBJECT_TYPE ObjectType, SECURIT
 	return True_GetSecurityInfo(handle, ObjectType, SecurityInfo, ppsidOwner, ppsidGroup, ppDacl, ppSacl, ppSecurityDescriptor);
 }
 
-void WriteStdFile(LPCVOID lpBuffer, DWORD nNumberOfBytesToWrite)
+void WriteStdFile(LPCVOID lpBuffer, DWORD nNumberOfBytesToWrite, bool isError)
 {
 	if (!g_echoOn)
 		return;
@@ -516,7 +516,7 @@ void WriteStdFile(LPCVOID lpBuffer, DWORD nNumberOfBytesToWrite)
 				--len;
 			if (len)
 				g_stdFile.Appendf(L"%.*hs", len, bufferStr + start);
-			Rpc_WriteLog(g_stdFile.data, g_stdFile.count, false);
+			Rpc_WriteLog(g_stdFile.data, g_stdFile.count, false, isError);
 			g_stdFile.Clear();
 			start = i + 1;
 		}
@@ -538,7 +538,7 @@ BOOL Detoured_WriteFile(HANDLE hFile, LPCVOID lpBuffer, DWORD nNumberOfBytesToWr
 
 		if (dh.type == HandleType_Std)
 		{
-			WriteStdFile(lpBuffer, nNumberOfBytesToWrite);
+			WriteStdFile(lpBuffer, nNumberOfBytesToWrite, hFile == g_stdHandle[0]);
 			*lpNumberOfBytesWritten = nNumberOfBytesToWrite;
 			SetLastError(ERROR_SUCCESS);
 			return true;
@@ -564,7 +564,7 @@ BOOL Detoured_WriteFile(HANDLE hFile, LPCVOID lpBuffer, DWORD nNumberOfBytesToWr
 	}
 	else if (hFile == g_stdHandle[1] || hFile == g_stdHandle[0])
 	{
-		WriteStdFile(lpBuffer, nNumberOfBytesToWrite);
+		WriteStdFile(lpBuffer, nNumberOfBytesToWrite, hFile == g_stdHandle[0]);
 		SetLastError(ERROR_SUCCESS);
 		return true;
 	}
diff --git a/Engine/Source/Programs/UnrealBuildAccelerator/Detours/Private/Windows/UbaDetoursFunctionsUcrtBase.inl b/Engine/Source/Programs/UnrealBuildAccelerator/Detours/Private/Windows/UbaDetoursFunctionsUcrtBase.inl
--- a/Engine/Source/Programs/UnrealBuildAccelerator/Detours/Private/Windows/UbaDetoursFunctionsUcrtBase.inl
+++ b/Engine/Source/Programs/UnrealBuildAccelerator/Detours/Private/Windows/UbaDetoursFunctionsUcrtBase.inl
@@ -255,6 +255,18 @@ int Detoured__write(int fd, const void* buffer, unsigned int count)
 	return res;
 }
 
+int Detoured_fputs(const char* str, FILE* stream)
+{
+	int fno = _fileno(stream); (void)fno;
+	int errfno = _fileno(stderr);
+	if (fno == errfno || fno == _fileno(stdout))
+	{
+		Shared_WriteConsole(str, u32(strlen(str)), fno == errfno);
+		return 1;
+	}
+	return True_fputs(str, stream);
+}
+
 int Detoured__wspawnl(int mode, const wchar_t* cmdname, const wchar_t* arg0, const wchar_t* arg1, ...)
 {
 	DETOURED_CALL(_wspawnl);
diff --git a/Engine/Source/Programs/UnrealBuildAccelerator/Detours/Private/Windows/UbaDetoursFunctionsWin.cpp b/Engine/Source/Programs/UnrealBuildAccelerator/Detours/Private/Windows/UbaDetoursFunctionsWin.cpp
--- a/Engine/Source/Programs/UnrealBuildAccelerator/Detours/Private/Windows/UbaDetoursFunctionsWin.cpp
+++ b/Engine/Source/Programs/UnrealBuildAccelerator/Detours/Private/Windows/UbaDetoursFunctionsWin.cpp
@@ -375,7 +375,7 @@ void UbaAssert(const wchar_t* text, const char* file, u32 line, const char* expr
 	
 	StringBuffer<32*1024> b;
 	WriteAssertInfo(b, text, file, line, expr, 1);
-	Rpc_WriteLog(b.data, b.count, true);
+	Rpc_WriteLog(b.data, b.count, true, true);
 	#if UBA_DEBUG_LOG_ENABLED
 	FlushDebugLog();
 	#endif
diff --git a/Engine/Source/Programs/UnrealBuildAccelerator/Detours/Private/Windows/UbaDetoursFunctionsWin.h b/Engine/Source/Programs/UnrealBuildAccelerator/Detours/Private/Windows/UbaDetoursFunctionsWin.h
--- a/Engine/Source/Programs/UnrealBuildAccelerator/Detours/Private/Windows/UbaDetoursFunctionsWin.h
+++ b/Engine/Source/Programs/UnrealBuildAccelerator/Detours/Private/Windows/UbaDetoursFunctionsWin.h
@@ -11,6 +11,7 @@
 #include <mbstring.h>
 #include <Shlwapi.h>
 #include <wchar.h>
+#include <stdio.h>
 
 
 #if UBA_DEBUG
@@ -145,6 +146,7 @@
 	DETOURED_FUNCTION(_get_osfhandle) \
 	DETOURED_FUNCTION(_isatty) \
 	DETOURED_FUNCTION(_write) \
+	DETOURED_FUNCTION(fputs) \
 	DETOURED_FUNCTIONS_CRTBASE_DEBUG \
 
 #if defined(UBA_USE_MIMALLOC)
