[UBA]
* Moved directory cache out to its own class

diff --git a/Engine/Source/Programs/UnrealBuildAccelerator/Common/Private/UbaFile.cpp b/Engine/Source/Programs/UnrealBuildAccelerator/Common/Private/UbaFile.cpp
--- a/Engine/Source/Programs/UnrealBuildAccelerator/Common/Private/UbaFile.cpp
+++ b/Engine/Source/Programs/UnrealBuildAccelerator/Common/Private/UbaFile.cpp
@@ -741,4 +741,47 @@ namespace uba
 		return true;
 #endif
 	}
+
+	bool DirectoryCache::CreateDirectory(Logger& logger, const tchar* dir)
+	{
+		u64 dirLen = TStrlen(dir);
+		if (dir[dirLen - 1] == PathSeparator)
+			--dirLen;
+		TString key(dir, dir + dirLen);
+		dir = key.c_str();
+
+		ScopedWriteLock lock(m_createdDirsLock);
+		CreatedDir& cd = m_createdDirs.try_emplace(key).first->second;
+		lock.Leave();
+		ScopedWriteLock dirLock(cd.lock);
+		if (cd.handled)
+			return true;
+		cd.handled = true;
+		if (uba::CreateDirectoryW(dir))
+			return true;
+		u32 lastError = GetLastError();
+		if (lastError == ERROR_ALREADY_EXISTS)
+			return true;
+		if (lastError != ERROR_PATH_NOT_FOUND)
+			return logger.Error(TC("Failed to create directory %s (%s)"), dir, LastErrorToText(lastError).data);
+
+		tchar temp[512];
+		const tchar* lastSep = TStrrchr(dir, PathSeparator);
+		u64 pos = u64(lastSep - dir);
+		memcpy(temp, dir, pos * sizeof(tchar));
+		temp[pos] = 0;
+		if (pos == 2 && temp[1] == ':')
+			return false;
+		if (!CreateDirectory(logger, temp))
+			return false;
+		if (!uba::CreateDirectoryW(dir))
+			return logger.Error(TC("Failed to create directory %s (%s)"), dir, LastErrorToText().data);
+		return true;
+	}
+
+	void DirectoryCache::Clear()
+	{
+		ScopedWriteLock lock(m_createdDirsLock);
+		m_createdDirs.clear();
+	}
 }
diff --git a/Engine/Source/Programs/UnrealBuildAccelerator/Common/Private/UbaStorage.cpp b/Engine/Source/Programs/UnrealBuildAccelerator/Common/Private/UbaStorage.cpp
--- a/Engine/Source/Programs/UnrealBuildAccelerator/Common/Private/UbaStorage.cpp
+++ b/Engine/Source/Programs/UnrealBuildAccelerator/Common/Private/UbaStorage.cpp
@@ -1505,7 +1505,7 @@ namespace uba
 
 		bool res = DeleteAllFiles(m_logger, m_rootDir.data, false, &deleteCount);
 		m_logger.Info(TC("Deleted %u cas files"), deleteCount);
-		m_createdDirs.clear();
+		m_dirCache.Clear();
 		return res;
 	}
 
@@ -2371,38 +2371,6 @@ namespace uba
 
 	bool StorageImpl::CreateDirectory(const tchar* dir)
 	{
-		u64 dirLen = TStrlen(dir);
-		if (dir[dirLen - 1] == PathSeparator)
-			--dirLen;
-		TString key(dir, dir + dirLen);
-		dir = key.c_str();
-
-		ScopedWriteLock lock(m_createdDirsLock);
-		CreatedDir& cd = m_createdDirs.try_emplace(key).first->second;
-		lock.Leave();
-		ScopedWriteLock dirLock(cd.lock);
-		if (cd.handled)
-			return true;
-		cd.handled = true;
-		if (uba::CreateDirectoryW(dir))
-			return true;
-		u32 lastError = GetLastError();
-		if (lastError == ERROR_ALREADY_EXISTS)
-			return true;
-		if (lastError != ERROR_PATH_NOT_FOUND)
-			return m_logger.Error(TC("Failed to create directory %s (%s)"), dir, LastErrorToText(lastError).data);
-
-		tchar temp[512];
-		const tchar* lastSep = TStrrchr(dir, PathSeparator);
-		u64 pos = u64(lastSep - dir);
-		memcpy(temp, dir, pos*sizeof(tchar));
-		temp[pos] = 0;
-		if (pos == 2 && temp[1] == ':')
-			return false;
-		if (!CreateDirectory(temp))
-			return false;
-		if (!uba::CreateDirectoryW(dir))
-			return m_logger.Error(TC("Failed to create directory %s (%s)"), dir, LastErrorToText().data);
-		return true;
+		return m_dirCache.CreateDirectory(m_logger, dir);
 	}
 }
diff --git a/Engine/Source/Programs/UnrealBuildAccelerator/Common/Public/UbaFile.h b/Engine/Source/Programs/UnrealBuildAccelerator/Common/Public/UbaFile.h
--- a/Engine/Source/Programs/UnrealBuildAccelerator/Common/Public/UbaFile.h
+++ b/Engine/Source/Programs/UnrealBuildAccelerator/Common/Public/UbaFile.h
@@ -2,8 +2,9 @@
 
 #pragma once
 
-#include "UbaPlatform.h"
 #include "UbaLogger.h"
+#include "UbaMemory.h"
+#include "UbaPlatform.h"
 
 namespace uba
 {
@@ -76,4 +77,17 @@ namespace uba
 	u64 GetSystemTimeAsFileTime();
 	u64 GetFileTimeAsSeconds(u64 fileTime);
 	bool GetCurrentDirectoryW(StringBufferBase& out);
+
+
+	class DirectoryCache
+	{
+	public:
+		bool CreateDirectory(Logger& logger, const tchar* dir);
+		void Clear();
+
+	private:
+		ReaderWriterLock m_createdDirsLock;
+		struct CreatedDir { ReaderWriterLock lock; bool handled = false; };
+		UnorderedMap<TString, CreatedDir> m_createdDirs;
+	};
 }
diff --git a/Engine/Source/Programs/UnrealBuildAccelerator/Common/Public/UbaStorage.h b/Engine/Source/Programs/UnrealBuildAccelerator/Common/Public/UbaStorage.h
--- a/Engine/Source/Programs/UnrealBuildAccelerator/Common/Public/UbaStorage.h
+++ b/Engine/Source/Programs/UnrealBuildAccelerator/Common/Public/UbaStorage.h
@@ -204,9 +204,7 @@ namespace uba
 		UnorderedMap<CasKey, TString> m_deferredCasCreationLookup;
 		UnorderedMap<const tchar*, CasKey, HashString, EqualString> m_deferredCasCreationLookupByName;
 
-		ReaderWriterLock m_createdDirsLock;
-		struct CreatedDir { ReaderWriterLock lock; bool handled = false; };
-		UnorderedMap<TString, CreatedDir> m_createdDirs;
+		DirectoryCache m_dirCache;
 
 		static constexpr u64 BufferSlotSize = 16*1024*1024;
 		static constexpr u64 BufferSlotHalfSize = BufferSlotSize/2; // This must be three times a msg size or more.
