[UBA]
* Added error when trying to use crypto on non-windows platforms (not implemented yet)

diff --git a/Engine/Source/Programs/UnrealBuildAccelerator/Common/Private/UbaCrypto.cpp b/Engine/Source/Programs/UnrealBuildAccelerator/Common/Private/UbaCrypto.cpp
--- a/Engine/Source/Programs/UnrealBuildAccelerator/Common/Private/UbaCrypto.cpp
+++ b/Engine/Source/Programs/UnrealBuildAccelerator/Common/Private/UbaCrypto.cpp
@@ -57,6 +57,7 @@ namespace uba
 
 		return (CryptoKey)(u64)keyHandle;
 #else
+		logger.Error(TC("ERROR: Crypto not supported on non-windows platforms"));
 		return InvalidCryptoKey;
 #endif // UBA_CRYPTO_TYPE
 	}
