[UBA]
* Disabled sparse files and overlapped writing.. quite large performance hit.  Will revisit and see if it matters for larger files
* Changed write size from 1gb to 8mb for WriteFile.. this seems to be a sweet spot based on some testing. Needs more testing

diff --git a/Engine/Source/Programs/UnrealBuildAccelerator/Common/Private/UbaFileAccessor.cpp b/Engine/Source/Programs/UnrealBuildAccelerator/Common/Private/UbaFileAccessor.cpp
--- a/Engine/Source/Programs/UnrealBuildAccelerator/Common/Private/UbaFileAccessor.cpp
+++ b/Engine/Source/Programs/UnrealBuildAccelerator/Common/Private/UbaFileAccessor.cpp
@@ -211,7 +211,7 @@ namespace uba
 		u8* pos = (u8*)data;
 		while (writeLeft)
 		{
-			u32 toWrite = u32(Min(writeLeft, 1024llu * 1024 * 1024));
+			u32 toWrite = u32(Min(writeLeft, 8llu * 1024 * 1024));
 			u32 toActuallyWrite = toWrite;
 
 #if 0
diff --git a/Engine/Source/Programs/UnrealBuildAccelerator/Common/Private/UbaSession.cpp b/Engine/Source/Programs/UnrealBuildAccelerator/Common/Private/UbaSession.cpp
--- a/Engine/Source/Programs/UnrealBuildAccelerator/Common/Private/UbaSession.cpp
+++ b/Engine/Source/Programs/UnrealBuildAccelerator/Common/Private/UbaSession.cpp
@@ -1943,7 +1943,7 @@ namespace uba
 			auto memClose = MakeGuard([&](){ UnmapViewOfFile(mem, fileSize); });
 
 			constexpr bool useFileMapForWrite = false;
-			bool useOverlap = fileSize > 1024 * 1024;
+			bool useOverlap = false;// fileSize > 8 * 1024 * 1024;
 
 			u32 attributes = DefaultAttributes();
 			if (useOverlap)
