[UBA]
* Implement simple commandline parser and added unit tests
* Removed usage of wordexp on non-windows platforms because it seems to have a _huge_ penalty on macos..
* Added trace scope around WriteFilesToDisk
* Fixed assert in stdout/err pipes on macos and added handling for poll hangup
*

diff --git a/Engine/Source/Programs/UnrealBuildAccelerator/Common/Private/UbaProcess.cpp b/Engine/Source/Programs/UnrealBuildAccelerator/Common/Private/UbaProcess.cpp
--- a/Engine/Source/Programs/UnrealBuildAccelerator/Common/Private/UbaProcess.cpp
+++ b/Engine/Source/Programs/UnrealBuildAccelerator/Common/Private/UbaProcess.cpp
@@ -16,7 +16,6 @@
 #include <poll.h>
 #include <stdio.h>
 #include <spawn.h>
-#include <wordexp.h>
 
 // These headers are used for tracking child and beyond
 // processes and making sure they clean up properly
@@ -422,11 +421,6 @@ namespace uba
 
 		m_session.ProcessExited(*this, m_processStats.wallTime);
 
-		#if PLATFORM_MAC
-		//int res = WaitForProcessGroup(getpgrp());
-		//UBA_ASSERT(res == 0);
-		#endif
-
 		// For some reason a parent can exit before a child. Need to figure out repro for this but I've seen it happen on ClangEditor win64
 		for (auto& child : m_childProcesses)
 			while (!((ProcessImpl*)child.m_process)->m_hasExited)
@@ -818,17 +812,9 @@ namespace uba
 
 					m_processStats.Add(stats);
 
-					if (m_parentProcess)
-						m_parentProcess->m_processStats.Add(m_processStats);
-
-					if (m_startInfo.outputStatsThresholdMs && TimeToMs(m_processStats.GetTotalTime()) > m_startInfo.outputStatsThresholdMs)
-					{
-						m_session.PrintProcessStats(m_processStats, logName.data);
-						m_processStats.Print(m_session.m_logger);
-					}
-
 					if (g_applicationRules[m_rulesIndex].rules->IsExitCodeSuccess(m_nativeProcessExitCode) && !IsCancelled())
 					{
+						TimerScope ts(m_processStats.writeFiles);
 						ScopedWriteLock lock(m_writtenFilesLock);
 						for (auto& kv : m_writtenFiles)
 						{
@@ -840,6 +826,16 @@ namespace uba
 									m_messageSuccess = false;
 						}
 					}
+
+					if (m_parentProcess)
+						m_parentProcess->m_processStats.Add(m_processStats);
+
+					if (m_startInfo.outputStatsThresholdMs && TimeToMs(m_processStats.GetTotalTime()) > m_startInfo.outputStatsThresholdMs)
+					{
+						m_session.PrintProcessStats(m_processStats, logName.data);
+						m_processStats.Print(m_session.m_logger);
+					}
+
 					return false;
 				}
 			case MessageType_Custom:
@@ -1001,57 +997,6 @@ namespace uba
 		return temp.data;
 	}
 
-#if PLATFORM_MAC
-	int ProcessImpl::WaitForProcessGroup(pid_t pgid)
-	{
-		int name[] = {CTL_KERN, KERN_PROC, KERN_PROC_PGRP, pgid};
-
-		for (;;) {
-			// Query the list of processes in the group by using sysctl(3).
-			// This is "hard" because we don't know how big that list is, so we
-			// have to first query the size of the output data and then account for
-			// the fact that the size might change by the time we actually issue
-			// the query.
-			struct kinfo_proc *procs = NULL;
-			size_t nprocs = 0;
-			do {
-				size_t len;
-				if (sysctl(name, 4, 0, &len, NULL, 0) == -1) {
-					printf("Something went wrong\n");
-					return -1;
-				}
-				procs = (struct kinfo_proc *)malloc(len);
-				if (sysctl(name, 4, procs, &len, NULL, 0) == -1) {
-					UBA_ASSERT(errno == ENOMEM);
-					free(procs);
-					procs = NULL;
-				} else {
-					nprocs = len / sizeof(struct kinfo_proc);
-				}
-			} while (procs == NULL);
-			UBA_ASSERT(nprocs >= 1);  // Must have found the group leader at least.
-
-			if (nprocs == 1) {
-				// Found only one process, which must be the leader because we have
-				// purposely expect it as a zombie.
-				UBA_ASSERT(procs->kp_proc.p_pid == pgid);
-				free(procs);
-				return 0;
-			}
-
-			// More than one process left in the process group.  Pause a little bit
-			// before retrying to avoid burning CPU.
-			struct timespec ts;
-			ts.tv_sec = 0;
-			ts.tv_nsec = 1000000;
-			if (nanosleep(&ts, NULL) == -1) {
-				UBA_ASSERT(FALSE);
-				return -1;
-			}
-		}
-	}
-#endif
-
 	struct ProcessImpl::PipeReader
 	{
 		PipeReader(ProcessImpl& p, LogEntryType lt) : process(p), logType(lt) {}
@@ -1352,8 +1297,6 @@ namespace uba
 
 		if (!m_parentProcess)
 		{
-			wordexp_t  w;
-
 			const char* realApplication = m_realApplication.c_str();
 			StringBuffer<> tempApplication;
 			if (m_realApplication.find(' ') != TString::npos)
@@ -1362,30 +1305,19 @@ namespace uba
 				realApplication = tempApplication.data;
 			}
 
-			static ReaderWriterLock g_wordexpLock; // wordexp is not thread-safe
-			ScopedWriteLock wordExpLock(g_wordexpLock);
-
-			auto expRes = wordexp(realApplication, &w, 0);
-			if (expRes != 0)
-			{
-				logger.Error("wordexp failed (%i) parsing application name: %s", expRes, realApplication);
-				return UBA_EXIT_CODE(16);
-			}
-			const char* args = m_startInfo.arguments;
-			expRes = wordexp(args, &w, WRDE_APPEND);
-			if (expRes != 0)
+			Vector<TString> arguments;
+			if (!ParseArguments(arguments, m_startInfo.arguments))
 			{
-				logger.Error("wordexp failed (%i) parsing arguments: %s", expRes, args);
+				logger.Error("Failed to parse arguments: %s", m_startInfo.arguments);
 				return UBA_EXIT_CODE(16);
 			}
-
-			wordExpLock.Leave();
-
-			auto wordExpGuard = MakeGuard([&]() { wordfree(&w); });
-
-			//logger.Info(TC("ARGS: %hs\n", args.c_str());
-			//const char* argv[] = { app.c_str(), args.c_str(), nullptr };
-			//const char* argv[] = { m_realApplication.c_str(), "-o", "code", "Code.cpp", nullptr };
+			Vector<const char*> arguments2;
+			arguments2.reserve(arguments.size() + 2);
+			arguments2.push_back(realApplication);
+			for (auto& s : arguments)
+				arguments2.push_back(s.data());
+			arguments2.push_back(nullptr);
+			auto argsArray = arguments2.data();
 
 			short flags = POSIX_SPAWN_SETPGROUP;
 
@@ -1414,7 +1346,7 @@ namespace uba
 			StringBuffer<32> rulesStr;
 			StringBuffer<512> logFile;
 			StringBuffer<512> ldLibraryPath;
-			StringBuffer<128> detoursVar;
+			StringBuffer<512> detoursVar;
 
 			Vector<const char*> envvars;
 
@@ -1503,7 +1435,7 @@ namespace uba
 			envvars.push_back(nullptr);
 
 			pid_t processID;
-			res = posix_spawnp(&processID, m_realApplication.c_str(), &fileActions, &attr, w.we_wordv, (char**)envvars.data());
+			res = posix_spawnp(&processID, m_realApplication.c_str(), &fileActions, &attr, (char**)argsArray, (char**)envvars.data());
 
 			posix_spawn_file_actions_destroy(&fileActions);
 			posix_spawnattr_destroy(&attr);
@@ -1524,10 +1456,18 @@ namespace uba
 				PipeReader outReader(*this, LogEntryType_Info);
 				PipeReader errReader(*this, LogEntryType_Error);
 
-				pollfd plist[] = { {outPipe[0],POLLIN}, {errPipe[0],POLLIN} };
+				pollfd plist[] = { {outPipe[0],POLLIN, 0}, {errPipe[0],POLLIN, 0} };
 
 				for (int rval; (rval = poll(plist, sizeof_array(plist), -1)) > 0;)
 				{
+					if (plist[0].revents & POLLHUP && plist[1].revents & POLLHUP) // If they both have hung up we hang up
+						break;
+
+					if (plist[0].revents & POLLERR || plist[1].revents & POLLERR) // If there is an error on any of them we hang up
+					{
+						logger.Error(TC("pipe polling error"));
+						break;
+					}
 					int fd = 0;
 					PipeReader* pipeReader = nullptr;
 					if (plist[0].revents & POLLIN)
@@ -1545,7 +1485,6 @@ namespace uba
 
 					char buffer[1024];
 					int bytesRead = read(fd, buffer, sizeof_array(buffer) - 1);
-					UBA_ASSERT(bytesRead > 0);
 					buffer[bytesRead] = 0;
 					pipeReader->ReadData(buffer, bytesRead);
 				}
@@ -1703,4 +1642,61 @@ namespace uba
 				CloseFileMapping(pair.second.mappingHandle);
 		m_tempFiles.clear();
 	}
+
+	bool ParseArguments(Vector<TString>& outArguments, const tchar* argumentString)
+	{
+		const tchar* argStart = argumentString;
+		bool isInArg = false;
+		bool isInQuotes = false;
+		bool isEnd = *argumentString == 0;
+		tchar lastChar = 0;
+		for (const tchar* it = argumentString; !isEnd; lastChar = *it, ++it)
+		{
+			isEnd = *it == 0;
+			if (*it == ' ' || *it == '\t' || isEnd)
+			{
+				if (isInQuotes || !isInArg)
+					continue;
+
+				TString result(argStart, it);
+				tchar lastChar2 = 0;
+				for (auto i = result.begin(); i != result.end();)
+				{
+					if (*i == '\"')
+					{
+						if (lastChar2 == '\\')
+							i = result.erase(i - 1) + 1;
+						else
+							i = result.erase(i);
+						lastChar2 = 0;
+						continue;
+					}
+					lastChar2 = *i;
+					++i;
+				}
+
+				outArguments.push_back(std::move(result));
+				isInArg = false;
+				continue;
+			}
+
+			if (!isInArg)
+			{
+				isInArg = true;
+				argStart = it;
+				if (*it == '\"')
+					isInQuotes = true;
+				continue;
+			}
+
+			if (*it == '\"')
+			{
+				if (isInQuotes && lastChar == '\\')
+					continue;
+
+				isInQuotes = !isInQuotes;
+			}
+		}
+		return true;
+	}
 }
diff --git a/Engine/Source/Programs/UnrealBuildAccelerator/Common/Public/UbaAWS.h b/Engine/Source/Programs/UnrealBuildAccelerator/Common/Public/UbaAWS.h
--- a/Engine/Source/Programs/UnrealBuildAccelerator/Common/Public/UbaAWS.h
+++ b/Engine/Source/Programs/UnrealBuildAccelerator/Common/Public/UbaAWS.h
@@ -6,7 +6,9 @@
 #include "UbaNetworkBackend.h"
 #include "UbaStringBuffer.h"
 
-#define UBA_USE_AWS
+#define UBA_USE_AWS 1
+
+#if UBA_USE_AWS
 
 namespace uba
 {
@@ -146,3 +148,5 @@ namespace uba
 		bool m_isAutoscaling = false;
 	};
 }
+
+#endif
\ No newline at end of file
diff --git a/Engine/Source/Programs/UnrealBuildAccelerator/Common/Public/UbaProcess.h b/Engine/Source/Programs/UnrealBuildAccelerator/Common/Public/UbaProcess.h
--- a/Engine/Source/Programs/UnrealBuildAccelerator/Common/Public/UbaProcess.h
+++ b/Engine/Source/Programs/UnrealBuildAccelerator/Common/Public/UbaProcess.h
@@ -144,4 +144,6 @@ namespace uba
 		void operator=(const ProcessImpl&) = delete;
 		struct PipeReader;
 	};
+
+	bool ParseArguments(Vector<TString>& outArguments, const tchar* argumentString);
 }
diff --git a/Engine/Source/Programs/UnrealBuildAccelerator/Common/Public/UbaTrace.h b/Engine/Source/Programs/UnrealBuildAccelerator/Common/Public/UbaTrace.h
--- a/Engine/Source/Programs/UnrealBuildAccelerator/Common/Public/UbaTrace.h
+++ b/Engine/Source/Programs/UnrealBuildAccelerator/Common/Public/UbaTrace.h
@@ -49,7 +49,7 @@ namespace uba
 		TraceType_FileStoreLight,
 	};
 
-	static constexpr u32 TraceVersion = 18;
+	static constexpr u32 TraceVersion = 19;
 	static constexpr u32 TraceReadCompatibilityVersion = 6;
 
 	class Trace
diff --git a/Engine/Source/Programs/UnrealBuildAccelerator/Core/Public/UbaProcessStats.h b/Engine/Source/Programs/UnrealBuildAccelerator/Core/Public/UbaProcessStats.h
--- a/Engine/Source/Programs/UnrealBuildAccelerator/Core/Public/UbaProcessStats.h
+++ b/Engine/Source/Programs/UnrealBuildAccelerator/Core/Public/UbaProcessStats.h
@@ -39,6 +39,7 @@ namespace uba
 		UBA_PROCESS_STAT(virtualAllocFailed, 0) \
 		UBA_PROCESS_STAT(log, 0) \
 		UBA_PROCESS_STAT(sendFiles, 0) \
+		UBA_PROCESS_STAT(writeFiles, 19) \
 
 
 		#define UBA_PROCESS_STAT(T, ver) Timer T;
diff --git a/Engine/Source/Programs/UnrealBuildAccelerator/Host/Private/UbaExports.cpp b/Engine/Source/Programs/UnrealBuildAccelerator/Host/Private/UbaExports.cpp
--- a/Engine/Source/Programs/UnrealBuildAccelerator/Host/Private/UbaExports.cpp
+++ b/Engine/Source/Programs/UnrealBuildAccelerator/Host/Private/UbaExports.cpp
@@ -170,7 +170,7 @@ extern "C"
 	{
 		using namespace uba;
 
-#if defined(UBA_USE_AWS)
+		#if UBA_USE_AWS
 		AWS aws;
 		if (!zone || !*zone)
 		{
@@ -178,7 +178,7 @@ extern "C"
 			if (aws.QueryAvailabilityZone(logger, rootDir))
 				zone = aws.GetAvailabilityZone();
 		}
-#endif
+		#endif
 
 		StorageServerCreateInfo info(server, rootDir, writer);
 		info.casCapacityBytes = casCapacityBytes;
diff --git a/Engine/Source/Programs/UnrealBuildAccelerator/Test/Private/UbaTestAll.h b/Engine/Source/Programs/UnrealBuildAccelerator/Test/Private/UbaTestAll.h
--- a/Engine/Source/Programs/UnrealBuildAccelerator/Test/Private/UbaTestAll.h
+++ b/Engine/Source/Programs/UnrealBuildAccelerator/Test/Private/UbaTestAll.h
@@ -15,6 +15,7 @@ namespace uba
 		UBA_TEST(TestEvents) \
 		UBA_TEST(TestPaths) \
 		UBA_TEST(TestFiles) \
+		UBA_TEST(TestParseArguments) \
 		UBA_TEST(TestSockets) \
 		UBA_TEST(TestClientServer) \
 		UBA_TEST(TestClientServer2) \
diff --git a/Engine/Source/Programs/UnrealBuildAccelerator/Test/Private/UbaTestBasics.h b/Engine/Source/Programs/UnrealBuildAccelerator/Test/Private/UbaTestBasics.h
--- a/Engine/Source/Programs/UnrealBuildAccelerator/Test/Private/UbaTestBasics.h
+++ b/Engine/Source/Programs/UnrealBuildAccelerator/Test/Private/UbaTestBasics.h
@@ -6,6 +6,7 @@
 
 #include "UbaBinaryReaderWriter.h"
 #include "UbaFileAccessor.h"
+#include "UbaProcess.h"
 #include "UbaPathUtils.h"
 #include "UbaPlatform.h"
 #include "UbaEvent.h"
@@ -15,24 +16,28 @@
 #include "UbaTimer.h"
 #include "UbaDirectoryIterator.h"
 
+#define VA_ARGS(...) , ##__VA_ARGS__
+#define UBA_TEST_CHECK(expr, fmt, ...) if (!(expr)) return logger.Error(TC(fmt) VA_ARGS(__VA_ARGS__));
+
+
 namespace uba
 {
 	bool TestTime(Logger& logger, const StringBufferBase& rootDir)
 	{
-		#if 0
+#if 0
 		LoggerWithWriter consoleLogger(g_consoleLogWriter); (void)consoleLogger;
 		u64 time1 = GetSystemTimeUs();
 		Sleep(1000);
 		u64 time2 = GetSystemTimeUs();
 		u64 ms = (time2 - time1) / 1000;
 		consoleLogger.Info(TC("Slept ms: %llu"), ms);
-		
+
 		time1 = GetTime();
 		Sleep(1000);
 		time2 = GetTime();
 		ms = (time2 * 1000 / GetFrequency()) - (time1 * 1000 / GetFrequency());
 		consoleLogger.Info(TC("Slept ms: %llu"), ms);
-		#endif
+#endif
 
 		return true;
 	}
@@ -67,14 +72,14 @@ namespace uba
 
 	bool TestPaths(Logger& logger, const StringBufferBase& rootDir)
 	{
-		#if PLATFORM_WINDOWS
+#if PLATFORM_WINDOWS
 		const tchar* workingDir = TC("e:\\dev\\");
 		const tchar* filename = TC("\"e:\\temp\"");
 		tchar buffer[1024];
 		u32 lengthResult;
 		if (!FixPath2(filename, workingDir, TStrlen(workingDir), buffer, &lengthResult))
 			return logger.Error(TC("FixPath2 (1) failed"));
-		#endif
+#endif
 		return true;
 	}
 
@@ -186,4 +191,42 @@ namespace uba
 		block.Free(mem);
 		return true;
 	}
+
+	bool TestParseArguments(Logger& logger, const StringBufferBase& rootDir)
+	{
+		Vector<TString> arguments;
+		ParseArguments(arguments, TC("foo bar"));
+		UBA_TEST_CHECK(arguments.size() == 2, "ParseAgsuments 1 failed (%llu)", arguments.size());
+
+		Vector<TString> arguments2;
+		ParseArguments(arguments2, TC("\"foo\" bar"));
+		UBA_TEST_CHECK(arguments2.size() == 2, "ParseArguments 2 failed");
+
+		Vector<TString> arguments3;
+		ParseArguments(arguments3, TC("\"foo meh\" bar"));
+		UBA_TEST_CHECK(arguments3.size() == 2, "ParseArguments 3 failed");
+		UBA_TEST_CHECK(Contains(arguments3[0].data(), TC(" ")), "ParseArguments 3 failed");
+
+		Vector<TString> arguments4;
+		ParseArguments(arguments4, TC("\"app\" @\"rsp\""));
+		UBA_TEST_CHECK(arguments4.size() == 2, "ParseArguments 4 failed");
+		UBA_TEST_CHECK(!Contains(arguments4[1].data(), TC("\"")), "ParseArguments 4 failed");
+
+		Vector<TString> arguments5;
+		ParseArguments(arguments5, TC("\"app\" @\"rsp foo\""));
+		UBA_TEST_CHECK(arguments5.size() == 2, "ParseArguments 4 failed");
+		UBA_TEST_CHECK(!Contains(arguments5[1].data(), TC("\"")), "ParseArguments 5 failed");
+		UBA_TEST_CHECK(Contains(arguments5[1].data(), TC(" ")), "ParseArguments 5 failed");
+
+		Vector<TString> arguments6;
+		ParseArguments(arguments6, TC("\"app\"\"1\" @\"rsp foo\""));
+		UBA_TEST_CHECK(arguments6.size() == 2, "ParseArguments 6 failed");
+		UBA_TEST_CHECK(Equals(arguments6[0].data(), TC("app1")), "ParseArguments 6 failed");
+
+		Vector<TString> arguments7;
+		ParseArguments(arguments7, TC("app \" \\\"foo\\\" bar\""));
+		UBA_TEST_CHECK(arguments7.size() == 2, "ParseArguments 7 failed");
+		UBA_TEST_CHECK(Contains(arguments7[1].data(), TC("\"")), "ParseArguments 7 failed");
+		return true;
+	}
 }
diff --git a/Engine/Source/Programs/UnrealBuildAccelerator/Test/Private/UbaTestSession.h b/Engine/Source/Programs/UnrealBuildAccelerator/Test/Private/UbaTestSession.h
--- a/Engine/Source/Programs/UnrealBuildAccelerator/Test/Private/UbaTestSession.h
+++ b/Engine/Source/Programs/UnrealBuildAccelerator/Test/Private/UbaTestSession.h
@@ -300,12 +300,23 @@ namespace uba
 			});
 	}
 
+	const tchar* GetPingApplication()
+	{
+		#if PLATFORM_WINDOWS
+		return TC("c:\\windows\\system32\\ping.exe");
+		#elif PLATFORM_LINUX
+		return TC("/usr/bin/ping");
+		#else
+		return TC("/sbin/ping");
+		#endif
+	}
+
 	bool TestMultipleDetouredProcesses(LoggerWithWriter& logger, const StringBufferBase& testRootDir)
 	{
 		return RunLocal(logger, testRootDir, [](LoggerWithWriter& logger, SessionServer& session, const tchar* workingDir, const RunProcessFunction& runProcess)
 			{
 				ProcessStartInfo processInfo;
-				processInfo.application = IsWindows ? TC("c:\\windows\\system32\\ping.exe") : TC("/usr/bin/ping");
+				processInfo.application = GetPingApplication();
 				processInfo.workingDir = workingDir;
 				processInfo.arguments = IsWindows ? TC("-n 2 localhost") : TC("-c 2 localhost");
 				Vector<ProcessHandle> processes;
@@ -331,7 +342,7 @@ namespace uba
 		return RunLocal(logger, testRootDir, [](LoggerWithWriter& logger, SessionServer& session, const tchar* workingDir, const RunProcessFunction& runProcess)
 			{
 				ProcessStartInfo processInfo;
-				processInfo.application = IsWindows ? TC("c:\\windows\\system32\\ping.exe") : TC("/usr/bin/ping");
+				processInfo.application = GetPingApplication();
 				processInfo.workingDir = workingDir;
 				processInfo.arguments = IsWindows ? TC("-n 1 localhost") : TC("-c 1 localhost");
 
