[UBA]
* Agent - Fixed fixup code for provided paths -dir and -eventfile
* Cli - Fixed fixup code for provided paths -dir
* Fixed so root dir is created before aws code write potential .isNotAWS file

diff --git a/Engine/Source/Programs/UnrealBuildAccelerator/Agent/Private/UbaAgent.cpp b/Engine/Source/Programs/UnrealBuildAccelerator/Agent/Private/UbaAgent.cpp
--- a/Engine/Source/Programs/UnrealBuildAccelerator/Agent/Private/UbaAgent.cpp
+++ b/Engine/Source/Programs/UnrealBuildAccelerator/Agent/Private/UbaAgent.cpp
@@ -42,7 +42,7 @@ namespace uba
 		if (IsWindows)
 			ExpandEnvironmentStringsW(TC("%ProgramData%\\Epic\\" UE_APP_NAME), buf, sizeof(buf));
 		else
-			TStrcpy_s(buf, sizeof(buf), TC("~/" UE_APP_NAME));
+			GetFullPathNameW(TC("~/" UE_APP_NAME), sizeof_array(buf), buf, nullptr);
 		return buf;
 	}();
 	u32				DefaultProcessorCount = []() { return GetLogicalProcessorCount(); }();
@@ -58,9 +58,9 @@ namespace uba
 			logger.Error(TC("%s"), message);
 		}
 		logger.Info(TC(""));
-		logger.Info(TC("------------------------"));
+		logger.Info(TC("-------------------------------------------"));
 		logger.Info(TC("   UbaAgent v%hs"), Version);
-		logger.Info(TC("------------------------"));
+		logger.Info(TC("-------------------------------------------"));
 		logger.Info(TC(""));
 		logger.Info(TC("  When started UbaAgent will keep trying to connect to provided host address."));
 		logger.Info(TC("  Once connected it will start helping out. Nothing else is needed :)"));
@@ -287,20 +287,6 @@ namespace uba
 		return 0;
 	}
 
-	int FixCommandLinePath(StringBufferBase& str)
-	{
-		if (int res = ExpandEnvironmentVariables(str))
-			return res;
-		if (str[1] != ':')
-		{
-			StringBuffer<> currentDir;
-			GetCurrentDirectoryW(currentDir);
-			currentDir.EnsureEndsWithSlash().Append(str);
-			str.Clear().Append(currentDir);
-		}
-		return 0;
-	}
-
 	bool IsTerminating(Logger& logger, const tchar* eventFile, StringBufferBase& outReason, u64& outTerminationTimeMs)
 	{
 		if (!*eventFile)
@@ -526,10 +512,10 @@ namespace uba
 			{
 				if (value.IsEmpty())
 					return PrintHelp(TC("-dir needs a value"));
-				g_rootDir.Clear();
-				g_rootDir.Append(value).Replace('/', PathSeparator);
-				if (int res = FixCommandLinePath(g_rootDir))
+				if (int res = ExpandEnvironmentVariables(value))
 					return res;
+				if ((g_rootDir.count = GetFullPathNameW(value.Replace('\\', PathSeparator).data, g_rootDir.capacity, g_rootDir.data, nullptr)) == 0)
+					return PrintHelp(StringBuffer<>().Appendf(TC("-dir has invalid path %s"), value.data).data);
 			}
 			else if (name.Equals(TC("-name")))
 			{
@@ -578,9 +564,10 @@ namespace uba
 			{
 				if (value.IsEmpty())
 					return PrintHelp(TC("-eventfile needs a value"));
-				eventFile.Append(value);
-				if (int res = FixCommandLinePath(eventFile))
+				if (int res = ExpandEnvironmentVariables(value))
 					return res;
+				if ((eventFile.count = GetFullPathNameW(value.Replace('\\', PathSeparator).data, eventFile.capacity, eventFile.data, nullptr)) == 0)
+					return PrintHelp(StringBuffer<>().Appendf(TC("-eventfile has invalid path %s"), value.data).data);
 			}
 			else if (name.Equals(TC("-killrandom")))
 			{
@@ -684,11 +671,15 @@ namespace uba
 
 
 		// Check if AWS
-		#if defined(UBA_USE_AWS)
+		#if UBA_USE_AWS
 		AWS aws;
-		aws.QueryInformation(logger, extraInfo, g_rootDir.data);
-		if (zone.IsEmpty())
-			zone.Append(aws.GetAvailabilityZone());
+		{
+			DirectoryCache dirCache;
+			dirCache.CreateDirectory(logger, g_rootDir.data);
+			aws.QueryInformation(logger, extraInfo, g_rootDir.data);
+			if (zone.IsEmpty())
+				zone.Append(aws.GetAvailabilityZone());
+		}
 		#endif
 
 		if (zone.count)
@@ -795,7 +786,7 @@ namespace uba
 		StringBuffer<512> terminationReason;
 		u64 terminationTimeMs = 0;
 
-		#if defined(UBA_USE_AWS)
+		#if UBA_USE_AWS
 		if (aws.IsTerminating(logger, terminationReason, terminationTimeMs))
 		{
 			LoggerWithWriter(g_consoleLogWriter, TC("")).Info(TC("%s. Exiting UbaAgent before starting session"), terminationReason.data);
@@ -1064,7 +1055,7 @@ namespace uba
 				{
 					if (IsTerminating(logger, eventFile.data, terminationReason, terminationTimeMs))
 						isTerminating = true;
-					#if defined(UBA_USE_AWS)
+					#if UBA_USE_AWS
 					else if (aws.IsTerminating(logger, terminationReason, terminationTimeMs))
 						isTerminating = true;
 					#endif
diff --git a/Engine/Source/Programs/UnrealBuildAccelerator/Cli/Private/UbaCli.cpp b/Engine/Source/Programs/UnrealBuildAccelerator/Cli/Private/UbaCli.cpp
--- a/Engine/Source/Programs/UnrealBuildAccelerator/Cli/Private/UbaCli.cpp
+++ b/Engine/Source/Programs/UnrealBuildAccelerator/Cli/Private/UbaCli.cpp
@@ -25,7 +25,7 @@ namespace uba
 		if (IsWindows)
 			ExpandEnvironmentStringsW(TC("%ProgramData%\\Epic\\" UE_APP_NAME), buf, sizeof(buf));
 		else
-			TStrcpy_s(buf, sizeof(buf), TC("~/" UE_APP_NAME));
+			GetFullPathNameW(TC("~/" UE_APP_NAME), sizeof_array(buf), buf, nullptr);
 		return buf;
 		}();
 	u32				DefaultProcessorCount = []() { return GetLogicalProcessorCount(); }();
@@ -39,9 +39,9 @@ namespace uba
 			logger.Error(TC("%s"), message);
 		}
 		logger.Info(TC(""));
-		logger.Info(TC("------------------------"));
+		logger.Info(TC("-------------------------------------------"));
 		logger.Info(TC("   UbaCli v%hs"), Version);
-		logger.Info(TC("------------------------"));
+		logger.Info(TC("-------------------------------------------"));
 		logger.Info(TC(""));
 		logger.Info(TC("  UbaCli.exe [options...] <commandtype> <executable> [arguments...]"));
 		logger.Info(TC(""));
@@ -261,7 +261,8 @@ namespace uba
 			{
 				if (value.IsEmpty())
 					return PrintHelp(TC("-dir needs a value"));
-				g_rootDir.Clear().Append(value);
+				if ((g_rootDir.count = GetFullPathNameW(value.Replace('/', PathSeparator).data, g_rootDir.capacity, g_rootDir.data, nullptr)) == 0)
+					return PrintHelp(StringBuffer<>().Appendf(TC("-dir has invalid path %s"), g_rootDir.data).data);
 			}
 			else if (name.Equals(TC("-?")))
 			{
@@ -313,6 +314,9 @@ namespace uba
 			return PrintHelp(errorMsg);
 		}
 
+		if (application.empty())
+			return PrintHelp(TC("No executable provided"));
+
 		StringBuffer<512> currentDir;
 		GetCurrentDirectoryW(currentDir);
 
