[Uba]
* Changed max size for WriteFile to 256mb after testing various sizes. Bigger seems better but we have had reported hangs when having this value set to 1gb to let's see if 256mb works for everyone

On my machine (Threadripper 128 logical cores and fast ssd) , linking 471 unrealeditor dlls takes 19 seconds with 256mb/write... 128mb/s takes ~21s, 16mb takes ~25 seconds

diff --git a/Engine/Source/Programs/UnrealBuildAccelerator/Common/Private/UbaFileAccessor.cpp b/Engine/Source/Programs/UnrealBuildAccelerator/Common/Private/UbaFileAccessor.cpp
--- a/Engine/Source/Programs/UnrealBuildAccelerator/Common/Private/UbaFileAccessor.cpp
+++ b/Engine/Source/Programs/UnrealBuildAccelerator/Common/Private/UbaFileAccessor.cpp
@@ -211,7 +211,7 @@ namespace uba
 		u8* pos = (u8*)data;
 		while (writeLeft)
 		{
-			u32 toWrite = u32(Min(writeLeft, 8llu * 1024 * 1024));
+			u32 toWrite = u32(Min(writeLeft, 256llu * 1024 * 1024));
 			u32 toActuallyWrite = toWrite;
 
 #if 0
