[UBA]
* Removed "EARLY_INIT" macro since this doesn't seem to work properly on macos.
* Changed so only one __attribute(constructor..) exists
* Replaced g_applicationRules with GetApplicationRules() to try to resolve initializationordering issues

diff --git a/Engine/Source/Programs/UnrealBuildAccelerator/Common/Private/UbaProcess.cpp b/Engine/Source/Programs/UnrealBuildAccelerator/Common/Private/UbaProcess.cpp
--- a/Engine/Source/Programs/UnrealBuildAccelerator/Common/Private/UbaProcess.cpp
+++ b/Engine/Source/Programs/UnrealBuildAccelerator/Common/Private/UbaProcess.cpp
@@ -812,7 +812,7 @@ namespace uba
 
 					m_processStats.Add(stats);
 
-					if (g_applicationRules[m_rulesIndex].rules->IsExitCodeSuccess(m_nativeProcessExitCode) && !IsCancelled())
+					if (GetApplicationRules()[m_rulesIndex].rules->IsExitCodeSuccess(m_nativeProcessExitCode) && !IsCancelled())
 					{
 						TimerScope ts(m_processStats.writeFiles);
 						ScopedWriteLock lock(m_writtenFilesLock);
@@ -936,11 +936,13 @@ namespace uba
 		StringBuffer<128> exeName;
 		exeName.Append(m_virtualApplication.c_str() + exeNameStart, exeNameEnd - exeNameStart);
 		
+		auto rules = GetApplicationRules();
+		
 		while (true)
 		{
 			for (u32 i = 1;; ++i)
 			{
-				const tchar* app = g_applicationRules[i].app;
+				const tchar* app = rules[i].app;
 				if (!app)
 					break;
 				if (!exeName.Equals(app))
@@ -1094,7 +1096,7 @@ namespace uba
 				return ProcessCancelExitCode;
 			}
 
-			bool isDetachedProcess = g_applicationRules[m_rulesIndex].rules->AllowDetach() && m_detourEnabled;
+			bool isDetachedProcess = GetApplicationRules()[m_rulesIndex].rules->AllowDetach() && m_detourEnabled;
 
 			HANDLE hJob = CreateJobObject(nullptr, nullptr);
 			JOBOBJECT_EXTENDED_LIMIT_INFORMATION info = { };
@@ -1243,7 +1245,7 @@ namespace uba
 			payload.runningRemote = runningRemote;
 			payload.isChild = m_parentProcess != nullptr;
 			payload.trackInputs = m_startInfo.trackInputs;
-			payload.useCustomAllocator = m_startInfo.useCustomAllocator && g_applicationRules[m_rulesIndex].rules->AllowMiMalloc();
+			payload.useCustomAllocator = m_startInfo.useCustomAllocator && GetApplicationRules()[m_rulesIndex].rules->AllowMiMalloc();
 			payload.isRunningWine = IsRunningWine();
 			payload.uiLanguage = m_startInfo.uiLanguage;
 			if (*m_startInfo.logFile)
diff --git a/Engine/Source/Programs/UnrealBuildAccelerator/Common/Private/UbaSession.cpp b/Engine/Source/Programs/UnrealBuildAccelerator/Common/Private/UbaSession.cpp
--- a/Engine/Source/Programs/UnrealBuildAccelerator/Common/Private/UbaSession.cpp
+++ b/Engine/Source/Programs/UnrealBuildAccelerator/Common/Private/UbaSession.cpp
@@ -1487,12 +1487,12 @@ namespace uba
 
 	bool Session::IsRarelyRead(ProcessImpl& process, const StringBufferBase& fileName) const
 	{
-		return g_applicationRules[process.m_rulesIndex].rules->IsRarelyRead(fileName);
+		return GetApplicationRules()[process.m_rulesIndex].rules->IsRarelyRead(fileName);
 	}
 
 	bool Session::IsRarelyReadAfterWritten(ProcessImpl& process, const tchar* fileName, u64 fileNameLen) const
 	{
-		return g_applicationRules[process.m_rulesIndex].rules->IsRarelyReadAfterWritten(fileName, fileNameLen);
+		return GetApplicationRules()[process.m_rulesIndex].rules->IsRarelyReadAfterWritten(fileName, fileNameLen);
 	}
 
 	bool Session::ShouldWriteToDisk(const tchar* fileName, u64 fileNameLen)
diff --git a/Engine/Source/Programs/UnrealBuildAccelerator/Core/Private/UbaApplicationRules.cpp b/Engine/Source/Programs/UnrealBuildAccelerator/Core/Private/UbaApplicationRules.cpp
--- a/Engine/Source/Programs/UnrealBuildAccelerator/Core/Private/UbaApplicationRules.cpp
+++ b/Engine/Source/Programs/UnrealBuildAccelerator/Core/Private/UbaApplicationRules.cpp
@@ -337,32 +337,37 @@ namespace uba
 		}
 	};
 
-	RulesRec g_applicationRules[]
+	const RulesRec* GetApplicationRules()
 	{
-		{ TC(""),							new ApplicationRules() },		// Must be index 0
-		{ TC("cl.exe"),						new ApplicationRulesClExe() },	// Must be index 1
-		{ TC("link.exe"),					new ApplicationRulesLinkExe() }, // Must be index 2
-		{ TC("lib.exe"),					new ApplicationRulesLinkExe() },
-		{ TC("cvtres.exe"),					new ApplicationRulesLinkExe() },
-		{ TC("mt.exe"),						new ApplicationRulesLinkExe() },
-		{ TC("rc.exe"),						new ApplicationRulesLinkExe() },
-		{ TC("clang++.exe"),				new ApplicationRulesClangPlusPlusExe() },
-		{ TC("clang-cl.exe"),				new ApplicationRulesClangPlusPlusExe() },
-		{ TC("verse-clang-cl.exe"),			new ApplicationRulesClangPlusPlusExe() },
-		{ TC("ispc.exe"),					new ApplicationRulesISPCExe() },
-		{ TC("orbis-clang.exe"),			new ApplicationRulesOrbisClangPlusPlusExe() },
-		{ TC("orbis-ld.exe"),				new ApplicationRulesOrbisLdExe() },
-		{ TC("orbis-ltop.exe"),				new ApplicationRulesOrbisLdExe() },
-		{ TC("prospero-clang.exe"),			new ApplicationRulesProsperoClangPlusPlusExe() },
-		{ TC("prospero-lld.exe"),			new ApplicationRulesProsperoLldExe() },
-		{ TC("dump_syms.exe"),				new ApplicationRulesDumpSymsExe() },
-		{ TC("ld.lld.exe"),					new ApplicationRulesLdLLdExe() },
-		{ TC("llvm-objcopy.exe"),			new ApplicationRulesLlvmObjCopyExe() },
-		{ TC("UnrealBuildTool.dll"),		new ApplicationRulesUBTDll() },
-		{ TC("PVS-Studio.exe"),				new ApplicationRulesPVSStudio() },
-		//{ L"MSBuild.dll"),				new ApplicationRules() },
-		//{ L"BreakpadSymbolEncoder.exe"),	new ApplicationRulesClang() },
-		//{ L"cmd.exe"),		new ApplicationRules() },
-		{ nullptr, nullptr }
-	};
+		static RulesRec rules[]
+		{
+			{ TC(""),							new ApplicationRules() },		// Must be index 0
+			{ TC("cl.exe"),						new ApplicationRulesClExe() },	// Must be index 1
+			{ TC("link.exe"),					new ApplicationRulesLinkExe() }, // Must be index 2
+			{ TC("lib.exe"),					new ApplicationRulesLinkExe() },
+			{ TC("cvtres.exe"),					new ApplicationRulesLinkExe() },
+			{ TC("mt.exe"),						new ApplicationRulesLinkExe() },
+			{ TC("rc.exe"),						new ApplicationRulesLinkExe() },
+			{ TC("clang++.exe"),				new ApplicationRulesClangPlusPlusExe() },
+			{ TC("clang-cl.exe"),				new ApplicationRulesClangPlusPlusExe() },
+			{ TC("verse-clang-cl.exe"),			new ApplicationRulesClangPlusPlusExe() },
+			{ TC("ispc.exe"),					new ApplicationRulesISPCExe() },
+			{ TC("orbis-clang.exe"),			new ApplicationRulesOrbisClangPlusPlusExe() },
+			{ TC("orbis-ld.exe"),				new ApplicationRulesOrbisLdExe() },
+			{ TC("orbis-ltop.exe"),				new ApplicationRulesOrbisLdExe() },
+			{ TC("prospero-clang.exe"),			new ApplicationRulesProsperoClangPlusPlusExe() },
+			{ TC("prospero-lld.exe"),			new ApplicationRulesProsperoLldExe() },
+			{ TC("dump_syms.exe"),				new ApplicationRulesDumpSymsExe() },
+			{ TC("ld.lld.exe"),					new ApplicationRulesLdLLdExe() },
+			{ TC("llvm-objcopy.exe"),			new ApplicationRulesLlvmObjCopyExe() },
+			{ TC("UnrealBuildTool.dll"),		new ApplicationRulesUBTDll() },
+			{ TC("PVS-Studio.exe"),				new ApplicationRulesPVSStudio() },
+			//{ L"MSBuild.dll"),				new ApplicationRules() },
+			//{ L"BreakpadSymbolEncoder.exe"),	new ApplicationRulesClang() },
+			//{ L"cmd.exe"),		new ApplicationRules() },
+			{ nullptr, nullptr }
+		};
+
+		return rules;
+	}
 }
\ No newline at end of file
diff --git a/Engine/Source/Programs/UnrealBuildAccelerator/Core/Public/UbaApplicationRules.h b/Engine/Source/Programs/UnrealBuildAccelerator/Core/Public/UbaApplicationRules.h
--- a/Engine/Source/Programs/UnrealBuildAccelerator/Core/Public/UbaApplicationRules.h
+++ b/Engine/Source/Programs/UnrealBuildAccelerator/Core/Public/UbaApplicationRules.h
@@ -118,5 +118,5 @@ namespace uba
 
 	struct RulesRec { const tchar* app; ApplicationRules* rules; };
 
-	extern RulesRec g_applicationRules[];
+	const RulesRec* GetApplicationRules();
 }
diff --git a/Engine/Source/Programs/UnrealBuildAccelerator/Detours/Private/PosixOS/UbaDetoursFunctionsPosix.cpp b/Engine/Source/Programs/UnrealBuildAccelerator/Detours/Private/PosixOS/UbaDetoursFunctionsPosix.cpp
--- a/Engine/Source/Programs/UnrealBuildAccelerator/Detours/Private/PosixOS/UbaDetoursFunctionsPosix.cpp
+++ b/Engine/Source/Programs/UnrealBuildAccelerator/Detours/Private/PosixOS/UbaDetoursFunctionsPosix.cpp
@@ -22,7 +22,6 @@ namespace uba
 {
 	constexpr bool g_logToScreen = false;
 
-	EARLY_INIT StringBuffer<MaxPath> g_exeDir;
 	bool g_isDetouring;
 	bool g_isInitialized;
 	void Deinit();
@@ -234,9 +233,6 @@ namespace uba
 		return access;
 	}
 
-	extern ReaderWriterLock g_communicationLock;
-	extern ProcessStats g_stats;
-
 	struct FileObject
 	{
 		FileInfo* fileInfo = nullptr;
@@ -254,8 +250,10 @@ namespace uba
 		FileObject* fileObject = nullptr;
 	};
 
-	UnorderedMap<int, DetouredHandle> g_fileHandles;
-	ReaderWriterLock g_fileHandlesLock;
+	using FileHandles = UnorderedMap<int, DetouredHandle>;
+	VARIABLE_MEM(StringBuffer<MaxPath>, g_exeDir);
+	VARIABLE_MEM(FileHandles, g_fileHandles);
+	VARIABLE_MEM(ReaderWriterLock, g_fileHandlesLock);
 }
 
 // Shared functions
@@ -1734,6 +1732,10 @@ namespace uba
 
 	void PreInit(const char* logFile)
 	{
+		g_exeDirMem.Create();
+		g_fileHandlesMem.Create();
+		g_fileHandlesLockMem.Create();
+
 		SuppressDetourScope s;
 
 		g_systemTemp.Append(getenv("TMPDIR"));
@@ -1914,6 +1916,7 @@ namespace uba
 	{
 		int t = errno;
 		TRUE_WRAPPER(write)(g_debugFile, str, strLen);
+		//fsync(g_debugFile);
 		errno = t;
 	}
 	#endif
diff --git a/Engine/Source/Programs/UnrealBuildAccelerator/Detours/Private/PosixOS/UbaDetoursMainPosix.cpp b/Engine/Source/Programs/UnrealBuildAccelerator/Detours/Private/PosixOS/UbaDetoursMainPosix.cpp
--- a/Engine/Source/Programs/UnrealBuildAccelerator/Detours/Private/PosixOS/UbaDetoursMainPosix.cpp
+++ b/Engine/Source/Programs/UnrealBuildAccelerator/Detours/Private/PosixOS/UbaDetoursMainPosix.cpp
@@ -54,20 +54,18 @@ static void segfault_sigaction(int signal)
 
 static void __attribute__((constructor(102))) PreInitCtor()
 {
+	using namespace uba;
+	SuppressDetourScope s;
+
 	#if PLATFORM_LINUX
 	prctl(PR_SET_PDEATHSIG, SIGHUP, 0, 0, 0); // We want the process to die if the parent die
 	#endif
 	
+	InitSharedVariables();
+
 	const char* logFile = getenv("UBA_LOGFILE");
 	PreInit(logFile);
 	unsetenv("UBA_LOGFILE");
-}
-
-static void __attribute__((constructor (65535))) InitCtor()
-{
-	using namespace uba;
-
-	SuppressDetourScope s;
 
 	const char* comId = getenv("UBA_COMID");
 	//printf("Starting up %s... (comid: %s) %u\n", __progname, (comId && *comId) ? comId : "NOTSET", getpid());
@@ -125,7 +123,7 @@ static void __attribute__((constructor (65535))) InitCtor()
 	g_runningRemote = getenv("UBA_REMOTE") != nullptr;
 
 	g_rulesIndex = strtoul(getenv("UBA_RULES"), nullptr, 10);
-	g_rules = g_applicationRules[g_rulesIndex].rules;
+	g_rules = GetApplicationRules()[g_rulesIndex].rules;
 	unsetenv("UBA_RULES");
 
 	const char* realCwd = getenv("UBA_CWD");
diff --git a/Engine/Source/Programs/UnrealBuildAccelerator/Detours/Private/UbaDetoursShared.cpp b/Engine/Source/Programs/UnrealBuildAccelerator/Detours/Private/UbaDetoursShared.cpp
--- a/Engine/Source/Programs/UnrealBuildAccelerator/Detours/Private/UbaDetoursShared.cpp
+++ b/Engine/Source/Programs/UnrealBuildAccelerator/Detours/Private/UbaDetoursShared.cpp
@@ -7,51 +7,42 @@
 
 namespace uba
 {
+	VARIABLE_MEM(StringBuffer<512>, g_virtualApplication);
+	VARIABLE_MEM(StringBuffer<512>, g_virtualApplicationDir);
+	VARIABLE_MEM(ProcessStats, g_stats);
+	VARIABLE_MEM(ReaderWriterLock, g_communicationLock);
+	VARIABLE_MEM(StringBuffer<256>, g_logName);
+	VARIABLE_MEM(StringBuffer<512>, g_virtualWorkingDir);
+	VARIABLE_MEM(StringBuffer<128>, g_systemRoot);
+	VARIABLE_MEM(StringBuffer<128>, g_systemTemp);
+	VARIABLE_MEM(MemoryBlock, g_memoryBlock);
+	VARIABLE_MEM(DirectoryTable, g_directoryTable);
+	VARIABLE_MEM(MappedFileTable, g_mappedFileTable);
+	VARIABLE_MEM(ReaderWriterLock, g_consoleStringCs);
 
-	EARLY_INIT StringBuffer<512> g_virtualApplication;
-	EARLY_INIT StringBuffer<512> g_virtualApplicationDir;
-	EARLY_INIT ProcessStats g_stats;
-	EARLY_INIT ReaderWriterLock g_communicationLock;
-	EARLY_INIT StringBuffer<256> g_logName;
-	EARLY_INIT StringBuffer<512> g_virtualWorkingDir;
-	EARLY_INIT StringBuffer<128> g_systemRoot;
-	EARLY_INIT StringBuffer<128> g_systemTemp;
 	bool g_echoOn = true;
-
-	// This is annoyingly complicated but it comes down to initialization ordering. On windows we detour everything at a certain point in time
-	// but on linux everything is detoured from the start. This becomes problematic with these statics and we have to use some party tricks to
-	// make sure they are initialized before detoured functions are called.
-	// (If this doesn't work we will need to resort to Functions returning a static variable)
-
-	static u64 g_memoryBlockMem[sizeof(MemoryBlock)/sizeof(u64)];
-	MemoryBlock& g_memoryBlock = *(MemoryBlock*)g_memoryBlockMem;
-
-	static u64 g_directoryTableMem[sizeof(DirectoryTable)/sizeof(u64)];
-	DirectoryTable& g_directoryTable = *(DirectoryTable*)g_directoryTableMem;
-
-	static u64 g_mappedFileTableMem[sizeof(MappedFileTable)/sizeof(u64)];
-	MappedFileTable& g_mappedFileTable = *(MappedFileTable*)g_mappedFileTableMem;
-
-	void InitStatics()
-	{
-		new (g_memoryBlockMem) MemoryBlock(160 * 1024 * 1024); // Leak this
-		new (g_directoryTableMem) DirectoryTable(&g_memoryBlock); // Leak this to speed up shutdown
-		new (g_mappedFileTableMem) MappedFileTable(g_memoryBlock); // Leak this to speed up shutdown
-	}
-
-
-	#if PLATFORM_WINDOWS
-	static bool g_staticInit = []() { InitStatics(); return true; }();
-	#else
-	static void __attribute__((constructor(102))) CtorInitStatics() { InitStatics(); }
-	#endif
-
-
 	u32 g_rulesIndex;
 	ApplicationRules* g_rules;
 	bool g_runningRemote;
 	bool g_isChild;
 
+	void InitSharedVariables()
+	{
+		g_virtualApplicationMem.Create();
+		g_virtualApplicationDirMem.Create();
+		g_statsMem.Create();
+		g_communicationLockMem.Create();
+		g_logNameMem.Create();
+		g_virtualWorkingDirMem.Create();
+		g_systemRootMem.Create();
+		g_systemTempMem.Create();
+
+		g_memoryBlockMem.Create(160 * 1024 * 1024);
+		g_directoryTableMem.Create(&g_memoryBlock);
+		g_mappedFileTableMem.Create(g_memoryBlock);
+		g_consoleStringCsMem.Create();
+	}
+
 #if UBA_DEBUG_LOG_ENABLED
 	FileHandle g_debugFile = InvalidFileHandle;
 	void WriteDebug(const char* str, u32 strLen);
@@ -192,7 +183,6 @@ namespace uba
 
 	tchar g_consoleString[4096];
 	u32 g_consoleStringIndex;
-	EARLY_INIT ReaderWriterLock g_consoleStringCs;
 
 	template<typename CharType>
 	void Shared_WriteConsoleT(const CharType* chars, u32 charCount, bool isError)
diff --git a/Engine/Source/Programs/UnrealBuildAccelerator/Detours/Private/UbaDetoursShared.h b/Engine/Source/Programs/UnrealBuildAccelerator/Detours/Private/UbaDetoursShared.h
--- a/Engine/Source/Programs/UnrealBuildAccelerator/Detours/Private/UbaDetoursShared.h
+++ b/Engine/Source/Programs/UnrealBuildAccelerator/Detours/Private/UbaDetoursShared.h
@@ -51,23 +51,17 @@ namespace uba
 		#define DEBUG_LOG_PIPE(...) ts.Leave();
 	#endif
 
-	#if PLATFORM_WINDOWS
-	#define EARLY_INIT
-	#else
-	#define EARLY_INIT __attribute__((init_priority(101)))
-	#endif
+	extern StringBuffer<512>& g_virtualApplication;
+	extern StringBuffer<512>& g_virtualApplicationDir;
+	extern StringBuffer<256>& g_logName;
 
-	extern StringBuffer<512> g_virtualApplication;
-	extern StringBuffer<512> g_virtualApplicationDir;
-	extern StringBuffer<256> g_logName;
+	extern StringBuffer<512>& g_virtualWorkingDir;
+	extern StringBuffer<128>& g_systemTemp;
+	extern StringBuffer<128>& g_systemRoot;
 
-	extern StringBuffer<512> g_virtualWorkingDir;
-	extern StringBuffer<128> g_systemTemp;
-	extern StringBuffer<128> g_systemRoot;
-
-	extern ProcessStats g_stats;
+	extern ProcessStats& g_stats;
 	extern bool g_echoOn;
-	extern ReaderWriterLock g_communicationLock;
+	extern ReaderWriterLock& g_communicationLock;
 	extern MemoryBlock& g_memoryBlock;
 	extern DirectoryTable& g_directoryTable;
 	extern MappedFileTable& g_mappedFileTable;
@@ -132,4 +126,9 @@ namespace uba
 	void Shared_WriteConsole(const wchar_t* chars, u32 charCount, bool isError);
 
 	const tchar* Shared_GetFileAttributes(FileAttributes& outAttr, const tchar* fileName, bool checkIfDir = false);
-}
\ No newline at end of file
+
+	void InitSharedVariables();
+
+	template<typename T> struct VariableMem { template<typename... Args> void Create(Args&&... args) { new (data) T(args...); }; u64 data[AlignUp(sizeof(T), sizeof(u64)) / sizeof(u64)]; };
+	#define VARIABLE_MEM(type, name) VariableMem<type> name##Mem; type& name = (type&)name##Mem.data;
+}
diff --git a/Engine/Source/Programs/UnrealBuildAccelerator/Detours/Private/Windows/UbaDetoursFunctionsKernelBase.inl b/Engine/Source/Programs/UnrealBuildAccelerator/Detours/Private/Windows/UbaDetoursFunctionsKernelBase.inl
--- a/Engine/Source/Programs/UnrealBuildAccelerator/Detours/Private/Windows/UbaDetoursFunctionsKernelBase.inl
+++ b/Engine/Source/Programs/UnrealBuildAccelerator/Detours/Private/Windows/UbaDetoursFunctionsKernelBase.inl
@@ -2539,7 +2539,7 @@ BOOL Detoured_CreateProcessW(LPCWSTR lpApplicationName, LPWSTR lpCommandLine, LP
 	lpStartupInfo->dwFlags |= STARTF_USESHOWWINDOW;
 	lpStartupInfo->wShowWindow = SW_HIDE;
 
-	if (g_applicationRules[rulesIndex].rules->AllowDetach())
+	if (g_rules->AllowDetach())
 		dwCreationFlags |= DETACHED_PROCESS;
 	else
 		dwCreationFlags |= CREATE_NO_WINDOW;
diff --git a/Engine/Source/Programs/UnrealBuildAccelerator/Detours/Private/Windows/UbaDetoursFunctionsNtDll.inl b/Engine/Source/Programs/UnrealBuildAccelerator/Detours/Private/Windows/UbaDetoursFunctionsNtDll.inl
--- a/Engine/Source/Programs/UnrealBuildAccelerator/Detours/Private/Windows/UbaDetoursFunctionsNtDll.inl
+++ b/Engine/Source/Programs/UnrealBuildAccelerator/Detours/Private/Windows/UbaDetoursFunctionsNtDll.inl
@@ -467,7 +467,7 @@ NTSTATUS NTAPI Local_NtCreateFile(bool IsCreateFunc, PHANDLE hFileHandle, ACCESS
 		StringBuffer<> b;
 		b.Append(ObjectAttributes->ObjectName->Buffer, ObjectAttributes->ObjectName->Length / 2);
 		if (!b.Contains(L"\\Device\\") && !b.EndsWith(L"\\nul"))
-			Rpc_WriteLogf(L"[%ls] WRITTEN: %ls", g_rulesIndex ? g_applicationRules[g_rulesIndex].app : wcsrchr(g_virtualApplication.data, '\\') + 1, ObjectAttributes->ObjectName->Buffer);
+			Rpc_WriteLogf(L"[%ls] WRITTEN: %ls", g_rulesIndex ? GetApplicationRules()[g_rulesIndex].app : wcsrchr(g_virtualApplication.data, '\\') + 1, ObjectAttributes->ObjectName->Buffer);
 	}
 #endif
 
diff --git a/Engine/Source/Programs/UnrealBuildAccelerator/Detours/Private/Windows/UbaDetoursFunctionsWin.cpp b/Engine/Source/Programs/UnrealBuildAccelerator/Detours/Private/Windows/UbaDetoursFunctionsWin.cpp
--- a/Engine/Source/Programs/UnrealBuildAccelerator/Detours/Private/Windows/UbaDetoursFunctionsWin.cpp
+++ b/Engine/Source/Programs/UnrealBuildAccelerator/Detours/Private/Windows/UbaDetoursFunctionsWin.cpp
@@ -692,8 +692,10 @@ void PreInit(const DetoursPayload& payload)
 	mi_option_disable(mi_option_abandoned_page_reset);
 	#endif
 
+	InitSharedVariables();
+
 	g_rulesIndex = payload.rulesIndex;
-	g_rules = g_applicationRules[payload.rulesIndex].rules;
+	g_rules = GetApplicationRules()[payload.rulesIndex].rules;
 	g_useMiMalloc = payload.useCustomAllocator;
 	g_runningRemote = payload.runningRemote;
 	g_isChild = payload.isChild;
diff --git a/Engine/Source/Programs/UnrealBuildAccelerator/Detours/Private/Windows/UbaDetoursMainWin.cpp b/Engine/Source/Programs/UnrealBuildAccelerator/Detours/Private/Windows/UbaDetoursMainWin.cpp
--- a/Engine/Source/Programs/UnrealBuildAccelerator/Detours/Private/Windows/UbaDetoursMainWin.cpp
+++ b/Engine/Source/Programs/UnrealBuildAccelerator/Detours/Private/Windows/UbaDetoursMainWin.cpp
@@ -10,8 +10,6 @@
 
 namespace uba
 {
-	extern ProcessStats g_stats;
-
 	HANDLE g_hostProcess;
 	HANDLE g_writeEvent;
 	HANDLE g_readEvent;
