[Uba]
* Fixed so posix path honors CaseInsensitiveFs

diff --git a/Engine/Source/Programs/UnrealBuildAccelerator/Detours/Private/PosixOS/UbaDetoursFunctionsPosix.cpp b/Engine/Source/Programs/UnrealBuildAccelerator/Detours/Private/PosixOS/UbaDetoursFunctionsPosix.cpp
--- a/Engine/Source/Programs/UnrealBuildAccelerator/Detours/Private/PosixOS/UbaDetoursFunctionsPosix.cpp
+++ b/Engine/Source/Programs/UnrealBuildAccelerator/Detours/Private/PosixOS/UbaDetoursFunctionsPosix.cpp
@@ -254,6 +254,12 @@ namespace uba
 	VARIABLE_MEM(StringBuffer<MaxPath>, g_exeDir);
 	VARIABLE_MEM(FileHandles, g_fileHandles);
 	VARIABLE_MEM(ReaderWriterLock, g_fileHandlesLock);
+
+
+	StringKey ToFilenameKey(const StringBufferBase& b)
+	{
+		return CaseInsensitiveFs ? ToStringKeyLower(b) : ToStringKey(b);
+	}
 }
 
 // Shared functions
@@ -270,7 +276,7 @@ int Shared_open(const char* funcName, const char* file, int flags, int mode, con
 	StringBuffer<> fileName;
 	FixPath(fileName, file);
 
-	StringKey fileNameKey = ToStringKey(fileName);
+	StringKey fileNameKey = ToFilenameKey(fileName);
 
 	bool keepInMemory = false;
 	bool isSystemOrTempFile = StartsWith(file, g_systemTemp.data);
@@ -305,7 +311,7 @@ int Shared_open(const char* funcName, const char* file, int flags, int mode, con
 				{
 					//SetLastError(ERROR_FILE_NOT_FOUND); // Don't think this is needed
 					errno = ENOENT;
-					DEBUG_LOG_DETOURED(funcName, "NOTFOUND_USINGTABLE (%s) -> -1", fileName.data);
+					DEBUG_LOG_DETOURED(funcName, "NOTFOUND_USINGTABLE (%s) (%s) -> -1", fileName.data, KeyToString(fileNameKey).data);
 					return -1;
 				}
 			}
@@ -1165,11 +1171,11 @@ UBA_EXPORT int UBA_WRAPPER(rename)(const char* oldpath, const char* newpath)
 
 	StringBuffer<> fixedOldPath;
 	FixPath(fixedOldPath, oldpath);
-	StringKey oldKey = ToStringKey(fixedOldPath);
+	StringKey oldKey = ToFilenameKey(fixedOldPath);
 
 	StringBuffer<> fixedNewPath;
 	FixPath(fixedNewPath, newpath);
-	StringKey newKey = ToStringKey(fixedNewPath);
+	StringKey newKey = ToFilenameKey(fixedNewPath);
 
 	// TODO: This might be really slow but it seems you can rename files on linux while they are open and they won't be properly renamed until closed
 	{
@@ -1248,7 +1254,7 @@ UBA_EXPORT int UBA_WRAPPER(chmod)(const char* pathname, mode_t mode)
 
 	StringBuffer<> fixedName;
 	FixPath(fixedName, pathname);
-	StringKey key = ToStringKey(fixedName);
+	StringKey key = ToFilenameKey(fixedName);
 	u32 errorCode;
 	{
 		TimerScope ts(g_stats.chmod);
@@ -1335,7 +1341,7 @@ UBA_EXPORT int UBA_WRAPPER(remove)(const char* pathname)
 	//	return TRUE;
 	//}
 
-	StringKey fileNameKey = ToStringKey(fixedName);
+	StringKey fileNameKey = ToFilenameKey(fixedName);
 
 	u32 directoryTableSize;
 	bool result;
@@ -1377,12 +1383,12 @@ UBA_EXPORT int UBA_WRAPPER(posix_spawn)(pid_t* pid, const char* path, const posi
 
 	t_inVfork = 0;
 
-	StringBuffer<8*1024> cmdLine;
+	TString cmdLine;
 	for (u32 i = 0; argv[i]; ++i)
 	{
 		if (i != 0)
-			cmdLine.Append(' ');
-		cmdLine.Append(argv[i]);
+			cmdLine.append(" ");
+		cmdLine.append(argv[i]);
 	}
 
 	TString commandLine;
