[Uba]
* Fixed handling of paths that has dirs that has double dots in it.. still a valid path

diff --git a/Engine/Source/Programs/UnrealBuildAccelerator/Core/Public/UbaPathUtils.h b/Engine/Source/Programs/UnrealBuildAccelerator/Core/Public/UbaPathUtils.h
--- a/Engine/Source/Programs/UnrealBuildAccelerator/Core/Public/UbaPathUtils.h
+++ b/Engine/Source/Programs/UnrealBuildAccelerator/Core/Public/UbaPathUtils.h
@@ -141,6 +141,7 @@ namespace uba
 			tchar* folders[128];
 			u32 folderCount = 0;
 
+			tchar lastLastLastChar = 0;
 			lastLastChar = 0;
 			lastChar = 0;
 			while (true)
@@ -149,7 +150,7 @@ namespace uba
 
 				if (c == '\\' || c == 0)
 				{
-					if (lastChar == '.' && lastLastChar == '.')
+					if (lastChar == '.' && lastLastChar == '.' && lastLastLastChar == '\\')
 					{
 						if (folderCount > 1)
 							--folderCount;
@@ -171,6 +172,7 @@ namespace uba
 						break;
 				}
 
+				lastLastLastChar = lastLastChar;
 				lastLastChar = lastChar;
 				lastChar = c;
 
@@ -234,6 +236,7 @@ namespace uba
 			char* folders[128];
 			u32 folderCount = 0;
 
+			char lastLastLastChar = 0;
 			char lastLastChar = 0;
 			char lastChar = 0;
 			while (true)
@@ -242,7 +245,7 @@ namespace uba
 
 				if (c == '/' || c == 0)
 				{
-					if (lastChar == '.' && lastLastChar == '.')
+					if (lastChar == '.' && lastLastChar == '.' && lastLastLastChar == '/')
 					{
 						if (folderCount > 1)
 							--folderCount;
@@ -263,6 +266,7 @@ namespace uba
 						break;
 				}
 
+				lastLastLastChar = lastLastChar;
 				lastLastChar = lastChar;
 				lastChar = c;
 
diff --git a/Engine/Source/Programs/UnrealBuildAccelerator/Core/Public/UbaPlatform.h b/Engine/Source/Programs/UnrealBuildAccelerator/Core/Public/UbaPlatform.h
--- a/Engine/Source/Programs/UnrealBuildAccelerator/Core/Public/UbaPlatform.h
+++ b/Engine/Source/Programs/UnrealBuildAccelerator/Core/Public/UbaPlatform.h
@@ -114,7 +114,6 @@ namespace uba
 	#define TStrcat_s wcscat_s
 	#define TStrdup _wcsdup
 #else
-	inline constexpr bool CaseInsensitiveFs = false;
 	inline constexpr tchar PathSeparator = '/';
 	inline constexpr u32 MaxPath = 512;
 	#define TStrlen(s) u32(strlen(s))
@@ -132,7 +131,10 @@ namespace uba
 	inline u64 FromTimeSpec(const timespec& ts) { return u64(ts.tv_sec) * 10'000'000ull + u64(ts.tv_nsec/100); }
 	inline timespec ToTimeSpec(u64 time) { timespec ts; ts.tv_sec = time / 10'000'000ull; ts.tv_nsec = (time - (u64(ts.tv_sec) * 10'000'000ull)) * 100; return ts; }
 	#if PLATFORM_LINUX
+	inline constexpr bool CaseInsensitiveFs = false;
 	#define st_mtimespec st_mtim
-	#endif
+	#else
+	inline constexpr bool CaseInsensitiveFs = true;
+#endif
 #endif // PLATFORM_WINDOWS
 }
diff --git a/Engine/Source/Programs/UnrealBuildAccelerator/Test/Private/UbaTestBasics.h b/Engine/Source/Programs/UnrealBuildAccelerator/Test/Private/UbaTestBasics.h
--- a/Engine/Source/Programs/UnrealBuildAccelerator/Test/Private/UbaTestBasics.h
+++ b/Engine/Source/Programs/UnrealBuildAccelerator/Test/Private/UbaTestBasics.h
@@ -72,14 +72,38 @@ namespace uba
 
 	bool TestPaths(Logger& logger, const StringBufferBase& rootDir)
 	{
-#if PLATFORM_WINDOWS
-		const tchar* workingDir = TC("e:\\dev\\");
-		const tchar* filename = TC("\"e:\\temp\"");
+		const tchar* workingDir = IsWindows ? TC("e:\\dev\\") : TC("/dev/bar/");
 		tchar buffer[1024];
 		u32 lengthResult;
-		if (!FixPath2(filename, workingDir, TStrlen(workingDir), buffer, &lengthResult))
+
+		auto TestPath = [&](const tchar* path) { return FixPath2(path, workingDir, TStrlen(workingDir), buffer, &lengthResult); };
+
+#if PLATFORM_WINDOWS
+		if (!FixPath2(TC("\"e:\\temp\""), workingDir, TStrlen(workingDir), buffer, &lengthResult))
 			return logger.Error(TC("FixPath2 (1) failed"));
+#else
+
+		if (!TestPath(TC("/..")))
+			return logger.Error(TC("FixPath2 should have failed"));
+		UBA_TEST_CHECK(Equals(buffer, TC("/")), "Should not contain ..");
+
+		if (!FixPath2(TC("/../Foo"), workingDir, TStrlen(workingDir), buffer, &lengthResult))
+			return logger.Error(TC("FixPath2 should have failed"));
+		UBA_TEST_CHECK(Equals(buffer, TC("/Foo")), "Should not contain ..");
 #endif
+
+		if (!FixPath2(TC("../Foo"), workingDir, TStrlen(workingDir), buffer, &lengthResult))
+			return logger.Error(TC("FixPath2 (1) failed"));
+		UBA_TEST_CHECK(!Contains(buffer, TC("..")), "Should not contain ..");
+
+		if (!FixPath2(TC("@../Foo"), workingDir, TStrlen(workingDir), buffer, &lengthResult))
+			return logger.Error(TC("FixPath2 (1) failed"));
+		UBA_TEST_CHECK(Contains(buffer, TC("..")), "Should contain ..");
+
+		if (!FixPath2(TC("..@/Foo"), workingDir, TStrlen(workingDir), buffer, &lengthResult))
+			return logger.Error(TC("FixPath2 (1) failed"));
+		UBA_TEST_CHECK(Contains(buffer, TC("..")), "Should contain ..");
+
 		return true;
 	}
 
