[UBA]
* Added code to prevent recursion of already handled directories (caused by symlinks)
* Fixed path issue with double forward slashes on macos
* Added unit test for double slash
* Improved assert description

diff --git a/Engine/Source/Programs/UnrealBuildAccelerator/Common/Private/UbaFile.cpp b/Engine/Source/Programs/UnrealBuildAccelerator/Common/Private/UbaFile.cpp
--- a/Engine/Source/Programs/UnrealBuildAccelerator/Common/Private/UbaFile.cpp
+++ b/Engine/Source/Programs/UnrealBuildAccelerator/Common/Private/UbaFile.cpp
@@ -654,7 +654,7 @@ namespace uba
 		if (res == 0)
 			return true;
 
-		UBA_ASSERTF(false, TC("CreateHardLinkW %s to %s error handling not implemented %s"), existingFileName, newFileName, strerror(errno));
+		UBA_ASSERTF(false, TC("CreateHardLinkW %s to %s error handling not implemented (%s)"), existingFileName, newFileName, strerror(errno));
 		return false;
 #endif
 	}
diff --git a/Engine/Source/Programs/UnrealBuildAccelerator/Common/Private/UbaStorageClient.cpp b/Engine/Source/Programs/UnrealBuildAccelerator/Common/Private/UbaStorageClient.cpp
--- a/Engine/Source/Programs/UnrealBuildAccelerator/Common/Private/UbaStorageClient.cpp
+++ b/Engine/Source/Programs/UnrealBuildAccelerator/Common/Private/UbaStorageClient.cpp
@@ -78,8 +78,10 @@ namespace uba
 		WorkManagerImpl workManager(workerCount);
 		ReaderWriterLock lock;
 		bool success = true;
+		UnorderedSet<u64> seenIds;
+		ReaderWriterLock seenIdsLock;
 		for (auto& dir : directories)
-			success = PopulateCasFromDirsRecursive(dir.c_str(), workManager) && success;
+			success = PopulateCasFromDirsRecursive(dir.c_str(), workManager, seenIds, seenIdsLock) && success;
 		workManager.Wait();
 
 		if (u32 fileCount = u32(m_localStorageFiles.size()))
@@ -1036,7 +1038,7 @@ namespace uba
 		return true;
 	}
 
-	bool StorageClient::PopulateCasFromDirsRecursive(const tchar* dir, WorkManager& workManager)
+	bool StorageClient::PopulateCasFromDirsRecursive(const tchar* dir, WorkManager& workManager, UnorderedSet<u64>& seenIds, ReaderWriterLock& seenIdsLock)
 	{
 		StringBuffer<> fullPath;
 		fullPath.Append(dir).EnsureEndsWithSlash();
@@ -1046,7 +1048,11 @@ namespace uba
 				fullPath.Resize(dirLen).Append(e.name);
 				if (IsDirectory(e.attributes))
 				{
-					PopulateCasFromDirsRecursive(fullPath.data, workManager);
+					ScopedWriteLock lock(seenIdsLock);
+					if (!seenIds.insert(e.id).second)
+						return;
+					lock.Leave();
+					PopulateCasFromDirsRecursive(fullPath.data, workManager, seenIds, seenIdsLock);
 					return;
 				}
 
diff --git a/Engine/Source/Programs/UnrealBuildAccelerator/Common/Public/UbaStorageClient.h b/Engine/Source/Programs/UnrealBuildAccelerator/Common/Public/UbaStorageClient.h
--- a/Engine/Source/Programs/UnrealBuildAccelerator/Common/Public/UbaStorageClient.h
+++ b/Engine/Source/Programs/UnrealBuildAccelerator/Common/Public/UbaStorageClient.h
@@ -54,7 +54,7 @@ namespace uba
 
 	private:
 		bool SendFile(const CasKey& casKey, const tchar* fileName, u8* sourceMem, u64 sourceSize, const tchar* hint);
-		bool PopulateCasFromDirsRecursive(const tchar* dir, WorkManager& workManager);
+		bool PopulateCasFromDirsRecursive(const tchar* dir, WorkManager& workManager, UnorderedSet<u64>& seenIds, ReaderWriterLock& seenIdsLock);
 
 		NetworkClient& m_client;
 		bool m_sendCompressed;
diff --git a/Engine/Source/Programs/UnrealBuildAccelerator/Detours/Private/PosixOS/UbaDetoursFunctionsPosix.cpp b/Engine/Source/Programs/UnrealBuildAccelerator/Detours/Private/PosixOS/UbaDetoursFunctionsPosix.cpp
--- a/Engine/Source/Programs/UnrealBuildAccelerator/Detours/Private/PosixOS/UbaDetoursFunctionsPosix.cpp
+++ b/Engine/Source/Programs/UnrealBuildAccelerator/Detours/Private/PosixOS/UbaDetoursFunctionsPosix.cpp
@@ -1790,13 +1790,14 @@ namespace uba
 		}
 		#endif
 
-		g_exeDir.count = GetProcessExecutablePath(g_exeDir.data, g_exeDir.capacity);
-		if (g_exeDir.count == -1)
-		{
-			UBA_ASSERT(false);
-		}
-		char* lastSlash = strrchr(g_exeDir.data, '/');
-		g_exeDir.Resize(lastSlash - g_exeDir.data);
+		StringBuffer<> exePath;
+		exePath.count = GetProcessExecutablePath(exePath.data, exePath.capacity);
+		UBA_ASSERT(exePath.count > 0);
+		char* lastSlash = strrchr(exePath.data, '/');
+		UBA_ASSERT(lastSlash);
+		exePath.Resize(lastSlash - exePath.data);
+		FixPath(g_exeDir, exePath.data);
+		g_exeDir.EnsureEndsWithSlash();
 	}
 
 	void Init()
diff --git a/Engine/Source/Programs/UnrealBuildAccelerator/Test/Private/UbaTestBasics.h b/Engine/Source/Programs/UnrealBuildAccelerator/Test/Private/UbaTestBasics.h
--- a/Engine/Source/Programs/UnrealBuildAccelerator/Test/Private/UbaTestBasics.h
+++ b/Engine/Source/Programs/UnrealBuildAccelerator/Test/Private/UbaTestBasics.h
@@ -90,6 +90,10 @@ namespace uba
 		if (!FixPath2(TC("/../Foo"), workingDir, TStrlen(workingDir), buffer, &lengthResult))
 			return logger.Error(TC("FixPath2 should have failed"));
 		UBA_TEST_CHECK(Equals(buffer, TC("/Foo")), "Should not contain ..");
+
+		if (!FixPath2(TC("/usr/bin//clang++"), workingDir, TStrlen(workingDir), buffer, &lengthResult))
+			return logger.Error(TC("FixPath2 should have failed"));
+		UBA_TEST_CHECK(!Contains(buffer, TC("//")), "Should not contain //");
 #endif
 
 		if (!FixPath2(TC("../Foo"), workingDir, TStrlen(workingDir), buffer, &lengthResult))
