[Uba]
* Removed special case code that should not be needed anymore

diff --git a/Engine/Source/Programs/UnrealBuildAccelerator/Common/Private/UbaSession.cpp b/Engine/Source/Programs/UnrealBuildAccelerator/Common/Private/UbaSession.cpp
--- a/Engine/Source/Programs/UnrealBuildAccelerator/Common/Private/UbaSession.cpp
+++ b/Engine/Source/Programs/UnrealBuildAccelerator/Common/Private/UbaSession.cpp
@@ -1665,7 +1665,6 @@ namespace uba
 
 			if (m_runningRemote && !fileName.StartsWith(m_tempPath.data))
 			{
-				UBA_ASSERT(!fileName.EndsWith(TC(".o.tmp")) && !fileName.EndsWith(TC(".obj.tmp")));  // clang outputs a .o.tmp that we don'want to send back
 				ScopedWriteLock lock(m_outputFilesLock);
 				auto insres = m_outputFiles.try_emplace(fileName.data);
 				if (insres.second)
diff --git a/Engine/Source/Programs/UnrealBuildAccelerator/Common/Private/UbaSessionClient.cpp b/Engine/Source/Programs/UnrealBuildAccelerator/Common/Private/UbaSessionClient.cpp
--- a/Engine/Source/Programs/UnrealBuildAccelerator/Common/Private/UbaSessionClient.cpp
+++ b/Engine/Source/Programs/UnrealBuildAccelerator/Common/Private/UbaSessionClient.cpp
@@ -322,22 +322,9 @@ namespace uba
 		StringKey fileNameKey = msg.fileNameKey;
 
 		if ((msg.access & FileAccess_Write) != 0)
-		{
-			if (fileName.EndsWith(TC(".o.tmp")) || fileName.EndsWith(TC(".obj.tmp"))) // This is clang's temporary object file that is moved into place
-			{
-				#if PLATFORM_WINDOWS
-				out.fileName.Append(m_sessionDir).Append(TC("temp\\")).Append(KeyToString(fileNameKey));
-				#else
-				out.fileName.Append(fileName);
-				#endif
-				return true;
-			}
-
 			return Session::CreateFile(out, msg, virtualApplicationDir);
-		}
 	
 		CasKey casKey;
-
 		if (!GetCasKeyForFile(casKey, msg.process.m_id, fileName, fileNameKey))
 			return false;
 
