[Uba]
* Added sigsegv handler to Uba that is initialized on CreateCallbackLogWriter. This is a temporary hack just to try to get info from a seg fault happening on the farm now and then

diff --git a/Engine/Source/Programs/UnrealBuildAccelerator/Core/Public/UbaPlatform.h b/Engine/Source/Programs/UnrealBuildAccelerator/Core/Public/UbaPlatform.h
--- a/Engine/Source/Programs/UnrealBuildAccelerator/Core/Public/UbaPlatform.h
+++ b/Engine/Source/Programs/UnrealBuildAccelerator/Core/Public/UbaPlatform.h
@@ -56,7 +56,7 @@ namespace uba
 	#define	sizeof_array(array) int(sizeof(array)/sizeof(array[0]))
 	#endif
 
-	#if UBA_DEBUG
+	#if UBA_DEBUG || PLATFORM_LINUX
 	#define UBA_ASSERT(x) do { if (x) break; uba::UbaAssert(TC(""), __FILE__, __LINE__, #x, 543221); } while(false)
 	#define UBA_ASSERTF(x, ...) do { if (x) break; uba::StringBuffer<> _buf; _buf.Appendf(__VA_ARGS__); uba::UbaAssert(_buf.data, __FILE__, __LINE__, #x, 543221); } while(false)
 	#define UBA_ASSERT_MESSAGEBOX 0
diff --git a/Engine/Source/Programs/UnrealBuildAccelerator/Host/Private/UbaExports.cpp b/Engine/Source/Programs/UnrealBuildAccelerator/Host/Private/UbaExports.cpp
--- a/Engine/Source/Programs/UnrealBuildAccelerator/Host/Private/UbaExports.cpp
+++ b/Engine/Source/Programs/UnrealBuildAccelerator/Host/Private/UbaExports.cpp
@@ -46,6 +46,13 @@ namespace uba
 
 		NetworkBackend* networkBackend;
 	};
+
+	#if PLATFORM_LINUX
+	void SigsegvHandler(int signo)
+	{
+		UbaAssert("Segmentation fault", "", 0, "", -1);
+	}
+	#endif
 }
 
 extern "C"
@@ -57,6 +64,10 @@ extern "C"
 
 	uba::LogWriter* CreateCallbackLogWriter(uba::CallbackLogWriter::BeginScopeCallback begin, uba::CallbackLogWriter::EndScopeCallback end, uba::CallbackLogWriter::LogCallback log)
 	{
+		#if PLATFORM_LINUX
+		signal(SIGSEGV, uba::SigsegvHandler);
+		#endif
+
 		return new uba::CallbackLogWriter(begin, end, log);
 	}
 
