[UBA]
* Fixes to make macos work properly with remote helpers

diff --git a/Engine/Source/Programs/UnrealBuildAccelerator/Cli/Private/UbaCli.cpp b/Engine/Source/Programs/UnrealBuildAccelerator/Cli/Private/UbaCli.cpp
--- a/Engine/Source/Programs/UnrealBuildAccelerator/Cli/Private/UbaCli.cpp
+++ b/Engine/Source/Programs/UnrealBuildAccelerator/Cli/Private/UbaCli.cpp
@@ -288,7 +288,7 @@ namespace uba
 			bool success = storage.CheckCasContent(DefaultProcessorCount);
 			return success ? 0 : -1;
 		}
-
+#if UBA_USE_AWS
 		if (checkAws)
 		{
 			AWS aws;
@@ -306,7 +306,8 @@ namespace uba
 				logger.Info(TC("Seems like we are not running inside aws."));
 			return 0;
 		}
-
+#endif
+		
 		if (commandType == CommandType_NotSet)
 		{
 			const tchar* errorMsg = argc == 1 ? TC("") : TC("\nERROR: First argument must be command type. Options are 'local,remote or native'");
diff --git a/Engine/Source/Programs/UnrealBuildAccelerator/Common/Private/UbaFile.cpp b/Engine/Source/Programs/UnrealBuildAccelerator/Common/Private/UbaFile.cpp
--- a/Engine/Source/Programs/UnrealBuildAccelerator/Common/Private/UbaFile.cpp
+++ b/Engine/Source/Programs/UnrealBuildAccelerator/Common/Private/UbaFile.cpp
@@ -484,7 +484,7 @@ namespace uba
 			SetLastError(ERROR_FILE_NOT_FOUND);
 			return false;
 		}
-		UBA_ASSERTF(false, TC("RemoveDirectoryW error handling not implemented (%s)"), strerror(errno));
+		UBA_ASSERTF(false, TC("RemoveDirectoryW error handling not implemented (%s): %s"), strerror(errno), pathName);
 		return false;
 #endif
 	}
@@ -649,8 +649,12 @@ namespace uba
 	{
 #if PLATFORM_WINDOWS
 		return ::CreateHardLinkW(newFileName, existingFileName, NULL);
+#else
+#if PLATFORM_MAC
+		int res = symlink(existingFileName, newFileName);
 #else
 		int res = link(existingFileName, newFileName);
+#endif
 		if (res == 0)
 			return true;
 
diff --git a/Engine/Source/Programs/UnrealBuildAccelerator/Common/Private/UbaProcess.cpp b/Engine/Source/Programs/UnrealBuildAccelerator/Common/Private/UbaProcess.cpp
--- a/Engine/Source/Programs/UnrealBuildAccelerator/Common/Private/UbaProcess.cpp
+++ b/Engine/Source/Programs/UnrealBuildAccelerator/Common/Private/UbaProcess.cpp
@@ -419,8 +419,6 @@ namespace uba
 
 		SystemStats::GetGlobal().Add(m_systemStats);
 
-		m_session.ProcessExited(*this, m_processStats.wallTime);
-
 		// For some reason a parent can exit before a child. Need to figure out repro for this but I've seen it happen on ClangEditor win64
 		for (auto& child : m_childProcesses)
 			while (!((ProcessImpl*)child.m_process)->m_hasExited)
@@ -428,6 +426,8 @@ namespace uba
 
 		UBA_ASSERT(!m_parentProcess || !m_parentProcess->m_hasExited);
 
+		m_session.ProcessExited(*this, m_processStats.wallTime);
+
 		m_hasExited = true;
 
 		#if !PLATFORM_WINDOWS
@@ -674,11 +674,13 @@ namespace uba
 					writer.WriteU32(detoursLibLen);
 					writer.WriteBytes(detoursLib, detoursLibLen);
 
+					writer.WriteString(m_realWorkingDir);
+					#if PLATFORM_WINDOWS
 					TString realCommandLine = TC("\"") + process.m_realApplication + TC("\" ") + commandLine;
 					writer.WriteString(realCommandLine);
-					writer.WriteString(m_realWorkingDir);
-
-					#if !PLATFORM_WINDOWS
+					#else
+					writer.WriteString(process.m_realApplication);
+					writer.WriteString(commandLine);
 					writer.WriteU64(process.m_comMemory.handle.uid);
 					writer.WriteU32(process.m_comMemory.offset);
 					writer.WriteString(info.logFile);
@@ -1315,7 +1317,7 @@ namespace uba
 			}
 			Vector<const char*> arguments2;
 			arguments2.reserve(arguments.size() + 2);
-			arguments2.push_back(realApplication);
+			arguments2.push_back(m_virtualApplication.data());
 			for (auto& s : arguments)
 				arguments2.push_back(s.data());
 			arguments2.push_back(nullptr);
diff --git a/Engine/Source/Programs/UnrealBuildAccelerator/Common/Private/UbaSession.cpp b/Engine/Source/Programs/UnrealBuildAccelerator/Common/Private/UbaSession.cpp
--- a/Engine/Source/Programs/UnrealBuildAccelerator/Common/Private/UbaSession.cpp
+++ b/Engine/Source/Programs/UnrealBuildAccelerator/Common/Private/UbaSession.cpp
@@ -1020,9 +1020,12 @@ namespace uba
 					if (seconds <= info.deleteSessionsOlderThanSeconds)
 						return;
 
-					StringBuffer<> sessionDir(sessionsDir);
-					sessionDir.EnsureEndsWithSlash().Append(e.name);
-					DeleteAllFiles(m_logger, sessionDir.data);
+					if (IsDirectory(e.attributes)) // on macos we get a ".ds_store" file created by the os
+					{
+						StringBuffer<> sessionDir(sessionsDir);
+						sessionDir.EnsureEndsWithSlash().Append(e.name);
+						DeleteAllFiles(m_logger, sessionDir.data);
+					}
 				});
 		}
 
diff --git a/Engine/Source/Programs/UnrealBuildAccelerator/Common/Private/UbaStorage.cpp b/Engine/Source/Programs/UnrealBuildAccelerator/Common/Private/UbaStorage.cpp
--- a/Engine/Source/Programs/UnrealBuildAccelerator/Common/Private/UbaStorage.cpp
+++ b/Engine/Source/Programs/UnrealBuildAccelerator/Common/Private/UbaStorage.cpp
@@ -1977,8 +1977,14 @@ namespace uba
 					#if !PLATFORM_WINDOWS
 					if (fileAttributes & S_IXUSR)
 					{
-						int res = chmod(destination, S_IRUSR | S_IWUSR | S_IXUSR); (void)res;
-						UBA_ASSERT(res == 0);
+						struct stat destStat;
+						int res = stat(destination, &destStat);
+						UBA_ASSERTF(res == 0, TC("stat failed (%s) error: %s"), destination, strerror(errno));
+						if ((destStat.st_mode & S_IXUSR) == 0)
+						{
+							res = chmod(destination, S_IRUSR | S_IWUSR | S_IXUSR); (void)res;
+							UBA_ASSERTF(res == 0, TC("chmod failed (%s) error: %s"), destination, strerror(errno));
+						}
 					}
 					#endif
 					return true;
diff --git a/Engine/Source/Programs/UnrealBuildAccelerator/Common/Private/UbaStorageClient.cpp b/Engine/Source/Programs/UnrealBuildAccelerator/Common/Private/UbaStorageClient.cpp
--- a/Engine/Source/Programs/UnrealBuildAccelerator/Common/Private/UbaStorageClient.cpp
+++ b/Engine/Source/Programs/UnrealBuildAccelerator/Common/Private/UbaStorageClient.cpp
@@ -94,7 +94,7 @@ namespace uba
 	bool StorageClient::GetCasFileName(StringBufferBase& out, const CasKey& casKey)
 	{
 		ScopedReadLock tempLock(m_localStorageFilesLock);
-		auto findIt = m_localStorageFiles.find(casKey);
+		auto findIt = m_localStorageFiles.find(AsCompressed(casKey, false));
 		if (findIt != m_localStorageFiles.end())
 		{
 			if (findIt->second.casEntry.mappingHandle.IsValid())
diff --git a/Engine/Source/Programs/UnrealBuildAccelerator/Common/Public/UbaAWS.h b/Engine/Source/Programs/UnrealBuildAccelerator/Common/Public/UbaAWS.h
--- a/Engine/Source/Programs/UnrealBuildAccelerator/Common/Public/UbaAWS.h
+++ b/Engine/Source/Programs/UnrealBuildAccelerator/Common/Public/UbaAWS.h
@@ -6,7 +6,7 @@
 #include "UbaNetworkBackend.h"
 #include "UbaStringBuffer.h"
 
-#define UBA_USE_AWS 1
+#define UBA_USE_AWS !PLATFORM_MAC
 
 #if UBA_USE_AWS
 
@@ -149,4 +149,4 @@ namespace uba
 	};
 }
 
-#endif
\ No newline at end of file
+#endif
diff --git a/Engine/Source/Programs/UnrealBuildAccelerator/Common/Public/UbaDirectoryIterator.h b/Engine/Source/Programs/UnrealBuildAccelerator/Common/Public/UbaDirectoryIterator.h
--- a/Engine/Source/Programs/UnrealBuildAccelerator/Common/Public/UbaDirectoryIterator.h
+++ b/Engine/Source/Programs/UnrealBuildAccelerator/Common/Public/UbaDirectoryIterator.h
@@ -164,7 +164,7 @@ namespace uba
 		DIR* dir = opendir(dirPath2);
 		if (!dir)
 		{
-			if (errno == ENOENT)
+			if (errno == ENOENT || errno == ENOTDIR)
 				return !errorOnNotFound;
 			UBA_ASSERTF(false, "TraverseDir error handling not implemented %s (%s)", dirPath, strerror(errno));
 			return false;
diff --git a/Engine/Source/Programs/UnrealBuildAccelerator/Detours/Private/PosixOS/UbaDetoursFunctionsPosix.cpp b/Engine/Source/Programs/UnrealBuildAccelerator/Detours/Private/PosixOS/UbaDetoursFunctionsPosix.cpp
--- a/Engine/Source/Programs/UnrealBuildAccelerator/Detours/Private/PosixOS/UbaDetoursFunctionsPosix.cpp
+++ b/Engine/Source/Programs/UnrealBuildAccelerator/Detours/Private/PosixOS/UbaDetoursFunctionsPosix.cpp
@@ -16,6 +16,8 @@
 
 #if PLATFORM_LINUX
 #include <sys/prctl.h>
+#else
+#include <crt_externs.h>
 #endif
 
 namespace uba
@@ -94,6 +96,7 @@ using namespace uba;
 	DETOURED_FUNCTION(_Exit) \
 	DETOURED_FUNCTION_DEBUG \
 	DETOURED_FUNCTION_LINUX \
+	DETOURED_FUNCTION_MACOS \
 
 #if UBA_DEBUG && PLATFORM_LINUX
 #define DETOURED_FUNCTION_DEBUG \
@@ -103,6 +106,7 @@ using namespace uba;
 #endif
 
 #if PLATFORM_LINUX
+#define DETOURED_FUNCTION_MACOS
 #define DETOURED_FUNCTION_LINUX \
 	DETOURED_FUNCTION(get_current_dir_name) \
 	DETOURED_FUNCTION(fopen64) \
@@ -119,6 +123,9 @@ using namespace uba;
 
 #else
 #define DETOURED_FUNCTION_LINUX
+#define DETOURED_FUNCTION_MACOS \
+	DETOURED_FUNCTION(_NSGetExecutablePath) \
+
 #endif
 
 #if (PLATFORM_MAC)
@@ -623,7 +630,7 @@ int Shared_stat(const char* funcName, const char* file, struct stat* attr, const
 	{
 		struct stat attr2;
 		int res2 = trueStat(file, &attr2);
-		UBA_ASSERTF(res == res2, "stat: return value differs for %s (%i vs %i)", file, res, res2);
+		UBA_ASSERTF(res == res2, "stat: return value differs for %s (%i vs %i) [fixed: %s]", file, res, res2, fixedFile.data);
 		if (res != -1)
 		{
 			bool isDir = S_ISDIR(attr->st_mode);
@@ -647,31 +654,35 @@ int Shared_stat(const char* funcName, const char* file, struct stat* attr, const
 
 
 // Detoured functions
-// #if PLATFORM_MAC
-// extern "C" {
-// char* realpath$DARWIN_EXTSN(const char* path, char* resolved_path);
-// }
-// UBA_EXPORT int UBA_WRAPPER(_NSGetExecutablePath)(char* buf, uint32_t* bufsize)
-// {
-// 	int res = TRUE_WRAPPER(_NSGetExecutablePath)(buf, bufsize);
-// 	return TRUE_WRAPPER(_NSGetExecutablePath)(buf, bufsize);
-// }
-
-// UBA_EXPORT int UBA_WRAPPER(_NSGetEnviron)(char* buf, uint32_t* bufsize)
-// {
-// 	int res = TRUE_WRAPPER(_NSGetEnviron)(buf, bufsize);
-// 	printf("%s for %s\n", __func__, buf);
-// 	return TRUE_WRAPPER(_NSGetEnviron)(buf, bufsize);
-// }
-// extern "C" char* UBA_WRAPPER(realpath$DARWIN_EXTSN)(const char* path, char* resolved_path)
-// {
-// 	printf(">>>>>> uba_realpathDARWIN: %s\n", path);
-// 	char* res = TRUE_WRAPPER(realpath$DARWIN_EXTSN)(path, resolved_path);
-// 	printf("<<<<< uba_realpathDARWIN: %s\n", strlen(resolved_path) > 0 ? resolved_path : "(NULL)");
-
-// 	return res;
-// }
-// #endif
+#if PLATFORM_MAC
+//extern "C" {
+//char* realpath$DARWIN_EXTSN(const char* path, char* resolved_path);
+//}
+UBA_EXPORT int UBA_WRAPPER(_NSGetExecutablePath)(char* buf, uint32_t* bufsize)
+{
+	if (!g_isDetouring)
+	{
+		return TRUE_WRAPPER(_NSGetExecutablePath)(buf, bufsize);
+	}
+	memcpy(buf, g_virtualApplication.data, g_virtualApplication.count + 1);
+	return 0;
+}
+
+//UBA_EXPORT int UBA_WRAPPER(_NSGetEnviron)(char* buf, uint32_t* bufsize)
+//{
+//int res = TRUE_WRAPPER(_NSGetEnviron)(buf, bufsize);
+//printf("%s for %s\n", __func__, buf);
+//return TRUE_WRAPPER(_NSGetEnviron)(buf, bufsize);
+//}
+//extern "C" char* UBA_WRAPPER(realpath$DARWIN_EXTSN)(const char* path, char* resolved_path)
+//{
+//printf(">>>>>> uba_realpathDARWIN: %s\n", path);
+//char* res = TRUE_WRAPPER(realpath$DARWIN_EXTSN)(path, resolved_path);
+//printf("<<<<< uba_realpathDARWIN: %s\n", strlen(resolved_path) > 0 ? resolved_path : "(NULL)");
+//
+//return res;
+//}
+#endif
 
 
 UBA_EXPORT int UBA_WRAPPER(chdir)(const char* path)
@@ -1122,16 +1133,22 @@ UBA_EXPORT int UBA_WRAPPER(access)(const char* pathname, int mode)
 		return res;
 	}
 
-	//errno = 0;
-	//auto res2 = TRUE_WRAPPER(access)(realName, mode);
-	//int eo = errno;
-
 	int res = attr.lastError == 0 ? 0 : -1;
-	DEBUG_LOG_DETOURED("access", "%s %i (%s) -> %i %s", pathname, mode, realName, res, StrError(res, attr.lastError));
-
-	//UBA_ASSERTF(res2 == res, "MISMATCH OF RESULTS for %s - %i %i", realName, res2, res);
-	//UBA_ASSERTF(eo == attr.lastError, "MISMATCH OF ERRORS for %s - %i %i", realName, eo, attr.lastError);
-
+	
+	#if UBA_DEBUG_VALIDATE
+	if (!g_runningRemote)
+	{
+		//errno = 0;
+		auto res2 = TRUE_WRAPPER(access)(realName, mode);
+		//int eo = errno;
+		
+		DEBUG_LOG_DETOURED("access", "%s %i (%s) -> %i %s", pathname, mode, realName, res, StrError(res, attr.lastError));
+		
+		UBA_ASSERTF(res2 == res, "MISMATCH OF RESULTS for %s - %i %i (err = %s) (exedir %s)", realName, res2, res, StrError(res, attr.lastError), g_exeDir.data);
+		//UBA_ASSERTF(eo == attr.lastError, "MISMATCH OF ERRORS for %s - %i %i", realName, eo, attr.lastError);
+	}
+	#endif
+	
 	errno = attr.lastError;
 	return res;
 }
@@ -1390,7 +1407,8 @@ UBA_EXPORT int UBA_WRAPPER(posix_spawn)(pid_t* pid, const char* path, const posi
 			cmdLine.append(" ");
 		cmdLine.append(argv[i]);
 	}
-
+	
+	TString realApplication;
 	TString commandLine;
 	u32 processId = 0;
 	StringBuffer<512> currentDir;
@@ -1429,11 +1447,12 @@ UBA_EXPORT int UBA_WRAPPER(posix_spawn)(pid_t* pid, const char* path, const posi
 		u32 dllNameSize = reader.ReadU32();
 		reader.Skip(dllNameSize);
 
-		commandLine = reader.ReadString();
-
 		currentDir.Append("UBA_CWD=");
 		reader.ReadString(currentDir);
 
+		realApplication = reader.ReadString();
+		commandLine = reader.ReadString();
+
 		comIdVar.Append("UBA_COMID=").AppendValue(reader.ReadU64()).Append('+').AppendValue(reader.ReadU32());
 
 		logFile.Append("UBA_LOGFILE=");
@@ -1443,7 +1462,10 @@ UBA_EXPORT int UBA_WRAPPER(posix_spawn)(pid_t* pid, const char* path, const posi
 	std::vector<const char*> envvars;
 	{
 		for (u32 i = 0; envp[i]; ++i)
-			envvars.push_back(envp[i]);
+		{
+//			if (envp[i][0] == 'D' || envp[i][0] == 'P')
+				envvars.push_back(envp[i]);
+		}
 		envvars.push_back(comIdVar.data);
 		envvars.push_back(currentDir.data);
 		envvars.push_back(rulesStr.data);
@@ -1453,13 +1475,21 @@ UBA_EXPORT int UBA_WRAPPER(posix_spawn)(pid_t* pid, const char* path, const posi
 		envvars.push_back(nullptr);
 	}
 
+//	int i=0;
+//	printf("spawnng %s\n", envvars.data()[i]);
+//	while (envvars.data()[i])
+//	{
+//		printf("env: %s\n", envvars.data()[i]);
+//		i++;
+//	}
+	
 	#if UBA_DEBUG_LOG_ENABLED
-	DEBUG_LOG_TRUE("posix_spawn", "%s", path);
+	DEBUG_LOG_TRUE("posix_spawn", "%s", realApplication.data());
 	for (u32 i = 0; argv[i]; ++i)
 		DEBUG_LOG("            %s", argv[i]);
 	#endif
 
-	int res = TRUE_WRAPPER(posix_spawn)(pid, path, file_actions, attrp, argv, (char**)envvars.data());
+	int res = TRUE_WRAPPER(posix_spawn)(pid, realApplication.data(), file_actions, attrp, argv, (char**)envvars.data());
 	bool success = res == 0;
 
 	{
@@ -1763,6 +1793,7 @@ namespace uba
 		if (g_debugFile != InvalidFileHandle)
 		{
 			char buf[32 * 1024];
+#if PLATFORM_LINUX
 			int fd = TRUE_WRAPPER(open)("/proc/self/cmdline", O_RDONLY);
 			if (fd != -1)
 			{
@@ -1787,6 +1818,16 @@ namespace uba
 					DEBUG_LOG("", it);
 				}
 			}
+#else
+			int ArgC = *_NSGetArgc();
+			char** ArgV = *_NSGetArgv();
+			DEBUG_LOG("Executable: %s", ArgV[0]);
+			for (int Arg = 1; Arg < ArgC; Arg++)
+			{
+				DEBUG_LOG("               %s", ArgV[Arg]);
+			}
+			DEBUG_LOG("");
+#endif
 		}
 		#endif
 
