[UBA]
* Fixed bug that could cause an infinite build hang if agent with multiple tcp connections disconnected while retreiving files

diff --git a/Engine/Source/Programs/UnrealBuildAccelerator/Common/Private/UbaNetworkServer.cpp b/Engine/Source/Programs/UnrealBuildAccelerator/Common/Private/UbaNetworkServer.cpp
--- a/Engine/Source/Programs/UnrealBuildAccelerator/Common/Private/UbaNetworkServer.cpp
+++ b/Engine/Source/Programs/UnrealBuildAccelerator/Common/Private/UbaNetworkServer.cpp
@@ -346,6 +346,11 @@ namespace uba
 		#endif
 	}
 
+	bool ConnectionInfo::ShouldDisconnect() const
+	{
+		return ((NetworkServer::Connection*)internalData)->m_shouldDisconnect;
+	}
+
 	void NetworkServer::Worker::ThreadWorker(NetworkServer& server)
 	{
 		u32 writeMemSize = server.m_sendSize;
diff --git a/Engine/Source/Programs/UnrealBuildAccelerator/Common/Private/UbaStorageServer.cpp b/Engine/Source/Programs/UnrealBuildAccelerator/Common/Private/UbaStorageServer.cpp
--- a/Engine/Source/Programs/UnrealBuildAccelerator/Common/Private/UbaStorageServer.cpp
+++ b/Engine/Source/Programs/UnrealBuildAccelerator/Common/Private/UbaStorageServer.cpp
@@ -129,6 +129,8 @@ namespace uba
 					ScopedWriteLock entryLock(store.casEntry->lock);
 					store.casEntry->verified = false;
 					store.casEntry->beingWritten = false;
+					if (m_traceStore)
+						m_trace->FileEndStore(clientId, store.casEntry->key);
 				}
 
 				m_casDataBuffer.UnmapView(store.mappedView, TC("OnDisconnected"));
@@ -156,6 +158,9 @@ namespace uba
 						i->event.Set();
 				}
 
+				if (m_traceFetch)
+					m_trace->FileEndFetch(clientId, AsCompressed(fetch.casKey, m_storeCompressed));
+
 				CloseFile(nullptr, fetch.readFileHandle);
 				it = m_activeFetches.erase(it);
 			}
@@ -411,6 +416,7 @@ namespace uba
 				}
 				
 				FileHandle readFileHandle = InvalidFileHandle;
+				auto rfg = MakeGuard([&]() { CloseFile(nullptr, readFileHandle); });
 				u64 fileSize;
 				u8* memoryBegin = nullptr;
 				u8* memoryPos = nullptr;
@@ -504,7 +510,6 @@ namespace uba
 				if (!left)
 				{
 					*fetchId = u16(~0);
-					CloseFile(nullptr, readFileHandle);
 					u64 sendCasTime = GetTime() - start;
 					stats.sendCas.Add(Timer{sendCasTime, 1});
 					return true;
@@ -512,6 +517,7 @@ namespace uba
 
 				mvg.Cancel();
 				cg.Cancel();
+				rfg.Cancel();
 
 				*fetchId = PopId();
 
@@ -556,7 +562,17 @@ namespace uba
 						fetch.firstItem = &item;
 						item.event.Create(true);
 						lock1.Leave();
-						item.event.IsSet();
+
+						bool shouldExit = false;
+						while (!item.event.IsSet(5000))
+						{
+							// There is a rare chance that connection got disconnected while we have received message with later index but not earlier (can happen when client has multiple tcp connections)
+							// In that case this will lock up forever unless we check the connection to see if we are still connected
+							if (!connectionInfo.ShouldDisconnect())
+								continue;
+							shouldExit = true;
+							break;
+						}
 
 						ScopedWriteLock lock2(fetch.readIndexLock);
 						if (item.prev)
@@ -565,6 +581,9 @@ namespace uba
 							fetch.firstItem = item.next;
 						if (item.next)
 							item.next->prev = item.prev;
+
+						if (shouldExit)
+							return false;
 					}
 				}
 
diff --git a/Engine/Source/Programs/UnrealBuildAccelerator/Common/Public/UbaNetworkServer.h b/Engine/Source/Programs/UnrealBuildAccelerator/Common/Public/UbaNetworkServer.h
--- a/Engine/Source/Programs/UnrealBuildAccelerator/Common/Public/UbaNetworkServer.h
+++ b/Engine/Source/Programs/UnrealBuildAccelerator/Common/Public/UbaNetworkServer.h
@@ -21,6 +21,7 @@ namespace uba
 		const Guid& GetUid() const;
 		u32 GetId() const;
 		bool GetName(StringBufferBase& out) const;
+		bool ShouldDisconnect() const;
 		void* internalData = nullptr;
 	};
 
