[UBA]
* Fixed visualizer crash

diff --git a/Engine/Source/Programs/UnrealBuildAccelerator/Visualizer/Private/UbaVisualizer.cpp b/Engine/Source/Programs/UnrealBuildAccelerator/Visualizer/Private/UbaVisualizer.cpp
--- a/Engine/Source/Programs/UnrealBuildAccelerator/Visualizer/Private/UbaVisualizer.cpp
+++ b/Engine/Source/Programs/UnrealBuildAccelerator/Visualizer/Private/UbaVisualizer.cpp
@@ -1648,7 +1648,8 @@ namespace uba
 						}
 					}
 
-					lastStop = Max(lastStop, processor.processes.rbegin()->stop);
+					if (!processor.processes.empty())
+						lastStop = Max(lastStop, processor.processes.rbegin()->stop);
 
 					posY += stepY;
 					++processorIndex;
