[UBA]
* Cleanup, no logical changes

diff --git a/Engine/Source/Programs/UnrealBuildAccelerator/Common/Private/UbaFile.cpp b/Engine/Source/Programs/UnrealBuildAccelerator/Common/Private/UbaFile.cpp
--- a/Engine/Source/Programs/UnrealBuildAccelerator/Common/Private/UbaFile.cpp
+++ b/Engine/Source/Programs/UnrealBuildAccelerator/Common/Private/UbaFile.cpp
@@ -504,7 +504,7 @@ namespace uba
 			SetLastError(ERROR_FILE_NOT_FOUND);
 			return false;
 	}
-		UBA_ASSERTF(false, TC("DeleteFileW error handling not implemented %s"), strerror(errno));
+		UBA_ASSERTF(false, TC("DeleteFileW failed on %s - Error handling not implemented (%s)"), fileName, strerror(errno));
 		return false;
 #endif
 	}
diff --git a/Engine/Source/Programs/UnrealBuildAccelerator/Common/Private/UbaProcess.cpp b/Engine/Source/Programs/UnrealBuildAccelerator/Common/Private/UbaProcess.cpp
--- a/Engine/Source/Programs/UnrealBuildAccelerator/Common/Private/UbaProcess.cpp
+++ b/Engine/Source/Programs/UnrealBuildAccelerator/Common/Private/UbaProcess.cpp
@@ -37,22 +37,6 @@ extern char **environ;
 
 namespace uba
 {
-	#if PLATFORM_WINDOWS
-	ReaderWriterLock g_envLock;
-	void* g_env;
-	#endif
-
-	TString GetDirectoryName(const TString& file)
-	{
-		size_t nameIndex = file.find_last_of('\\');
-		size_t lastSlash = file.find_last_of('/', nameIndex + 1);
-		if (lastSlash != -1)
-			nameIndex = lastSlash;
-		if (nameIndex == -1)
-			return TC("");
-		return file.substr(0, nameIndex);
-	}
-
 	void Process::AddRef()
 	{
 		++m_refCount;
@@ -129,21 +113,23 @@ namespace uba
 		m_startInfo.description = m_description.c_str();
 
 		m_virtualApplication = startInfo.application;
-		Replace(m_virtualApplication.data(), '/', PathSeparator);
+		FixPathSeparators(m_virtualApplication.data());
 		m_startInfo.application = m_virtualApplication.c_str();
 
 		m_arguments = startInfo.arguments;
 		m_startInfo.arguments = m_arguments.c_str();
 
 		m_virtualWorkingDir = startInfo.workingDir;
-		Replace(m_virtualWorkingDir.data(), '/', PathSeparator);
+		FixPathSeparators(m_virtualWorkingDir.data());
 		m_startInfo.workingDir = m_virtualWorkingDir.c_str();
 
 		m_logFile = startInfo.logFile;
-		Replace(m_logFile.data(), '/', PathSeparator);
+		FixPathSeparators(m_logFile.data());
 		m_startInfo.logFile = m_logFile.c_str();
 
-		m_virtualApplicationDir = GetDirectoryName(startInfo.application);
+		size_t nameIndex = m_virtualApplication.find_last_of(PathSeparator);
+		if (nameIndex != -1)
+			m_virtualApplicationDir = m_virtualApplication.substr(0, nameIndex);
 
 		m_realApplication = std::move(realApplication);
 		m_realWorkingDir = realWorkingDir;
@@ -481,7 +467,7 @@ namespace uba
 					msg.access = (FileAccess)reader.ReadByte();
 
 					CreateFileResponse response;
-					m_messageSuccess = m_session.CreateFile(response, msg, m_virtualApplicationDir.c_str()) && m_messageSuccess;
+					m_messageSuccess = m_session.CreateFile(response, msg) && m_messageSuccess;
 					writer.WriteString(response.fileName);
 					writer.WriteU64(response.size);
 					writer.WriteU32(response.closeId);
@@ -496,7 +482,7 @@ namespace uba
 					reader.ReadString(msg.fileName);
 					msg.fileNameKey = reader.ReadStringKey();
 					GetFullFileNameResponse response;
-					m_messageSuccess = m_session.GetFullFileName(response, msg, m_virtualApplicationDir.c_str()) && m_messageSuccess;
+					m_messageSuccess = m_session.GetFullFileName(response, msg) && m_messageSuccess;
 					writer.WriteString(response.fileName);
 					writer.WriteString(response.virtualFileName);
 					writer.WriteU32(response.mappedFileTableSize);
@@ -698,7 +684,7 @@ namespace uba
 					u32 lastError = reader.ReadU32();
 					if (!result)
 					{
-						m_session.m_logger.Logf(LogEntryType_Info, TC("DetourCreateProcessWithDllEx for child process failed - %s. %s (Working dir: %s)"), LastErrorToText(lastError).data, process.m_realApplication.c_str(), process.m_realWorkingDir);
+						m_session.m_logger.Logf(LogEntryType_Info, TC("Detoured process failed to start child process - %s. %s (Working dir: %s)"), LastErrorToText(lastError).data, process.m_realApplication.c_str(), process.m_realWorkingDir);
 						process.m_waitForParent.Set();
 						return true;
 					}
@@ -722,7 +708,6 @@ namespace uba
 					u64 nativeThreadHandle = reader.ReadU64();
 					process.m_nativeProcessHandle = (ProcHandle)nativeProcessHandle;
 					process.m_nativeProcessId = nativeProcessId;
-					//m_session.m_logger.Info("Got StartProcess with pid %u", nativeProcessId);
 #endif
 					process.m_waitForParent.Set();
 					return true;
diff --git a/Engine/Source/Programs/UnrealBuildAccelerator/Common/Private/UbaSession.cpp b/Engine/Source/Programs/UnrealBuildAccelerator/Common/Private/UbaSession.cpp
--- a/Engine/Source/Programs/UnrealBuildAccelerator/Common/Private/UbaSession.cpp
+++ b/Engine/Source/Programs/UnrealBuildAccelerator/Common/Private/UbaSession.cpp
@@ -1616,7 +1616,7 @@ namespace uba
 		return m_environmentVariables.data();
 	}
 
-	bool Session::CreateFile(CreateFileResponse& out, const CreateFileMessage& msg, const tchar* virtualApplicationDir)
+	bool Session::CreateFile(CreateFileResponse& out, const CreateFileMessage& msg)
 	{
 		const StringBufferBase& fileName = msg.fileName;
 		const StringKey& fileNameKey = msg.fileNameKey;
@@ -1918,7 +1918,7 @@ namespace uba
 		return true;
 	}
 
-	bool Session::GetFullFileName(GetFullFileNameResponse& out, const GetFullFileNameMessage& msg, const tchar* virtualApplicationDir)
+	bool Session::GetFullFileName(GetFullFileNameResponse& out, const GetFullFileNameMessage& msg)
 	{
 		UBA_ASSERTF(false, TC("SHOULD NOT HAPPEN (only remote)"));
 		return false;
diff --git a/Engine/Source/Programs/UnrealBuildAccelerator/Common/Private/UbaSessionClient.cpp b/Engine/Source/Programs/UnrealBuildAccelerator/Common/Private/UbaSessionClient.cpp
--- a/Engine/Source/Programs/UnrealBuildAccelerator/Common/Private/UbaSessionClient.cpp
+++ b/Engine/Source/Programs/UnrealBuildAccelerator/Common/Private/UbaSessionClient.cpp
@@ -316,13 +316,13 @@ namespace uba
 		return m_storage.CopyOrLink(casKey, out.data, fileAttributes);
 	}
 
-	bool SessionClient::CreateFile(CreateFileResponse& out, const CreateFileMessage& msg, const tchar* virtualApplicationDir)
+	bool SessionClient::CreateFile(CreateFileResponse& out, const CreateFileMessage& msg)
 	{
 		const StringBufferBase& fileName = msg.fileName;
 		StringKey fileNameKey = msg.fileNameKey;
 
 		if ((msg.access & FileAccess_Write) != 0)
-			return Session::CreateFile(out, msg, virtualApplicationDir);
+			return Session::CreateFile(out, msg);
 	
 		CasKey casKey;
 		if (!GetCasKeyForFile(casKey, msg.process.m_id, fileName, fileNameKey))
@@ -574,7 +574,7 @@ namespace uba
 		writeMsg.fileNameKey = msg.toKey;
 		writeMsg.access = FileAccess_Write;
 		CreateFileResponse writeOut;
-		if (!CreateFile(writeOut, writeMsg, nullptr))
+		if (!CreateFile(writeOut, writeMsg))
 			return false;
 
 		out.toName.Append(writeOut.fileName);
@@ -667,7 +667,7 @@ namespace uba
 		return true;
 	}
 
-	bool SessionClient::GetFullFileName(GetFullFileNameResponse& out, const GetFullFileNameMessage& msg, const tchar* virtualApplicationDir)
+	bool SessionClient::GetFullFileName(GetFullFileNameResponse& out, const GetFullFileNameMessage& msg)
 	{
 		ScopedWriteLock lock(m_nameToNameLookupLock);
 		auto insres = m_nameToNameLookup.try_emplace(msg.fileName.data);
@@ -683,7 +683,8 @@ namespace uba
 		}
 		rec.handled = true;
 
-		if (!EnsureBinaryFile(out.fileName, out.virtualFileName, msg.process.m_id, msg.fileName, msg.fileNameKey, virtualApplicationDir))
+		auto& dir = msg.process.m_virtualApplicationDir;
+		if (!EnsureBinaryFile(out.fileName, out.virtualFileName, msg.process.m_id, msg.fileName, msg.fileNameKey, dir.c_str()))
 			return false;
 
 		rec.name = out.fileName.data;
diff --git a/Engine/Source/Programs/UnrealBuildAccelerator/Common/Private/UbaSessionServer.cpp b/Engine/Source/Programs/UnrealBuildAccelerator/Common/Private/UbaSessionServer.cpp
--- a/Engine/Source/Programs/UnrealBuildAccelerator/Common/Private/UbaSessionServer.cpp
+++ b/Engine/Source/Programs/UnrealBuildAccelerator/Common/Private/UbaSessionServer.cpp
@@ -1405,7 +1405,7 @@ namespace uba
 		return true;
 	}
 
-	bool SessionServer::CreateFile(CreateFileResponse& out, const CreateFileMessage& msg, const tchar* virtualApplicationDir)
+	bool SessionServer::CreateFile(CreateFileResponse& out, const CreateFileMessage& msg)
 	{
 		if ((msg.access & FileAccess_Write) == 0)
 		{
@@ -1426,7 +1426,7 @@ namespace uba
 				return true;
 			}
 		}
-		return Session::CreateFile(out, msg, virtualApplicationDir);
+		return Session::CreateFile(out, msg);
 	}
 
 	void SessionServer::FileEntryAdded(StringKey fileNameKey, u64 lastWritten, u64 size)
diff --git a/Engine/Source/Programs/UnrealBuildAccelerator/Common/Public/UbaSession.h b/Engine/Source/Programs/UnrealBuildAccelerator/Common/Public/UbaSession.h
--- a/Engine/Source/Programs/UnrealBuildAccelerator/Common/Public/UbaSession.h
+++ b/Engine/Source/Programs/UnrealBuildAccelerator/Common/Public/UbaSession.h
@@ -145,14 +145,14 @@ namespace uba
 		virtual void PrintSessionStats(Logger& logger);
 
 		virtual bool GetInitResponse(InitResponse& out, const InitMessage& msg);
-		virtual bool CreateFile(CreateFileResponse& out, const CreateFileMessage& msg, const tchar* virtualApplicationDir);
+		virtual bool CreateFile(CreateFileResponse& out, const CreateFileMessage& msg);
 		virtual bool CloseFile(CloseFileResponse& out, const CloseFileMessage& msg);
 		virtual bool DeleteFile(DeleteFileResponse& out, const DeleteFileMessage& msg);
 		virtual bool CopyFile(CopyFileResponse& out, const CopyFileMessage& msg);
 		virtual bool MoveFile(MoveFileResponse& out, const MoveFileMessage& msg);
 		virtual bool Chmod(ChmodResponse& out, const ChmodMessage& msg);
 		virtual bool CreateDirectory(CreateDirectoryResponse& out, const CreateDirectoryMessage& msg);
-		virtual bool GetFullFileName(GetFullFileNameResponse& out, const GetFullFileNameMessage& msg, const tchar* virtualApplicationDir);
+		virtual bool GetFullFileName(GetFullFileNameResponse& out, const GetFullFileNameMessage& msg);
 		virtual bool GetListDirectoryInfo(ListDirectoryResponse& out, tchar* dirName, const StringKey& dirKey);
 		virtual bool WriteFileToDisk(ProcessImpl& process, WrittenFile& file);
 		virtual bool AllocFailed(Process& process, const tchar* allocType, u32 error);
diff --git a/Engine/Source/Programs/UnrealBuildAccelerator/Common/Public/UbaSessionClient.h b/Engine/Source/Programs/UnrealBuildAccelerator/Common/Public/UbaSessionClient.h
--- a/Engine/Source/Programs/UnrealBuildAccelerator/Common/Public/UbaSessionClient.h
+++ b/Engine/Source/Programs/UnrealBuildAccelerator/Common/Public/UbaSessionClient.h
@@ -44,13 +44,13 @@ namespace uba
 		virtual bool PrepareProcess(const ProcessStartInfo& startInfo, bool isChild, StringBufferBase& outRealApplication, const tchar*& outRealWorkingDir) override;
 		virtual u64 GetMemoryMapAlignment(const tchar* fileName, u64 fileNameLen) const override;
 		virtual void* GetProcessEnvironmentVariables() override;
-		virtual bool CreateFile(CreateFileResponse& out, const CreateFileMessage& msg, const tchar* virtualApplicationDir) override;
+		virtual bool CreateFile(CreateFileResponse& out, const CreateFileMessage& msg) override;
 		virtual bool DeleteFile(DeleteFileResponse& out, const DeleteFileMessage& msg) override;
 		virtual bool CopyFile(CopyFileResponse& out, const CopyFileMessage& msg) override;
 		virtual bool MoveFile(MoveFileResponse& out, const MoveFileMessage& msg) override;
 		virtual bool Chmod(ChmodResponse& out, const ChmodMessage& msg) override;
 		virtual bool CreateDirectory(CreateDirectoryResponse& out, const CreateDirectoryMessage& msg) override;
-		virtual bool GetFullFileName(GetFullFileNameResponse& out, const GetFullFileNameMessage& msg, const tchar* virtualApplicationDir) override;
+		virtual bool GetFullFileName(GetFullFileNameResponse& out, const GetFullFileNameMessage& msg) override;
 		virtual bool GetListDirectoryInfo(ListDirectoryResponse& out, tchar* dirName, const StringKey& dirKey) override;
 		virtual bool WriteFileToDisk(ProcessImpl& process, WrittenFile& file) override;
 		virtual bool AllocFailed(Process& process, const tchar* allocType, u32 error) override;
diff --git a/Engine/Source/Programs/UnrealBuildAccelerator/Common/Public/UbaSessionServer.h b/Engine/Source/Programs/UnrealBuildAccelerator/Common/Public/UbaSessionServer.h
--- a/Engine/Source/Programs/UnrealBuildAccelerator/Common/Public/UbaSessionServer.h
+++ b/Engine/Source/Programs/UnrealBuildAccelerator/Common/Public/UbaSessionServer.h
@@ -54,7 +54,7 @@ namespace uba
 		ProcessHandle ProcessRemoved(u32 processId);
 
 		virtual bool PrepareProcess(const ProcessStartInfo& startInfo, bool isChild, StringBufferBase& outRealApplication, const tchar*& outRealWorkingDir) override final;
-		virtual bool CreateFile(CreateFileResponse& out, const CreateFileMessage& msg, const tchar* virtualApplicationDir) override final;
+		virtual bool CreateFile(CreateFileResponse& out, const CreateFileMessage& msg) override final;
 		virtual void FileEntryAdded(StringKey fileNameKey, u64 lastWritten, u64 size) override final;
 		virtual void PrintSessionStats(Logger& logger) override final;
 		virtual void TraceSessionUpdate() override final;
diff --git a/Engine/Source/Programs/UnrealBuildAccelerator/Core/Private/UbaStringBuffer.cpp b/Engine/Source/Programs/UnrealBuildAccelerator/Core/Private/UbaStringBuffer.cpp
--- a/Engine/Source/Programs/UnrealBuildAccelerator/Core/Private/UbaStringBuffer.cpp
+++ b/Engine/Source/Programs/UnrealBuildAccelerator/Core/Private/UbaStringBuffer.cpp
@@ -122,6 +122,11 @@ namespace uba
 		}
 	}
 
+	void FixPathSeparators(tchar* str)
+	{
+		Replace(str, NonPathSeparator, PathSeparator);
+	}
+
 	StringBufferBase& StringBufferBase::Append(const tchar* str)
 	{
 		return Append(str, u32(TStrlen(str)));
@@ -271,4 +276,65 @@ namespace uba
 		data[count] = 0;
 		return *this;
 	}
+
+	StringBufferBase& StringBufferBase::FixPathSeparators()
+	{
+		uba::FixPathSeparators(data);
+		return *this;
+	}
+
+	StringBufferBase& StringBufferBase::MakeLower()
+	{
+		for (tchar* it = data; *it; ++it)
+			*it = ToLower(*it);
+		return *this;
+	}
+
+	bool StringBufferBase::Parse(u64& out)
+	{
+		if (!count)
+			return false;
+
+		#if PLATFORM_WINDOWS
+		out = wcstoull(data, nullptr, 10);
+		#else
+		out = strtoull(data, nullptr, 10);
+		#endif
+		return out != 0 || Equals(TC("0"));
+	}
+
+	bool StringBufferBase::Parse(u32& out)
+	{
+		if (!count)
+			return false;
+		#if PLATFORM_WINDOWS
+		out = wcstoul(data, nullptr, 10);
+		#else
+		out = strtoul(data, nullptr, 10);
+		#endif
+		return out != 0 || Equals(TC("0"));
+	}
+
+	bool StringBufferBase::Parse(u16& out)
+	{
+		u32 temp;
+		if (!Parse(temp))
+			return false;
+		if (temp > 65535)
+			return false;
+		out = u16(temp);
+		return true;
+	}
+
+	bool StringBufferBase::Parse(float& out)
+	{
+		if (!count)
+			return false;
+		#if PLATFORM_WINDOWS
+		out = wcstof(data, nullptr);
+		#else
+		out = strtof(data, nullptr);
+		#endif
+		return out != 0 || Equals(TC("0"));
+	}
 }
diff --git a/Engine/Source/Programs/UnrealBuildAccelerator/Core/Public/UbaPlatform.h b/Engine/Source/Programs/UnrealBuildAccelerator/Core/Public/UbaPlatform.h
--- a/Engine/Source/Programs/UnrealBuildAccelerator/Core/Public/UbaPlatform.h
+++ b/Engine/Source/Programs/UnrealBuildAccelerator/Core/Public/UbaPlatform.h
@@ -98,6 +98,7 @@ namespace uba
 #if PLATFORM_WINDOWS
 	inline constexpr bool CaseInsensitiveFs = true;
 	inline constexpr tchar PathSeparator = '\\';
+	inline constexpr tchar NonPathSeparator = '/';
 	inline constexpr u32 MaxPath = 512;
 	inline DWORD ToLow(u64 v) { LARGE_INTEGER li; li.QuadPart = (LONGLONG)v; return li.LowPart; }
 	inline LONG ToHigh(u64 v) { LARGE_INTEGER li; li.QuadPart = (LONGLONG)v; return li.HighPart; }
@@ -115,6 +116,7 @@ namespace uba
 	#define TStrdup _wcsdup
 #else
 	inline constexpr tchar PathSeparator = '/';
+	inline constexpr tchar NonPathSeparator = '\\';
 	inline constexpr u32 MaxPath = 512;
 	#define TStrlen(s) u32(strlen(s))
 	#define TStrchr(a, b) strchr(a, b)
diff --git a/Engine/Source/Programs/UnrealBuildAccelerator/Core/Public/UbaStringBuffer.h b/Engine/Source/Programs/UnrealBuildAccelerator/Core/Public/UbaStringBuffer.h
--- a/Engine/Source/Programs/UnrealBuildAccelerator/Core/Public/UbaStringBuffer.h
+++ b/Engine/Source/Programs/UnrealBuildAccelerator/Core/Public/UbaStringBuffer.h
@@ -12,6 +12,7 @@ namespace uba
 	bool Equals(const tchar* str1, const tchar* str2, bool ignoreCase = true);
 	bool Equals(const tchar* str1, const tchar* str2, u64 count, bool ignoreCase = true);
 	void Replace(tchar* str, tchar from, tchar to);
+	void FixPathSeparators(tchar* str);
 	inline void ToLower(tchar* str) { while (tchar c = *str) { if (c >= 'A' && c <= 'Z') *str = c - 'A' + 'a'; ++str; } }
 	inline tchar ToLower(tchar c) { return (c >= 'A' && c <= 'Z') ? (c - 'A' + 'a') : c;}
 	inline tchar ToUpper(tchar c) { return (c >= 'a' && c <= 'z') ? (c - 'a' + 'A') : c;}
@@ -52,61 +53,13 @@ namespace uba
 		inline StringBufferBase& Replace(tchar from, tchar to) { uba::Replace(data, from, to);  return *this; }
 
 		StringBufferBase& EnsureEndsWithSlash();
+		StringBufferBase& FixPathSeparators();
+		StringBufferBase& MakeLower();
 
-		StringBufferBase& MakeLower()
-		{
-			for (tchar* it = data; *it; ++it)
-				*it = ToLower(*it);
-			return *this;
-		}
-
-		bool Parse(u64& out)
-		{
-			if (!count)
-				return false;
-
-			#if PLATFORM_WINDOWS
-			out = wcstoull(data, nullptr, 10);
-			#else
-			out = strtoull(data, nullptr, 10);
-			#endif
-			return out != 0 || Equals(TC("0"));
-		}
-
-		bool Parse(u32& out)
-		{
-			if (!count)
-				return false;
-			#if PLATFORM_WINDOWS
-			out = wcstoul(data, nullptr, 10);
-			#else
-			out = strtoul(data, nullptr, 10);
-			#endif
-			return out != 0 || Equals(TC("0"));
-		}
-
-		bool Parse(u16& out)
-		{
-			u32 temp;
-			if (!Parse(temp))
-				return false;
-			if (temp > 65535)
-				return false;
-			out = u16(temp);
-			return true;
-		}
-
-		bool Parse(float& out)
-		{
-			if (!count)
-				return false;
-			#if PLATFORM_WINDOWS
-			out = wcstof(data, nullptr);
-			#else
-			out = strtof(data, nullptr);
-			#endif
-			return out != 0 || Equals(TC("0"));
-		}
+		bool Parse(u64& out);
+		bool Parse(u32& out);
+		bool Parse(u16& out);
+		bool Parse(float& out);
 
 		u32 count;
 		u32 capacity;
diff --git a/Engine/Source/Programs/UnrealBuildAccelerator/Detours/Private/Windows/UbaDetoursFunctionsKernelBase.inl b/Engine/Source/Programs/UnrealBuildAccelerator/Detours/Private/Windows/UbaDetoursFunctionsKernelBase.inl
--- a/Engine/Source/Programs/UnrealBuildAccelerator/Detours/Private/Windows/UbaDetoursFunctionsKernelBase.inl
+++ b/Engine/Source/Programs/UnrealBuildAccelerator/Detours/Private/Windows/UbaDetoursFunctionsKernelBase.inl
@@ -2325,8 +2325,7 @@ HMODULE Shared_LoadLibrary(LPCWSTR lpLibFileName, HANDLE hFile, DWORD dwFlags)
 		return NULL;
 
 	StringBuffer<> path;
-	path.Append(lpLibFileName);
-	Replace(path.data, '/', PathSeparator);
+	path.Append(lpLibFileName).FixPathSeparators();
 
 	bool detourDll = path.EndsWith(L".exe") || path.EndsWith(L".dll");
 	if (detourDll && path.StartsWith(g_systemRoot.data))
