[UBA]
* Fixed static analysis warning

diff --git a/Engine/Source/Programs/UnrealBuildAccelerator/Detours/Private/Windows/UbaDetoursFunctionsKernelBase.inl b/Engine/Source/Programs/UnrealBuildAccelerator/Detours/Private/Windows/UbaDetoursFunctionsKernelBase.inl
--- a/Engine/Source/Programs/UnrealBuildAccelerator/Detours/Private/Windows/UbaDetoursFunctionsKernelBase.inl
+++ b/Engine/Source/Programs/UnrealBuildAccelerator/Detours/Private/Windows/UbaDetoursFunctionsKernelBase.inl
@@ -2503,7 +2503,6 @@ BOOL Detoured_CreateProcessW(LPCWSTR lpApplicationName, LPWSTR lpCommandLine, LP
 	TString commandLine;
 	TString currentDir;
 	u32 processId = 0;
-	u32 rulesIndex = 0;
 	char dll[1024];
 	{
 		TimerScope ts(g_stats.createProcess);
@@ -2518,7 +2517,7 @@ BOOL Detoured_CreateProcessW(LPCWSTR lpApplicationName, LPWSTR lpCommandLine, LP
 		processId = reader.ReadU32();
 		UBA_ASSERT(processId > 0);
 
-		rulesIndex = reader.ReadU32();
+		reader.Skip(sizeof(u32)); // Rules index
 
 		u32 dllNameSize = reader.ReadU32();
 		reader.ReadBytes(dll, dllNameSize);
