[UBA]
* Changed so UbaHordeSession is not trying to connect to agent until the agent has outputted the "Listening on" message
* Added -populateCasFromXcode argument and prepped implementation in UbaAgent

diff --git a/Engine/Source/Programs/UnrealBuildAccelerator/Agent/Private/UbaAgent.cpp b/Engine/Source/Programs/UnrealBuildAccelerator/Agent/Private/UbaAgent.cpp
--- a/Engine/Source/Programs/UnrealBuildAccelerator/Agent/Private/UbaAgent.cpp
+++ b/Engine/Source/Programs/UnrealBuildAccelerator/Agent/Private/UbaAgent.cpp
@@ -98,6 +98,9 @@ namespace uba
 		logger.Info(TC("  -memkill=<percent>      The amount of memory needed before processes starts to be killed. Set this to 100 to disable. Defaults to 90%%"));
 		logger.Info(TC("  -crypto=<key>           16 bytes crypto key used for secure network transfer"));
 		logger.Info(TC("  -populateCas=<dir>      Prepopulate cas database with files in dir. If files needed exists on machine this can be an optimization"));
+		#if PLATFORM_MAC
+		logger.Info(TC("  -populateCasFromXcode   Prepopulate cas database with files from local xcode installation."));
+		#endif
 		logger.Info(TC(""));
 		return -1;
 	}
@@ -403,6 +406,10 @@ namespace uba
 		bool hasCrypto = false;
 		Vector<TString> populateCasDirs;
 
+		#if PLATFORM_MAC
+		bool populateCasFromXCode = false;
+		#endif
+
 		for (int i=1; i!=argc; ++i)
 		{
 			StringBuffer<> name;
@@ -597,6 +604,12 @@ namespace uba
 					return PrintHelp(TC("-populateCas needs a dir"));
 				populateCasDirs.push_back(value.data);
 			}
+			#if PLATFORM_MAC
+			else if (name.Equals(TC("-populateCasFromXcode")))
+			{
+				populateCasFromXcode = true;
+			}
+			#endif
 			else if (name.Equals(TC("-sentry")))
 			{
 				if (value.IsEmpty())
@@ -730,6 +743,39 @@ namespace uba
 				return -1;
 		}
 
+#if PLATFORM_MAC
+		if (populateCasFromXCode)
+		{
+			// .. do your think Josh.... maybe something like this?
+
+			FILE* xcSelect = popen("/usr/bin/xcode-select -p", "r");
+			if (!xcSelect)
+			{
+				logger.Error("Unable to run /usr/bin/xcode-select. Is xcode installed?");
+				return -1;
+			}
+			StringBuffer<> developerPath;
+			bool success = fgets(developerPath.data, developerPath.capacity, xcSelect) != NULL;
+			pclose(xcSelect);
+
+			if (!success)
+			{
+				logger.Error("Unable to parse string from /usr/bin/xcode-select");
+				return -1;
+			}
+			developerPath.count = strlen(developerPath.data);
+			developerPath.EnsureEndsWithSlash();
+
+			u32 developerPathLen = developerPath.count;
+			const char* subDirs[] = { "Toolchains", "Platforms" };
+			for (auto subDir : subDirs)
+			{
+				developerPath.Resize(developerPathLen).Append(subDir);
+				populateCasDirs.push_back(developerPath.data);
+			}
+		}
+#endif
+
 
 		Vector<ProcessLogLine> logLines[2];
 		u32 logLinesIndex = 0;
diff --git a/Engine/Source/Programs/UnrealBuildTool/Executors/Experimental/UBAHordeSession.cs b/Engine/Source/Programs/UnrealBuildTool/Executors/Experimental/UBAHordeSession.cs
--- a/Engine/Source/Programs/UnrealBuildTool/Executors/Experimental/UBAHordeSession.cs
+++ b/Engine/Source/Programs/UnrealBuildTool/Executors/Experimental/UBAHordeSession.cs
@@ -631,6 +631,11 @@ namespace UnrealBuildTool
 							arguments.Add("-Log");
 						}
 
+						if (OperatingSystem.IsMacOS())
+						{
+							arguments.Add("-populateCasFromXcode");
+						}
+
 						logger.LogInformation("Executing child process: {Executable} {Arguments}", executable, CommandLineArguments.Join(arguments));
 
 						ExecuteProcessFlags execFlags = _allowWine ? ExecuteProcessFlags.UseWine : ExecuteProcessFlags.None;
@@ -640,7 +645,7 @@ namespace UnrealBuildTool
 						string? line;
 						while ((line = await process.ReadLineAsync(cancellationToken)) != null)
 						{
-							if (shouldConnect)
+							if (shouldConnect && line.Contains("Listening on")) // This log entry means that the agent is ready for connections.
 							{
 								long totalMs = self.StartTime.ElapsedMilliseconds;
 								logger.LogInformation("Connecting to UbaAgent on {Ip}:{Port} (local agent port {AgentPort}) {Seconds}.{Milliseconds} seconds after assigned", 
