[UBA]
* Added guards for null function pointers in tcp backend
* Changed signal handler for SIGSEGV to use sigaction
* Added time to log entry when populating cas client side

diff --git a/Engine/Source/Programs/UnrealBuildAccelerator/Common/Private/UbaNetworkBackendTcp.cpp b/Engine/Source/Programs/UnrealBuildAccelerator/Common/Private/UbaNetworkBackendTcp.cpp
--- a/Engine/Source/Programs/UnrealBuildAccelerator/Common/Private/UbaNetworkBackendTcp.cpp
+++ b/Engine/Source/Programs/UnrealBuildAccelerator/Common/Private/UbaNetworkBackendTcp.cpp
@@ -441,13 +441,27 @@ namespace uba
 					break;
 				isFirst = false;
 
-				if (!connection.headerCallback(connection.recvContext, headerData, bodyContext, bodyData, bodySize))
+				auto hc = connection.headerCallback;
+				if (!hc)
+				{
+					m_logger.Error(TC("Header callback not set"));
+					break;
+				}
+
+				if (!hc(connection.recvContext, headerData, bodyContext, bodyData, bodySize))
 					break;
 				if (!bodySize)
 					continue;
 
 				bool success = RecvSocket(m_logger, connection.socket, bodyData, bodySize, connection.recvTimeoutMs, connectionUid, connection.recvHint, TC("Body"), false);
-				if (!connection.bodyCallback(connection.recvContext, !success, headerData, bodyContext, bodyData, bodySize))
+				auto bc = connection.bodyCallback;
+				if (!bc)
+				{
+					m_logger.Error(TC("Body callback not set"));
+					break;
+				}
+
+				if (!bc(connection.recvContext, !success, headerData, bodyContext, bodyData, bodySize))
 					break;
 				if (!success)
 					break;
@@ -462,9 +476,9 @@ namespace uba
 		SOCKET s = connection.socket;
 		connection.socket = INVALID_SOCKET;
 
-		if (connection.disconnectCallback)
+		auto cb = connection.disconnectCallback;
+		if (cb)
 		{
-			auto cb = connection.disconnectCallback;
 			auto context = connection.disconnectContext;
 			connection.disconnectCallback = nullptr;
 			connection.disconnectContext = nullptr;
diff --git a/Engine/Source/Programs/UnrealBuildAccelerator/Common/Private/UbaStorageClient.cpp b/Engine/Source/Programs/UnrealBuildAccelerator/Common/Private/UbaStorageClient.cpp
--- a/Engine/Source/Programs/UnrealBuildAccelerator/Common/Private/UbaStorageClient.cpp
+++ b/Engine/Source/Programs/UnrealBuildAccelerator/Common/Private/UbaStorageClient.cpp
@@ -75,17 +75,20 @@ namespace uba
 		if (directories.empty())
 			return true;
 
+		u64 start = GetTime();
+
 		WorkManagerImpl workManager(workerCount);
 		ReaderWriterLock lock;
 		bool success = true;
 		UnorderedSet<u64> seenIds;
 		ReaderWriterLock seenIdsLock;
+
 		for (auto& dir : directories)
 			success = PopulateCasFromDirsRecursive(dir.c_str(), workManager, seenIds, seenIdsLock) && success;
 		workManager.Wait();
 
 		if (u32 fileCount = u32(m_localStorageFiles.size()))
-			m_logger.Info(TC("Prepopulated %u files to cas"), fileCount);
+			m_logger.Info(TC("Prepopulated %u files to cas in %s"), fileCount, TimeToText(GetTime() - start).str);
 
 		return success;
 	}
diff --git a/Engine/Source/Programs/UnrealBuildAccelerator/Host/Private/UbaExports.cpp b/Engine/Source/Programs/UnrealBuildAccelerator/Host/Private/UbaExports.cpp
--- a/Engine/Source/Programs/UnrealBuildAccelerator/Host/Private/UbaExports.cpp
+++ b/Engine/Source/Programs/UnrealBuildAccelerator/Host/Private/UbaExports.cpp
@@ -47,10 +47,14 @@ namespace uba
 		NetworkBackend* networkBackend;
 	};
 
-	#if PLATFORM_LINUX
-	void SigsegvHandler(int signo)
+	#define UBA_USE_SIGNALHANDLER PLATFORM_LINUX
+
+	#if UBA_USE_SIGNALHANDLER
+	void SignalHandler(int sig, siginfo_t* si, void* unused)
 	{
-		UbaAssert("Segmentation fault", "", 0, "", -1);
+		StringBuffer<256> desc;
+		desc.Append("Segmentation fault at +0x").AppendHex(u64(si->si_addr));
+		UbaAssert(desc.data, "", 0, "", -1);
 	}
 	#endif
 }
@@ -64,8 +68,14 @@ extern "C"
 
 	uba::LogWriter* CreateCallbackLogWriter(uba::CallbackLogWriter::BeginScopeCallback begin, uba::CallbackLogWriter::EndScopeCallback end, uba::CallbackLogWriter::LogCallback log)
 	{
-		#if PLATFORM_LINUX
-		signal(SIGSEGV, uba::SigsegvHandler);
+		#if UBA_USE_SIGNALHANDLER
+		struct sigaction action;
+		memset(&action, 0, sizeof(action));
+		sigfillset(&action.sa_mask);
+		action.sa_flags = SA_SIGINFO | SA_RESTART | SA_ONSTACK;
+		action.sa_sigaction = uba::SignalHandler;
+		sigaction(SIGSEGV, &action, NULL);
+		//sigaction(SIGABRT, &action, NULL);
 		#endif
 
 		return new uba::CallbackLogWriter(begin, end, log);
