[UBA]
* Added retry for ETXTBSY failures when doing posix_spawn. Have no idea how these can happen but they do.

diff --git a/Engine/Source/Programs/UnrealBuildAccelerator/Common/Private/UbaProcess.cpp b/Engine/Source/Programs/UnrealBuildAccelerator/Common/Private/UbaProcess.cpp
--- a/Engine/Source/Programs/UnrealBuildAccelerator/Common/Private/UbaProcess.cpp
+++ b/Engine/Source/Programs/UnrealBuildAccelerator/Common/Private/UbaProcess.cpp
@@ -1312,14 +1312,18 @@ namespace uba
 
 			posix_spawnattr_t attr;
 			int res = posix_spawnattr_init(&attr);
-			UBA_ASSERTF(res == 0, TC("posix_spawnattr_init"));
+			UBA_ASSERTF(res == 0, TC("posix_spawnattr_init (%s)"), strerror(errno));
+			auto attrGuard = MakeGuard([&]() { posix_spawnattr_destroy(&attr); });
+
 			res = posix_spawnattr_setflags(&attr, flags);
-			UBA_ASSERTF(res == 0, TC("posix_spawnattr_setflags"));
+			UBA_ASSERTF(res == 0, TC("posix_spawnattr_setflags (%s)"), strerror(errno));
 			res = posix_spawnattr_setpgroup(&attr, getpgrp());
-			UBA_ASSERTF(res == 0, TC("posix_spawnattr_setpgroup"));
+			UBA_ASSERTF(res == 0, TC("posix_spawnattr_setpgroup (%s)"), strerror(errno));
 
 			posix_spawn_file_actions_t fileActions;
-			posix_spawn_file_actions_init(&fileActions);
+			res = posix_spawn_file_actions_init(&fileActions);
+			UBA_ASSERTF(res == 0, TC("posix_spawn_file_actions_init (%s)"), strerror(errno));
+			auto actionsGuard = MakeGuard([&]() { posix_spawn_file_actions_destroy(&fileActions); });
 
 			if (!*m_realWorkingDir)
 			{
@@ -1423,14 +1427,22 @@ namespace uba
 
 			envvars.push_back(nullptr);
 
+			u32 retryCount = 0;
 			pid_t processID;
-			res = posix_spawnp(&processID, m_realApplication.c_str(), &fileActions, &attr, (char**)argsArray, (char**)envvars.data());
+			while (true)
+			{
+				res = posix_spawnp(&processID, m_realApplication.c_str(), &fileActions, &attr, (char**)argsArray, (char**)envvars.data());
+				if (res == 0)
+					break;
 
-			posix_spawn_file_actions_destroy(&fileActions);
-			posix_spawnattr_destroy(&attr);
+				if (errno == ETXTBSY && retryCount < 5)
+				{
+					logger.Warning(TC("posix_spawn failed with ETXTBSY, will retry %s %s (Working dir: %s)"), m_realApplication.c_str(), m_startInfo.arguments, m_realWorkingDir);
+					Sleep(2000);
+					++retryCount;
+					continue;
+				}
 
-			if (res != 0)
-			{
 				logger.Error(TC("posix_spawn failed: %s %s (Working dir: %s) -> %i (%s)"), m_realApplication.c_str(), m_startInfo.arguments, m_realWorkingDir, res, strerror(errno));
 				return UBA_EXIT_CODE(12);
 			}
