[UBA]
* Changed CreateHardlink to use symlink on linux (macos already does this)

diff --git a/Engine/Source/Programs/UnrealBuildAccelerator/Common/Private/UbaFile.cpp b/Engine/Source/Programs/UnrealBuildAccelerator/Common/Private/UbaFile.cpp
--- a/Engine/Source/Programs/UnrealBuildAccelerator/Common/Private/UbaFile.cpp
+++ b/Engine/Source/Programs/UnrealBuildAccelerator/Common/Private/UbaFile.cpp
@@ -650,7 +650,7 @@ namespace uba
 #if PLATFORM_WINDOWS
 		return ::CreateHardLinkW(newFileName, existingFileName, NULL);
 #else
-#if PLATFORM_MAC
+#if 1//PLATFORM_MAC
 		int res = symlink(existingFileName, newFileName);
 #else
 		int res = link(existingFileName, newFileName);
