[Uba]
* Desperate changes trying to figure out the segfault happening rarely on the farm which contains _zero_ information. It seems to happen directly after a remote is connecting but it never reaches the log entry that writes "Connected to ".

diff --git a/Engine/Source/Programs/UnrealBuildAccelerator/Common/Private/UbaNetworkBackendTcp.cpp b/Engine/Source/Programs/UnrealBuildAccelerator/Common/Private/UbaNetworkBackendTcp.cpp
--- a/Engine/Source/Programs/UnrealBuildAccelerator/Common/Private/UbaNetworkBackendTcp.cpp
+++ b/Engine/Source/Programs/UnrealBuildAccelerator/Common/Private/UbaNetworkBackendTcp.cpp
@@ -73,6 +73,9 @@ namespace uba
 		CriticalSection shutdownLock;
 
 		Thread recvThread;
+
+		Connection(const Connection&) = delete;
+		void operator=(const Connection&) = delete;
 	};
 
 	bool SetKeepAlive(Logger& logger, SOCKET socket);
@@ -405,7 +408,7 @@ namespace uba
 			ScopedWriteLock lock(m_connectionsLock);
 			auto it = m_connections.emplace(m_connections.end(), clientSocket);
 			auto& conn = *it;
-			conn.recvThread.Start([this, &conn] { ThreadRecv(conn); return 0; });
+			conn.recvThread.Start([this, connPtr = &conn] { ThreadRecv(*connPtr); return 0; });
 			lock.Leave();
 
 			if (!m_connectedFunc(&conn, remoteSockAddr))
@@ -638,7 +641,7 @@ namespace uba
 		ScopedWriteLock lock(m_connectionsLock);
 		auto it = m_connections.emplace(m_connections.end(), socketFd);
 		auto& conn = *it;
-		conn.recvThread.Start([this, &conn] { ThreadRecv(conn); return 0; });
+		conn.recvThread.Start([this, connPtr = &conn] { ThreadRecv(*connPtr); return 0; });
 		lock.Leave();
 
 		if (!connectedFunc(&conn, remoteSocketAddr, timedOut))
diff --git a/Engine/Source/Programs/UnrealBuildAccelerator/Common/Private/UbaThread.cpp b/Engine/Source/Programs/UnrealBuildAccelerator/Common/Private/UbaThread.cpp
--- a/Engine/Source/Programs/UnrealBuildAccelerator/Common/Private/UbaThread.cpp
+++ b/Engine/Source/Programs/UnrealBuildAccelerator/Common/Private/UbaThread.cpp
@@ -93,6 +93,7 @@ namespace uba
 		int res = pthread_join(*(pthread_t*)&m_handle, (void**)&ptr);
 		UBA_ASSERT(res == 0);
 #endif
+		m_func = {};
 		m_handle = nullptr;
 		return true;
 	}
