[UBA]
* One more place to fix thread lambda taking reference to reference. This is a long shot but according to c++ standard it is undefined behavior to take reference of reference and then first reference go out of scope (so maybe msvc and clang produce different code since we only see this problem on linux)

diff --git a/Engine/Source/Programs/UnrealBuildAccelerator/Common/Private/UbaNetworkServer.cpp b/Engine/Source/Programs/UnrealBuildAccelerator/Common/Private/UbaNetworkServer.cpp
--- a/Engine/Source/Programs/UnrealBuildAccelerator/Common/Private/UbaNetworkServer.cpp
+++ b/Engine/Source/Programs/UnrealBuildAccelerator/Common/Private/UbaNetworkServer.cpp
@@ -899,14 +899,14 @@ namespace uba
 					return true;
 				u32 toAdd = connectionCount - c.connectionCount;
 
-				auto& conn = *(NetworkServer::Connection*)connectionInfo.internalData;
-				auto remoteAddr = conn.m_remoteSockAddr;
+				auto connPtr = (NetworkServer::Connection*)connectionInfo.internalData;
 				ScopedWriteLock lock2(m_addConnectionsLock);
 				for (u32 i = 0; i != toAdd; ++i)
 				{
-					m_addConnections.emplace_back([this, &conn, remoteAddr]()
+					m_addConnections.emplace_back([this, connPtr]()
 						{
-							conn.m_backend.Connect(m_logger, remoteAddr, [this, &conn](void* connection, const sockaddr& remoteSocketAddr, bool* timedOut)
+							auto& conn = *connPtr;
+							conn.m_backend.Connect(m_logger, conn.m_remoteSockAddr, [this, &conn](void* connection, const sockaddr& remoteSocketAddr, bool* timedOut)
 								{
 									CryptoKey cryptoKey = InvalidCryptoKey;
 									if (conn.m_cryptoKey)
