[UBA]
* Fixed so storage client keeps track of last tested proxy to prevent testing of same proxy over and over again
* Removed Sleep to emulate timeout if timeout variable is not set

diff --git a/Engine/Source/Programs/UnrealBuildAccelerator/Common/Private/UbaNetworkBackendTcp.cpp b/Engine/Source/Programs/UnrealBuildAccelerator/Common/Private/UbaNetworkBackendTcp.cpp
--- a/Engine/Source/Programs/UnrealBuildAccelerator/Common/Private/UbaNetworkBackendTcp.cpp
+++ b/Engine/Source/Programs/UnrealBuildAccelerator/Common/Private/UbaNetworkBackendTcp.cpp
@@ -523,8 +523,10 @@ namespace uba
 		if (!success)
 			return false;
 
-		if (timedOut)
-			*timedOut = true;
+		if (!timedOut)
+			return false;
+
+		*timedOut = true;
 		int connectTimeMs = int(TimeToMs(GetTime() - startTime));
 		int timeoutMs = 2000;
 		if (connectTimeMs < timeoutMs)
diff --git a/Engine/Source/Programs/UnrealBuildAccelerator/Common/Private/UbaStorageClient.cpp b/Engine/Source/Programs/UnrealBuildAccelerator/Common/Private/UbaStorageClient.cpp
--- a/Engine/Source/Programs/UnrealBuildAccelerator/Common/Private/UbaStorageClient.cpp
+++ b/Engine/Source/Programs/UnrealBuildAccelerator/Common/Private/UbaStorageClient.cpp
@@ -246,6 +246,7 @@ namespace uba
 					{
 						delete m_proxyClient;
 						m_proxyClient = nullptr;
+						m_lastTestedProxyIp.clear();
 					}
 				}
 				wantsProxy = proxy == nullptr && m_startProxyCallback;
@@ -262,7 +263,10 @@ namespace uba
 							{
 								delete proxy;
 								if (proxy == m_proxyClient)
+								{
 									m_proxyClient = nullptr;
+									m_lastTestedProxyIp.clear();
+								}
 							}
 						}
 					}
@@ -314,6 +318,13 @@ namespace uba
 					if (m_proxyClient)
 						continue;
 
+					// If we have x processes sitting here with an already closed proxy we would do x connect attempts to that bad proxy
+					// Therefore we keep track of last checked proxy to prevent repetition of already failing proxy
+					// There is a risk we have multiple threads asking for two bad proxies.. but that is low enough for us not to care
+					if (m_lastTestedProxyIp == proxyHost.data)
+						continue;
+					m_lastTestedProxyIp = proxyHost.data;
+
 					NetworkBackendTcp& proxyBackend = m_client.GetTcpBackend();
 					NetworkClientCreateInfo ncci(m_logger.m_writer);
 					bool ctorSuccess = true;
diff --git a/Engine/Source/Programs/UnrealBuildAccelerator/Common/Public/UbaStorageClient.h b/Engine/Source/Programs/UnrealBuildAccelerator/Common/Public/UbaStorageClient.h
--- a/Engine/Source/Programs/UnrealBuildAccelerator/Common/Public/UbaStorageClient.h
+++ b/Engine/Source/Programs/UnrealBuildAccelerator/Common/Public/UbaStorageClient.h
@@ -73,6 +73,7 @@ namespace uba
 		static constexpr u8 ServiceId = StorageServiceId;
 
 		struct ProxyClient;
+		TString m_lastTestedProxyIp;
 		ReaderWriterLock m_proxyClientLock;
 		ProxyClient* m_proxyClient = nullptr;
 		u64 m_proxyClientKeepAliveTime = 0;
