[UBA]
* Disabled asserts in shipping again for linux platform since we haven't seen an assert since enabling it

diff --git a/Engine/Source/Programs/UnrealBuildAccelerator/Core/Public/UbaPlatform.h b/Engine/Source/Programs/UnrealBuildAccelerator/Core/Public/UbaPlatform.h
--- a/Engine/Source/Programs/UnrealBuildAccelerator/Core/Public/UbaPlatform.h
+++ b/Engine/Source/Programs/UnrealBuildAccelerator/Core/Public/UbaPlatform.h
@@ -56,7 +56,7 @@ namespace uba
 	#define	sizeof_array(array) int(sizeof(array)/sizeof(array[0]))
 	#endif
 
-	#if UBA_DEBUG || PLATFORM_LINUX
+	#if UBA_DEBUG
 	#define UBA_ASSERT(x) do { if (x) break; uba::UbaAssert(TC(""), __FILE__, __LINE__, #x, 543221); } while(false)
 	#define UBA_ASSERTF(x, ...) do { if (x) break; uba::StringBuffer<> _buf; _buf.Appendf(__VA_ARGS__); uba::UbaAssert(_buf.data, __FILE__, __LINE__, #x, 543221); } while(false)
 	#define UBA_ASSERT_MESSAGEBOX 0
