[UBA]
* Added -listenTimeout option to UbaAgent

diff --git a/Engine/Source/Programs/UnrealBuildAccelerator/Agent/Private/UbaAgent.cpp b/Engine/Source/Programs/UnrealBuildAccelerator/Agent/Private/UbaAgent.cpp
--- a/Engine/Source/Programs/UnrealBuildAccelerator/Agent/Private/UbaAgent.cpp
+++ b/Engine/Source/Programs/UnrealBuildAccelerator/Agent/Private/UbaAgent.cpp
@@ -36,7 +36,8 @@
 namespace uba
 {
 	const char*		Version = GetVersionString();
-	u32				DefaultCapacityGb = 20;
+	constexpr u32	DefaultCapacityGb = 20;
+	constexpr u32	DefaultListenTimeout = 5;
 	const tchar*	DefaultRootDir = [](){
 		static tchar buf[256];
 		if (IsWindows)
@@ -68,6 +69,7 @@ namespace uba
 		logger.Info(TC("  -dir=<rootdir>          The directory used to store data. Defaults to \"%s\""), DefaultRootDir);
 		logger.Info(TC("  -host=<host>[:<port>]   The ip/name and port (default: %u) of the machine we want to help"), DefaultPort);
 		logger.Info(TC("  -listen[=port]          Agent will listen for connections on port (default: %u) and help when connected"), DefaultPort);
+		logger.Info(TC("  -listenTimeout=<sec>    Number of seconds agent will listen for host before giving up (default: %u)"), DefaultListenTimeout);
 		logger.Info(TC("  -proxyport=<port>       Which port that agent will use if being assigned to be proxy for other agents (default: %u)"), DefaultStorageProxyPort);
 		logger.Info(TC("  -maxcpu=<number>        Max number of processes that can be started. Defaults to \"%u\" on this machine"), DefaultProcessorCount);
 		logger.Info(TC("  -mulcpu=<number>        This value multiplies with number of cpu to figure out max cpu. Defaults to 1.0"));
@@ -402,6 +404,7 @@ namespace uba
 		u32 waitProcessId = ~0u;
 		u32 memWaitLoadPercent = 80;
 		u32 memKillLoadPercent = 90;
+		u32 listenTimeoutSec = DefaultListenTimeout;
 		u8 crypto[16];
 		bool hasCrypto = false;
 		Vector<TString> populateCasDirs;
@@ -484,6 +487,13 @@ namespace uba
 						return PrintHelp(TC("Invalid value for -capacity"));
 				useListen = true;
 			}
+			else if (name.Equals(TC("-listenTimeout")))
+			{
+				if (value.IsEmpty())
+					return PrintHelp(TC("-listenTimeout needs a value"));
+				if (!value.Parse(listenTimeoutSec))
+					return PrintHelp(TC("Invalid value for -listenTimeout"));
+			}
 			else if (name.Equals(TC("-named")))
 			{
 				if (value.IsEmpty())
@@ -883,7 +893,7 @@ namespace uba
 					}
 
 					u64 waitTime = GetTime() - startTime;
-					if (!poll && TimeToMs(waitTime) > 5*1000)
+					if (!poll && TimeToMs(waitTime) > listenTimeoutSec*1000)
 					{
 						logger.Error(TC("Failed to get connection while listening for %s"), TimeToText(waitTime).str);
 						return -1;
