[UBA]
* Added delete to copy-ctor and assign operator for Thread and event and fixed related code
* Some minor cleanups in thread class
* Fixed so UbaAgent is not writing timeout waiting for connection if server was connected but quickly disconnected

diff --git a/Engine/Source/Programs/UnrealBuildAccelerator/Agent/Private/UbaAgent.cpp b/Engine/Source/Programs/UnrealBuildAccelerator/Agent/Private/UbaAgent.cpp
--- a/Engine/Source/Programs/UnrealBuildAccelerator/Agent/Private/UbaAgent.cpp
+++ b/Engine/Source/Programs/UnrealBuildAccelerator/Agent/Private/UbaAgent.cpp
@@ -884,7 +884,7 @@ namespace uba
 			{
 				client->StartListen(*networkBackend, port);
 				u64 startTime = GetTime();
-				while (!client->IsConnected(200))
+				while (!client->IsOrWasConnected(200))
 				{
 					if (IsEscapePressed())
 					{
diff --git a/Engine/Source/Programs/UnrealBuildAccelerator/Common/Private/UbaNetworkBackendTcp.cpp b/Engine/Source/Programs/UnrealBuildAccelerator/Common/Private/UbaNetworkBackendTcp.cpp
--- a/Engine/Source/Programs/UnrealBuildAccelerator/Common/Private/UbaNetworkBackendTcp.cpp
+++ b/Engine/Source/Programs/UnrealBuildAccelerator/Common/Private/UbaNetworkBackendTcp.cpp
@@ -212,7 +212,7 @@ namespace uba
 
 		auto AddAddr = [&](const tchar* addr)
 			{
-				m_listenEntries.push_back({});
+				m_listenEntries.emplace_back();
 				auto& entry = m_listenEntries.back();
 				entry.ip.Append(addr);
 				entry.port = port;
diff --git a/Engine/Source/Programs/UnrealBuildAccelerator/Common/Private/UbaNetworkClient.cpp b/Engine/Source/Programs/UnrealBuildAccelerator/Common/Private/UbaNetworkClient.cpp
--- a/Engine/Source/Programs/UnrealBuildAccelerator/Common/Private/UbaNetworkClient.cpp
+++ b/Engine/Source/Programs/UnrealBuildAccelerator/Common/Private/UbaNetworkClient.cpp
@@ -15,6 +15,7 @@ namespace uba
 	,	m_logWriter(info.logWriter)
 	,	m_logger(info.logWriter, SetGetPrefix(name))
 	,	m_isConnected(true)
+	,	m_isOrWasConnected(true)
 	,	m_tcpBackend(new NetworkBackendTcp(info.logWriter))
 	{
 		outCtorSuccess = true;
@@ -125,6 +126,8 @@ namespace uba
 		if (!rc.recvEvent.IsSet(timeoutSeconds * 1000))
 			return m_logger.Error(TC("Timed out waiting for connection response from server"));
 
+		m_isOrWasConnected.Set();
+
 		if (rc.error == 1) // Bad version
 			return m_logger.Error(TC("Version mismatch with server"));
 
@@ -281,6 +284,11 @@ namespace uba
 		return m_isConnected.IsSet(waitTimeoutMs);
 	}
 
+	bool NetworkClient::IsOrWasConnected(u32 waitTimeoutMs)
+	{
+		return m_isOrWasConnected.IsSet(waitTimeoutMs);
+	}
+
 	void NetworkClient::PrintSummary(Logger& logger)
 	{
 		Timer sendTimer;
diff --git a/Engine/Source/Programs/UnrealBuildAccelerator/Common/Private/UbaStorage.cpp b/Engine/Source/Programs/UnrealBuildAccelerator/Common/Private/UbaStorage.cpp
--- a/Engine/Source/Programs/UnrealBuildAccelerator/Common/Private/UbaStorage.cpp
+++ b/Engine/Source/Programs/UnrealBuildAccelerator/Common/Private/UbaStorage.cpp
@@ -329,9 +329,11 @@ namespace uba
 
 		struct WorkRec
 		{
+			WorkRec(u32 wc) { workCount = wc; events = new Event[workCount]; }
+			~WorkRec() { delete[] events; }
 			Atomic<u64> refCount;
 			Atomic<u64> compressCounter;
-			Vector<Event> events;
+			Event* events;
 			const u8* uncompressedData = nullptr;
 #if !UBA_USE_SPARSEFILE
 			FileAccessor* destination = nullptr;
@@ -346,9 +348,7 @@ namespace uba
 			bool error = false;
 		};
 
-		WorkRec* rec = new WorkRec();
-		rec->workCount = workCount;
-		rec->events.resize(workCount);
+		WorkRec* rec = new WorkRec(workCount);
 		rec->uncompressedData = uncompressedData;
 		rec->maxUncompressedBlock = maxUncompressedBlock;
 		rec->fileSize = fileSize;
@@ -359,8 +359,8 @@ namespace uba
 		rec->memPos = totalWritten;
 #endif
 
-		for (Event& e : rec->events)
-			e.Create(true);
+		for (u32 i=0; i!=workCount; ++i)
+			rec->events[i].Create(true);
 
 		StorageStats& stats = Stats();
 
@@ -418,7 +418,7 @@ namespace uba
 		rec->refCount = workerCount + 1; // We need to keep refcount up 1 to make sure it is not deleted before we read rec->written
 		m_workManager->AddWork(work, workerCount-1, TC("Compress")); // We are a worker ourselves
 		work();
-		rec->events[rec->events.size() - 1].IsSet();
+		rec->events[rec->workCount - 1].IsSet();
 
 		totalWritten += rec->written;
 		bool error = rec->error;
diff --git a/Engine/Source/Programs/UnrealBuildAccelerator/Common/Private/UbaThread.cpp b/Engine/Source/Programs/UnrealBuildAccelerator/Common/Private/UbaThread.cpp
--- a/Engine/Source/Programs/UnrealBuildAccelerator/Common/Private/UbaThread.cpp
+++ b/Engine/Source/Programs/UnrealBuildAccelerator/Common/Private/UbaThread.cpp
@@ -26,7 +26,6 @@ namespace uba
 	}
 
 	Thread::Thread()
-	:	m_handle(nullptr)
 	{
 	}
 
@@ -37,13 +36,7 @@ namespace uba
 
 	Thread::~Thread()
 	{
-		if (!m_handle)
-			return;
 		Wait();
-
-		#if PLATFORM_WINDOWS
-		CloseHandle(m_handle);
-		#endif
 	}
 
 	void Thread::Start(Function<u32()>&& f)
diff --git a/Engine/Source/Programs/UnrealBuildAccelerator/Common/Public/UbaNetworkClient.h b/Engine/Source/Programs/UnrealBuildAccelerator/Common/Public/UbaNetworkClient.h
--- a/Engine/Source/Programs/UnrealBuildAccelerator/Common/Public/UbaNetworkClient.h
+++ b/Engine/Source/Programs/UnrealBuildAccelerator/Common/Public/UbaNetworkClient.h
@@ -42,6 +42,7 @@ namespace uba
 		bool SendKeepAlive();
 
 		bool IsConnected(u32 waitTimeoutMs = 0);
+		bool IsOrWasConnected(u32 waitTimeoutMs = 0);
 
 		void PrintSummary(Logger& logger);
 
@@ -94,6 +95,7 @@ namespace uba
 		u32 m_receiveTimeoutSeconds;
 
 		Event m_isConnected;
+		Event m_isOrWasConnected;
 		Atomic<u32> m_connectionCount;
 		ReaderWriterLock m_onConnectedFunctionsLock;
 		Vector<OnConnectedFunction> m_onConnectedFunctions;
diff --git a/Engine/Source/Programs/UnrealBuildAccelerator/Common/Public/UbaThread.h b/Engine/Source/Programs/UnrealBuildAccelerator/Common/Public/UbaThread.h
--- a/Engine/Source/Programs/UnrealBuildAccelerator/Common/Public/UbaThread.h
+++ b/Engine/Source/Programs/UnrealBuildAccelerator/Common/Public/UbaThread.h
@@ -19,8 +19,14 @@ namespace uba
 		bool Wait(u32 milliseconds = ~0u, Event* wakeupEvent = nullptr);
 	private:
 		Function<u32()>	m_func;
-		void* m_handle;
-		Event m_finished; // Only used for non windows
+		void* m_handle = nullptr;
+
+		#if !PLATFORM_WINDOWS
+		Event m_finished;
+		#endif
+
+		Thread(const Thread&) = delete;
+		void operator=(const Thread&) = delete;
 	};
 
 
diff --git a/Engine/Source/Programs/UnrealBuildAccelerator/Core/Public/UbaEvent.h b/Engine/Source/Programs/UnrealBuildAccelerator/Core/Public/UbaEvent.h
--- a/Engine/Source/Programs/UnrealBuildAccelerator/Core/Public/UbaEvent.h
+++ b/Engine/Source/Programs/UnrealBuildAccelerator/Core/Public/UbaEvent.h
@@ -24,5 +24,8 @@ namespace uba
 		#else
 		u64 m_data[16];
 		#endif
+
+		Event(const Event&) = delete;
+		void operator=(const Event&) = delete;
 	};
 }
