UBA: Fix for populateCasFromXcode.

diff --git a/Engine/Source/Programs/UnrealBuildAccelerator/Agent/Private/UbaAgent.cpp b/Engine/Source/Programs/UnrealBuildAccelerator/Agent/Private/UbaAgent.cpp
--- a/Engine/Source/Programs/UnrealBuildAccelerator/Agent/Private/UbaAgent.cpp
+++ b/Engine/Source/Programs/UnrealBuildAccelerator/Agent/Private/UbaAgent.cpp
@@ -410,7 +410,7 @@ namespace uba
 		Vector<TString> populateCasDirs;
 
 		#if PLATFORM_MAC
-		bool populateCasFromXCode = false;
+		bool populateCasFromXcode = false;
 		#endif
 
 		for (int i=1; i!=argc; ++i)
@@ -754,7 +754,7 @@ namespace uba
 		}
 
 #if PLATFORM_MAC
-		if (populateCasFromXCode)
+		if (populateCasFromXcode)
 		{
 			// .. do your think Josh.... maybe something like this?
 
