- Mac Xcode versioning for Uba

diff --git a/Engine/Source/Programs/UnrealBuildAccelerator/Agent/Private/UbaAgent.cpp b/Engine/Source/Programs/UnrealBuildAccelerator/Agent/Private/UbaAgent.cpp
--- a/Engine/Source/Programs/UnrealBuildAccelerator/Agent/Private/UbaAgent.cpp
+++ b/Engine/Source/Programs/UnrealBuildAccelerator/Agent/Private/UbaAgent.cpp
@@ -6,6 +6,7 @@
 #include "UbaSessionClient.h"
 #include "UbaStorageClient.h"
 #include "UbaStorageProxy.h"
+#include "UbaDirectoryIterator.h"
 #include "UbaAWS.h"
 #include "UbaSentry.h"
 #include "UbaVersion.h"
@@ -101,7 +102,8 @@ namespace uba
 		logger.Info(TC("  -crypto=<key>           16 bytes crypto key used for secure network transfer"));
 		logger.Info(TC("  -populateCas=<dir>      Prepopulate cas database with files in dir. If files needed exists on machine this can be an optimization"));
 		#if PLATFORM_MAC
-		logger.Info(TC("  -populateCasFromXcode   Prepopulate cas database with files from local xcode installation."));
+		logger.Info(TC("  -populateCasFromXcodeVersion=<version>   Prepopulate cas database with files from local xcode installation that matches the version."));
+		logger.Info(TC("  -populateCasFromAllXcodes   Prepopulate cas database with files from local xcode installation that matches the version."));
 		#endif
 		logger.Info(TC(""));
 		return -1;
@@ -410,7 +412,8 @@ namespace uba
 		Vector<TString> populateCasDirs;
 
 		#if PLATFORM_MAC
-		bool populateCasFromXcode = false;
+		StringBuffer<32> populateCasFromXcodeVersion;
+		bool populateCasFromAllXcodes;
 		#endif
 
 		for (int i=1; i!=argc; ++i)
@@ -615,9 +618,13 @@ namespace uba
 				populateCasDirs.push_back(value.data);
 			}
 			#if PLATFORM_MAC
-			else if (name.Equals(TC("-populateCasFromXcode")))
+			else if (name.Equals(TC("-populateCasFromXcodeVersion")))
 			{
-				populateCasFromXcode = true;
+				populateCasFromXcodeVersion.Append(value.data);
+			}
+			else if (name.Equals(TC("-populateCasFromAllXcodes")))
+			{
+				populateCasFromAllXcodes = true;
 			}
 			#endif
 			else if (name.Equals(TC("-sentry")))
@@ -754,34 +761,103 @@ namespace uba
 		}
 
 #if PLATFORM_MAC
-		if (populateCasFromXcode)
+		
+		Vector<TString> xcodeDirectories;
+		
+		if (populateCasFromXcodeVersion.count > 0 || populateCasFromAllXcodes)
 		{
-			// .. do your think Josh.... maybe something like this?
-
-			FILE* xcSelect = popen("/usr/bin/xcode-select -p", "r");
-			if (!xcSelect)
+			// look for all xcodes in /Applications (is there a function to get Applications dir location for other locales?)
+			StringBuffer<> applicationsDir;
+			applicationsDir.Append("/Applications");
+			
+			TraverseDir(logger, applicationsDir.data,
+						[&](const DirectoryEntry& e)
+						{
+				if (IsDirectory(e.attributes) && StartsWith(e.name, "Xcode"))
+				{
+					StringBuffer<128> xcodeDir("/Applications/");
+					xcodeDir.Append(e.name).Append("/Contents/Developer/");
+					if (FileExists(logger, xcodeDir.data))
+					{
+						if (populateCasFromAllXcodes)
+						{
+							xcodeDirectories.push_back(xcodeDir.data);
+						}
+						else
+						{
+							StringBuffer<512> command;
+							StringBuffer<32> xcodeVer;
+							
+							// look for short version like 15.1 or 15, or BuildVersion like 15C610
+							bool bUseShortVersion = (populateCasFromXcodeVersion.Contains('.')) || populateCasFromXcodeVersion.count <= 3;
+							const char* key = bUseShortVersion ? "CFBundleShortVersionString" : "ProductBuildVersion";
+							command.Append("/usr/bin/defaults read /Applications/").Append(e.name).Append("/Contents/version.plist ").Append(key);
+							
+							FILE* getver = popen(command.data, "r");
+							if (getver == nullptr || fgets(xcodeVer.data, xcodeVer.capacity, getver) == nullptr)
+							{
+								pclose(getver);
+								logger.Error("Failed to get DTXcodeBuild from /Applications/%s", e.name);
+								return;
+							}
+							pclose(getver);
+							xcodeVer.count = strlen(xcodeVer.data);
+							while (isspace(xcodeVer.data[xcodeVer.count-1]))
+							{
+								xcodeVer.data[xcodeVer.count-1] = 0;
+								xcodeVer.count--;
+							}
+							
+							logger.Info("/Applications/%s has version '%s' (looking for %s)", e.name, xcodeVer.data, populateCasFromXcodeVersion.data);
+							
+							if (xcodeVer.Equals(populateCasFromXcodeVersion.data))
+							{
+								xcodeDirectories.push_back(xcodeDir.data);
+							}
+						}
+					}
+				}
+			});
+		}
+		// if we didn't want a single version, or all xcodes, then use active xcode (useful for user running their own agents)
+		else
+		{
+			StringBuffer<512> command;
+			StringBuffer<32> xcodeSelectOutput;
+			command.Append("/usr/bin/defaults read %s/../Contents/version.plist ProductBuildVersion");
+			FILE* xcodeSelect = popen("/usr/bin/xcode-select -p", "r");
+			if (xcodeSelect == nullptr || fgets(xcodeSelectOutput.data, xcodeSelectOutput.capacity, xcodeSelect) == nullptr || pclose(xcodeSelect) != 0)
 			{
-				logger.Error("Unable to run /usr/bin/xcode-select. Is xcode installed?");
+				logger.Error("Failed to get an Xcode from xcode-select");
 				return -1;
 			}
-			StringBuffer<> developerPath;
-			bool success = fgets(developerPath.data, developerPath.capacity, xcSelect) != NULL;
-			pclose(xcSelect);
 
-			if (!success)
+			xcodeSelectOutput.count = strlen(xcodeSelectOutput.data);
+			while (isspace(xcodeSelectOutput.data[xcodeSelectOutput.count-1]))
 			{
-				logger.Error("Unable to parse string from /usr/bin/xcode-select");
-				return -1;
+				xcodeSelectOutput.data[xcodeSelectOutput.count-1] = 0;
+				xcodeSelectOutput.count--;
 			}
-			developerPath.count = strlen(developerPath.data);
-			developerPath.EnsureEndsWithSlash();
+			
+			xcodeDirectories.push_back(xcodeSelectOutput.data);
+		}
+
+		if (xcodeDirectories.size() == 0)
+		{
+			logger.Error("Unable to populate from any Xcodes. Agent is unusable.");
+			return -1;
+		}
 
-			u32 developerPathLen = developerPath.count;
-			const char* subDirs[] = { "Toolchains", "Platforms" };
+		for (TString& xcodeDir : xcodeDirectories)
+		{
+			logger.Info("Populating cas with %s", xcodeDir.data());
+			
+			const char* subDirs[] = { "/Toolchains", "/Platforms" };
 			for (auto subDir : subDirs)
 			{
-				developerPath.Resize(developerPathLen).Append(subDir);
-				populateCasDirs.push_back(developerPath.data);
+				TString populateDir(xcodeDir);
+				populateDir.append(subDir);
+				populateCasDirs.push_back(populateDir);
 			}
 		}
 #endif
@@ -1181,4 +1257,4 @@ int main(int argc, char* argv[])
 {
 	return uba::WrappedMain(argc, argv);
 }
-#endif
\ No newline at end of file
+#endif
diff --git a/Engine/Source/Programs/UnrealBuildTool/Executors/Experimental/UBAHordeSession.cs b/Engine/Source/Programs/UnrealBuildTool/Executors/Experimental/UBAHordeSession.cs
--- a/Engine/Source/Programs/UnrealBuildTool/Executors/Experimental/UBAHordeSession.cs
+++ b/Engine/Source/Programs/UnrealBuildTool/Executors/Experimental/UBAHordeSession.cs
@@ -623,6 +623,13 @@ namespace UnrealBuildTool
 						//arguments.Add("-NoStore");
 						//arguments.Add("-KillRandom"); // For debugging
 
+						if (System.OperatingSystem.IsMacOS())
+						{
+							// we need to populate the cas with all known xcodes so we can serve all the ones we have installed
+							string XcodeVersion = Utils.RunLocalProcessAndReturnStdOut("/bin/sh", "-c '/usr/bin/defaults read $(xcode-select -p)/../version.plist ProductBuildVersion");
+							arguments.Add($"-populateCasFromXcodeVersion={XcodeVersion}");
+						}
+
 						arguments.Add("-Dir=%UE_HORDE_SHARED_DIR%\\Uba");
 						arguments.Add("-Eventfile=%UE_HORDE_TERMINATION_SIGNAL_FILE%");
 						arguments.Add("-MaxIdle=15");
@@ -631,11 +638,6 @@ namespace UnrealBuildTool
 							arguments.Add("-Log");
 						}
 
-						if (OperatingSystem.IsMacOS())
-						{
-							arguments.Add("-populateCasFromXcode");
-						}
-
 						logger.LogInformation("Executing child process: {Executable} {Arguments}", executable, CommandLineArguments.Join(arguments));
 
 						ExecuteProcessFlags execFlags = _allowWine ? ExecuteProcessFlags.UseWine : ExecuteProcessFlags.None;
