UBA: Address issue of agents quitting instead of waiting to connect to host.

diff --git a/Engine/Source/Programs/UnrealBuildAccelerator/Common/Private/UbaNetworkBackendTcp.cpp b/Engine/Source/Programs/UnrealBuildAccelerator/Common/Private/UbaNetworkBackendTcp.cpp
--- a/Engine/Source/Programs/UnrealBuildAccelerator/Common/Private/UbaNetworkBackendTcp.cpp
+++ b/Engine/Source/Programs/UnrealBuildAccelerator/Common/Private/UbaNetworkBackendTcp.cpp
@@ -563,8 +563,13 @@ namespace uba
 				return false;
 #else
 		if (res != 0)
+		{
 			if (errno != EINPROGRESS)
+			{
+				logger.Error(TC("Connect failed (%d: %s)"), WSAGetLastError(), LastErrorToText(WSAGetLastError()).data);
 				return false;
+			}
+		}
 #endif
 
 		// Return to blocking since we want select to block
@@ -619,7 +624,22 @@ namespace uba
 		}
 
 
+
 #if !PLATFORM_WINDOWS
+		// Before we send anything even though the
+		// the socket is writable, but let's make sure
+		// the connection is actually valid by getting
+		// information about what we've connected to
+		struct sockaddr_in junk;
+		socklen_t length = sizeof(junk);
+		memset(&junk, 0, sizeof(junk));
+		if (getpeername(socketFd, (struct sockaddr *)&junk, &length) != 0)
+		{
+			//logger.Info(TC("Connection not ready..."));
+			*timedOut = true;
+			return false;
+		}
+
 		int sent = (int)send(socketFd, nullptr, 0, 0);
 		if (sent == SOCKET_ERROR)
 		{
