- Fixed path length issue with xcode-select

diff --git a/Engine/Source/Programs/UnrealBuildAccelerator/Agent/Private/UbaAgent.cpp b/Engine/Source/Programs/UnrealBuildAccelerator/Agent/Private/UbaAgent.cpp
--- a/Engine/Source/Programs/UnrealBuildAccelerator/Agent/Private/UbaAgent.cpp
+++ b/Engine/Source/Programs/UnrealBuildAccelerator/Agent/Private/UbaAgent.cpp
@@ -823,8 +823,7 @@ namespace uba
 		else
 		{
 			StringBuffer<512> command;
-			StringBuffer<32> xcodeSelectOutput;
-			command.Append("/usr/bin/defaults read %s/../Contents/version.plist ProductBuildVersion");
+			StringBuffer<512> xcodeSelectOutput;
 			FILE* xcodeSelect = popen("/usr/bin/xcode-select -p", "r");
 			if (xcodeSelect == nullptr || fgets(xcodeSelectOutput.data, xcodeSelectOutput.capacity, xcodeSelect) == nullptr || pclose(xcodeSelect) != 0)
 			{
