UBA: Add support for CPU time for MacOS using mach messages.

diff --git a/Engine/Source/Programs/UnrealBuildAccelerator/Common/Private/UbaSession.cpp b/Engine/Source/Programs/UnrealBuildAccelerator/Common/Private/UbaSession.cpp
--- a/Engine/Source/Programs/UnrealBuildAccelerator/Common/Private/UbaSession.cpp
+++ b/Engine/Source/Programs/UnrealBuildAccelerator/Common/Private/UbaSession.cpp
@@ -26,6 +26,7 @@
 #include <mach/mach_types.h>
 #include <mach/mach_init.h>
 #include <mach/mach_host.h>
+#include <mach/mach.h>
 extern char **environ;
 #endif
 
@@ -2246,10 +2247,7 @@ namespace uba
 		m_previousTotalCpuTime = totalTime;
 		m_previousIdleCpuTime = idleTime;
 
-		// TODO: This is the wrong solution.. but can't repro the bad values some people get
-		if (cpuLoad >= 0 && cpuLoad <= 1.0f)
-			m_cpuLoad = cpuLoad;
-#else
+#elif PLATFORM_LINUX
 		int fd = open("/proc/stat", O_RDONLY);
 		if (fd != -1)
 		{
@@ -2295,16 +2293,52 @@ namespace uba
 
 						m_previousTotalCpuTime = totalTime;
 						m_previousIdleCpuTime = idleTime;
-
-						if (cpuLoad >= 0 && cpuLoad <= 1.0f)
-							m_cpuLoad = cpuLoad;
 					}
 				}
 			}
 			close(fd);
 		}
 		// TODO: Read "/proc/stat" to get cpu cycles
+#else // PLATFORM_MAC
+        mach_msg_type_number_t  CpuMsgCount = 0;
+        processor_flavor_t CpuInfoType = PROCESSOR_CPU_LOAD_INFO;;
+        natural_t CpuCount = 0;
+        processor_cpu_load_info_t CpuData;
+        host_t host = mach_host_self();
+
+        int res = 0;
+        u64 work = 0;
+        u64 idleTime = 0;
+
+        res = host_processor_info(host, CpuInfoType, &CpuCount, (processor_info_array_t *)&CpuData, &CpuMsgCount);
+		if(res != KERN_SUCCESS)
+        {
+				return m_logger.Error(TC("Kernel error: %s"), mach_error_string(res));
+        }
+
+        for(int i = 0; i < (int)CpuCount; i++)
+        {
+                work += CpuData[i].cpu_ticks[CPU_STATE_SYSTEM];
+                work += CpuData[i].cpu_ticks[CPU_STATE_USER];
+                work += CpuData[i].cpu_ticks[CPU_STATE_NICE];
+                idleTime += CpuData[i].cpu_ticks[CPU_STATE_IDLE];
+        }
+
+		u64 totalTime = work + idleTime;
+
+		u64 totalTimeSinceLastTime = totalTime - m_previousTotalCpuTime;
+		u64 idleTimeSinceLastTime = idleTime - m_previousIdleCpuTime;
+
+		float cpuLoad = 1.0f - ((totalTimeSinceLastTime > 0) ? (float(idleTimeSinceLastTime) / float(totalTimeSinceLastTime)) : 0);
+
+		m_previousTotalCpuTime = totalTime;
+		m_previousIdleCpuTime = idleTime;
+
 #endif
+		// TODO: This is the wrong solution.. but can't repro the bad values some people get
+		if (cpuLoad >= 0 && cpuLoad <= 1.0f)
+			m_cpuLoad = cpuLoad;
+
 		return m_cpuLoad;
 	}
 
