[Uba]
* Removed usage of "select" in tcp socket connect code because it can't handle file descriptors over 1024 so segfaults (phew, hard one to find)
* Added option to have clients send back log to server (SessionServerCreateInfo.remoteLogEnabled) .. this forced a network protocol bump
* Improved asserts for when network messages are corrupt
* Some small fixes of asserts in debug in win32 detoured code

diff --git a/Engine/Source/Programs/UnrealBuildAccelerator/Common/Private/UbaNetworkBackendTcp.cpp b/Engine/Source/Programs/UnrealBuildAccelerator/Common/Private/UbaNetworkBackendTcp.cpp
--- a/Engine/Source/Programs/UnrealBuildAccelerator/Common/Private/UbaNetworkBackendTcp.cpp
+++ b/Engine/Source/Programs/UnrealBuildAccelerator/Common/Private/UbaNetworkBackendTcp.cpp
@@ -20,6 +20,7 @@
 #include <netdb.h>
 #include <arpa/inet.h>
 #include <ifaddrs.h>
+#include <poll.h>
 #define TIMEVAL timeval
 #define SOCKET_ERROR -1
 #define SOCKET int
@@ -32,6 +33,8 @@
 #define FreeAddrInfoW freeaddrinfo
 #define WSAGetLastError() errno
 #define strcpy_s(a, b, c) strcpy(a, c)
+#define WSAPOLLFD pollfd
+#define WSAPoll poll
 #endif
 
 namespace uba
@@ -577,54 +580,33 @@ namespace uba
 			return false;
 
 		int timeoutMs = 2000;
-		TIMEVAL timeval = { 1, 0 };
-		//		timeval.tv_sec = 1;
-		timeval.tv_usec = timeoutMs * 100;
 
-		fd_set write, err;
-		FD_ZERO(&write);
-		FD_ZERO(&err);
-		FD_SET(socketFd, &write);
-		FD_SET(socketFd, &err);
-
-		int maxFds = 0;
-		#if !PLATFORM_WINDOWS
-		maxFds = socketFd + 1;
-		#endif
-		// check if the socket is ready
-		res = select(maxFds, NULL, &write, &err, &timeval);
-		if (res == SOCKET_ERROR)
-			return logger.Error(TC("select failed (%s)"), LastErrorToText(WSAGetLastError()).data);
-
-		if (res == 0)
+		WSAPOLLFD p;
+		p.fd = socketFd;
+		p.revents = 0;
+		p.events = POLLOUT;
+		int pollRes = WSAPoll(&p, 1, timeoutMs);
+		if (!pollRes)
 		{
 			if (timedOut)
 				*timedOut = true;
 			return false;
 		}
 
-		if (FD_ISSET(socketFd, &err))
+		if (pollRes == SOCKET_ERROR)
 		{
-			// When running in wine we end up in this if host is not listening. (it returns immediately instead of 2000ms timeout on select)
-			// "Connection refused") is the response
-			//int resLen = sizeof(res);
-			//getsockopt(s, SOL_SOCKET, SO_ERROR, (void*)&res, &resLen);
-
-			if (timedOut)
-				*timedOut = true;
-			//logger.Info(TC("Error %s"), LastErrorToText(res));
+			int lastError = WSAGetLastError();
+			logger.Warning(TC("WSAPoll returned error %s (%s%s)"), LastErrorToText(lastError).data, nameHint);
 			return false;
 		}
 
-		if (!FD_ISSET(socketFd, &write))
+		if (p.revents & (POLLERR | POLLHUP | POLLNVAL))
 		{
-			if (timedOut)
-				*timedOut = true;
+			logger.Warning(TC("WSAPoll returned successful but with unexpected flags: %u"), p.revents);
 			return false;
 		}
 
 
-
 #if !PLATFORM_WINDOWS
 		// Before we send anything even though the
 		// the socket is writable, but let's make sure
@@ -729,7 +711,7 @@ namespace uba
 			int sent = (int)send(socket, (char*)b, u32(bufferLen), 0);
 			if (sent == SOCKET_ERROR)
 			{
-				#ifdef _DEBUG
+				#if UBA_DEBUG
 				logger.Warning(TC("ERROR sending socket (error: %s)"), LastErrorToText(WSAGetLastError()).data);
 				#endif
 				return false;
@@ -750,7 +732,6 @@ namespace uba
 		u32 recvLeft = bufferLen;
 		while (recvLeft)
 		{
-#if PLATFORM_WINDOWS
 			if (timeoutMs)
 			{
 				WSAPOLLFD p;
@@ -765,7 +746,7 @@ namespace uba
 				}
 				if (res == SOCKET_ERROR)
 				{
-#ifdef _DEBUG
+#if UBA_DEBUG && PLATFORM_WINDOWS
 					// When cancelling all kinds of errors can happen..
 					int lastError = WSAGetLastError();
 					if (lastError != WSAEINTR && lastError != WSAESHUTDOWN && lastError != WSAECONNABORTED && lastError != WSAECONNRESET) // Interrupted by cancel
@@ -773,12 +754,11 @@ namespace uba
 #endif
 				}
 			}
-#endif
 
 			int read = (int)recv(socket, (char*)buffer, recvLeft, 0);
 			if (read == 0)
 			{
-				//#ifdef _DEBUG
+				//#if UBA_DEBUG
 				//logger.Warning(TC("Socket closed while in recv"));
 				//#endif
 				return false;
@@ -787,14 +767,14 @@ namespace uba
 			if (read == SOCKET_ERROR)
 			{
 #if PLATFORM_WINDOWS
-#ifdef _DEBUG
+#if UBA_DEBUG
 				// When cancelling all kinds of errors can happen..
 				int lastError = WSAGetLastError();
 				if (lastError != WSAEINTR && lastError != WSAESHUTDOWN && lastError != WSAECONNABORTED && lastError != WSAECONNRESET) // Interrupted by cancel
 					logger.Warning(TC("ERROR receiving socket: %s (%s%s)"), LastErrorToText(lastError).data, hint1, hint2);
 #endif
 #else
-				if (!isFirstCall || errno != ECONNRESET)
+				if (!isFirstCall && errno != ECONNRESET)
 					logger.Error(TC("ERROR receiving socket %i after %u bytes (%s%s) (%s)"), socket, bufferLen, hint1, hint2, strerror(errno));
 #endif
 				return false;
diff --git a/Engine/Source/Programs/UnrealBuildAccelerator/Common/Private/UbaNetworkServer.cpp b/Engine/Source/Programs/UnrealBuildAccelerator/Common/Private/UbaNetworkServer.cpp
--- a/Engine/Source/Programs/UnrealBuildAccelerator/Common/Private/UbaNetworkServer.cpp
+++ b/Engine/Source/Programs/UnrealBuildAccelerator/Common/Private/UbaNetworkServer.cpp
@@ -225,7 +225,8 @@ namespace uba
 			u8 messageType = serviceIdAndMessageType & 0b111111;
 			u16 messageId = u16(headerData[1] << 8) | u16((*(u32*)(headerData + 2) & 0xff000000) >> 24);
 			u32 messageSize = *(u32*)(headerData + 2) & 0x00ffffff;
-			UBA_ASSERT(messageSize <= SendMaxSize);
+			UBA_ASSERTF(messageSize <= SendMaxSize, TC("Got message size %u which is larger than max %u. Protocol error?"), messageSize, SendMaxSize);
+			UBA_ASSERTF(serviceId < sizeof(NetworkServer::m_workerFunctions), TC("Got message with service id %u which is out of range. Protocol error?"), serviceId);
 
 			//m_logger.Debug(TC("Recv: %u, %u, %u, %u"), serviceId, messageType, id, size);
 			Worker* worker = conn.m_server.PopWorker();
diff --git a/Engine/Source/Programs/UnrealBuildAccelerator/Common/Private/UbaSession.cpp b/Engine/Source/Programs/UnrealBuildAccelerator/Common/Private/UbaSession.cpp
--- a/Engine/Source/Programs/UnrealBuildAccelerator/Common/Private/UbaSession.cpp
+++ b/Engine/Source/Programs/UnrealBuildAccelerator/Common/Private/UbaSession.cpp
@@ -2233,20 +2233,14 @@ namespace uba
 
 	float Session::UpdateCpuLoad()
 	{
+		u64 idleTime = 0;
+		u64 totalTime = 0;
+
 #if PLATFORM_WINDOWS
-		u64 idleTime, kernelTime, userTime;
+		u64 kernelTime, userTime;
 		if (!GetSystemTimes((FILETIME*)&idleTime, (FILETIME*)&kernelTime, (FILETIME*)&userTime))
 			return m_cpuLoad;
-		u64 totalTime = kernelTime + userTime;
-
-		u64 totalTimeSinceLastTime = totalTime - m_previousTotalCpuTime;
-		u64 idleTimeSinceLastTime = idleTime - m_previousIdleCpuTime;
-
-		float cpuLoad = 1.0f - ((totalTimeSinceLastTime > 0) ? (float(idleTimeSinceLastTime) / float(totalTimeSinceLastTime)) : 0);
-
-		m_previousTotalCpuTime = totalTime;
-		m_previousIdleCpuTime = idleTime;
-
+		totalTime = kernelTime + userTime;
 #elif PLATFORM_LINUX
 		int fd = open("/proc/stat", O_RDONLY);
 		if (fd != -1)
@@ -2283,16 +2277,8 @@ namespace uba
 					if (valueCount > 6)
 					{
 						u64 work = values[0] + values[1] + values[2];
-						u64 idleTime = values[3] + values[4] + values[5] + values[6];
-						u64 totalTime = work + idleTime;
-
-						u64 totalTimeSinceLastTime = totalTime - m_previousTotalCpuTime;
-						u64 idleTimeSinceLastTime = idleTime - m_previousIdleCpuTime;
-
-						float cpuLoad = 1.0f - ((totalTimeSinceLastTime > 0) ? (float(idleTimeSinceLastTime) / float(totalTimeSinceLastTime)) : 0);
-
-						m_previousTotalCpuTime = totalTime;
-						m_previousIdleCpuTime = idleTime;
+						idleTime = values[3] + values[4] + values[5] + values[6];
+						totalTime = work + idleTime;
 					}
 				}
 			}
@@ -2308,7 +2294,7 @@ namespace uba
 
         int res = 0;
         u64 work = 0;
-        u64 idleTime = 0;
+        idleTime = 0;
 
         res = host_processor_info(host, CpuInfoType, &CpuCount, (processor_info_array_t *)&CpuData, &CpuMsgCount);
 		if(res != KERN_SUCCESS)
@@ -2324,7 +2310,8 @@ namespace uba
                 idleTime += CpuData[i].cpu_ticks[CPU_STATE_IDLE];
         }
 
-		u64 totalTime = work + idleTime;
+		totalTime = work + idleTime;
+#endif
 
 		u64 totalTimeSinceLastTime = totalTime - m_previousTotalCpuTime;
 		u64 idleTimeSinceLastTime = idleTime - m_previousIdleCpuTime;
@@ -2334,7 +2321,6 @@ namespace uba
 		m_previousTotalCpuTime = totalTime;
 		m_previousIdleCpuTime = idleTime;
 
-#endif
 		// TODO: This is the wrong solution.. but can't repro the bad values some people get
 		if (cpuLoad >= 0 && cpuLoad <= 1.0f)
 			m_cpuLoad = cpuLoad;
diff --git a/Engine/Source/Programs/UnrealBuildAccelerator/Common/Private/UbaSessionClient.cpp b/Engine/Source/Programs/UnrealBuildAccelerator/Common/Private/UbaSessionClient.cpp
--- a/Engine/Source/Programs/UnrealBuildAccelerator/Common/Private/UbaSessionClient.cpp
+++ b/Engine/Source/Programs/UnrealBuildAccelerator/Common/Private/UbaSessionClient.cpp
@@ -974,7 +974,9 @@ namespace uba
 		m_sessionId = reader.ReadU32();
 		m_uiLanguage = reader.ReadU32();
 		m_detailedTrace = reader.ReadBool();
-
+		m_shouldSendLogToServer = reader.ReadBool();
+		if (m_shouldSendLogToServer)
+			m_logToFile = true;
 		BuildEnvironmentVariables(reader);
 
 		m_loop = true;
@@ -1379,13 +1381,13 @@ namespace uba
 						ProcessRec* rec = er->rec;
 						delete er;
 
-						#if UBA_DEBUG_LOG_ENABLED
-						if (false)
+						if (session.m_shouldSendLogToServer)
 						{
 							if (const tchar* logFile = h.GetStartInfo().logFile)
 							{
 								WrittenFile f;
 								f.name = logFile;
+								f.attributes = DefaultAttributes();
 								StringBuffer<> dest;
 								if (const tchar* lastSlash = TStrrchr(logFile, PathSeparator))
 									logFile = lastSlash + 1;
@@ -1394,7 +1396,6 @@ namespace uba
 								session.SendFile(*(ProcessImpl*)h.m_process, f, dest.data);
 							}
 						}
-						#endif
 
 						auto decreaseWeight = MakeGuard([&]()
 							{
diff --git a/Engine/Source/Programs/UnrealBuildAccelerator/Common/Private/UbaSessionServer.cpp b/Engine/Source/Programs/UnrealBuildAccelerator/Common/Private/UbaSessionServer.cpp
--- a/Engine/Source/Programs/UnrealBuildAccelerator/Common/Private/UbaSessionServer.cpp
+++ b/Engine/Source/Programs/UnrealBuildAccelerator/Common/Private/UbaSessionServer.cpp
@@ -132,6 +132,7 @@ namespace uba
 		m_remoteExecutionEnabled = info.remoteExecutionEnabled;
 		m_nameToHashTableEnabled = info.nameToHashTableEnabled;
 		m_memKillLoadPercent = info.memKillLoadPercent;
+		m_remoteLogEnabled = info.remoteLogEnabled;
 
 		if (m_resetCas)
 			m_storage.Reset();
@@ -571,6 +572,7 @@ namespace uba
 				writer.WriteU32(sessionId);
 				writer.WriteU32(m_uiLanguage);
 				writer.WriteBool(m_detailedTrace);
+				writer.WriteBool(m_remoteLogEnabled);
 				WriteRemoteEnvironmentVariables(writer);
 
 				m_trace.SessionAdded(sessionId, connectionInfo.GetId(), name.data, info.data); // Must be inside lock for TraceSessionUpdate() to not include
diff --git a/Engine/Source/Programs/UnrealBuildAccelerator/Common/Public/UbaNetwork.h b/Engine/Source/Programs/UnrealBuildAccelerator/Common/Public/UbaNetwork.h
--- a/Engine/Source/Programs/UnrealBuildAccelerator/Common/Public/UbaNetwork.h
+++ b/Engine/Source/Programs/UnrealBuildAccelerator/Common/Public/UbaNetwork.h
@@ -18,7 +18,7 @@ namespace uba
 	
 	static constexpr u32 SystemNetworkVersion = 1338;
 	static constexpr u32 StorageNetworkVersion = 3;
-	static constexpr u32 SessionNetworkVersion = 21;
+	static constexpr u32 SessionNetworkVersion = 22;
 
 
 	// Messages used over network between client/server (system, storage and session)
diff --git a/Engine/Source/Programs/UnrealBuildAccelerator/Common/Public/UbaSessionClient.h b/Engine/Source/Programs/UnrealBuildAccelerator/Common/Public/UbaSessionClient.h
--- a/Engine/Source/Programs/UnrealBuildAccelerator/Common/Public/UbaSessionClient.h
+++ b/Engine/Source/Programs/UnrealBuildAccelerator/Common/Public/UbaSessionClient.h
@@ -103,6 +103,7 @@ namespace uba
 		bool m_connected = false;
 		bool m_dedicated = false;
 		bool m_useStorage = true;
+		bool m_shouldSendLogToServer = false;
 		const tchar* m_terminationReason = nullptr;
 		u64 m_terminationTime = ~0ull;
 
diff --git a/Engine/Source/Programs/UnrealBuildAccelerator/Common/Public/UbaSessionServer.h b/Engine/Source/Programs/UnrealBuildAccelerator/Common/Public/UbaSessionServer.h
--- a/Engine/Source/Programs/UnrealBuildAccelerator/Common/Public/UbaSessionServer.h
+++ b/Engine/Source/Programs/UnrealBuildAccelerator/Common/Public/UbaSessionServer.h
@@ -141,6 +141,7 @@ namespace uba
 		ReaderWriterLock m_waitingProcessesLock;
 		bool m_allowWaitOnMem = false;
 		bool m_allowKillOnMem = false;
+		bool m_remoteLogEnabled = false;
 
 		SessionServer(const SessionServer&) = delete;
 		void operator=(const SessionServer&) = delete;
diff --git a/Engine/Source/Programs/UnrealBuildAccelerator/Common/Public/UbaSessionServerCreateInfo.h b/Engine/Source/Programs/UnrealBuildAccelerator/Common/Public/UbaSessionServerCreateInfo.h
--- a/Engine/Source/Programs/UnrealBuildAccelerator/Common/Public/UbaSessionServerCreateInfo.h
+++ b/Engine/Source/Programs/UnrealBuildAccelerator/Common/Public/UbaSessionServerCreateInfo.h
@@ -21,5 +21,6 @@ namespace uba
 		bool checkMemory = true;
 		bool allowWaitOnMem = false;
 		bool allowKillOnMem = false;
+		bool remoteLogEnabled = false; // If Uba is built in debug, then the logs will be sent back to server
 	};
 }
diff --git a/Engine/Source/Programs/UnrealBuildAccelerator/Core/Public/UbaHash.h b/Engine/Source/Programs/UnrealBuildAccelerator/Core/Public/UbaHash.h
--- a/Engine/Source/Programs/UnrealBuildAccelerator/Core/Public/UbaHash.h
+++ b/Engine/Source/Programs/UnrealBuildAccelerator/Core/Public/UbaHash.h
@@ -167,15 +167,15 @@ namespace uba
 		tchar str[41];
 	};
 
-	#ifdef _DEBUG
+	#if UBA_DEBUG
 	inline void CheckPath(const tchar* fileName)
 	{
 		auto pos = fileName;
-		tchar lastChar = L' ';
+		tchar lastChar = ' ';
 		//bool hasSlash = false;
 		while (tchar c = *pos++)
 		{
-			UBA_ASSERTF(c < 'A' || c > 'Z', TC("Path is not valid (%s)"), fileName);
+			UBA_ASSERTF(!CaseInsensitiveFs || c < 'A' || c > 'Z', TC("Path is not valid (%s)"), fileName);
 			UBA_ASSERTF(c != L'.' || lastChar != L'.', TC("Path is not valid (%s)"), fileName);
 
 			#if PLATFORM_WINDOWS
diff --git a/Engine/Source/Programs/UnrealBuildAccelerator/Detours/Private/Windows/UbaDetoursFunctionsKernelBase.inl b/Engine/Source/Programs/UnrealBuildAccelerator/Detours/Private/Windows/UbaDetoursFunctionsKernelBase.inl
--- a/Engine/Source/Programs/UnrealBuildAccelerator/Detours/Private/Windows/UbaDetoursFunctionsKernelBase.inl
+++ b/Engine/Source/Programs/UnrealBuildAccelerator/Detours/Private/Windows/UbaDetoursFunctionsKernelBase.inl
@@ -2812,7 +2812,7 @@ DWORD Detoured_GetFullPathNameA(LPCSTR lpFileName, DWORD nBufferLength, LPSTR lp
 {
 	DETOURED_CALL(GetFullPathNameA);
 	DEBUG_LOG_TRUE(L"GetFullPathNameA", L"");
-	UBA_ASSERT(!g_runningRemote);
+	UBA_ASSERT(!g_runningRemote || g_isRunningWine); // Wine is calling GetFullPathNameW.. need to check so windows is doing that too
 	return True_GetFullPathNameA(lpFileName, nBufferLength, lpBuffer, lpFilePart);
 }
 
@@ -3575,8 +3575,14 @@ HRESULT Detoured_SHCreateStreamOnFileW(LPCWSTR pszFile, DWORD grfMode, IStream**
 
 BOOL Detoured_PathFileExistsW(LPCWSTR pszPath)
 {
-	UBA_ASSERTF(!g_runningRemote, L"%ls", pszPath);
-	return True_PathFileExistsW(pszPath);
+	DEBUG_LOG_DETOURED(L"PathFileExistsW", L"CALLING GetFileAttributesW (%s)", pszPath);
+	DWORD attributes = Detoured_GetFileAttributesW(pszPath);
+	return attributes != INVALID_FILE_ATTRIBUTES;
+
+	//UBA_ASSERTF(!g_runningRemote, L"%ls", pszPath);
+	//auto res = True_PathFileExistsW(pszPath);
+	//DEBUG_LOG_TRUE(L"PathFileExistsW", L"(%ls) -> %ls", pszPath, res);
+	//return res;
 }
 
 #endif // DETOURED_INCLUDE_DEBUG
diff --git a/Engine/Source/Programs/UnrealBuildAccelerator/Host/Private/UbaExports.cpp b/Engine/Source/Programs/UnrealBuildAccelerator/Host/Private/UbaExports.cpp
--- a/Engine/Source/Programs/UnrealBuildAccelerator/Host/Private/UbaExports.cpp
+++ b/Engine/Source/Programs/UnrealBuildAccelerator/Host/Private/UbaExports.cpp
@@ -301,6 +301,7 @@ uba::StorageClient* CreateStorageClient(uba::NetworkClient& client, const uba::t
 		info->detailedTrace = detailedTrace;
 		info->allowWaitOnMem = allowWaitOnMem;
 		info->allowKillOnMem = allowKillOnMem;
+		//info->remoteLogEnabled = true;
 		return info;
 	}
 
diff --git a/Engine/Source/Programs/UnrealBuildAccelerator/Test/Private/UbaTestBasics.h b/Engine/Source/Programs/UnrealBuildAccelerator/Test/Private/UbaTestBasics.h
--- a/Engine/Source/Programs/UnrealBuildAccelerator/Test/Private/UbaTestBasics.h
+++ b/Engine/Source/Programs/UnrealBuildAccelerator/Test/Private/UbaTestBasics.h
@@ -187,6 +187,8 @@ namespace uba
 		if (FileExists(logger, TC("UbaTestFile")))
 			return logger.Error(TC("Found non-existing file UbaTestFile"));
 
+		// CreateHardLinkW is a symbolic link on non-windows.. need to revisit
+		#if PLATFORM_WINDOWS
 		if (!FileExists(logger, TC("UbaTestFile2")))
 			return logger.Error(TC("Failed to find file UbaTestFile2"));
 
@@ -203,6 +205,7 @@ namespace uba
 
 		if (!DeleteFileW(TC("UbaTestFile2")))
 			return false;
+		#endif
 
 		LoggerWithWriter nullLogger(g_nullLogWriter);
 		if (TraverseDir(nullLogger, TC("TestDir2"), [&](const DirectoryEntry&) {}, true))
