[UBA]
* Fixed so all visualization of times in visualizer is taking timer frequency into account

diff --git a/Engine/Source/Programs/UnrealBuildAccelerator/Common/Private/UbaProcess.cpp b/Engine/Source/Programs/UnrealBuildAccelerator/Common/Private/UbaProcess.cpp
--- a/Engine/Source/Programs/UnrealBuildAccelerator/Common/Private/UbaProcess.cpp
+++ b/Engine/Source/Programs/UnrealBuildAccelerator/Common/Private/UbaProcess.cpp
@@ -820,7 +820,7 @@ namespace uba
 					if (m_startInfo.outputStatsThresholdMs && TimeToMs(m_processStats.GetTotalTime()) > m_startInfo.outputStatsThresholdMs)
 					{
 						m_session.PrintProcessStats(m_processStats, logName.data);
-						m_processStats.Print(m_session.m_logger);
+						m_processStats.Print(m_session.m_logger, GetFrequency());
 					}
 
 					return false;
diff --git a/Engine/Source/Programs/UnrealBuildAccelerator/Common/Private/UbaStats.cpp b/Engine/Source/Programs/UnrealBuildAccelerator/Common/Private/UbaStats.cpp
--- a/Engine/Source/Programs/UnrealBuildAccelerator/Common/Private/UbaStats.cpp
+++ b/Engine/Source/Programs/UnrealBuildAccelerator/Common/Private/UbaStats.cpp
@@ -5,11 +5,11 @@
 
 namespace uba
 {
-	void ProcessStats::Print(Logger& logger)
+	void ProcessStats::Print(Logger& logger, u64 frequency)
 	{
-		logger.Info(TC("  Total              %8u %9s"), GetTotalCount(), TimeToText(GetTotalTime()).str);
-		logger.Info(TC("  WaitOnResponse     %8u %9s"), waitOnResponse.count.load(), TimeToText(waitOnResponse.time).str);
-		logger.Info(TC("  Host                %17s"), TimeToText(hostTotalTime).str);
+		logger.Info(TC("  Total              %8u %9s"), GetTotalCount(), TimeToText(GetTotalTime(), false, frequency).str);
+		logger.Info(TC("  WaitOnResponse     %8u %9s"), waitOnResponse.count.load(), TimeToText(waitOnResponse.time, false, frequency).str);
+		logger.Info(TC("  Host                %17s"), TimeToText(hostTotalTime, false, frequency).str);
 		logger.Info(TC(""));
 
 		struct Stat { const char* name; u64 nameLen; const Timer& timer; };
@@ -23,18 +23,18 @@ namespace uba
 		const tchar empty[] = TC("                   ");
 		for (Stat& s : stats)
 			if (s.timer.count)
-				logger.Info(TC("  %c%hs%s %8u %9s"), ToUpper(s.name[0]), s.name+1, empty + s.nameLen, s.timer.count.load(), TimeToText(s.timer.time).str);
+				logger.Info(TC("  %c%hs%s %8u %9s"), ToUpper(s.name[0]), s.name+1, empty + s.nameLen, s.timer.count.load(), TimeToText(s.timer.time, false, frequency).str);
 
 		logger.Info(TC(""));
 
 		logger.Info(TC("  HighestMem                  %9s"), BytesToText(usedMemory).str);
-		logger.Info(TC("  Startup Time                %9s"), TimeToText(startupTime).str);
-		logger.Info(TC("  Exit Time                   %9s"), TimeToText(exitTime).str);
-		logger.Info(TC("  CPU Time                    %9s"), TimeToText(cpuTime).str);
-		logger.Info(TC("  Wall Time                   %9s"), TimeToText(wallTime).str);
+		logger.Info(TC("  Startup Time                %9s"), TimeToText(startupTime, false, frequency).str);
+		logger.Info(TC("  Exit Time                   %9s"), TimeToText(exitTime, false, frequency).str);
+		logger.Info(TC("  CPU Time                    %9s"), TimeToText(cpuTime, false, frequency).str);
+		logger.Info(TC("  Wall Time                   %9s"), TimeToText(wallTime, false, frequency).str);
 	}
 
-	void SystemStats::Print(Logger& logger, bool writeHeader)
+	void SystemStats::Print(Logger& logger, bool writeHeader, u64 frequency)
 	{
 		if (writeHeader)
 			logger.Info(TC("  --- Platform system stats summary ---"));
@@ -50,7 +50,7 @@ namespace uba
 		const tchar empty[] = TC("                   ");
 		for (Stat& s : stats)
 			if (s.timer.count)
-				logger.Info(TC("  %c%hs%s %8u %9s"), ToUpper(s.name[0]), s.name+1, empty + s.nameLen, s.timer.count.load(), TimeToText(s.timer.time).str);
+				logger.Info(TC("  %c%hs%s %8u %9s"), ToUpper(s.name[0]), s.name+1, empty + s.nameLen, s.timer.count.load(), TimeToText(s.timer.time, false, frequency).str);
 
 		if (writeHeader)
 			logger.Info(TC(""));
@@ -98,37 +98,37 @@ namespace uba
 		#undef UBA_STORAGE_STAT
 	}
 
-	void StorageStats::Print(Logger& logger)
+	void StorageStats::Print(Logger& logger, u64 frequency)
 	{
 		if (calculateCasKey.count)
-			logger.Info(TC("  CalculateCasKeys     %6u %9s"), calculateCasKey.count.load(), TimeToText(calculateCasKey.time).str);
+			logger.Info(TC("  CalculateCasKeys     %6u %9s"), calculateCasKey.count.load(), TimeToText(calculateCasKey.time, false, frequency).str);
 		if (ensureCas.count)
-			logger.Info(TC("  EnsureCas            %6u %9s"), ensureCas.count.load(), TimeToText(ensureCas.time).str);
+			logger.Info(TC("  EnsureCas            %6u %9s"), ensureCas.count.load(), TimeToText(ensureCas.time, false, frequency).str);
 		if (recvCas.count)
 		{
 			logger.Info(TC("  ReceiveCas           %6u %9s"), recvCas.count.load(), TimeToText(recvCas.time).str);
 			logger.Info(TC("     Bytes Raw/Comp %9s %9s"), BytesToText(recvCasBytesRaw).str, BytesToText(recvCasBytesComp).str);
 			if (decompressRecv.count)
-				logger.Info(TC("     Decompress        %6u %9s"), decompressRecv.count.load(), TimeToText(decompressRecv.time).str);
+				logger.Info(TC("     Decompress        %6u %9s"), decompressRecv.count.load(), TimeToText(decompressRecv.time, false, frequency).str);
 		}
 		if (sendCas.count)
 		{
-			logger.Info(TC("  SendCas              %6u %9s"), sendCas.count.load(), TimeToText(sendCas.time).str);
+			logger.Info(TC("  SendCas              %6u %9s"), sendCas.count.load(), TimeToText(sendCas.time, false, frequency).str);
 			logger.Info(TC("     Bytes Raw/Comp %9s %9s"), BytesToText(sendCasBytesRaw).str, BytesToText(sendCasBytesComp).str);
-			logger.Info(TC("     Compress          %6u %9s"), compressSend.count.load(), TimeToText(compressSend.time).str);
+			logger.Info(TC("     Compress          %6u %9s"), compressSend.count.load(), TimeToText(compressSend.time, false, frequency).str);
 		}
 		if (createCas.count)
 		{
-			logger.Info(TC("  CreateCas            %6u %9s"), createCas.count.load(), TimeToText(createCas.time).str);
+			logger.Info(TC("  CreateCas            %6u %9s"), createCas.count.load(), TimeToText(createCas.time, false, frequency).str);
 			logger.Info(TC("     Bytes Raw/Comp %9s %9s"), BytesToText(createCasBytesRaw).str, BytesToText(createCasBytesComp).str);
-			logger.Info(TC("     Compress          %6u %9s"), compressWrite.count.load(), TimeToText(compressWrite.time).str);
+			logger.Info(TC("     Compress          %6u %9s"), compressWrite.count.load(), TimeToText(compressWrite.time, false, frequency).str);
 		}
 		if (copyOrLink.count)
-			logger.Info(TC("  CopyOrLink           %6u %9s"), copyOrLink.count.load(), TimeToText(copyOrLink.time).str);
+			logger.Info(TC("  CopyOrLink           %6u %9s"), copyOrLink.count.load(), TimeToText(copyOrLink.time, false, frequency).str);
 		if (copyOrLinkWait.count)
-			logger.Info(TC("  CopyOrLinkWait       %6u %9s"), copyOrLinkWait.count.load(), TimeToText(copyOrLinkWait.time).str);
+			logger.Info(TC("  CopyOrLinkWait       %6u %9s"), copyOrLinkWait.count.load(), TimeToText(copyOrLinkWait.time, false, frequency).str);
 		if (decompressToMem.count)
-			logger.Info(TC("  DecompressToMem      %6u %9s"), decompressToMem.count.load(), TimeToText(decompressToMem.time).str);
+			logger.Info(TC("  DecompressToMem      %6u %9s"), decompressToMem.count.load(), TimeToText(decompressToMem.time, false, frequency).str);
 	}
 
 	thread_local StorageStats* t_storageStats;
@@ -174,17 +174,17 @@ namespace uba
 	template<typename T>
 	void LogStat(Logger& logger, const char* name, const T&) {}
 
-	void LogStat(Logger& logger, const char* name, const Timer& timer)
+	void LogStat(Logger& logger, const char* name, const Timer& timer, u64 frequency)
 	{
 		if (!timer.count)
 			return;
 		const tchar empty[] = TC("                   ");
-		logger.Info(TC("  %c%hs%s %8u %9s"), ToUpper(name[0]), name+1, empty + strlen(name)+1, timer.count.load(), TimeToText(timer.time).str);
+		logger.Info(TC("  %c%hs%s %8u %9s"), ToUpper(name[0]), name+1, empty + strlen(name)+1, timer.count.load(), TimeToText(timer.time, false, frequency).str);
 	}
 
-	void SessionStats::Print(Logger& logger)
+	void SessionStats::Print(Logger& logger, u64 frequency)
 	{
-		#define UBA_SESSION_STAT(type, var, ver) LogStat(logger, #var, var);
+		#define UBA_SESSION_STAT(type, var, ver) LogStat(logger, #var, var, frequency);
 		UBA_SESSION_STATS
 		#undef UBA_SESSION_STAT
 	}
@@ -223,13 +223,13 @@ namespace uba
 		#undef UBA_SESSION_SUMMARY_STAT
 	}
 
-	void SessionSummaryStats::Print(Logger& logger)
+	void SessionSummaryStats::Print(Logger& logger, u64 frequency)
 	{
 		#define UBA_SESSION_SUMMARY_STAT(T, V) LogStat(logger, #V, V);
 		UBA_SESSION_SUMMARY_STATS
 		#undef UBA_SESSION_SUMMARY_STAT
-		stats.Print(logger);
-		logger.Info(TC("  MemoryPressureWait          %9s"), TimeToText(waitMemPressure).str);
+		stats.Print(logger, frequency);
+		logger.Info(TC("  MemoryPressureWait          %9s"), TimeToText(waitMemPressure, false, frequency).str);
 		logger.Info(TC("  ProcessesKilled             %9llu"), killCount);
 		logger.Info(TC(""));
 	}
diff --git a/Engine/Source/Programs/UnrealBuildAccelerator/Common/Public/UbaStats.h b/Engine/Source/Programs/UnrealBuildAccelerator/Common/Public/UbaStats.h
--- a/Engine/Source/Programs/UnrealBuildAccelerator/Common/Public/UbaStats.h
+++ b/Engine/Source/Programs/UnrealBuildAccelerator/Common/Public/UbaStats.h
@@ -40,7 +40,7 @@ namespace uba
 		void Write(BinaryWriter& writer);
 		void Read(BinaryReader& reader);
 		void Add(const StorageStats& other);
-		void Print(Logger& logger);
+		void Print(Logger& logger, u64 frequency = GetFrequency());
 		static StorageStats* GetCurrent();
 	};
 
@@ -74,7 +74,7 @@ namespace uba
 		void Write(BinaryWriter& writer);
 		void Read(BinaryReader& reader, u32 version);
 		void Add(const SessionStats& other);
-		void Print(Logger& logger);
+		void Print(Logger& logger, u64 frequency = GetFrequency());
 		static SessionStats* GetCurrent();
 	};
 
@@ -109,7 +109,7 @@ namespace uba
 		SessionSummaryStats() : waitMemPressure(0), killCount(0) {}
 		void Write(BinaryWriter& writer);
 		void Read(BinaryReader& reader, u32 version);
-		void Print(Logger& logger);
+		void Print(Logger& logger, u64 frequency = GetFrequency());
 	};
 
 }
diff --git a/Engine/Source/Programs/UnrealBuildAccelerator/Core/Public/UbaProcessStats.h b/Engine/Source/Programs/UnrealBuildAccelerator/Core/Public/UbaProcessStats.h
--- a/Engine/Source/Programs/UnrealBuildAccelerator/Core/Public/UbaProcessStats.h
+++ b/Engine/Source/Programs/UnrealBuildAccelerator/Core/Public/UbaProcessStats.h
@@ -57,7 +57,7 @@ namespace uba
 
 		u64 hostTotalTime = 0;
 
-		void Print(Logger& logger);
+		void Print(Logger& logger, u64 frequency = GetFrequency());
 
 		u64 GetTotalTime()
 		{
@@ -147,7 +147,7 @@ namespace uba
 
 		void Read(BinaryReader& reader);
 		void Write(BinaryWriter& writer);
-		void Print(Logger& logger, bool writeHeader);
+		void Print(Logger& logger, bool writeHeader, u64 frequency = GetFrequency());
 		void Add(const SystemStats& other);
 		static SystemStats& GetCurrent();
 		static SystemStats& GetGlobal();
diff --git a/Engine/Source/Programs/UnrealBuildAccelerator/Visualizer/Private/UbaVisualizer.cpp b/Engine/Source/Programs/UnrealBuildAccelerator/Visualizer/Private/UbaVisualizer.cpp
--- a/Engine/Source/Programs/UnrealBuildAccelerator/Visualizer/Private/UbaVisualizer.cpp
+++ b/Engine/Source/Programs/UnrealBuildAccelerator/Visualizer/Private/UbaVisualizer.cpp
@@ -1027,11 +1027,11 @@ namespace uba
 			if (hasExited)
 			{
 				LineCountLogger counter;
-				process.processStats.Print(counter);
+				process.processStats.Print(counter, m_traceView.frequency);
 				u32 prevLineCount = counter.lineCount;
-				process.sessionStats.Print(counter);
-				process.storageStats.Print(counter);
-				process.systemStats.Print(counter, false);
+				process.sessionStats.Print(counter, m_traceView.frequency);
+				process.storageStats.Print(counter, m_traceView.frequency);
+				process.systemStats.Print(counter, false, m_traceView.frequency);
 				hasStorageStats = prevLineCount != counter.lineCount;
 				lineCount = 6 + counter.lineCount;
 				if (hasStorageStats)
@@ -1081,18 +1081,18 @@ namespace uba
 			if (process.stop != ~u64(0))
 			{
 				logger.Info(L"  ----------- Process stats -----------");
-				process.processStats.Print(logger);
+				process.processStats.Print(logger, m_traceView.frequency);
 				if (hasStorageStats)
 				{
 					logger.Info(L"");
 					logger.Info(L"  ----------- Session stats -----------");
-					process.sessionStats.Print(logger);
+					process.sessionStats.Print(logger, m_traceView.frequency);
 					logger.Info(L"");
 					logger.Info(L"  ----------- Storage stats -----------");
-					process.storageStats.Print(logger);
+					process.storageStats.Print(logger, m_traceView.frequency);
 					logger.Info(L"");
 					logger.Info(L"  ----------- System stats ------------");
-					process.systemStats.Print(logger, false);
+					process.systemStats.Print(logger, false, m_traceView.frequency);
 				}
 			}
 		}
