UBA: Fix for UbaAgents waiting for host now that poll is used.

diff --git a/Engine/Source/Programs/UnrealBuildAccelerator/Common/Private/UbaNetworkBackendTcp.cpp b/Engine/Source/Programs/UnrealBuildAccelerator/Common/Private/UbaNetworkBackendTcp.cpp
--- a/Engine/Source/Programs/UnrealBuildAccelerator/Common/Private/UbaNetworkBackendTcp.cpp
+++ b/Engine/Source/Programs/UnrealBuildAccelerator/Common/Private/UbaNetworkBackendTcp.cpp
@@ -600,7 +600,14 @@ namespace uba
 			return false;
 		}
 
-		if (p.revents & (POLLERR | POLLHUP | POLLNVAL))
+		if (p.revents & (POLLHUP))
+		{
+			//logger.Info(TC("Connection not ready..."));
+			*timedOut = true;
+			return false;
+		}
+
+		if (p.revents & (POLLERR | POLLNVAL))
 		{
 			logger.Warning(TC("WSAPoll returned successful but with unexpected flags: %u"), p.revents);
 			return false;
