[UBA]
* Fixed path validation code
* Added null guard in NetworkBackendTcp::Connect and collapsed to early outs into one

diff --git a/Engine/Source/Programs/UnrealBuildAccelerator/Common/Private/UbaNetworkBackendTcp.cpp b/Engine/Source/Programs/UnrealBuildAccelerator/Common/Private/UbaNetworkBackendTcp.cpp
--- a/Engine/Source/Programs/UnrealBuildAccelerator/Common/Private/UbaNetworkBackendTcp.cpp
+++ b/Engine/Source/Programs/UnrealBuildAccelerator/Common/Private/UbaNetworkBackendTcp.cpp
@@ -586,12 +586,6 @@ namespace uba
 		p.revents = 0;
 		p.events = POLLOUT;
 		int pollRes = WSAPoll(&p, 1, timeoutMs);
-		if (!pollRes)
-		{
-			if (timedOut)
-				*timedOut = true;
-			return false;
-		}
 
 		if (pollRes == SOCKET_ERROR)
 		{
@@ -600,10 +594,10 @@ namespace uba
 			return false;
 		}
 
-		if (p.revents & (POLLHUP))
+		if (!pollRes || p.revents & (POLLHUP)) // Treat hangup as timeout (since we want retry if that happens)
 		{
-			//logger.Info(TC("Connection not ready..."));
-			*timedOut = true;
+			if (timedOut)
+				*timedOut = true;
 			return false;
 		}
 
@@ -624,8 +618,8 @@ namespace uba
 		memset(&junk, 0, sizeof(junk));
 		if (getpeername(socketFd, (struct sockaddr *)&junk, &length) != 0)
 		{
-			//logger.Info(TC("Connection not ready..."));
-			*timedOut = true;
+			if (timedOut)
+				*timedOut = true;
 			return false;
 		}
 
diff --git a/Engine/Source/Programs/UnrealBuildAccelerator/Core/Public/UbaHash.h b/Engine/Source/Programs/UnrealBuildAccelerator/Core/Public/UbaHash.h
--- a/Engine/Source/Programs/UnrealBuildAccelerator/Core/Public/UbaHash.h
+++ b/Engine/Source/Programs/UnrealBuildAccelerator/Core/Public/UbaHash.h
@@ -171,23 +171,19 @@ namespace uba
 	inline void CheckPath(const tchar* fileName)
 	{
 		auto pos = fileName;
-		tchar lastChar = ' ';
-		//bool hasSlash = false;
+		tchar history[3] = { 0 };
 		while (tchar c = *pos++)
 		{
-			UBA_ASSERTF(!CaseInsensitiveFs || c < 'A' || c > 'Z', TC("Path is not valid (%s)"), fileName);
-			UBA_ASSERTF(c != L'.' || lastChar != L'.', TC("Path is not valid (%s)"), fileName);
-
-			#if PLATFORM_WINDOWS
-			UBA_ASSERTF(c != '/', TC("Path is not valid (%s)"), fileName);
-			UBA_ASSERTF(c != '\\' || lastChar != '\\' || TStrstr(fileName, TC("\\\\.\\pipe")), TC("Path is not valid (%s)"), fileName);
-			#else
-			UBA_ASSERTF(c != '\\', TC("Path is not valid (%s)"), fileName);
-			#endif
-			//hasSlash |= c == '\\';
-			lastChar = c;
+			UBA_ASSERTF(!CaseInsensitiveFs || c < 'A' || c > 'Z', TC("Path is not valid (%s)"), fileName); // No upper case characters if case insensitive
+			UBA_ASSERTF(!(c == PathSeparator && history[0] == '.' && history[1] == PathSeparator) || (IsWindows && TStrstr(fileName, TC("\\\\.\\pipe"))), TC("Path is not valid (%s)"), fileName);  //  /./ not allowed
+			UBA_ASSERTF(!(c == PathSeparator && history[0] == '.' && history[1] == '.' && history[2] == PathSeparator), TC("Path is not valid (%s)"), fileName); // /../ not allowed
+			UBA_ASSERTF(c != NonPathSeparator, TC("Path is not valid (%s)"), fileName);
+			UBA_ASSERTF(c != PathSeparator || history[0] != PathSeparator || (IsWindows && TStrstr(fileName, TC("\\\\.\\pipe"))), TC("Path is not valid (%s)"), fileName); // Double backslash
+			history[2] = history[1];
+			history[1] = history[0];
+			history[0] = c;
 		}
-		UBA_ASSERTF(lastChar != ' ' && (lastChar != '\\' || (fileName[1] == ':' && fileName[3] == 0)), TC("Path is not valid (%s)"), fileName);
+		UBA_ASSERTF(history[0] != ' ' && (history[0] != PathSeparator || (fileName[1] == ':' && fileName[3] == 0)), TC("Path is not valid (%s)"), fileName);
 		// Commented out because this asserts on dlls in subfolders on remote machines (example: 2057\clui.dll). Likely related to helper not having some language id as host
 		// UBA_ASSERTF(!hasSlash || fileName[1] == ':', TC("Invalid path: %s"), fileName);
 	}
