UnrealBuildAccelerator: Always compile universal for Mac

#rnx

diff --git a/Engine/Source/Programs/UnrealBuildAccelerator/UbaAgent.Target.cs b/Engine/Source/Programs/UnrealBuildAccelerator/UbaAgent.Target.cs
--- a/Engine/Source/Programs/UnrealBuildAccelerator/UbaAgent.Target.cs
+++ b/Engine/Source/Programs/UnrealBuildAccelerator/UbaAgent.Target.cs
@@ -83,6 +83,14 @@ public class UbaAgentTarget : TargetRules
 
 		if (Target.Platform.IsInGroup(UnrealPlatformGroup.Apple))
 		{
+			if (!Rules.Architectures.Contains(UnrealArch.X64))
+			{
+				Rules.Architectures.Architectures.Add(UnrealArch.X64);
+			}
+			if (!Rules.Architectures.Contains(UnrealArch.Arm64))
+			{
+				Rules.Architectures.Architectures.Add(UnrealArch.Arm64);
+			}
 			string BinaryPrefix = Rules.bShouldCompileAsDLL ? "lib" : string.Empty;
 			string BinaryExt = Rules.bShouldCompileAsDLL ? ".dylib" : string.Empty;
 			Rules.OutputFile = Path.Combine(BinariesFolder, $"{BinaryPrefix}{Rules.LaunchModuleName}{BinaryExt}");
