UBA: Fix bus error on Intel Macs due to a stacksize limit

diff --git a/Engine/Source/Programs/UnrealBuildAccelerator/Common/Private/UbaThread.cpp b/Engine/Source/Programs/UnrealBuildAccelerator/Common/Private/UbaThread.cpp
--- a/Engine/Source/Programs/UnrealBuildAccelerator/Common/Private/UbaThread.cpp
+++ b/Engine/Source/Programs/UnrealBuildAccelerator/Common/Private/UbaThread.cpp
@@ -46,10 +46,24 @@ namespace uba
 		m_handle = CreateThread(NULL, 0, [](LPVOID p) -> DWORD { return ((Thread*)p)->m_func(); }, this, 0, NULL);
 		AlternateGroupAffinity(m_handle);
 #else
+		int err = 0;
+
 		m_finished.Create(true);
 		static_assert(sizeof(pthread_t) <= sizeof(m_handle), "");
 		auto& pth = *(pthread_t*)&m_handle;
-		int err = pthread_create(&pth, NULL, [](void* p) -> void*
+
+		pthread_attr_t tattr;
+		// initialized with default attributes
+		err = pthread_attr_init(&tattr);
+
+		// TODO: Need to figure out a better value, or decrease stack usage
+		// without this though we get a bus error on Intel Macs
+		#if !defined(__arm__) && !defined(__arm64__)
+		size_t size = PTHREAD_STACK_MIN * 500;
+		err = pthread_attr_setstacksize(&tattr, size);
+		#endif
+
+		err = pthread_create(&pth, &tattr, [](void* p) -> void*
 			{
 				auto& t = *(Thread*)p;
 				int res = t.m_func();
