[UBA]
* Reduced stack usage in client session by reusing 256kb stack binary reader
* Changed so update dir table and hash table messages are not sent if client is already enough up to date
* Fixed race condition writing dir table memory size. This has been in here forever and maybe can't happen in practice.. but wrong lock was taken when updating table memory size (different than what GetDirectoryTableSize was using)

diff --git a/Engine/Source/Programs/UnrealBuildAccelerator/Common/Private/UbaSessionClient.cpp b/Engine/Source/Programs/UnrealBuildAccelerator/Common/Private/UbaSessionClient.cpp
--- a/Engine/Source/Programs/UnrealBuildAccelerator/Common/Private/UbaSessionClient.cpp
+++ b/Engine/Source/Programs/UnrealBuildAccelerator/Common/Private/UbaSessionClient.cpp
@@ -526,14 +526,17 @@ namespace uba
 		NetworkMessage networkMsg(m_client, ServiceId, SessionMessageType_DeleteFile, writer);
 		writer.WriteStringKey(msg.fileNameKey);
 		writer.WriteString(msg.fileName);
-		StackBinaryReader<1024> reader;
+		StackBinaryReader<SendMaxSize> reader;
 		if (!networkMsg.Send(reader, Stats().deleteFileMsg))
 			return false;
 		out.result = reader.ReadBool();
 		out.errorCode = reader.ReadU32();
 		if (out.result)
-			if (!SendUpdateDirectoryTable())
+		{
+			reader.Reset();
+			if (!SendUpdateDirectoryTable(reader))
 				return false;
+		}
 		out.directoryTableSize = GetDirectoryTableSize();
 		return true;
 	}
@@ -552,7 +555,7 @@ namespace uba
 			writer.WriteString(msg.fromName);
 			writer.WriteStringKey(msg.toKey);
 			writer.WriteString(msg.toName);
-			StackBinaryReader<1024> reader;
+			StackBinaryReader<SendMaxSize> reader;
 			if (!networkMsg.Send(reader, Stats().copyFileMsg))
 				return false;
 			out.fromName.Append(msg.fromName);
@@ -560,8 +563,11 @@ namespace uba
 			out.closeId = ~0u;
 			out.errorCode = reader.ReadU32();
 			if (!out.errorCode)
-				if (!SendUpdateDirectoryTable())
+			{
+				reader.Reset();
+				if (!SendUpdateDirectoryTable(reader))
 					return false;
+			}
 			out.directoryTableSize = GetDirectoryTableSize();
 			return true;
 		}
@@ -781,7 +787,7 @@ namespace uba
 		}
 	};
 
-	bool SessionClient::UpdateDirectoryTableFromServer(BinaryReader& reader)
+	bool SessionClient::UpdateDirectoryTableFromServer(StackBinaryReader<SendMaxSize>& reader)
 	{
 		auto& dirTable = m_directoryTable;
 
@@ -790,7 +796,7 @@ namespace uba
 		{
 			if (!isFirst)
 			{
-				reader.SetPosition(0);
+				reader.Reset();
 
 				StackBinaryWriter<1024> writer;
 				NetworkMessage msg(m_client, ServiceId, SessionMessageType_GetDirectoriesFromServer, writer);
@@ -827,7 +833,11 @@ namespace uba
 
 			if (reader.GetPosition() < m_client.GetMessageMaxSize() - m_client.GetMessageReceiveHeaderSize())
 			{
-				dirTable.ParseDirectoryTable(m_directoryTableMemPos);
+				//dirTable.ParseDirectoryTable(m_directoryTableMemPos); // This is not needed.. we never read from the directory table in the client session
+				{
+					ScopedWriteLock lock2(dirTable.m_memoryLock);
+					dirTable.m_memorySize = m_directoryTableMemPos;
+				}
 				ActiveUpdateDirectoryEntry::UpdateReadPosLess(m_firstEmptyWait, m_directoryTableMemPos);
 				break;
 			}
@@ -837,7 +847,7 @@ namespace uba
 		return true;
 	}
 
-	bool SessionClient::UpdateNameToHashTableFromServer(BinaryReader& reader)
+	bool SessionClient::UpdateNameToHashTableFromServer(StackBinaryReader<SendMaxSize>& reader)
 	{
 		u32 serverTableSize = 0;
 		bool isFirst = true;
@@ -858,7 +868,7 @@ namespace uba
 				writer.WriteU32(serverTableSize);
 				writer.WriteU32(localTableSize);
 
-				reader.SetPosition(0);
+				reader.Reset();
 				if (!msg.Send(reader, Stats().getHashesMsg))
 					return false;
 			}
@@ -1066,9 +1076,27 @@ namespace uba
 
 		}
 
+		u32 neededDirectoryTableSize = reader.ReadU32();
+		u32 neededHashTableSize = reader.ReadU32();
+
 		if (!out.empty())
-			if (!SendUpdateDirectoryTable())
+		{
+			if (neededDirectoryTableSize > GetDirectoryTableSize())
+			{
+				reader.Reset();
+				if (!SendUpdateDirectoryTable(reader))
+					return false;
+			}
+		}
+
+		// Always nice to update name-to-hash table since it can reduce number of messages while building.
+		u32 hashTableMemSize = m_nameToHashMemLock.ScopedRead([this]() { return u32(m_nameToHashTableMem.writtenSize); });
+		if (neededHashTableSize > hashTableMemSize)
+		{
+			reader.Reset();
+			if (!SendUpdateNameToHashTable(reader))
 				return false;
+		}
 
 		return true;
 	}
@@ -1084,19 +1112,19 @@ namespace uba
 			return;
 	}
 
-	bool SessionClient::SendUpdateDirectoryTable()
+	bool SessionClient::SendUpdateDirectoryTable(StackBinaryReader<SendMaxSize>& reader)
 	{
+		UBA_ASSERT(reader.GetPosition() == 0);
 		StackBinaryWriter<32> writer;
 		NetworkMessage msg(m_client, ServiceId, SessionMessageType_GetDirectoriesFromServer, writer);
 		writer.WriteU32(m_sessionId);
 		writer.WriteU32(~u32(0));
-		StackBinaryReader<SendMaxSize> reader;
 		if (!msg.Send(reader, Stats().getDirsMsg))
 			return false;
 		return UpdateDirectoryTableFromServer(reader);
 	}
 
-	bool SessionClient::SendUpdateNameToHashTable()
+	bool SessionClient::SendUpdateNameToHashTable(StackBinaryReader<SendMaxSize>& reader)
 	{
 		StackBinaryWriter<32> writer;
 		NetworkMessage msg(m_client, ServiceId, SessionMessageType_GetNameToHashFromServer, writer);
@@ -1105,7 +1133,6 @@ namespace uba
 		ScopedWriteLock lock(m_nameToHashMemLock);
 		writer.WriteU32(u32(m_nameToHashTableMem.writtenSize));
 
-		StackBinaryReader<SendMaxSize> reader;
 		if (!msg.Send(reader, Stats().getHashesMsg))
 			return false;
 		return UpdateNameToHashTableFromServer(reader);
@@ -1315,8 +1342,6 @@ namespace uba
 					waitTimeoutMs = 200;
 				}
 
-				SendUpdateNameToHashTable();
-
 				for (InternalProcessStartInfo& startInfo : startInfos)
 				{
 					startInfo.description = startInfo.descriptionBuf.c_str();
@@ -1505,6 +1530,8 @@ namespace uba
 
 			SendPing(memAvail, memTotal);
 
+			//SendUpdateNameToHashTable(); // It is always nice to populate this at certain cadence since it might speed up running processes queries
+
 			m_waitToSendEvent.IsSet(waitTimeoutMs);
 
 			RemoveInactiveProcesses();
@@ -1606,7 +1633,7 @@ namespace uba
 
 	bool SessionClient::UpdateEnvironment(ProcessImpl& process, const tchar* reason)
 	{
-		StackBinaryWriter<SendMaxSize> writer;
+		StackBinaryWriter<16*1024> writer;
 		NetworkMessage msg(m_client, ServiceId, SessionMessageType_UpdateEnvironment, writer);
 		writer.WriteU32(process.m_id);
 		writer.WriteString(reason);
@@ -1614,9 +1641,10 @@ namespace uba
 		process.m_sessionStats.Write(writer);
 		process.m_storageStats.Write(writer);
 		process.m_systemStats.Write(writer);
-		StackBinaryReader<32> reader;
+		StackBinaryReader<SendMaxSize> reader;
 		if (!msg.Send(reader, m_stats.customMsg))
 			return false;
-		return SendUpdateDirectoryTable();
+		reader.Reset();
+		return SendUpdateDirectoryTable(reader);
 	}
 }
diff --git a/Engine/Source/Programs/UnrealBuildAccelerator/Common/Private/UbaSessionServer.cpp b/Engine/Source/Programs/UnrealBuildAccelerator/Common/Private/UbaSessionServer.cpp
--- a/Engine/Source/Programs/UnrealBuildAccelerator/Common/Private/UbaSessionServer.cpp
+++ b/Engine/Source/Programs/UnrealBuildAccelerator/Common/Private/UbaSessionServer.cpp
@@ -883,6 +883,9 @@ namespace uba
 				}
 				fillLock.Leave();
 
+				u32 neededDirectoryTableSize = GetDirectoryTableSize();
+				u32 neededHashTableSize = m_nameToHashLookupLock.ScopedRead([this]() { return u32(m_nameToHashTableMem.writtenSize); });
+
 				ScopedCriticalSection lock(m_remoteProcessAndSessionLock);
 				ClientSession& session = *m_clientSessions[sessionIndex];
 
@@ -915,6 +918,12 @@ namespace uba
 				}
 
 				writer.WriteU32(remoteExecutionEnabled ? SessionProcessAvailableResponse_None : SessionProcessAvailableResponse_RemoteExecutionDisabled);
+
+				
+				// Write in the needed dir and hash table offset to be up-to-date (to potentially avoid additional messages from client
+				writer.WriteU32(neededDirectoryTableSize);
+				writer.WriteU32(neededHashTableSize);
+
 				return true;
 			}
 			case SessionMessageType_ProcessFinished:
diff --git a/Engine/Source/Programs/UnrealBuildAccelerator/Common/Public/UbaNetwork.h b/Engine/Source/Programs/UnrealBuildAccelerator/Common/Public/UbaNetwork.h
--- a/Engine/Source/Programs/UnrealBuildAccelerator/Common/Public/UbaNetwork.h
+++ b/Engine/Source/Programs/UnrealBuildAccelerator/Common/Public/UbaNetwork.h
@@ -18,7 +18,7 @@ namespace uba
 	
 	static constexpr u32 SystemNetworkVersion = 1338;
 	static constexpr u32 StorageNetworkVersion = 3;
-	static constexpr u32 SessionNetworkVersion = 22;
+	static constexpr u32 SessionNetworkVersion = 23;
 
 
 	// Messages used over network between client/server (system, storage and session)
diff --git a/Engine/Source/Programs/UnrealBuildAccelerator/Common/Public/UbaSessionClient.h b/Engine/Source/Programs/UnrealBuildAccelerator/Common/Public/UbaSessionClient.h
--- a/Engine/Source/Programs/UnrealBuildAccelerator/Common/Public/UbaSessionClient.h
+++ b/Engine/Source/Programs/UnrealBuildAccelerator/Common/Public/UbaSessionClient.h
@@ -67,10 +67,10 @@ namespace uba
 		bool WriteBinFile(StringBufferBase& out, const tchar* binaryName, const CasKey& casKey, const KeyToString& applicationDir, u32 fileAttributes);
 		bool SendFiles(ProcessImpl& process, Timer& sendFiles);
 		bool SendFile(ProcessImpl& process, WrittenFile& source, const tchar* destination);
-		bool SendUpdateDirectoryTable();
-		bool UpdateDirectoryTableFromServer(BinaryReader& reader);
-		bool SendUpdateNameToHashTable();
-		bool UpdateNameToHashTableFromServer(BinaryReader& reader);
+		bool SendUpdateDirectoryTable(StackBinaryReader<SendMaxSize>& reader); // Note, reader is sent in to save stack space.
+		bool UpdateDirectoryTableFromServer(StackBinaryReader<SendMaxSize>& reader);
+		bool SendUpdateNameToHashTable(StackBinaryReader<SendMaxSize>& reader);
+		bool UpdateNameToHashTableFromServer(StackBinaryReader<SendMaxSize>& reader);
 		void Connect();
 		void BuildEnvironmentVariables(BinaryReader& reader);
 		bool SendProcessAvailable(Vector<InternalProcessStartInfo>& out, float availableWeight, bool& outRemoteExecutionEnabled);
diff --git a/Engine/Source/Programs/UnrealBuildAccelerator/Core/Public/UbaBinaryReaderWriter.h b/Engine/Source/Programs/UnrealBuildAccelerator/Core/Public/UbaBinaryReaderWriter.h
--- a/Engine/Source/Programs/UnrealBuildAccelerator/Core/Public/UbaBinaryReaderWriter.h
+++ b/Engine/Source/Programs/UnrealBuildAccelerator/Core/Public/UbaBinaryReaderWriter.h
@@ -76,7 +76,7 @@ namespace uba
 
 		BinaryReader();
 
-	private:
+	protected:
 		inline u64 InternalReadString(tchar* str, u64 charLen);
 		const u8* m_begin;
 		const u8* m_pos;
@@ -87,6 +87,7 @@ namespace uba
 	struct StackBinaryReader : BinaryReader
 	{
 		StackBinaryReader() : BinaryReader(buffer, 0, Capacity) { *buffer = 0; }
+		void Reset() { m_pos = m_begin; m_end = buffer + Capacity; }
 		u8 buffer[Capacity];
 	};
 
