[UBA]
* Stats fixes for network transfers

diff --git a/Engine/Source/Programs/UnrealBuildAccelerator/Common/Private/UbaNetworkServer.cpp b/Engine/Source/Programs/UnrealBuildAccelerator/Common/Private/UbaNetworkServer.cpp
--- a/Engine/Source/Programs/UnrealBuildAccelerator/Common/Private/UbaNetworkServer.cpp
+++ b/Engine/Source/Programs/UnrealBuildAccelerator/Common/Private/UbaNetworkServer.cpp
@@ -81,7 +81,7 @@ namespace uba
 				{
 					auto& conn = *(Connection*)context;
 					if (auto c = conn.m_client)
-						c->sendBytes += bytes;
+						c->recvBytes += bytes;
 					conn.m_server.m_sendBytes += bytes;
 				});
 
@@ -254,7 +254,7 @@ namespace uba
 				return false;
 			}
 
-			conn.m_client->recvBytes += worker->m_dataSize;
+			conn.m_client->sendBytes += worker->m_dataSize;
 			conn.m_server.m_recvBytes += worker->m_dataSize;
 			++conn.m_server.m_recvCount;
 
diff --git a/Engine/Source/Programs/UnrealBuildAccelerator/Common/Private/UbaSessionClient.cpp b/Engine/Source/Programs/UnrealBuildAccelerator/Common/Private/UbaSessionClient.cpp
--- a/Engine/Source/Programs/UnrealBuildAccelerator/Common/Private/UbaSessionClient.cpp
+++ b/Engine/Source/Programs/UnrealBuildAccelerator/Common/Private/UbaSessionClient.cpp
@@ -389,6 +389,7 @@ namespace uba
 				}
 				else
 				{
+					TimerScope s(m_stats.storageRetrieve);
 					casKey = AsCompressed(casKey, false);
 					entry.handled = true;
 					Storage::RetrieveResult result;
@@ -960,9 +961,12 @@ namespace uba
 
 		{
 			CasKey detoursBinaryKey = reader.ReadCasKey();
-			Storage::RetrieveResult result;
-			if (!m_storage.RetrieveCasFile(result, AsCompressed(detoursBinaryKey, false), UBA_DETOURS_LIBRARY))
-				return;
+			{
+				TimerScope s(m_stats.storageRetrieve);
+				Storage::RetrieveResult result;
+				if (!m_storage.RetrieveCasFile(result, AsCompressed(detoursBinaryKey, false), UBA_DETOURS_LIBRARY))
+					return;
+			}
 			KeyToString dir(StringKeyZero);
 			StringBuffer<> detoursFile;
 			if (!WriteBinFile(detoursFile, UBA_DETOURS_LIBRARY, detoursBinaryKey, dir, DefaultAttributes()))
diff --git a/Engine/Source/Programs/UnrealBuildAccelerator/Common/Private/UbaStats.cpp b/Engine/Source/Programs/UnrealBuildAccelerator/Common/Private/UbaStats.cpp
--- a/Engine/Source/Programs/UnrealBuildAccelerator/Common/Private/UbaStats.cpp
+++ b/Engine/Source/Programs/UnrealBuildAccelerator/Common/Private/UbaStats.cpp
@@ -172,7 +172,7 @@ namespace uba
 	}
 
 	template<typename T>
-	void LogStat(Logger& logger, const char* name, const T&) {}
+	void LogStat(Logger& logger, const char* name, const T&, u64 frequency) {}
 
 	void LogStat(Logger& logger, const char* name, const Timer& timer, u64 frequency)
 	{
@@ -225,7 +225,7 @@ namespace uba
 
 	void SessionSummaryStats::Print(Logger& logger, u64 frequency)
 	{
-		#define UBA_SESSION_SUMMARY_STAT(T, V) LogStat(logger, #V, V);
+		#define UBA_SESSION_SUMMARY_STAT(T, V) LogStat(logger, #V, V, frequency);
 		UBA_SESSION_SUMMARY_STATS
 		#undef UBA_SESSION_SUMMARY_STAT
 		stats.Print(logger, frequency);
