[UBA]
* Fixed one more stats visualization number that was wrong (not taking frequency in to account)

diff --git a/Engine/Source/Programs/UnrealBuildAccelerator/Common/Private/UbaProcess.cpp b/Engine/Source/Programs/UnrealBuildAccelerator/Common/Private/UbaProcess.cpp
--- a/Engine/Source/Programs/UnrealBuildAccelerator/Common/Private/UbaProcess.cpp
+++ b/Engine/Source/Programs/UnrealBuildAccelerator/Common/Private/UbaProcess.cpp
@@ -815,7 +815,12 @@ namespace uba
 					}
 
 					if (m_parentProcess)
+					{
 						m_parentProcess->m_processStats.Add(m_processStats);
+						m_parentProcess->m_sessionStats.Add(m_sessionStats);
+						m_parentProcess->m_storageStats.Add(m_storageStats);
+						m_parentProcess->m_systemStats.Add(m_systemStats);
+					}
 
 					if (m_startInfo.outputStatsThresholdMs && TimeToMs(m_processStats.GetTotalTime()) > m_startInfo.outputStatsThresholdMs)
 					{
diff --git a/Engine/Source/Programs/UnrealBuildAccelerator/Common/Private/UbaStats.cpp b/Engine/Source/Programs/UnrealBuildAccelerator/Common/Private/UbaStats.cpp
--- a/Engine/Source/Programs/UnrealBuildAccelerator/Common/Private/UbaStats.cpp
+++ b/Engine/Source/Programs/UnrealBuildAccelerator/Common/Private/UbaStats.cpp
@@ -106,7 +106,7 @@ namespace uba
 			logger.Info(TC("  EnsureCas            %6u %9s"), ensureCas.count.load(), TimeToText(ensureCas.time, false, frequency).str);
 		if (recvCas.count)
 		{
-			logger.Info(TC("  ReceiveCas           %6u %9s"), recvCas.count.load(), TimeToText(recvCas.time).str);
+			logger.Info(TC("  ReceiveCas           %6u %9s"), recvCas.count.load(), TimeToText(recvCas.time, false, frequency).str);
 			logger.Info(TC("     Bytes Raw/Comp %9s %9s"), BytesToText(recvCasBytesRaw).str, BytesToText(recvCasBytesComp).str);
 			if (decompressRecv.count)
 				logger.Info(TC("     Decompress        %6u %9s"), decompressRecv.count.load(), TimeToText(decompressRecv.time, false, frequency).str);
diff --git a/Engine/Source/Programs/UnrealBuildAccelerator/Common/Public/UbaStats.h b/Engine/Source/Programs/UnrealBuildAccelerator/Common/Public/UbaStats.h
--- a/Engine/Source/Programs/UnrealBuildAccelerator/Common/Public/UbaStats.h
+++ b/Engine/Source/Programs/UnrealBuildAccelerator/Common/Public/UbaStats.h
@@ -48,6 +48,8 @@ namespace uba
 	{
 		StorageStatsScope(StorageStats& stats);
 		~StorageStatsScope();
+		StorageStatsScope(const StorageStatsScope&) = delete;
+		void operator=(const StorageStatsScope&) = delete;
 		StorageStats& stats;
 	};
 
@@ -82,6 +84,8 @@ namespace uba
 	{
 		SessionStatsScope(SessionStats& stats);
 		~SessionStatsScope();
+		SessionStatsScope(const SessionStatsScope&) = delete;
+		void operator=(const SessionStatsScope&) = delete;
 		SessionStats& stats;
 	};
 
diff --git a/Engine/Source/Programs/UnrealBuildAccelerator/Core/Public/UbaTimer.h b/Engine/Source/Programs/UnrealBuildAccelerator/Core/Public/UbaTimer.h
--- a/Engine/Source/Programs/UnrealBuildAccelerator/Core/Public/UbaTimer.h
+++ b/Engine/Source/Programs/UnrealBuildAccelerator/Core/Public/UbaTimer.h
@@ -44,6 +44,8 @@ namespace uba
 		void Cancel() { start = ~u64(0); }
 		Timer& timer;
 		u64 start;
+		TimerScope(const TimerScope&) = delete;
+		void operator=(const TimerScope&) = delete;
 	};
 
 	struct ExtendedTimerScope 
