UnrealBuildAccelerator: Fix frequency of wall time on mac and linux

#rnx

diff --git a/Engine/Source/Programs/UnrealBuildAccelerator/Core/Public/UbaTimer.h b/Engine/Source/Programs/UnrealBuildAccelerator/Core/Public/UbaTimer.h
--- a/Engine/Source/Programs/UnrealBuildAccelerator/Core/Public/UbaTimer.h
+++ b/Engine/Source/Programs/UnrealBuildAccelerator/Core/Public/UbaTimer.h
@@ -10,14 +10,21 @@ namespace uba
 	u64 GetFrequency();						// Number of counts per second. GetTime() / GetFrequency() = seconds
 	u64 GetSystemTimeUs();					// Get time in microseconds since 1 Jan 1970 00:00 UTC
 
-	inline u64 TimeToMs(u64 time, u64 frequency)	{ return time * 1000 / frequency; }
+	inline u64 TimeToTick(u64 time, u64 frequency)	{ return time * 10'000'000 / frequency; }
+	inline u64 TimeToUs(u64 time, u64 frequency)	{ return time * 1'000'000 / frequency; }
+	inline u64 TimeToMs(u64 time, u64 frequency)	{ return time * 1'000 / frequency; }
 	inline float TimeToS(u64 time, u64 frequency)	{ return float(double(time) / double(frequency)); }
 
+	inline u64 TimeToTick(u64 time)					{ return TimeToTick(time, GetFrequency()); }
+	inline u64 TimeToUs(u64 time)					{ return TimeToUs(time, GetFrequency()); }
 	inline u64 TimeToMs(u64 time)					{ return TimeToMs(time, GetFrequency()); }
-	inline u64 UsToTime(u64 us)						{ return us * GetFrequency() / 1'000'000; }
-	inline u64 MsToTime(u64 ms)						{ return ms * GetFrequency() / 1000; }
 	inline float TimeToS(u64 time)					{ return TimeToS(time, GetFrequency()); }
 
+	inline u64 TickToTime(u64 us)					{ return us * GetFrequency() / 10'000'000; }
+	inline u64 UsToTime(u64 us)						{ return us * GetFrequency() / 1'000'000; }
+	inline u64 MsToTime(u64 ms)						{ return ms * GetFrequency() / 1'000; }
+
+
 
 	struct Timer
 	{
diff --git a/Engine/Source/Programs/UnrealBuildAccelerator/Host/Private/UbaExports.cpp b/Engine/Source/Programs/UnrealBuildAccelerator/Host/Private/UbaExports.cpp
--- a/Engine/Source/Programs/UnrealBuildAccelerator/Host/Private/UbaExports.cpp
+++ b/Engine/Source/Programs/UnrealBuildAccelerator/Host/Private/UbaExports.cpp
@@ -262,14 +262,16 @@ uba::StorageClient* CreateStorageClient(uba::NetworkClient& client, const uba::t
 		return handle->GetHash();
 	}
 
+	// 100ns ticks
 	uba::u64 ProcessHandle_GetTotalProcessorTime(uba::ProcessHandle* handle)
 	{
-		return handle->GetTotalProcessorTime();
+		return uba::TimeToTick(handle->GetTotalProcessorTime());
 	}
 
+	// 100ns ticks
 	uba::u64 ProcessHandle_GetTotalWallTime(uba::ProcessHandle* handle)
 	{
-		return handle->GetTotalWallTime();
+		return uba::TimeToTick(handle->GetTotalWallTime());
 	}
 
 	void ProcessHandle_Cancel(uba::ProcessHandle* handle, bool terminate)
