UBA: Clean up thread attrs after creating the thread.

diff --git a/Engine/Source/Programs/UnrealBuildAccelerator/Common/Private/UbaThread.cpp b/Engine/Source/Programs/UnrealBuildAccelerator/Common/Private/UbaThread.cpp
--- a/Engine/Source/Programs/UnrealBuildAccelerator/Common/Private/UbaThread.cpp
+++ b/Engine/Source/Programs/UnrealBuildAccelerator/Common/Private/UbaThread.cpp
@@ -71,6 +71,9 @@ namespace uba
 				return (void*)(uintptr_t)res;
 			}, this);
 		UBA_ASSERT(err == 0); (void)err;
+
+		err = pthread_attr_destroy(&tattr);
+		UBA_ASSERT(err == 0); (void)err;
 #endif
 	}
 
