UBA: Handle the scenario where files may be spread across mounts and linking won't work. When this happens fall back to copying/sendfile-ing files around.

diff --git a/Engine/Source/Programs/UnrealBuildAccelerator/Common/Private/UbaFileAccessor.cpp b/Engine/Source/Programs/UnrealBuildAccelerator/Common/Private/UbaFileAccessor.cpp
--- a/Engine/Source/Programs/UnrealBuildAccelerator/Common/Private/UbaFileAccessor.cpp
+++ b/Engine/Source/Programs/UnrealBuildAccelerator/Common/Private/UbaFileAccessor.cpp
@@ -4,6 +4,12 @@
 #include "UbaFile.h"
 #include "UbaStats.h"
 
+#if PLATFORM_LINUX
+#include <sys/sendfile.h>
+#elif PLATFORM_MAC
+#include <copyfile.h>
+#endif
+
 namespace uba
 {
 	#if PLATFORM_WINDOWS
@@ -334,7 +340,26 @@ namespace uba
 						return m_logger.Error(TC("Failed to remove delete on close for file %s (%s)"), realFileName, LastErrorToText().data);
 					#else
 					if (m_tempPath && rename(realFileName, m_fileName) == -1)
-						return m_logger.Error(TC("Failed to rename temporary file %s to %s (%s)"), realFileName, m_fileName, strerror(errno));
+					{
+						if (errno != EXDEV)
+							return m_logger.Error(TC("Failed to rename temporary file %s to %s (%s)"), realFileName, m_fileName, strerror(errno));
+
+						// Need to copy, can't rename over devices
+						int targetFd = open(m_fileName, O_CREAT | O_TRUNC | O_WRONLY);
+						auto g = MakeGuard([targetFd]() { close(targetFd); });
+						if (targetFd == -1)
+							return m_logger.Error(TC("Failed to create file %s for move from temporary file %s (%s)"), m_fileName, realFileName, strerror(errno));
+						
+						#if PLATFORM_MAC
+						if (fcopyfile(asFileDescriptor(m_fileHandle), targetFd, 0, COPYFILE_ALL) == -1) {
+						#else
+						if (sendfile(targetFd, asFileDescriptor(m_fileHandle), NULL, m_size) == -1) {
+						#endif
+							return m_logger.Error(TC("Failed to do sendfile from temporary %s to file %s (%s)"), realFileName, m_fileName, strerror(errno));
+						}
+
+						remove(realFileName); // Remove real file now when we have copied it over
+					}
 					#endif
 
 					if (lastWriteTime)
