[UBA]
* Added stub for providing known inputs to RunProcessRemote
* Added stub for RunProcessRacing that can be used to have local machine race against remotely scheduled tasks.
* Added documentation

diff --git a/Engine/Source/Programs/UnrealBuildAccelerator/Common/Private/UbaSessionServer.cpp b/Engine/Source/Programs/UnrealBuildAccelerator/Common/Private/UbaSessionServer.cpp
--- a/Engine/Source/Programs/UnrealBuildAccelerator/Common/Private/UbaSessionServer.cpp
+++ b/Engine/Source/Programs/UnrealBuildAccelerator/Common/Private/UbaSessionServer.cpp
@@ -255,7 +255,13 @@ namespace uba
 		#endif
 	}
 
-	ProcessHandle SessionServer::RunProcessRemote(const ProcessStartInfo& startInfo, float weight)
+	ProcessHandle SessionServer::RunProcessRacing(u32 raceAgainstRemoteProcessId)
+	{
+		// TODO: Implement
+		return {};
+	}
+
+	ProcessHandle SessionServer::RunProcessRemote(const ProcessStartInfo& startInfo, float weight, const void* knownInputs, u32 knownInputsSizeBytes)
 	{
 		FlushDeadProcesses();
 		ValidateStartInfo(startInfo);
diff --git a/Engine/Source/Programs/UnrealBuildAccelerator/Common/Public/UbaSessionServer.h b/Engine/Source/Programs/UnrealBuildAccelerator/Common/Public/UbaSessionServer.h
--- a/Engine/Source/Programs/UnrealBuildAccelerator/Common/Public/UbaSessionServer.h
+++ b/Engine/Source/Programs/UnrealBuildAccelerator/Common/Public/UbaSessionServer.h
@@ -14,25 +14,49 @@ namespace uba
 	class SessionServer final : public Session
 	{
 	public:
+		// Ctor/dtor
 		SessionServer(const SessionServerCreateInfo& info);
 		~SessionServer();
-		
-		ProcessHandle RunProcessRemote(const ProcessStartInfo& startInfo, float weight = 1.0f);
+
+		// Run process remotely.
+		// startInfo contains info about the process to run remotely
+		// weight is the expected core usage of the process. If processes are multithreaded it makes sense to increase weight. As an example, in UnrealBuildTool we see cl.exe as 1.5 and clang.exe as 1.0
+		// knownInputs is a memory block with null terminated strings followed by an empty null terminated string. knownInputSize is the size of the memoryBlock in bytes
+		// knownInputs strings should be of type tchar and a path relative to working dir or absolute. Session will make a copy of knownInputs internally
+		ProcessHandle RunProcessRemote(const ProcessStartInfo& startInfo, float weight = 1.0f, const void* knownInputs = nullptr, u32 knownInputsSizeBytes = 0);
+
+		// Will kick off a local process with the same startInfo as the one provided to start the process with id matching raceAgainstRemoteProcessId
+		// This can be useful if there are free local cores and we know local machine is faster or network connection to remote is slow
+		ProcessHandle RunProcessRacing(u32 raceAgainstRemoteProcessId);
+
+		// Disable remote execution. This will tell all clients to take on new processes and disconnect as soon as their current processes are finished
 		void DisableRemoteExecution();
 
+		// This can be used to set a custom cas key based on a file and its inputs. Can be used for non-deterministic outputs to still be able to use cached cas content on clients
+		// For example, when a pch is built on the host we can use all the input involved to create the pch and use that as the cas key for the resulting huge file
+		// This means that if a remote machine has a pch from an older run where all the input matches, it can reuse the cas content even though the output differs
 		void SetCustomCasKeyFromTrackedInputs(const tchar* fileName, const tchar* workingDir, const u8* trackedInputs, u32 trackedInputsBytes);
 		bool GetCasKeyFromTrackedInputs(CasKey& out, const tchar* fileName, const tchar* workingDir, const u8* data, u32 dataLen);
 
+		// Callback that will be called when a client is asking for processes to run. All remotes frequently ask for processes to run if they have free process slots
 		void SetRemoteProcessSlotAvailableEvent(const Function<void()>& remoteProcessSlotAvailableEvent);
+
+		// Callback that is called when process is returned. This could be because a client unexpectedly disconnected or is running out of memory
 		void SetRemoteProcessReturnedEvent(const Function<void(Process&)>& remoteProcessReturnedEvent);
 
+		// Wait for all queued up/active tasks to finish
 		void WaitOnAllTasks();
 
+		// Can be used to hint the session how many processes (that can run remotely that are left.. session can then start disabling remote execution on clients not needed anymore
 		void SetMaxRemoteProcessCount(u32 count);
 
+		// Report an external process just to get it visible in the trace stream/visualizer. Returns an unique id that should be sent in to EndExternalProcess
 		u32 BeginExternalProcess(const tchar* description);
+
+		// End external process.
 		void EndExternalProcess(u32 id, u32 exitCode);
 
+		// Get the network server used by this session
 		NetworkServer& GetServer();
 
 	protected:
diff --git a/Engine/Source/Programs/UnrealBuildAccelerator/Host/Private/UbaExports.cpp b/Engine/Source/Programs/UnrealBuildAccelerator/Host/Private/UbaExports.cpp
--- a/Engine/Source/Programs/UnrealBuildAccelerator/Host/Private/UbaExports.cpp
+++ b/Engine/Source/Programs/UnrealBuildAccelerator/Host/Private/UbaExports.cpp
@@ -338,9 +338,13 @@ uba::StorageClient* CreateStorageClient(uba::NetworkClient& client, const uba::t
 	{
 		return new uba::ProcessHandle(server->RunProcess(info, async, enableDetour));
 	}
-	uba::ProcessHandle* SessionServer_RunProcessRemote(uba::SessionServer* server, uba::ProcessStartInfo& info, float weight)
+	uba::ProcessHandle* SessionServer_RunProcessRemote(uba::SessionServer* server, uba::ProcessStartInfo& info, float weight, const void* knownInputs, uba::u32 knownInputsSizeBytes)
 	{
-		return new uba::ProcessHandle(server->RunProcessRemote(info, weight));
+		return new uba::ProcessHandle(server->RunProcessRemote(info, weight, knownInputs, knownInputsSizeBytes));
+	}
+	uba::ProcessHandle* SessionServer_RunProcessRacing(uba::SessionServer* server, uba::u32 raceAgainstRemoteProcessId)
+	{
+		return new uba::ProcessHandle(server->RunProcessRacing(raceAgainstRemoteProcessId));
 	}
 	void SessionServer_SetMaxRemoteProcessCount(uba::SessionServer* server, uba::u32 count)
 	{
diff --git a/Engine/Source/Programs/UnrealBuildAccelerator/Host/Public/UbaExports.h b/Engine/Source/Programs/UnrealBuildAccelerator/Host/Public/UbaExports.h
--- a/Engine/Source/Programs/UnrealBuildAccelerator/Host/Public/UbaExports.h
+++ b/Engine/Source/Programs/UnrealBuildAccelerator/Host/Public/UbaExports.h
@@ -86,7 +86,9 @@ extern "C"
 	UBA_API void SessionServer_RefreshDirectory(uba::SessionServer* server, const uba::tchar* directory);
 	UBA_API void SessionServer_RegisterNewFile(uba::SessionServer* server, const uba::tchar* filePath);
 	UBA_API uba::ProcessHandle* SessionServer_RunProcess(uba::SessionServer* server, uba::ProcessStartInfo& info, bool async, bool enableDetour);
-	UBA_API uba::ProcessHandle* SessionServer_RunProcessRemote(uba::SessionServer* server, uba::ProcessStartInfo& info, float weight);
+	UBA_API uba::ProcessHandle* SessionServer_RunProcessRemote(uba::SessionServer* server, uba::ProcessStartInfo& info, float weight, const void* knownInputs = nullptr, uba::u32 knownInputsSizeBytes = 0);
+	UBA_API uba::ProcessHandle* SessionServer_RunProcessRacing(uba::SessionServer* server, uba::u32 raceAgainstRemoteProcessId);
+
 	UBA_API void SessionServer_SetMaxRemoteProcessCount(uba::SessionServer* server, uba::u32 count);
 	UBA_API void SessionServer_DisableRemoteExecution(uba::SessionServer* server);
 	UBA_API void SessionServer_PrintSummary(uba::SessionServer* server);
