[UBA]
* Fixed so if shm exists when there is no file we just unlink it and try to create it again. We assume it is somehow orphaned, caused by manual deleting of files in /tmp/uba_shm_locks or if the machine has been on long enough for folder to have been cleaned up by os.

diff --git a/Engine/Source/Programs/UnrealBuildAccelerator/Common/Private/UbaFileMapping.cpp b/Engine/Source/Programs/UnrealBuildAccelerator/Common/Private/UbaFileMapping.cpp
--- a/Engine/Source/Programs/UnrealBuildAccelerator/Common/Private/UbaFileMapping.cpp
+++ b/Engine/Source/Programs/UnrealBuildAccelerator/Common/Private/UbaFileMapping.cpp
@@ -131,10 +131,16 @@ namespace uba
 			shmFd = shm_open(uidName.data, oflags, S_IRUSR | S_IWUSR);
 			if (shmFd != -1)
 				break;
+
+			shm_unlink(uidName.data);
+			shmFd = shm_open(uidName.data, oflags, S_IRUSR | S_IWUSR);
+			if (shmFd != -1)
+				break;
+
 			remove(lockFile.data);
 			close(lockFd);
 			SetLastError(errno);
-			UBA_ASSERTF(false, "This should not happen.. someone created shm %s without lock-file %s (%s)", uidName.data, lockFile.data, strerror(errno));
+			UBA_ASSERTF(false, "Failed to create shm %s after getting lock-file %s (%s)", uidName.data, lockFile.data, strerror(errno));
 			return {};
 		}
 
