Fixed UBA on Mac to hve children die when session process goes away even if not notified

diff --git a/Engine/Source/Programs/UnrealBuildAccelerator/Common/Private/UbaProcess.cpp b/Engine/Source/Programs/UnrealBuildAccelerator/Common/Private/UbaProcess.cpp
--- a/Engine/Source/Programs/UnrealBuildAccelerator/Common/Private/UbaProcess.cpp
+++ b/Engine/Source/Programs/UnrealBuildAccelerator/Common/Private/UbaProcess.cpp
@@ -1325,6 +1325,14 @@ namespace uba
 			res = posix_spawnattr_setpgroup(&attr, getpgrp());
 			UBA_ASSERTF(res == 0, TC("posix_spawnattr_setpgroup (%s)"), strerror(errno));
 
+			/* Set the child's priority to be two levels lower. */
+			sched_param param;
+			res = posix_spawnattr_getschedparam(&attr, &param);
+			UBA_ASSERTF(res == 0, TC("posix_spawnattr_getschedparam (%s)"), strerror(errno));
+			param.sched_priority -= 2;
+			res = posix_spawnattr_setschedparam(&attr, &param);
+			UBA_ASSERTF(res == 0, TC("posix_spawnattr_setschedparam (%s)"), strerror(errno));
+
 			posix_spawn_file_actions_t fileActions;
 			res = posix_spawn_file_actions_init(&fileActions);
 			UBA_ASSERTF(res == 0, TC("posix_spawn_file_actions_init (%s)"), strerror(errno));
@@ -1345,6 +1353,7 @@ namespace uba
 			StringBuffer<512> logFile;
 			StringBuffer<512> ldLibraryPath;
 			StringBuffer<512> detoursVar;
+			StringBuffer<32> processVar;
 
 			Vector<const char*> envvars;
 
@@ -1387,6 +1396,8 @@ namespace uba
 				detoursVar.Append("DYLD_INSERT_LIBRARIES=").Append(detoursLib);
 #endif
 
+				processVar.Append("UBA_SESSION_PROCESS=").AppendValue(getpid());
+
 
 				comIdVar.Append("UBA_COMID=").AppendValue(communicationHandle.uid).Append('+').AppendValue(communicationOffset);
 				workingDir.Append("UBA_CWD=").Append(m_realWorkingDir);
@@ -1403,6 +1414,7 @@ namespace uba
 				if (ldLibraryPath.count)
 					envvars.push_back(ldLibraryPath.data);
 				envvars.push_back(detoursVar.data);
+				envvars.push_back(processVar.data);
 				envvars.push_back(comIdVar.data);
 				envvars.push_back(workingDir.data);
 				envvars.push_back(rulesStr.data);
diff --git a/Engine/Source/Programs/UnrealBuildAccelerator/Detours/Private/PosixOS/UbaDetoursMainPosix.cpp b/Engine/Source/Programs/UnrealBuildAccelerator/Detours/Private/PosixOS/UbaDetoursMainPosix.cpp
--- a/Engine/Source/Programs/UnrealBuildAccelerator/Detours/Private/PosixOS/UbaDetoursMainPosix.cpp
+++ b/Engine/Source/Programs/UnrealBuildAccelerator/Detours/Private/PosixOS/UbaDetoursMainPosix.cpp
@@ -14,6 +14,7 @@ using namespace uba;
 namespace uba
 {
 	int g_comFd = -1;
+	int g_sessionPid = -1;
 	Event* g_cancelEvent;
 	Event* g_readEvent;
 	Event* g_writeEvent;
@@ -120,6 +121,7 @@ static void __attribute__((constructor(102))) PreInitCtor()
 		sigaction(signal, &sa, NULL);
 	#endif
 
+	g_sessionPid = atoi(getenv("UBA_SESSION_PROCESS"));
 	g_runningRemote = getenv("UBA_REMOTE") != nullptr;
 
 	g_rulesIndex = strtoul(getenv("UBA_RULES"), nullptr, 10);
@@ -178,6 +180,13 @@ namespace uba
 		{
 			if (g_readEvent->IsSet(1000))
 				break;
+			
+#if PLATFORM_MAC
+			// check if session process is gone
+			if (kill(g_sessionPid, 0) == -1 && errno == ESRCH)
+				exit(1337);
+#endif
+			
 			if (g_cancelEvent->IsSet(0))
 				exit(1339);
 		} while (true);
