[UBA]
* Fixed a bug where compressed cas files ended up with the uncompressed file size on disk

diff --git a/Engine/Source/Programs/UnrealBuildAccelerator/Common/Private/UbaStorageClient.cpp b/Engine/Source/Programs/UnrealBuildAccelerator/Common/Private/UbaStorageClient.cpp
--- a/Engine/Source/Programs/UnrealBuildAccelerator/Common/Private/UbaStorageClient.cpp
+++ b/Engine/Source/Programs/UnrealBuildAccelerator/Common/Private/UbaStorageClient.cpp
@@ -390,26 +390,27 @@ namespace uba
 					actualSize = *(u64*)readBuffer;
 			}
 
+			sizeOnDisk = IsCompressed(casKey) ? fileSize : actualSize;
+
 #if !UBA_USE_SPARSEFILE
 			FileAccessor destinationFile(m_logger, casFile.data);
 			if (!mappingBuffer)
 			{
 				u32 extraFlags = DefaultAttributes();
-				bool useOverlap = !IsRunningWine() && isCompressed == IsCompressed(casKey) && actualSize > 1024 * 1024;
+				bool useOverlap = !IsRunningWine() && isCompressed == IsCompressed(casKey) && sizeOnDisk > 1024 * 1024;
 				if (useOverlap)
 					extraFlags |= FILE_FLAG_OVERLAPPED;
-				if (!destinationFile.CreateWrite(false, extraFlags, actualSize, m_tempPath.data))
+				if (!destinationFile.CreateWrite(false, extraFlags, sizeOnDisk, m_tempPath.data))
 					return false;
 			}
 #endif
 
 			if (mappingBuffer)
 			{
-				u64 allocSize = IsCompressed(casKey) ? fileSize : actualSize;
-				UBA_ASSERT(!writeMem || mappedView.size == allocSize);
+				UBA_ASSERT(!writeMem || mappedView.size == sizeOnDisk);
 				if (!writeMem)
 				{
-					mappedView = mappingBuffer->AllocAndMapView(mappingType, allocSize, memoryMapAlignment, hint);
+					mappedView = mappingBuffer->AllocAndMapView(mappingType, sizeOnDisk, memoryMapAlignment, hint);
 					writeMem = mappedView.memory;
 					if (!writeMem)
 						return false;
@@ -469,8 +470,6 @@ namespace uba
 				}
 				if (tryAgain)
 					continue;
-
-				sizeOnDisk = fileSize;
 			}
 			else
 			{
@@ -599,7 +598,6 @@ namespace uba
 					if (tryAgain)
 						continue;
 				}
-				sizeOnDisk = actualSize;
 			}
 
 			if (sendEnd)
