[UBA]
* Enabled logic that looks for sessionpid if read event times out on linux as well (was ifdeffed for osx only)

diff --git a/Engine/Source/Programs/UnrealBuildAccelerator/Detours/Private/PosixOS/UbaDetoursMainPosix.cpp b/Engine/Source/Programs/UnrealBuildAccelerator/Detours/Private/PosixOS/UbaDetoursMainPosix.cpp
--- a/Engine/Source/Programs/UnrealBuildAccelerator/Detours/Private/PosixOS/UbaDetoursMainPosix.cpp
+++ b/Engine/Source/Programs/UnrealBuildAccelerator/Detours/Private/PosixOS/UbaDetoursMainPosix.cpp
@@ -181,11 +181,9 @@ namespace uba
 			if (g_readEvent->IsSet(1000))
 				break;
 			
-#if PLATFORM_MAC
 			// check if session process is gone
 			if (kill(g_sessionPid, 0) == -1 && errno == ESRCH)
 				exit(1337);
-#endif
 			
 			if (g_cancelEvent->IsSet(0))
 				exit(1339);
