[UBA]
* Implemented knownInputs support. It is now fully supported to provide known inputs when doing RunProcessRemote. This is optional and can be a nice speedup if latency to helpers is high since it reduce number of messages going back and forth.
* Changed so GetMemoryMapAlignment is checking only file name and no flags. Flags is now handled on the outside. This made it possible to send over map alignment for files without having to send the file names
* Renamed WorkManagerImpl::Wait to WorkManagerImpl::FlushWork
* Renamed willBeUsedUnCompressed to storeUncompressed
* Added ScopedCriticalSection::Enter

diff --git a/Engine/Source/Programs/UnrealBuildAccelerator/Common/Private/UbaSession.cpp b/Engine/Source/Programs/UnrealBuildAccelerator/Common/Private/UbaSession.cpp
--- a/Engine/Source/Programs/UnrealBuildAccelerator/Common/Private/UbaSession.cpp
+++ b/Engine/Source/Programs/UnrealBuildAccelerator/Common/Private/UbaSession.cpp
@@ -1002,6 +1002,8 @@ namespace uba
 		else
 			m_keepOutputFileMemoryMapsThreshold = info.keepOutputFileMemoryMapsThreshold;
 		m_shouldWriteToDisk = info.shouldWriteToDisk;
+		UBA_ASSERTF(m_shouldWriteToDisk || m_allowMemoryMaps, TC("Can't disable both should write to disk and allow memory maps"));
+
 		m_detailedTrace = info.detailedTrace;
 		m_logToFile = info.logToFile;
 		if (info.extraInfo)
@@ -1511,11 +1513,8 @@ namespace uba
 		return true;
 	}
 
-	u64 Session::GetMemoryMapAlignment(const tchar* fileName, u64 fileNameLen) const
+	u32 Session::GetMemoryMapAlignment(const tchar* fileName, u64 fileNameLen) const
 	{
-		if (!m_allowMemoryMaps)
-			return 0;
-
 		// It is not necessarily better to make mem maps of everything.. only things that are read more than once in the build.
 		// Reason is because there is additional overhead to use memory mappings.
 		// Upside is that all things that are memory mapped can be stored compressed in cas storage so it saves space.
@@ -1622,85 +1621,88 @@ namespace uba
 		const StringBufferBase& fileName = msg.fileName;
 		const StringKey& fileNameKey = msg.fileNameKey;
 
+		auto tableSizeGuard = MakeGuard([&]()
+			{
+				out.mappedFileTableSize = GetFileMappingSize();
+				out.directoryTableSize = GetDirectoryTableSize();
+			});
+
 		if ((msg.access & ~FileAccess_Read) == 0)
 		{
 			if (!IsWindows)
 			{
 				out.fileName.Append(fileName);
+				return true;
 			}
-			else if (fileName.EndsWith(TC(".dll")) || fileName.EndsWith(TC(".exe")))
+		
+			if (fileName.EndsWith(TC(".dll")) || fileName.EndsWith(TC(".exe")))
 			{
-				if (fileName[1] != ':')
-				{
-					UBA_ASSERTF(false, TC("Got bad filename from process %s"), fileName.data);
-				}
-				else
-				{
-					AddFileMapping(fileNameKey, fileName.data, TC("#"));
-					out.fileName.Append(TC("#"));
-				}
+				UBA_ASSERTF(fileName[1] == ':', TC("Got bad filename from process %s"), fileName.data);
+				AddFileMapping(fileNameKey, fileName.data, TC("#"));
+				out.fileName.Append(TC("#"));
+				return true;
 			}
-			else if (u64 alignment = GetMemoryMapAlignment(fileName.data, fileName.count))
+			
+			if (m_allowMemoryMaps)
 			{
-				MemoryMap map;
-				if (CreateMemoryMapFromFile(map, fileNameKey, fileName.data, false, alignment))
+				if (u64 alignment = GetMemoryMapAlignment(fileName.data, fileName.count))
 				{
-					out.size = map.size;
-					out.fileName.Append(map.name);
-				}
-				else
-				{
-					out.fileName.Append(fileName);
+					MemoryMap map;
+					if (CreateMemoryMapFromFile(map, fileNameKey, fileName.data, false, alignment))
+					{
+						out.size = map.size;
+						out.fileName.Append(map.name);
+					}
+					else
+					{
+						out.fileName.Append(fileName);
+					}
+					return true;
 				}
 			}
-			else if (!IsRarelyRead(msg.process, fileName))
+
+			if (!IsRarelyRead(msg.process, fileName))
 			{
 				AddFileMapping(fileNameKey, fileName.data, TC("#"));
 				out.fileName.Append(TC("#"));
+				return true;
 			}
-			else
-			{
-				out.fileName.Append(fileName);
-			}
+
+			out.fileName.Append(fileName);
+			return true;
 		}
-		else // if ((message.Access & FileAccess.Write) != 0)
-		{
-			m_storage.ReportFileWrite(fileName.data);
+		
+		// if ((message.Access & FileAccess.Write) != 0)
+		m_storage.ReportFileWrite(fileName.data);
 
-			if (m_runningRemote && !fileName.StartsWith(m_tempPath.data))
+		if (m_runningRemote && !fileName.StartsWith(m_tempPath.data))
+		{
+			ScopedWriteLock lock(m_outputFilesLock);
+			auto insres = m_outputFiles.try_emplace(fileName.data);
+			if (insres.second)
 			{
-				ScopedWriteLock lock(m_outputFilesLock);
-				auto insres = m_outputFiles.try_emplace(fileName.data);
-				if (insres.second)
-				{
-					out.fileName.Append(m_sessionOutputDir).Append(KeyToString(fileNameKey));
-					insres.first->second = out.fileName.data;
-				}
-				else
-				{
-					out.fileName.Append(insres.first->second.c_str());
-				}
+				out.fileName.Append(m_sessionOutputDir).Append(KeyToString(fileNameKey));
+				insres.first->second = out.fileName.data;
 			}
 			else
 			{
-				out.fileName.Append(fileName);
-			}
-
-			ScopedWriteLock lock(m_activeFilesLock);
-			u32 wantsOnCloseId = m_wantsOnCloseIdCounter++;
-			out.closeId = wantsOnCloseId;
-			auto insres = m_activeFiles.try_emplace(wantsOnCloseId);
-			if (insres.second)
-			{
-				insres.first->second.name = fileName.data;
-				insres.first->second.nameKey = fileNameKey;
+				out.fileName.Append(insres.first->second.c_str());
 			}
-			else
-				return m_logger.Error(TC("TRYING TO ADD %s twice!"), out.fileName.data);
+		}
+		else
+		{
+			out.fileName.Append(fileName);
 		}
 
-		out.mappedFileTableSize = GetFileMappingSize();
-		out.directoryTableSize = GetDirectoryTableSize();
+		ScopedWriteLock lock(m_activeFilesLock);
+		u32 wantsOnCloseId = m_wantsOnCloseIdCounter++;
+		out.closeId = wantsOnCloseId;
+		auto insres = m_activeFiles.try_emplace(wantsOnCloseId);
+		if (!insres.second)
+			return m_logger.Error(TC("TRYING TO ADD %s twice!"), out.fileName.data);
+
+		insres.first->second.name = fileName.data;
+		insres.first->second.nameKey = fileNameKey;
 		return true;
 	}
 
diff --git a/Engine/Source/Programs/UnrealBuildAccelerator/Common/Private/UbaSessionClient.cpp b/Engine/Source/Programs/UnrealBuildAccelerator/Common/Private/UbaSessionClient.cpp
--- a/Engine/Source/Programs/UnrealBuildAccelerator/Common/Private/UbaSessionClient.cpp
+++ b/Engine/Source/Programs/UnrealBuildAccelerator/Common/Private/UbaSessionClient.cpp
@@ -82,19 +82,19 @@ namespace uba
 		return m_bestPing;
 	}
 
-	bool SessionClient::RetrieveCasFile(CasKey& outNewKey, u64& outSize, const CasKey& casKey, const tchar* hint, bool willBeUsedUncompressed, bool allowProxy)
+	bool SessionClient::RetrieveCasFile(CasKey& outNewKey, u64& outSize, const CasKey& casKey, const tchar* hint, bool storeUncompressed, bool allowProxy)
 	{
 		TimerScope s(m_stats.storageRetrieve);
 		CasKey tempKey = casKey;
 
 #if UBA_USE_SPARSEFILE
-		willBeUsedUncompressed = false;
+		storeUncompressed = false;
 #endif
 
-		if (willBeUsedUncompressed)
+		if (storeUncompressed)
 			tempKey = AsCompressed(casKey, false);
 		//else
-		//	willBeUsedUncompressed = !IsCompressed(casKey);
+		//	storeUncompressed = !IsCompressed(casKey);
 
 		Storage::RetrieveResult result;
 		bool res = m_storage.RetrieveCasFile(result, tempKey, hint, nullptr, 1, allowProxy);
@@ -172,10 +172,10 @@ namespace uba
 			out.Append(fileName);
 			return true;
 		}
-		bool willBeUsedUncompressed = true;
+		bool storeUncompressed = true;
 		CasKey newKey;
 		u64 fileSize;
-		if (!RetrieveCasFile(newKey, fileSize, casKey, outVirtual.data, willBeUsedUncompressed))
+		if (!RetrieveCasFile(newKey, fileSize, casKey, outVirtual.data, storeUncompressed))
 			UBA_ASSERTF(false, TC("Casfile not found for %s using %s"), outVirtual.data, CasKeyString(casKey).str);
 		StringBuffer<> destFile;
 		if (fileName[1] == ':')
@@ -197,14 +197,6 @@ namespace uba
 		return EnsureApplicationEnvironment(outRealApplication, 0, startInfo.application);
 	}
 
-	u64 SessionClient::GetMemoryMapAlignment(const tchar* fileName, u64 fileNameLen) const
-	{
-		u64 alignment = Session::GetMemoryMapAlignment(fileName, fileNameLen);
-		if (!alignment && !m_useStorage && m_allowMemoryMaps)
-			return 64 * 1024;
-		return alignment;
-	}
-
 	bool SessionClient::EnsureApplicationEnvironment(StringBufferBase& out, u32 processId, const tchar* application)
 	{
 		StringBuffer<> applicationDir;
@@ -253,9 +245,9 @@ namespace uba
 				}
 
 				CasKey newCasKey;
-				bool willBeUsedUncompressed = true;
+				bool storeUncompressed = true;
 				u64 fileSize;
-				if (!RetrieveCasFile(newCasKey, fileSize, casKey, moduleFile.data, willBeUsedUncompressed))
+				if (!RetrieveCasFile(newCasKey, fileSize, casKey, moduleFile.data, storeUncompressed))
 					return m_logger.Error(TC("Casfile not found for %s (%s)"), moduleFile.data, CasKeyString(casKey).str);
 
 				const tchar* moduleName = moduleFile.data;
@@ -342,7 +334,14 @@ namespace uba
 		bool isDir = casKey == CasKeyIsDirectory;
 		u64 fileSize = InvalidValue;
 		CasKey newCasKey;
-		u64 memoryMapAlignment = GetMemoryMapAlignment(fileName.data, fileName.count);
+
+		u32 memoryMapAlignment = 0;
+		if (m_allowMemoryMaps)
+		{
+			memoryMapAlignment = GetMemoryMapAlignment(fileName.data, fileName.count);
+			if (!memoryMapAlignment && !m_useStorage)
+				memoryMapAlignment = 64 * 1024;
+		}
 
 		if (isDir)
 		{
@@ -352,8 +351,8 @@ namespace uba
 		{
 			if (m_useStorage || memoryMapAlignment == 0)
 			{
-				bool willBeUsedUncompressed = memoryMapAlignment == 0;
-				if (!RetrieveCasFile(newCasKey, fileSize, casKey, fileName.data, willBeUsedUncompressed, !IsRarelyRead(msg.process, fileName)))
+				bool storeUncompressed = memoryMapAlignment == 0;
+				if (!RetrieveCasFile(newCasKey, fileSize, casKey, fileName.data, storeUncompressed, !IsRarelyRead(msg.process, fileName)))
 					return m_logger.Error(TC("Error retrieving cas entry %s (%s)"), CasKeyString(casKey).str, fileName.data);
 
 				#if !UBA_USE_SPARSEFILE
@@ -1083,6 +1082,26 @@ namespace uba
 		u32 neededDirectoryTableSize = reader.ReadU32();
 		u32 neededHashTableSize = reader.ReadU32();
 
+		if (u32 knownInputsCount = reader.ReadU32())
+		{
+			while (knownInputsCount--)
+			{
+				CasKey knownInputKey = reader.ReadCasKey();
+				u32 mappingAlignment = reader.ReadU32();
+				bool storeUncompressed = !m_allowMemoryMaps || mappingAlignment == 0;
+				if (storeUncompressed)
+					knownInputKey = AsCompressed(knownInputKey, false);
+
+				m_client.AddWork([knownInputKey, this]()
+					{
+						Storage::RetrieveResult result;
+						bool allowProxy = true;
+						bool res = m_storage.RetrieveCasFile(result, knownInputKey, TC(""), nullptr, 1, allowProxy);
+						(void)res;
+					}, 1, TC("KnownInput"));
+			}
+		}
+
 		if (!out.empty())
 		{
 			if (neededDirectoryTableSize > GetDirectoryTableSize())
@@ -1561,6 +1580,8 @@ namespace uba
 		};
 
 		SendSummary();
+
+		m_client.FlushWork();
 	}
 
 	u32 SessionClient::CountLogLines(ProcessImpl& process)
diff --git a/Engine/Source/Programs/UnrealBuildAccelerator/Common/Private/UbaSessionServer.cpp b/Engine/Source/Programs/UnrealBuildAccelerator/Common/Private/UbaSessionServer.cpp
--- a/Engine/Source/Programs/UnrealBuildAccelerator/Common/Private/UbaSessionServer.cpp
+++ b/Engine/Source/Programs/UnrealBuildAccelerator/Common/Private/UbaSessionServer.cpp
@@ -30,6 +30,7 @@ namespace uba
 		{
 			StringBuffer<512> temp;
 			FixFileName(temp, si.workingDir, nullptr);
+			temp.EnsureEndsWithSlash();
 			workingDir = temp.data;
 
 			temp.EnsureEndsWithSlash();
@@ -47,6 +48,11 @@ namespace uba
 			m_startInfo.logFile = nullptr;
 		}
 
+		~RemoteProcess()
+		{
+			delete[] m_knownInputs;
+		}
+
 		virtual const ProcessStartInfo& GetStartInfo() override { return m_startInfo; }
 		virtual u32 GetId() override { return m_processId; }
 		virtual u32 GetExitCode() override { UBA_ASSERT(m_done.IsSet(0)); return m_exitCode; }
@@ -99,6 +105,10 @@ namespace uba
 		u32 m_sessionId = 0;
 		float m_weight;
 		TString m_executingHost;
+
+		struct KnownInput { CasKey key; u32 mappingAlignment = 0; };
+		KnownInput* m_knownInputs = nullptr;
+		u32 m_knownInputsCount = 0;
 	};
 
 
@@ -261,12 +271,53 @@ namespace uba
 		return {};
 	}
 
-	ProcessHandle SessionServer::RunProcessRemote(const ProcessStartInfo& startInfo, float weight, const void* knownInputs, u32 knownInputsSizeBytes)
+	ProcessHandle SessionServer::RunProcessRemote(const ProcessStartInfo& startInfo, float weight, const void* knownInputs, u32 knownInputsCount)
 	{
 		FlushDeadProcesses();
 		ValidateStartInfo(startInfo);
 		u32 processId = ++m_processIdCounter;
 		RemoteProcess* remoteProcess = new RemoteProcess(this, startInfo, processId, weight);
+		
+		if (knownInputsCount)
+		{
+			auto keys = remoteProcess->m_knownInputs = new RemoteProcess::KnownInput[knownInputsCount];
+
+			u32 keysIndex = 0;
+			const TString& workingDir = remoteProcess->workingDir;
+			for (auto kiIt = (const tchar*)knownInputs; *kiIt; kiIt += TStrlen(kiIt) + 1)
+			{
+				StringBuffer<> fileName;
+				FixPath(kiIt, workingDir.c_str(), u32(workingDir.size()), fileName);
+
+				// Make sure cas entry exists and caskey is calculated (cas content creation is deferred in case client already has it)
+				CasKey casKey;
+				bool deferCreation = true;
+				if (!m_storage.StoreCasFile(casKey, fileName.data, CasKeyZero, deferCreation) || casKey == CasKeyZero)
+					continue;
+
+				auto& ki = keys[keysIndex++];
+				ki.key = casKey;
+				ki.mappingAlignment = GetMemoryMapAlignment(fileName.data, fileName.count);
+
+
+				// Update name to hash table
+				if (CaseInsensitiveFs)
+					fileName.MakeLower();
+				StringKey fileNameKey = ToStringKey(fileName);
+				ScopedWriteLock lock(m_nameToHashLookupLock);
+				CasKey& lookupCasKey = m_nameToHashLookup[fileNameKey];
+				if (lookupCasKey != casKey)
+				{
+					lookupCasKey = casKey;
+					BinaryWriter w(m_nameToHashTableMem.memory, m_nameToHashTableMem.writtenSize, NameToHashMemSize);
+					m_nameToHashTableMem.AllocateNoLock(sizeof(StringKey) + sizeof(CasKey), 1, TC("NameToHashTable"));
+					w.WriteStringKey(fileNameKey);
+					w.WriteCasKey(lookupCasKey);
+				}
+			}
+			remoteProcess->m_knownInputsCount = keysIndex;
+		}
+		
 		ScopedCriticalSection lock(m_remoteProcessAndSessionLock);
 		m_queuedRemoteProcesses.push_back(remoteProcess);
 		if (m_remoteProcessReturnedEvent)
@@ -859,9 +910,15 @@ namespace uba
 				u32 sessionId = reader.ReadU32();
 				u32 sessionIndex = sessionId - 1;
 
+				ScopedCriticalSection sessionsLock(m_remoteProcessAndSessionLock);
+				ClientSession& session = *m_clientSessions[sessionIndex];
+				sessionsLock.Leave();
+
 				u32 weight32 = reader.ReadU32();
 				float availableWeight = *(float*)&weight32;
 
+				Vector<RemoteProcess::KnownInput*> knownInputsToSend;
+
 				float weightLeft = availableWeight;
 				u32 addCount = 0;
 				ScopedWriteLock fillLock(m_fillUpOneAtTheTimeLock); // This is a lock to group files better (all clients connect at the same time)
@@ -880,6 +937,10 @@ namespace uba
 					writer.WriteU32(*(u32*)&process->m_weight);
 					writer.WriteU64(process->m_startInfo.outputStatsThresholdMs);
 
+					for (auto kiIt = process->m_knownInputs, kiEnd = kiIt + process->m_knownInputsCount; kiIt!=kiEnd; ++kiIt)
+						if (session.sentKeys.insert(kiIt->key).second)
+							knownInputsToSend.push_back(kiIt);
+
 					++addCount;
 
 					if (writer.GetCapacityLeft() < 5000) // Arbitrary number to cover all parameters above
@@ -892,9 +953,7 @@ namespace uba
 				u32 neededDirectoryTableSize = GetDirectoryTableSize();
 				u32 neededHashTableSize = m_nameToHashLookupLock.ScopedRead([this]() { return u32(m_nameToHashTableMem.writtenSize); });
 
-				ScopedCriticalSection lock(m_remoteProcessAndSessionLock);
-				ClientSession& session = *m_clientSessions[sessionIndex];
-
+				sessionsLock.Enter();
 				//if (addCount)
 				//	m_logger.Debug(TC("Gave %u processes to %s using up %.1f weight out of %.1f available"), addCount, session.name.c_str(), availableWeight - weightLeft, availableWeight);
 
@@ -922,6 +981,7 @@ namespace uba
 						}
 					}
 				}
+				sessionsLock.Leave();
 
 				writer.WriteU32(remoteExecutionEnabled ? SessionProcessAvailableResponse_None : SessionProcessAvailableResponse_RemoteExecutionDisabled);
 
@@ -930,6 +990,18 @@ namespace uba
 				writer.WriteU32(neededDirectoryTableSize);
 				writer.WriteU32(neededHashTableSize);
 
+				// Send caskeys of known inputs so client can start retrieving them straight away
+				u32 kiCapacity = u32(writer.GetCapacityLeft() - sizeof(u32)) / sizeof(RemoteProcess::KnownInput);
+				u32 toSendCount = Min(kiCapacity, u32(knownInputsToSend.size()));
+				writer.WriteU32(toSendCount);
+				for (auto kv : knownInputsToSend)
+				{
+					if (!toSendCount--)
+						break;
+					writer.WriteCasKey(kv->key);
+					writer.WriteU32(kv->mappingAlignment);
+				}
+
 				return true;
 			}
 			case SessionMessageType_ProcessFinished:
@@ -1424,7 +1496,7 @@ namespace uba
 
 	bool SessionServer::CreateFile(CreateFileResponse& out, const CreateFileMessage& msg)
 	{
-		if ((msg.access & FileAccess_Write) == 0)
+		if (!m_shouldWriteToDisk && ((msg.access & FileAccess_Write) == 0))
 		{
 			ScopedReadLock lock(m_receivedFilesLock);
 			auto findIt = m_receivedFiles.find(msg.fileNameKey);
diff --git a/Engine/Source/Programs/UnrealBuildAccelerator/Common/Private/UbaStorage.cpp b/Engine/Source/Programs/UnrealBuildAccelerator/Common/Private/UbaStorage.cpp
--- a/Engine/Source/Programs/UnrealBuildAccelerator/Common/Private/UbaStorage.cpp
+++ b/Engine/Source/Programs/UnrealBuildAccelerator/Common/Private/UbaStorage.cpp
@@ -1434,7 +1434,7 @@ namespace uba
 			});
 		m_logger.Info(TC("Validating %u entries..."), entryCount);
 
-		workManager.Wait();
+		workManager.FlushWork();
 
 		StringBuffer<> newestLastWrittenStr;
 		writeTimeAgo(newestLastWrittenStr, newestWrittenError);
@@ -1499,7 +1499,7 @@ namespace uba
 							atomicDeleteCount += deleteCountTemp;
 						}, 1, TC(""));
 				});
-			workManager.Wait();
+			workManager.FlushWork();
 			deleteCount += atomicDeleteCount;
 		}
 
@@ -1551,7 +1551,8 @@ namespace uba
 		if (fileEntry.verified)
 		{
 			UBA_ASSERT(casKeyOverride == CasKeyZero || casKeyOverride == fileEntry.casKey);
-			CasEntryAccessed(fileEntry.casKey);
+			if (!AddCasFile(fileName, fileEntry.casKey, deferCreation))
+				return false;
 			out = fileEntry.casKey;
 			return true;
 		}
diff --git a/Engine/Source/Programs/UnrealBuildAccelerator/Common/Private/UbaStorageClient.cpp b/Engine/Source/Programs/UnrealBuildAccelerator/Common/Private/UbaStorageClient.cpp
--- a/Engine/Source/Programs/UnrealBuildAccelerator/Common/Private/UbaStorageClient.cpp
+++ b/Engine/Source/Programs/UnrealBuildAccelerator/Common/Private/UbaStorageClient.cpp
@@ -85,7 +85,7 @@ namespace uba
 
 		for (auto& dir : directories)
 			success = PopulateCasFromDirsRecursive(dir.c_str(), workManager, seenIds, seenIdsLock) && success;
-		workManager.Wait();
+		workManager.FlushWork();
 
 		if (u32 fileCount = u32(m_localStorageFiles.size()))
 			m_logger.Info(TC("Prepopulated %u files to cas in %s"), fileCount, TimeToText(GetTime() - start).str);
diff --git a/Engine/Source/Programs/UnrealBuildAccelerator/Common/Private/UbaWorkManager.cpp b/Engine/Source/Programs/UnrealBuildAccelerator/Common/Private/UbaWorkManager.cpp
--- a/Engine/Source/Programs/UnrealBuildAccelerator/Common/Private/UbaWorkManager.cpp
+++ b/Engine/Source/Programs/UnrealBuildAccelerator/Common/Private/UbaWorkManager.cpp
@@ -115,7 +115,7 @@ namespace uba
 		--m_activeWorkerCount;
 	}
 
-	void WorkManagerImpl::Wait()
+	void WorkManagerImpl::FlushWork()
 	{
 		while (true)
 		{
diff --git a/Engine/Source/Programs/UnrealBuildAccelerator/Common/Public/UbaNetwork.h b/Engine/Source/Programs/UnrealBuildAccelerator/Common/Public/UbaNetwork.h
--- a/Engine/Source/Programs/UnrealBuildAccelerator/Common/Public/UbaNetwork.h
+++ b/Engine/Source/Programs/UnrealBuildAccelerator/Common/Public/UbaNetwork.h
@@ -18,7 +18,7 @@ namespace uba
 	
 	static constexpr u32 SystemNetworkVersion = 1338;
 	static constexpr u32 StorageNetworkVersion = 3;
-	static constexpr u32 SessionNetworkVersion = 23;
+	static constexpr u32 SessionNetworkVersion = 24;
 
 
 	// Messages used over network between client/server (system, storage and session)
diff --git a/Engine/Source/Programs/UnrealBuildAccelerator/Common/Public/UbaSession.h b/Engine/Source/Programs/UnrealBuildAccelerator/Common/Public/UbaSession.h
--- a/Engine/Source/Programs/UnrealBuildAccelerator/Common/Public/UbaSession.h
+++ b/Engine/Source/Programs/UnrealBuildAccelerator/Common/Public/UbaSession.h
@@ -120,7 +120,8 @@ namespace uba
 		void StopTraceThread();
 		u32 GetDirectoryTableSize();
 		u32 GetFileMappingSize();
-		
+		u32 GetMemoryMapAlignment(const tchar* fileName, u64 fileNameLen) const;
+
 		SessionStats& Stats();
 
 		struct BinaryModule { TString name; TString path; u32 fileAttributes = 0; bool isSystem = false; };
@@ -140,7 +141,6 @@ namespace uba
 		u32 RegisterDeleteFile(StringKey fileNameKey, const tchar* fileName);
 
 		virtual bool PrepareProcess(const ProcessStartInfo& startInfo, bool isChild, StringBufferBase& outRealApplication, const tchar*& outRealWorkingDir);
-		virtual u64 GetMemoryMapAlignment(const tchar* fileName, u64 fileNameLen) const;
 		virtual void* GetProcessEnvironmentVariables();
 		virtual void PrintSessionStats(Logger& logger);
 
diff --git a/Engine/Source/Programs/UnrealBuildAccelerator/Common/Public/UbaSessionClient.h b/Engine/Source/Programs/UnrealBuildAccelerator/Common/Public/UbaSessionClient.h
--- a/Engine/Source/Programs/UnrealBuildAccelerator/Common/Public/UbaSessionClient.h
+++ b/Engine/Source/Programs/UnrealBuildAccelerator/Common/Public/UbaSessionClient.h
@@ -39,10 +39,9 @@ namespace uba
 
 		u64 GetBestPing();
 
-		bool RetrieveCasFile(CasKey& outNewKey, u64& outSize, const CasKey& casKey, const tchar* hint, bool willBeUsedUnCompressed, bool allowProxy = true);
+		bool RetrieveCasFile(CasKey& outNewKey, u64& outSize, const CasKey& casKey, const tchar* hint, bool storeUncompressed, bool allowProxy = true);
 
 		virtual bool PrepareProcess(const ProcessStartInfo& startInfo, bool isChild, StringBufferBase& outRealApplication, const tchar*& outRealWorkingDir) override;
-		virtual u64 GetMemoryMapAlignment(const tchar* fileName, u64 fileNameLen) const override;
 		virtual void* GetProcessEnvironmentVariables() override;
 		virtual bool CreateFile(CreateFileResponse& out, const CreateFileMessage& msg) override;
 		virtual bool DeleteFile(DeleteFileResponse& out, const DeleteFileMessage& msg) override;
diff --git a/Engine/Source/Programs/UnrealBuildAccelerator/Common/Public/UbaSessionServer.h b/Engine/Source/Programs/UnrealBuildAccelerator/Common/Public/UbaSessionServer.h
--- a/Engine/Source/Programs/UnrealBuildAccelerator/Common/Public/UbaSessionServer.h
+++ b/Engine/Source/Programs/UnrealBuildAccelerator/Common/Public/UbaSessionServer.h
@@ -21,9 +21,9 @@ namespace uba
 		// Run process remotely.
 		// startInfo contains info about the process to run remotely
 		// weight is the expected core usage of the process. If processes are multithreaded it makes sense to increase weight. As an example, in UnrealBuildTool we see cl.exe as 1.5 and clang.exe as 1.0
-		// knownInputs is a memory block with null terminated strings followed by an empty null terminated string. knownInputSize is the size of the memoryBlock in bytes
-		// knownInputs strings should be of type tchar and a path relative to working dir or absolute. Session will make a copy of knownInputs internally
-		ProcessHandle RunProcessRemote(const ProcessStartInfo& startInfo, float weight = 1.0f, const void* knownInputs = nullptr, u32 knownInputsSizeBytes = 0);
+		// knownInputs is a memory block with null terminated tchar strings followed by an empty null terminated string to end. knownInputsCount is the number of strings in the memory block
+		// strings should be absolute or relative to working dir.
+		ProcessHandle RunProcessRemote(const ProcessStartInfo& startInfo, float weight = 1.0f, const void* knownInputs = nullptr, u32 knownInputsCount = 0);
 
 		// Will kick off a local process with the same startInfo as the one provided to start the process with id matching raceAgainstRemoteProcessId
 		// This can be useful if there are free local cores and we know local machine is faster or network connection to remote is slow
@@ -128,6 +128,8 @@ namespace uba
 			ReaderWriterLock dirTablePosLock;
 			u32 dirTablePos = 0;
 
+			UnorderedSet<CasKey> sentKeys;
+
 		};
 		Vector<ClientSession*> m_clientSessions;
 
diff --git a/Engine/Source/Programs/UnrealBuildAccelerator/Common/Public/UbaWorkManager.h b/Engine/Source/Programs/UnrealBuildAccelerator/Common/Public/UbaWorkManager.h
--- a/Engine/Source/Programs/UnrealBuildAccelerator/Common/Public/UbaWorkManager.h
+++ b/Engine/Source/Programs/UnrealBuildAccelerator/Common/Public/UbaWorkManager.h
@@ -21,7 +21,7 @@ namespace uba
 		virtual ~WorkManagerImpl();
 		virtual void AddWork(const Function<void()>& work, u32 count, const tchar* desc) override;
 		virtual u32 GetWorkerCount() override;
-		void Wait();
+		void FlushWork();
 	private:
 		struct Worker;
 		void PushWorker(Worker* worker);
diff --git a/Engine/Source/Programs/UnrealBuildAccelerator/Core/Public/UbaSynchronization.h b/Engine/Source/Programs/UnrealBuildAccelerator/Core/Public/UbaSynchronization.h
--- a/Engine/Source/Programs/UnrealBuildAccelerator/Core/Public/UbaSynchronization.h
+++ b/Engine/Source/Programs/UnrealBuildAccelerator/Core/Public/UbaSynchronization.h
@@ -45,6 +45,7 @@ namespace uba
 	public:
 		ScopedCriticalSection(CriticalSection& cs) : m_cs(cs), m_active(true) { cs.Enter(); }
 		~ScopedCriticalSection() { Leave(); }
+		void Enter() { if (m_active) return; m_cs.Enter(); m_active = true; }
 		void Leave() { if (!m_active) return; m_cs.Leave(); m_active = false; }
 	private:
 		CriticalSection& m_cs;
diff --git a/Engine/Source/Programs/UnrealBuildAccelerator/Host/Private/UbaExports.cpp b/Engine/Source/Programs/UnrealBuildAccelerator/Host/Private/UbaExports.cpp
--- a/Engine/Source/Programs/UnrealBuildAccelerator/Host/Private/UbaExports.cpp
+++ b/Engine/Source/Programs/UnrealBuildAccelerator/Host/Private/UbaExports.cpp
@@ -338,9 +338,9 @@ uba::StorageClient* CreateStorageClient(uba::NetworkClient& client, const uba::t
 	{
 		return new uba::ProcessHandle(server->RunProcess(info, async, enableDetour));
 	}
-	uba::ProcessHandle* SessionServer_RunProcessRemote(uba::SessionServer* server, uba::ProcessStartInfo& info, float weight, const void* knownInputs, uba::u32 knownInputsSizeBytes)
+	uba::ProcessHandle* SessionServer_RunProcessRemote(uba::SessionServer* server, uba::ProcessStartInfo& info, float weight, const void* knownInputs, uba::u32 knownInputsCount)
 	{
-		return new uba::ProcessHandle(server->RunProcessRemote(info, weight, knownInputs, knownInputsSizeBytes));
+		return new uba::ProcessHandle(server->RunProcessRemote(info, weight, knownInputs, knownInputsCount));
 	}
 	uba::ProcessHandle* SessionServer_RunProcessRacing(uba::SessionServer* server, uba::u32 raceAgainstRemoteProcessId)
 	{
diff --git a/Engine/Source/Programs/UnrealBuildAccelerator/Host/Public/UbaExports.h b/Engine/Source/Programs/UnrealBuildAccelerator/Host/Public/UbaExports.h
--- a/Engine/Source/Programs/UnrealBuildAccelerator/Host/Public/UbaExports.h
+++ b/Engine/Source/Programs/UnrealBuildAccelerator/Host/Public/UbaExports.h
@@ -86,7 +86,7 @@ extern "C"
 	UBA_API void SessionServer_RefreshDirectory(uba::SessionServer* server, const uba::tchar* directory);
 	UBA_API void SessionServer_RegisterNewFile(uba::SessionServer* server, const uba::tchar* filePath);
 	UBA_API uba::ProcessHandle* SessionServer_RunProcess(uba::SessionServer* server, uba::ProcessStartInfo& info, bool async, bool enableDetour);
-	UBA_API uba::ProcessHandle* SessionServer_RunProcessRemote(uba::SessionServer* server, uba::ProcessStartInfo& info, float weight, const void* knownInputs = nullptr, uba::u32 knownInputsSizeBytes = 0);
+	UBA_API uba::ProcessHandle* SessionServer_RunProcessRemote(uba::SessionServer* server, uba::ProcessStartInfo& info, float weight, const void* knownInputs = nullptr, uba::u32 knownInputsCount = 0);
 	UBA_API uba::ProcessHandle* SessionServer_RunProcessRacing(uba::SessionServer* server, uba::u32 raceAgainstRemoteProcessId);
 
 	UBA_API void SessionServer_SetMaxRemoteProcessCount(uba::SessionServer* server, uba::u32 count);
diff --git a/Engine/Source/Programs/UnrealBuildTool/Executors/Experimental/UBAExecutor.cs b/Engine/Source/Programs/UnrealBuildTool/Executors/Experimental/UBAExecutor.cs
--- a/Engine/Source/Programs/UnrealBuildTool/Executors/Experimental/UBAExecutor.cs
+++ b/Engine/Source/Programs/UnrealBuildTool/Executors/Experimental/UBAExecutor.cs
@@ -752,7 +752,7 @@ namespace UnrealBuildTool
 				byte[]? knownInputs = null;
 				if (UBAConfig.bUseKnownInputs)
 				{
-					int sizeOfChar = Marshal.SizeOf<char>();
+					int sizeOfChar = System.OperatingSystem.IsWindows() ? 2 : 1;
 
 					int byteCount = 0;
 					foreach (var item in action.PrerequisiteItems)
