[Uba]
* #if 0 on code that does not exist on mac

diff --git a/Engine/Source/Programs/UnrealBuildAccelerator/Common/Private/UbaProcess.cpp b/Engine/Source/Programs/UnrealBuildAccelerator/Common/Private/UbaProcess.cpp
--- a/Engine/Source/Programs/UnrealBuildAccelerator/Common/Private/UbaProcess.cpp
+++ b/Engine/Source/Programs/UnrealBuildAccelerator/Common/Private/UbaProcess.cpp
@@ -1326,12 +1326,14 @@ namespace uba
 			UBA_ASSERTF(res == 0, TC("posix_spawnattr_setpgroup (%s)"), strerror(errno));
 
 			/* Set the child's priority to be two levels lower. */
+			#if 0//PLATFORM_LINUX
 			sched_param param;
 			res = posix_spawnattr_getschedparam(&attr, &param);
 			UBA_ASSERTF(res == 0, TC("posix_spawnattr_getschedparam (%s)"), strerror(errno));
 			param.sched_priority -= 2;
 			res = posix_spawnattr_setschedparam(&attr, &param);
 			UBA_ASSERTF(res == 0, TC("posix_spawnattr_setschedparam (%s)"), strerror(errno));
+			#endif
 
 			posix_spawn_file_actions_t fileActions;
 			res = posix_spawn_file_actions_init(&fileActions);
