[Uba]
* Added so if NetworkServer fails to stop connections after x seconds because of running workers we abort the entire process. I'm hoping we will get

diff --git a/Engine/Source/Programs/UnrealBuildAccelerator/Common/Private/UbaNetworkServer.cpp b/Engine/Source/Programs/UnrealBuildAccelerator/Common/Private/UbaNetworkServer.cpp
--- a/Engine/Source/Programs/UnrealBuildAccelerator/Common/Private/UbaNetworkServer.cpp
+++ b/Engine/Source/Programs/UnrealBuildAccelerator/Common/Private/UbaNetworkServer.cpp
@@ -102,7 +102,7 @@ namespace uba
 				TestDisconnect();
 		}
 
-		void Stop()
+		bool Stop()
 		{
 			Disconnect();
 				
@@ -111,11 +111,12 @@ namespace uba
 			{
 				if (TimeToMs(GetTime() - startTimer) > 3000)
 				{
-					m_server.m_logger.Error(TC("Connection has waited 3 seconds to stop... something is stuck"));
-					break;
+					m_server.m_logger.Error(TC("Connection has waited 3 seconds to stop... something is stuck (Active worker count: %u). Need someone to attach a debugger to process to get callstacks"), m_activeWorkerCount.load());
+					return false;
 				}
 				Sleep(1);
 			}
+			return true;
 		}
 
 		bool SendAsync(u8 value)
@@ -591,13 +592,18 @@ namespace uba
 			auto connections(std::move(m_connections));
 			lock.Leave();
 
+			bool success = true;
 			for (auto& c : connections)
 			{
-				c.Stop();
+				success = c.Stop() && success;
 				m_sendTimer.Add(c.m_sendTimer);
 				m_encryptTimer.Add(c.m_encryptTimer);
 				m_decryptTimer.Add(c.m_decryptTimer);
 			}
+
+			// If stopping connections fail we need to abort because we will most likely run into a deadlock when deleting the workers.
+			if (!success)
+				abort(); // TODO: Does this produce core dump on windows?
 		}
 
 		auto deleteWorkers = [](Worker*& start)
