- Lowered the priority of the processes Uba spawns, so that too many clangs don't starve the session process. Doesn't always repro, but it was causing some networking errors when talking to horde

diff --git a/Engine/Source/Programs/UnrealBuildAccelerator/Common/Private/UbaProcess.cpp b/Engine/Source/Programs/UnrealBuildAccelerator/Common/Private/UbaProcess.cpp
--- a/Engine/Source/Programs/UnrealBuildAccelerator/Common/Private/UbaProcess.cpp
+++ b/Engine/Source/Programs/UnrealBuildAccelerator/Common/Private/UbaProcess.cpp
@@ -1325,16 +1325,6 @@ namespace uba
 			res = posix_spawnattr_setpgroup(&attr, getpgrp());
 			UBA_ASSERTF(res == 0, TC("posix_spawnattr_setpgroup (%s)"), strerror(errno));
 
-			/* Set the child's priority to be two levels lower. */
-			#if 0//PLATFORM_LINUX
-			sched_param param;
-			res = posix_spawnattr_getschedparam(&attr, &param);
-			UBA_ASSERTF(res == 0, TC("posix_spawnattr_getschedparam (%s)"), strerror(errno));
-			param.sched_priority -= 2;
-			res = posix_spawnattr_setschedparam(&attr, &param);
-			UBA_ASSERTF(res == 0, TC("posix_spawnattr_setschedparam (%s)"), strerror(errno));
-			#endif
-
 			posix_spawn_file_actions_t fileActions;
 			res = posix_spawn_file_actions_init(&fileActions);
 			UBA_ASSERTF(res == 0, TC("posix_spawn_file_actions_init (%s)"), strerror(errno));
@@ -1465,6 +1455,13 @@ namespace uba
 				logger.Error(TC("posix_spawn failed: %s %s (Working dir: %s) -> %i (%s)"), m_realApplication.c_str(), m_startInfo.arguments, m_realWorkingDir, res, strerror(errno));
 				return UBA_EXIT_CODE(12);
 			}
+			
+			#if PLATFORM MAC
+			int prio = getpriority(PRIO_PROCESS, processID);
+			errno = 0;
+			res = setpriority(PRIO_PROCESS, processID, prio + 2);
+			UBA_ASSERTF(res == 0, TC("setpriority (%s)"), strerror(errno));
+			#endif
 
 			m_nativeProcessHandle = (ProcHandle)1;
 			m_nativeProcessId = u32(processID);
