- Removed unneeded PLATFORM_MAC

diff --git a/Engine/Source/Programs/UnrealBuildAccelerator/Common/Private/UbaProcess.cpp b/Engine/Source/Programs/UnrealBuildAccelerator/Common/Private/UbaProcess.cpp
--- a/Engine/Source/Programs/UnrealBuildAccelerator/Common/Private/UbaProcess.cpp
+++ b/Engine/Source/Programs/UnrealBuildAccelerator/Common/Private/UbaProcess.cpp
@@ -1456,12 +1456,10 @@ namespace uba
 				return UBA_EXIT_CODE(12);
 			}
 			
-			#if PLATFORM MAC
 			int prio = getpriority(PRIO_PROCESS, processID);
 			errno = 0;
 			res = setpriority(PRIO_PROCESS, processID, prio + 2);
 			UBA_ASSERTF(res == 0, TC("setpriority (%s)"), strerror(errno));
-			#endif
 
 			m_nativeProcessHandle = (ProcHandle)1;
 			m_nativeProcessId = u32(processID);
