UnrealBuildAccelerator: Workarounds to get latest PVS-Studio working

#rnx

diff --git a/Engine/Source/Programs/UnrealBuildAccelerator/Core/Private/UbaApplicationRules.cpp b/Engine/Source/Programs/UnrealBuildAccelerator/Core/Private/UbaApplicationRules.cpp
--- a/Engine/Source/Programs/UnrealBuildAccelerator/Core/Private/UbaApplicationRules.cpp
+++ b/Engine/Source/Programs/UnrealBuildAccelerator/Core/Private/UbaApplicationRules.cpp
@@ -1,6 +1,7 @@
 // Copyright Epic Games, Inc. All Rights Reserved.
 
 #include "UbaApplicationRules.h"
+#include "UbaPlatform.h"
 
 namespace uba
 {
@@ -335,6 +336,23 @@ namespace uba
 			return EndsWith(file, fileLen, TC(".pvslog"))
 				|| EndsWith(file, fileLen, TC(".stacktrace.txt"));
 		}
+		
+		virtual bool IsRarelyRead(const StringBufferBase& file) override
+		{
+			return file.EndsWith(TC(".i"))
+				|| file.EndsWith(TC(".pvslog"))
+				|| file.EndsWith(TC(".stacktrace.txt"));
+		}
+
+		virtual void RepairMalformedLibPath(const wchar_t* path) override
+		{
+			// There is a bug where the path passed into wsplitpath_s is malformed and not null terminated correctly
+			const wchar_t* pext = TStrstr(path, L".dll");
+			if (pext == nullptr) pext = TStrstr(path, L".DLL");
+			if (pext == nullptr) pext = TStrstr(path, L".exe");
+			if (pext == nullptr) pext = TStrstr(path, L".EXE");
+			if (pext != nullptr && *(pext + 4) != 0) *(const_cast<wchar_t*>(pext + 4)) = 0;
+		}
 	};
 
 	const RulesRec* GetApplicationRules()
diff --git a/Engine/Source/Programs/UnrealBuildAccelerator/Core/Public/UbaApplicationRules.h b/Engine/Source/Programs/UnrealBuildAccelerator/Core/Public/UbaApplicationRules.h
--- a/Engine/Source/Programs/UnrealBuildAccelerator/Core/Public/UbaApplicationRules.h
+++ b/Engine/Source/Programs/UnrealBuildAccelerator/Core/Public/UbaApplicationRules.h
@@ -113,6 +113,11 @@ namespace uba
 		{
 			return true;
 		}
+
+		virtual void RepairMalformedLibPath(const wchar_t* path)
+		{
+			// Do nothing
+		}
 	};
 
 
diff --git a/Engine/Source/Programs/UnrealBuildAccelerator/Detours/Private/Windows/UbaDetoursFunctionsKernelBase.inl b/Engine/Source/Programs/UnrealBuildAccelerator/Detours/Private/Windows/UbaDetoursFunctionsKernelBase.inl
--- a/Engine/Source/Programs/UnrealBuildAccelerator/Detours/Private/Windows/UbaDetoursFunctionsKernelBase.inl
+++ b/Engine/Source/Programs/UnrealBuildAccelerator/Detours/Private/Windows/UbaDetoursFunctionsKernelBase.inl
@@ -2246,6 +2246,8 @@ DWORD Detoured_SearchPathW(LPCWSTR lpPath, LPCWSTR lpFileName, LPCWSTR lpExtensi
 	DETOURED_CALL(SearchPathW);
 	if (g_runningRemote && !t_disallowDetour)
 	{
+		g_rules->RepairMalformedLibPath(lpFileName);
+
 		const wchar_t* original = lpFileName; (void)original;
 		u64 pathLen = wcslen(lpFileName);
 		StringBuffer<512> tempBuf;
@@ -2647,6 +2649,27 @@ BOOL Detoured_GetExitCodeProcess(HANDLE hProcess, LPDWORD lpExitCode)
 	return res;
 }
 
+BOOL Detoured_CreateTimerQueueTimer(PHANDLE phNewTimer, HANDLE TimerQueue, WAITORTIMERCALLBACK Callback, PVOID Parameter, DWORD DueTime, DWORD Period, ULONG Flags)
+{
+	DETOURED_CALL(CreateTimerQueueTimer);
+	BOOL res = True_CreateTimerQueueTimer(phNewTimer, TimerQueue, Callback, Parameter, DueTime, Period, Flags);
+	DEBUG_LOG_TRUE(L"CreateTimerQueueTimer", L"%p -> %ls", *phNewTimer, ToString(res));
+	return res;
+}
+
+BOOL Detoured_DeleteTimerQueueTimer(HANDLE TimerQueue, HANDLE Timer, HANDLE CompletionEvent)
+{
+	DETOURED_CALL(DeleteTimerQueueTimer);
+	BOOL res = True_DeleteTimerQueueTimer(TimerQueue, Timer, CompletionEvent);
+	if (!res && IsRunningWine())
+	{
+		DEBUG_LOG_DETOURED(L"DeleteTimerQueueTimer", L"%p %p %p -> %ls (WINE ignored)", TimerQueue, Timer, CompletionEvent, ToString(res));
+		return true;
+	}
+	DEBUG_LOG_TRUE(L"DeleteTimerQueueTimer", L"%p %p %p -> %ls", TimerQueue, Timer, CompletionEvent, ToString(res));
+	return res;
+}
+
 DWORD Detoured_WaitForSingleObject(HANDLE hHandle, DWORD dwMilliseconds)
 {
 	return WaitForSingleObjectEx(hHandle, dwMilliseconds, false);
diff --git a/Engine/Source/Programs/UnrealBuildAccelerator/Detours/Private/Windows/UbaDetoursFunctionsNtDll.inl b/Engine/Source/Programs/UnrealBuildAccelerator/Detours/Private/Windows/UbaDetoursFunctionsNtDll.inl
--- a/Engine/Source/Programs/UnrealBuildAccelerator/Detours/Private/Windows/UbaDetoursFunctionsNtDll.inl
+++ b/Engine/Source/Programs/UnrealBuildAccelerator/Detours/Private/Windows/UbaDetoursFunctionsNtDll.inl
@@ -513,7 +513,7 @@ NTSTATUS NTAPI Shared_NtCreateFile(bool IsCreateFunc, PHANDLE hFileHandle, ACCES
 			if (StartsWith(createFileName, L"\\\\?\\"))
 				createFileName += 4;
 			if (!FixPath(fileName, createFileName))
-				UBA_ASSERT(false);
+				UBA_ASSERTF(false, L"FixPath failed for string '%ls'", createFileName);
 			if (fileName.StartsWith(L"\\\\.\\pipe"))
 				createFileName = nullptr;
 
@@ -523,7 +523,7 @@ NTSTATUS NTAPI Shared_NtCreateFile(bool IsCreateFunc, PHANDLE hFileHandle, ACCES
 		else if (memcmp(buf, L"\\??\\", 8) == 0)
 		{
 			if (!FixPath(fileName, buf + 4))
-				UBA_ASSERT(false);
+				UBA_ASSERTF(false, L"FixPath failed for string '%ls'", buf + 4);
 		}
 		else
 		{
diff --git a/Engine/Source/Programs/UnrealBuildAccelerator/Detours/Private/Windows/UbaDetoursFunctionsUcrtBase.inl b/Engine/Source/Programs/UnrealBuildAccelerator/Detours/Private/Windows/UbaDetoursFunctionsUcrtBase.inl
--- a/Engine/Source/Programs/UnrealBuildAccelerator/Detours/Private/Windows/UbaDetoursFunctionsUcrtBase.inl
+++ b/Engine/Source/Programs/UnrealBuildAccelerator/Detours/Private/Windows/UbaDetoursFunctionsUcrtBase.inl
@@ -329,6 +329,15 @@ int Detoured__wspawnl(int mode, const wchar_t* cmdname, const wchar_t* arg0, con
 	return exitCode;
 }
 
+errno_t Detoured__wsplitpath_s(const wchar_t* path, wchar_t* drive, size_t driveNumberOfElements, wchar_t* dir, size_t dirNumberOfElements, wchar_t* fname, size_t nameNumberOfElements, wchar_t* ext, size_t extNumberOfElements)
+{
+	DETOURED_CALL(_wsplitpath_s);
+
+	g_rules->RepairMalformedLibPath(path);
+	auto res = True__wsplitpath_s(path, drive, driveNumberOfElements, dir, dirNumberOfElements, fname, nameNumberOfElements, ext, extNumberOfElements);
+	DEBUG_LOG_TRUE(L"_wsplitpath_s", L"%ls %ls %ls %ls %ls", path, drive, dir, fname, ext);
+	return res;
+}
 
 #if defined(DETOURED_INCLUDE_DEBUG)
 int Detoured__wcsnicoll_l(const wchar_t* string1, const wchar_t* string2, size_t count, _locale_t locale)
@@ -386,13 +395,6 @@ errno_t Detoured_getenv_s(size_t* pReturnValue, char* buffer, size_t numberOfEle
 	return res;
 }
 
-errno_t Detoured__wsplitpath_s(const wchar_t* path, wchar_t* drive, size_t driveNumberOfElements, wchar_t* dir, size_t dirNumberOfElements, wchar_t* fname, size_t nameNumberOfElements, wchar_t* ext, size_t extNumberOfElements)
-{
-	auto res = True__wsplitpath_s(path, drive, driveNumberOfElements, dir, dirNumberOfElements, fname, nameNumberOfElements, ext, extNumberOfElements);
-	DEBUG_LOG_TRUE(L"_wsplitpath_s", L"%ls %ls %ls %ls %ls", path, drive, dir, fname, ext);
-	return res;
-}
-
 errno_t Detoured__wmakepath_s(wchar_t* path, size_t sizeInWords, const wchar_t* drive, const wchar_t* dir, const wchar_t* fname, const wchar_t* ext)
 {
 	auto res = True__wmakepath_s(path, sizeInWords, drive, dir, fname, ext);
diff --git a/Engine/Source/Programs/UnrealBuildAccelerator/Detours/Private/Windows/UbaDetoursFunctionsWin.h b/Engine/Source/Programs/UnrealBuildAccelerator/Detours/Private/Windows/UbaDetoursFunctionsWin.h
--- a/Engine/Source/Programs/UnrealBuildAccelerator/Detours/Private/Windows/UbaDetoursFunctionsWin.h
+++ b/Engine/Source/Programs/UnrealBuildAccelerator/Detours/Private/Windows/UbaDetoursFunctionsWin.h
@@ -112,6 +112,8 @@
 #define DETOURED_FUNCTIONS_KERNEL32 \
 	DETOURED_FUNCTION(CreateFileMappingA) \
 	DETOURED_FUNCTION(GetExitCodeProcess) \
+	DETOURED_FUNCTION(CreateTimerQueueTimer) \
+	DETOURED_FUNCTION(DeleteTimerQueueTimer) \
 	DETOURED_FUNCTIONS_KERNEL32_DEBUG \
 
 #define DETOURED_FUNCTIONS_NTDLL \
@@ -147,7 +149,8 @@
 	DETOURED_FUNCTION(_isatty) \
 	DETOURED_FUNCTION(_write) \
 	DETOURED_FUNCTION(fputs) \
-	DETOURED_FUNCTIONS_CRTBASE_DEBUG \
+	DETOURED_FUNCTION(_wsplitpath_s) \
+	DETOURED_FUNCTIONS_UCRTBASE_DEBUG \
 
 #if defined(UBA_USE_MIMALLOC)
 #define DETOURED_FUNCTIONS_MEMORY \
@@ -300,13 +303,12 @@
 	DETOURED_FUNCTION(PathFileExistsW) \
 
 
-#define DETOURED_FUNCTIONS_CRTBASE_DEBUG \
+#define DETOURED_FUNCTIONS_UCRTBASE_DEBUG \
 	DETOURED_FUNCTION(_wcsnicoll_l) \
 	DETOURED_FUNCTION(_wgetenv) \
 	DETOURED_FUNCTION(_wgetenv_s) \
 	DETOURED_FUNCTION(getenv) \
 	DETOURED_FUNCTION(getenv_s) \
-	DETOURED_FUNCTION(_wsplitpath_s) \
 	DETOURED_FUNCTION(_wmakepath_s) \
 	//DETOURED_FUNCTION(_wsopen_s) \
 	//DETOURED_FUNCTION(_fileno) \
@@ -323,7 +325,7 @@
 #define DETOURED_FUNCTIONS_KERNEL32_DEBUG
 #define DETOURED_FUNCTIONS_NTDLL_DEBUG
 #define DETOURED_FUNCTIONS_SHLWAPI_DEBUG
-#define DETOURED_FUNCTIONS_CRTBASE_DEBUG
+#define DETOURED_FUNCTIONS_UCRTBASE_DEBUG
 #define DETOURED_FUNCTIONS_MEMORY_DEBUG
 #endif
 
