UnrealBuildAccelerator: Fix Linux build

#rnx

diff --git a/Engine/Source/Programs/UnrealBuildAccelerator/Core/Private/UbaApplicationRules.cpp b/Engine/Source/Programs/UnrealBuildAccelerator/Core/Private/UbaApplicationRules.cpp
--- a/Engine/Source/Programs/UnrealBuildAccelerator/Core/Private/UbaApplicationRules.cpp
+++ b/Engine/Source/Programs/UnrealBuildAccelerator/Core/Private/UbaApplicationRules.cpp
@@ -344,15 +344,17 @@ namespace uba
 				|| file.EndsWith(TC(".stacktrace.txt"));
 		}
 
-		virtual void RepairMalformedLibPath(const wchar_t* path) override
+#if PLATFORM_WINDOWS
+		virtual void RepairMalformedLibPath(const tchar* path) override
 		{
 			// There is a bug where the path passed into wsplitpath_s is malformed and not null terminated correctly
-			const wchar_t* pext = TStrstr(path, L".dll");
-			if (pext == nullptr) pext = TStrstr(path, L".DLL");
-			if (pext == nullptr) pext = TStrstr(path, L".exe");
-			if (pext == nullptr) pext = TStrstr(path, L".EXE");
-			if (pext != nullptr && *(pext + 4) != 0) *(const_cast<wchar_t*>(pext + 4)) = 0;
+			const tchar* pext = TStrstr(path, TC(".dll"));
+			if (pext == nullptr) pext = TStrstr(path, TC(".DLL"));
+			if (pext == nullptr) pext = TStrstr(path, TC(".exe"));
+			if (pext == nullptr) pext = TStrstr(path, TC(".EXE"));
+			if (pext != nullptr && *(pext + 4) != 0) *(const_cast<tchar*>(pext + 4)) = 0;
 		}
+#endif // #if PLATFORM_WINDOWS
 	};
 
 	const RulesRec* GetApplicationRules()
diff --git a/Engine/Source/Programs/UnrealBuildAccelerator/Core/Public/UbaApplicationRules.h b/Engine/Source/Programs/UnrealBuildAccelerator/Core/Public/UbaApplicationRules.h
--- a/Engine/Source/Programs/UnrealBuildAccelerator/Core/Public/UbaApplicationRules.h
+++ b/Engine/Source/Programs/UnrealBuildAccelerator/Core/Public/UbaApplicationRules.h
@@ -114,7 +114,7 @@ namespace uba
 			return true;
 		}
 
-		virtual void RepairMalformedLibPath(const wchar_t* path)
+		virtual void RepairMalformedLibPath(const tchar* path)
 		{
 			// Do nothing
 		}
