[UBA - Visualizer]
* Fixed bug in listen flow so it sometimes didn't pick up new traces.
* Fixed annoying messagebox that popped up when starting listen session sometimes

diff --git a/Engine/Source/Programs/UnrealBuildAccelerator/Visualizer/Private/UbaTraceReader.cpp b/Engine/Source/Programs/UnrealBuildAccelerator/Visualizer/Private/UbaTraceReader.cpp
--- a/Engine/Source/Programs/UnrealBuildAccelerator/Visualizer/Private/UbaTraceReader.cpp
+++ b/Engine/Source/Programs/UnrealBuildAccelerator/Visualizer/Private/UbaTraceReader.cpp
@@ -176,12 +176,13 @@ namespace uba
 		return ReadMemory(out, false);
 	}
 
-	bool TraceReader::StartReadNamed(TraceView& out, const tchar* namedTrace)
+	bool TraceReader::StartReadNamed(TraceView& out, const tchar* namedTrace, bool silentFail)
 	{
 		m_memoryHandle.handle = ::OpenFileMappingW(PAGE_READWRITE, false, namedTrace);
 		if (!m_memoryHandle.IsValid())
 		{
-			m_logger.Error(L"OpenFileMappingW - Failed to open file mapping %ls (%ls)", namedTrace, LastErrorToText().data);
+			if (!silentFail)
+				m_logger.Error(L"OpenFileMappingW - Failed to open file mapping %ls (%ls)", namedTrace, LastErrorToText().data);
 			return false;
 		}
 		m_memoryBegin = MapViewOfFile(m_memoryHandle, FILE_MAP_READ, 0, 0);
@@ -189,7 +190,7 @@ namespace uba
 			return false;
 		m_memoryPos = m_memoryBegin;
 		m_memoryEnd = m_memoryBegin;
-
+		m_hostProcess = 0;
 		out.finished = false;
 		out.sessions.emplace_back();
 		out.sessions.back().name = L"LOCAL";
@@ -205,7 +206,7 @@ namespace uba
 		m_memoryEnd = m_memoryBegin + *(u32*)m_memoryBegin;
 		outChanged = !m_activeProcesses.empty() || m_memoryPos != m_memoryEnd;
 		bool res = ReadMemory(out, true);
-		if (WaitForSingleObject(m_hostProcess, 0) != WAIT_TIMEOUT)
+		if (m_hostProcess  && WaitForSingleObject(m_hostProcess, 0) != WAIT_TIMEOUT)
 			StopAllActive(out, GetTime() - m_startTime);
 		return res;
 	}
diff --git a/Engine/Source/Programs/UnrealBuildAccelerator/Visualizer/Private/UbaTraceReader.h b/Engine/Source/Programs/UnrealBuildAccelerator/Visualizer/Private/UbaTraceReader.h
--- a/Engine/Source/Programs/UnrealBuildAccelerator/Visualizer/Private/UbaTraceReader.h
+++ b/Engine/Source/Programs/UnrealBuildAccelerator/Visualizer/Private/UbaTraceReader.h
@@ -109,7 +109,7 @@ namespace uba
 		};
 
 		Process* GetProcess(const ProcessLocation& loc) { return &(sessions[loc.sessionIndex].processors[loc.processorIndex].processes[loc.processIndex]); }
-		void Clear() { sessions.clear(); workTracks.clear(); strings.clear(); startTime = 0; finished = false; };
+		void Clear() { sessions.clear(); workTracks.clear(); strings.clear(); startTime = 0; finished = true; };
 
 		Vector<Session> sessions;
 		Vector<WorkTrack> workTracks;
@@ -138,7 +138,7 @@ namespace uba
 		bool UpdateReadClient(TraceView& out, NetworkClient& client, bool& outChanged);
 
 		// Use for local
-		bool StartReadNamed(TraceView& out, const tchar* namedTrace);
+		bool StartReadNamed(TraceView& out, const tchar* namedTrace, bool silentFail = false);
 		bool UpdateReadNamed(TraceView& out, bool& outChanged);
 
 		bool ReadMemory(TraceView& out, bool trackHost);
diff --git a/Engine/Source/Programs/UnrealBuildAccelerator/Visualizer/Private/UbaVisualizer.cpp b/Engine/Source/Programs/UnrealBuildAccelerator/Visualizer/Private/UbaVisualizer.cpp
--- a/Engine/Source/Programs/UnrealBuildAccelerator/Visualizer/Private/UbaVisualizer.cpp
+++ b/Engine/Source/Programs/UnrealBuildAccelerator/Visualizer/Private/UbaVisualizer.cpp
@@ -370,8 +370,6 @@ namespace uba
 
 		m_startTime = GetTime();
 
-		SetTimer(m_hwnd, 0, 200, NULL);
-
 		while (m_looping)
 		{
 			MSG msg;
@@ -1766,14 +1764,17 @@ namespace uba
 			}
 			else
 			{
-				if (!m_trace.StartReadNamed(m_traceView, m_newTraceName.data))
+				if (!m_trace.StartReadNamed(m_traceView, m_newTraceName.data, true))
 					return false;
 				m_namedTrace.Clear().Append(m_newTraceName);
 				title.Appendf(L"%s (Listening for new sessions on channel '%s')", m_namedTrace.data, m_listenChannel.data);
 			}
+
+			m_traceView.finished = false;
 			SetWindowTextW(m_hwnd, title.data);
 			SendMessage(m_hwnd, WM_TIMER, 0, 0);
 			UpdateScrollbars(true);
+			SetTimer(m_hwnd, 0, 200, NULL);
 			return 0;
 		}
 
@@ -1847,7 +1848,8 @@ namespace uba
 				RedrawWindow(hWnd, NULL, NULL, RDW_INVALIDATE|RDW_UPDATENOW);
 				u64 paintTimeMs = TimeToMs(GetTime() - startTime);
 				u32 waitTime = u32(Min(paintTimeMs * 5, 200ull));
-				SetTimer(m_hwnd, 0, waitTime, NULL);
+				if (!m_traceView.finished)
+					SetTimer(m_hwnd, 0, waitTime, NULL);
 			}
 			break;
 		}
