[UBA]
* Fixed a bug where uba failed to download dlls to remote if it was in the system32 folder but did not exist on the remote machine's system32 and was not part of known system modules. (this solves the problem with ShaderCompilerWorker not downloading opengl32.dll to barebone machines)
* Fixed bug where relative path of dll was sent into RetriveCasFile which in turn failed because it couldn't find the file
* Fixed so code could handle loading libraries with circular dependencies
* Removed assert for now related to GetThreadPreferredUILanguages and asking for language name instead of id for remote machines)

diff --git a/Engine/Source/Programs/UnrealBuildAccelerator/Common/Private/UbaSession.cpp b/Engine/Source/Programs/UnrealBuildAccelerator/Common/Private/UbaSession.cpp
--- a/Engine/Source/Programs/UnrealBuildAccelerator/Common/Private/UbaSession.cpp
+++ b/Engine/Source/Programs/UnrealBuildAccelerator/Common/Private/UbaSession.cpp
@@ -894,61 +894,8 @@ namespace uba
 		FixPath(applicationName, nullptr, 0, temp3);
 
 		bool isSystem = StartsWith(applicationName, m_systemPath.data);
-
-		if (isSystem)
-		{
-			// These dlls should _always_ exist on all machines so we can filter them out from the list of imports to copy
-
-			static const wchar_t* knownSystemFiles[] = // Sorted lowercase
-			{
-				L"advapi32.dll",
-				L"bcrypt.dll",
-				L"bcryptprimitives.dll",
-				L"combase.dll",
-				L"cryptbase.dll",
-				L"dbghelp.dll",
-				L"dnsapi.dll",
-				L"fwpuclnt.dll",
-				L"gdi32.dll",
-				L"gdi32full.dll",
-				L"imm32.dll",
-				L"iphlpapi.dll",
-				L"kernel32.dll",
-				L"kernelbase.dll",
-				L"msvcp_win.dll",
-				L"msvcrt.dll",
-				L"mswsock.dll",
-				L"ncrypt.dll",
-				L"nsi.dll",
-				L"ntasn1.dll",
-				L"ntdll.dll",
-				L"ole32.dll",
-				L"oleaut32.dll",
-				L"ondemandconnroutehelper.dll",
-				L"powrprof.dll",
-				L"rasadhlp.dll",
-				L"rpcrt4.dll",
-				L"rstrtmgr.dll",
-				L"sechost.dll",
-				L"sspicli.dll",
-				L"ucrtbase.dll",
-				L"umpdc.dll",
-				L"umppc17706.dll",
-				L"user32.dll",
-				L"version.dll",
-				L"webio.dll",
-				L"win32u.dll",
-				L"winhttp.dll",
-				L"winnsi.dll",
-				L"ws2_32.dll",
-			};
-
-			auto lastSeparator = TStrrchr(applicationName, PathSeparator);
-			StringBuffer<> nameLower;
-			nameLower.Append(lastSeparator + 1).MakeLower();
-			if (std::binary_search(knownSystemFiles, knownSystemFiles + sizeof_array(knownSystemFiles), nameLower.data, [](const wchar_t* a, const wchar_t* b) { return TStrcmp(a, b) < 0; }))
-				return true;
-		}
+		if (isSystem && IsKnownSystemFile(applicationName))
+			return true;
 
 
 		out.push_back({ library, temp3.data, attr, isSystem });
@@ -1501,6 +1448,64 @@ namespace uba
 		return GetApplicationRules()[process.m_rulesIndex].rules->IsRarelyReadAfterWritten(fileName, fileNameLen);
 	}
 
+	bool Session::IsKnownSystemFile(const tchar* applicationName)
+	{
+#if PLATFORM_WINDOWS
+		// These dlls should _always_ exist on all machines so we can filter them out from the list of imports to copy
+
+		static const wchar_t* knownSystemFiles[] = // Sorted lowercase
+		{
+			L"advapi32.dll",
+			L"bcrypt.dll",
+			L"bcryptprimitives.dll",
+			L"combase.dll",
+			L"cryptbase.dll",
+			L"dbghelp.dll",
+			L"dnsapi.dll",
+			L"fwpuclnt.dll",
+			L"gdi32.dll",
+			L"gdi32full.dll",
+			L"imm32.dll",
+			L"iphlpapi.dll",
+			L"kernel32.dll",
+			L"kernelbase.dll",
+			L"msvcp_win.dll",
+			L"msvcrt.dll",
+			L"mswsock.dll",
+			L"ncrypt.dll",
+			L"nsi.dll",
+			L"ntasn1.dll",
+			L"ntdll.dll",
+			L"ole32.dll",
+			L"oleaut32.dll",
+			L"ondemandconnroutehelper.dll",
+			L"powrprof.dll",
+			L"rasadhlp.dll",
+			L"rpcrt4.dll",
+			L"rstrtmgr.dll",
+			L"sechost.dll",
+			L"sspicli.dll",
+			L"ucrtbase.dll",
+			L"umpdc.dll",
+			L"umppc17706.dll",
+			L"user32.dll",
+			L"version.dll",
+			L"webio.dll",
+			L"win32u.dll",
+			L"winhttp.dll",
+			L"winnsi.dll",
+			L"ws2_32.dll",
+		};
+
+		auto lastSeparator = TStrrchr(applicationName, PathSeparator);
+		StringBuffer<> nameLower;
+		nameLower.Append(lastSeparator + 1).MakeLower();
+		return std::binary_search(knownSystemFiles, knownSystemFiles + sizeof_array(knownSystemFiles), nameLower.data, [](const wchar_t* a, const wchar_t* b) { return TStrcmp(a, b) < 0; });
+#else
+		return false;
+#endif
+	}
+
 	bool Session::ShouldWriteToDisk(const tchar* fileName, u64 fileNameLen)
 	{
 		if (m_shouldWriteToDisk)
diff --git a/Engine/Source/Programs/UnrealBuildAccelerator/Common/Private/UbaSessionClient.cpp b/Engine/Source/Programs/UnrealBuildAccelerator/Common/Private/UbaSessionClient.cpp
--- a/Engine/Source/Programs/UnrealBuildAccelerator/Common/Private/UbaSessionClient.cpp
+++ b/Engine/Source/Programs/UnrealBuildAccelerator/Common/Private/UbaSessionClient.cpp
@@ -147,8 +147,10 @@ namespace uba
 		{
 			if (!GetCasKeyForFile(casKey, processId, fileName, fileNameKey))
 				return false;
-			const tchar* lastSlash = TStrrchr(fileName.data, PathSeparator);
-			outVirtual.Append(lastSlash + 1);
+			// This needs to be absolute virtual path (the path on the host).. don't remember why this was written like this. Changed to just use fileName (needed for shadercompilerworker)
+			//const tchar* lastSlash = TStrrchr(fileName.data, PathSeparator);
+			//outVirtual.Append(lastSlash + 1);
+			outVirtual.Append(fileName);
 		}
 		else
 		{
@@ -242,6 +244,7 @@ namespace uba
 					localSystemModule.Append(m_systemPath).Append(moduleFile.data + serverSystemPathLen);
 					if (FileExists(m_logger, localSystemModule.data))
 						continue;
+					moduleFile.Clear().Append(localSystemModule);
 				}
 
 				CasKey newCasKey;
diff --git a/Engine/Source/Programs/UnrealBuildAccelerator/Common/Private/UbaSessionServer.cpp b/Engine/Source/Programs/UnrealBuildAccelerator/Common/Private/UbaSessionServer.cpp
--- a/Engine/Source/Programs/UnrealBuildAccelerator/Common/Private/UbaSessionServer.cpp
+++ b/Engine/Source/Programs/UnrealBuildAccelerator/Common/Private/UbaSessionServer.cpp
@@ -667,7 +667,7 @@ namespace uba
 				CasKey casKey = CasKeyZero;
 				StringBuffer<> absoluteFile;
 				if (SearchPathForFile(m_logger, absoluteFile, fileName.data, applicationDir.data))
-					if (!absoluteFile.StartsWith(m_systemPath.data))
+					if (!absoluteFile.StartsWith(m_systemPath.data) || !IsKnownSystemFile(absoluteFile.data))
 						if (!StoreCasFile(casKey, fileNameKey, absoluteFile.data))
 							return false;
 
diff --git a/Engine/Source/Programs/UnrealBuildAccelerator/Common/Private/UbaStorage.cpp b/Engine/Source/Programs/UnrealBuildAccelerator/Common/Private/UbaStorage.cpp
--- a/Engine/Source/Programs/UnrealBuildAccelerator/Common/Private/UbaStorage.cpp
+++ b/Engine/Source/Programs/UnrealBuildAccelerator/Common/Private/UbaStorage.cpp
@@ -721,7 +721,7 @@ namespace uba
 
 	bool StorageImpl::AddCasFile(const tchar* fileName, const CasKey& casKey, bool deferCreation)
 	{
-		UBA_ASSERT(IsCompressed(casKey) == m_storeCompressed);
+		UBA_ASSERTF(IsCompressed(casKey) == m_storeCompressed, TC("CasKey compress mode must match storage compress mode (%s)"), fileName);
 		ScopedWriteLock lookupLock(m_casLookupLock);
 		auto insres = m_casLookup.try_emplace(casKey);
 		CasEntry& casEntry = insres.first->second;
@@ -1550,9 +1550,12 @@ namespace uba
 	
 		if (fileEntry.verified)
 		{
-			UBA_ASSERT(casKeyOverride == CasKeyZero || casKeyOverride == fileEntry.casKey);
-			if (!AddCasFile(fileName, fileEntry.casKey, deferCreation))
-				return false;
+			if (fileEntry.casKey != CasKeyZero)
+			{
+				UBA_ASSERT(casKeyOverride == CasKeyZero || casKeyOverride == fileEntry.casKey);
+				if (!AddCasFile(fileName, fileEntry.casKey, deferCreation))
+					return false;
+			}
 			out = fileEntry.casKey;
 			return true;
 		}
diff --git a/Engine/Source/Programs/UnrealBuildAccelerator/Common/Public/UbaSession.h b/Engine/Source/Programs/UnrealBuildAccelerator/Common/Public/UbaSession.h
--- a/Engine/Source/Programs/UnrealBuildAccelerator/Common/Public/UbaSession.h
+++ b/Engine/Source/Programs/UnrealBuildAccelerator/Common/Public/UbaSession.h
@@ -129,6 +129,7 @@ namespace uba
 		void Free(Vector<BinaryModule>& v);
 		bool IsRarelyRead(ProcessImpl& process, const StringBufferBase& fileName) const;
 		bool IsRarelyReadAfterWritten(ProcessImpl& process, const tchar* fileName, u64 fileNameLen) const;
+		bool IsKnownSystemFile(const tchar* applicationName);
 		bool ShouldWriteToDisk(const tchar* fileName, u64 fileNameLen);
 		u32 WriteDirectoryEntries(const StringKey& dirKey, tchar* dirPath, u32& outTableOffset);
 		u32 AddFileMapping(StringKey fileNameKey, const tchar* fileName, const tchar* newFileName, u64 fileSize = InvalidValue);
diff --git a/Engine/Source/Programs/UnrealBuildAccelerator/Detours/Private/Windows/UbaDetoursFunctionsKernelBase.inl b/Engine/Source/Programs/UnrealBuildAccelerator/Detours/Private/Windows/UbaDetoursFunctionsKernelBase.inl
--- a/Engine/Source/Programs/UnrealBuildAccelerator/Detours/Private/Windows/UbaDetoursFunctionsKernelBase.inl
+++ b/Engine/Source/Programs/UnrealBuildAccelerator/Detours/Private/Windows/UbaDetoursFunctionsKernelBase.inl
@@ -2264,9 +2264,13 @@ DWORD Detoured_SearchPathW(LPCWSTR lpPath, LPCWSTR lpFileName, LPCWSTR lpExtensi
 }
 
 using AdditionalLoads = Vector<HMODULE, GrowingAllocator<HMODULE>>;
+using VisitedModules = std::unordered_set<StringKey, std::hash<StringKey>, std::equal_to<StringKey>, GrowingAllocator<StringKey>>;
 
-HMODULE Recursive_LoadLibraryExW(LPCWSTR lpLibFileName, LPCWSTR originalName, DWORD dwFlags, AdditionalLoads& additionalLoads)
+HMODULE Recursive_LoadLibraryExW(LPCWSTR lpLibFileName, LPCWSTR originalName, DWORD dwFlags, AdditionalLoads& additionalLoads, VisitedModules& visitedModules)
 {
+	if (!visitedModules.insert(ToStringKeyNoCheck(lpLibFileName, wcslen(lpLibFileName))).second)
+		return 0;
+
 	// Important that this code is not doing allocations.. it could cause a recursive stack overflow
 	struct Import { char name[128]; Import(const char* s) { strcpy_s(name, sizeof_array(name), s); } };
 	std::vector<Import, GrowingAllocator<Import>> importedModules(&g_memoryBlock);
@@ -2297,7 +2301,8 @@ HMODULE Recursive_LoadLibraryExW(LPCWSTR lpLibFileName, LPCWSTR originalName, DW
 		StringBuffer<512> tempBuf;
 		Rpc_GetFullFileName(path, pathLen, tempBuf, false);
 
-		additionalLoads.push_back(Recursive_LoadLibraryExW(path, moduleNameW, dwFlags, additionalLoads));
+		if (HMODULE r = Recursive_LoadLibraryExW(path, moduleNameW, dwFlags, additionalLoads, visitedModules))
+			additionalLoads.push_back(r);
 	}
 
 	StringBuffer<512> newName;
@@ -2358,7 +2363,8 @@ HMODULE Shared_LoadLibrary(LPCWSTR lpLibFileName, HANDLE hFile, DWORD dwFlags)
 		Rpc_GetFullFileName(newPath, newPathLen, tempBuf, false);
 
 	AdditionalLoads additionalLoads(&g_memoryBlock); // Don't do allocations
-	HMODULE res = Recursive_LoadLibraryExW(newPath, fileName, dwFlags, additionalLoads);
+	VisitedModules visitedModules(&g_memoryBlock);
+	HMODULE res = Recursive_LoadLibraryExW(newPath, fileName, dwFlags, additionalLoads, visitedModules);
 	for (HMODULE h : additionalLoads)
 		FreeLibrary(h);
 	return res;
@@ -2792,26 +2798,28 @@ BOOL Detoured_GetThreadPreferredUILanguages(DWORD dwFlags, PULONG pulNumLanguage
 {
 	DETOURED_CALL(GetThreadPreferredUILanguages);
 
-	if (!(dwFlags & MUI_LANGUAGE_ID))
+	if (dwFlags & MUI_LANGUAGE_ID)
 	{
-		UBA_ASSERT(!g_runningRemote);
+		DEBUG_LOG_DETOURED(L"GetThreadPreferredUILanguages", L"");
+		//UBA_ASSERT(!(dwFlags & ~MUI_LANGUAGE_ID);
+		*pulNumLanguages = 1;
+		*pcchLanguagesBuffer = 6;
+
+		if (!pwszLanguagesBuffer)
+			return TRUE;
+
+		swprintf_s(pwszLanguagesBuffer, 6, L"%04x", g_uiLanguage);
+		pwszLanguagesBuffer[5] = 0;
+		return TRUE;
+	}
+	else // MUI_LANGUAGE_NAME
+	{
+		// TODO: We need to get the string of g_uiLanguage
+		//UBA_ASSERTF(!g_runningRemote, L"GetThreadPreferredUILanguages uses unsupported flag on remote execution: %u", dwFlags);
 		auto res = True_GetThreadPreferredUILanguages(dwFlags, pulNumLanguages, pwszLanguagesBuffer, pcchLanguagesBuffer);
 		DEBUG_LOG_TRUE(L"GetThreadPreferredUILanguages", L"-> %ls", ToString(res));
 		return res;
 	}
-
-	DEBUG_LOG_DETOURED(L"GetThreadPreferredUILanguages", L"");
-
-	UBA_ASSERT(dwFlags & MUI_LANGUAGE_ID);
-	*pulNumLanguages = 1;
-	*pcchLanguagesBuffer = 6;
-
-	if (!pwszLanguagesBuffer)
-		return TRUE;
-
-	swprintf_s(pwszLanguagesBuffer, 6, L"%04x", g_uiLanguage);
-	pwszLanguagesBuffer[5] = 0;
-	return TRUE;
 }
 
 #if defined(DETOURED_INCLUDE_DEBUG)
