[UBA]
* Fixed linux compile error

diff --git a/Engine/Source/Programs/UnrealBuildAccelerator/Common/Private/UbaProcess.cpp b/Engine/Source/Programs/UnrealBuildAccelerator/Common/Private/UbaProcess.cpp
--- a/Engine/Source/Programs/UnrealBuildAccelerator/Common/Private/UbaProcess.cpp
+++ b/Engine/Source/Programs/UnrealBuildAccelerator/Common/Private/UbaProcess.cpp
@@ -23,6 +23,7 @@
 // Mac has to roll it's own solution
 #if PLATFORM_LINUX
 #include <sys/prctl.h>
+#include <sys/resource.h>
 #elif PLATFORM_MAC
 #include <sys/types.h>
 #include <sys/sysctl.h>
