[UBA]
* Added Scheduler class that is a very simple scheduler that handles processes that has no intra dependencies. It handles scheduling remotely and also reschedule processes returned from remote machines.
* Changed so application dependencies are retrieved in parallel by clients
* Added so custom assert handler can be set from the outside
* Added traceEnabled to SessionCreateInfo so trace shared mem can be created without needing to launch visualizer or write to file
* Added so userData can be provided in RemoteProcessAvailable and RemoteProcessReturned callbacks
* Added more files to known system files (based on what exists in wine)
* Improved FindImports code and made it available in export. (it is now automatically filtering out known system files)
* Moved ProcessHandle to its own file
* Added ProcessStartInfoHolder which is a class that can wrap a ProcessStartInfo and make sure all strings are allocated

diff --git a/Engine/Source/Programs/UnrealBuildAccelerator/Common/Private/UbaLogger.cpp b/Engine/Source/Programs/UnrealBuildAccelerator/Common/Private/UbaLogger.cpp
--- a/Engine/Source/Programs/UnrealBuildAccelerator/Common/Private/UbaLogger.cpp
+++ b/Engine/Source/Programs/UnrealBuildAccelerator/Common/Private/UbaLogger.cpp
@@ -17,6 +17,8 @@
 
 namespace uba
 {
+	CustomAssertHandler* g_assertHandler;
+
 	ANALYSIS_NORETURN void UbaAssert(const tchar* text, const char* file, u32 line, const char* expr, u32 terminateCode)
 	{
 		static ReaderWriterLock assertLock;
@@ -24,6 +26,13 @@ namespace uba
 
 		StringBuffer<4096> b;
 		WriteAssertInfo(b, text, file, line, expr, 1);
+
+		if (g_assertHandler)
+		{
+			g_assertHandler(b.data);
+			return;
+		}
+
 		Fputs(b.data, stdout);
 		Fputs(TC("\n"), stdout);
 		fflush(stdout);
@@ -47,6 +56,11 @@ namespace uba
 #endif
 	}
 
+	void SetCustomAssertHandler(CustomAssertHandler* handler)
+	{
+		g_assertHandler = handler;
+	}
+
 	ANALYSIS_NORETURN void FatalError(u32 code, const tchar* format, ...)
 	{
 		va_list arg;
diff --git a/Engine/Source/Programs/UnrealBuildAccelerator/Common/Private/UbaProcessStartInfoHolder.h b/Engine/Source/Programs/UnrealBuildAccelerator/Common/Private/UbaProcessStartInfoHolder.h
new file mode 100644
--- /dev/null
+++ b/Engine/Source/Programs/UnrealBuildAccelerator/Common/Private/UbaProcessStartInfoHolder.h
@@ -0,0 +1,88 @@
+// Copyright Epic Games, Inc. All Rights Reserved.
+
+#pragma once
+
+#define Local_GetLongPathNameW uba::GetLongPathNameW
+
+#include "UbaFile.h"
+#include "UbaPathUtils.h"
+
+namespace uba
+{
+	inline void FixFileName(StringBufferBase& out, const tchar* fileName, const tchar* workingDir)
+	{
+		tchar buffer[1024];
+		u32 charLen;
+		u64 workingDirLen = 0;
+		if (workingDir)
+			workingDirLen = TStrlen(workingDir);
+		FixPath2(fileName, workingDir, workingDirLen, buffer, &charLen);
+		out.Append(buffer);
+	}
+
+	struct ProcessStartInfoHolder
+	{
+		ProcessStartInfoHolder(const ProcessStartInfo& si)
+		{
+			startInfo = si;
+
+			StringBuffer<512> temp;
+			FixFileName(temp, si.workingDir, nullptr);
+			temp.EnsureEndsWithSlash();
+			workingDir = temp.data;
+			startInfo.workingDir = workingDir.c_str();
+
+			temp.EnsureEndsWithSlash();
+			StringBuffer<512> temp2;
+			FixFileName(temp2, si.application, temp.data);
+			application = temp2.data;
+			startInfo.application = application.c_str();
+
+			arguments = si.arguments;
+			startInfo.arguments = arguments.c_str();
+
+			description = si.description;
+			startInfo.description = description.c_str();
+
+			if (si.logFile)
+			{
+				logFile = si.logFile;
+				startInfo.logFile = logFile.c_str();
+			}
+		}
+
+		ProcessStartInfoHolder(const ProcessStartInfoHolder& o)
+		{
+			*this = o;
+		}
+
+		void operator=(const ProcessStartInfoHolder& o)
+		{
+			startInfo = o.startInfo;
+
+			workingDir = o.workingDir;
+			startInfo.workingDir = workingDir.c_str();
+			application = o.application;
+			startInfo.application = application.c_str();
+			arguments = o.arguments;
+			startInfo.arguments = arguments.c_str();
+			description = o.description;
+			startInfo.description = description.c_str();
+
+			if (o.startInfo.logFile)
+			{
+				logFile = o.logFile;
+				startInfo.logFile = logFile.c_str();
+			}
+		}
+
+		ProcessStartInfo startInfo;
+
+		TString description;
+		TString application;
+		TString arguments;
+		TString workingDir;
+		TString logFile;
+	};
+
+}
\ No newline at end of file
diff --git a/Engine/Source/Programs/UnrealBuildAccelerator/Common/Private/UbaScheduler.cpp b/Engine/Source/Programs/UnrealBuildAccelerator/Common/Private/UbaScheduler.cpp
new file mode 100644
--- /dev/null
+++ b/Engine/Source/Programs/UnrealBuildAccelerator/Common/Private/UbaScheduler.cpp
@@ -0,0 +1,181 @@
+// Copyright Epic Games, Inc. All Rights Reserved.
+
+#include "UbaScheduler.h"
+#include "UbaProcess.h"
+#include "UbaProcessStartInfoHolder.h"
+#include "UbaSessionServer.h"
+#include "UbaStringBuffer.h"
+
+namespace uba
+{
+	struct Scheduler::QueuedProcess : ProcessStartInfoHolder
+	{
+		QueuedProcess(const ProcessStartInfo& si, const u8* ki, u32 kic)
+		:	ProcessStartInfoHolder(si)
+		, knownInputs(ki)
+		, knownInputsCount(kic)
+		{
+		};
+
+		const u8* knownInputs;
+		u32 knownInputsCount;
+	};
+
+	struct Scheduler::ExitProcessInfo
+	{
+		Scheduler* scheduler = nullptr;
+		ProcessStartInfo::ExitedCallback* originalExitedFunc = nullptr;
+		void* originalExitedData = nullptr;
+		const u8* knownInputs = nullptr;
+		u32 knownInputsCount = 0;
+		bool wasReturned = false;
+		bool isRemote = false;
+	};
+
+	Scheduler::Scheduler(SessionServer& session, u32 maxLocalProcessors)
+	:	m_session(session)
+	,	m_maxLocalProcessors(maxLocalProcessors ? maxLocalProcessors : GetLogicalProcessorCount())
+	,	m_updateThreadLoop(false)
+	{
+	}
+
+	Scheduler::~Scheduler()
+	{
+		Stop();
+	}
+
+	void Scheduler::Start()
+	{
+		m_loop = true;
+		m_thread.Start([this]() { ThreadLoop(); return 0; });
+
+		m_session.SetRemoteProcessReturnedEvent([this](Process& process) { RemoteProcessReturned(process); });
+		m_session.SetRemoteProcessSlotAvailableEvent([this]() { RemoteSlotAvailable(); });
+	}
+
+	void Scheduler::Stop()
+	{
+		m_loop = false;
+		m_updateThreadLoop.Set();
+		m_thread.Wait();
+		m_session.WaitOnAllTasks();
+	}
+
+	u32 Scheduler::EnqueueProcess(const ProcessStartInfo& info, float weight, const void* knownInputs, u32 knownInputsBytes, u32 knownInputsCount)
+	{
+		UBA_ASSERTF(weight == 1.0f, TC("weight not implemented"));
+
+		if (knownInputsCount)
+		{
+			void* ki = new u8[knownInputsBytes];
+			memcpy(ki, knownInputs, knownInputsBytes);
+			knownInputs = ki;
+		}
+		else
+			knownInputs = nullptr;
+
+		ScopedWriteLock lock(m_queuedProcessesLock);
+		m_queuedProcesses.emplace_back(info, (const u8*)knownInputs, knownInputsCount);
+		m_updateThreadLoop.Set();
+		return 0;
+	}
+
+	void Scheduler::GetStats(u32& outQueued, u32& outActive, u32& outFinished)
+	{
+		outActive = m_activeLocalProcesses + m_activeRemoteProcesses;
+		outFinished = m_finishedProcesses;
+		ScopedReadLock lock(m_queuedProcessesLock);
+		outQueued = u32(m_queuedProcesses.size());
+	}
+
+	void Scheduler::ThreadLoop()
+	{
+		while (m_loop)
+		{
+			if (!m_updateThreadLoop.IsSet())
+				break;
+
+			while (m_activeLocalProcesses < m_maxLocalProcessors)
+				if (!RunQueuedProcess(true))
+					break;
+		}
+	}
+
+	void Scheduler::RemoteProcessReturned(Process& process)
+	{
+		auto& ei = *(ExitProcessInfo*)process.GetStartInfo().exitedUserData;
+
+		ScopedWriteLock lock(m_queuedProcessesLock);
+		ProcessStartInfo info = process.GetStartInfo();
+		info.exitedFunc = ei.originalExitedFunc;
+		info.exitedUserData = ei.originalExitedData;
+		m_queuedProcesses.emplace_front(info, ei.knownInputs, ei.knownInputsCount);
+		ei.knownInputs = nullptr;
+		lock.Leave();
+
+		ei.wasReturned = true;
+		process.Cancel(true);
+		--m_activeRemoteProcesses;
+		m_updateThreadLoop.Set();
+	}
+
+	void Scheduler::RemoteSlotAvailable()
+	{
+		RunQueuedProcess(false);
+	}
+
+	void Scheduler::ProcessExited(ExitProcessInfo* info, const ProcessHandle& handle)
+	{
+		auto ig = MakeGuard([info]() { delete[] info->knownInputs;  delete info; });
+
+		if (info->wasReturned)
+			return;
+
+		if (info->originalExitedFunc)
+			info->originalExitedFunc(info->originalExitedData, handle);
+
+		++m_finishedProcesses;
+
+		if (info->isRemote)
+		{
+			--m_activeRemoteProcesses;
+		}
+		else
+		{
+			--m_activeLocalProcesses;
+		}
+		m_updateThreadLoop.Set();
+	}
+
+	bool Scheduler::RunQueuedProcess(bool runLocal)
+	{
+		ScopedWriteLock lock(m_queuedProcessesLock);
+		if (m_queuedProcesses.empty())
+			return false;
+		auto si = m_queuedProcesses.front();
+		m_queuedProcesses.pop_front();
+		lock.Leave();
+
+		auto exitInfo = new ExitProcessInfo();
+		exitInfo->scheduler = this;
+		exitInfo->originalExitedFunc = si.startInfo.exitedFunc;
+		exitInfo->originalExitedData = si.startInfo.exitedUserData;
+		exitInfo->knownInputs = si.knownInputs;
+		exitInfo->knownInputsCount = si.knownInputsCount;
+		exitInfo->isRemote = !runLocal;
+		si.startInfo.exitedUserData = exitInfo;
+		si.startInfo.exitedFunc = [](void* userData, const ProcessHandle& handle) { auto ei = (ExitProcessInfo*)userData; ei->scheduler->ProcessExited(ei, handle); };
+
+		if (runLocal)
+		{
+			++m_activeLocalProcesses;
+			m_session.RunProcess(si.startInfo);
+		}
+		else
+		{
+			++m_activeRemoteProcesses;
+			m_session.RunProcessRemote(si.startInfo, 1.0f, si.knownInputs, si.knownInputsCount);
+		}
+		return true;
+	}
+}
diff --git a/Engine/Source/Programs/UnrealBuildAccelerator/Common/Private/UbaSession.cpp b/Engine/Source/Programs/UnrealBuildAccelerator/Common/Private/UbaSession.cpp
--- a/Engine/Source/Programs/UnrealBuildAccelerator/Common/Private/UbaSession.cpp
+++ b/Engine/Source/Programs/UnrealBuildAccelerator/Common/Private/UbaSession.cpp
@@ -10,11 +10,6 @@
 #include "UbaProtocol.h"
 
 #if PLATFORM_WINDOWS
-#include <algorithm>
-#define True_CreateFileW ::CreateFileW
-#define True_CreateFileMappingW ::CreateFileMappingW
-#define True_MapViewOfFile ::MapViewOfFile
-#define True_UnmapViewOfFile ::UnmapViewOfFile
 #include "UbaWinBinDependencyParser.h"
 #include <powerbase.h>
 #pragma comment(lib, "Powrprof.lib")
@@ -901,13 +896,10 @@ namespace uba
 		out.push_back({ library, temp3.data, attr, isSystem });
 
 		bool result = true;
-		FindImports(applicationName, [&](const char* import)
+		FindImports(applicationName, [&](const tchar* importName)
 			{
-				if (!result)
-					return;
-				tchar importW[128];
-				swprintf_s(importW, 128, TC("%hs"), import);
-				result = CopyImports(out, importW, applicationDir, applicationDirEnd, handledImports);
+				if (result)
+					result = CopyImports(out, importName, applicationDir, applicationDirEnd, handledImports);
 			});
 		return result;
 		#else
@@ -1031,7 +1023,7 @@ namespace uba
 		StringBuffer<> traceName;
 		if (info.traceName && *info.traceName)
 			traceName.Append(info.traceName);
-		else if (info.launchVisualizer || !m_traceOutputFile.IsEmpty())
+		else if (info.launchVisualizer || !m_traceOutputFile.IsEmpty() || info.traceEnabled)
 			traceName.Append(m_id);
 
 		if (!traceName.IsEmpty())
@@ -1451,56 +1443,7 @@ namespace uba
 	bool Session::IsKnownSystemFile(const tchar* applicationName)
 	{
 #if PLATFORM_WINDOWS
-		// These dlls should _always_ exist on all machines so we can filter them out from the list of imports to copy
-
-		static const wchar_t* knownSystemFiles[] = // Sorted lowercase
-		{
-			L"advapi32.dll",
-			L"bcrypt.dll",
-			L"bcryptprimitives.dll",
-			L"combase.dll",
-			L"cryptbase.dll",
-			L"dbghelp.dll",
-			L"dnsapi.dll",
-			L"fwpuclnt.dll",
-			L"gdi32.dll",
-			L"gdi32full.dll",
-			L"imm32.dll",
-			L"iphlpapi.dll",
-			L"kernel32.dll",
-			L"kernelbase.dll",
-			L"msvcp_win.dll",
-			L"msvcrt.dll",
-			L"mswsock.dll",
-			L"ncrypt.dll",
-			L"nsi.dll",
-			L"ntasn1.dll",
-			L"ntdll.dll",
-			L"ole32.dll",
-			L"oleaut32.dll",
-			L"ondemandconnroutehelper.dll",
-			L"powrprof.dll",
-			L"rasadhlp.dll",
-			L"rpcrt4.dll",
-			L"rstrtmgr.dll",
-			L"sechost.dll",
-			L"sspicli.dll",
-			L"ucrtbase.dll",
-			L"umpdc.dll",
-			L"umppc17706.dll",
-			L"user32.dll",
-			L"version.dll",
-			L"webio.dll",
-			L"win32u.dll",
-			L"winhttp.dll",
-			L"winnsi.dll",
-			L"ws2_32.dll",
-		};
-
-		auto lastSeparator = TStrrchr(applicationName, PathSeparator);
-		StringBuffer<> nameLower;
-		nameLower.Append(lastSeparator + 1).MakeLower();
-		return std::binary_search(knownSystemFiles, knownSystemFiles + sizeof_array(knownSystemFiles), nameLower.data, [](const wchar_t* a, const wchar_t* b) { return TStrcmp(a, b) < 0; });
+		return uba::IsKnownSystemFile(applicationName);
 #else
 		return false;
 #endif
diff --git a/Engine/Source/Programs/UnrealBuildAccelerator/Common/Private/UbaSessionClient.cpp b/Engine/Source/Programs/UnrealBuildAccelerator/Common/Private/UbaSessionClient.cpp
--- a/Engine/Source/Programs/UnrealBuildAccelerator/Common/Private/UbaSessionClient.cpp
+++ b/Engine/Source/Programs/UnrealBuildAccelerator/Common/Private/UbaSessionClient.cpp
@@ -227,6 +227,11 @@ namespace uba
 			if (moduleCount == 0)
 				return m_logger.Error(TC("Application %s not found"), application);
 
+			struct ModuleInfo { TString name; CasKey casKey; u32 attributes; Event done; };
+			List<ModuleInfo> modules;
+
+			Atomic<bool> success = true;
+
 			while (moduleCount--)
 			{
 				StringBuffer<> moduleFile;
@@ -247,20 +252,35 @@ namespace uba
 					moduleFile.Clear().Append(localSystemModule);
 				}
 
-				CasKey newCasKey;
-				bool storeUncompressed = true;
-				u64 fileSize;
-				if (!RetrieveCasFile(newCasKey, fileSize, casKey, moduleFile.data, storeUncompressed))
-					return m_logger.Error(TC("Casfile not found for %s (%s)"), moduleFile.data, CasKeyString(casKey).str);
+				auto& m = modules.emplace_back(moduleFile.data, casKey, fileAttributes, true);
+				m_client.AddWork([&]()
+					{
+						auto g = MakeGuard([&]() { m.done.Set(); });
+						CasKey newCasKey;
+						bool storeUncompressed = true;
+						u64 fileSize;
+						const tchar* moduleName = m.name.c_str();
+						if (!RetrieveCasFile(newCasKey, fileSize, m.casKey, moduleName, storeUncompressed))
+						{
+							m_logger.Error(TC("Casfile not found for %s (%s)"), moduleName, CasKeyString(m.casKey).str);
+							success = false;
+							return;
+						}
+						if (const tchar* lastSeparator = TStrrchr(moduleName, PathSeparator))
+							moduleName = lastSeparator + 1;
+						StringBuffer<> temp;
+						if (!WriteBinFile(temp, moduleName, newCasKey, keyStr, m.attributes))
+							success = false;
+					}, 1, TC("EnsureApp"));
+			}
 
-				const tchar* moduleName = moduleFile.data;
-				if (const tchar* lastSeparator = TStrrchr(moduleFile.data, PathSeparator))
-					moduleName = lastSeparator + 1;
+			// Wait for all to be done
+			for (auto& m : modules)
+				if (!m.done.IsSet(10 * 60 * 1000)) // 10 minutes is a very long time
+					return m_logger.Error(TC("Timed out while waiting for application cas files to be downloaded"));
 
-				StringBuffer<> temp;
-				if (!WriteBinFile(temp, moduleName, newCasKey, keyStr, fileAttributes))
-					return false;
-			}
+			if (!success)
+				return false;
 
 			failGuard.Cancel();
 		}
@@ -678,6 +698,9 @@ namespace uba
 
 	bool SessionClient::GetFullFileName(GetFullFileNameResponse& out, const GetFullFileNameMessage& msg)
 	{
+		// TODO: There is a potential risk here where two different applications asks for the full name of a file
+		// and they have different bin/working dir.. and there are two versions of this file.
+
 		ScopedWriteLock lock(m_nameToNameLookupLock);
 		auto insres = m_nameToNameLookup.try_emplace(msg.fileName.data);
 		NameRec& rec = insres.first->second;
@@ -1099,7 +1122,7 @@ namespace uba
 					{
 						Storage::RetrieveResult result;
 						bool allowProxy = true;
-						bool res = m_storage.RetrieveCasFile(result, knownInputKey, TC(""), nullptr, 1, allowProxy);
+						bool res = m_storage.RetrieveCasFile(result, knownInputKey, TC("KnownInput"), nullptr, 1, allowProxy);
 						(void)res;
 					}, 1, TC("KnownInput"));
 			}
diff --git a/Engine/Source/Programs/UnrealBuildAccelerator/Common/Private/UbaSessionServer.cpp b/Engine/Source/Programs/UnrealBuildAccelerator/Common/Private/UbaSessionServer.cpp
--- a/Engine/Source/Programs/UnrealBuildAccelerator/Common/Private/UbaSessionServer.cpp
+++ b/Engine/Source/Programs/UnrealBuildAccelerator/Common/Private/UbaSessionServer.cpp
@@ -1,51 +1,23 @@
 // Copyright Epic Games, Inc. All Rights Reserved.
 
 #include "UbaSessionServer.h"
-#include "UbaStorage.h"
-#include "UbaProcess.h"
 #include "UbaNetworkServer.h"
+#include "UbaProcess.h"
+#include "UbaProcessStartInfoHolder.h"
+#include "UbaStorage.h"
 
 namespace uba
 {
-	void FixFileName(StringBufferBase& out, const tchar* fileName, const tchar* workingDir)
-	{
-		tchar buffer[1024];
-		u32 charLen;
-		u64 workingDirLen = 0;
-		if (workingDir)
-			workingDirLen = TStrlen(workingDir);
-		FixPath2(fileName, workingDir, workingDirLen, buffer, &charLen);
-		out.Append(buffer);
-	}
-
-	class SessionServer::RemoteProcess : public Process
+	class SessionServer::RemoteProcess : public Process, public ProcessStartInfoHolder
 	{
 	public:
 		RemoteProcess(SessionServer* server, const ProcessStartInfo& si, u32 processId, float weight)
-		:	m_server(server)
-		,	m_startInfo(si)
+		:	ProcessStartInfoHolder(si)
+		,	m_server(server)
 		,	m_processId(processId)
 		,	m_done(true)
 		,	m_weight(weight)
 		{
-			StringBuffer<512> temp;
-			FixFileName(temp, si.workingDir, nullptr);
-			temp.EnsureEndsWithSlash();
-			workingDir = temp.data;
-
-			temp.EnsureEndsWithSlash();
-			StringBuffer<512> temp2;
-			FixFileName(temp2, si.application, temp.data);
-			application = temp2.data;
-
-			description = si.description;
-			arguments = si.arguments;
-
-			m_startInfo.description = description.c_str();
-			m_startInfo.application = application.c_str();
-			m_startInfo.arguments = arguments.c_str();
-			m_startInfo.workingDir = workingDir.c_str();
-			m_startInfo.logFile = nullptr;
 		}
 
 		~RemoteProcess()
@@ -53,7 +25,7 @@ namespace uba
 			delete[] m_knownInputs;
 		}
 
-		virtual const ProcessStartInfo& GetStartInfo() override { return m_startInfo; }
+		virtual const ProcessStartInfo& GetStartInfo() override { return startInfo; }
 		virtual u32 GetId() override { return m_processId; }
 		virtual u32 GetExitCode() override { UBA_ASSERT(m_done.IsSet(0)); return m_exitCode; }
 		virtual bool HasExited() override { return m_done.IsSet(0); }
@@ -73,13 +45,13 @@ namespace uba
 			else
 				m_done.Set();
 
-			if (m_startInfo.exitedFunc)
+			if (startInfo.exitedFunc)
 			{
 				ProcessHandle h;
 				h.m_process = this;
-				m_startInfo.exitedFunc(m_startInfo.exitedUserData, h);
+				startInfo.exitedFunc(startInfo.exitedUserData, h);
 				h.m_process = nullptr;
-				m_startInfo.exitedFunc = nullptr;
+				startInfo.exitedFunc = nullptr;
 			}
 		}
 
@@ -88,11 +60,6 @@ namespace uba
 		virtual bool IsChild() override { return false; }
 
 		SessionServer* m_server;
-		ProcessStartInfo m_startInfo;
-		TString description;
-		TString application;
-		TString arguments;
-		TString workingDir;
 		u32 m_processId;
 		u32 m_exitCode = ~u32(0);
 		u64 m_processorTime = 0;
@@ -930,12 +897,12 @@ namespace uba
 
 					ProcessAdded(*process, sessionId);
 					writer.WriteU32(process->m_processId);
-					writer.WriteString(process->m_startInfo.description);
-					writer.WriteString(process->m_startInfo.application);
-					writer.WriteString(process->m_startInfo.arguments);
-					writer.WriteString(process->m_startInfo.workingDir);
+					writer.WriteString(process->startInfo.description);
+					writer.WriteString(process->startInfo.application);
+					writer.WriteString(process->startInfo.arguments);
+					writer.WriteString(process->startInfo.workingDir);
 					writer.WriteU32(*(u32*)&process->m_weight);
-					writer.WriteU64(process->m_startInfo.outputStatsThresholdMs);
+					writer.WriteU64(process->startInfo.outputStatsThresholdMs);
 
 					for (auto kiIt = process->m_knownInputs, kiEnd = kiIt + process->m_knownInputsCount; kiIt!=kiEnd; ++kiIt)
 						if (session.sentKeys.insert(kiIt->key).second)
@@ -1044,9 +1011,9 @@ namespace uba
 					process.m_logLines.push_back({ std::move(text), type });
 				}
 
-				if (auto func = process.m_startInfo.logLineFunc)
+				if (auto func = process.startInfo.logLineFunc)
 					for (auto& line : process.m_logLines)
-						func(process.m_startInfo.logLineUserData, line.text.c_str(), u32(line.text.size()), line.type);
+						func(process.startInfo.logLineUserData, line.text.c_str(), u32(line.text.size()), line.type);
 
 				u32 id = process.m_processId;
 				m_trace.ProcessExited(id, exitCode, reader.GetPositionData(), reader.GetLeft());
@@ -1064,10 +1031,10 @@ namespace uba
 				process.m_server = nullptr;
 				process.m_done.Set();
 				
-				if (process.m_startInfo.exitedFunc)
+				if (process.startInfo.exitedFunc)
 				{
-					process.m_startInfo.exitedFunc(process.m_startInfo.exitedUserData, h);
-					process.m_startInfo.exitedFunc = nullptr;
+					process.startInfo.exitedFunc(process.startInfo.exitedUserData, h);
+					process.startInfo.exitedFunc = nullptr;
 				}
 				return true;
 			}
diff --git a/Engine/Source/Programs/UnrealBuildAccelerator/Common/Public/UbaProcessHandle.h b/Engine/Source/Programs/UnrealBuildAccelerator/Common/Public/UbaProcessHandle.h
new file mode 100644
--- /dev/null
+++ b/Engine/Source/Programs/UnrealBuildAccelerator/Common/Public/UbaProcessHandle.h
@@ -0,0 +1,59 @@
+// Copyright Epic Games, Inc. All Rights Reserved.
+
+#pragma once
+
+#include "UbaLogger.h"
+#include "UbaMemory.h"
+
+namespace uba
+{
+	class Process;
+	struct ProcessLogLine;
+	struct ProcessStartInfo;
+
+	// Exit code returned by ProcessHandle.GetExitCode() if build was cancelled
+	inline constexpr u32 ProcessCancelExitCode = 99999;
+
+
+	struct ProcessLogLine
+	{
+		TString text;
+		LogEntryType type = LogEntryType_Info;
+	};
+
+
+	class ProcessHandle
+	{
+	public:
+
+		const ProcessStartInfo& GetStartInfo() const;		// Get the info used to start process
+		u32 GetId() const;									// Unique id of process. Is an increasing index internally
+		u32 GetExitCode() const;							// Exit code of process. Will return ~0u if not finished or default error
+		bool HasExited() const;								// Returns true if process has exited
+		bool WaitForExit(u32 millisecondsTimeout) const;	// Wait for process until it exists or time reach timeout
+		const Vector<ProcessLogLine>& GetLogLines() const;	// Log lines produced by the process
+		const Vector<u8>& GetTrackedInputs() const;			// If ProcessStartInfo.trackInputs was true, this will return a buffer to wchar/char strings 
+		u64 GetTotalProcessorTime() const;					// Total used cpu time (in time units. Use TimeToMs etc)
+		u64 GetTotalWallTime() const;						// Total wall time (in time units. Use TimeToMs etc)
+		void Cancel(bool terminate) const;					// Request to cancel process. Set terminate to true to quickly kill process
+		const tchar* GetExecutingHost() const;				// Host that is executing process.
+		bool IsRemote() const;								// Returns true if process is a remote process
+
+		ProcessHandle();
+		ProcessHandle(const ProcessHandle& o);
+		ProcessHandle(ProcessHandle&& o) noexcept;
+		ProcessHandle& operator=(const ProcessHandle& o);
+		ProcessHandle& operator=(ProcessHandle&& o) noexcept;
+		~ProcessHandle();
+		inline bool operator==(const ProcessHandle& o) const { return m_process == o.m_process; }
+		inline u64 GetHash() const { return u64(m_process); }
+
+	private:
+		ProcessHandle(Process* process);
+		Process* m_process;
+		friend class ProcessImpl;
+		friend class Session;
+		friend class SessionClient;
+		friend class SessionServer;
+	};
+}
\ No newline at end of file
diff --git a/Engine/Source/Programs/UnrealBuildAccelerator/Common/Public/UbaScheduler.h b/Engine/Source/Programs/UnrealBuildAccelerator/Common/Public/UbaScheduler.h
new file mode 100644
--- /dev/null
+++ b/Engine/Source/Programs/UnrealBuildAccelerator/Common/Public/UbaScheduler.h
@@ -0,0 +1,55 @@
+// Copyright Epic Games, Inc. All Rights Reserved.
+
+#pragma once
+
+#include "UbaMemory.h"
+#include "UbaProcessStartInfo.h"
+#include "UbaThread.h"
+
+namespace uba
+{
+	class Process;
+	class SessionServer;
+	struct ProcessStartInfoHolder;
+
+	// This is a simple scheduler that does not handle any dependencies. Can be added based on need
+
+	class Scheduler
+	{
+	public:
+		Scheduler(SessionServer& session, u32 maxLocalProcessors = 0);
+		~Scheduler();
+
+		void Start();
+		void Stop();
+
+		u32 EnqueueProcess(const ProcessStartInfo& info, float weight = 1.0f, const void* knownInputs = nullptr, u32 knownInputsBytes = 0, u32 knownInputsCount = 0);
+
+		void GetStats(u32& outQueued, u32& outActive, u32& outFinished);
+	private:
+		struct ExitProcessInfo;
+
+		void ThreadLoop();
+		void RemoteProcessReturned(Process& process);
+		void RemoteSlotAvailable();
+		void ProcessExited(ExitProcessInfo* info, const ProcessHandle& handle);
+		bool RunQueuedProcess(bool runLocal);
+
+		SessionServer& m_session;
+		u32 m_maxLocalProcessors;
+
+		ReaderWriterLock m_queuedProcessesLock;
+		struct QueuedProcess;
+		List<QueuedProcess> m_queuedProcesses;
+		Event m_updateThreadLoop;
+		Thread m_thread;
+		bool m_loop = false;
+
+		Atomic<u32> m_activeLocalProcesses;
+		Atomic<u32> m_activeRemoteProcesses;
+		Atomic<u32> m_finishedProcesses;
+
+		Scheduler(const Scheduler&) = delete;
+		void operator=(const Scheduler&) = delete;
+	};
+}
diff --git a/Engine/Source/Programs/UnrealBuildAccelerator/Common/Public/UbaSession.h b/Engine/Source/Programs/UnrealBuildAccelerator/Common/Public/UbaSession.h
--- a/Engine/Source/Programs/UnrealBuildAccelerator/Common/Public/UbaSession.h
+++ b/Engine/Source/Programs/UnrealBuildAccelerator/Common/Public/UbaSession.h
@@ -7,6 +7,7 @@
 #include "UbaFile.h"
 #include "UbaDirectoryTable.h"
 #include "UbaFileMapping.h"
+#include "UbaProcessHandle.h"
 #include "UbaProcessStartInfo.h"
 #include "UbaSessionCreateInfo.h"
 #include "UbaStats.h"
@@ -43,51 +44,6 @@ namespace uba
 	struct ListDirectoryResponse;
 	struct WrittenFile;
 
-	// Exit code returned by ProcessHandle.GetExitCode() if build was cancelled
-	inline constexpr u32 ProcessCancelExitCode = 99999;
-
-	struct ProcessLogLine
-	{
-		TString text;
-		LogEntryType type = LogEntryType_Info;
-	};
-
-	class ProcessHandle
-	{
-	public:
-
-		const ProcessStartInfo& GetStartInfo() const;		// Get the info used to start process
-		u32 GetId() const;									// Unique id of process. Is an increasing index internally
-		u32 GetExitCode() const;							// Exit code of process. Will return ~0u if not finished or default error
-		bool HasExited() const;								// Returns true if process has exited
-		bool WaitForExit(u32 millisecondsTimeout) const;	// Wait for process until it exists or time reach timeout
-		const Vector<ProcessLogLine>& GetLogLines() const;	// Log lines produced by the process
-		const Vector<u8>& GetTrackedInputs() const;			// If ProcessStartInfo.trackInputs was true, this will return a buffer to wchar/char strings 
-		u64 GetTotalProcessorTime() const;					// Total used cpu time (in time units. Use TimeToMs etc)
-		u64 GetTotalWallTime() const;						// Total wall time (in time units. Use TimeToMs etc)
-		void Cancel(bool terminate) const;					// Request to cancel process. Set terminate to true to quickly kill process
-		const tchar* GetExecutingHost() const;				// Host that is executing process.
-		bool IsRemote() const;								// Returns true if process is a remote process
-
-		ProcessHandle();
-		ProcessHandle(const ProcessHandle& o);
-		ProcessHandle(ProcessHandle&& o) noexcept;
-		ProcessHandle& operator=(const ProcessHandle& o);
-		ProcessHandle& operator=(ProcessHandle&& o) noexcept;
-		~ProcessHandle();
-		inline bool operator==(const ProcessHandle& o) const { return m_process == o.m_process; }
-		inline u64 GetHash() const { return u64(m_process); }
-
-	private:
-		ProcessHandle(Process* process);
-		Process* m_process;
-		friend class ProcessImpl;
-		friend class Session;
-		friend class SessionClient;
-		friend class SessionServer;
-	};
-
-
 	class Session
 	{
 	public:
diff --git a/Engine/Source/Programs/UnrealBuildAccelerator/Common/Public/UbaSessionCreateInfo.h b/Engine/Source/Programs/UnrealBuildAccelerator/Common/Public/UbaSessionCreateInfo.h
--- a/Engine/Source/Programs/UnrealBuildAccelerator/Common/Public/UbaSessionCreateInfo.h
+++ b/Engine/Source/Programs/UnrealBuildAccelerator/Common/Public/UbaSessionCreateInfo.h
@@ -24,6 +24,7 @@ namespace uba
 		bool launchVisualizer = false;			// Launch a UbaVisualizer process (this automatically enable trace)
 		bool allowMemoryMaps = IsWindows;		// Use memory maps where possible. Session creates memory maps of files that processes use
 		bool shouldWriteToDisk = true;			// Set to false to skip writing output files to disk
+		bool traceEnabled = false;				// Set to true to always create in-memory trace data. Is not needed if traceName, traceOutputFile or launchVisualizer is set
 		bool detailedTrace = false;				// Enable detailed trace to include jobs, individual file I/O etc in trace dump
 		u64 deleteSessionsOlderThanSeconds = 12 * 60 * 60; // Delete session folders older than 12 hours by default . Set to 0 to not delete or 1 to delete all
 		u64 keepOutputFileMemoryMapsThreshold = 256 * 1024; // If allowMemoryMaps is true, output files will be kept in memory if smaller than this size
diff --git a/Engine/Source/Programs/UnrealBuildAccelerator/Core/Public/UbaPlatform.h b/Engine/Source/Programs/UnrealBuildAccelerator/Core/Public/UbaPlatform.h
--- a/Engine/Source/Programs/UnrealBuildAccelerator/Core/Public/UbaPlatform.h
+++ b/Engine/Source/Programs/UnrealBuildAccelerator/Core/Public/UbaPlatform.h
@@ -70,6 +70,8 @@ namespace uba
 
 	void WriteAssertInfo(class StringBufferBase& out, const tchar* text, const char* file, u32 line, const char* expr, u32 skipCallstack);
 	ANALYSIS_NORETURN void UbaAssert(const tchar* text, const char* file, u32 line, const char* expr, u32 terminateCode);
+	using CustomAssertHandler = void(const tchar* text);
+	void SetCustomAssertHandler(CustomAssertHandler* handler);
 	bool CreateGuid(Guid& out);
 	bool IsRunningWine();
 	bool IsEscapePressed();
diff --git a/Engine/Source/Programs/UnrealBuildAccelerator/Core/Public/UbaWinBinDependencyParser.h b/Engine/Source/Programs/UnrealBuildAccelerator/Core/Public/UbaWinBinDependencyParser.h
--- a/Engine/Source/Programs/UnrealBuildAccelerator/Core/Public/UbaWinBinDependencyParser.h
+++ b/Engine/Source/Programs/UnrealBuildAccelerator/Core/Public/UbaWinBinDependencyParser.h
@@ -1,9 +1,89 @@
 // Copyright Epic Games, Inc. All Rights Reserved.
 
+#pragma once
+
+#include "UbaStringBuffer.h"
+#include <algorithm>
+
+#if !defined(IS_DETOURED_INCLUDE)
+#define True_CreateFileW ::CreateFileW
+#define True_CreateFileMappingW ::CreateFileMappingW
+#define True_MapViewOfFile ::MapViewOfFile
+#define True_UnmapViewOfFile ::UnmapViewOfFile
+#endif
+
 namespace uba
 {
+	inline bool IsKnownSystemFile(const tchar* applicationName)
+	{
+		// These dlls should _always_ exist on all machines so we can filter them out from the list of imports to copy
+
+		static const wchar_t* knownSystemFiles[] = // Sorted lowercase
+		{
+			L"advapi32.dll",
+			L"bcrypt.dll",
+			L"bcryptprimitives.dll",
+			L"combase.dll",
+			L"concrt140.dll",
+			L"crypt32.dll",
+			L"cryptbase.dll",
+			L"dbghelp.dll",
+			L"dnsapi.dll",
+			L"dwmapi.dll",
+			L"dxgi.dll",
+			L"fwpuclnt.dll",
+			L"gdi32.dll",
+			L"gdi32full.dll",
+			L"imagehlp.dll",
+			L"imm32.dll",
+			L"iphlpapi.dll",
+			L"kernel32.dll",
+			L"kernelbase.dll",
+			L"msvcp_win.dll",
+			L"msvcrt.dll",
+			L"mswsock.dll",
+			L"ncrypt.dll",
+			L"netapi32.dll",
+			L"nsi.dll",
+			L"ntasn1.dll",
+			L"ntdll.dll",
+			L"ole32.dll",
+			L"oleaut32.dll",
+			L"ondemandconnroutehelper.dll",
+			L"powrprof.dll",
+			L"psapi.dll",
+			L"rasadhlp.dll",
+			L"rpcrt4.dll",
+			L"rstrtmgr.dll",
+			L"sechost.dll",
+			L"setupapi.dll"
+			L"shell32.dll",
+			L"sspicli.dll",
+			L"ucrtbase.dll",
+			L"umpdc.dll",
+			L"umppc17706.dll",
+			L"user32.dll",
+			L"userenv.dll",
+			L"uxtheme.dll",
+			L"version.dll",
+			L"webio.dll",
+			L"win32u.dll",
+			L"winhttp.dll",
+			L"winmm.dll",
+			L"winnsi.dll",
+			L"ws2_32.dll",
+		};
+
+		StringBuffer<> nameLower;
+		if (auto lastSeparator = TStrrchr(applicationName, PathSeparator))
+			nameLower.Append(lastSeparator + 1);
+		else
+			nameLower.Append(applicationName);
+		nameLower.MakeLower();
+		return std::binary_search(knownSystemFiles, knownSystemFiles + sizeof_array(knownSystemFiles), nameLower.data, [](const wchar_t* a, const wchar_t* b) { return TStrcmp(a, b) < 0; });
+	}
 
-	DWORD GetOffset(DWORD rva, PIMAGE_SECTION_HEADER psh, PIMAGE_NT_HEADERS pnt)
+	inline DWORD GetOffset(DWORD rva, PIMAGE_SECTION_HEADER psh, PIMAGE_NT_HEADERS pnt)
 	{
 		if(rva == 0)
 			return rva;
@@ -14,7 +94,7 @@ namespace uba
 		return rva - seh->VirtualAddress + seh->PointerToRawData;
 	} 
 
-	bool FindImportsInMem(LPCWSTR fileName, void* mem, const Function<void(const char* import)>& func)
+	inline bool FindImportsInMem(LPCWSTR fileName, void* mem, const Function<void(const tchar* import)>& func)
 	{
 		auto hdrs = (PIMAGE_NT_HEADERS)(PCHAR(mem) + PIMAGE_DOS_HEADER(mem)->e_lfanew);   
 		auto pSech=IMAGE_FIRST_SECTION(hdrs);
@@ -28,7 +108,10 @@ namespace uba
 			while(importDesc->Name != NULL)
 			{
 				const char* name = (PCHAR)((DWORD_PTR)mem + GetOffset(importDesc->Name, pSech, hdrs));
-				func(name);
+				StringBuffer<256> wname;
+				wname.Append(name);
+				if (!IsKnownSystemFile(wname.data))
+					func(wname.data);
 				++importDesc;
 
 			}
@@ -43,7 +126,7 @@ namespace uba
 		return true;
 	}
 
-	bool FindImports(LPCWSTR fileName, const Function<void(const char* import)>& func)
+	inline bool FindImports(LPCWSTR fileName, const Function<void(const tchar* import)>& func)
 	{
 		HANDLE fileHandle = True_CreateFileW(fileName, GENERIC_READ, FILE_SHARE_READ, 0, OPEN_EXISTING, FILE_ATTRIBUTE_NORMAL | FILE_FLAG_SEQUENTIAL_SCAN, 0);
 		if (fileHandle == INVALID_HANDLE_VALUE)
diff --git a/Engine/Source/Programs/UnrealBuildAccelerator/Detours/Private/Windows/UbaDetoursFunctionsKernelBase.inl b/Engine/Source/Programs/UnrealBuildAccelerator/Detours/Private/Windows/UbaDetoursFunctionsKernelBase.inl
--- a/Engine/Source/Programs/UnrealBuildAccelerator/Detours/Private/Windows/UbaDetoursFunctionsKernelBase.inl
+++ b/Engine/Source/Programs/UnrealBuildAccelerator/Detours/Private/Windows/UbaDetoursFunctionsKernelBase.inl
@@ -2272,13 +2272,13 @@ HMODULE Recursive_LoadLibraryExW(LPCWSTR lpLibFileName, LPCWSTR originalName, DW
 		return 0;
 
 	// Important that this code is not doing allocations.. it could cause a recursive stack overflow
-	struct Import { char name[128]; Import(const char* s) { strcpy_s(name, sizeof_array(name), s); } };
+	struct Import { wchar_t name[128]; Import(const wchar_t* s) { wcscpy_s(name, sizeof_array(name), s); } };
 	std::vector<Import, GrowingAllocator<Import>> importedModules(&g_memoryBlock);
 	{
 		SuppressCreateFileDetourScope cfs;
-		if (!FindImports(lpLibFileName, [&](const char* import)
+		if (!FindImports(lpLibFileName, [&](const wchar_t* import)
 			{
-				if (!GetModuleHandleA(import))
+				if (!GetModuleHandleW(import))
 					importedModules.emplace_back(import);
 			}))
 		{
@@ -2287,13 +2287,10 @@ HMODULE Recursive_LoadLibraryExW(LPCWSTR lpLibFileName, LPCWSTR originalName, DW
 	}
 	for (auto& importedModule : importedModules)
 	{
-		if (GetModuleHandleA(importedModule.name))
+		if (GetModuleHandleW(importedModule.name))
 			continue;
 
-		wchar_t moduleNameW[256];
-		swprintf_s(moduleNameW, 256, L"%hs", importedModule.name);
-
-		const wchar_t* path = moduleNameW;
+		const wchar_t* path = importedModule.name;
 		if (path[1] == ':')
 			if (const wchar_t* lastSlash = wcsrchr(path, '\\'))
 				path = lastSlash + 1;
@@ -2301,7 +2298,7 @@ HMODULE Recursive_LoadLibraryExW(LPCWSTR lpLibFileName, LPCWSTR originalName, DW
 		StringBuffer<512> tempBuf;
 		Rpc_GetFullFileName(path, pathLen, tempBuf, false);
 
-		if (HMODULE r = Recursive_LoadLibraryExW(path, moduleNameW, dwFlags, additionalLoads, visitedModules))
+		if (HMODULE r = Recursive_LoadLibraryExW(path, importedModule.name, dwFlags, additionalLoads, visitedModules))
 			additionalLoads.push_back(r);
 	}
 
diff --git a/Engine/Source/Programs/UnrealBuildAccelerator/Detours/Private/Windows/UbaDetoursFunctionsWin.cpp b/Engine/Source/Programs/UnrealBuildAccelerator/Detours/Private/Windows/UbaDetoursFunctionsWin.cpp
--- a/Engine/Source/Programs/UnrealBuildAccelerator/Detours/Private/Windows/UbaDetoursFunctionsWin.cpp
+++ b/Engine/Source/Programs/UnrealBuildAccelerator/Detours/Private/Windows/UbaDetoursFunctionsWin.cpp
@@ -3,6 +3,8 @@
 #include "UbaDetoursFunctionsWin.h"
 #include "UbaDetoursFileMappingTable.h"
 
+#define IS_DETOURED_INCLUDE 1
+
 #if !defined(UBA_USE_MIMALLOC)
 #define True_malloc malloc
 #else
diff --git a/Engine/Source/Programs/UnrealBuildAccelerator/Host/Private/UbaExports.cpp b/Engine/Source/Programs/UnrealBuildAccelerator/Host/Private/UbaExports.cpp
--- a/Engine/Source/Programs/UnrealBuildAccelerator/Host/Private/UbaExports.cpp
+++ b/Engine/Source/Programs/UnrealBuildAccelerator/Host/Private/UbaExports.cpp
@@ -4,9 +4,15 @@
 #include "UbaNetworkBackendQuic.h"
 #include "UbaNetworkBackendTcp.h"
 #include "UbaAWS.h"
+#include "UbaProcess.h"
+#include "UbaScheduler.h"
 #include "UbaStorageServer.h"
 #include "UbaSessionServer.h"
 
+#if PLATFORM_WINDOWS
+#include "UbaWinBinDependencyParser.h"
+#endif
+
 namespace uba
 {
 	CallbackLogWriter::CallbackLogWriter(BeginScopeCallback begin, EndScopeCallback end, LogCallback log) : m_beginScope(begin), m_endScope(end), m_logCallback(log)
@@ -283,6 +289,12 @@ uba::StorageClient* CreateStorageClient(uba::NetworkClient& client, const uba::t
 	{
 		delete handle;
 	}
+
+	const uba::ProcessStartInfo* Process_GetStartInfo(uba::Process& process)
+	{
+		return &process.GetStartInfo();
+	}
+
 /*
 	uba::SessionClient* CreateSessionClient(const uba::SessionClientCreateInfo& info)
 	{
@@ -318,13 +330,13 @@ uba::StorageClient* CreateStorageClient(uba::NetworkClient& client, const uba::t
 	{
 		return new uba::SessionServer(info);
 	}
-	void SessionServer_SetRemoteProcessAvailable(uba::SessionServer* server, SessionServer_RemoteProcessAvailableCallback* available)
+	void SessionServer_SetRemoteProcessAvailable(uba::SessionServer* server, SessionServer_RemoteProcessAvailableCallback* available, void* userData)
 	{
-		server->SetRemoteProcessSlotAvailableEvent(available);
+		server->SetRemoteProcessSlotAvailableEvent([available, userData]() { available(userData); });
 	}
-	void SessionServer_SetRemoteProcessReturned(uba::SessionServer* server, SessionServer_RemoteProcessReturnedCallback* returned)
+	void SessionServer_SetRemoteProcessReturned(uba::SessionServer* server, SessionServer_RemoteProcessReturnedCallback* returned, void* userData)
 	{
-		server->SetRemoteProcessReturnedEvent(returned);
+		server->SetRemoteProcessReturnedEvent([returned, userData](uba::Process& process) { returned(process, userData); });
 	}
 	void SessionServer_RefreshDirectory(uba::SessionServer* server, const uba::tchar* directory)
 	{
@@ -416,4 +428,46 @@ uba::StorageClient* CreateStorageClient(uba::NetworkClient& client, const uba::t
 		free((void*)info->logFile);
 		delete info;
 	}
+
+	uba::Scheduler* Scheduler_Create(uba::SessionServer* session, uba::u32 maxLocalProcessors)
+	{
+		return new uba::Scheduler(*session, maxLocalProcessors);
+	}
+
+	void Scheduler_Start(uba::Scheduler* scheduler)
+	{
+		scheduler->Start();
+	}
+
+	void Scheduler_EnqueueProcess(uba::Scheduler* scheduler, const uba::ProcessStartInfo& info, float weight, const void* knownInputs, uba::u32 knownInputsBytes, uba::u32 knownInputsCount)
+	{
+		scheduler->EnqueueProcess(info, weight, knownInputs, knownInputsBytes, knownInputsCount);
+	}
+
+	void Scheduler_Stop(uba::Scheduler* scheduler)
+	{
+		scheduler->Stop();
+	}
+
+	void Scheduler_Destroy(uba::Scheduler* scheduler)
+	{
+		delete scheduler;
+	}
+
+	void Scheduler_GetStats(uba::Scheduler* scheduler, uba::u32& outQueued, uba::u32& outActive, uba::u32& outFinished)
+	{
+		scheduler->GetStats(outQueued, outActive, outFinished);
+	}
+
+	void Uba_SetCustomAssertHandler(Uba_CustomAssertHandler* handler)
+	{
+		uba::SetCustomAssertHandler(handler);
+	}
+
+	void Uba_FindImports(const uba::tchar* binary, ImportFunc* func, void* userData)
+	{
+#if PLATFORM_WINDOWS
+		uba::FindImports(binary, [&](const uba::tchar* importName) { func(importName, userData); });
+#endif
+	}
 }
diff --git a/Engine/Source/Programs/UnrealBuildAccelerator/Host/Public/UbaExports.h b/Engine/Source/Programs/UnrealBuildAccelerator/Host/Public/UbaExports.h
--- a/Engine/Source/Programs/UnrealBuildAccelerator/Host/Public/UbaExports.h
+++ b/Engine/Source/Programs/UnrealBuildAccelerator/Host/Public/UbaExports.h
@@ -12,6 +12,7 @@ namespace uba
 	class ProcessHandle;
 	class Session;
 	class SessionServer;
+	class Scheduler;
 	class Storage;
 	class StorageServer;
 	struct ProcessStartInfo;
@@ -72,17 +73,18 @@ extern "C"
 	UBA_API uba::u64 ProcessHandle_GetTotalWallTime(uba::ProcessHandle* handle);
 	UBA_API void ProcessHandle_Cancel(uba::ProcessHandle* handle, bool terminate);
 	UBA_API void DestroyProcessHandle(uba::ProcessHandle* handle);
+	UBA_API const uba::ProcessStartInfo* Process_GetStartInfo(uba::Process& process);
 
-	using SessionServer_RemoteProcessAvailableCallback = void();
-	using SessionServer_RemoteProcessReturnedCallback = void(uba::Process& process);
+	using SessionServer_RemoteProcessAvailableCallback = void(void* userData);
+	using SessionServer_RemoteProcessReturnedCallback = void(uba::Process& process, void* userData);
 	using SessionServer_CustomServiceFunction = uba::u32(const uba::Guid& clientUid, const void* recv, uba::u32 recvSize, void* send, uba::u32 sendCapacity, void* userData);
 
 	UBA_API uba::SessionServerCreateInfo* CreateSessionServerCreateInfo(uba::Storage& storage, uba::NetworkServer& client, uba::LogWriter& writer, const uba::tchar* rootDir, const uba::tchar* traceOutputFile,
 		bool disableCustomAllocator, bool launchVisualizer, bool resetCas, bool writeToDisk, bool detailedTrace, bool allowWaitOnMem = false, bool allowKillOnMem = false);
 	UBA_API void DestroySessionServerCreateInfo(uba::SessionServerCreateInfo* info);
 	UBA_API uba::SessionServer* CreateSessionServer(const uba::SessionServerCreateInfo& info);
-	UBA_API void SessionServer_SetRemoteProcessAvailable(uba::SessionServer* server, SessionServer_RemoteProcessAvailableCallback* available);
-	UBA_API void SessionServer_SetRemoteProcessReturned(uba::SessionServer* server, SessionServer_RemoteProcessReturnedCallback* returned);
+	UBA_API void SessionServer_SetRemoteProcessAvailable(uba::SessionServer* server, SessionServer_RemoteProcessAvailableCallback* available, void* userData);
+	UBA_API void SessionServer_SetRemoteProcessReturned(uba::SessionServer* server, SessionServer_RemoteProcessReturnedCallback* returned, void* userData);
 	UBA_API void SessionServer_RefreshDirectory(uba::SessionServer* server, const uba::tchar* directory);
 	UBA_API void SessionServer_RegisterNewFile(uba::SessionServer* server, const uba::tchar* filePath);
 	UBA_API uba::ProcessHandle* SessionServer_RunProcess(uba::SessionServer* server, uba::ProcessStartInfo& info, bool async, bool enableDetour);
@@ -105,4 +107,20 @@ extern "C"
 
 	UBA_API uba::ProcessStartInfo* CreateProcessStartInfo(const uba::tchar* application, const uba::tchar* arguments, const uba::tchar* workingDir, const uba::tchar* description, uba::u32 priorityClass, uba::u64 outputStatsThresholdMs, bool trackInputs, const uba::tchar* logFile, ProcessHandle_ExitCallback* exit);
 	UBA_API void DestroyProcessStartInfo(uba::ProcessStartInfo* info);
+
+
+	// Scheduler
+	UBA_API uba::Scheduler* Scheduler_Create(uba::SessionServer* session, uba::u32 maxLocalProcessors = 0);
+	UBA_API void Scheduler_Start(uba::Scheduler* scheduler);
+	UBA_API void Scheduler_EnqueueProcess(uba::Scheduler* scheduler, const uba::ProcessStartInfo& info, float weight = 1.0f, const void* knownInputs = nullptr, uba::u32 knownInputsBytes = 0, uba::u32 knownInputsCount = 0);
+	UBA_API void Scheduler_Stop(uba::Scheduler* scheduler);
+	UBA_API void Scheduler_Destroy(uba::Scheduler* scheduler);
+	UBA_API void Scheduler_GetStats(uba::Scheduler* scheduler, uba::u32& outQueued, uba::u32& outActive, uba::u32& outFinished);
+
+	// Misc
+	using Uba_CustomAssertHandler = void(const uba::tchar* text);
+	UBA_API void Uba_SetCustomAssertHandler(Uba_CustomAssertHandler* handler);
+
+	using ImportFunc = void(const uba::tchar* importName, void* userData);
+	UBA_API void Uba_FindImports(const uba::tchar* binary, ImportFunc* func, void* userData);
 }
diff --git a/Engine/Source/Programs/UnrealBuildAccelerator/Visualizer/Private/UbaVisualizer.cpp b/Engine/Source/Programs/UnrealBuildAccelerator/Visualizer/Private/UbaVisualizer.cpp
--- a/Engine/Source/Programs/UnrealBuildAccelerator/Visualizer/Private/UbaVisualizer.cpp
+++ b/Engine/Source/Programs/UnrealBuildAccelerator/Visualizer/Private/UbaVisualizer.cpp
@@ -607,7 +607,10 @@ namespace uba
 			}
 
 			if (m_visibleComponents[ComponentType_DetailedData])
-				PaintDetailedStats(hdc, posY, progressRect, session, processLocation.sessionIndex != 0, playTime);
+			{
+				auto drawText = [&](const StringBufferBase& text, RECT& rect) { DrawTextW(hdc, text.data, text.count, &rect, DT_SINGLELINE); };
+				PaintDetailedStats(posY, progressRect, session, processLocation.sessionIndex != 0, playTime, drawText);
+			}
 
 			if (m_visibleComponents[ComponentType_Bars])
 			{
@@ -1224,6 +1227,51 @@ namespace uba
 			MoveToEx(hdc, left, 2, NULL);
 			LineTo(hdc, left + 1, progressRect.bottom);
 		}
+		else if (m_fetchedFilesSelected != ~0u)
+		{
+			auto& session = m_traceView.sessions[m_fetchedFilesSelected];
+			auto& fetchedFiles = session.fetchedFiles;
+
+			int width = 1000;
+			int height = Min(int(clientRect.bottom), int(fetchedFiles.size() * m_popupFontHeight));
+
+
+			POINT p;
+			GetCursorPos(&p);
+			ScreenToClient(m_hwnd, &p);
+			RECT r;
+			r.left = p.x;
+			r.top = p.y;
+			r.right = r.left + width;
+			r.bottom = r.top + height;
+
+			if (r.right > clientRect.right)
+				OffsetRect(&r, -width, 0);
+			if (r.bottom > clientRect.bottom)
+			{
+				OffsetRect(&r, 0, -height);
+				if (r.top < 0)
+					OffsetRect(&r, 0, -r.top);
+			}
+			FillRect(hdc, &r, m_tooltipBackgroundBrush);
+
+			SelectObject(hdc, m_font);
+			DrawTextLogger logger(hdc, r, FontHeight);
+			for (auto& f : fetchedFiles)
+			{
+				if (f.hint == TC("KnownInput"))
+					continue;
+				if (logger.rect.top >= r.bottom - FontHeight)
+				{
+					logger.rect.top = r.top;
+					logger.rect.left += 500;
+					if (logger.rect.left >= r.right)
+						break;
+				}
+				logger.Info(L"%s", f.hint.c_str());
+			}
+			logger.Info(L"...");
+		}
 	}
 
 	void Visualizer::PaintProcessRect(TraceView::Process& process, HDC hdc, RECT rect, const RECT& progressRect, bool selected, bool writingBitmap)
@@ -1397,7 +1445,7 @@ namespace uba
 		*/
 	}
 
-	void Visualizer::PaintDetailedStats(HDC hdc, int& posY, const RECT& progressRect, TraceView::Session& session, bool isRemote, u64 playTime)
+	void Visualizer::PaintDetailedStats(int& posY, const RECT& progressRect, TraceView::Session& session, bool isRemote, u64 playTime, const DrawTextFunc& drawTextFunc)
 	{
 		int stepY = FontHeight;
 		int startPosY = posY;
@@ -1413,14 +1461,12 @@ namespace uba
 				textRect.top = posY;
 				textRect.bottom = posY + stepY;
 				posY += stepY;
-				if (!hdc)
-					return;
 				StringBuffer<> str;
 				va_list arg;
 				va_start(arg, format);
 				str.Append(format, arg);
 				va_end(arg);
-				DrawTextW(hdc, str.data, str.count, &textRect, DT_SINGLELINE);
+				drawTextFunc(str, textRect);
 			};
 
 		drawText(L"Finished: %u", session.processExitedCount);
@@ -1593,7 +1639,15 @@ namespace uba
 			}
 
 			if (m_visibleComponents[ComponentType_DetailedData])
-				PaintDetailedStats(0, posY, progressRect, session, i != 0, playTime);
+			{
+				auto drawText = [&](const StringBufferBase& text, RECT& rect)
+					{
+						//DrawTextW(hdc, text.data, text.count, &rect, DT_SINGLELINE);
+						if (pos.x >= rect.left && pos.x < rect.right && pos.y >= rect.top && pos.y < rect.bottom && text.StartsWith(TC("Fetched Files")))
+							outResult.fetchedFilesSelected = sessionIndex;
+					};
+				PaintDetailedStats(posY, progressRect, session, i != 0, playTime, drawText);
+			}
 
 			if (m_visibleComponents[ComponentType_Bars])
 			{
@@ -1710,7 +1764,8 @@ namespace uba
 		if (res.processSelected == m_processSelected && res.processLocation == m_processSelectedLocation &&
 			res.sessionSelectedIndex == m_sessionSelectedIndex &&
 			res.statsSelected == m_statsSelected && memcmp(&res.stats, &m_stats, sizeof(Stats)) == 0 &&
-			res.buttonSelected == m_buttonSelected && res.timelineSelected == m_timelineSelected)
+			res.buttonSelected == m_buttonSelected && res.timelineSelected == m_timelineSelected &&
+			res.fetchedFilesSelected == m_fetchedFilesSelected)
 			return false;
 		m_processSelected = res.processSelected;
 		m_processSelectedLocation = res.processLocation;
@@ -1719,6 +1774,7 @@ namespace uba
 		m_stats = res.stats;
 		m_buttonSelected = res.buttonSelected;
 		m_timelineSelected = res.timelineSelected;
+		m_fetchedFilesSelected = res.fetchedFilesSelected;
 		return true;
 	}
 
@@ -1965,13 +2021,14 @@ namespace uba
 			tme.hwndTrack = hWnd;
 			TrackMouseEvent(&tme);
 
-			if (m_processSelected || m_sessionSelectedIndex != ~0u || m_statsSelected || m_timelineSelected)
+			if (m_processSelected || m_sessionSelectedIndex != ~0u || m_statsSelected || m_timelineSelected || m_fetchedFilesSelected != ~0u)
 			{
 				m_processSelected = false;
 				m_sessionSelectedIndex = ~0u;
 				m_statsSelected = false;
 				m_buttonSelected = ~0u;
 				m_timelineSelected = 0;
+				m_fetchedFilesSelected = ~0u;
 				RedrawWindow(hWnd, NULL, NULL, RDW_INVALIDATE);
 			}
 			break;
@@ -1983,6 +2040,7 @@ namespace uba
 			m_statsSelected = false;
 			m_buttonSelected = ~0u;
 			m_timelineSelected = 0;
+			m_fetchedFilesSelected = ~0u;
 			m_autoScroll = false;
 			POINTS p = MAKEPOINTS(lParam);
 			m_mouseAnchor = {p.x, p.y};
diff --git a/Engine/Source/Programs/UnrealBuildAccelerator/Visualizer/Private/UbaVisualizer.h b/Engine/Source/Programs/UnrealBuildAccelerator/Visualizer/Private/UbaVisualizer.h
--- a/Engine/Source/Programs/UnrealBuildAccelerator/Visualizer/Private/UbaVisualizer.h
+++ b/Engine/Source/Programs/UnrealBuildAccelerator/Visualizer/Private/UbaVisualizer.h
@@ -30,7 +30,8 @@ namespace uba
 		void PaintAll(HDC hdc, const RECT& clientRect);
 		void PaintProcessRect(TraceView::Process& process, HDC hdc, RECT rect, const RECT& progressRect, bool selected, bool writingBitmap);
 		void PaintTimeline(HDC hdc, const RECT& clientRect);
-		void PaintDetailedStats(HDC hdc, int& posY, const RECT& progressRect, TraceView::Session& session, bool isRemote, u64 playTime);
+		using DrawTextFunc = Function<void(const StringBufferBase& text, RECT& rect)>;
+		void PaintDetailedStats(int& posY, const RECT& progressRect, TraceView::Session& session, bool isRemote, u64 playTime, const DrawTextFunc& drawTextFunc);
 
 		struct Stats
 		{
@@ -50,6 +51,7 @@ namespace uba
 			Stats stats;
 			u32 buttonSelected = ~0u;
 			float timelineSelected = 0;
+			u32 fetchedFilesSelected = ~0u;
 
 		};
 		void HitTest(HitTestResult& outResult, const POINT& pos);
@@ -134,6 +136,7 @@ namespace uba
 		Stats m_stats;
 		u32 m_buttonSelected = ~0u;
 		float m_timelineSelected = 0;
+		u32 m_fetchedFilesSelected = ~0u;
 		
 		bool m_mouseOverWindow = false;
 
