[UBA]
* Fixed macos compile error

diff --git a/Engine/Source/Programs/UnrealBuildAccelerator/Common/Private/UbaSessionClient.cpp b/Engine/Source/Programs/UnrealBuildAccelerator/Common/Private/UbaSessionClient.cpp
--- a/Engine/Source/Programs/UnrealBuildAccelerator/Common/Private/UbaSessionClient.cpp
+++ b/Engine/Source/Programs/UnrealBuildAccelerator/Common/Private/UbaSessionClient.cpp
@@ -227,7 +227,7 @@ namespace uba
 			if (moduleCount == 0)
 				return m_logger.Error(TC("Application %s not found"), application);
 
-			struct ModuleInfo { TString name; CasKey casKey; u32 attributes; Event done; };
+			struct ModuleInfo { ModuleInfo(const tchar* n, const CasKey& c, u32 a) : name(n), casKey(c), attributes(a), done(true) {} TString name; CasKey casKey; u32 attributes; Event done; };
 			List<ModuleInfo> modules;
 
 			Atomic<bool> success = true;
@@ -252,7 +252,7 @@ namespace uba
 					moduleFile.Clear().Append(localSystemModule);
 				}
 
-				auto& m = modules.emplace_back(moduleFile.data, casKey, fileAttributes, true);
+				auto& m = modules.emplace_back(moduleFile.data, casKey, fileAttributes);
 				m_client.AddWork([&]()
 					{
 						auto g = MakeGuard([&]() { m.done.Set(); });
