[UBA]
* Fixed so when downloading application modules, thread itself helps out downloading

diff --git a/Engine/Source/Programs/UnrealBuildAccelerator/Common/Private/UbaSessionClient.cpp b/Engine/Source/Programs/UnrealBuildAccelerator/Common/Private/UbaSessionClient.cpp
--- a/Engine/Source/Programs/UnrealBuildAccelerator/Common/Private/UbaSessionClient.cpp
+++ b/Engine/Source/Programs/UnrealBuildAccelerator/Common/Private/UbaSessionClient.cpp
@@ -231,7 +231,8 @@ namespace uba
 			List<ModuleInfo> modules;
 
 			Atomic<bool> success = true;
-
+			Atomic<u32> handledCount;
+			u32 workCount = 0;
 			while (moduleCount--)
 			{
 				StringBuffer<> moduleFile;
@@ -252,9 +253,12 @@ namespace uba
 					moduleFile.Clear().Append(localSystemModule);
 				}
 
+				++workCount;
+
 				auto& m = modules.emplace_back(moduleFile.data, casKey, fileAttributes);
 				m_client.AddWork([&]()
 					{
+						++handledCount;
 						auto g = MakeGuard([&]() { m.done.Set(); });
 						CasKey newCasKey;
 						bool storeUncompressed = true;
@@ -274,6 +278,9 @@ namespace uba
 					}, 1, TC("EnsureApp"));
 			}
 
+			while (handledCount < workCount)
+				m_client.DoWork();
+
 			// Wait for all to be done
 			for (auto& m : modules)
 				if (!m.done.IsSet(10 * 60 * 1000)) // 10 minutes is a very long time
diff --git a/Engine/Source/Programs/UnrealBuildAccelerator/Common/Private/UbaWorkManager.cpp b/Engine/Source/Programs/UnrealBuildAccelerator/Common/Private/UbaWorkManager.cpp
--- a/Engine/Source/Programs/UnrealBuildAccelerator/Common/Private/UbaWorkManager.cpp
+++ b/Engine/Source/Programs/UnrealBuildAccelerator/Common/Private/UbaWorkManager.cpp
@@ -115,6 +115,21 @@ namespace uba
 		--m_activeWorkerCount;
 	}
 
+	void WorkManagerImpl::DoWork(u32 count)
+	{
+		while (count--)
+		{
+			ScopedWriteLock lock(m_workLock);
+			if (m_work.empty())
+				break;
+			Work work = m_work.front();
+			m_work.pop_front();
+			lock.Leave();
+
+			work.func();
+		}
+	}
+
 	void WorkManagerImpl::FlushWork()
 	{
 		while (true)
diff --git a/Engine/Source/Programs/UnrealBuildAccelerator/Common/Public/UbaWorkManager.h b/Engine/Source/Programs/UnrealBuildAccelerator/Common/Public/UbaWorkManager.h
--- a/Engine/Source/Programs/UnrealBuildAccelerator/Common/Public/UbaWorkManager.h
+++ b/Engine/Source/Programs/UnrealBuildAccelerator/Common/Public/UbaWorkManager.h
@@ -21,7 +21,9 @@ namespace uba
 		virtual ~WorkManagerImpl();
 		virtual void AddWork(const Function<void()>& work, u32 count, const tchar* desc) override;
 		virtual u32 GetWorkerCount() override;
+		void DoWork(u32 count = 1);
 		void FlushWork();
+
 	private:
 		struct Worker;
 		void PushWorker(Worker* worker);
