[UBA]
* Implemented process reuse logic in UbaScheduler. It is now possible for running processes to fetch more work
* Fixed bugs in custom message path
* Fixed bugs in FlushWrittenFiles message
* Changed So UpdateEnvironment is resetting stats
* Fixed potential race condition related to process reuse
* Added special rule for ShaderCompileWorker.exe which detours ImageGetDigestStream in Imagehlp.dll (because wine implementation does not match windows implementation)
* Renamed "exitedUserData" to just userData since it is used for more than when exiting process
* Added a couple more dlls to "known system dlls"
* Fixed so visualizer can visualize process reuse properly

diff --git a/Engine/Source/Programs/UnrealBuildAccelerator/Common/Private/UbaProcess.cpp b/Engine/Source/Programs/UnrealBuildAccelerator/Common/Private/UbaProcess.cpp
--- a/Engine/Source/Programs/UnrealBuildAccelerator/Common/Private/UbaProcess.cpp
+++ b/Engine/Source/Programs/UnrealBuildAccelerator/Common/Private/UbaProcess.cpp
@@ -431,9 +431,13 @@ namespace uba
 
 		if (m_startInfo.exitedFunc)
 		{
+			auto exitedFunc = m_startInfo.exitedFunc;
+			auto userData = m_startInfo.userData;
+			m_startInfo.exitedFunc = nullptr;
+			m_startInfo.userData = nullptr;
 			ProcessHandle h;
 			h.m_process = this;
-			m_startInfo.exitedFunc(m_startInfo.exitedUserData, h);
+			exitedFunc(userData, h);
 			h.m_process = nullptr;
 		}
 	}
@@ -801,19 +805,7 @@ namespace uba
 					m_processStats.Add(stats);
 
 					if (GetApplicationRules()[m_rulesIndex].rules->IsExitCodeSuccess(m_nativeProcessExitCode) && !IsCancelled())
-					{
-						TimerScope ts(m_processStats.writeFiles);
-						ScopedWriteLock lock(m_writtenFilesLock);
-						for (auto& kv : m_writtenFiles)
-						{
-							if (kv.second.owner != this)
-								continue;
-							kv.second.owner = nullptr;
-							if (kv.second.mappingHandle.IsValid())
-								if (!m_session.WriteFileToDisk(*this, kv.second))
-									m_messageSuccess = false;
-						}
-					}
+						WriteFilesToDisk();
 
 					if (m_parentProcess)
 					{
@@ -833,19 +825,24 @@ namespace uba
 				}
 			case MessageType_Custom:
 				{
-					m_session.CustomMessage(reader, writer);
+					m_session.CustomMessage(*this, reader, writer);
 					return true;
 				}
 			case MessageType_FlushWrittenFiles:
 				{
-					return m_session.FlushWrittenFiles(*this);
+					WriteFilesToDisk();
+					bool result = m_session.FlushWrittenFiles(*this);
+					writer.WriteBool(result);
+					return true;
 				}
 
 			case MessageType_UpdateEnvironment:
 				{
 					StringBuffer<> reason;
 					reader.ReadString(reason);
-					return m_session.UpdateEnvironment(*this, reason.data);
+					bool result = m_session.UpdateEnvironment(*this, reason.data);
+					writer.WriteBool(result);
+					return true;
 				}
 		}
 		return m_session.m_logger.Error(TC("Unknown message type %u"), messageType);
@@ -981,6 +978,22 @@ namespace uba
 		}
 	}
 
+	bool ProcessImpl::WriteFilesToDisk()
+	{
+		TimerScope ts(m_processStats.writeFiles);
+		ScopedWriteLock lock(m_writtenFilesLock);
+		for (auto& kv : m_writtenFiles)
+		{
+			if (kv.second.owner != this)
+				continue;
+			kv.second.owner = nullptr;
+			if (kv.second.mappingHandle.IsValid())
+				if (!m_session.WriteFileToDisk(*this, kv.second))
+					return false;
+		}
+		return true;
+	}
+
 	const tchar* ProcessImpl::InternalGetChildLogFile(StringBufferBase& temp)
 	{
 		if (!*m_startInfo.logFile)
diff --git a/Engine/Source/Programs/UnrealBuildAccelerator/Common/Private/UbaScheduler.cpp b/Engine/Source/Programs/UnrealBuildAccelerator/Common/Private/UbaScheduler.cpp
--- a/Engine/Source/Programs/UnrealBuildAccelerator/Common/Private/UbaScheduler.cpp
+++ b/Engine/Source/Programs/UnrealBuildAccelerator/Common/Private/UbaScheduler.cpp
@@ -25,18 +25,23 @@ namespace uba
 	{
 		Scheduler* scheduler = nullptr;
 		ProcessStartInfo::ExitedCallback* originalExitedFunc = nullptr;
-		void* originalExitedData = nullptr;
+		void* originalUserData = nullptr;
 		const u8* knownInputs = nullptr;
 		u32 knownInputsCount = 0;
 		bool wasReturned = false;
 		bool isRemote = false;
 	};
 
-	Scheduler::Scheduler(SessionServer& session, u32 maxLocalProcessors)
+	Scheduler::Scheduler(SessionServer& session, u32 maxLocalProcessors, bool enableProcessReuse)
 	:	m_session(session)
 	,	m_maxLocalProcessors(maxLocalProcessors ? maxLocalProcessors : GetLogicalProcessorCount())
 	,	m_updateThreadLoop(false)
 	{
+		if (enableProcessReuse)
+			session.RegisterCustomService([this](Process& process, const void* recv, u32 recvSize, void* send, u32 sendCapacity)
+				{
+					return HandleReuseMessage(process, recv, recvSize, send, sendCapacity);
+				});
 	}
 
 	Scheduler::~Scheduler()
@@ -103,18 +108,18 @@ namespace uba
 
 	void Scheduler::RemoteProcessReturned(Process& process)
 	{
-		auto& ei = *(ExitProcessInfo*)process.GetStartInfo().exitedUserData;
+		auto& ei = *(ExitProcessInfo*)process.GetStartInfo().userData;
 
 		ScopedWriteLock lock(m_queuedProcessesLock);
 		ProcessStartInfo info = process.GetStartInfo();
 		info.exitedFunc = ei.originalExitedFunc;
-		info.exitedUserData = ei.originalExitedData;
+		info.userData = ei.originalUserData;
 		m_queuedProcesses.emplace_front(info, ei.knownInputs, ei.knownInputsCount);
 		ei.knownInputs = nullptr;
+		ei.wasReturned = true;
 		lock.Leave();
 
-		ei.wasReturned = true;
-		process.Cancel(true);
+		process.Cancel(true); // Cancel will call ProcessExited
 		--m_activeRemoteProcesses;
 		m_updateThreadLoop.Set();
 	}
@@ -126,13 +131,13 @@ namespace uba
 
 	void Scheduler::ProcessExited(ExitProcessInfo* info, const ProcessHandle& handle)
 	{
-		auto ig = MakeGuard([info]() { delete[] info->knownInputs;  delete info; });
+		auto ig = MakeGuard([info]() { delete[] info->knownInputs; delete info; });
 
 		if (info->wasReturned)
 			return;
 
-		if (info->originalExitedFunc)
-			info->originalExitedFunc(info->originalExitedData, handle);
+		if (auto func = info->originalExitedFunc)
+			func(info->originalUserData, handle);
 
 		++m_finishedProcesses;
 
@@ -154,21 +159,26 @@ namespace uba
 			return false;
 		auto si = m_queuedProcesses.front();
 		m_queuedProcesses.pop_front();
+		if (runLocal)
+			++m_activeLocalProcesses;
 		lock.Leave();
 
 		auto exitInfo = new ExitProcessInfo();
 		exitInfo->scheduler = this;
 		exitInfo->originalExitedFunc = si.startInfo.exitedFunc;
-		exitInfo->originalExitedData = si.startInfo.exitedUserData;
+		exitInfo->originalUserData = si.startInfo.userData;
 		exitInfo->knownInputs = si.knownInputs;
 		exitInfo->knownInputsCount = si.knownInputsCount;
 		exitInfo->isRemote = !runLocal;
-		si.startInfo.exitedUserData = exitInfo;
-		si.startInfo.exitedFunc = [](void* userData, const ProcessHandle& handle) { auto ei = (ExitProcessInfo*)userData; ei->scheduler->ProcessExited(ei, handle); };
+		si.startInfo.userData = exitInfo;
+		si.startInfo.exitedFunc = [](void* userData, const ProcessHandle& handle)
+			{
+				auto ei = (ExitProcessInfo*)userData;
+				ei->scheduler->ProcessExited(ei, handle);
+			};
 
 		if (runLocal)
 		{
-			++m_activeLocalProcesses;
 			m_session.RunProcess(si.startInfo);
 		}
 		else
@@ -178,4 +188,52 @@ namespace uba
 		}
 		return true;
 	}
+
+	u32 Scheduler::HandleReuseMessage(Process& process, const void* recv, u32 recvSize, void* send, u32 sendCapacity)
+	{
+		auto& currentStartInfo = process.GetStartInfo();
+		auto info = (ExitProcessInfo*)currentStartInfo.userData;
+		if (!info) // If null, process has already exited from some other thread
+			return 0;
+
+		// Call ExitedFunc and cleanup
+		if (auto func = info->originalExitedFunc)
+		{
+			ProcessHandle h;
+			h.m_process = &process;
+			func(info->originalUserData, h);
+			h.m_process = nullptr;
+		}
+		delete[] info->knownInputs;
+		info->knownInputs = nullptr;
+		info->originalExitedFunc = nullptr;
+		info->originalUserData = nullptr;
+
+		// Try to get queued process to send back
+		ScopedWriteLock lock(m_queuedProcessesLock);
+		if (info->wasReturned)
+			return 0;
+		if (m_queuedProcesses.empty())
+			return 0;
+		auto qp = m_queuedProcesses.front();
+		m_queuedProcesses.pop_front();
+		lock.Leave();
+
+		delete[] qp.knownInputs;
+		auto& si = qp.startInfo;
+
+		UBA_ASSERT(Equals(currentStartInfo.application, si.application));
+
+		// Move over exited func and user data for the queued process
+		info->originalExitedFunc = si.exitedFunc;
+		info->originalUserData = si.userData;
+
+		// TODO: Don't think we need to udpate the other parts of StartInfo
+
+		BinaryWriter writer((u8*)send, 0, sendCapacity);
+		writer.WriteString(si.arguments);
+		writer.WriteString(si.workingDir);
+		writer.WriteString(si.description);
+		return u32(writer.GetPosition());
+	}
 }
diff --git a/Engine/Source/Programs/UnrealBuildAccelerator/Common/Private/UbaSession.cpp b/Engine/Source/Programs/UnrealBuildAccelerator/Common/Private/UbaSession.cpp
--- a/Engine/Source/Programs/UnrealBuildAccelerator/Common/Private/UbaSession.cpp
+++ b/Engine/Source/Programs/UnrealBuildAccelerator/Common/Private/UbaSession.cpp
@@ -1270,10 +1270,9 @@ namespace uba
 		RegisterCreateFileForWrite(key, fixedFilePath.data, fixedFilePath.count, true);
 	}
 
-	void Session::RegisterCustomService(CustomServiceFunction* function, void* userData)
+	void Session::RegisterCustomService(CustomServiceFunction&& function)
 	{
 		m_customServiceFunction = function;
-		m_customServiceUserData = userData;
 	}
 
 	const tchar* Session::GetId() { return m_id.data; }
@@ -1968,14 +1967,14 @@ namespace uba
 		return true;
 	}
 
-	bool Session::CustomMessage(BinaryReader& reader, BinaryWriter& writer)
+	bool Session::CustomMessage(Process& process, BinaryReader& reader, BinaryWriter& writer)
 	{
 		if (!m_customServiceFunction)
 			return false;
 		u32 recvSize = reader.ReadU32();
 		u32* sendSize = (u32*)writer.AllocWrite(4);
 		void* sendData = writer.GetData() + writer.GetPosition();
-		u32 written = m_customServiceFunction({}, reader.GetPositionData(), recvSize, sendData, u32(writer.GetCapacityLeft()), m_customServiceUserData);
+		u32 written = m_customServiceFunction(process, reader.GetPositionData(), recvSize, sendData, u32(writer.GetCapacityLeft()));
 		*sendSize = written;
 		writer.AllocWrite(written);
 		return true;
@@ -1987,12 +1986,22 @@ namespace uba
 
 	bool Session::FlushWrittenFiles(ProcessImpl& process)
 	{
-		return false;
+		return true;
 	}
 
 	bool Session::UpdateEnvironment(ProcessImpl& process, const tchar* reason)
 	{
-		return false;
+		StackBinaryWriter<16 * 1024> writer;
+		process.m_processStats.Write(writer);
+		process.m_sessionStats.Write(writer);
+		process.m_storageStats.Write(writer);
+		process.m_systemStats.Write(writer);
+		m_trace.ProcessEnvironmentUpdated(process.GetId(), reason, writer.GetData(), writer.GetPosition());
+		process.m_processStats = {};
+		process.m_sessionStats = {};
+		process.m_storageStats = {};
+		process.m_systemStats = {};
+		return true;
 	}
 
 	void Session::PrintSessionStats(Logger& logger)
diff --git a/Engine/Source/Programs/UnrealBuildAccelerator/Common/Private/UbaSessionClient.cpp b/Engine/Source/Programs/UnrealBuildAccelerator/Common/Private/UbaSessionClient.cpp
--- a/Engine/Source/Programs/UnrealBuildAccelerator/Common/Private/UbaSessionClient.cpp
+++ b/Engine/Source/Programs/UnrealBuildAccelerator/Common/Private/UbaSessionClient.cpp
@@ -1452,7 +1452,7 @@ namespace uba
 					};
 
 					ExitedRec* exitedRec = new ExitedRec(*this, activeWeightLock, activeWeight, rec);
-					startInfo.exitedUserData = exitedRec;
+					startInfo.userData = exitedRec;
 					startInfo.exitedFunc = [](void* userData, const ProcessHandle& h)
 					{
 						auto er = (ExitedRec*)userData;
@@ -1662,12 +1662,13 @@ namespace uba
 		m_stats.Print(logger);
 	}
 
-	bool SessionClient::CustomMessage(BinaryReader& reader, BinaryWriter& writer)
+	bool SessionClient::CustomMessage(Process& process, BinaryReader& reader, BinaryWriter& writer)
 	{
 		StackBinaryWriter<SendMaxSize> msgWriter;
 		NetworkMessage msg(m_client, ServiceId, SessionMessageType_Custom, msgWriter);
 
 		u32 recvSize = reader.ReadU32();
+		msgWriter.WriteU32(process.GetId());
 		msgWriter.WriteU32(recvSize);
 		msgWriter.WriteBytes(reader.GetPositionData(), recvSize);
 
@@ -1699,6 +1700,12 @@ namespace uba
 		process.m_sessionStats.Write(writer);
 		process.m_storageStats.Write(writer);
 		process.m_systemStats.Write(writer);
+
+		process.m_processStats = {};
+		process.m_sessionStats = {};
+		process.m_storageStats = {};
+		process.m_systemStats = {};
+
 		StackBinaryReader<SendMaxSize> reader;
 		if (!msg.Send(reader, m_stats.customMsg))
 			return false;
diff --git a/Engine/Source/Programs/UnrealBuildAccelerator/Common/Private/UbaSessionServer.cpp b/Engine/Source/Programs/UnrealBuildAccelerator/Common/Private/UbaSessionServer.cpp
--- a/Engine/Source/Programs/UnrealBuildAccelerator/Common/Private/UbaSessionServer.cpp
+++ b/Engine/Source/Programs/UnrealBuildAccelerator/Common/Private/UbaSessionServer.cpp
@@ -45,21 +45,30 @@ namespace uba
 			else
 				m_done.Set();
 
-			if (startInfo.exitedFunc)
-			{
-				ProcessHandle h;
-				h.m_process = this;
-				startInfo.exitedFunc(startInfo.exitedUserData, h);
-				h.m_process = nullptr;
-				startInfo.exitedFunc = nullptr;
-			}
+			ProcessHandle h;
+			h.m_process = this;
+			CallProcessExit(h);
+			h.m_process = nullptr;
 		}
 
 		virtual const tchar* GetExecutingHost() const override { return m_executingHost.c_str(); }
 		virtual bool IsRemote() const override { return true; }
 		virtual bool IsChild() override { return false; }
 
+		void CallProcessExit(ProcessHandle& h)
+		{
+			if (!startInfo.exitedFunc)
+				return;
+			ScopedWriteLock lock(m_exitedLock);
+			auto exitedFunc = startInfo.exitedFunc;
+			auto userData = startInfo.userData;
+			startInfo.exitedFunc = nullptr;
+			startInfo.userData = nullptr;
+			exitedFunc(userData, h);
+		}
+
 		SessionServer* m_server;
+		ReaderWriterLock m_exitedLock;
 		u32 m_processId;
 		u32 m_exitCode = ~u32(0);
 		u64 m_processorTime = 0;
@@ -1030,12 +1039,7 @@ namespace uba
 				process.m_wallTime = processStats.wallTime;
 				process.m_server = nullptr;
 				process.m_done.Set();
-				
-				if (process.startInfo.exitedFunc)
-				{
-					process.startInfo.exitedFunc(process.startInfo.exitedUserData, h);
-					process.startInfo.exitedFunc = nullptr;
-				}
+				process.CallProcessExit(h);
 				return true;
 			}
 
@@ -1117,7 +1121,17 @@ namespace uba
 			}
 			case SessionMessageType_Custom:
 			{
-				CustomMessage(reader, writer);
+				u32 processId = reader.ReadU32();
+				ScopedWriteLock lock(m_processesLock);
+				auto findIt = m_processes.find(processId);
+				if (findIt == m_processes.end())
+					return m_logger.Error(TC("Failed to find process for id %u when receiving custom message"), processId);
+				ProcessHandle h(findIt->second);
+				lock.Leave();
+
+				auto& remoteProcess = *(RemoteProcess*)h.m_process;
+				ScopedWriteLock exitedLock(remoteProcess.m_exitedLock);
+				CustomMessage(remoteProcess, reader, writer);
 				return true;
 			}
 			case SessionMessageType_UpdateEnvironment:
diff --git a/Engine/Source/Programs/UnrealBuildAccelerator/Common/Public/UbaProcess.h b/Engine/Source/Programs/UnrealBuildAccelerator/Common/Public/UbaProcess.h
--- a/Engine/Source/Programs/UnrealBuildAccelerator/Common/Public/UbaProcess.h
+++ b/Engine/Source/Programs/UnrealBuildAccelerator/Common/Public/UbaProcess.h
@@ -66,6 +66,7 @@ namespace uba
 		void LogLine(bool printInSession, TString&& line, LogEntryType logType);
 		bool CreateTempFile(BinaryReader& reader, ProcHandle nativeProcessHandle, const tchar* application);
 		bool OpenTempFile(BinaryReader& reader, BinaryWriter& writer, const tchar* application);
+		bool WriteFilesToDisk();
 
 		void SetRulesIndex(const ProcessStartInfo& si);
 		const tchar* InternalGetChildLogFile(StringBufferBase& temp);
diff --git a/Engine/Source/Programs/UnrealBuildAccelerator/Common/Public/UbaProcessHandle.h b/Engine/Source/Programs/UnrealBuildAccelerator/Common/Public/UbaProcessHandle.h
--- a/Engine/Source/Programs/UnrealBuildAccelerator/Common/Public/UbaProcessHandle.h
+++ b/Engine/Source/Programs/UnrealBuildAccelerator/Common/Public/UbaProcessHandle.h
@@ -48,10 +48,11 @@ namespace uba
 		inline bool operator==(const ProcessHandle& o) const { return m_process == o.m_process; }
 		inline u64 GetHash() const { return u64(m_process); }
 
-	private:
 		ProcessHandle(Process* process);
+	private:
 		Process* m_process;
 		friend class ProcessImpl;
+		friend class Scheduler;
 		friend class Session;
 		friend class SessionClient;
 		friend class SessionServer;
diff --git a/Engine/Source/Programs/UnrealBuildAccelerator/Common/Public/UbaProcessStartInfo.h b/Engine/Source/Programs/UnrealBuildAccelerator/Common/Public/UbaProcessStartInfo.h
--- a/Engine/Source/Programs/UnrealBuildAccelerator/Common/Public/UbaProcessStartInfo.h
+++ b/Engine/Source/Programs/UnrealBuildAccelerator/Common/Public/UbaProcessStartInfo.h
@@ -30,7 +30,7 @@ namespace uba
 
 		using ExitedCallback = void(void* userData, const ProcessHandle&);
 		ExitedCallback* exitedFunc = nullptr;	// Callback for when process is done (it has already exited)
-		void* exitedUserData = nullptr;			// User data provided to exit callback
+		void* userData = nullptr;				// User data provided to exit callback
 
 		int uiLanguage = 1033;					// Internal use
 	};
diff --git a/Engine/Source/Programs/UnrealBuildAccelerator/Common/Public/UbaScheduler.h b/Engine/Source/Programs/UnrealBuildAccelerator/Common/Public/UbaScheduler.h
--- a/Engine/Source/Programs/UnrealBuildAccelerator/Common/Public/UbaScheduler.h
+++ b/Engine/Source/Programs/UnrealBuildAccelerator/Common/Public/UbaScheduler.h
@@ -17,7 +17,7 @@ namespace uba
 	class Scheduler
 	{
 	public:
-		Scheduler(SessionServer& session, u32 maxLocalProcessors = 0);
+		Scheduler(SessionServer& session, u32 maxLocalProcessors = 0, bool enableProcessReuse = false);
 		~Scheduler();
 
 		void Start();
@@ -34,6 +34,7 @@ namespace uba
 		void RemoteSlotAvailable();
 		void ProcessExited(ExitProcessInfo* info, const ProcessHandle& handle);
 		bool RunQueuedProcess(bool runLocal);
+		u32 HandleReuseMessage(Process& process, const void* recv, u32 recvSize, void* send, u32 sendCapacity);
 
 		SessionServer& m_session;
 		u32 m_maxLocalProcessors;
diff --git a/Engine/Source/Programs/UnrealBuildAccelerator/Common/Public/UbaSession.h b/Engine/Source/Programs/UnrealBuildAccelerator/Common/Public/UbaSession.h
--- a/Engine/Source/Programs/UnrealBuildAccelerator/Common/Public/UbaSession.h
+++ b/Engine/Source/Programs/UnrealBuildAccelerator/Common/Public/UbaSession.h
@@ -54,8 +54,8 @@ namespace uba
 		void RefreshDirectory(const tchar* dirName); // Tell uba a directory on disk has been changed by some other system while session is running
 		void RegisterNewFile(const tchar* filePath); // Tell uba a new file on disk has been added by some other system while session is running
 
-		using CustomServiceFunction = u32(const Guid& clientUid, const void* recv, u32 recvSize, void* send, u32 sendCapacity, void* userData);
-		void RegisterCustomService(CustomServiceFunction* function, void* userData = nullptr); // Register a custom service (that can be communicated with from the remote agents)
+		using CustomServiceFunction = Function<u32(Process& handle, const void* recv, u32 recvSize, void* send, u32 sendCapacity)>;
+		void RegisterCustomService(CustomServiceFunction&& function); // Register a custom service (that can be communicated with from the remote agents)
 
 		const tchar* GetId();			// Id for session. Will be "yymmdd_hhmmss" unless SessionCreateInfo.useUniqueId is set to false
 		u32 GetActiveProcessCount();	// Current active processes running inside session
@@ -113,7 +113,7 @@ namespace uba
 		virtual bool GetListDirectoryInfo(ListDirectoryResponse& out, tchar* dirName, const StringKey& dirKey);
 		virtual bool WriteFileToDisk(ProcessImpl& process, WrittenFile& file);
 		virtual bool AllocFailed(Process& process, const tchar* allocType, u32 error);
-		virtual bool CustomMessage(BinaryReader& reader, BinaryWriter& writer);
+		virtual bool CustomMessage(Process& process, BinaryReader& reader, BinaryWriter& writer);
 		virtual void FileEntryAdded(StringKey fileNameKey, u64 lastWritten, u64 size);
 		virtual bool FlushWrittenFiles(ProcessImpl& process);
 		virtual bool UpdateEnvironment(ProcessImpl& process, const tchar* reason);
@@ -220,8 +220,7 @@ namespace uba
 		Vector<tchar> m_environmentVariables;
 		UnorderedSet<const tchar*, HashStringNoCase, EqualStringNoCase> m_localEnvironmentVariables;
 
-		CustomServiceFunction* m_customServiceFunction = nullptr;
-		void* m_customServiceUserData = nullptr;
+		CustomServiceFunction m_customServiceFunction;
 
 		friend class ProcessImpl;
 	};
diff --git a/Engine/Source/Programs/UnrealBuildAccelerator/Common/Public/UbaSessionClient.h b/Engine/Source/Programs/UnrealBuildAccelerator/Common/Public/UbaSessionClient.h
--- a/Engine/Source/Programs/UnrealBuildAccelerator/Common/Public/UbaSessionClient.h
+++ b/Engine/Source/Programs/UnrealBuildAccelerator/Common/Public/UbaSessionClient.h
@@ -54,7 +54,7 @@ namespace uba
 		virtual bool WriteFileToDisk(ProcessImpl& process, WrittenFile& file) override;
 		virtual bool AllocFailed(Process& process, const tchar* allocType, u32 error) override;
 		virtual void PrintSessionStats(Logger& logger) override;
-		virtual bool CustomMessage(BinaryReader& reader, BinaryWriter& writer) override;
+		virtual bool CustomMessage(Process& process, BinaryReader& reader, BinaryWriter& writer) override;
 		virtual bool FlushWrittenFiles(ProcessImpl& process) override;
 		virtual bool UpdateEnvironment(ProcessImpl& process, const tchar* reason) override;
 
diff --git a/Engine/Source/Programs/UnrealBuildAccelerator/Core/Private/UbaApplicationRules.cpp b/Engine/Source/Programs/UnrealBuildAccelerator/Core/Private/UbaApplicationRules.cpp
--- a/Engine/Source/Programs/UnrealBuildAccelerator/Core/Private/UbaApplicationRules.cpp
+++ b/Engine/Source/Programs/UnrealBuildAccelerator/Core/Private/UbaApplicationRules.cpp
@@ -357,6 +357,14 @@ namespace uba
 #endif // #if PLATFORM_WINDOWS
 	};
 
+	class ApplicationRulesShaderCompileWorker : public ApplicationRules
+	{
+		virtual bool DetourImageGetDigestStream() override
+		{
+			return true;
+		}
+	};
+
 	const RulesRec* GetApplicationRules()
 	{
 		static RulesRec rules[]
@@ -382,6 +390,7 @@ namespace uba
 			{ TC("llvm-objcopy.exe"),			new ApplicationRulesLlvmObjCopyExe() },
 			{ TC("UnrealBuildTool.dll"),		new ApplicationRulesUBTDll() },
 			{ TC("PVS-Studio.exe"),				new ApplicationRulesPVSStudio() },
+			{ TC("ShaderCompileWorker.exe"),	new ApplicationRulesShaderCompileWorker() },
 			//{ L"MSBuild.dll"),				new ApplicationRules() },
 			//{ L"BreakpadSymbolEncoder.exe"),	new ApplicationRulesClang() },
 			//{ L"cmd.exe"),		new ApplicationRules() },
diff --git a/Engine/Source/Programs/UnrealBuildAccelerator/Core/Public/UbaApplicationRules.h b/Engine/Source/Programs/UnrealBuildAccelerator/Core/Public/UbaApplicationRules.h
--- a/Engine/Source/Programs/UnrealBuildAccelerator/Core/Public/UbaApplicationRules.h
+++ b/Engine/Source/Programs/UnrealBuildAccelerator/Core/Public/UbaApplicationRules.h
@@ -118,6 +118,12 @@ namespace uba
 		{
 			// Do nothing
 		}
+
+		// ImageGetDigestStream is buggy in wine so we have to detour it for ShaderCompileWorker
+		virtual bool DetourImageGetDigestStream()
+		{
+			return false;
+		}
 	};
 
 
diff --git a/Engine/Source/Programs/UnrealBuildAccelerator/Core/Public/UbaSynchronization.h b/Engine/Source/Programs/UnrealBuildAccelerator/Core/Public/UbaSynchronization.h
--- a/Engine/Source/Programs/UnrealBuildAccelerator/Core/Public/UbaSynchronization.h
+++ b/Engine/Source/Programs/UnrealBuildAccelerator/Core/Public/UbaSynchronization.h
@@ -16,6 +16,7 @@ namespace uba
 		AtomicU64(u64 initialValue = 0) : Atomic<u64>(initialValue) {}
 		AtomicU64(AtomicU64&& o) noexcept : Atomic<u64>(o.load()) {}
 		void operator=(u64 o) { store(o); }
+		void operator=(const AtomicU64& o) { store(o); }
 	};
 
 	class CriticalSection
diff --git a/Engine/Source/Programs/UnrealBuildAccelerator/Core/Public/UbaWinBinDependencyParser.h b/Engine/Source/Programs/UnrealBuildAccelerator/Core/Public/UbaWinBinDependencyParser.h
--- a/Engine/Source/Programs/UnrealBuildAccelerator/Core/Public/UbaWinBinDependencyParser.h
+++ b/Engine/Source/Programs/UnrealBuildAccelerator/Core/Public/UbaWinBinDependencyParser.h
@@ -5,7 +5,7 @@
 #include "UbaStringBuffer.h"
 #include <algorithm>
 
-#if !defined(IS_DETOURED_INCLUDE)
+#if !defined(UBA_IS_DETOURED_INCLUDE)
 #define True_CreateFileW ::CreateFileW
 #define True_CreateFileMappingW ::CreateFileMappingW
 #define True_MapViewOfFile ::MapViewOfFile
@@ -34,6 +34,7 @@ namespace uba
 			L"fwpuclnt.dll",
 			L"gdi32.dll",
 			L"gdi32full.dll",
+			L"glu32.dll",
 			L"imagehlp.dll",
 			L"imm32.dll",
 			L"iphlpapi.dll",
@@ -50,6 +51,7 @@ namespace uba
 			L"ole32.dll",
 			L"oleaut32.dll",
 			L"ondemandconnroutehelper.dll",
+			L"opengl32.dll",
 			L"powrprof.dll",
 			L"psapi.dll",
 			L"rasadhlp.dll",
diff --git a/Engine/Source/Programs/UnrealBuildAccelerator/Detours/Private/Windows/UbaDetoursFunctionsImagehlp.inl b/Engine/Source/Programs/UnrealBuildAccelerator/Detours/Private/Windows/UbaDetoursFunctionsImagehlp.inl
new file mode 100644
--- /dev/null
+++ b/Engine/Source/Programs/UnrealBuildAccelerator/Detours/Private/Windows/UbaDetoursFunctionsImagehlp.inl
@@ -0,0 +1,90 @@
+// Copyright Epic Games, Inc. All Rights Reserved.
+
+typedef PVOID DIGEST_HANDLE;
+typedef BOOL(WINAPI* DIGEST_FUNCTION) (DIGEST_HANDLE refdata, PBYTE pData, DWORD dwLength);
+
+using ImageGetDigestStreamFunc = BOOL(HANDLE FileHandle, DWORD DigestLevel, DIGEST_FUNCTION DigestFunction, DIGEST_HANDLE DigestHandle);
+ImageGetDigestStreamFunc* True_ImageGetDigestStream;
+
+BOOL Detoured_ImageGetDigestStream(HANDLE FileHandle, DWORD DigestLevel, DIGEST_FUNCTION DigestFunction, DIGEST_HANDLE DigestHandle)
+{
+	// This function is _very_ limited and does not cover all kinds of binaries. Only here to handle the dlls used by ShaderCompileWorker
+	// Expects CERT_PE_IMAGE_DIGEST_ALL_IMPORT_INFO | CERT_PE_IMAGE_DIGEST_RESOURCES
+	UBA_ASSERT(DigestLevel == (DWORD)(0x04 | 0x02));
+
+	auto ErrorInvalidParameter = []() { SetLastError(ERROR_INVALID_PARAMETER); return false; };
+
+	if (!FileHandle)
+		return ErrorInvalidParameter();
+	DWORD fileSize = GetFileSize(FileHandle, NULL);
+	if (fileSize == INVALID_FILE_SIZE || fileSize < sizeof(IMAGE_DOS_HEADER))
+		return ErrorInvalidParameter();
+
+	auto mappingHandle = CreateFileMappingW(FileHandle, NULL, PAGE_READONLY, 0, 0, NULL);
+	if (mappingHandle == INVALID_HANDLE_VALUE)
+		return ErrorInvalidParameter();
+	auto hg = MakeGuard([&]() { CloseHandle(mappingHandle); });
+
+	auto mem = (BYTE*)MapViewOfFile(mappingHandle, FILE_MAP_COPY, 0, 0, 0);
+	if (!mem)
+		return ErrorInvalidParameter();
+	auto mg = MakeGuard([&]() { UnmapViewOfFile(mem); });
+
+	auto& dosHeader = *(IMAGE_DOS_HEADER*)mem;
+	if (dosHeader.e_magic != IMAGE_DOS_SIGNATURE)
+		return ErrorInvalidParameter();
+	DWORD offset = dosHeader.e_lfanew;
+	if (!offset || offset > fileSize)
+		return ErrorInvalidParameter();
+
+	if (!DigestFunction(DigestHandle, mem, offset))
+		return false;
+
+	if (offset + sizeof(IMAGE_NT_HEADERS) > fileSize)
+		return ErrorInvalidParameter();
+
+	auto& ntHeader = *(IMAGE_NT_HEADERS*)(mem + offset);
+	if (ntHeader.Signature != IMAGE_NT_SIGNATURE)
+		return ErrorInvalidParameter();
+
+	ntHeader.OptionalHeader.CheckSum = 0;
+	if (ntHeader.OptionalHeader.Magic == 0x20b)
+		ntHeader.OptionalHeader.DataDirectory[4] = { 0, 0 }; // Debug table
+
+	DWORD headerSize = sizeof(ntHeader.Signature) + sizeof(ntHeader.FileHeader) + ntHeader.FileHeader.SizeOfOptionalHeader;
+	if (!DigestFunction(DigestHandle, mem + offset, headerSize))
+		return false;
+
+	offset += headerSize;
+	DWORD sectionCount = ntHeader.FileHeader.NumberOfSections;
+	DWORD sectionsSize = sectionCount * sizeof(IMAGE_SECTION_HEADER);
+	if (offset + sectionsSize > fileSize)
+		return ErrorInvalidParameter();
+
+	if (!DigestFunction(DigestHandle, mem + offset, sectionsSize))
+		return false;
+
+	auto sections = (IMAGE_SECTION_HEADER*)(mem + offset);
+
+	for (IMAGE_SECTION_HEADER* i = sections, *e = sections + sectionCount; i < e; ++i)
+		if (i->Characteristics & IMAGE_SCN_CNT_CODE)
+			if (!DigestFunction(DigestHandle, mem + i->PointerToRawData, i->SizeOfRawData))
+				return false;
+
+	auto ReportSection = [&](const char* name)
+		{
+			for (IMAGE_SECTION_HEADER* i = sections, *e = sections + sectionCount; i < e; ++i)
+				if (memcmp(i->Name, name, strlen(name)) == 0)
+					if (DWORD offset = i->PointerToRawData)
+						DigestFunction(DigestHandle, mem + offset, i->SizeOfRawData);
+		};
+
+	ReportSection(".rdata");
+	ReportSection(".data");
+	ReportSection(".pdata");
+	ReportSection(".debug");
+	ReportSection(".didat");
+	ReportSection(".rsrc");
+	ReportSection(".reloc");
+	return true;
+}
diff --git a/Engine/Source/Programs/UnrealBuildAccelerator/Detours/Private/Windows/UbaDetoursFunctionsWin.cpp b/Engine/Source/Programs/UnrealBuildAccelerator/Detours/Private/Windows/UbaDetoursFunctionsWin.cpp
--- a/Engine/Source/Programs/UnrealBuildAccelerator/Detours/Private/Windows/UbaDetoursFunctionsWin.cpp
+++ b/Engine/Source/Programs/UnrealBuildAccelerator/Detours/Private/Windows/UbaDetoursFunctionsWin.cpp
@@ -1,9 +1,10 @@
 // Copyright Epic Games, Inc. All Rights Reserved.
 
+#define UBA_IS_DETOURED_INCLUDE 1
+
 #include "UbaDetoursFunctionsWin.h"
 #include "UbaDetoursFileMappingTable.h"
-
-#define IS_DETOURED_INCLUDE 1
+#include "UbaDetoursApi.h"
 
 #if !defined(UBA_USE_MIMALLOC)
 #define True_malloc malloc
@@ -584,6 +585,7 @@ thread_local const wchar_t* t_createFileFileName;
 #include "UbaDetoursFunctionsNtDll.inl"
 #include "UbaDetoursFunctionsKernelBase.inl"
 #include "UbaDetoursFunctionsUcrtBase.inl"
+#include "UbaDetoursFunctionsImagehlp.inl"
 
 void DetourAttachFunction(void** trueFunc, void* detouredFunc, const char* funcName)
 {
@@ -639,7 +641,19 @@ int DetourAttachFunctions(bool runningRemote)
 		DETOURED_FUNCTIONS_SHLWAPI
 	}
 
-#undef DETOURED_FUNCTION
+	#undef DETOURED_FUNCTION
+
+
+	if (g_isRunningWine && g_rules->DetourImageGetDigestStream())
+	{
+		if (HMODULE moduleHandle = LoadLibraryW(L"Imagehlp.dll"))
+		{
+			True_ImageGetDigestStream = (ImageGetDigestStreamFunc*)GetProcAddress(moduleHandle, "ImageGetDigestStream");
+			DetourAttachFunction((PVOID*)&True_ImageGetDigestStream, Detoured_ImageGetDigestStream, "ImageGetDigestStream");
+		}
+	}
+
+
 
 	// Can't attach to these when running through debugger with some vs extensions (Microsoft child process debugging)
 #if UBA_DEBUG
@@ -968,7 +982,8 @@ void PostDeinit()
 extern "C"
 {
 	using namespace uba;
-	__declspec(dllexport) u32 UbaSendMessage(const void* send, u32 sendSize, void* recv, u32 recvCapacity)
+
+	UBA_DETOURED_API u32 UbaSendCustomMessage(const void* send, u32 sendSize, void* recv, u32 recvCapacity)
 	{
 		//TimerScope ts(g_stats.init);
 		ScopedWriteLock pcs(g_communicationLock);
@@ -984,27 +999,45 @@ extern "C"
 		return recvSize;
 	}
 
-	__declspec(dllexport) void UbaFlushWrittenFiles()
+	UBA_DETOURED_API bool UbaFlushWrittenFiles()
 	{
 		ScopedWriteLock pcs(g_communicationLock);
 		BinaryWriter writer;
 		writer.WriteByte(MessageType_FlushWrittenFiles);
 		writer.Flush();
+		BinaryReader reader;
+		return reader.ReadBool();
 	}
 
-	__declspec(dllexport) void UbaUpdateEnvironment(const wchar_t* reason)
+	UBA_DETOURED_API bool UbaUpdateEnvironment(const wchar_t* reason)
 	{
-		ScopedWriteLock pcs(g_communicationLock);
-		BinaryWriter writer;
-		writer.WriteByte(MessageType_UpdateEnvironment);
-		writer.WriteString(reason ? reason : L"");
-		writer.Flush();
-
+		{
+			ScopedWriteLock pcs(g_communicationLock);
+			BinaryWriter writer;
+			writer.WriteByte(MessageType_UpdateEnvironment);
+			writer.WriteString(reason ? reason : L"");
+			writer.Flush();
+			BinaryReader reader;
+			if (!reader.ReadBool())
+				return false;
+		}
 		Rpc_UpdateTables();
+		return true;
 	}
 
-	__declspec(dllexport) bool UbaRunningRemote()
+	UBA_DETOURED_API bool UbaRunningRemote()
 	{
 		return g_runningRemote;
 	}
+
+	UBA_DETOURED_API bool UbaRequestNextProcess(wchar_t* outArguments, u32 outArgumentsCapacity)
+	{
+		*outArguments = 0;
+		StackBinaryReader<1024> reader;
+		reader.SetSize(UbaSendCustomMessage(nullptr, 0, reader.buffer, 1024));
+		if (!reader.GetLeft())
+			return false;
+		reader.ReadString(outArguments, outArgumentsCapacity);
+		return true;
+	}
 }
diff --git a/Engine/Source/Programs/UnrealBuildAccelerator/Detours/Public/UbaDetoursApi.h b/Engine/Source/Programs/UnrealBuildAccelerator/Detours/Public/UbaDetoursApi.h
new file mode 100644
--- /dev/null
+++ b/Engine/Source/Programs/UnrealBuildAccelerator/Detours/Public/UbaDetoursApi.h
@@ -0,0 +1,50 @@
+// Copyright Epic Games, Inc. All Rights Reserved.
+
+#if UBA_IS_DETOURED_INCLUDE
+#define UBA_DETOURED_API __declspec(dllexport)
+#else
+#define UBA_DETOURED_API __declspec(dllimport)
+#endif
+
+extern "C"
+{
+	// Send custom message from detoured process all the way out to callback registered with RegisterCustomService
+	// Handles remote processes as well and network the message
+	UBA_DETOURED_API unsigned int UbaSendCustomMessage(const void* send, unsigned int sendSize, void* recv, unsigned int recvCapacity);
+
+	// Make sure all written files are flushed to destination (lots of files are kept in memory until process ends otherwise)
+	UBA_DETOURED_API bool UbaFlushWrittenFiles();
+
+	// Update environment. Making sure caches are up-to-date in detoured process.
+	// reason ends up in the visualizer for the next block of process
+	UBA_DETOURED_API bool UbaUpdateEnvironment(const wchar_t* reason);
+
+	// Returns true if process is running on a remote machine
+	UBA_DETOURED_API bool UbaRunningRemote();
+
+	// Using custom message inside UbaScheduler so can't be combined with RegisterCustomService.
+	UBA_DETOURED_API bool UbaRequestNextProcess(wchar_t* outArguments, unsigned int outArgumentsCapacity);
+}
+
+
+// Above functions can be found and used using code looking like this:
+
+#if 0
+// Windows
+static HMODULE UbaDetoursModule = GetModuleHandleW(L"UbaDetours.dll");
+if (UbaDetoursModule)
+{
+	using UbaFlushWrittenFilesFunc = bool();
+	using UbaRequestNextProcessFunc = bool(TCHAR* outArguments, uint32 outArgumentsCapacity);
+	using UbaUpdateEnvironmentFunc = bool(const wchar_t* reason);
+
+	static UbaFlushWrittenFilesFunc* FlushWrittenFiles = (UbaFlushWrittenFilesFunc*)(void*)GetProcAddress(UbaDetoursModule, "UbaFlushWrittenFiles");
+	static UbaRequestNextProcessFunc* RequestNextProcess = (UbaRequestNextProcessFunc*)(void*)GetProcAddress(UbaDetoursModule, "UbaRequestNextProcess");
+	static UbaUpdateEnvironmentFunc* UpdateEnvironment = (UbaUpdateEnvironmentFunc*)(void*)GetProcAddress(UbaDetoursModule, "UbaUpdateEnvironment");
+
+	// ...
+}
+
+// Linux/Mac not implemented yet
+...
+#endif
\ No newline at end of file
diff --git a/Engine/Source/Programs/UnrealBuildAccelerator/Host/Private/UbaExports.cpp b/Engine/Source/Programs/UnrealBuildAccelerator/Host/Private/UbaExports.cpp
--- a/Engine/Source/Programs/UnrealBuildAccelerator/Host/Private/UbaExports.cpp
+++ b/Engine/Source/Programs/UnrealBuildAccelerator/Host/Private/UbaExports.cpp
@@ -395,7 +395,11 @@ uba::StorageClient* CreateStorageClient(uba::NetworkClient& client, const uba::t
 
 	void SessionServer_RegisterCustomService(uba::SessionServer* server, SessionServer_CustomServiceFunction* function, void* userData)
 	{
-		server->RegisterCustomService(function, userData);
+		server->RegisterCustomService([function, userData](uba::Process& process, const void* recv, uba::u32 recvSize, void* send, uba::u32 sendCapacity)
+			{
+				uba::ProcessHandle h(&process);
+				return function(&h, recv, recvSize, send, sendCapacity, userData);
+			});
 	}
 
 	void DestroySessionServer(uba::SessionServer* server)
@@ -429,9 +433,9 @@ uba::StorageClient* CreateStorageClient(uba::NetworkClient& client, const uba::t
 		delete info;
 	}
 
-	uba::Scheduler* Scheduler_Create(uba::SessionServer* session, uba::u32 maxLocalProcessors)
+	uba::Scheduler* Scheduler_Create(uba::SessionServer* session, uba::u32 maxLocalProcessors, bool enableProcessReuse)
 	{
-		return new uba::Scheduler(*session, maxLocalProcessors);
+		return new uba::Scheduler(*session, maxLocalProcessors, enableProcessReuse);
 	}
 
 	void Scheduler_Start(uba::Scheduler* scheduler)
diff --git a/Engine/Source/Programs/UnrealBuildAccelerator/Host/Public/UbaExports.h b/Engine/Source/Programs/UnrealBuildAccelerator/Host/Public/UbaExports.h
--- a/Engine/Source/Programs/UnrealBuildAccelerator/Host/Public/UbaExports.h
+++ b/Engine/Source/Programs/UnrealBuildAccelerator/Host/Public/UbaExports.h
@@ -77,7 +77,7 @@ extern "C"
 
 	using SessionServer_RemoteProcessAvailableCallback = void(void* userData);
 	using SessionServer_RemoteProcessReturnedCallback = void(uba::Process& process, void* userData);
-	using SessionServer_CustomServiceFunction = uba::u32(const uba::Guid& clientUid, const void* recv, uba::u32 recvSize, void* send, uba::u32 sendCapacity, void* userData);
+	using SessionServer_CustomServiceFunction = uba::u32(uba::ProcessHandle* handle, const void* recv, uba::u32 recvSize, void* send, uba::u32 sendCapacity, void* userData);
 
 	UBA_API uba::SessionServerCreateInfo* CreateSessionServerCreateInfo(uba::Storage& storage, uba::NetworkServer& client, uba::LogWriter& writer, const uba::tchar* rootDir, const uba::tchar* traceOutputFile,
 		bool disableCustomAllocator, bool launchVisualizer, bool resetCas, bool writeToDisk, bool detailedTrace, bool allowWaitOnMem = false, bool allowKillOnMem = false);
@@ -110,7 +110,7 @@ extern "C"
 
 
 	// Scheduler
-	UBA_API uba::Scheduler* Scheduler_Create(uba::SessionServer* session, uba::u32 maxLocalProcessors = 0);
+	UBA_API uba::Scheduler* Scheduler_Create(uba::SessionServer* session, uba::u32 maxLocalProcessors = 0, bool enableProcessReuse = false);
 	UBA_API void Scheduler_Start(uba::Scheduler* scheduler);
 	UBA_API void Scheduler_EnqueueProcess(uba::Scheduler* scheduler, const uba::ProcessStartInfo& info, float weight = 1.0f, const void* knownInputs = nullptr, uba::u32 knownInputsBytes = 0, uba::u32 knownInputsCount = 0);
 	UBA_API void Scheduler_Stop(uba::Scheduler* scheduler);
diff --git a/Engine/Source/Programs/UnrealBuildAccelerator/Test/Private/UbaTestSession.h b/Engine/Source/Programs/UnrealBuildAccelerator/Test/Private/UbaTestSession.h
--- a/Engine/Source/Programs/UnrealBuildAccelerator/Test/Private/UbaTestSession.h
+++ b/Engine/Source/Programs/UnrealBuildAccelerator/Test/Private/UbaTestSession.h
@@ -182,15 +182,15 @@ namespace uba
 	bool RunCustomService(LoggerWithWriter& logger, SessionServer& session, const tchar* workingDir, const RunProcessFunction& runProcess)
 	{
 		bool gotMessage = false;
-		session.RegisterCustomService([](const Guid& connectionUid, const void* recv, u32 recvSize, void* send, u32 sendCapacity, void* gotMessagePtr)
+		session.RegisterCustomService([&](uba::Process& process, const void* recv, u32 recvSize, void* send, u32 sendCapacity)
 			{
-				*(bool*)gotMessagePtr = true;
+				gotMessage = true;
 				//wprintf(L"GOT MESSAGE: %.*s\n", recvSize / 2, (const wchar_t*)recv);
 				const wchar_t* hello = L"Hello response from server";
 				u64 helloBytes = wcslen(hello) * 2;
 				memcpy(send, hello, helloBytes);
 				return u32(helloBytes);
-			}, &gotMessage);
+			});
 
 		StringBuffer<> testApp;
 		GetTestAppPath(logger, testApp);
diff --git a/Engine/Source/Programs/UnrealBuildAccelerator/TestApp/Private/Windows/UbaTestAppWin.cpp b/Engine/Source/Programs/UnrealBuildAccelerator/TestApp/Private/Windows/UbaTestAppWin.cpp
--- a/Engine/Source/Programs/UnrealBuildAccelerator/TestApp/Private/Windows/UbaTestAppWin.cpp
+++ b/Engine/Source/Programs/UnrealBuildAccelerator/TestApp/Private/Windows/UbaTestAppWin.cpp
@@ -134,11 +134,11 @@ int wmain(int argc, wchar_t* argv[])
 	else
 	{
 		using u32 = unsigned int;
-		using UbaSendMessageFunc = u32(const void* send, u32 sendSize, void* recv, u32 recvCapacity);
+		using UbaSendCustomMessageFunc = u32(const void* send, u32 sendSize, void* recv, u32 recvCapacity);
 
-		UbaSendMessageFunc* sendMessage = (UbaSendMessageFunc*)GetProcAddress(detoursHandle, "UbaSendMessage");
+		UbaSendCustomMessageFunc* sendMessage = (UbaSendCustomMessageFunc*)GetProcAddress(detoursHandle, "UbaSendCustomMessage");
 		if (!sendMessage)
-			return LogError(L"Couldn't find UbaSendMessage function in UbaDetours.dll");
+			return LogError(L"Couldn't find UbaSendCustomMessage function in UbaDetours.dll");
 
 		const wchar_t* helloMsg = L"Hello from client";
 		wchar_t response[256];
diff --git a/Engine/Source/Programs/UnrealBuildAccelerator/Visualizer/Private/UbaTraceReader.cpp b/Engine/Source/Programs/UnrealBuildAccelerator/Visualizer/Private/UbaTraceReader.cpp
--- a/Engine/Source/Programs/UnrealBuildAccelerator/Visualizer/Private/UbaTraceReader.cpp
+++ b/Engine/Source/Programs/UnrealBuildAccelerator/Visualizer/Private/UbaTraceReader.cpp
@@ -538,12 +538,24 @@ namespace uba
 			TraceView::ProcessLocation active = findIt->second;
 			auto& session = GetSession(out, active.sessionIndex);
 
-			TraceView::Process& process = session.processors[active.processorIndex].processes[active.processIndex];
+			auto& processes = session.processors[active.processorIndex].processes;
+			TraceView::Process& process = processes[active.processIndex];
 			process.processStats.Read(reader, out.version);
 			process.sessionStats.Read(reader, out.version);
 			process.storageStats.Read(reader);
 			process.systemStats.Read(reader);
-			process.updates.push_back(TraceView::Process::Update{ time,reason.data });
+			process.exitCode = 0u;
+			process.stop = time;
+			process.bitmapDirty = true;
+
+			processes.emplace_back();
+			auto& process2 = processes.back();
+			m_activeProcesses[processId].processIndex = u32(processes.size() - 1);
+			process2.id = processId;
+			process2.description = reason.data;
+			process2.start = time;
+			process2.stop = ~u64(0);
+			process2.exitCode = ~0u;
 			break;
 		}
 		case TraceType_ProcessReturned:
diff --git a/Engine/Source/Programs/UnrealBuildAccelerator/Visualizer/Private/UbaTraceReader.h b/Engine/Source/Programs/UnrealBuildAccelerator/Visualizer/Private/UbaTraceReader.h
--- a/Engine/Source/Programs/UnrealBuildAccelerator/Visualizer/Private/UbaTraceReader.h
+++ b/Engine/Source/Programs/UnrealBuildAccelerator/Visualizer/Private/UbaTraceReader.h
@@ -27,8 +27,6 @@ namespace uba
 			SessionStats sessionStats;
 			StorageStats storageStats;
 			SystemStats systemStats;
-			struct Update { u64 time; TString reason; };
-			Vector<Update> updates;
 		};
 
 		struct Processor
diff --git a/Engine/Source/Programs/UnrealBuildAccelerator/Visualizer/Private/UbaVisualizer.cpp b/Engine/Source/Programs/UnrealBuildAccelerator/Visualizer/Private/UbaVisualizer.cpp
--- a/Engine/Source/Programs/UnrealBuildAccelerator/Visualizer/Private/UbaVisualizer.cpp
+++ b/Engine/Source/Programs/UnrealBuildAccelerator/Visualizer/Private/UbaVisualizer.cpp
@@ -456,6 +456,7 @@ namespace uba
 		SelectObject(textDC, m_font);
 		SelectObject(textDC, GetStockObject(NULL_BRUSH));
 		SetBkMode(textDC, TRANSPARENT);
+		SetBkColor(hdc, m_useDarkMode ? RGB(70, 70, 70) : RGB(180, 180, 180));
 
 		//TEXTMETRIC metric;
 		//GetTextMetrics(textDC, &metric);
@@ -519,10 +520,7 @@ namespace uba
 					rect.bottom = textBottom;
 
 					if (selected)
-					{
 						SetBkMode(hdc, OPAQUE);
-						SetBkColor(hdc, m_useDarkMode ? RGB(70, 70, 70) : RGB(180, 180, 180));
-					}
 					ExtTextOutW(hdc, 5, posY+2, ETO_CLIPPED, &rect, text.data, text.count, NULL);
 					if (selected)
 						SetBkMode(hdc, TRANSPARENT);
@@ -608,7 +606,15 @@ namespace uba
 
 			if (m_visibleComponents[ComponentType_DetailedData])
 			{
-				auto drawText = [&](const StringBufferBase& text, RECT& rect) { DrawTextW(hdc, text.data, text.count, &rect, DT_SINGLELINE); };
+				auto drawText = [&](const StringBufferBase& text, RECT& rect)
+					{
+						bool selected = m_fetchedFilesSelected == processLocation.sessionIndex && text.StartsWith(TC("Fetched Files"));
+						if (selected)
+							SetBkMode(hdc, OPAQUE);
+						DrawTextW(hdc, text.data, text.count, &rect, DT_SINGLELINE);
+						if (selected)
+							SetBkMode(hdc, TRANSPARENT);
+					};
 				PaintDetailedStats(posY, progressRect, session, processLocation.sessionIndex != 0, playTime, drawText);
 			}
 
@@ -1231,46 +1237,51 @@ namespace uba
 		{
 			auto& session = m_traceView.sessions[m_fetchedFilesSelected];
 			auto& fetchedFiles = session.fetchedFiles;
+			if (!fetchedFiles.empty() && !fetchedFiles[0].hint.empty())
+			{
+				int colWidth = 500;
+				int width = colWidth * 2;
+				int height = Min(int(clientRect.bottom), int(fetchedFiles.size() * m_popupFontHeight));
 
-			int width = 1000;
-			int height = Min(int(clientRect.bottom), int(fetchedFiles.size() * m_popupFontHeight));
-
-
-			POINT p;
-			GetCursorPos(&p);
-			ScreenToClient(m_hwnd, &p);
-			RECT r;
-			r.left = p.x;
-			r.top = p.y;
-			r.right = r.left + width;
-			r.bottom = r.top + height;
 
-			if (r.right > clientRect.right)
-				OffsetRect(&r, -width, 0);
-			if (r.bottom > clientRect.bottom)
-			{
-				OffsetRect(&r, 0, -height);
-				if (r.top < 0)
-					OffsetRect(&r, 0, -r.top);
-			}
-			FillRect(hdc, &r, m_tooltipBackgroundBrush);
+				POINT p;
+				GetCursorPos(&p);
+				ScreenToClient(m_hwnd, &p);
+				RECT r;
+				r.left = p.x;
+				r.top = p.y;
+				r.right = r.left + width;
+				r.bottom = r.top + height;
+
+				if (r.right > clientRect.right)
+					OffsetRect(&r, -width, 0);
+				if (r.bottom > clientRect.bottom)
+				{
+					OffsetRect(&r, 0, -height);
+					if (r.top < 0)
+						OffsetRect(&r, 0, -r.top);
+				}
+				FillRect(hdc, &r, m_tooltipBackgroundBrush);
 
-			SelectObject(hdc, m_font);
-			DrawTextLogger logger(hdc, r, FontHeight);
-			for (auto& f : fetchedFiles)
-			{
-				if (f.hint == TC("KnownInput"))
-					continue;
-				if (logger.rect.top >= r.bottom - FontHeight)
+				SelectObject(hdc, m_font);
+				DrawTextLogger logger(hdc, r, FontHeight);
+				for (auto& f : fetchedFiles)
 				{
-					logger.rect.top = r.top;
-					logger.rect.left += 500;
-					if (logger.rect.left >= r.right)
-						break;
+					if (f.hint == TC("KnownInput"))
+						continue;
+					if (logger.rect.top >= r.bottom - FontHeight)
+					{
+						if (logger.rect.left + colWidth >= r.right)
+						{
+							logger.Info(L"...");
+							break;
+						}
+						logger.rect.top = r.top;
+						logger.rect.left += colWidth;
+					}
+					logger.Info(L"%s", f.hint.c_str());
 				}
-				logger.Info(L"%s", f.hint.c_str());
 			}
-			logger.Info(L"...");
 		}
 	}
 
