[UBA]
* Changed so UbaRequestNextProcess is flushing written files and updating environment automatically (can't think of a reason when this is not desired)

diff --git a/Engine/Source/Programs/UnrealBuildAccelerator/Detours/Private/Windows/UbaDetoursFunctionsWin.cpp b/Engine/Source/Programs/UnrealBuildAccelerator/Detours/Private/Windows/UbaDetoursFunctionsWin.cpp
--- a/Engine/Source/Programs/UnrealBuildAccelerator/Detours/Private/Windows/UbaDetoursFunctionsWin.cpp
+++ b/Engine/Source/Programs/UnrealBuildAccelerator/Detours/Private/Windows/UbaDetoursFunctionsWin.cpp
@@ -1032,12 +1032,19 @@ extern "C"
 
 	UBA_DETOURED_API bool UbaRequestNextProcess(wchar_t* outArguments, u32 outArgumentsCapacity)
 	{
+		if (!UbaFlushWrittenFiles())
+			return false;
+
 		*outArguments = 0;
-		StackBinaryReader<1024> reader;
-		reader.SetSize(UbaSendCustomMessage(nullptr, 0, reader.buffer, 1024));
+		StackBinaryReader<2048> reader;
+		reader.SetSize(UbaSendCustomMessage(nullptr, 0, reader.buffer, 2048));
 		if (!reader.GetLeft())
 			return false;
 		reader.ReadString(outArguments, outArgumentsCapacity);
-		return true;
+		reader.SkipString(); // Working dir
+		StringBuffer<256> description;
+		reader.ReadString(description);
+
+		return UbaUpdateEnvironment(description.data);
 	}
 }
diff --git a/Engine/Source/Programs/UnrealBuildAccelerator/Detours/Public/UbaDetoursApi.h b/Engine/Source/Programs/UnrealBuildAccelerator/Detours/Public/UbaDetoursApi.h
--- a/Engine/Source/Programs/UnrealBuildAccelerator/Detours/Public/UbaDetoursApi.h
+++ b/Engine/Source/Programs/UnrealBuildAccelerator/Detours/Public/UbaDetoursApi.h
@@ -23,6 +23,7 @@ extern "C"
 	UBA_DETOURED_API bool UbaRunningRemote();
 
 	// Using custom message inside UbaScheduler so can't be combined with RegisterCustomService.
+	// This function will automatically flush written files before requesting next process and will also update environment if new process is retrieved
 	UBA_DETOURED_API bool UbaRequestNextProcess(wchar_t* outArguments, unsigned int outArgumentsCapacity);
 }
 
