[UBA]
* Fixed bug with detoured process requesting more processes if custom handler is not registered

diff --git a/Engine/Source/Programs/UnrealBuildAccelerator/Common/Private/UbaSession.cpp b/Engine/Source/Programs/UnrealBuildAccelerator/Common/Private/UbaSession.cpp
--- a/Engine/Source/Programs/UnrealBuildAccelerator/Common/Private/UbaSession.cpp
+++ b/Engine/Source/Programs/UnrealBuildAccelerator/Common/Private/UbaSession.cpp
@@ -1969,12 +1969,12 @@ namespace uba
 
 	bool Session::CustomMessage(Process& process, BinaryReader& reader, BinaryWriter& writer)
 	{
-		if (!m_customServiceFunction)
-			return false;
 		u32 recvSize = reader.ReadU32();
 		u32* sendSize = (u32*)writer.AllocWrite(4);
 		void* sendData = writer.GetData() + writer.GetPosition();
-		u32 written = m_customServiceFunction(process, reader.GetPositionData(), recvSize, sendData, u32(writer.GetCapacityLeft()));
+		u32 written = 0;
+		if (m_customServiceFunction)
+			written = m_customServiceFunction(process, reader.GetPositionData(), recvSize, sendData, u32(writer.GetCapacityLeft()));
 		*sendSize = written;
 		writer.AllocWrite(written);
 		return true;
