[UBA]
* Fixed so paths returned from SearchPathForFile are cleaned up (not having .. etc)
* Fixed so UbaAgent handle WSAPoll POLLERR as timeout so it retries when failing (this made the process exit on wine)
* Fixed so returned processes that had been reused returned the reused process to the scheduler
* Fixed so EnsureBinaryFile is using correct fileNameKeys
* Removed assert in detoured GetFullPathNameA since it is always calling GetFullPathNameW

diff --git a/Engine/Source/Programs/UnrealBuildAccelerator/Common/Private/UbaFile.cpp b/Engine/Source/Programs/UnrealBuildAccelerator/Common/Private/UbaFile.cpp
--- a/Engine/Source/Programs/UnrealBuildAccelerator/Common/Private/UbaFile.cpp
+++ b/Engine/Source/Programs/UnrealBuildAccelerator/Common/Private/UbaFile.cpp
@@ -1,8 +1,11 @@
 // Copyright Epic Games, Inc. All Rights Reserved.
 
+#define Local_GetLongPathNameW uba::GetLongPathNameW
+
 #include "UbaFile.h"
 #include "UbaDirectoryIterator.h"
 #include "UbaEvent.h"
+#include "UbaPathUtils.h"
 #include "UbaProcessStats.h"
 
 #if PLATFORM_WINDOWS
@@ -336,7 +339,7 @@ namespace uba
 		fullPath.Append(file);
 		if (GetFileAttributesW(fullPath.data) != INVALID_FILE_ATTRIBUTES)
 		{
-			out.Append(fullPath);
+			FixPath(fullPath.data, nullptr, 0, out);
 			return true;
 		}
 
@@ -369,7 +372,7 @@ namespace uba
 			fullPath.Append(file);
 			if (GetFileAttributesW(fullPath.data) != INVALID_FILE_ATTRIBUTES)
 			{
-				out.Append(fullPath);
+				FixPath(fullPath.data, nullptr, 0, out);
 				return true;
 			}
 
diff --git a/Engine/Source/Programs/UnrealBuildAccelerator/Common/Private/UbaNetworkBackendTcp.cpp b/Engine/Source/Programs/UnrealBuildAccelerator/Common/Private/UbaNetworkBackendTcp.cpp
--- a/Engine/Source/Programs/UnrealBuildAccelerator/Common/Private/UbaNetworkBackendTcp.cpp
+++ b/Engine/Source/Programs/UnrealBuildAccelerator/Common/Private/UbaNetworkBackendTcp.cpp
@@ -594,14 +594,15 @@ namespace uba
 			return false;
 		}
 
-		if (!pollRes || p.revents & (POLLHUP)) // Treat hangup as timeout (since we want retry if that happens)
+		u16 validFlags = POLLERR | POLLHUP; // Treat hangup as timeout (since we want retry if that happens). Also treat error as timeout. This is needed for Wine agent to be able to retry
+		if (!pollRes || p.revents & validFlags)
 		{
 			if (timedOut)
 				*timedOut = true;
 			return false;
 		}
 
-		if (p.revents & (POLLERR | POLLNVAL))
+		if (p.revents & POLLNVAL)
 		{
 			logger.Warning(TC("WSAPoll returned successful but with unexpected flags: %u"), p.revents);
 			return false;
diff --git a/Engine/Source/Programs/UnrealBuildAccelerator/Common/Private/UbaScheduler.cpp b/Engine/Source/Programs/UnrealBuildAccelerator/Common/Private/UbaScheduler.cpp
--- a/Engine/Source/Programs/UnrealBuildAccelerator/Common/Private/UbaScheduler.cpp
+++ b/Engine/Source/Programs/UnrealBuildAccelerator/Common/Private/UbaScheduler.cpp
@@ -30,6 +30,11 @@ namespace uba
 		u32 knownInputsCount = 0;
 		bool wasReturned = false;
 		bool isRemote = false;
+
+		// If reusing processes then these are the actual info of the process
+		TString lastArguments;
+		TString lastWorkingDir;
+		TString lastDescription;
 	};
 
 	Scheduler::Scheduler(SessionServer& session, u32 maxLocalProcessors, bool enableProcessReuse)
@@ -114,7 +119,14 @@ namespace uba
 		ProcessStartInfo info = process.GetStartInfo();
 		info.exitedFunc = ei.originalExitedFunc;
 		info.userData = ei.originalUserData;
+		if (!ei.lastArguments.empty())
+		{
+			info.arguments = ei.lastArguments.c_str();
+			info.workingDir = ei.lastWorkingDir.c_str();
+			info.description = ei.lastDescription.c_str();
+		}
 		m_queuedProcesses.emplace_front(info, ei.knownInputs, ei.knownInputsCount);
+		ei.knownInputsCount = 0;
 		ei.knownInputs = nullptr;
 		ei.wasReturned = true;
 		lock.Leave();
@@ -199,16 +211,22 @@ namespace uba
 		// Call ExitedFunc and cleanup
 		if (auto func = info->originalExitedFunc)
 		{
+			UBA_ASSERT(!info->wasReturned);
 			ProcessHandle h;
 			h.m_process = &process;
 			func(info->originalUserData, h);
 			h.m_process = nullptr;
 		}
+		info->knownInputsCount = 0;
 		delete[] info->knownInputs;
 		info->knownInputs = nullptr;
 		info->originalExitedFunc = nullptr;
 		info->originalUserData = nullptr;
 
+		info->lastArguments.clear();
+		info->lastWorkingDir.clear();
+		info->lastDescription.clear();
+
 		// Try to get queued process to send back
 		ScopedWriteLock lock(m_queuedProcessesLock);
 		if (info->wasReturned)
@@ -227,6 +245,9 @@ namespace uba
 		// Move over exited func and user data for the queued process
 		info->originalExitedFunc = si.exitedFunc;
 		info->originalUserData = si.userData;
+		info->lastArguments = si.arguments;
+		info->lastWorkingDir = si.workingDir;
+		info->lastDescription = si.description;
 
 		// TODO: Don't think we need to udpate the other parts of StartInfo
 
diff --git a/Engine/Source/Programs/UnrealBuildAccelerator/Common/Private/UbaSessionClient.cpp b/Engine/Source/Programs/UnrealBuildAccelerator/Common/Private/UbaSessionClient.cpp
--- a/Engine/Source/Programs/UnrealBuildAccelerator/Common/Private/UbaSessionClient.cpp
+++ b/Engine/Source/Programs/UnrealBuildAccelerator/Common/Private/UbaSessionClient.cpp
@@ -143,8 +143,11 @@ namespace uba
 	{
 		CasKey casKey;
 		u32 fileAttributes = DefaultAttributes(); // TODO: This is wrong.. need to retrieve from server if this is executable or not
-		if (fileName[1] == ':')
+
+		bool isAbsolute = IsWindows ? (fileName.count > 1 && fileName[1] == ':') : (fileName.count > 0 && fileName[0] == '/');
+		if (isAbsolute)
 		{
+			UBA_ASSERT(fileNameKey != StringKeyZero);
 			if (!GetCasKeyForFile(casKey, processId, fileName, fileNameKey))
 				return false;
 			// This needs to be absolute virtual path (the path on the host).. don't remember why this was written like this. Changed to just use fileName (needed for shadercompilerworker)
@@ -154,6 +157,7 @@ namespace uba
 		}
 		else
 		{
+			UBA_ASSERT(fileNameKey == StringKeyZero);
 			StackBinaryWriter<1024> writer;
 			NetworkMessage msg(m_client, ServiceId, SessionMessageType_EnsureBinaryFile, writer);
 			writer.WriteU32(processId);
@@ -180,7 +184,7 @@ namespace uba
 		if (!RetrieveCasFile(newKey, fileSize, casKey, outVirtual.data, storeUncompressed))
 			UBA_ASSERTF(false, TC("Casfile not found for %s using %s"), outVirtual.data, CasKeyString(casKey).str);
 		StringBuffer<> destFile;
-		if (fileName[1] == ':')
+		if (isAbsolute || fileName.Contains(TC(".."))) // This is not beautiful, but we need to keep some dlls in the sub folder (for cl.exe etc)
 			destFile.AppendFileName(fileName.data);
 		else
 			destFile.Append(fileName);
@@ -726,9 +730,13 @@ namespace uba
 		if (!EnsureBinaryFile(out.fileName, out.virtualFileName, msg.process.m_id, msg.fileName, msg.fileNameKey, dir.c_str()))
 			return false;
 
+		StringKey fileNameKey = msg.fileNameKey;
+		if (fileNameKey == StringKeyZero)
+			fileNameKey =  CaseInsensitiveFs ? ToStringKeyLower(out.virtualFileName) : ToStringKey(out.virtualFileName);
+
 		rec.name = out.fileName.data;
 		rec.virtualName = out.virtualFileName.data;
-		out.mappedFileTableSize = AddFileMapping(msg.fileNameKey, msg.fileName.data, out.fileName.data);
+		out.mappedFileTableSize = AddFileMapping(fileNameKey, msg.fileName.data, out.fileName.data);
 		return true;
 	}
 
diff --git a/Engine/Source/Programs/UnrealBuildAccelerator/Common/Private/UbaSessionServer.cpp b/Engine/Source/Programs/UnrealBuildAccelerator/Common/Private/UbaSessionServer.cpp
--- a/Engine/Source/Programs/UnrealBuildAccelerator/Common/Private/UbaSessionServer.cpp
+++ b/Engine/Source/Programs/UnrealBuildAccelerator/Common/Private/UbaSessionServer.cpp
@@ -626,7 +626,7 @@ namespace uba
 				StringBuffer<> lookupStr;
 				lookupStr.Append(fileName).Append(applicationDir).Append('#');
 				lookupStr.MakeLower();
-				StringKey lookupKey = ToStringKey(lookupStr);
+				StringKey lookupKey = ToStringKeyNoCheck(lookupStr.data, lookupStr.count);
 
 				ScopedWriteLock lock(m_applicationDataLock);
 				auto insres = m_applicationData.try_emplace(lookupKey);
@@ -644,8 +644,11 @@ namespace uba
 				StringBuffer<> absoluteFile;
 				if (SearchPathForFile(m_logger, absoluteFile, fileName.data, applicationDir.data))
 					if (!absoluteFile.StartsWith(m_systemPath.data) || !IsKnownSystemFile(absoluteFile.data))
+					{
+						fileNameKey = ToStringKeyLower(absoluteFile);
 						if (!StoreCasFile(casKey, fileNameKey, absoluteFile.data))
 							return false;
+					}
 
 				u64 startPos = writer.GetPosition();
 				writer.WriteCasKey(casKey);
diff --git a/Engine/Source/Programs/UnrealBuildAccelerator/Detours/Private/UbaDetoursFileMappingTable.cpp b/Engine/Source/Programs/UnrealBuildAccelerator/Detours/Private/UbaDetoursFileMappingTable.cpp
--- a/Engine/Source/Programs/UnrealBuildAccelerator/Detours/Private/UbaDetoursFileMappingTable.cpp
+++ b/Engine/Source/Programs/UnrealBuildAccelerator/Detours/Private/UbaDetoursFileMappingTable.cpp
@@ -212,11 +212,17 @@ namespace uba
 
 	void Rpc_GetFullFileName(const tchar*& path, u64& pathLen, StringBufferBase& tempBuf, bool useVirtualName)
 	{
-		tempBuf.Append(path, pathLen);
-		if (CaseInsensitiveFs)
-			tempBuf.MakeLower();
-		CHECK_PATH(tempBuf.data);
-		StringKey fileNameKey = ToStringKey(tempBuf);
+		StringKey fileNameKey;
+		bool isAbsolute = IsWindows ? (pathLen > 1 && path[1] == ':') : (pathLen > 0 && path[0] == '/');
+		if (isAbsolute)
+		{
+			tempBuf.Append(path, pathLen);
+			if (CaseInsensitiveFs)
+				tempBuf.MakeLower();
+			fileNameKey = ToStringKey(tempBuf);
+			tempBuf.Clear();
+		}
+
 		u32 mappedFileTableSize;
 
 		{
@@ -228,7 +234,6 @@ namespace uba
 			writer.WriteStringKey(fileNameKey);
 			writer.Flush();
 			BinaryReader reader;
-			tempBuf.Clear();
 			reader.ReadString(tempBuf);
 			if (useVirtualName)
 			{
diff --git a/Engine/Source/Programs/UnrealBuildAccelerator/Detours/Private/Windows/UbaDetoursFunctionsKernelBase.inl b/Engine/Source/Programs/UnrealBuildAccelerator/Detours/Private/Windows/UbaDetoursFunctionsKernelBase.inl
--- a/Engine/Source/Programs/UnrealBuildAccelerator/Detours/Private/Windows/UbaDetoursFunctionsKernelBase.inl
+++ b/Engine/Source/Programs/UnrealBuildAccelerator/Detours/Private/Windows/UbaDetoursFunctionsKernelBase.inl
@@ -2838,9 +2838,9 @@ DWORD Detoured_GetLongPathNameA(LPCSTR lpszShortPath, LPSTR lpszLongPath, DWORD
 
 DWORD Detoured_GetFullPathNameA(LPCSTR lpFileName, DWORD nBufferLength, LPSTR lpBuffer, LPSTR* lpFilePart)
 {
+	// Is verified that both windows and wine are calling GetFullPathNameW
 	DETOURED_CALL(GetFullPathNameA);
 	DEBUG_LOG_TRUE(L"GetFullPathNameA", L"");
-	UBA_ASSERT(!g_runningRemote || g_isRunningWine); // Wine is calling GetFullPathNameW.. need to check so windows is doing that too
 	return True_GetFullPathNameA(lpFileName, nBufferLength, lpBuffer, lpFilePart);
 }
 
