[Uba]
* Added RegisterDeleteFile to external api that should be used when file is deleted outside of uba but uba needs to know about it
* Reduced lock scope around FlushDeadProcesses
* Fixed disconnect issue crash in ubaagent that could cause access violation in UbaStorage. Fix was to cleanup when leaving StorageClient::SendAllSegments
* Changed so UbaRequestNextProcess is doing a full environment update (not resetting stats) if no next process is found
* Fixed so Scheduler enableProcessReuse=false works properly
* Improved some assert descriptions
* Fixed so files in m_outputFiles in session client is erased when flushed to server
* Fixed so Rpc_GetFullName cleans up .. in paths
* Fixed so files that existed but has been deleted by external process is seen as not existing by remote detoured process.

diff --git a/Engine/Source/Programs/UnrealBuildAccelerator/Common/Private/UbaProcess.cpp b/Engine/Source/Programs/UnrealBuildAccelerator/Common/Private/UbaProcess.cpp
--- a/Engine/Source/Programs/UnrealBuildAccelerator/Common/Private/UbaProcess.cpp
+++ b/Engine/Source/Programs/UnrealBuildAccelerator/Common/Private/UbaProcess.cpp
@@ -840,7 +840,8 @@ namespace uba
 				{
 					StringBuffer<> reason;
 					reader.ReadString(reason);
-					bool result = m_session.UpdateEnvironment(*this, reason.data);
+					bool resetStats = reader.ReadBool();
+					bool result = m_session.UpdateEnvironment(*this, reason.data, resetStats);
 					writer.WriteBool(result);
 					return true;
 				}
diff --git a/Engine/Source/Programs/UnrealBuildAccelerator/Common/Private/UbaScheduler.cpp b/Engine/Source/Programs/UnrealBuildAccelerator/Common/Private/UbaScheduler.cpp
--- a/Engine/Source/Programs/UnrealBuildAccelerator/Common/Private/UbaScheduler.cpp
+++ b/Engine/Source/Programs/UnrealBuildAccelerator/Common/Private/UbaScheduler.cpp
@@ -41,12 +41,12 @@ namespace uba
 	:	m_session(session)
 	,	m_maxLocalProcessors(maxLocalProcessors ? maxLocalProcessors : GetLogicalProcessorCount())
 	,	m_updateThreadLoop(false)
+	,	m_enableProcessReuse(enableProcessReuse)
 	{
-		if (enableProcessReuse)
-			session.RegisterCustomService([this](Process& process, const void* recv, u32 recvSize, void* send, u32 sendCapacity)
-				{
-					return HandleReuseMessage(process, recv, recvSize, send, sendCapacity);
-				});
+		session.RegisterCustomService([this](Process& process, const void* recv, u32 recvSize, void* send, u32 sendCapacity)
+			{
+				return HandleReuseMessage(process, recv, recvSize, send, sendCapacity);
+			});
 	}
 
 	Scheduler::~Scheduler()
@@ -203,6 +203,9 @@ namespace uba
 
 	u32 Scheduler::HandleReuseMessage(Process& process, const void* recv, u32 recvSize, void* send, u32 sendCapacity)
 	{
+		if (!m_enableProcessReuse)
+			return 0;
+
 		auto& currentStartInfo = process.GetStartInfo();
 		auto info = (ExitProcessInfo*)currentStartInfo.userData;
 		if (!info) // If null, process has already exited from some other thread
diff --git a/Engine/Source/Programs/UnrealBuildAccelerator/Common/Private/UbaSession.cpp b/Engine/Source/Programs/UnrealBuildAccelerator/Common/Private/UbaSession.cpp
--- a/Engine/Source/Programs/UnrealBuildAccelerator/Common/Private/UbaSession.cpp
+++ b/Engine/Source/Programs/UnrealBuildAccelerator/Common/Private/UbaSession.cpp
@@ -657,7 +657,7 @@ namespace uba
 		if (!shouldWriteToDisk)
 		{
 			u32 offset;
-			bool res = WriteDirectoryEntries(dirKey, dirName.data, offset);
+			u32 res = WriteDirectoryEntries(dirKey, dirName.data, offset);
 			UBA_ASSERT(res); (void)res;
 		}
 
@@ -685,7 +685,7 @@ namespace uba
 		{
 			u32 offset;
 			dirLock.Leave();
-			bool res = WriteDirectoryEntries(dirKey, dirName.data, offset);
+			u32 res = WriteDirectoryEntries(dirKey, dirName.data, offset);
 			UBA_ASSERT(res); (void)res;
 			dirLock.Enter();
 		}
@@ -1243,11 +1243,8 @@ namespace uba
 		WriteDirectoryEntriesInternal(dir, dirKey, dirPath.data, true, tableOffset);
 	}
 
-	void Session::RegisterNewFile(const tchar* filePath)
+	StringKey GetKeyAndFixedName(StringBuffer<>& fixedFilePath, const tchar* filePath)
 	{
-		UBA_ASSERT(!m_runningRemote);
-
-		StringBuffer<> fixedFilePath;
 		FixPath2(filePath, nullptr, 0, fixedFilePath.data, &fixedFilePath.count);
 
 		StringKey dirKey;
@@ -1260,16 +1257,30 @@ namespace uba
 
 		StringKeyHasher hasher;
 		hasher.Update(dirNameForHash.data, dirNameForHash.count);
-		
+
 		StringBuffer<128> baseFileNameForHash;
 		baseFileNameForHash.Append(baseFileName);
 		if (CaseInsensitiveFs)
 			baseFileNameForHash.MakeLower();
-		auto key = ToStringKey(hasher, baseFileNameForHash.data, baseFileNameForHash.count);
+		return ToStringKey(hasher, baseFileNameForHash.data, baseFileNameForHash.count);
+	}
 
+	void Session::RegisterNewFile(const tchar* filePath)
+	{
+		UBA_ASSERT(!m_runningRemote);
+		StringBuffer<> fixedFilePath;
+		auto key = GetKeyAndFixedName(fixedFilePath, filePath);
 		RegisterCreateFileForWrite(key, fixedFilePath.data, fixedFilePath.count, true);
 	}
 
+	void Session::RegisterDeleteFile(const tchar* filePath)
+	{
+		UBA_ASSERT(!m_runningRemote);
+		StringBuffer<> fixedFilePath;
+		auto key = GetKeyAndFixedName(fixedFilePath, filePath);
+		RegisterDeleteFile(key, fixedFilePath.data);
+	}
+
 	void Session::RegisterCustomService(CustomServiceFunction&& function)
 	{
 		m_customServiceFunction = function;
@@ -1323,7 +1334,9 @@ namespace uba
 	void Session::FlushDeadProcesses()
 	{
 		ScopedWriteLock lock(m_processesLock);
-		m_deadProcesses.clear();
+		Vector<ProcessHandle> deadProcesses;
+		deadProcesses.swap(m_deadProcesses);
+		lock.Leave();
 	}
 
 	bool Session::GetInitResponse(InitResponse& out, const InitMessage& msg)
@@ -1989,7 +2002,7 @@ namespace uba
 		return true;
 	}
 
-	bool Session::UpdateEnvironment(ProcessImpl& process, const tchar* reason)
+	bool Session::UpdateEnvironment(ProcessImpl& process, const tchar* reason, bool resetStats)
 	{
 		StackBinaryWriter<16 * 1024> writer;
 		process.m_processStats.Write(writer);
@@ -2084,7 +2097,7 @@ namespace uba
 		#endif
 
 		if (!m_extraInfo.empty())
-			out.Append(' ').Append(m_extraInfo);
+			out.Append(m_extraInfo);
 
 		#if UBA_DEBUG
 		out.Append(TC(" DEBUG"));
diff --git a/Engine/Source/Programs/UnrealBuildAccelerator/Common/Private/UbaSessionClient.cpp b/Engine/Source/Programs/UnrealBuildAccelerator/Common/Private/UbaSessionClient.cpp
--- a/Engine/Source/Programs/UnrealBuildAccelerator/Common/Private/UbaSessionClient.cpp
+++ b/Engine/Source/Programs/UnrealBuildAccelerator/Common/Private/UbaSessionClient.cpp
@@ -625,6 +625,7 @@ namespace uba
 	bool SessionClient::MoveFile(MoveFileResponse& out, const MoveFileMessage& msg)
 	{
 		const tchar* fromName = msg.fromName.data;
+		const tchar* toName = msg.toName.data;
 
 		{
 			ScopedWriteLock lock(msg.process.m_writtenFilesLock);
@@ -632,8 +633,8 @@ namespace uba
 			auto findIt = writtenFiles.find(fromName);
 			if (findIt != writtenFiles.end())
 			{
-				auto insres = writtenFiles.try_emplace(msg.toName.data);
-				UBA_ASSERT(insres.second);
+				auto insres = writtenFiles.try_emplace(toName);
+				UBA_ASSERTF(insres.second, TC("Moving written file to other written file."));
 				insres.first->second = findIt->second;
 				insres.first->second.owner = &msg.process;
 				writtenFiles.erase(findIt);
@@ -643,11 +644,11 @@ namespace uba
 		bool sendMove = true;
 		{
 			ScopedWriteLock lock(m_outputFilesLock);
-			auto findIt = m_outputFiles.find(msg.fromName.data);
+			auto findIt = m_outputFiles.find(fromName);
 			if (findIt != m_outputFiles.end())
 			{
-				auto insres = m_outputFiles.try_emplace(msg.toName.data);
-				UBA_ASSERT(insres.second);
+				auto insres = m_outputFiles.try_emplace(toName);
+				UBA_ASSERTF(insres.second, TC("Failed to add move destination file %s as output file because it is already added."), toName);
 				insres.first->second = findIt->second;
 				m_outputFiles.erase(findIt);
 				sendMove = false;
@@ -661,7 +662,7 @@ namespace uba
 			return true;
 		}
 
-		out.result = uba::MoveFileExW(fromName, msg.toName.data, 0);
+		out.result = uba::MoveFileExW(fromName, toName, 0);
 		out.errorCode = GetLastError();
 
 		return true;
@@ -1694,30 +1695,39 @@ namespace uba
 		ScopedWriteLock lock(process.m_writtenFilesLock);
 		if (!SendFiles(process, process.m_processStats.sendFiles))
 			return false;
+		{
+			ScopedWriteLock lock2(m_outputFilesLock);
+			for (auto& kv : process.m_writtenFiles)
+				m_outputFiles.erase(kv.first);
+		}
 		process.m_writtenFiles.clear();
 		return true;
 	}
 
-	bool SessionClient::UpdateEnvironment(ProcessImpl& process, const tchar* reason)
+	bool SessionClient::UpdateEnvironment(ProcessImpl& process, const tchar* reason, bool resetStats)
 	{
-		StackBinaryWriter<16*1024> writer;
-		NetworkMessage msg(m_client, ServiceId, SessionMessageType_UpdateEnvironment, writer);
-		writer.WriteU32(process.m_id);
-		writer.WriteString(reason);
-		process.m_processStats.Write(writer);
-		process.m_sessionStats.Write(writer);
-		process.m_storageStats.Write(writer);
-		process.m_systemStats.Write(writer);
-
-		process.m_processStats = {};
-		process.m_sessionStats = {};
-		process.m_storageStats = {};
-		process.m_systemStats = {};
-
 		StackBinaryReader<SendMaxSize> reader;
-		if (!msg.Send(reader, m_stats.customMsg))
-			return false;
-		reader.Reset();
+
+		if (resetStats)
+		{
+			StackBinaryWriter<16 * 1024> writer;
+			NetworkMessage msg(m_client, ServiceId, SessionMessageType_UpdateEnvironment, writer);
+			writer.WriteU32(process.m_id);
+			writer.WriteString(reason);
+			process.m_processStats.Write(writer);
+			process.m_sessionStats.Write(writer);
+			process.m_storageStats.Write(writer);
+			process.m_systemStats.Write(writer);
+
+			process.m_processStats = {};
+			process.m_sessionStats = {};
+			process.m_storageStats = {};
+			process.m_systemStats = {};
+
+			if (!msg.Send(reader, m_stats.customMsg))
+				return false;
+			reader.Reset();
+		}
 		return SendUpdateDirectoryTable(reader);
 	}
 }
diff --git a/Engine/Source/Programs/UnrealBuildAccelerator/Common/Private/UbaStorageClient.cpp b/Engine/Source/Programs/UnrealBuildAccelerator/Common/Private/UbaStorageClient.cpp
--- a/Engine/Source/Programs/UnrealBuildAccelerator/Common/Private/UbaStorageClient.cpp
+++ b/Engine/Source/Programs/UnrealBuildAccelerator/Common/Private/UbaStorageClient.cpp
@@ -849,21 +849,28 @@ namespace uba
 
 		u32 entryIndex = 0;
 		u32 fetchIndex = 0;
+		bool success = true;
 		while (left)
 		{
 			if (fetchIndex >= EntryCount)
 			{
 				auto& entry = entries[entryIndex];
 				if (!entry.message.WaitForAsync(entry.reader))
-					return false; // TODO: Cleanup
+				{
+					success = false;
+					break;
+				}
 				entry.~Entry();
 				UBA_ASSERT(entry.reader.GetLeft() == messageMaxSize);
 			}
 
 			auto& entry = *new (entries + entryIndex) Entry(client, fetchId, readBuffer, fetchIndex);
 			if (!entry.message.SendAsync(entry.reader))
-				return false; // TODO: Cleanup
-
+			{
+				entry.~Entry();
+				success = false;
+				break;
+			}
 			++fetchIndex;
 			entryIndex = (entryIndex + 1) % EntryCount;
 
@@ -884,11 +891,11 @@ namespace uba
 		{
 			auto& entry = entries[waitIndex];
 			if (!entry.message.WaitForAsync(entry.reader))
-				return false;  // TODO: Cleanup
+				success = false;
 			entry.~Entry();
 			waitIndex = (waitIndex + 1) % EntryCount;
 		}
-		return true;
+		return success;
 	}
 
 
diff --git a/Engine/Source/Programs/UnrealBuildAccelerator/Common/Public/UbaScheduler.h b/Engine/Source/Programs/UnrealBuildAccelerator/Common/Public/UbaScheduler.h
--- a/Engine/Source/Programs/UnrealBuildAccelerator/Common/Public/UbaScheduler.h
+++ b/Engine/Source/Programs/UnrealBuildAccelerator/Common/Public/UbaScheduler.h
@@ -45,6 +45,7 @@ namespace uba
 		Event m_updateThreadLoop;
 		Thread m_thread;
 		bool m_loop = false;
+		bool m_enableProcessReuse;
 
 		Atomic<u32> m_activeLocalProcesses;
 		Atomic<u32> m_activeRemoteProcesses;
diff --git a/Engine/Source/Programs/UnrealBuildAccelerator/Common/Public/UbaSession.h b/Engine/Source/Programs/UnrealBuildAccelerator/Common/Public/UbaSession.h
--- a/Engine/Source/Programs/UnrealBuildAccelerator/Common/Public/UbaSession.h
+++ b/Engine/Source/Programs/UnrealBuildAccelerator/Common/Public/UbaSession.h
@@ -53,6 +53,7 @@ namespace uba
 		void PrintSummary(Logger& logger); // Print summary stats of session
 		void RefreshDirectory(const tchar* dirName); // Tell uba a directory on disk has been changed by some other system while session is running
 		void RegisterNewFile(const tchar* filePath); // Tell uba a new file on disk has been added by some other system while session is running
+		void RegisterDeleteFile(const tchar* filePath); // Tell uba a file on disk has been deleted by some other system while session is running
 
 		using CustomServiceFunction = Function<u32(Process& handle, const void* recv, u32 recvSize, void* send, u32 sendCapacity)>;
 		void RegisterCustomService(CustomServiceFunction&& function); // Register a custom service (that can be communicated with from the remote agents)
@@ -116,7 +117,7 @@ namespace uba
 		virtual bool CustomMessage(Process& process, BinaryReader& reader, BinaryWriter& writer);
 		virtual void FileEntryAdded(StringKey fileNameKey, u64 lastWritten, u64 size);
 		virtual bool FlushWrittenFiles(ProcessImpl& process);
-		virtual bool UpdateEnvironment(ProcessImpl& process, const tchar* reason);
+		virtual bool UpdateEnvironment(ProcessImpl& process, const tchar* reason, bool resetStats);
 
 		static constexpr CasKey CasKeyIsDirectory = { ~u64(0), ~u64(0), ~u32(0) };
 
diff --git a/Engine/Source/Programs/UnrealBuildAccelerator/Common/Public/UbaSessionClient.h b/Engine/Source/Programs/UnrealBuildAccelerator/Common/Public/UbaSessionClient.h
--- a/Engine/Source/Programs/UnrealBuildAccelerator/Common/Public/UbaSessionClient.h
+++ b/Engine/Source/Programs/UnrealBuildAccelerator/Common/Public/UbaSessionClient.h
@@ -56,7 +56,7 @@ namespace uba
 		virtual void PrintSessionStats(Logger& logger) override;
 		virtual bool CustomMessage(Process& process, BinaryReader& reader, BinaryWriter& writer) override;
 		virtual bool FlushWrittenFiles(ProcessImpl& process) override;
-		virtual bool UpdateEnvironment(ProcessImpl& process, const tchar* reason) override;
+		virtual bool UpdateEnvironment(ProcessImpl& process, const tchar* reason, bool resetStats) override;
 
 		struct InternalProcessStartInfo;
 
diff --git a/Engine/Source/Programs/UnrealBuildAccelerator/Core/Public/UbaPathUtils.h b/Engine/Source/Programs/UnrealBuildAccelerator/Core/Public/UbaPathUtils.h
--- a/Engine/Source/Programs/UnrealBuildAccelerator/Core/Public/UbaPathUtils.h
+++ b/Engine/Source/Programs/UnrealBuildAccelerator/Core/Public/UbaPathUtils.h
@@ -109,7 +109,7 @@ namespace uba
 				--charLen;
 			}
 
-			UBA_ASSERTF(workingDir && *workingDir, TC("No working dir provided"));
+			UBA_ASSERTF(workingDir && *workingDir, TC("No working dir provided but path is relative (%s)"), buffer);
 			tchar temp2[1024];
 			memcpy(temp2, workingDir, workingDirCharLen*sizeof(tchar));
 			memcpy(temp2 + workingDirCharLen, copyFrom, charLen*sizeof(tchar));
diff --git a/Engine/Source/Programs/UnrealBuildAccelerator/Detours/Private/UbaDetoursFileMappingTable.cpp b/Engine/Source/Programs/UnrealBuildAccelerator/Detours/Private/UbaDetoursFileMappingTable.cpp
--- a/Engine/Source/Programs/UnrealBuildAccelerator/Detours/Private/UbaDetoursFileMappingTable.cpp
+++ b/Engine/Source/Programs/UnrealBuildAccelerator/Detours/Private/UbaDetoursFileMappingTable.cpp
@@ -216,7 +216,7 @@ namespace uba
 		bool isAbsolute = IsWindows ? (pathLen > 1 && path[1] == ':') : (pathLen > 0 && path[0] == '/');
 		if (isAbsolute)
 		{
-			tempBuf.Append(path, pathLen);
+			FixPath(tempBuf, path);
 			if (CaseInsensitiveFs)
 				tempBuf.MakeLower();
 			fileNameKey = ToStringKey(tempBuf);
diff --git a/Engine/Source/Programs/UnrealBuildAccelerator/Detours/Private/Windows/UbaDetoursFunctionsNtDll.inl b/Engine/Source/Programs/UnrealBuildAccelerator/Detours/Private/Windows/UbaDetoursFunctionsNtDll.inl
--- a/Engine/Source/Programs/UnrealBuildAccelerator/Detours/Private/Windows/UbaDetoursFunctionsNtDll.inl
+++ b/Engine/Source/Programs/UnrealBuildAccelerator/Detours/Private/Windows/UbaDetoursFunctionsNtDll.inl
@@ -673,6 +673,14 @@ NTSTATUS NTAPI Shared_NtCreateFile(bool IsCreateFunc, PHANDLE hFileHandle, ACCES
 						return STATUS_OBJECT_NAME_NOT_FOUND;
 					}
 				}
+				else if (!checkIfDir)
+				{
+					// File could have been deleted.
+					DirectoryTable::EntryInformation entryInfo;
+					g_directoryTable.GetEntryInformation(entryInfo, dirTableOffset);
+					if (entryInfo.attributes == 0)
+						return STATUS_OBJECT_NAME_NOT_FOUND;
+				}
 
 				bool isWriteAttributes = (DesiredAccess & FILE_WRITE_ATTRIBUTES) != 0;
 
diff --git a/Engine/Source/Programs/UnrealBuildAccelerator/Detours/Private/Windows/UbaDetoursFunctionsWin.cpp b/Engine/Source/Programs/UnrealBuildAccelerator/Detours/Private/Windows/UbaDetoursFunctionsWin.cpp
--- a/Engine/Source/Programs/UnrealBuildAccelerator/Detours/Private/Windows/UbaDetoursFunctionsWin.cpp
+++ b/Engine/Source/Programs/UnrealBuildAccelerator/Detours/Private/Windows/UbaDetoursFunctionsWin.cpp
@@ -1009,13 +1009,14 @@ extern "C"
 		return reader.ReadBool();
 	}
 
-	UBA_DETOURED_API bool UbaUpdateEnvironment(const wchar_t* reason)
+	UBA_DETOURED_API bool UbaUpdateEnvironment(const wchar_t* reason, bool resetStats)
 	{
 		{
 			ScopedWriteLock pcs(g_communicationLock);
 			BinaryWriter writer;
 			writer.WriteByte(MessageType_UpdateEnvironment);
 			writer.WriteString(reason ? reason : L"");
+			writer.WriteBool(resetStats);
 			writer.Flush();
 			BinaryReader reader;
 			if (!reader.ReadBool())
@@ -1038,13 +1039,17 @@ extern "C"
 		*outArguments = 0;
 		StackBinaryReader<2048> reader;
 		reader.SetSize(UbaSendCustomMessage(nullptr, 0, reader.buffer, 2048));
-		if (!reader.GetLeft())
-			return false;
-		reader.ReadString(outArguments, outArgumentsCapacity);
-		reader.SkipString(); // Working dir
-		StringBuffer<256> description;
-		reader.ReadString(description);
+		bool hasProcess = reader.GetLeft();
 
-		return UbaUpdateEnvironment(description.data);
+		StringBuffer<256> description;
+		if (hasProcess)
+		{
+			reader.ReadString(outArguments, outArgumentsCapacity);
+			reader.SkipString(); // Working dir
+			reader.ReadString(description);
+		}
+		if (!UbaUpdateEnvironment(description.data, hasProcess))
+			return false;
+		return hasProcess;
 	}
 }
diff --git a/Engine/Source/Programs/UnrealBuildAccelerator/Detours/Public/UbaDetoursApi.h b/Engine/Source/Programs/UnrealBuildAccelerator/Detours/Public/UbaDetoursApi.h
--- a/Engine/Source/Programs/UnrealBuildAccelerator/Detours/Public/UbaDetoursApi.h
+++ b/Engine/Source/Programs/UnrealBuildAccelerator/Detours/Public/UbaDetoursApi.h
@@ -16,8 +16,8 @@ extern "C"
 	UBA_DETOURED_API bool UbaFlushWrittenFiles();
 
 	// Update environment. Making sure caches are up-to-date in detoured process.
-	// reason ends up in the visualizer for the next block of process
-	UBA_DETOURED_API bool UbaUpdateEnvironment(const wchar_t* reason);
+	// reason ends up in the visualizer for the next block of process if resetStats is true (if resetStats is not true, this will just update the environment within same process)
+	UBA_DETOURED_API bool UbaUpdateEnvironment(const wchar_t* reason, bool resetStats);
 
 	// Returns true if process is running on a remote machine
 	UBA_DETOURED_API bool UbaRunningRemote();
@@ -37,7 +37,7 @@ if (UbaDetoursModule)
 {
 	using UbaFlushWrittenFilesFunc = bool();
 	using UbaRequestNextProcessFunc = bool(TCHAR* outArguments, uint32 outArgumentsCapacity);
-	using UbaUpdateEnvironmentFunc = bool(const wchar_t* reason);
+	using UbaUpdateEnvironmentFunc = bool(const wchar_t* reason, bool resetStats);
 
 	static UbaFlushWrittenFilesFunc* FlushWrittenFiles = (UbaFlushWrittenFilesFunc*)(void*)GetProcAddress(UbaDetoursModule, "UbaFlushWrittenFiles");
 	static UbaRequestNextProcessFunc* RequestNextProcess = (UbaRequestNextProcessFunc*)(void*)GetProcAddress(UbaDetoursModule, "UbaRequestNextProcess");
diff --git a/Engine/Source/Programs/UnrealBuildAccelerator/Host/Private/UbaExports.cpp b/Engine/Source/Programs/UnrealBuildAccelerator/Host/Private/UbaExports.cpp
--- a/Engine/Source/Programs/UnrealBuildAccelerator/Host/Private/UbaExports.cpp
+++ b/Engine/Source/Programs/UnrealBuildAccelerator/Host/Private/UbaExports.cpp
@@ -346,6 +346,10 @@ uba::StorageClient* CreateStorageClient(uba::NetworkClient& client, const uba::t
 	{
 		server->RegisterNewFile(filePath);
 	}
+	void SessionServer_RegisterDeleteFile(uba::SessionServer* server, const uba::tchar* filePath)
+	{
+		server->RegisterDeleteFile(filePath);
+	}
 	uba::ProcessHandle* SessionServer_RunProcess(uba::SessionServer* server, uba::ProcessStartInfo& info, bool async, bool enableDetour)
 	{
 		return new uba::ProcessHandle(server->RunProcess(info, async, enableDetour));
diff --git a/Engine/Source/Programs/UnrealBuildAccelerator/Host/Public/UbaExports.h b/Engine/Source/Programs/UnrealBuildAccelerator/Host/Public/UbaExports.h
--- a/Engine/Source/Programs/UnrealBuildAccelerator/Host/Public/UbaExports.h
+++ b/Engine/Source/Programs/UnrealBuildAccelerator/Host/Public/UbaExports.h
@@ -87,6 +87,7 @@ extern "C"
 	UBA_API void SessionServer_SetRemoteProcessReturned(uba::SessionServer* server, SessionServer_RemoteProcessReturnedCallback* returned, void* userData);
 	UBA_API void SessionServer_RefreshDirectory(uba::SessionServer* server, const uba::tchar* directory);
 	UBA_API void SessionServer_RegisterNewFile(uba::SessionServer* server, const uba::tchar* filePath);
+	UBA_API void SessionServer_RegisterDeleteFile(uba::SessionServer* server, const uba::tchar* filePath);
 	UBA_API uba::ProcessHandle* SessionServer_RunProcess(uba::SessionServer* server, uba::ProcessStartInfo& info, bool async, bool enableDetour);
 	UBA_API uba::ProcessHandle* SessionServer_RunProcessRemote(uba::SessionServer* server, uba::ProcessStartInfo& info, float weight, const void* knownInputs = nullptr, uba::u32 knownInputsCount = 0);
 	UBA_API uba::ProcessHandle* SessionServer_RunProcessRacing(uba::SessionServer* server, uba::u32 raceAgainstRemoteProcessId);
