[UBA]
* Added popup menu on right mouse click
* Added so it is possible to save active visualizing session

diff --git a/Engine/Source/Programs/UnrealBuildAccelerator/Visualizer/Private/UbaTraceReader.cpp b/Engine/Source/Programs/UnrealBuildAccelerator/Visualizer/Private/UbaTraceReader.cpp
--- a/Engine/Source/Programs/UnrealBuildAccelerator/Visualizer/Private/UbaTraceReader.cpp
+++ b/Engine/Source/Programs/UnrealBuildAccelerator/Visualizer/Private/UbaTraceReader.cpp
@@ -2,6 +2,7 @@
 
 #include "UbaTraceReader.h"
 #include "UbaFile.h"
+#include "UbaFileAccessor.h"
 #include "UbaNetworkClient.h"
 #include "UbaNetworkMessage.h"
 
@@ -58,10 +59,8 @@ namespace uba
 		u32 traceSize = reader.ReadU32(); (void)traceSize;
 		u32 version = reader.ReadU32();
 		if (version < TraceReadCompatibilityVersion || version > TraceVersion)
-		{
-			m_logger.Error(L"Incompatible trace version (%u). Current executable supports version %u to %u.", version, TraceReadCompatibilityVersion, TraceVersion);
-			return false;
-		}
+			return m_logger.Error(L"Incompatible trace version (%u). Current executable supports version %u to %u.", version, TraceReadCompatibilityVersion, TraceVersion);
+
 		out.version = version;
 		reader.ReadU32(); // ProcessId
 		u64 traceSystemStartTimeUs = 0;
@@ -864,6 +863,22 @@ namespace uba
 		out.finished = true;
 	}
 
+	bool TraceReader::SaveAs(const tchar* fileName)
+	{
+		if (!m_memoryBegin)
+		{
+			m_logger.Warning(L"Can only save traces that are opened using -listen/-name.");
+			return false;
+		}
+
+		FileAccessor file(m_logger, fileName);
+		if (!file.CreateWrite())
+			return false;
+		if (!file.Write(m_memoryBegin, m_memoryPos - m_memoryBegin))
+			return false;
+		return file.Close();
+	}
+
 	Guid TraceReader::ReadClientId(TraceView& out, BinaryReader& reader)
 	{
 		if (out.version < 15)
diff --git a/Engine/Source/Programs/UnrealBuildAccelerator/Visualizer/Private/UbaTraceReader.h b/Engine/Source/Programs/UnrealBuildAccelerator/Visualizer/Private/UbaTraceReader.h
--- a/Engine/Source/Programs/UnrealBuildAccelerator/Visualizer/Private/UbaTraceReader.h
+++ b/Engine/Source/Programs/UnrealBuildAccelerator/Visualizer/Private/UbaTraceReader.h
@@ -143,6 +143,8 @@ namespace uba
 		bool ReadTrace(TraceView& out, BinaryReader& reader, u64 maxTime);
 		void StopAllActive(TraceView& out, u64 stopTime);
 
+		bool SaveAs(const tchar* fileName);
+
 		Guid ReadClientId(TraceView& out, BinaryReader& reader);
 		TraceView::Session& GetSession(TraceView& out, u32 sessionIndex);
 		TraceView::Session* GetSession(TraceView& out, const Guid& clientUid);
diff --git a/Engine/Source/Programs/UnrealBuildAccelerator/Visualizer/Private/UbaVisualizer.cpp b/Engine/Source/Programs/UnrealBuildAccelerator/Visualizer/Private/UbaVisualizer.cpp
--- a/Engine/Source/Programs/UnrealBuildAccelerator/Visualizer/Private/UbaVisualizer.cpp
+++ b/Engine/Source/Programs/UnrealBuildAccelerator/Visualizer/Private/UbaVisualizer.cpp
@@ -9,6 +9,14 @@
 #pragma comment (lib, "Dwmapi.lib")
 #define WM_NEWTRACE WM_USER+1
 
+enum
+{
+	Popup_CopySessionInfo = 3,
+	Popup_CopyProcessInfo,
+	Popup_SaveAs,
+	Popup_Quit,
+};
+
 namespace uba
 {
 	class LineCountLogger : public Logger
@@ -40,6 +48,15 @@ namespace uba
 		int fontHeight;
 	};
 
+	class WriteTextLogger : public Logger
+	{
+	public:
+		WriteTextLogger(TString& out) : m_out(out) {}
+		virtual void BeginScope() override {}
+		virtual void EndScope() override {}
+		virtual void Log(LogEntryType type, const wchar_t* str, u32 strLen) override { m_out.append(str, strLen).append(TC("\n")); }
+		TString& m_out;
+	};
 
 	Visualizer::Visualizer(Logger& logger)
 	:	m_logger(logger)
@@ -1744,6 +1761,67 @@ namespace uba
 		m_contentHeight = posY - int(m_scrollPosY) + stepY + 14;
 	}
 
+	void Visualizer::WriteProcessStats(Logger& out, TraceView::Process& process)
+	{
+		bool hasStorageStats = true;
+		bool hasExited = process.stop != ~u64(0);
+		out.Info(L"  %ls", process.description.c_str());
+		out.Info(L"  Start:     %ls", TimeToText(process.start, true).str);
+		if (hasExited)
+			out.Info(L"  Duration:  %ls", TimeToText(process.stop - process.start, true).str);
+		if (hasExited && process.exitCode != 0)
+			out.Info(L"  ExitCode:  %u", process.exitCode);
+		out.Info(L"");
+
+		if (process.stop != ~u64(0))
+		{
+			out.Info(L"  ----------- Process stats -----------");
+			process.processStats.Print(out, m_traceView.frequency);
+			if (hasStorageStats)
+			{
+				out.Info(L"");
+				out.Info(L"  ----------- Session stats -----------");
+				process.sessionStats.Print(out, m_traceView.frequency);
+				out.Info(L"");
+				out.Info(L"  ----------- Storage stats -----------");
+				process.storageStats.Print(out, m_traceView.frequency);
+				out.Info(L"");
+				out.Info(L"  ----------- System stats ------------");
+				process.systemStats.Print(out, false, m_traceView.frequency);
+			}
+		}
+	}
+
+	void Visualizer::CopyTextToClipboard(const TString& str)
+	{
+		if (!OpenClipboard(m_hwnd))
+			return;
+		if (auto hglbCopy = GlobalAlloc(GMEM_MOVEABLE, (str.size() + 1) * sizeof(TCHAR)))
+		{
+			if (auto lptstrCopy = GlobalLock(hglbCopy))
+			{
+				memcpy(lptstrCopy, str.data(), (str.size() + 1) * sizeof(TCHAR));
+				GlobalUnlock(hglbCopy);
+				EmptyClipboard();
+				SetClipboardData(CF_UNICODETEXT, hglbCopy);
+			}
+		}
+		CloseClipboard();
+	}
+
+	void Visualizer::UnselectAndRedraw()
+	{
+		if (m_processSelected || m_sessionSelectedIndex != ~0u || m_statsSelected || m_timelineSelected || m_fetchedFilesSelected != ~0u)
+		{
+			m_processSelected = false;
+			m_sessionSelectedIndex = ~0u;
+			m_statsSelected = false;
+			m_buttonSelected = ~0u;
+			m_timelineSelected = 0;
+			m_fetchedFilesSelected = ~0u;
+			RedrawWindow(m_hwnd, NULL, NULL, RDW_INVALIDATE);
+		}
+	}
 
 	bool Visualizer::UpdateAutoscroll()
 	{
@@ -2032,16 +2110,8 @@ namespace uba
 			tme.hwndTrack = hWnd;
 			TrackMouseEvent(&tme);
 
-			if (m_processSelected || m_sessionSelectedIndex != ~0u || m_statsSelected || m_timelineSelected || m_fetchedFilesSelected != ~0u)
-			{
-				m_processSelected = false;
-				m_sessionSelectedIndex = ~0u;
-				m_statsSelected = false;
-				m_buttonSelected = ~0u;
-				m_timelineSelected = 0;
-				m_fetchedFilesSelected = ~0u;
-				RedrawWindow(hWnd, NULL, NULL, RDW_INVALIDATE);
-			}
+			if (!m_showPopup)
+				UnselectAndRedraw();
 			break;
 
 		case WM_MBUTTONDOWN:
@@ -2103,8 +2173,84 @@ namespace uba
 			break;
 		}
 
-		case WM_RBUTTONDOWN:
+		case WM_RBUTTONUP:
 		{
+			POINT point;
+			point.x = LOWORD(lParam);
+			point.y = HIWORD(lParam);
+
+			HMENU hMenu = CreatePopupMenu();
+			ClientToScreen(hWnd, &point);
+
+			//AppendMenuW(hMenu, MF_STRING, 1, L"&New");
+			if (m_sessionSelectedIndex != ~0u)
+			{
+				AppendMenuW(hMenu, MF_STRING, Popup_CopySessionInfo, L"&Copy Session Info");
+				AppendMenuW(hMenu, MF_SEPARATOR, 0, NULL);
+			}
+			else if (m_processSelected)
+			{
+				AppendMenuW(hMenu, MF_STRING, Popup_CopyProcessInfo, L"&Copy Process Info");
+				AppendMenuW(hMenu, MF_SEPARATOR, 0, NULL);
+			}
+			AppendMenuW(hMenu, MF_STRING, Popup_SaveAs, L"&Save Trace");
+			AppendMenuW(hMenu, MF_SEPARATOR, 0, NULL);
+			AppendMenuW(hMenu, MF_STRING, Popup_Quit, L"&Quit");
+			m_showPopup = true;
+			switch (TrackPopupMenu(hMenu, TPM_RETURNCMD | TPM_RIGHTBUTTON, point.x, point.y, 0, hWnd, NULL))
+			{
+			case Popup_SaveAs:
+			{
+				OPENFILENAME ofn;       // common dialog box structure
+				TCHAR szFile[260] = { 0 };       // if using TCHAR macros
+
+				// Initialize OPENFILENAME
+				ZeroMemory(&ofn, sizeof(ofn));
+				ofn.lStructSize = sizeof(ofn);
+				ofn.hwndOwner = hWnd;
+				ofn.lpstrFile = szFile;
+				ofn.nMaxFile = sizeof(szFile);
+				ofn.lpstrDefExt = TC("uba");
+				ofn.lpstrFilter = TC("Uba\0*.uba\0All\0*.*\0");
+				ofn.nFilterIndex = 1;
+				ofn.lpstrFileTitle = NULL;
+				ofn.nMaxFileTitle = 0;
+				ofn.lpstrInitialDir = NULL;
+				//ofn.Flags = OFN_PATHMUSTEXIST | OFN_FILEMUSTEXIST;
+				if (GetSaveFileName(&ofn))
+					m_trace.SaveAs(ofn.lpstrFile);
+				break;
+			}
+
+			case Popup_Quit: // Quit
+				m_looping = false;
+				break;
+
+			case Popup_CopySessionInfo:
+			{
+				TString str;
+				auto& session = m_traceView.sessions[m_sessionSelectedIndex];
+				str.append(session.name).append(TC("\n"));
+				for (auto& line : session.summary)
+					str.append(line).append(TC("\n"));
+				CopyTextToClipboard(str);
+				break;
+			}
+
+			case Popup_CopyProcessInfo:
+			{
+				TString str;
+				WriteTextLogger logger(str);
+				TraceView::Process& process = *m_traceView.GetProcess(m_processSelectedLocation);
+				WriteProcessStats(logger, process);
+				CopyTextToClipboard(str);
+				break;
+			}
+			}
+
+			DestroyMenu(hMenu);
+			m_showPopup = false;
+			UnselectAndRedraw();
 			break;
 		}
 
diff --git a/Engine/Source/Programs/UnrealBuildAccelerator/Visualizer/Private/UbaVisualizer.h b/Engine/Source/Programs/UnrealBuildAccelerator/Visualizer/Private/UbaVisualizer.h
--- a/Engine/Source/Programs/UnrealBuildAccelerator/Visualizer/Private/UbaVisualizer.h
+++ b/Engine/Source/Programs/UnrealBuildAccelerator/Visualizer/Private/UbaVisualizer.h
@@ -56,6 +56,9 @@ namespace uba
 		};
 		void HitTest(HitTestResult& outResult, const POINT& pos);
 
+		void WriteProcessStats(Logger& out, TraceView::Process& process);
+		void CopyTextToClipboard(const TString& str);
+		void UnselectAndRedraw();
 		bool UpdateAutoscroll();
 		bool UpdateSelection();
 		void UpdateScrollbars(bool redraw);
@@ -139,6 +142,7 @@ namespace uba
 		u32 m_fetchedFilesSelected = ~0u;
 		
 		bool m_mouseOverWindow = false;
+		bool m_showPopup = false;
 
 		enum ComponentType
 		{
diff --git a/Engine/Source/Programs/UnrealBuildAccelerator/Visualizer/Private/UbaVisualizerApp.cpp b/Engine/Source/Programs/UnrealBuildAccelerator/Visualizer/Private/UbaVisualizerApp.cpp
--- a/Engine/Source/Programs/UnrealBuildAccelerator/Visualizer/Private/UbaVisualizerApp.cpp
+++ b/Engine/Source/Programs/UnrealBuildAccelerator/Visualizer/Private/UbaVisualizerApp.cpp
@@ -47,15 +47,16 @@ struct MessageBoxLogWriter : public LogWriter
 
 	virtual void Log(LogEntryType type, const tchar* str, u32 strLen, const tchar* prefix = nullptr, u32 prefixLen = 0) override
 	{
+		if (type > LogEntryType_Warning)
+			return;
+
+		HWND hwnd = m_visualizer ? m_visualizer->GetHwnd() : NULL;
+		UINT flags = MB_ICONERROR;
+		if (!hwnd)
+			flags |= MB_TOPMOST;
+		MessageBox(hwnd, str, TC("UbaVisualizer"), flags);
 		if (type == LogEntryType_Error)
-		{
-			HWND hwnd = m_visualizer ? m_visualizer->GetHwnd() : NULL;
-			UINT flags = MB_ICONERROR;
-			if (!hwnd)
-				flags |= MB_TOPMOST;
-			MessageBox(hwnd, str, TC("UbaVisualizer"), flags);
 			ExitProcess(~0u);
-		}
 	}
 	
 	Visualizer* m_visualizer = nullptr;
