[UBA]
* Beeeautiful hack to try to get around ETXTBSY error when spawning ispc processes on linux. Our theory is that with perfectly bad timing the symlink creation of ispc could happen at the exact same time as a posix_spawn happening and inside symlink there is a temporarily open file descriptor that gets cloned into the child and kept open. and this causes ETXTBSY when trying to open that binary.

Solution is simply to have a writelock around the symlink creation and a readlock around posix_spawn

diff --git a/Engine/Source/Programs/UnrealBuildAccelerator/Common/Private/UbaProcess.cpp b/Engine/Source/Programs/UnrealBuildAccelerator/Common/Private/UbaProcess.cpp
--- a/Engine/Source/Programs/UnrealBuildAccelerator/Common/Private/UbaProcess.cpp
+++ b/Engine/Source/Programs/UnrealBuildAccelerator/Common/Private/UbaProcess.cpp
@@ -1455,6 +1455,8 @@ namespace uba
 			pid_t processID;
 			while (true)
 			{
+				ScopedReadLock lock(m_session.m_hackToPreventETXTBUSY);
+
 				res = posix_spawnp(&processID, m_realApplication.c_str(), &fileActions, &attr, (char**)argsArray, (char**)envvars.data());
 				if (res == 0)
 					break;
diff --git a/Engine/Source/Programs/UnrealBuildAccelerator/Common/Private/UbaSessionClient.cpp b/Engine/Source/Programs/UnrealBuildAccelerator/Common/Private/UbaSessionClient.cpp
--- a/Engine/Source/Programs/UnrealBuildAccelerator/Common/Private/UbaSessionClient.cpp
+++ b/Engine/Source/Programs/UnrealBuildAccelerator/Common/Private/UbaSessionClient.cpp
@@ -339,6 +339,11 @@ namespace uba
 			if (!m_storage.CreateDirectory(binaryDir.data))
 				return false;
 		}
+
+		#if !PLATFORM_WINDOWS
+		ScopedWriteLock hackLock(m_hackToPreventETXTBUSY);
+		#endif
+
 		return m_storage.CopyOrLink(casKey, out.data, fileAttributes);
 	}
 
diff --git a/Engine/Source/Programs/UnrealBuildAccelerator/Common/Public/UbaSession.h b/Engine/Source/Programs/UnrealBuildAccelerator/Common/Public/UbaSession.h
--- a/Engine/Source/Programs/UnrealBuildAccelerator/Common/Public/UbaSession.h
+++ b/Engine/Source/Programs/UnrealBuildAccelerator/Common/Public/UbaSession.h
@@ -223,6 +223,12 @@ namespace uba
 
 		CustomServiceFunction m_customServiceFunction;
 
+		// This is a horrible hack to try to prevent ETXTBUSY. Our theory is that if posix_spawn happens with the exact same timing as symlink has an open handle we end up with the child process cloning that handle
+		// Causing posix_spawn to fail on the executable that was created with the symlink
+		#if !PLATFORM_WINDOWS
+		ReaderWriterLock m_hackToPreventETXTBUSY;
+		#endif
+
 		friend class ProcessImpl;
 	};
 
