[UBA]
* Added code fix in posix detour for bug that happens when reusing shadercompileworker on remotes. Don't know if this will be used on mac/linux but just in case

diff --git a/Engine/Source/Programs/UnrealBuildAccelerator/Detours/Private/PosixOS/UbaDetoursFunctionsPosix.cpp b/Engine/Source/Programs/UnrealBuildAccelerator/Detours/Private/PosixOS/UbaDetoursFunctionsPosix.cpp
--- a/Engine/Source/Programs/UnrealBuildAccelerator/Detours/Private/PosixOS/UbaDetoursFunctionsPosix.cpp
+++ b/Engine/Source/Programs/UnrealBuildAccelerator/Detours/Private/PosixOS/UbaDetoursFunctionsPosix.cpp
@@ -322,6 +322,17 @@ int Shared_open(const char* funcName, const char* file, int flags, int mode, con
 					return -1;
 				}
 			}
+			else
+			{
+				// File could have been deleted.
+				DirectoryTable::EntryInformation entryInfo;
+				g_directoryTable.GetEntryInformation(entryInfo, dirTableOffset);
+				if (entryInfo.attributes == 0)
+				{
+					errno = ENOENT;
+					return -1;
+				}
+			}
 		}
 	}
 
