[UBA]
* Fixed so scheduler can take in 0 as max local processors and use 0. Default is now ~0u which becomes max logical processors

diff --git a/Engine/Source/Programs/UnrealBuildAccelerator/Common/Private/UbaScheduler.cpp b/Engine/Source/Programs/UnrealBuildAccelerator/Common/Private/UbaScheduler.cpp
--- a/Engine/Source/Programs/UnrealBuildAccelerator/Common/Private/UbaScheduler.cpp
+++ b/Engine/Source/Programs/UnrealBuildAccelerator/Common/Private/UbaScheduler.cpp
@@ -39,7 +39,7 @@ namespace uba
 
 	Scheduler::Scheduler(SessionServer& session, u32 maxLocalProcessors, bool enableProcessReuse)
 	:	m_session(session)
-	,	m_maxLocalProcessors(maxLocalProcessors ? maxLocalProcessors : GetLogicalProcessorCount())
+	,	m_maxLocalProcessors(maxLocalProcessors != ~0u ? maxLocalProcessors : GetLogicalProcessorCount())
 	,	m_updateThreadLoop(false)
 	,	m_enableProcessReuse(enableProcessReuse)
 	{
diff --git a/Engine/Source/Programs/UnrealBuildAccelerator/Common/Public/UbaScheduler.h b/Engine/Source/Programs/UnrealBuildAccelerator/Common/Public/UbaScheduler.h
--- a/Engine/Source/Programs/UnrealBuildAccelerator/Common/Public/UbaScheduler.h
+++ b/Engine/Source/Programs/UnrealBuildAccelerator/Common/Public/UbaScheduler.h
@@ -17,7 +17,7 @@ namespace uba
 	class Scheduler
 	{
 	public:
-		Scheduler(SessionServer& session, u32 maxLocalProcessors = 0, bool enableProcessReuse = false);
+		Scheduler(SessionServer& session, u32 maxLocalProcessors = ~0u, bool enableProcessReuse = false);
 		~Scheduler();
 
 		void Start();
diff --git a/Engine/Source/Programs/UnrealBuildAccelerator/Host/Public/UbaExports.h b/Engine/Source/Programs/UnrealBuildAccelerator/Host/Public/UbaExports.h
--- a/Engine/Source/Programs/UnrealBuildAccelerator/Host/Public/UbaExports.h
+++ b/Engine/Source/Programs/UnrealBuildAccelerator/Host/Public/UbaExports.h
@@ -111,7 +111,7 @@ extern "C"
 
 
 	// Scheduler
-	UBA_API uba::Scheduler* Scheduler_Create(uba::SessionServer* session, uba::u32 maxLocalProcessors = 0, bool enableProcessReuse = false);
+	UBA_API uba::Scheduler* Scheduler_Create(uba::SessionServer* session, uba::u32 maxLocalProcessors = ~0u, bool enableProcessReuse = false);
 	UBA_API void Scheduler_Start(uba::Scheduler* scheduler);
 	UBA_API void Scheduler_EnqueueProcess(uba::Scheduler* scheduler, const uba::ProcessStartInfo& info, float weight = 1.0f, const void* knownInputs = nullptr, uba::u32 knownInputsBytes = 0, uba::u32 knownInputsCount = 0);
 	UBA_API void Scheduler_Stop(uba::Scheduler* scheduler);
