[UBA]
* Added unit test for scheduler

diff --git a/Engine/Source/Programs/UnrealBuildAccelerator/Common/Private/UbaSession.cpp b/Engine/Source/Programs/UnrealBuildAccelerator/Common/Private/UbaSession.cpp
--- a/Engine/Source/Programs/UnrealBuildAccelerator/Common/Private/UbaSession.cpp
+++ b/Engine/Source/Programs/UnrealBuildAccelerator/Common/Private/UbaSession.cpp
@@ -2004,6 +2004,8 @@ namespace uba
 
 	bool Session::UpdateEnvironment(ProcessImpl& process, const tchar* reason, bool resetStats)
 	{
+		if (!resetStats)
+			return true;
 		StackBinaryWriter<16 * 1024> writer;
 		process.m_processStats.Write(writer);
 		process.m_sessionStats.Write(writer);
diff --git a/Engine/Source/Programs/UnrealBuildAccelerator/Detours/Public/UbaDetoursApi.h b/Engine/Source/Programs/UnrealBuildAccelerator/Detours/Public/UbaDetoursApi.h
--- a/Engine/Source/Programs/UnrealBuildAccelerator/Detours/Public/UbaDetoursApi.h
+++ b/Engine/Source/Programs/UnrealBuildAccelerator/Detours/Public/UbaDetoursApi.h
@@ -24,6 +24,7 @@ extern "C"
 
 	// Using custom message inside UbaScheduler so can't be combined with RegisterCustomService.
 	// This function will automatically flush written files before requesting next process and will also update environment if new process is retrieved
+	// Will write arguments into outArguments. outArgumentsCapacity should be in characters (not bytes)
 	UBA_DETOURED_API bool UbaRequestNextProcess(wchar_t* outArguments, unsigned int outArgumentsCapacity);
 }
 
@@ -36,12 +37,12 @@ static HMODULE UbaDetoursModule = GetModuleHandleW(L"UbaDetours.dll");
 if (UbaDetoursModule)
 {
 	using UbaFlushWrittenFilesFunc = bool();
-	using UbaRequestNextProcessFunc = bool(TCHAR* outArguments, uint32 outArgumentsCapacity);
+	using UbaRequestNextProcessFunc = bool(const wchar_t* outArguments, unsigned int outArgumentsCapacity);
 	using UbaUpdateEnvironmentFunc = bool(const wchar_t* reason, bool resetStats);
 
-	static UbaFlushWrittenFilesFunc* FlushWrittenFiles = (UbaFlushWrittenFilesFunc*)(void*)GetProcAddress(UbaDetoursModule, "UbaFlushWrittenFiles");
-	static UbaRequestNextProcessFunc* RequestNextProcess = (UbaRequestNextProcessFunc*)(void*)GetProcAddress(UbaDetoursModule, "UbaRequestNextProcess");
-	static UbaUpdateEnvironmentFunc* UpdateEnvironment = (UbaUpdateEnvironmentFunc*)(void*)GetProcAddress(UbaDetoursModule, "UbaUpdateEnvironment");
+	static UbaFlushWrittenFilesFunc* flushWrittenFiles = (UbaFlushWrittenFilesFunc*)(void*)GetProcAddress(UbaDetoursModule, "UbaFlushWrittenFiles");
+	static UbaRequestNextProcessFunc* requestNextProcess = (UbaRequestNextProcessFunc*)(void*)GetProcAddress(UbaDetoursModule, "UbaRequestNextProcess");
+	static UbaUpdateEnvironmentFunc* updateEnvironment = (UbaUpdateEnvironmentFunc*)(void*)GetProcAddress(UbaDetoursModule, "UbaUpdateEnvironment");
 
 	// ...
 }
diff --git a/Engine/Source/Programs/UnrealBuildAccelerator/Test/Private/UbaTestAll.h b/Engine/Source/Programs/UnrealBuildAccelerator/Test/Private/UbaTestAll.h
--- a/Engine/Source/Programs/UnrealBuildAccelerator/Test/Private/UbaTestAll.h
+++ b/Engine/Source/Programs/UnrealBuildAccelerator/Test/Private/UbaTestAll.h
@@ -4,6 +4,7 @@
 
 #include "UbaTestBasics.h"
 #include "UbaTestNetwork.h"
+#include "UbaTestScheduler.h"
 #include "UbaTestStorage.h"
 #include "UbaTestSession.h"
 
@@ -27,6 +28,8 @@ namespace uba
 		UBA_TEST(TestCustomService) \
 		UBA_TEST(TestMultipleDetouredProcesses) \
 		UBA_TEST(TestLogLines) \
+		UBA_TEST(TestLocalSchedule) \
+		UBA_TEST(TestLocalScheduleReuse) \
 
 	#define UBA_TEST(x) \
 		logger.Info(TC("Running %hs..."), #x); \
@@ -57,7 +60,7 @@ namespace uba
 
 		logger.Info(TC("Running tests (Test rootdir: %s)"), testRootDir.data);
 
-		//UBA_TEST(TestDetouredTestApp)
+		//UBA_TEST(TestLocalScheduleReuse)
 		UBA_TESTS
 
 		logger.Info(TC("Tests finished successfully!"));
diff --git a/Engine/Source/Programs/UnrealBuildAccelerator/Test/Private/UbaTestScheduler.h b/Engine/Source/Programs/UnrealBuildAccelerator/Test/Private/UbaTestScheduler.h
new file mode 100644
--- /dev/null
+++ b/Engine/Source/Programs/UnrealBuildAccelerator/Test/Private/UbaTestScheduler.h
@@ -0,0 +1,61 @@
+// Copyright Epic Games, Inc. All Rights Reserved.
+
+#pragma once
+
+#include "UbaTestSession.h"
+#include "UbaScheduler.h"
+
+namespace uba
+{
+	bool TestLocalSchedule(LoggerWithWriter& logger, const StringBufferBase& testRootDir)
+	{
+		return RunLocal(logger, testRootDir, [](LoggerWithWriter& logger, SessionServer& session, const tchar* workingDir, const RunProcessFunction& runProcess)
+			{
+
+				Scheduler scheduler(session, true);
+
+				ProcessStartInfo processInfo;
+				processInfo.application = GetPingApplication();
+				processInfo.workingDir = workingDir;
+				processInfo.arguments = IsWindows ? TC("-n 2 localhost") : TC("-c 2 localhost");
+
+				scheduler.EnqueueProcess(processInfo);
+				scheduler.Start();
+
+				u32 queued, active, finished;
+				do { scheduler.GetStats(queued, active, finished); } while (queued != 0 || active != 0);
+
+				scheduler.Stop();
+				return true;
+			});
+	}
+
+	bool TestLocalScheduleReuse(LoggerWithWriter& logger, const StringBufferBase& testRootDir)
+	{
+		if (!IsWindows)
+			return true;
+
+		return RunLocal(logger, testRootDir, [](LoggerWithWriter& logger, SessionServer& session, const tchar* workingDir, const RunProcessFunction& runProcess)
+			{
+
+				Scheduler scheduler(session, true);
+
+				StringBuffer<> testApp;
+				GetTestAppPath(logger, testApp);
+
+				ProcessStartInfo processInfo;
+				processInfo.application = testApp.data;
+				processInfo.workingDir = workingDir;
+				processInfo.arguments = TC("-reuse");
+
+				scheduler.EnqueueProcess(processInfo);
+				scheduler.Start();
+
+				u32 queued, active, finished;
+				do { scheduler.GetStats(queued, active, finished); } while (queued != 0 || active != 0);
+
+				scheduler.Stop();
+				return true;
+			});
+	}
+}
diff --git a/Engine/Source/Programs/UnrealBuildAccelerator/TestApp/Private/Windows/UbaTestAppWin.cpp b/Engine/Source/Programs/UnrealBuildAccelerator/TestApp/Private/Windows/UbaTestAppWin.cpp
--- a/Engine/Source/Programs/UnrealBuildAccelerator/TestApp/Private/Windows/UbaTestAppWin.cpp
+++ b/Engine/Source/Programs/UnrealBuildAccelerator/TestApp/Private/Windows/UbaTestAppWin.cpp
@@ -131,6 +131,15 @@ int wmain(int argc, wchar_t* argv[])
 			CopyFile(L"FileW2", L"FileWF", false);
 		}
 	}
+	else if (wcscmp(argv[1], L"-reuse") == 0)
+	{
+		using UbaRequestNextProcessFunc = bool(wchar_t* outArguments, unsigned int outArgumentsCapacity);
+		static UbaRequestNextProcessFunc* requestNextProcess = (UbaRequestNextProcessFunc*)(void*)GetProcAddress(detoursHandle, "UbaRequestNextProcess");
+
+		wchar_t arguments[1024];
+		if (requestNextProcess(arguments, sizeof(arguments)))
+			return LogError(L"Didn't expect another process");
+	}
 	else
 	{
 		using u32 = unsigned int;
