[UBA]
* Fixed visualizer drawing lines between two sessions started from the same client
* Removed the "- Disconnected" entries where client was disconnected super early before any information was available

diff --git a/Engine/Source/Programs/UnrealBuildAccelerator/Visualizer/Private/UbaTraceReader.cpp b/Engine/Source/Programs/UnrealBuildAccelerator/Visualizer/Private/UbaTraceReader.cpp
--- a/Engine/Source/Programs/UnrealBuildAccelerator/Visualizer/Private/UbaTraceReader.cpp
+++ b/Engine/Source/Programs/UnrealBuildAccelerator/Visualizer/Private/UbaTraceReader.cpp
@@ -376,10 +376,9 @@ namespace uba
 				session.prevUpdateTime = 0;
 				session.prevSend = 0;
 				session.prevRecv = 0;
-				//session.lastPing = 0;
-				//session.memAvail = 0;
 				session.memTotal = 0;
-				//session.cpuLoad = 0;
+				if (!session.updates.empty())
+					session.updates.push_back({ time, 0, 0, lastPing, memAvail, cpuLoad, connectionCount });
 			}
 			else
 			{
diff --git a/Engine/Source/Programs/UnrealBuildAccelerator/Visualizer/Private/UbaVisualizer.cpp b/Engine/Source/Programs/UnrealBuildAccelerator/Visualizer/Private/UbaVisualizer.cpp
--- a/Engine/Source/Programs/UnrealBuildAccelerator/Visualizer/Private/UbaVisualizer.cpp
+++ b/Engine/Source/Programs/UnrealBuildAccelerator/Visualizer/Private/UbaVisualizer.cpp
@@ -495,6 +495,10 @@ namespace uba
 		for (u64 i = 0, e = m_traceView.sessions.size(); i != e; ++i)
 		{
 			auto& session = *sortedSessions[i].session;
+			bool hasUpdates = !session.updates.empty();
+			if (!hasUpdates && session.processors.empty())
+				continue;
+
 			processLocation.sessionIndex = sortedSessions[i].index;
 			bool isFirst = i == 0;
 			if (!isFirst)
@@ -506,8 +510,6 @@ namespace uba
 				MoveToEx(hdc, 0, posY, NULL);
 				LineTo(hdc, clientRect.right, posY);
 
-				bool hasUpdates = !session.updates.empty();
-
 				StringBuffer<> text;
 				text.Append(session.name);
 
@@ -546,7 +548,7 @@ namespace uba
 			posY += SessionStepY;
 
 			bool showGraph = m_visibleComponents[ComponentType_SendRecv] || m_visibleComponents[ComponentType_CpuMem];
-			if (showGraph && !session.updates.empty())
+			if (showGraph && hasUpdates)
 			{
 				if (posY + GraphHeight >= progressRect.top && posY + GraphHeight - 5 < progressRect.bottom)
 				{
@@ -581,6 +583,9 @@ namespace uba
 						double duration = TimeToS(update.time - prevTime);
 						if (update.time == 0)
 							isFirstUpdate = true;
+						else if (prevSend > update.send || prevRecv > update.recv)
+							isFirstUpdate = true;
+
 						if (double sendInvScaleY = duration * sendScale)
 							sendY = graphBaseY - int(double(update.send - prevSend) / sendInvScaleY);
 						if (double recvInvScaleY = duration * recvScale)
@@ -1612,6 +1617,11 @@ namespace uba
 		for (u64 i = 0, e = m_traceView.sessions.size(); i != e; ++i)
 		{
 			auto& session = *sortedSessions[i].session;
+
+			bool hasUpdates = !session.updates.empty();
+			if (!hasUpdates && session.processors.empty())
+				continue;
+
 			u32 sessionIndex = sortedSessions[i].index;
 			bool isFirst = i == 0;
 			if (!isFirst)
@@ -1641,6 +1651,15 @@ namespace uba
 					for (auto& update : session.updates)
 					{
 						int x = int(posX + TimeToS(update.time) * scaleX);
+
+						if (prevSend > update.send || prevRecv > update.recv)
+						{
+							prevSend = update.send;
+							prevRecv = update.recv;
+							prevX = x;
+							continue;
+						}
+
 						int hitOffset = (prevX - x)/2;
 						if (pos.x + hitOffset >= prevX && pos.x + hitOffset <= x)
 						{
