[UBA]
* Removed assert in detouring code for GetFileAttributesA since we have verfied it is forwarded to GetFileAttributesW on both windows and wine
* Added code to fixup path sent in when creating storage server through export functions
* Added path separator fix in session setup code

diff --git a/Engine/Source/Programs/UnrealBuildAccelerator/Common/Private/UbaSession.cpp b/Engine/Source/Programs/UnrealBuildAccelerator/Common/Private/UbaSession.cpp
--- a/Engine/Source/Programs/UnrealBuildAccelerator/Common/Private/UbaSession.cpp
+++ b/Engine/Source/Programs/UnrealBuildAccelerator/Common/Private/UbaSession.cpp
@@ -930,7 +930,7 @@ namespace uba
 
 		UBA_ASSERTF(info.rootDir && *info.rootDir, TC("No root dir set when creating session"));
 		m_rootDir.count = GetFullPathNameW(info.rootDir, m_rootDir.capacity, m_rootDir.data, NULL);
-		m_rootDir.EnsureEndsWithSlash();
+		m_rootDir.Replace('/', PathSeparator).EnsureEndsWithSlash();
 		m_uid = u32(HashString()(m_rootDir.data));
 
 		m_runningRemote = runningRemote;
@@ -1908,7 +1908,10 @@ namespace uba
 			auto memClose = MakeGuard([&](){ UnmapViewOfFile(mem, fileSize); });
 
 			constexpr bool useFileMapForWrite = false;
-			bool useOverlap = false;// fileSize > 8 * 1024 * 1024;
+			//bool useOverlap = false;// fileSize > 8 * 1024 * 1024;
+
+			static bool useOverlap = false;
+
 
 			u32 attributes = DefaultAttributes();
 			if (useOverlap)
@@ -1926,8 +1929,17 @@ namespace uba
 			{
 				if (!destinationFile.CreateWrite(false, attributes, fileSize, m_tempPath.data))
 					return false;
+
+				u64 start = GetTime();
 				if (!destinationFile.Write(mem, fileSize))
 					return false;
+
+				if (!useOverlap && TimeToMs(GetTime() - start) > 4 * 1000)
+				{
+					useOverlap = true;
+					m_logger.Info(TC("Switching to overlapped I/O due to slow writes"));
+				}
+
 			}
 			if (u64 time = file.lastWriteTime)
 				if (!SetFileLastWriteTime(destinationFile.GetHandle(), time))
diff --git a/Engine/Source/Programs/UnrealBuildAccelerator/Detours/Private/Windows/UbaDetoursFunctionsKernelBase.inl b/Engine/Source/Programs/UnrealBuildAccelerator/Detours/Private/Windows/UbaDetoursFunctionsKernelBase.inl
--- a/Engine/Source/Programs/UnrealBuildAccelerator/Detours/Private/Windows/UbaDetoursFunctionsKernelBase.inl
+++ b/Engine/Source/Programs/UnrealBuildAccelerator/Detours/Private/Windows/UbaDetoursFunctionsKernelBase.inl
@@ -2846,8 +2846,8 @@ DWORD Detoured_GetFullPathNameA(LPCSTR lpFileName, DWORD nBufferLength, LPSTR lp
 
 DWORD Detoured_GetFileAttributesA(LPCSTR lpFileName)
 {
+	// Is verified that both windows and wine are calling GetFullPathNameW
 	DEBUG_LOG_TRUE(L"GetFileAttributesA", L"");
-	UBA_ASSERT(!g_runningRemote);
 	return True_GetFileAttributesA(lpFileName);
 }
 
diff --git a/Engine/Source/Programs/UnrealBuildAccelerator/Host/Private/UbaExports.cpp b/Engine/Source/Programs/UnrealBuildAccelerator/Host/Private/UbaExports.cpp
--- a/Engine/Source/Programs/UnrealBuildAccelerator/Host/Private/UbaExports.cpp
+++ b/Engine/Source/Programs/UnrealBuildAccelerator/Host/Private/UbaExports.cpp
@@ -198,6 +198,10 @@ extern "C"
 		using namespace uba;
 
 		#if UBA_USE_AWS
+		StringBuffer<> fixedRootDir;
+		fixedRootDir.count = GetFullPathNameW(rootDir, fixedRootDir.capacity, fixedRootDir.data, NULL);
+		fixedRootDir.Replace('/', PathSeparator).EnsureEndsWithSlash();
+		rootDir = fixedRootDir.data;
 		AWS aws;
 		if (!zone || !*zone)
 		{
