[UBA]
* Added Bottleneck type that can be used in combination with BottleneckScope to prevent more than x threads to pass through at the same time

diff --git a/Engine/Source/Programs/UnrealBuildAccelerator/Common/Private/UbaSession.cpp b/Engine/Source/Programs/UnrealBuildAccelerator/Common/Private/UbaSession.cpp
--- a/Engine/Source/Programs/UnrealBuildAccelerator/Common/Private/UbaSession.cpp
+++ b/Engine/Source/Programs/UnrealBuildAccelerator/Common/Private/UbaSession.cpp
@@ -1,6 +1,7 @@
 // Copyright Epic Games, Inc. All Rights Reserved.
 
 #include "UbaSession.h"
+#include "UbaBottleneck.h"
 #include "UbaFileAccessor.h"
 #include "UbaProcess.h"
 #include "UbaStorage.h"
@@ -1908,9 +1909,7 @@ namespace uba
 			auto memClose = MakeGuard([&](){ UnmapViewOfFile(mem, fileSize); });
 
 			constexpr bool useFileMapForWrite = false;
-			//bool useOverlap = false;// fileSize > 8 * 1024 * 1024;
-
-			static bool useOverlap = false;
+			bool useOverlap = false;// fileSize > 8 * 1024 * 1024;
 
 
 			u32 attributes = DefaultAttributes();
@@ -1930,16 +1929,14 @@ namespace uba
 				if (!destinationFile.CreateWrite(false, attributes, fileSize, m_tempPath.data))
 					return false;
 
-				u64 start = GetTime();
+				// This is to kill I/O when writing lots of pdb/dlls in parallel
+				#if PLATFORM_WINDOWS
+				static Bottleneck bottleneck(16);
+				BottleneckScope scope(bottleneck);
+				#endif
+
 				if (!destinationFile.Write(mem, fileSize))
 					return false;
-
-				if (!useOverlap && TimeToMs(GetTime() - start) > 4 * 1000)
-				{
-					useOverlap = true;
-					m_logger.Info(TC("Switching to overlapped I/O due to slow writes"));
-				}
-
 			}
 			if (u64 time = file.lastWriteTime)
 				if (!SetFileLastWriteTime(destinationFile.GetHandle(), time))
diff --git a/Engine/Source/Programs/UnrealBuildAccelerator/Common/Public/UbaBottleneck.h b/Engine/Source/Programs/UnrealBuildAccelerator/Common/Public/UbaBottleneck.h
new file mode 100644
--- /dev/null
+++ b/Engine/Source/Programs/UnrealBuildAccelerator/Common/Public/UbaBottleneck.h
@@ -0,0 +1,52 @@
+// Copyright Epic Games, Inc. All Rights Reserved.
+
+#pragma once
+
+#include "UbaEvent.h"
+#include "UbaSynchronization.h"
+
+namespace uba
+{
+
+	struct Bottleneck
+	{
+		Bottleneck(u32 mc) : underMax(true), activeCount(0), maxCount(mc) {}
+		ReaderWriterLock lock;
+		Event underMax;
+		u32 activeCount;
+		u32 maxCount;
+	};
+
+	struct BottleneckScope
+	{
+		BottleneckScope(Bottleneck& b) : bottleneck(b)
+		{
+			ScopedWriteLock lock(bottleneck.lock);
+			while (true)
+			{
+				if (bottleneck.activeCount < bottleneck.maxCount)
+				{
+					++bottleneck.activeCount;
+					if (bottleneck.activeCount == bottleneck.maxCount)
+						bottleneck.underMax.Reset();
+					break;
+				}
+
+				lock.Leave();
+				bottleneck.underMax.IsSet();
+				lock.Enter();
+			}
+		}
+
+		~BottleneckScope()
+		{
+			ScopedWriteLock lock(bottleneck.lock);
+			if (bottleneck.activeCount == bottleneck.maxCount)
+				bottleneck.underMax.Set();
+			--bottleneck.activeCount;
+		}
+
+		Bottleneck& bottleneck;
+	};
+
+}
