[UBA]
* Downgraded tcp bind error to info since it could happen if ppl are running compile and cook at the same time.

diff --git a/Engine/Source/Programs/UnrealBuildAccelerator/Common/Private/UbaNetworkBackendTcp.cpp b/Engine/Source/Programs/UnrealBuildAccelerator/Common/Private/UbaNetworkBackendTcp.cpp
--- a/Engine/Source/Programs/UnrealBuildAccelerator/Common/Private/UbaNetworkBackendTcp.cpp
+++ b/Engine/Source/Programs/UnrealBuildAccelerator/Common/Private/UbaNetworkBackendTcp.cpp
@@ -27,6 +27,7 @@
 #define INVALID_SOCKET -1
 #define SD_BOTH SHUT_RDWR
 #define WSAHOST_NOT_FOUND 0
+#define WSAEADDRINUSE EADDRINUSE
 #define closesocket(a) close(a)
 #define addrinfoW addrinfo
 #define GetAddrInfoW getaddrinfo
@@ -326,7 +327,11 @@ namespace uba
 
 		if (res == SOCKET_ERROR)
 		{
-			logger.Error(TC("bind %s:%hu failed (%s)"), entry.ip.data, entry.port, LastErrorToText(WSAGetLastError()).data);
+			int lastError = WSAGetLastError();
+			if (lastError == WSAEADDRINUSE)
+				logger.Info(TC("bind %s:%hu failed because address/port is in use. Some other process is already using this address/port"), entry.ip.data, entry.port);
+			else
+				logger.Error(TC("bind %s:%hu failed (%s)"), entry.ip.data, entry.port, LastErrorToText(lastError).data);
 			return;
 		}
 
