[UBA]
* Enabled overlap write and bottleneck max to 16 when writing files to disk... this seems to give best results on both machines with ntfs compression and non-compression

diff --git a/Engine/Source/Programs/UnrealBuildAccelerator/Common/Private/UbaSession.cpp b/Engine/Source/Programs/UnrealBuildAccelerator/Common/Private/UbaSession.cpp
--- a/Engine/Source/Programs/UnrealBuildAccelerator/Common/Private/UbaSession.cpp
+++ b/Engine/Source/Programs/UnrealBuildAccelerator/Common/Private/UbaSession.cpp
@@ -1908,8 +1908,12 @@ namespace uba
 
 			auto memClose = MakeGuard([&](){ UnmapViewOfFile(mem, fileSize); });
 
+			// Seems like best combo (for windows at least) is to use writes with overlap and max 16 at the same time.
+			// On one machine we get twice as fast without overlap if no bottleneck. On another machine (ntfs compression on) we get twice as slow without overlap
+			// Both machines behaves well with overlap AND bottleneck. Both machine are 128 logical core thread rippers.
 			constexpr bool useFileMapForWrite = false;
-			bool useOverlap = false;// fileSize > 8 * 1024 * 1024;
+			constexpr u32 bottleneckMax = 16;
+			bool useOverlap = fileSize > 8 * 1024 * 1024;
 
 
 			u32 attributes = DefaultAttributes();
@@ -1931,7 +1935,7 @@ namespace uba
 
 				// This is to kill I/O when writing lots of pdb/dlls in parallel
 				#if PLATFORM_WINDOWS
-				static Bottleneck bottleneck(16);
+				static Bottleneck bottleneck(bottleneckMax);
 				BottleneckScope scope(bottleneck);
 				#endif
 
