[UBA]
* Added ProcessStartInfo::writeOutputFilesOnFail that can be set to true if we want the output files to be written/sent back even though process exited with errors
* Changed so session client use ProcessStartInfoHolder and moved serialization code to that class

diff --git a/Engine/Source/Programs/UnrealBuildAccelerator/Common/Private/UbaProcess.cpp b/Engine/Source/Programs/UnrealBuildAccelerator/Common/Private/UbaProcess.cpp
--- a/Engine/Source/Programs/UnrealBuildAccelerator/Common/Private/UbaProcess.cpp
+++ b/Engine/Source/Programs/UnrealBuildAccelerator/Common/Private/UbaProcess.cpp
@@ -804,8 +804,9 @@ namespace uba
 
 					m_processStats.Add(stats);
 
-					if (GetApplicationRules()[m_rulesIndex].rules->IsExitCodeSuccess(m_nativeProcessExitCode) && !IsCancelled())
-						WriteFilesToDisk();
+					if (!IsCancelled())
+						if (m_startInfo.writeOutputFilesOnFail || GetApplicationRules()[m_rulesIndex].rules->IsExitCodeSuccess(m_nativeProcessExitCode))
+							WriteFilesToDisk();
 
 					if (m_parentProcess)
 					{
diff --git a/Engine/Source/Programs/UnrealBuildAccelerator/Common/Private/UbaProcessStartInfoHolder.h b/Engine/Source/Programs/UnrealBuildAccelerator/Common/Private/UbaProcessStartInfoHolder.h
--- a/Engine/Source/Programs/UnrealBuildAccelerator/Common/Private/UbaProcessStartInfoHolder.h
+++ b/Engine/Source/Programs/UnrealBuildAccelerator/Common/Private/UbaProcessStartInfoHolder.h
@@ -22,6 +22,8 @@ namespace uba
 
 	struct ProcessStartInfoHolder
 	{
+		ProcessStartInfoHolder() {}
+
 		ProcessStartInfoHolder(const ProcessStartInfo& si)
 		{
 			startInfo = si;
@@ -51,6 +53,38 @@ namespace uba
 			}
 		}
 
+		void Write(BinaryWriter& writer)
+		{
+			writer.WriteString(description);
+			writer.WriteString(application);
+			writer.WriteString(arguments);
+			writer.WriteString(workingDir);
+			writer.WriteU32(*(u32*)&weight);
+			writer.WriteBool(startInfo.writeOutputFilesOnFail);
+			writer.WriteU64(startInfo.outputStatsThresholdMs);
+		}
+
+		void Read(BinaryReader& reader)
+		{
+			description = reader.ReadString();
+			application = reader.ReadString();
+			arguments = reader.ReadString();
+			workingDir = reader.ReadString();
+
+			Replace(application.data(), '/', PathSeparator); // TODO: Is this needed?
+
+			u32 weight32 = reader.ReadU32();
+			weight = *(float*)&weight32;
+			
+			startInfo.outputStatsThresholdMs = reader.ReadU64();
+			startInfo.writeOutputFilesOnFail = reader.ReadBool();
+
+			startInfo.description = description.c_str();
+			startInfo.application = application.c_str();
+			startInfo.arguments = arguments.c_str();
+			startInfo.workingDir = workingDir.c_str();
+		}
+
 		ProcessStartInfoHolder(const ProcessStartInfoHolder& o)
 		{
 			*this = o;
@@ -83,6 +117,7 @@ namespace uba
 		TString arguments;
 		TString workingDir;
 		TString logFile;
+		float weight = 1.0f;
 	};
 
 }
\ No newline at end of file
diff --git a/Engine/Source/Programs/UnrealBuildAccelerator/Common/Private/UbaSessionClient.cpp b/Engine/Source/Programs/UnrealBuildAccelerator/Common/Private/UbaSessionClient.cpp
--- a/Engine/Source/Programs/UnrealBuildAccelerator/Common/Private/UbaSessionClient.cpp
+++ b/Engine/Source/Programs/UnrealBuildAccelerator/Common/Private/UbaSessionClient.cpp
@@ -4,6 +4,7 @@
 #include "UbaNetworkClient.h"
 #include "UbaNetworkMessage.h"
 #include "UbaProcess.h"
+#include "UbaProcessStartInfoHolder.h"
 #include "UbaProtocol.h"
 #include "UbaStorage.h"
 
@@ -1071,14 +1072,9 @@ namespace uba
 		m_environmentVariables.push_back(0);
 	}
 
-	struct SessionClient::InternalProcessStartInfo : ProcessStartInfo
+	struct SessionClient::InternalProcessStartInfo : ProcessStartInfoHolder
 	{
 		u32 processId = 0;
-		float weight = 1.0f;
-		TString descriptionBuf;
-		TString applicationBuf;
-		TString argumentsBuf;
-		TString workingDirBuf;
 	};
 
 
@@ -1114,16 +1110,7 @@ namespace uba
 			out.push_back({});
 			InternalProcessStartInfo& info = out.back();
 			info.processId = processId;
-			info.descriptionBuf = reader.ReadString();
-			info.applicationBuf = reader.ReadString();
-			info.argumentsBuf = reader.ReadString();
-			info.workingDirBuf = reader.ReadString();
-			u32 weight32 = reader.ReadU32();
-			info.outputStatsThresholdMs = reader.ReadU64();
-			info.weight = *(float*)&weight32;
-
-			Replace(info.applicationBuf.data(), '/', PathSeparator);
-
+			info.Read(reader);
 		}
 
 		u32 neededDirectoryTableSize = reader.ReadU32();
@@ -1412,13 +1399,9 @@ namespace uba
 					waitTimeoutMs = 200;
 				}
 
-				for (InternalProcessStartInfo& startInfo : startInfos)
+				for (InternalProcessStartInfo& info : startInfos)
 				{
-					startInfo.description = startInfo.descriptionBuf.c_str();
-					startInfo.application = startInfo.applicationBuf.c_str();
-					startInfo.arguments = startInfo.argumentsBuf.c_str();
-					startInfo.workingDir = startInfo.workingDirBuf.c_str();
-
+					auto& startInfo = info.startInfo;
 					startInfo.uiLanguage = int(m_uiLanguage);
 					startInfo.priorityClass = m_defaultPriorityClass;
 					startInfo.useCustomAllocator = !m_disableCustomAllocator;
@@ -1434,22 +1417,22 @@ namespace uba
 					}
 
 					StringBuffer<> realApplication;
-					if (!EnsureApplicationEnvironment(realApplication, startInfo.processId, startInfo.application))
+					if (!EnsureApplicationEnvironment(realApplication, info.processId, startInfo.application))
 					{
 						m_logger.Error(TC("Failed to ensure application environment for %s"), startInfo.application);
-						SendReturnProcess(startInfo.processId, TC("Failed to ensure application environment"));
+						SendReturnProcess(info.processId, TC("Failed to ensure application environment"));
 						m_loop = false;
 						break;
 					}
 
 					void* env = GetProcessEnvironmentVariables();
 
-					auto process = new ProcessImpl(*this, startInfo.processId, nullptr);
+					auto process = new ProcessImpl(*this, info.processId, nullptr);
 
 					activeProcesses.emplace_back(process);
 					ProcessRec* rec = &activeProcesses.back();
 
-					rec->weight = startInfo.weight;
+					rec->weight = info.weight;
 
 					{
 						ScopedWriteLock lock(activeWeightLock);
@@ -1476,9 +1459,10 @@ namespace uba
 						ProcessRec* rec = er->rec;
 						delete er;
 
+						auto& startInfo = h.GetStartInfo();
 						if (session.m_shouldSendLogToServer)
 						{
-							if (const tchar* logFile = h.GetStartInfo().logFile)
+							if (const tchar* logFile = startInfo.logFile)
 							{
 								WrittenFile f;
 								f.name = logFile;
@@ -1526,17 +1510,9 @@ namespace uba
 								rec->isDone = true;
 								return;
 							}
-
-							// TODO: Remove this once m_isTerminating is working properly
-							if (IsAWSTermination(process))
-							{
-								session.m_loop = false;
-								session.m_logger.Info(TC("Got error that is most likely because AWS termination. If not, this should be checked (%s)"), process.m_description.c_str());
-								rec->isDone = true;
-								return;
-							}
 						}
-						else
+
+						if (exitCode == 0 || startInfo.writeOutputFilesOnFail)
 						{
 							// Should we decrease weight before or after sending files?
 							//decreaseWeight.Execute();
@@ -1650,17 +1626,6 @@ namespace uba
 		}
 	}
 
-	bool SessionClient::IsAWSTermination(ProcessImpl& process)
-	{
-		// This seems to be AWS termination error code... let's just shutdown the entire agent
-		if (process.m_exitCode == 4 || (process.m_exitCode == 1 && Contains(process.m_realApplication.c_str(), TC("link.exe"))))
-			return true;
-		//for (auto c : process.m_childProcesses)
-		//	if (IsAWSTermination(*(ProcessImpl*)c.m_process))
-		//		return true;
-		return false;
-	}
-
 	bool SessionClient::AllocFailed(Process& process, const tchar* allocType, u32 error)
 	{
 		//StackBinaryWriter<32> writer;
diff --git a/Engine/Source/Programs/UnrealBuildAccelerator/Common/Private/UbaSessionServer.cpp b/Engine/Source/Programs/UnrealBuildAccelerator/Common/Private/UbaSessionServer.cpp
--- a/Engine/Source/Programs/UnrealBuildAccelerator/Common/Private/UbaSessionServer.cpp
+++ b/Engine/Source/Programs/UnrealBuildAccelerator/Common/Private/UbaSessionServer.cpp
@@ -11,13 +11,13 @@ namespace uba
 	class SessionServer::RemoteProcess : public Process, public ProcessStartInfoHolder
 	{
 	public:
-		RemoteProcess(SessionServer* server, const ProcessStartInfo& si, u32 processId, float weight)
+		RemoteProcess(SessionServer* server, const ProcessStartInfo& si, u32 processId, float weight_)
 		:	ProcessStartInfoHolder(si)
 		,	m_server(server)
 		,	m_processId(processId)
 		,	m_done(true)
-		,	m_weight(weight)
 		{
+			weight = weight_;
 		}
 
 		~RemoteProcess()
@@ -79,7 +79,6 @@ namespace uba
 		bool m_cancelled = false;
 		u32 m_clientId = ~0u;
 		u32 m_sessionId = 0;
-		float m_weight;
 		TString m_executingHost;
 
 		struct KnownInput { CasKey key; u32 mappingAlignment = 0; };
@@ -909,12 +908,7 @@ namespace uba
 
 					ProcessAdded(*process, sessionId);
 					writer.WriteU32(process->m_processId);
-					writer.WriteString(process->startInfo.description);
-					writer.WriteString(process->startInfo.application);
-					writer.WriteString(process->startInfo.arguments);
-					writer.WriteString(process->startInfo.workingDir);
-					writer.WriteU32(*(u32*)&process->m_weight);
-					writer.WriteU64(process->startInfo.outputStatsThresholdMs);
+					process->Write(writer);
 
 					for (auto kiIt = process->m_knownInputs, kiEnd = kiIt + process->m_knownInputsCount; kiIt!=kiEnd; ++kiIt)
 						if (session.sentKeys.insert(kiIt->key).second)
@@ -925,7 +919,7 @@ namespace uba
 					if (writer.GetCapacityLeft() < 5000) // Arbitrary number to cover all parameters above
 						break;
 
-					weightLeft -= process->m_weight;
+					weightLeft -= process->weight;
 				}
 				fillLock.Leave();
 
diff --git a/Engine/Source/Programs/UnrealBuildAccelerator/Common/Public/UbaNetwork.h b/Engine/Source/Programs/UnrealBuildAccelerator/Common/Public/UbaNetwork.h
--- a/Engine/Source/Programs/UnrealBuildAccelerator/Common/Public/UbaNetwork.h
+++ b/Engine/Source/Programs/UnrealBuildAccelerator/Common/Public/UbaNetwork.h
@@ -18,7 +18,7 @@ namespace uba
 	
 	static constexpr u32 SystemNetworkVersion = 1338;
 	static constexpr u32 StorageNetworkVersion = 3;
-	static constexpr u32 SessionNetworkVersion = 24;
+	static constexpr u32 SessionNetworkVersion = 25;
 
 
 	// Messages used over network between client/server (system, storage and session)
@@ -77,6 +77,7 @@ namespace uba
 		UBA_SESSION_MESSAGE(VirtualAllocFailed) \
 		UBA_SESSION_MESSAGE(GetTraceInformation) \
 		UBA_SESSION_MESSAGE(Ping) \
+		UBA_SESSION_MESSAGE(GetNextProcess) \
 		UBA_SESSION_MESSAGE(Custom) \
 		UBA_SESSION_MESSAGE(UpdateEnvironment) \
 		UBA_SESSION_MESSAGE(Summary) \
diff --git a/Engine/Source/Programs/UnrealBuildAccelerator/Common/Public/UbaProcessStartInfo.h b/Engine/Source/Programs/UnrealBuildAccelerator/Common/Public/UbaProcessStartInfo.h
--- a/Engine/Source/Programs/UnrealBuildAccelerator/Common/Public/UbaProcessStartInfo.h
+++ b/Engine/Source/Programs/UnrealBuildAccelerator/Common/Public/UbaProcessStartInfo.h
@@ -23,6 +23,7 @@ namespace uba
 		u64 outputStatsThresholdMs = ~u64(0);	// Threshold in milliseconds where process summary should be printed to log
 		bool trackInputs = false;				// Track all files read. Can read result in ProcessHandle.GetTrackedInputs()
 		bool useCustomAllocator = true;			// Disable detouring of allocator inside processes. If Session.disableCustomAllocator is false this will be overridden
+		bool writeOutputFilesOnFail = false;	// If set to true, output files will be written to disk regardless if process succeeds or not
 
 		using LogLineCallback = void(void* userData, const tchar* line, u32 length, LogEntryType type);
 		LogLineCallback* logLineFunc = nullptr;	// Callback for when log entries happens
diff --git a/Engine/Source/Programs/UnrealBuildAccelerator/Common/Public/UbaSessionClient.h b/Engine/Source/Programs/UnrealBuildAccelerator/Common/Public/UbaSessionClient.h
--- a/Engine/Source/Programs/UnrealBuildAccelerator/Common/Public/UbaSessionClient.h
+++ b/Engine/Source/Programs/UnrealBuildAccelerator/Common/Public/UbaSessionClient.h
@@ -78,7 +78,6 @@ namespace uba
 		void SendSummary();
 		u32 CountLogLines(ProcessImpl& process);
 		void WriteLogLines(BinaryWriter& writer, ProcessImpl& process);
-		static bool IsAWSTermination(ProcessImpl& process);
 
 		void ThreadCreateProcessLoop();
 
