[UBA]
* Fixed so two traces can't be writing to the same shared memory at the same time
* Added commented out code to track number of active file mappings

diff --git a/Engine/Source/Programs/UnrealBuildAccelerator/Common/Private/UbaFileMapping.cpp b/Engine/Source/Programs/UnrealBuildAccelerator/Common/Private/UbaFileMapping.cpp
--- a/Engine/Source/Programs/UnrealBuildAccelerator/Common/Private/UbaFileMapping.cpp
+++ b/Engine/Source/Programs/UnrealBuildAccelerator/Common/Private/UbaFileMapping.cpp
@@ -16,10 +16,14 @@ namespace uba
 #if PLATFORM_WINDOWS
 	ReaderWriterLock g_createFileHandleLock;
 
+	//Atomic<u64> g_fileMappingCount;
+
 	HANDLE asHANDLE(FileHandle fh);
 
 	HANDLE InternalCreateFileMappingW(HANDLE hFile, DWORD flProtect, DWORD dwMaximumSizeHigh, DWORD dwMaximumSizeLow, LPCWSTR lpName)
 	{
+		//++g_fileMappingCount;
+
 		if (flProtect != PAGE_READWRITE)
 			return ::CreateFileMappingW(hFile, NULL, flProtect, dwMaximumSizeHigh, dwMaximumSizeLow, lpName);
 
@@ -222,6 +226,14 @@ namespace uba
 	bool CloseFileMapping(FileMappingHandle h)
 	{
 #if PLATFORM_WINDOWS
+		//if (h.handle)
+		//{
+		//	--g_fileMappingCount;
+		//	StringBuffer<256> sb;
+		//	sb.Appendf(TC("FILEMAPPING COUNT: %llu\r\n"), g_fileMappingCount.load());
+		//	OutputDebugStringW(sb.data);
+		//}
+
 		return CloseHandle(h.handle);
 #else
 		if (h.shmFd == -1)
diff --git a/Engine/Source/Programs/UnrealBuildAccelerator/Common/Private/UbaTrace.cpp b/Engine/Source/Programs/UnrealBuildAccelerator/Common/Private/UbaTrace.cpp
--- a/Engine/Source/Programs/UnrealBuildAccelerator/Common/Private/UbaTrace.cpp
+++ b/Engine/Source/Programs/UnrealBuildAccelerator/Common/Private/UbaTrace.cpp
@@ -50,8 +50,9 @@ namespace uba
 		m_memoryHandle = uba::CreateMemoryMappingW(m_logger, PAGE_READWRITE|SEC_RESERVE, m_memoryCapacity, namedTrace);
 		if (!m_memoryHandle.IsValid())
 			return false;
+		if (GetLastError() != ERROR_ALREADY_EXISTS)
+			m_memoryBegin = MapViewOfFile(m_memoryHandle, FILE_MAP_WRITE, 0, m_memoryCapacity);
 
-		m_memoryBegin = MapViewOfFile(m_memoryHandle, FILE_MAP_WRITE, 0, m_memoryCapacity);
 		if (!m_memoryBegin)
 		{
 			CloseFileMapping(m_memoryHandle);
