[UBA]
* Added GetNextProcess as a real message type (not using Custom) to be able to stop fetching work if helper is being terminated by aws or if we want to scale down number of workers
* Did some cleanup in the Tcp backend code and made sure not to call close socket multiple times (since it can cause the code to close a different socket than what is owned)

diff --git a/Engine/Source/Programs/UnrealBuildAccelerator/Common/Private/UbaNetworkBackendTcp.cpp b/Engine/Source/Programs/UnrealBuildAccelerator/Common/Private/UbaNetworkBackendTcp.cpp
--- a/Engine/Source/Programs/UnrealBuildAccelerator/Common/Private/UbaNetworkBackendTcp.cpp
+++ b/Engine/Source/Programs/UnrealBuildAccelerator/Common/Private/UbaNetworkBackendTcp.cpp
@@ -47,7 +47,6 @@ namespace uba
 		Thread thread;
 		Event listening;
 		SOCKET socket = INVALID_SOCKET;
-		SOCKET socketToClose = INVALID_SOCKET;
 	};
 
 	struct NetworkBackendTcp::Connection
@@ -100,6 +99,13 @@ namespace uba
 		return true;
 	}
 
+	bool CloseSocket(Logger& logger, SOCKET s)
+	{
+		if (closesocket(s) == SOCKET_ERROR)
+			return logger.Error(TC("failed to close socket (%s)"), LastErrorToText(WSAGetLastError()).data);
+		return true;
+	}
+
 
 
 	NetworkBackendTcp::NetworkBackendTcp(LogWriter& writer, const tchar* prefix)
@@ -122,7 +128,7 @@ namespace uba
 			shutdown(s, SD_BOTH);
 			lock2.Leave();
 			conn.recvThread.Wait();
-			closesocket(s);
+			CloseSocket(m_logger, s);
 		}
 		m_connections.clear();
 
@@ -151,7 +157,7 @@ namespace uba
 		conn.socket = INVALID_SOCKET;
 		lock.Leave();
 		conn.recvThread.Wait();
-		closesocket(s);
+		CloseSocket(m_logger, s);
 	}
 
 	bool NetworkBackendTcp::Send(Logger& logger, void* connection, const void* data, u32 dataSize, SendContext& sendContext)
@@ -203,6 +209,11 @@ namespace uba
 
 	void NetworkBackendTcp::SetDisconnectCallback(void* connection, void* context, DisconnectCallback* callback)
 	{
+		{
+			ScopedReadLock lock(m_connectionsLock);
+			if (m_connections.empty())
+				return;
+		}
 		auto& conn = *(Connection*)connection;
 		ScopedCriticalSection lock(conn.shutdownLock);
 		conn.disconnectCallback = callback;
@@ -269,20 +280,13 @@ namespace uba
 	void NetworkBackendTcp::StopListen()
 	{
 		for (auto& e : m_listenEntries)
-		{
-			e.socketToClose = e.socket;
 			e.socket = INVALID_SOCKET;
-			shutdown(e.socketToClose, SD_BOTH);
-		}
 		for (auto& e : m_listenEntries)
-		{
 			e.thread.Wait();
-			closesocket(e.socketToClose);
-		}
 		m_listenEntries.clear();
 	}
 
-	void NetworkBackendTcp::ThreadListen(Logger& logger, ListenEntry& entry)
+	bool NetworkBackendTcp::ThreadListen(Logger& logger, ListenEntry& entry)
 	{
 		addrinfoW hints;
 		memset(&hints, 0, sizeof(hints));
@@ -300,27 +304,21 @@ namespace uba
 		auto listenEv = MakeGuard([&]() { entry.listening.Set(); });
 
 		if (res != 0)
-		{
-			logger.Error(TC("getaddrinfo failed (%d)"), res);
-			return;
-		}
+			return logger.Error(TC("getaddrinfo failed (%d)"), res);
+
 		UBA_ASSERT(result);
 		auto addrGuard = MakeGuard([result]() { FreeAddrInfoW(result); });
 
 		// Create a socket for listening to connections
 		SOCKET listenSocket = socket(result->ai_family, result->ai_socktype, result->ai_protocol);
 		if (listenSocket == INVALID_SOCKET)
-		{
-			logger.Error(TC("socket failed (%s)"), LastErrorToText(WSAGetLastError()).data);
-			return;
-		}
-		auto listenSocketCleanup = MakeGuard([&]() { closesocket(listenSocket); listenSocket = INVALID_SOCKET; });
+			return logger.Error(TC("socket failed (%s)"), LastErrorToText(WSAGetLastError()).data);
 
-		// This is here to be able to iterate fast when doing development.. seems like socket ends up in TIME_WAIT state after close and it takes some time to be able to use socket again
-		#if !PLATFORM_WINDOWS// && UBA_DEBUG
-		int optval = 1;
-		::setsockopt(listenSocket, SOL_SOCKET, SO_REUSEADDR, (const char*)&optval, sizeof optval);
-		#endif
+		auto listenSocketCleanup = MakeGuard([&]() { CloseSocket(m_logger, listenSocket); });
+
+		u32 reuseAddr = 1;
+		if (::setsockopt(listenSocket, SOL_SOCKET, SO_REUSEADDR, (const char*)&reuseAddr, sizeof reuseAddr) == SOCKET_ERROR)
+			return logger.Error(TC("setsockopt SO_REUSEADDR failed (error: %s)"), LastErrorToText(WSAGetLastError()).data);
 
 		// Setup the TCP listening socket
 		res = bind(listenSocket, result->ai_addr, (socklen_t)result->ai_addrlen);
@@ -328,76 +326,54 @@ namespace uba
 		if (res == SOCKET_ERROR)
 		{
 			int lastError = WSAGetLastError();
-			if (lastError == WSAEADDRINUSE)
-				logger.Info(TC("bind %s:%hu failed because address/port is in use. Some other process is already using this address/port"), entry.ip.data, entry.port);
-			else
-				logger.Error(TC("bind %s:%hu failed (%s)"), entry.ip.data, entry.port, LastErrorToText(lastError).data);
-			return;
+			if (lastError != WSAEADDRINUSE)
+				return logger.Error(TC("bind %s:%hu failed (%s)"), entry.ip.data, entry.port, LastErrorToText(lastError).data);
+			logger.Info(TC("bind %s:%hu failed because address/port is in use. Some other process is already using this address/port"), entry.ip.data, entry.port);
+			return false;
 		}
 
 		addrGuard.Execute();
 
 		res = listen(listenSocket, SOMAXCONN);
 		if (res == SOCKET_ERROR)
-		{
-			logger.Error(TC("Listen failed (%s)"), LastErrorToText(WSAGetLastError()).data);
-			return;
-		}
+			return logger.Error(TC("Listen failed (%s)"), LastErrorToText(WSAGetLastError()).data);
 
 		if (!SetKeepAlive(logger, listenSocket))
-			return;
-
-
-		listenSocketCleanup.Cancel();
+			return false;
 
 		logger.Info(TC("Listening on %s:%hu"), entry.ip.data, entry.port);
 		entry.socket = listenSocket;
 
 		listenEv.Execute();
 
-		fd_set set;
-		fd_set read_fds; // temp file descriptor list for select()
-		int fdmax = 1; // maximum file descriptor number
-
-		struct timeval timeout;
-		timeout.tv_sec = 1;
-		timeout.tv_usec = 0;
-
-		FD_ZERO(&set);
-		FD_ZERO(&read_fds);
-
-		FD_SET(listenSocket, &set);
-
-		#if !PLATFORM_WINDOWS
-		fdmax = listenSocket + 1;
-		#endif
-
 		while (entry.socket != INVALID_SOCKET)
 		{
-			read_fds = set;
-
-			// [honk] TODO: There was a new hang here that has never happened before so brought back this code
-			#if PLATFORM_WINDOWS
-			FD_ZERO(&read_fds);
-			FD_SET(listenSocket, &read_fds);
-			#endif
-
-			int ret = select(fdmax, &read_fds, NULL, NULL, &timeout);
+			WSAPOLLFD p;
+			p.fd = listenSocket;
+			p.revents = 0;
+			p.events = POLLIN;
+			int timeoutMs = 2000;
+			int pollRes = WSAPoll(&p, 1, timeoutMs);
+
+			if (pollRes == SOCKET_ERROR)
+			{
+				int lastError = WSAGetLastError();
+				logger.Warning(TC("WSAPoll returned error %s"), LastErrorToText(lastError).data);
+				break;
+			}
 
-			if (ret == 0)
+			if (!pollRes)
 				continue;
 
-			if (ret == -1)
+			if (p.revents & POLLNVAL)
 			{
-				if (entry.socket != INVALID_SOCKET)
-					logger.Info(TC("Select (for accept) failed with WSA error: %s"), LastErrorToText(WSAGetLastError()).data);
-				break;
+				logger.Warning(TC("WSAPoll returned successful but with unexpected flags: %u"), p.revents);
+				continue;
 			}
 
 			sockaddr remoteSockAddr = { 0 }; // for TCP/IP
 			socklen_t remoteSockAddrLen = sizeof(remoteSockAddr);
 			SOCKET clientSocket = accept(listenSocket, (sockaddr*)&remoteSockAddr, &remoteSockAddrLen);
-			FD_SET(clientSocket, &set);
 
 			if (clientSocket == INVALID_SOCKET)
 			{
@@ -406,7 +382,7 @@ namespace uba
 				break;
 			}
 
-			auto socketClose = MakeGuard([&]() { closesocket(clientSocket); });
+			auto socketClose = MakeGuard([&]() { CloseSocket(logger, clientSocket); });
 
 			if (!DisableNagle(logger, clientSocket))
 				continue;
@@ -431,6 +407,8 @@ namespace uba
 
 			socketClose.Cancel();
 		}
+
+		return true;
 	}
 
 	void NetworkBackendTcp::ThreadRecv(Connection& connection)
@@ -498,7 +476,7 @@ namespace uba
 		if (s == INVALID_SOCKET)
 			return;
 		shutdown(s, SD_BOTH);
-		closesocket(s);
+		CloseSocket(m_logger, s);
 	}
 
 	bool NetworkBackendTcp::Connect(Logger& logger, const tchar* ip, const ConnectedFunc& connectedFunc, u16 port, bool* timedOut)
@@ -556,7 +534,7 @@ namespace uba
 			return logger.Error(TC("socket failed (%s)"), LastErrorToText(WSAGetLastError()).data);
 
 		// Create guard in case we fail to connect (will be cancelled further down if we succeed)
-		auto socketClose = MakeGuard([&]() { closesocket(socketFd); });
+		auto socketClose = MakeGuard([&]() { CloseSocket(logger, socketFd); });
 
 		// Set to non-blocking just for the connect call (we want to control the connect timeout after connect using select instead)
 		if (!SetBlocking(logger, socketFd, false))
@@ -595,7 +573,7 @@ namespace uba
 		if (pollRes == SOCKET_ERROR)
 		{
 			int lastError = WSAGetLastError();
-			logger.Warning(TC("WSAPoll returned error %s (%s%s)"), LastErrorToText(lastError).data, nameHint);
+			logger.Warning(TC("WSAPoll returned error %s (%s)"), LastErrorToText(lastError).data, nameHint);
 			return false;
 		}
 
@@ -657,6 +635,7 @@ namespace uba
 
 		if (!connectedFunc(&conn, remoteSocketAddr, timedOut))
 		{
+			socketFd = conn.socket;
 			conn.socket = INVALID_SOCKET;
 			socketClose.Execute();
 			conn.ready.Set();
@@ -888,12 +867,16 @@ namespace uba
 
 	HttpConnection::~HttpConnection()
 	{
+		if (m_socket != INVALID_SOCKET)
+		{
+			LoggerWithWriter logger(g_nullLogWriter);
+			CloseSocket(logger, m_socket);
+		}
+
 		#if PLATFORM_WINDOWS
 		if (m_wsaInitDone)
 			WSACleanup();
 		#endif
-		if (m_socket != INVALID_SOCKET)
-			closesocket(m_socket);
 	}
 
 	bool HttpConnection::Connect(Logger& logger, const char* host)
@@ -943,7 +926,7 @@ namespace uba
 		// TODO: Fix so we reuse socket connection for multiple queries
 		if (*m_host)// && _stricmp(m_host, host) != 0)
 		{
-			closesocket(m_socket);
+			CloseSocket(logger, m_socket);
 			m_socket = INVALID_SOCKET;
 			*m_host = 0;
 		}
diff --git a/Engine/Source/Programs/UnrealBuildAccelerator/Common/Private/UbaProcess.cpp b/Engine/Source/Programs/UnrealBuildAccelerator/Common/Private/UbaProcess.cpp
--- a/Engine/Source/Programs/UnrealBuildAccelerator/Common/Private/UbaProcess.cpp
+++ b/Engine/Source/Programs/UnrealBuildAccelerator/Common/Private/UbaProcess.cpp
@@ -824,6 +824,22 @@ namespace uba
 
 					return false;
 				}
+
+			case MessageType_GetNextProcess:
+				{
+					bool prevExitCode = reader.ReadBool();
+					bool newProcess = false;
+					NextProcessInfo nextProcess;
+					m_messageSuccess = m_session.GetNextProcess(*this, newProcess, nextProcess, prevExitCode) && m_messageSuccess;
+					writer.WriteBool(newProcess);
+					if (!newProcess)
+						return true;
+					writer.WriteString(nextProcess.arguments);
+					writer.WriteString(nextProcess.workingDir);
+					writer.WriteString(nextProcess.description);
+					return true;
+				}
+
 			case MessageType_Custom:
 				{
 					m_session.CustomMessage(*this, reader, writer);
diff --git a/Engine/Source/Programs/UnrealBuildAccelerator/Common/Private/UbaScheduler.cpp b/Engine/Source/Programs/UnrealBuildAccelerator/Common/Private/UbaScheduler.cpp
--- a/Engine/Source/Programs/UnrealBuildAccelerator/Common/Private/UbaScheduler.cpp
+++ b/Engine/Source/Programs/UnrealBuildAccelerator/Common/Private/UbaScheduler.cpp
@@ -43,9 +43,9 @@ namespace uba
 	,	m_updateThreadLoop(false)
 	,	m_enableProcessReuse(enableProcessReuse)
 	{
-		session.RegisterCustomService([this](Process& process, const void* recv, u32 recvSize, void* send, u32 sendCapacity)
+		session.RegisterGetNextProcess([this](Process& process, NextProcessInfo& outNextProcess, u32 prevExitCode)
 			{
-				return HandleReuseMessage(process, recv, recvSize, send, sendCapacity);
+				return HandleReuseMessage(process, outNextProcess, prevExitCode);
 			});
 	}
 
@@ -71,6 +71,11 @@ namespace uba
 		m_session.WaitOnAllTasks();
 	}
 
+	void Scheduler::SetMaxLocalProcessors(u32 maxLocalProcessors)
+	{
+		m_maxLocalProcessors = maxLocalProcessors;
+	}
+
 	u32 Scheduler::EnqueueProcess(const ProcessStartInfo& info, float weight, const void* knownInputs, u32 knownInputsBytes, u32 knownInputsCount)
 	{
 		UBA_ASSERTF(weight == 1.0f, TC("weight not implemented"));
@@ -201,20 +206,21 @@ namespace uba
 		return true;
 	}
 
-	u32 Scheduler::HandleReuseMessage(Process& process, const void* recv, u32 recvSize, void* send, u32 sendCapacity)
+	bool Scheduler::HandleReuseMessage(Process& process, NextProcessInfo& outNextProcess, u32 prevExitCode)
 	{
 		if (!m_enableProcessReuse)
-			return 0;
+			return false;
 
 		auto& currentStartInfo = process.GetStartInfo();
 		auto info = (ExitProcessInfo*)currentStartInfo.userData;
 		if (!info) // If null, process has already exited from some other thread
-			return 0;
+			return false;
 
 		// Call ExitedFunc and cleanup
 		if (auto func = info->originalExitedFunc)
 		{
 			UBA_ASSERT(!info->wasReturned);
+			UBA_ASSERT(prevExitCode == 0); // This is not handled for anything else
 			ProcessHandle h;
 			h.m_process = &process;
 			func(info->originalUserData, h);
@@ -230,12 +236,15 @@ namespace uba
 		info->lastWorkingDir.clear();
 		info->lastDescription.clear();
 
+		if (m_activeLocalProcesses > m_maxLocalProcessors)
+			return false;
+
 		// Try to get queued process to send back
 		ScopedWriteLock lock(m_queuedProcessesLock);
 		if (info->wasReturned)
-			return 0;
+			return false;
 		if (m_queuedProcesses.empty())
-			return 0;
+			return false;
 		auto qp = m_queuedProcesses.front();
 		m_queuedProcesses.pop_front();
 		lock.Leave();
@@ -252,12 +261,9 @@ namespace uba
 		info->lastWorkingDir = si.workingDir;
 		info->lastDescription = si.description;
 
-		// TODO: Don't think we need to udpate the other parts of StartInfo
-
-		BinaryWriter writer((u8*)send, 0, sendCapacity);
-		writer.WriteString(si.arguments);
-		writer.WriteString(si.workingDir);
-		writer.WriteString(si.description);
-		return u32(writer.GetPosition());
+		outNextProcess.arguments = si.arguments;
+		outNextProcess.workingDir = si.workingDir;
+		outNextProcess.description = si.description;
+		return true;
 	}
 }
diff --git a/Engine/Source/Programs/UnrealBuildAccelerator/Common/Private/UbaSession.cpp b/Engine/Source/Programs/UnrealBuildAccelerator/Common/Private/UbaSession.cpp
--- a/Engine/Source/Programs/UnrealBuildAccelerator/Common/Private/UbaSession.cpp
+++ b/Engine/Source/Programs/UnrealBuildAccelerator/Common/Private/UbaSession.cpp
@@ -1287,6 +1287,11 @@ namespace uba
 		m_customServiceFunction = function;
 	}
 
+	void Session::RegisterGetNextProcess(GetNextProcessFunction&& function)
+	{
+		m_getNextProcessFunction = function;
+	}
+
 	const tchar* Session::GetId() { return m_id.data; }
 	Storage& Session::GetStorage() { return m_storage; }
 	Logger& Session::GetLogger() { return m_logger; }
@@ -1993,6 +1998,18 @@ namespace uba
 		return true;
 	}
 
+	bool Session::GetNextProcess(Process& process, bool& outNewProcess, NextProcessInfo& outNextProcess, u32 prevExitCode)
+	{
+		if (!m_getNextProcessFunction)
+		{
+			outNewProcess = false;
+			return true;
+		}
+
+		outNewProcess = m_getNextProcessFunction(process, outNextProcess, prevExitCode);
+		return true;
+	}
+
 	bool Session::CustomMessage(Process& process, BinaryReader& reader, BinaryWriter& writer)
 	{
 		u32 recvSize = reader.ReadU32();
diff --git a/Engine/Source/Programs/UnrealBuildAccelerator/Common/Private/UbaSessionClient.cpp b/Engine/Source/Programs/UnrealBuildAccelerator/Common/Private/UbaSessionClient.cpp
--- a/Engine/Source/Programs/UnrealBuildAccelerator/Common/Private/UbaSessionClient.cpp
+++ b/Engine/Source/Programs/UnrealBuildAccelerator/Common/Private/UbaSessionClient.cpp
@@ -1078,7 +1078,7 @@ namespace uba
 	};
 
 
-	bool SessionClient::SendProcessAvailable(Vector<InternalProcessStartInfo>& out, float availableWeight, bool& outRemoteExecutionEnabled)
+	bool SessionClient::SendProcessAvailable(Vector<InternalProcessStartInfo>& out, float availableWeight)
 	{
 		StackBinaryWriter<32> writer;
 		NetworkMessage msg(m_client, ServiceId, SessionMessageType_ProcessAvailable, writer);
@@ -1104,7 +1104,7 @@ namespace uba
 			}
 			if (processId == SessionProcessAvailableResponse_RemoteExecutionDisabled)
 			{
-				outRemoteExecutionEnabled = false;
+				m_remoteExecutionEnabled = false;
 				break;
 			}
 			out.push_back({});
@@ -1269,8 +1269,6 @@ namespace uba
 		u64 memRequiredToSpawn = u64(double(memTotal) * double(100 - m_memWaitLoadPercent) / 100.0);
 		u64 memRequiredFree = u64(double(memTotal) * double(100 - m_memKillLoadPercent) / 100.0);
 
-		bool remoteExecutionEnabled = true;
-
 		float maxWeight = float(m_maxProcessCount);
 		float activeWeight = 0;
 		ReaderWriterLock activeWeightLock;
@@ -1291,9 +1289,9 @@ namespace uba
 				it = activeProcesses.erase(it);
 			}
 
-			if (remoteExecutionEnabled && m_terminationReason)
+			if (m_remoteExecutionEnabled && m_terminationReason)
 			{
-				remoteExecutionEnabled = false;
+				m_remoteExecutionEnabled = false;
 				m_logger.Info(TC("%s. Will stop scheduling processes and send failing processes back for retry"), m_terminationReason);
 			}
 
@@ -1302,14 +1300,14 @@ namespace uba
 				idleStartTime = GetTime();
 				processRequestCount = 0;
 			}
-			else if (remoteExecutionEnabled)
+			else if (m_remoteExecutionEnabled)
 			{
 				u32 idleTime = u32(TimeToS(GetTime() - idleStartTime));
 				if (idleTime > m_maxIdleSeconds)
 				{
 					m_logger.Info(TC("Session has been idle longer than max idle time (%u seconds). Disconnecting (Did %u process requests during idle)"), m_maxIdleSeconds, processRequestCount);
 					m_waitToSendEvent.Set();
-					remoteExecutionEnabled = false;
+					m_remoteExecutionEnabled = false;
 				}
 			}
 		};
@@ -1346,7 +1344,7 @@ namespace uba
 
 			bool firstCall = true;
 
-			while (remoteExecutionEnabled && canSpawn && m_loop)
+			while (m_remoteExecutionEnabled && canSpawn && m_loop)
 			{
 				float availableWeight;
 				{
@@ -1381,14 +1379,14 @@ namespace uba
 				}
 
 				Vector<InternalProcessStartInfo> startInfos;
-				if (!SendProcessAvailable(startInfos, availableWeight, remoteExecutionEnabled))
+				if (!SendProcessAvailable(startInfos, availableWeight))
 				{
 					m_loop = false;
 					break;
 				}
 				++processRequestCount;
 
-				if (!remoteExecutionEnabled)
+				if (!m_remoteExecutionEnabled)
 				{
 					m_logger.Info(TC("Got remote execution disabled response from host (will finish %llu active processes)"), startInfos.size() + activeProcesses.size());
 				}
@@ -1582,7 +1580,7 @@ namespace uba
 
 			RemoveInactiveProcesses();
 
-			if (activeProcesses.empty() && !remoteExecutionEnabled)
+			if (activeProcesses.empty() && !m_remoteExecutionEnabled)
 				break;
 		}
 
@@ -1641,6 +1639,57 @@ namespace uba
 		m_stats.Print(logger);
 	}
 
+	bool SessionClient::GetNextProcess(Process& process, bool& outNewProcess, NextProcessInfo& outNextProcess, u32 prevExitCode)
+	{
+		outNewProcess = false;
+
+		if (!m_remoteExecutionEnabled)
+			return true;
+
+		auto& pi = (ProcessImpl&)process;
+
+		if (!FlushWrittenFiles(pi))
+			return false;
+
+		StackBinaryReader<SendMaxSize> reader;
+		StackBinaryWriter<16 * 1024> writer;
+		NetworkMessage msg(m_client, ServiceId, SessionMessageType_GetNextProcess, writer);
+		writer.WriteU32(pi.m_id);
+		writer.WriteU32(prevExitCode);
+		pi.m_processStats.Write(writer);
+		pi.m_sessionStats.Write(writer);
+		pi.m_storageStats.Write(writer);
+		pi.m_systemStats.Write(writer);
+
+		if (!msg.Send(reader, m_stats.customMsg))
+			return false;
+
+		outNewProcess = reader.ReadBool();
+		if (outNewProcess)
+		{
+			pi.m_exitCode = prevExitCode;
+			if (m_processFinished)
+				m_processFinished(&process);
+
+			pi.m_exitCode = ~0u;
+			pi.m_processStats = {};
+			pi.m_sessionStats = {};
+			pi.m_storageStats = {};
+			pi.m_systemStats = {};
+
+			outNextProcess.arguments = reader.ReadString();
+			outNextProcess.workingDir = reader.ReadString();
+			outNextProcess.description = reader.ReadString();
+
+			// TODO: Probably need to fill up with more stuff.. this is fine for current usecase
+			pi.m_arguments = outNextProcess.arguments;
+			pi.m_description = outNextProcess.description;
+		}
+
+		reader.Reset();
+		return SendUpdateDirectoryTable(reader);
+	}
+
 	bool SessionClient::CustomMessage(Process& process, BinaryReader& reader, BinaryWriter& writer)
 	{
 		StackBinaryWriter<SendMaxSize> msgWriter;
diff --git a/Engine/Source/Programs/UnrealBuildAccelerator/Common/Private/UbaSessionServer.cpp b/Engine/Source/Programs/UnrealBuildAccelerator/Common/Private/UbaSessionServer.cpp
--- a/Engine/Source/Programs/UnrealBuildAccelerator/Common/Private/UbaSessionServer.cpp
+++ b/Engine/Source/Programs/UnrealBuildAccelerator/Common/Private/UbaSessionServer.cpp
@@ -1116,6 +1116,34 @@ namespace uba
 
 				return true;
 			}
+			case SessionMessageType_GetNextProcess:
+			{
+				u32 processId = reader.ReadU32();
+				u32 prevExitCode = reader.ReadU32();
+				ScopedWriteLock lock(m_processesLock);
+				auto findIt = m_processes.find(processId);
+				if (findIt == m_processes.end())
+					return m_logger.Error(TC("Failed to find process for id %u when receiving custom message"), processId);
+				ProcessHandle h(findIt->second);
+				lock.Leave();
+
+				auto& remoteProcess = *(RemoteProcess*)h.m_process;
+				ScopedWriteLock exitedLock(remoteProcess.m_exitedLock);
+				NextProcessInfo nextProcess;
+				bool newProcess;
+				if (!GetNextProcess(remoteProcess, newProcess, nextProcess, prevExitCode))
+					return false;
+
+				writer.WriteBool(newProcess);
+				if (newProcess)
+				{
+					writer.WriteString(nextProcess.arguments);
+					writer.WriteString(nextProcess.workingDir);
+					writer.WriteString(nextProcess.description);
+					m_trace.ProcessEnvironmentUpdated(processId, nextProcess.description.c_str(), reader.GetPositionData(), reader.GetLeft());
+				}
+				return true;
+			}
 			case SessionMessageType_Custom:
 			{
 				u32 processId = reader.ReadU32();
diff --git a/Engine/Source/Programs/UnrealBuildAccelerator/Common/Public/UbaNetworkBackendTcp.h b/Engine/Source/Programs/UnrealBuildAccelerator/Common/Public/UbaNetworkBackendTcp.h
--- a/Engine/Source/Programs/UnrealBuildAccelerator/Common/Public/UbaNetworkBackendTcp.h
+++ b/Engine/Source/Programs/UnrealBuildAccelerator/Common/Public/UbaNetworkBackendTcp.h
@@ -33,7 +33,7 @@ namespace uba
 
 		struct Connection;
 		struct ListenEntry;
-		void ThreadListen(Logger& logger, ListenEntry& entry);
+		bool ThreadListen(Logger& logger, ListenEntry& entry);
 		void ThreadRecv(Connection& connection);
 
 		LoggerWithWriter m_logger;
diff --git a/Engine/Source/Programs/UnrealBuildAccelerator/Common/Public/UbaScheduler.h b/Engine/Source/Programs/UnrealBuildAccelerator/Common/Public/UbaScheduler.h
--- a/Engine/Source/Programs/UnrealBuildAccelerator/Common/Public/UbaScheduler.h
+++ b/Engine/Source/Programs/UnrealBuildAccelerator/Common/Public/UbaScheduler.h
@@ -10,6 +10,7 @@ namespace uba
 {
 	class Process;
 	class SessionServer;
+	struct NextProcessInfo;
 	struct ProcessStartInfoHolder;
 
 	// This is a simple scheduler that does not handle any dependencies. Can be added based on need
@@ -22,6 +23,7 @@ namespace uba
 
 		void Start();
 		void Stop();
+		void SetMaxLocalProcessors(u32 maxLocalProcessors);
 
 		u32 EnqueueProcess(const ProcessStartInfo& info, float weight = 1.0f, const void* knownInputs = nullptr, u32 knownInputsBytes = 0, u32 knownInputsCount = 0);
 
@@ -34,7 +36,7 @@ namespace uba
 		void RemoteSlotAvailable();
 		void ProcessExited(ExitProcessInfo* info, const ProcessHandle& handle);
 		bool RunQueuedProcess(bool runLocal);
-		u32 HandleReuseMessage(Process& process, const void* recv, u32 recvSize, void* send, u32 sendCapacity);
+		bool HandleReuseMessage(Process& process, NextProcessInfo& outNextProcess, u32 prevExitCode);
 
 		SessionServer& m_session;
 		u32 m_maxLocalProcessors;
diff --git a/Engine/Source/Programs/UnrealBuildAccelerator/Common/Public/UbaSession.h b/Engine/Source/Programs/UnrealBuildAccelerator/Common/Public/UbaSession.h
--- a/Engine/Source/Programs/UnrealBuildAccelerator/Common/Public/UbaSession.h
+++ b/Engine/Source/Programs/UnrealBuildAccelerator/Common/Public/UbaSession.h
@@ -43,6 +43,7 @@ namespace uba
 	struct ListDirectoryMessage;
 	struct ListDirectoryResponse;
 	struct WrittenFile;
+	struct NextProcessInfo;
 
 	class Session
 	{
@@ -58,6 +59,9 @@ namespace uba
 		using CustomServiceFunction = Function<u32(Process& handle, const void* recv, u32 recvSize, void* send, u32 sendCapacity)>;
 		void RegisterCustomService(CustomServiceFunction&& function); // Register a custom service (that can be communicated with from the remote agents)
 
+		using GetNextProcessFunction = Function<bool(Process& handle, NextProcessInfo& outNextProcess, u32 prevExitCode)>;
+		void RegisterGetNextProcess(GetNextProcessFunction&& function); // Register a custom service (that can be communicated with from the remote agents)
+
 		const tchar* GetId();			// Id for session. Will be "yymmdd_hhmmss" unless SessionCreateInfo.useUniqueId is set to false
 		u32 GetActiveProcessCount();	// Current active processes running inside session
 		Storage& GetStorage();		// Storage (only used when remote machines are connected)
@@ -114,6 +118,7 @@ namespace uba
 		virtual bool GetListDirectoryInfo(ListDirectoryResponse& out, tchar* dirName, const StringKey& dirKey);
 		virtual bool WriteFileToDisk(ProcessImpl& process, WrittenFile& file);
 		virtual bool AllocFailed(Process& process, const tchar* allocType, u32 error);
+		virtual bool GetNextProcess(Process& process, bool& outNewProcess, NextProcessInfo& outNextProcess, u32 prevExitCode);
 		virtual bool CustomMessage(Process& process, BinaryReader& reader, BinaryWriter& writer);
 		virtual void FileEntryAdded(StringKey fileNameKey, u64 lastWritten, u64 size);
 		virtual bool FlushWrittenFiles(ProcessImpl& process);
@@ -221,6 +226,7 @@ namespace uba
 		Vector<tchar> m_environmentVariables;
 		UnorderedSet<const tchar*, HashStringNoCase, EqualStringNoCase> m_localEnvironmentVariables;
 
+		GetNextProcessFunction m_getNextProcessFunction;
 		CustomServiceFunction m_customServiceFunction;
 
 		// This is a horrible hack to try to prevent ETXTBUSY. Our theory is that if posix_spawn happens with the exact same timing as symlink has an open handle we end up with the child process cloning that handle
@@ -408,6 +414,14 @@ namespace uba
 		u64 lastWriteTime = 0;
 		u32 attributes = 0;
 	};
+
+	struct NextProcessInfo
+	{
+		TString arguments;
+		TString workingDir;
+		TString description;
+	};
+
 }
 
 template<> struct std::hash<uba::ProcessHandle> { size_t operator()(const uba::ProcessHandle& g) const { return g.GetHash(); } };
diff --git a/Engine/Source/Programs/UnrealBuildAccelerator/Common/Public/UbaSessionClient.h b/Engine/Source/Programs/UnrealBuildAccelerator/Common/Public/UbaSessionClient.h
--- a/Engine/Source/Programs/UnrealBuildAccelerator/Common/Public/UbaSessionClient.h
+++ b/Engine/Source/Programs/UnrealBuildAccelerator/Common/Public/UbaSessionClient.h
@@ -54,6 +54,7 @@ namespace uba
 		virtual bool WriteFileToDisk(ProcessImpl& process, WrittenFile& file) override;
 		virtual bool AllocFailed(Process& process, const tchar* allocType, u32 error) override;
 		virtual void PrintSessionStats(Logger& logger) override;
+		virtual bool GetNextProcess(Process& process, bool& outNewProcess, NextProcessInfo& outNextProcess, u32 prevExitCode) override;
 		virtual bool CustomMessage(Process& process, BinaryReader& reader, BinaryWriter& writer) override;
 		virtual bool FlushWrittenFiles(ProcessImpl& process) override;
 		virtual bool UpdateEnvironment(ProcessImpl& process, const tchar* reason, bool resetStats) override;
@@ -72,7 +73,7 @@ namespace uba
 		bool UpdateNameToHashTableFromServer(StackBinaryReader<SendMaxSize>& reader);
 		void Connect();
 		void BuildEnvironmentVariables(BinaryReader& reader);
-		bool SendProcessAvailable(Vector<InternalProcessStartInfo>& out, float availableWeight, bool& outRemoteExecutionEnabled);
+		bool SendProcessAvailable(Vector<InternalProcessStartInfo>& out, float availableWeight);
 		void SendReturnProcess(u32 processId, const tchar* reason);
 		void SendPing(u64 memAvail, u64 memTotal);
 		void SendSummary();
@@ -102,6 +103,7 @@ namespace uba
 		bool m_dedicated = false;
 		bool m_useStorage = true;
 		bool m_shouldSendLogToServer = false;
+		bool m_remoteExecutionEnabled = true;
 		const tchar* m_terminationReason = nullptr;
 		u64 m_terminationTime = ~0ull;
 
diff --git a/Engine/Source/Programs/UnrealBuildAccelerator/Core/Public/UbaProtocol.h b/Engine/Source/Programs/UnrealBuildAccelerator/Core/Public/UbaProtocol.h
--- a/Engine/Source/Programs/UnrealBuildAccelerator/Core/Public/UbaProtocol.h
+++ b/Engine/Source/Programs/UnrealBuildAccelerator/Core/Public/UbaProtocol.h
@@ -31,6 +31,7 @@ namespace uba
 		MessageType_Exit,
 		MessageType_FlushWrittenFiles,
 		MessageType_UpdateEnvironment,
+		MessageType_GetNextProcess,
 		MessageType_Custom,
 	};
 
diff --git a/Engine/Source/Programs/UnrealBuildAccelerator/Detours/Private/Windows/UbaDetoursFunctionsWin.cpp b/Engine/Source/Programs/UnrealBuildAccelerator/Detours/Private/Windows/UbaDetoursFunctionsWin.cpp
--- a/Engine/Source/Programs/UnrealBuildAccelerator/Detours/Private/Windows/UbaDetoursFunctionsWin.cpp
+++ b/Engine/Source/Programs/UnrealBuildAccelerator/Detours/Private/Windows/UbaDetoursFunctionsWin.cpp
@@ -1031,26 +1031,26 @@ extern "C"
 		return g_runningRemote;
 	}
 
-	UBA_DETOURED_API bool UbaRequestNextProcess(wchar_t* outArguments, u32 outArgumentsCapacity)
+	UBA_DETOURED_API bool UbaRequestNextProcess(u32 prevExitCode, wchar_t* outArguments, u32 outArgumentsCapacity)
 	{
 		*outArguments = 0;
-
-		if (!UbaFlushWrittenFiles())
-			return false;
-
-		StackBinaryReader<2048> reader;
-		reader.SetSize(UbaSendCustomMessage(nullptr, 0, reader.buffer, 2048));
-		bool hasProcess = reader.GetLeft();
-
-		StringBuffer<256> description;
-		if (hasProcess)
+		bool newProcess;
 		{
-			reader.ReadString(outArguments, outArgumentsCapacity);
-			reader.SkipString(); // Working dir
-			reader.ReadString(description);
+			ScopedWriteLock pcs(g_communicationLock);
+			BinaryWriter writer;
+			writer.WriteByte(MessageType_GetNextProcess);
+			writer.WriteU32(prevExitCode);
+			writer.Flush();
+			BinaryReader reader;
+			newProcess = reader.ReadBool();
+			if (newProcess)
+			{
+				reader.ReadString(outArguments, outArgumentsCapacity);
+				//writer.SkipString(workingDir);
+				//writer.SkipString(description);
+			}
 		}
-		if (!UbaUpdateEnvironment(description.data, hasProcess))
-			return false;
-		return hasProcess;
+		Rpc_UpdateTables();
+		return newProcess;
 	}
 }
diff --git a/Engine/Source/Programs/UnrealBuildAccelerator/Detours/Public/UbaDetoursApi.h b/Engine/Source/Programs/UnrealBuildAccelerator/Detours/Public/UbaDetoursApi.h
--- a/Engine/Source/Programs/UnrealBuildAccelerator/Detours/Public/UbaDetoursApi.h
+++ b/Engine/Source/Programs/UnrealBuildAccelerator/Detours/Public/UbaDetoursApi.h
@@ -22,10 +22,10 @@ extern "C"
 	// Returns true if process is running on a remote machine
 	UBA_DETOURED_API bool UbaRunningRemote();
 
-	// Using custom message inside UbaScheduler so can't be combined with RegisterCustomService.
 	// This function will automatically flush written files before requesting next process and will also update environment if new process is retrieved
+	// prevExitCode should contain exit code of "finished" process.
 	// Will write arguments into outArguments. outArgumentsCapacity should be in characters (not bytes)
-	UBA_DETOURED_API bool UbaRequestNextProcess(wchar_t* outArguments, unsigned int outArgumentsCapacity);
+	UBA_DETOURED_API bool UbaRequestNextProcess(unsigned int prevExitCode, wchar_t* outArguments, unsigned int outArgumentsCapacity);
 }
 
 
@@ -37,7 +37,7 @@ static HMODULE UbaDetoursModule = GetModuleHandleW(L"UbaDetours.dll");
 if (UbaDetoursModule)
 {
 	using UbaFlushWrittenFilesFunc = bool();
-	using UbaRequestNextProcessFunc = bool(const wchar_t* outArguments, unsigned int outArgumentsCapacity);
+	using UbaRequestNextProcessFunc = bool(unsigned int prevExitCode, const wchar_t* outArguments, unsigned int outArgumentsCapacity);
 	using UbaUpdateEnvironmentFunc = bool(const wchar_t* reason, bool resetStats);
 
 	static UbaFlushWrittenFilesFunc* flushWrittenFiles = (UbaFlushWrittenFilesFunc*)(void*)GetProcAddress(UbaDetoursModule, "UbaFlushWrittenFiles");
diff --git a/Engine/Source/Programs/UnrealBuildAccelerator/Host/Private/UbaExports.cpp b/Engine/Source/Programs/UnrealBuildAccelerator/Host/Private/UbaExports.cpp
--- a/Engine/Source/Programs/UnrealBuildAccelerator/Host/Private/UbaExports.cpp
+++ b/Engine/Source/Programs/UnrealBuildAccelerator/Host/Private/UbaExports.cpp
@@ -456,6 +456,11 @@ uba::StorageClient* CreateStorageClient(uba::NetworkClient& client, const uba::t
 		scheduler->EnqueueProcess(info, weight, knownInputs, knownInputsBytes, knownInputsCount);
 	}
 
+	void Scheduler_SetMaxLocalProcessors(uba::Scheduler* scheduler, uba::u32 maxLocalProcessors)
+	{
+		scheduler->SetMaxLocalProcessors(maxLocalProcessors);
+	}
+
 	void Scheduler_Stop(uba::Scheduler* scheduler)
 	{
 		scheduler->Stop();
diff --git a/Engine/Source/Programs/UnrealBuildAccelerator/Host/Public/UbaExports.h b/Engine/Source/Programs/UnrealBuildAccelerator/Host/Public/UbaExports.h
--- a/Engine/Source/Programs/UnrealBuildAccelerator/Host/Public/UbaExports.h
+++ b/Engine/Source/Programs/UnrealBuildAccelerator/Host/Public/UbaExports.h
@@ -114,6 +114,7 @@ extern "C"
 	UBA_API uba::Scheduler* Scheduler_Create(uba::SessionServer* session, uba::u32 maxLocalProcessors = ~0u, bool enableProcessReuse = false);
 	UBA_API void Scheduler_Start(uba::Scheduler* scheduler);
 	UBA_API void Scheduler_EnqueueProcess(uba::Scheduler* scheduler, const uba::ProcessStartInfo& info, float weight = 1.0f, const void* knownInputs = nullptr, uba::u32 knownInputsBytes = 0, uba::u32 knownInputsCount = 0);
+	UBA_API void Scheduler_SetMaxLocalProcessors(uba::Scheduler* scheduler, uba::u32 maxLocalProcessors);
 	UBA_API void Scheduler_Stop(uba::Scheduler* scheduler);
 	UBA_API void Scheduler_Destroy(uba::Scheduler* scheduler);
 	UBA_API void Scheduler_GetStats(uba::Scheduler* scheduler, uba::u32& outQueued, uba::u32& outActive, uba::u32& outFinished);
diff --git a/Engine/Source/Programs/UnrealBuildAccelerator/Test/Private/UbaTestAll.h b/Engine/Source/Programs/UnrealBuildAccelerator/Test/Private/UbaTestAll.h
--- a/Engine/Source/Programs/UnrealBuildAccelerator/Test/Private/UbaTestAll.h
+++ b/Engine/Source/Programs/UnrealBuildAccelerator/Test/Private/UbaTestAll.h
@@ -30,6 +30,7 @@ namespace uba
 		UBA_TEST(TestLogLines) \
 		UBA_TEST(TestLocalSchedule) \
 		UBA_TEST(TestLocalScheduleReuse) \
+		UBA_TEST(TestRemoteScheduleReuse) \
 
 	#define UBA_TEST(x) \
 		logger.Info(TC("Running %hs..."), #x); \
@@ -60,8 +61,8 @@ namespace uba
 
 		logger.Info(TC("Running tests (Test rootdir: %s)"), testRootDir.data);
 
-		//UBA_TEST(TestLocalScheduleReuse)
-		UBA_TESTS
+		UBA_TEST(TestRemoteScheduleReuse)
+		//UBA_TESTS
 
 		logger.Info(TC("Tests finished successfully!"));
 		Sleep(3000);
diff --git a/Engine/Source/Programs/UnrealBuildAccelerator/Test/Private/UbaTestScheduler.h b/Engine/Source/Programs/UnrealBuildAccelerator/Test/Private/UbaTestScheduler.h
--- a/Engine/Source/Programs/UnrealBuildAccelerator/Test/Private/UbaTestScheduler.h
+++ b/Engine/Source/Programs/UnrealBuildAccelerator/Test/Private/UbaTestScheduler.h
@@ -38,7 +38,36 @@ namespace uba
 		return RunLocal(logger, testRootDir, [](LoggerWithWriter& logger, SessionServer& session, const tchar* workingDir, const RunProcessFunction& runProcess)
 			{
 
-				Scheduler scheduler(session, true);
+				Scheduler scheduler(session, ~0u, true);
+
+				StringBuffer<> testApp;
+				GetTestAppPath(logger, testApp);
+
+				ProcessStartInfo processInfo;
+				processInfo.application = testApp.data;
+				processInfo.workingDir = workingDir;
+				processInfo.arguments = TC("-reuse");
+
+				scheduler.EnqueueProcess(processInfo);
+				scheduler.Start();
+
+				u32 queued, active, finished;
+				do { scheduler.GetStats(queued, active, finished); } while (queued != 0 || active != 0);
+
+				scheduler.Stop();
+				return true;
+			});
+	}
+
+	bool TestRemoteScheduleReuse(LoggerWithWriter& logger, const StringBufferBase& testRootDir)
+	{
+		if (!IsWindows)
+			return true;
+
+		return RunRemote(logger, testRootDir, [](LoggerWithWriter& logger, SessionServer& session, const tchar* workingDir, const RunProcessFunction& runProcess)
+			{
+
+				Scheduler scheduler(session, 0u, true);
 
 				StringBuffer<> testApp;
 				GetTestAppPath(logger, testApp);
diff --git a/Engine/Source/Programs/UnrealBuildAccelerator/TestApp/Private/Windows/UbaTestAppWin.cpp b/Engine/Source/Programs/UnrealBuildAccelerator/TestApp/Private/Windows/UbaTestAppWin.cpp
--- a/Engine/Source/Programs/UnrealBuildAccelerator/TestApp/Private/Windows/UbaTestAppWin.cpp
+++ b/Engine/Source/Programs/UnrealBuildAccelerator/TestApp/Private/Windows/UbaTestAppWin.cpp
@@ -133,11 +133,11 @@ int wmain(int argc, wchar_t* argv[])
 	}
 	else if (wcscmp(argv[1], L"-reuse") == 0)
 	{
-		using UbaRequestNextProcessFunc = bool(wchar_t* outArguments, unsigned int outArgumentsCapacity);
+		using UbaRequestNextProcessFunc = bool(unsigned int prevExitCode, wchar_t* outArguments, unsigned int outArgumentsCapacity);
 		static UbaRequestNextProcessFunc* requestNextProcess = (UbaRequestNextProcessFunc*)(void*)GetProcAddress(detoursHandle, "UbaRequestNextProcess");
 
 		wchar_t arguments[1024];
-		if (requestNextProcess(arguments, sizeof(arguments)))
+		if (requestNextProcess(0, arguments, sizeof(arguments)))
 			return LogError(L"Didn't expect another process");
 	}
 	else
