[UBA]
* Fixed stats reporting bug for process reuse

diff --git a/Engine/Source/Programs/UnrealBuildAccelerator/Common/Private/UbaProcess.cpp b/Engine/Source/Programs/UnrealBuildAccelerator/Common/Private/UbaProcess.cpp
--- a/Engine/Source/Programs/UnrealBuildAccelerator/Common/Private/UbaProcess.cpp
+++ b/Engine/Source/Programs/UnrealBuildAccelerator/Common/Private/UbaProcess.cpp
@@ -827,13 +827,36 @@ namespace uba
 
 			case MessageType_GetNextProcess:
 				{
-					bool prevExitCode = reader.ReadBool();
-					bool newProcess = false;
+					u32 prevExitCode = reader.ReadU32();
 					NextProcessInfo nextProcess;
-					m_messageSuccess = m_session.GetNextProcess(*this, newProcess, nextProcess, prevExitCode) && m_messageSuccess;
+
+					StackBinaryWriter<16 * 1024> statsWriter;
+					
+					ProcessStats processStats;
+					processStats.Read(reader, TraceVersion);
+					processStats.startupTime = m_processStats.startupTime;
+					processStats.wallTime = GetTime() - m_startTime;
+					processStats.cpuTime = 0;
+					processStats.hostTotalTime = m_processStats.hostTotalTime;
+					
+					processStats.Write(statsWriter);
+					m_sessionStats.Write(statsWriter);
+					m_storageStats.Write(statsWriter);
+					m_systemStats.Write(statsWriter);
+					BinaryReader statsReader(statsWriter.GetData(), 0, statsWriter.GetPosition());
+
+					bool newProcess = false;
+					m_messageSuccess = m_session.GetNextProcess(*this, newProcess, nextProcess, prevExitCode, statsReader) && m_messageSuccess;
 					writer.WriteBool(newProcess);
 					if (!newProcess)
 						return true;
+
+					m_startTime = GetTime();
+					m_processStats = {};
+					m_sessionStats = {};
+					m_storageStats = {};
+					m_systemStats = {};
+
 					writer.WriteString(nextProcess.arguments);
 					writer.WriteString(nextProcess.workingDir);
 					writer.WriteString(nextProcess.description);
diff --git a/Engine/Source/Programs/UnrealBuildAccelerator/Common/Private/UbaSession.cpp b/Engine/Source/Programs/UnrealBuildAccelerator/Common/Private/UbaSession.cpp
--- a/Engine/Source/Programs/UnrealBuildAccelerator/Common/Private/UbaSession.cpp
+++ b/Engine/Source/Programs/UnrealBuildAccelerator/Common/Private/UbaSession.cpp
@@ -1998,7 +1998,7 @@ namespace uba
 		return true;
 	}
 
-	bool Session::GetNextProcess(Process& process, bool& outNewProcess, NextProcessInfo& outNextProcess, u32 prevExitCode)
+	bool Session::GetNextProcess(Process& process, bool& outNewProcess, NextProcessInfo& outNextProcess, u32 prevExitCode, BinaryReader& statsReader)
 	{
 		if (!m_getNextProcessFunction)
 		{
@@ -2007,6 +2007,11 @@ namespace uba
 		}
 
 		outNewProcess = m_getNextProcessFunction(process, outNextProcess, prevExitCode);
+		if (!outNewProcess)
+			return true;
+
+		m_trace.ProcessEnvironmentUpdated(process.GetId(), outNextProcess.description.c_str(), statsReader.GetPositionData(), statsReader.GetLeft());
+
 		return true;
 	}
 
diff --git a/Engine/Source/Programs/UnrealBuildAccelerator/Common/Private/UbaSessionClient.cpp b/Engine/Source/Programs/UnrealBuildAccelerator/Common/Private/UbaSessionClient.cpp
--- a/Engine/Source/Programs/UnrealBuildAccelerator/Common/Private/UbaSessionClient.cpp
+++ b/Engine/Source/Programs/UnrealBuildAccelerator/Common/Private/UbaSessionClient.cpp
@@ -1639,7 +1639,7 @@ namespace uba
 		m_stats.Print(logger);
 	}
 
-	bool SessionClient::GetNextProcess(Process& process, bool& outNewProcess, NextProcessInfo& outNextProcess, u32 prevExitCode)
+	bool SessionClient::GetNextProcess(Process& process, bool& outNewProcess, NextProcessInfo& outNextProcess, u32 prevExitCode, BinaryReader& statsReader)
 	{
 		outNewProcess = false;
 
@@ -1656,10 +1656,7 @@ namespace uba
 		NetworkMessage msg(m_client, ServiceId, SessionMessageType_GetNextProcess, writer);
 		writer.WriteU32(pi.m_id);
 		writer.WriteU32(prevExitCode);
-		pi.m_processStats.Write(writer);
-		pi.m_sessionStats.Write(writer);
-		pi.m_storageStats.Write(writer);
-		pi.m_systemStats.Write(writer);
+		writer.WriteBytes(statsReader.GetPositionData(), statsReader.GetLeft());
 
 		if (!msg.Send(reader, m_stats.customMsg))
 			return false;
diff --git a/Engine/Source/Programs/UnrealBuildAccelerator/Common/Private/UbaSessionServer.cpp b/Engine/Source/Programs/UnrealBuildAccelerator/Common/Private/UbaSessionServer.cpp
--- a/Engine/Source/Programs/UnrealBuildAccelerator/Common/Private/UbaSessionServer.cpp
+++ b/Engine/Source/Programs/UnrealBuildAccelerator/Common/Private/UbaSessionServer.cpp
@@ -1131,7 +1131,7 @@ namespace uba
 				ScopedWriteLock exitedLock(remoteProcess.m_exitedLock);
 				NextProcessInfo nextProcess;
 				bool newProcess;
-				if (!GetNextProcess(remoteProcess, newProcess, nextProcess, prevExitCode))
+				if (!GetNextProcess(remoteProcess, newProcess, nextProcess, prevExitCode, reader))
 					return false;
 
 				writer.WriteBool(newProcess);
@@ -1140,7 +1140,6 @@ namespace uba
 					writer.WriteString(nextProcess.arguments);
 					writer.WriteString(nextProcess.workingDir);
 					writer.WriteString(nextProcess.description);
-					m_trace.ProcessEnvironmentUpdated(processId, nextProcess.description.c_str(), reader.GetPositionData(), reader.GetLeft());
 				}
 				return true;
 			}
diff --git a/Engine/Source/Programs/UnrealBuildAccelerator/Common/Public/UbaSession.h b/Engine/Source/Programs/UnrealBuildAccelerator/Common/Public/UbaSession.h
--- a/Engine/Source/Programs/UnrealBuildAccelerator/Common/Public/UbaSession.h
+++ b/Engine/Source/Programs/UnrealBuildAccelerator/Common/Public/UbaSession.h
@@ -118,7 +118,7 @@ namespace uba
 		virtual bool GetListDirectoryInfo(ListDirectoryResponse& out, tchar* dirName, const StringKey& dirKey);
 		virtual bool WriteFileToDisk(ProcessImpl& process, WrittenFile& file);
 		virtual bool AllocFailed(Process& process, const tchar* allocType, u32 error);
-		virtual bool GetNextProcess(Process& process, bool& outNewProcess, NextProcessInfo& outNextProcess, u32 prevExitCode);
+		virtual bool GetNextProcess(Process& process, bool& outNewProcess, NextProcessInfo& outNextProcess, u32 prevExitCode, BinaryReader& statsReader);
 		virtual bool CustomMessage(Process& process, BinaryReader& reader, BinaryWriter& writer);
 		virtual void FileEntryAdded(StringKey fileNameKey, u64 lastWritten, u64 size);
 		virtual bool FlushWrittenFiles(ProcessImpl& process);
diff --git a/Engine/Source/Programs/UnrealBuildAccelerator/Common/Public/UbaSessionClient.h b/Engine/Source/Programs/UnrealBuildAccelerator/Common/Public/UbaSessionClient.h
--- a/Engine/Source/Programs/UnrealBuildAccelerator/Common/Public/UbaSessionClient.h
+++ b/Engine/Source/Programs/UnrealBuildAccelerator/Common/Public/UbaSessionClient.h
@@ -54,7 +54,7 @@ namespace uba
 		virtual bool WriteFileToDisk(ProcessImpl& process, WrittenFile& file) override;
 		virtual bool AllocFailed(Process& process, const tchar* allocType, u32 error) override;
 		virtual void PrintSessionStats(Logger& logger) override;
-		virtual bool GetNextProcess(Process& process, bool& outNewProcess, NextProcessInfo& outNextProcess, u32 prevExitCode) override;
+		virtual bool GetNextProcess(Process& process, bool& outNewProcess, NextProcessInfo& outNextProcess, u32 prevExitCode, BinaryReader& statsReader) override;
 		virtual bool CustomMessage(Process& process, BinaryReader& reader, BinaryWriter& writer) override;
 		virtual bool FlushWrittenFiles(ProcessImpl& process) override;
 		virtual bool UpdateEnvironment(ProcessImpl& process, const tchar* reason, bool resetStats) override;
diff --git a/Engine/Source/Programs/UnrealBuildAccelerator/Detours/Private/Windows/UbaDetoursFunctionsWin.cpp b/Engine/Source/Programs/UnrealBuildAccelerator/Detours/Private/Windows/UbaDetoursFunctionsWin.cpp
--- a/Engine/Source/Programs/UnrealBuildAccelerator/Detours/Private/Windows/UbaDetoursFunctionsWin.cpp
+++ b/Engine/Source/Programs/UnrealBuildAccelerator/Detours/Private/Windows/UbaDetoursFunctionsWin.cpp
@@ -1040,11 +1040,15 @@ extern "C"
 			BinaryWriter writer;
 			writer.WriteByte(MessageType_GetNextProcess);
 			writer.WriteU32(prevExitCode);
+			g_stats.Write(writer);
+
+
 			writer.Flush();
 			BinaryReader reader;
 			newProcess = reader.ReadBool();
 			if (newProcess)
 			{
+				g_stats = {};
 				reader.ReadString(outArguments, outArgumentsCapacity);
 				//writer.SkipString(workingDir);
 				//writer.SkipString(description);
