[UBA]
* Fixed so sendfile stats is correct for reuse process entries

diff --git a/Engine/Source/Programs/UnrealBuildAccelerator/Common/Private/UbaSessionClient.cpp b/Engine/Source/Programs/UnrealBuildAccelerator/Common/Private/UbaSessionClient.cpp
--- a/Engine/Source/Programs/UnrealBuildAccelerator/Common/Private/UbaSessionClient.cpp
+++ b/Engine/Source/Programs/UnrealBuildAccelerator/Common/Private/UbaSessionClient.cpp
@@ -1651,11 +1651,16 @@ namespace uba
 		if (!FlushWrittenFiles(pi))
 			return false;
 
+		ProcessStats processStats;
+		processStats.Read(statsReader, TraceVersion);
+		processStats.sendFiles = pi.m_processStats.sendFiles;
+
 		StackBinaryReader<SendMaxSize> reader;
 		StackBinaryWriter<16 * 1024> writer;
 		NetworkMessage msg(m_client, ServiceId, SessionMessageType_GetNextProcess, writer);
 		writer.WriteU32(pi.m_id);
 		writer.WriteU32(prevExitCode);
+		processStats.Write(writer);
 		writer.WriteBytes(statsReader.GetPositionData(), statsReader.GetLeft());
 
 		if (!msg.Send(reader, m_stats.customMsg))
