[UBA]
* Fixed so imagehlp.dll and dbghelp.dll are detoured when loaded. Detour ImageGetDigestStream and SymLoadModuleExW because both of them cause trouble on wine
* Fixed so reuse of processes also honor log files so a new log file is created when reuse happen
* Improved log file naming so host can set name and log files are sent back properly with the right name

diff --git a/Engine/Source/Programs/UnrealBuildAccelerator/Common/Private/UbaFile.cpp b/Engine/Source/Programs/UnrealBuildAccelerator/Common/Private/UbaFile.cpp
--- a/Engine/Source/Programs/UnrealBuildAccelerator/Common/Private/UbaFile.cpp
+++ b/Engine/Source/Programs/UnrealBuildAccelerator/Common/Private/UbaFile.cpp
@@ -122,7 +122,7 @@ namespace uba
 		dwFlagsAndAttributes |= (overlapped ? FILE_FLAG_OVERLAPPED : FILE_FLAG_SEQUENTIAL_SCAN);
 		#endif
 
-		outHandle = uba::CreateFileW(fileName, GENERIC_READ, FILE_SHARE_READ, OPEN_EXISTING, dwFlagsAndAttributes);
+		outHandle = uba::CreateFileW(fileName, GENERIC_READ, FILE_SHARE_READ|FILE_SHARE_WRITE, OPEN_EXISTING, dwFlagsAndAttributes);
 		if (outHandle != InvalidFileHandle)
 			return true;
 		u32 lastError = GetLastError();
diff --git a/Engine/Source/Programs/UnrealBuildAccelerator/Common/Private/UbaProcess.cpp b/Engine/Source/Programs/UnrealBuildAccelerator/Common/Private/UbaProcess.cpp
--- a/Engine/Source/Programs/UnrealBuildAccelerator/Common/Private/UbaProcess.cpp
+++ b/Engine/Source/Programs/UnrealBuildAccelerator/Common/Private/UbaProcess.cpp
@@ -860,6 +860,7 @@ namespace uba
 					writer.WriteString(nextProcess.arguments);
 					writer.WriteString(nextProcess.workingDir);
 					writer.WriteString(nextProcess.description);
+					writer.WriteString(nextProcess.logFile);
 					return true;
 				}
 
diff --git a/Engine/Source/Programs/UnrealBuildAccelerator/Common/Private/UbaProcessStartInfoHolder.h b/Engine/Source/Programs/UnrealBuildAccelerator/Common/Private/UbaProcessStartInfoHolder.h
--- a/Engine/Source/Programs/UnrealBuildAccelerator/Common/Private/UbaProcessStartInfoHolder.h
+++ b/Engine/Source/Programs/UnrealBuildAccelerator/Common/Private/UbaProcessStartInfoHolder.h
@@ -46,11 +46,8 @@ namespace uba
 			description = si.description;
 			startInfo.description = description.c_str();
 
-			if (si.logFile)
-			{
-				logFile = si.logFile;
-				startInfo.logFile = logFile.c_str();
-			}
+			logFile = si.logFile;
+			startInfo.logFile = logFile.c_str();
 		}
 
 		void Write(BinaryWriter& writer)
@@ -59,6 +56,7 @@ namespace uba
 			writer.WriteString(application);
 			writer.WriteString(arguments);
 			writer.WriteString(workingDir);
+			writer.WriteString(logFile);
 			writer.WriteU32(*(u32*)&weight);
 			writer.WriteBool(startInfo.writeOutputFilesOnFail);
 			writer.WriteU64(startInfo.outputStatsThresholdMs);
@@ -70,6 +68,7 @@ namespace uba
 			application = reader.ReadString();
 			arguments = reader.ReadString();
 			workingDir = reader.ReadString();
+			logFile = reader.ReadString();
 
 			Replace(application.data(), '/', PathSeparator); // TODO: Is this needed?
 
@@ -83,6 +82,7 @@ namespace uba
 			startInfo.application = application.c_str();
 			startInfo.arguments = arguments.c_str();
 			startInfo.workingDir = workingDir.c_str();
+			startInfo.logFile = logFile.c_str();
 		}
 
 		ProcessStartInfoHolder(const ProcessStartInfoHolder& o)
@@ -102,12 +102,8 @@ namespace uba
 			startInfo.arguments = arguments.c_str();
 			description = o.description;
 			startInfo.description = description.c_str();
-
-			if (o.startInfo.logFile)
-			{
-				logFile = o.logFile;
-				startInfo.logFile = logFile.c_str();
-			}
+			logFile = o.logFile;
+			startInfo.logFile = logFile.c_str();
 		}
 
 		ProcessStartInfo startInfo;
diff --git a/Engine/Source/Programs/UnrealBuildAccelerator/Common/Private/UbaScheduler.cpp b/Engine/Source/Programs/UnrealBuildAccelerator/Common/Private/UbaScheduler.cpp
--- a/Engine/Source/Programs/UnrealBuildAccelerator/Common/Private/UbaScheduler.cpp
+++ b/Engine/Source/Programs/UnrealBuildAccelerator/Common/Private/UbaScheduler.cpp
@@ -265,6 +265,7 @@ namespace uba
 		outNextProcess.arguments = si.arguments;
 		outNextProcess.workingDir = si.workingDir;
 		outNextProcess.description = si.description;
+		outNextProcess.logFile = si.logFile;
 		return true;
 	}
 }
diff --git a/Engine/Source/Programs/UnrealBuildAccelerator/Common/Private/UbaSession.cpp b/Engine/Source/Programs/UnrealBuildAccelerator/Common/Private/UbaSession.cpp
--- a/Engine/Source/Programs/UnrealBuildAccelerator/Common/Private/UbaSession.cpp
+++ b/Engine/Source/Programs/UnrealBuildAccelerator/Common/Private/UbaSession.cpp
@@ -897,9 +897,9 @@ namespace uba
 		out.push_back({ library, temp3.data, attr, isSystem });
 
 		bool result = true;
-		FindImports(applicationName, [&](const tchar* importName)
+		FindImports(applicationName, [&](const tchar* importName, bool isKnown)
 			{
-				if (result)
+				if (result && !isKnown)
 					result = CopyImports(out, importName, applicationDir, applicationDirEnd, handledImports);
 			});
 		return result;
@@ -1890,7 +1890,7 @@ namespace uba
 	bool Session::GetFullFileName(GetFullFileNameResponse& out, const GetFullFileNameMessage& msg)
 	{
 		UBA_ASSERTF(false, TC("SHOULD NOT HAPPEN (only remote)"));
-		return false;
+		return SearchPathForFile(m_logger, out.fileName, msg.fileName.data, msg.process.m_virtualApplicationDir.c_str());
 	}
 
 	bool Session::GetListDirectoryInfo(ListDirectoryResponse& out, tchar* dirName, const StringKey& dirKey)
diff --git a/Engine/Source/Programs/UnrealBuildAccelerator/Common/Private/UbaSessionClient.cpp b/Engine/Source/Programs/UnrealBuildAccelerator/Common/Private/UbaSessionClient.cpp
--- a/Engine/Source/Programs/UnrealBuildAccelerator/Common/Private/UbaSessionClient.cpp
+++ b/Engine/Source/Programs/UnrealBuildAccelerator/Common/Private/UbaSessionClient.cpp
@@ -1242,6 +1242,38 @@ namespace uba
 		msg.Send();
 	}
 
+	void SessionClient::SendLogFileToServer(ProcessImpl& pi)
+	{
+		auto logFile = pi.m_startInfo.logFile;
+		if (!logFile || !*logFile)
+			return;
+		WrittenFile f;
+		f.name = logFile;
+		f.attributes = DefaultAttributes();
+		StringBuffer<> dest;
+		if (const tchar* lastSlash = TStrrchr(logFile, PathSeparator))
+			logFile = lastSlash + 1;
+		dest.Append(TC("<log>")).Append(logFile);
+		f.key = ToStringKeyLower(dest);
+		SendFile(pi, f, dest.data);
+	}
+
+	void SessionClient::GetLogFileName(StringBufferBase& out, const tchar* logFile, const tchar* arguments)
+	{
+		out.Append(m_sessionLogDir.data);
+		if (logFile && *logFile)
+		{
+			if (const tchar* lastSeparator = TStrrchr(logFile, PathSeparator))
+				logFile = lastSeparator + 1;
+			out.Append(logFile);
+		}
+		else
+		{
+			GetNameFromArguments(out, arguments, true);
+			out.Append(TC(".log"));
+		}
+	}
+
 	void SessionClient::ThreadCreateProcessLoop()
 	{
 		struct ProcessRec
@@ -1408,9 +1440,7 @@ namespace uba
 					StringBuffer<> logFile;
 					if (m_logToFile)
 					{
-						logFile.Append(m_sessionLogDir.data);
-						GetNameFromArguments(logFile, startInfo.arguments, true);
-						logFile.Append(TC(".log"));
+						GetLogFileName(logFile, startInfo.logFile, startInfo.arguments);
 						startInfo.logFile = logFile.data;
 					}
 
@@ -1459,20 +1489,7 @@ namespace uba
 
 						auto& startInfo = h.GetStartInfo();
 						if (session.m_shouldSendLogToServer)
-						{
-							if (const tchar* logFile = startInfo.logFile)
-							{
-								WrittenFile f;
-								f.name = logFile;
-								f.attributes = DefaultAttributes();
-								StringBuffer<> dest;
-								if (const tchar* lastSlash = TStrrchr(logFile, PathSeparator))
-									logFile = lastSlash + 1;
-								dest.Append(TC("<log>")).Append(logFile);
-								f.key = ToStringKeyLower(dest);
-								session.SendFile(*(ProcessImpl*)h.m_process, f, dest.data);
-							}
-						}
+							session.SendLogFileToServer(*(ProcessImpl*)h.m_process);
 
 						auto decreaseWeight = MakeGuard([&]()
 							{
@@ -1669,6 +1686,9 @@ namespace uba
 		outNewProcess = reader.ReadBool();
 		if (outNewProcess)
 		{
+			if (m_shouldSendLogToServer)
+				SendLogFileToServer(pi);
+
 			pi.m_exitCode = prevExitCode;
 			if (m_processFinished)
 				m_processFinished(&process);
@@ -1682,10 +1702,23 @@ namespace uba
 			outNextProcess.arguments = reader.ReadString();
 			outNextProcess.workingDir = reader.ReadString();
 			outNextProcess.description = reader.ReadString();
+			outNextProcess.logFile = reader.ReadString();
+
+			if (m_logToFile)
+			{
+				StringBuffer<512> logFile;
+				GetLogFileName(logFile, outNextProcess.logFile.c_str(), outNextProcess.arguments.c_str());
+				outNextProcess.logFile = logFile.data;
+			}
 
 			// TODO: Probably need to fill up with more stuff.. this is fine for current usecase
 			pi.m_arguments = outNextProcess.arguments;
 			pi.m_description = outNextProcess.description;
+			pi.m_logFile = outNextProcess.logFile;
+
+			pi.m_startInfo.arguments = pi.m_arguments.c_str();
+			pi.m_startInfo.description = pi.m_description.c_str();
+			pi.m_startInfo.logFile = pi.m_logFile.c_str();
 		}
 
 		reader.Reset();
diff --git a/Engine/Source/Programs/UnrealBuildAccelerator/Common/Private/UbaSessionServer.cpp b/Engine/Source/Programs/UnrealBuildAccelerator/Common/Private/UbaSessionServer.cpp
--- a/Engine/Source/Programs/UnrealBuildAccelerator/Common/Private/UbaSessionServer.cpp
+++ b/Engine/Source/Programs/UnrealBuildAccelerator/Common/Private/UbaSessionServer.cpp
@@ -769,6 +769,7 @@ namespace uba
 						StringBuffer<> logPath;
 						logPath.Append(m_sessionLogDir).Append(destination.data + 5);
 						m_storage.CopyOrLink(casKey, logPath.data, attributes);
+						m_storage.DropCasFile(casKey, true);
 						writer.WriteBool(true);
 						return true;
 					}
@@ -1140,6 +1141,7 @@ namespace uba
 					writer.WriteString(nextProcess.arguments);
 					writer.WriteString(nextProcess.workingDir);
 					writer.WriteString(nextProcess.description);
+					writer.WriteString(nextProcess.logFile);
 				}
 				return true;
 			}
diff --git a/Engine/Source/Programs/UnrealBuildAccelerator/Common/Public/UbaNetwork.h b/Engine/Source/Programs/UnrealBuildAccelerator/Common/Public/UbaNetwork.h
--- a/Engine/Source/Programs/UnrealBuildAccelerator/Common/Public/UbaNetwork.h
+++ b/Engine/Source/Programs/UnrealBuildAccelerator/Common/Public/UbaNetwork.h
@@ -18,7 +18,7 @@ namespace uba
 	
 	static constexpr u32 SystemNetworkVersion = 1338;
 	static constexpr u32 StorageNetworkVersion = 3;
-	static constexpr u32 SessionNetworkVersion = 25;
+	static constexpr u32 SessionNetworkVersion = 26;
 
 
 	// Messages used over network between client/server (system, storage and session)
diff --git a/Engine/Source/Programs/UnrealBuildAccelerator/Common/Public/UbaSession.h b/Engine/Source/Programs/UnrealBuildAccelerator/Common/Public/UbaSession.h
--- a/Engine/Source/Programs/UnrealBuildAccelerator/Common/Public/UbaSession.h
+++ b/Engine/Source/Programs/UnrealBuildAccelerator/Common/Public/UbaSession.h
@@ -420,6 +420,7 @@ namespace uba
 		TString arguments;
 		TString workingDir;
 		TString description;
+		TString logFile;
 	};
 
 }
diff --git a/Engine/Source/Programs/UnrealBuildAccelerator/Common/Public/UbaSessionClient.h b/Engine/Source/Programs/UnrealBuildAccelerator/Common/Public/UbaSessionClient.h
--- a/Engine/Source/Programs/UnrealBuildAccelerator/Common/Public/UbaSessionClient.h
+++ b/Engine/Source/Programs/UnrealBuildAccelerator/Common/Public/UbaSessionClient.h
@@ -77,6 +77,8 @@ namespace uba
 		void SendReturnProcess(u32 processId, const tchar* reason);
 		void SendPing(u64 memAvail, u64 memTotal);
 		void SendSummary();
+		void SendLogFileToServer(ProcessImpl& pi);
+		void GetLogFileName(StringBufferBase& out, const tchar* logFile, const tchar* arguments);
 		u32 CountLogLines(ProcessImpl& process);
 		void WriteLogLines(BinaryWriter& writer, ProcessImpl& process);
 
diff --git a/Engine/Source/Programs/UnrealBuildAccelerator/Core/Private/UbaApplicationRules.cpp b/Engine/Source/Programs/UnrealBuildAccelerator/Core/Private/UbaApplicationRules.cpp
--- a/Engine/Source/Programs/UnrealBuildAccelerator/Core/Private/UbaApplicationRules.cpp
+++ b/Engine/Source/Programs/UnrealBuildAccelerator/Core/Private/UbaApplicationRules.cpp
@@ -357,14 +357,6 @@ namespace uba
 #endif // #if PLATFORM_WINDOWS
 	};
 
-	class ApplicationRulesShaderCompileWorker : public ApplicationRules
-	{
-		virtual bool DetourImageGetDigestStream() override
-		{
-			return true;
-		}
-	};
-
 	const RulesRec* GetApplicationRules()
 	{
 		static RulesRec rules[]
@@ -390,7 +382,6 @@ namespace uba
 			{ TC("llvm-objcopy.exe"),			new ApplicationRulesLlvmObjCopyExe() },
 			{ TC("UnrealBuildTool.dll"),		new ApplicationRulesUBTDll() },
 			{ TC("PVS-Studio.exe"),				new ApplicationRulesPVSStudio() },
-			{ TC("ShaderCompileWorker.exe"),	new ApplicationRulesShaderCompileWorker() },
 			//{ L"MSBuild.dll"),				new ApplicationRules() },
 			//{ L"BreakpadSymbolEncoder.exe"),	new ApplicationRulesClang() },
 			//{ L"cmd.exe"),		new ApplicationRules() },
diff --git a/Engine/Source/Programs/UnrealBuildAccelerator/Core/Public/UbaApplicationRules.h b/Engine/Source/Programs/UnrealBuildAccelerator/Core/Public/UbaApplicationRules.h
--- a/Engine/Source/Programs/UnrealBuildAccelerator/Core/Public/UbaApplicationRules.h
+++ b/Engine/Source/Programs/UnrealBuildAccelerator/Core/Public/UbaApplicationRules.h
@@ -118,12 +118,6 @@ namespace uba
 		{
 			// Do nothing
 		}
-
-		// ImageGetDigestStream is buggy in wine so we have to detour it for ShaderCompileWorker
-		virtual bool DetourImageGetDigestStream()
-		{
-			return false;
-		}
 	};
 
 
diff --git a/Engine/Source/Programs/UnrealBuildAccelerator/Core/Public/UbaWinBinDependencyParser.h b/Engine/Source/Programs/UnrealBuildAccelerator/Core/Public/UbaWinBinDependencyParser.h
--- a/Engine/Source/Programs/UnrealBuildAccelerator/Core/Public/UbaWinBinDependencyParser.h
+++ b/Engine/Source/Programs/UnrealBuildAccelerator/Core/Public/UbaWinBinDependencyParser.h
@@ -96,7 +96,7 @@ namespace uba
 		return rva - seh->VirtualAddress + seh->PointerToRawData;
 	} 
 
-	inline bool FindImportsInMem(LPCWSTR fileName, void* mem, const Function<void(const tchar* import)>& func)
+	inline bool FindImportsInMem(LPCWSTR fileName, void* mem, const Function<void(const tchar* import, bool isKnown)>& func)
 	{
 		auto hdrs = (PIMAGE_NT_HEADERS)(PCHAR(mem) + PIMAGE_DOS_HEADER(mem)->e_lfanew);   
 		auto pSech=IMAGE_FIRST_SECTION(hdrs);
@@ -112,8 +112,7 @@ namespace uba
 				const char* name = (PCHAR)((DWORD_PTR)mem + GetOffset(importDesc->Name, pSech, hdrs));
 				StringBuffer<256> wname;
 				wname.Append(name);
-				if (!IsKnownSystemFile(wname.data))
-					func(wname.data);
+				func(wname.data, IsKnownSystemFile(wname.data));
 				++importDesc;
 
 			}
@@ -128,7 +127,7 @@ namespace uba
 		return true;
 	}
 
-	inline bool FindImports(LPCWSTR fileName, const Function<void(const tchar* import)>& func)
+	inline bool FindImports(LPCWSTR fileName, const Function<void(const tchar* import, bool isKnown)>& func)
 	{
 		HANDLE fileHandle = True_CreateFileW(fileName, GENERIC_READ, FILE_SHARE_READ, 0, OPEN_EXISTING, FILE_ATTRIBUTE_NORMAL | FILE_FLAG_SEQUENTIAL_SCAN, 0);
 		if (fileHandle == INVALID_HANDLE_VALUE)
diff --git a/Engine/Source/Programs/UnrealBuildAccelerator/Detours/Private/Windows/UbaDetoursFunctionsDbgHelp.inl b/Engine/Source/Programs/UnrealBuildAccelerator/Detours/Private/Windows/UbaDetoursFunctionsDbgHelp.inl
new file mode 100644
--- /dev/null
+++ b/Engine/Source/Programs/UnrealBuildAccelerator/Detours/Private/Windows/UbaDetoursFunctionsDbgHelp.inl
@@ -0,0 +1,17 @@
+// Copyright Epic Games, Inc. All Rights Reserved.
+
+typedef struct _MODLOAD_DATA {} *PMODLOAD_DATA;
+using SymLoadModuleExWFunc = BOOL(HANDLE hProcess, HANDLE hFile, PCWSTR ImageName, DWORD64 BaseOfDll, DWORD DllSize, PMODLOAD_DATA Data, DWORD Flags);
+SymLoadModuleExWFunc* True_SymLoadModuleExW;
+
+BOOL Detoured_SymLoadModuleExW(HANDLE hProcess, HANDLE hFile, PCWSTR ImageName, DWORD64 BaseOfDll, DWORD DllSize, PMODLOAD_DATA Data, DWORD Flags)
+{
+	UBA_ASSERT(g_isRunningWine);
+
+	//u64 pathLen = wcslen(ImageName);
+	//StringBuffer<512> tempBuf;
+	//Rpc_GetFullFileName(ImageName, pathLen, tempBuf, false);
+	DEBUG_LOG_TRUE(L"SymLoadModuleExW", L"(%ls)", ImageName);
+	SuppressDetourScope _;
+	return True_SymLoadModuleExW(hProcess, hFile, ImageName, BaseOfDll, DllSize, Data, Flags);
+}
\ No newline at end of file
diff --git a/Engine/Source/Programs/UnrealBuildAccelerator/Detours/Private/Windows/UbaDetoursFunctionsKernelBase.inl b/Engine/Source/Programs/UnrealBuildAccelerator/Detours/Private/Windows/UbaDetoursFunctionsKernelBase.inl
--- a/Engine/Source/Programs/UnrealBuildAccelerator/Detours/Private/Windows/UbaDetoursFunctionsKernelBase.inl
+++ b/Engine/Source/Programs/UnrealBuildAccelerator/Detours/Private/Windows/UbaDetoursFunctionsKernelBase.inl
@@ -2272,14 +2272,14 @@ HMODULE Recursive_LoadLibraryExW(LPCWSTR lpLibFileName, LPCWSTR originalName, DW
 		return 0;
 
 	// Important that this code is not doing allocations.. it could cause a recursive stack overflow
-	struct Import { wchar_t name[128]; Import(const wchar_t* s) { wcscpy_s(name, sizeof_array(name), s); } };
+	struct Import { wchar_t name[128]; bool isKnown;  Import(const wchar_t* s, bool ik) : isKnown(ik) { wcscpy_s(name, sizeof_array(name), s); } };
 	std::vector<Import, GrowingAllocator<Import>> importedModules(&g_memoryBlock);
 	{
 		SuppressCreateFileDetourScope cfs;
-		if (!FindImports(lpLibFileName, [&](const wchar_t* import)
+		if (!FindImports(lpLibFileName, [&](const wchar_t* import, bool isKnown)
 			{
 				if (!GetModuleHandleW(import))
-					importedModules.emplace_back(import);
+					importedModules.emplace_back(import, isKnown);
 			}))
 		{
 			UBA_ASSERTF(false, L"Failed to find imports for binary %ls (%ls)", lpLibFileName, originalName);
@@ -2287,9 +2287,23 @@ HMODULE Recursive_LoadLibraryExW(LPCWSTR lpLibFileName, LPCWSTR originalName, DW
 	}
 	for (auto& importedModule : importedModules)
 	{
-		if (GetModuleHandleW(importedModule.name))
+		if (importedModule.isKnown && !g_isRunningWine)
 			continue;
 
+		HMODULE checkModule = GetModuleHandleW(importedModule.name);
+		if (checkModule)
+			continue;
+
+		if (importedModule.isKnown) // We need to catch dbghelp.dll and imagehlp.dll
+		{
+			if (HMODULE h = True_LoadLibraryExW(importedModule.name, 0, 0))
+			{
+				OnModuleLoaded(h, importedModule.name);
+				additionalLoads.push_back(h);
+			}
+			continue;
+		}
+
 		const wchar_t* path = importedModule.name;
 		if (path[1] == ':')
 			if (const wchar_t* lastSlash = wcsrchr(path, '\\'))
@@ -2315,10 +2329,15 @@ HMODULE Recursive_LoadLibraryExW(LPCWSTR lpLibFileName, LPCWSTR originalName, DW
 
 	SuppressCreateFileDetourScope cfs;
 	auto res = True_LoadLibraryExW(lpLibFileName, 0, 0);
-	if (res && originalName[1] == ':')
+	if (res)
 	{
-		ScopedWriteLock lock(g_loadedModulesLock);
-		g_loadedModules[res] = originalName;
+		if (originalName[1] == ':')
+		{
+			ScopedWriteLock lock(g_loadedModulesLock);
+			g_loadedModules[res] = originalName;
+		}
+		if (g_isRunningWine)
+			OnModuleLoaded(res, lpLibFileName);
 	}
 	return res;
 }
diff --git a/Engine/Source/Programs/UnrealBuildAccelerator/Detours/Private/Windows/UbaDetoursFunctionsNtDll.inl b/Engine/Source/Programs/UnrealBuildAccelerator/Detours/Private/Windows/UbaDetoursFunctionsNtDll.inl
--- a/Engine/Source/Programs/UnrealBuildAccelerator/Detours/Private/Windows/UbaDetoursFunctionsNtDll.inl
+++ b/Engine/Source/Programs/UnrealBuildAccelerator/Detours/Private/Windows/UbaDetoursFunctionsNtDll.inl
@@ -1317,7 +1317,6 @@ NTSTATUS NTAPI Detoured_NtClose(HANDLE handle)
 NTSTATUS Detoured_NtQueryObject(HANDLE Handle, OBJECT_INFORMATION_CLASS ObjectInformationClass, PVOID ObjectInformation, ULONG ObjectInformationLength, PULONG ReturnLength)
 {
 	DETOURED_CALL(NtQueryObject);
-	DEBUG_LOG_TRUE(L"NtQueryObject", L"%llu", uintptr_t(Handle));
 
 	// This can be other things than FILES.. Is used by GetHandleInformation
 	if (isDetouredHandle(Handle))
@@ -1325,7 +1324,9 @@ NTSTATUS Detoured_NtQueryObject(HANDLE Handle, OBJECT_INFORMATION_CLASS ObjectIn
 		Handle = asDetouredHandle(Handle).trueHandle;
 		UBA_ASSERTF(Handle != INVALID_HANDLE_VALUE, L"NtQueryObject");
 	}
-	return True_NtQueryObject(Handle, ObjectInformationClass, ObjectInformation, ObjectInformationLength, ReturnLength);
+	auto res = True_NtQueryObject(Handle, ObjectInformationClass, ObjectInformation, ObjectInformationLength, ReturnLength);
+	DEBUG_LOG_TRUE(L"NtQueryObject", L"(%i) %llu -> %ls", ObjectInformationClass, uintptr_t(Handle), ToString(res));
+	return res;
 }
 
 NTSTATUS Detoured_NtQueryInformationProcess(HANDLE ProcessHandle, PROCESSINFOCLASS ProcessInformationClass, PVOID ProcessInformation, ULONG ProcessInformationLength, PULONG ReturnLength)
diff --git a/Engine/Source/Programs/UnrealBuildAccelerator/Detours/Private/Windows/UbaDetoursFunctionsWin.cpp b/Engine/Source/Programs/UnrealBuildAccelerator/Detours/Private/Windows/UbaDetoursFunctionsWin.cpp
--- a/Engine/Source/Programs/UnrealBuildAccelerator/Detours/Private/Windows/UbaDetoursFunctionsWin.cpp
+++ b/Engine/Source/Programs/UnrealBuildAccelerator/Detours/Private/Windows/UbaDetoursFunctionsWin.cpp
@@ -577,6 +577,8 @@ void SendExitMessage(DWORD exitCode, u64 startTime)
 	writer.Flush(g_isChild);
 }
 
+void OnModuleLoaded(HMODULE moduleHandle, const wchar_t* name);
+
 // Variables used to communicate state from kernelbase functions to ntdll functions
 thread_local const wchar_t* t_renameFileNewName;
 thread_local const wchar_t* t_createFileFileName;
@@ -586,6 +588,7 @@ thread_local const wchar_t* t_createFileFileName;
 #include "UbaDetoursFunctionsKernelBase.inl"
 #include "UbaDetoursFunctionsUcrtBase.inl"
 #include "UbaDetoursFunctionsImagehlp.inl"
+#include "UbaDetoursFunctionsDbgHelp.inl"
 
 void DetourAttachFunction(void** trueFunc, void* detouredFunc, const char* funcName)
 {
@@ -610,6 +613,9 @@ void DetourDetachFunction(void** trueFunc, void* detouredFunc, const char* funcN
 
 int DetourAttachFunctions(bool runningRemote)
 {
+	DetourTransactionBegin();
+	DetourUpdateThread(GetCurrentThread());
+
 	#define DETOURED_FUNCTION(Func) True_##Func = (decltype(True_##Func))GetProcAddress(moduleHandle, #Func);
 
 	if (HMODULE moduleHandle = GetModuleHandleW(L"kernelbase.dll"))
@@ -643,18 +649,6 @@ int DetourAttachFunctions(bool runningRemote)
 
 	#undef DETOURED_FUNCTION
 
-
-	if (g_isRunningWine && g_rules->DetourImageGetDigestStream())
-	{
-		if (HMODULE moduleHandle = LoadLibraryW(L"Imagehlp.dll"))
-		{
-			True_ImageGetDigestStream = (ImageGetDigestStreamFunc*)GetProcAddress(moduleHandle, "ImageGetDigestStream");
-			DetourAttachFunction((PVOID*)&True_ImageGetDigestStream, Detoured_ImageGetDigestStream, "ImageGetDigestStream");
-		}
-	}
-
-
-
 	// Can't attach to these when running through debugger with some vs extensions (Microsoft child process debugging)
 #if UBA_DEBUG
 	if (IsDebuggerPresent())
@@ -675,9 +669,46 @@ int DetourAttachFunctions(bool runningRemote)
 	}
 	#undef DETOURED_FUNCTION
 
+
+	LONG error = DetourTransactionCommit();
+	if (error != NO_ERROR)
+	{
+		printf("Error detouring: %ld\n", error);
+		ExitProcess(1343);
+	}
+
 	return 0;
 }
 
+void OnModuleLoaded(HMODULE moduleHandle, const wchar_t* name)
+{
+	UBA_ASSERT(g_isRunningWine);
+
+	// SymLoadModuleExW do something bad that cause remote wine to fail everything after this call.. TODO: Revisit
+	if (!True_SymLoadModuleExW && Contains(name, L"dbghelp.dll"))
+	{
+		True_SymLoadModuleExW = (SymLoadModuleExWFunc*)GetProcAddress(moduleHandle, "SymLoadModuleExW");
+		UBA_ASSERT(True_SymLoadModuleExW);
+		DetourTransactionBegin();
+		DetourUpdateThread(GetCurrentThread());
+		DetourAttachFunction((PVOID*)&True_SymLoadModuleExW, Detoured_SymLoadModuleExW, "SymLoadModuleExW");
+		LONG error = DetourTransactionCommit(); (void)error;
+		UBA_ASSERT(!error);
+	}
+
+	// ImageGetDigestStream is buggy in wine so we have to detour it for ShaderCompileWorker
+	if (!True_ImageGetDigestStream && Contains(name, L"imagehlp.dll"))
+	{
+		True_ImageGetDigestStream = (ImageGetDigestStreamFunc*)GetProcAddress(moduleHandle, "ImageGetDigestStream");
+		UBA_ASSERT(True_ImageGetDigestStream);
+		DetourTransactionBegin();
+		DetourUpdateThread(GetCurrentThread());
+		DetourAttachFunction((PVOID*)&True_ImageGetDigestStream, Detoured_ImageGetDigestStream, "ImageGetDigestStream");
+		LONG error = DetourTransactionCommit(); (void)error;
+		UBA_ASSERT(!error);
+	}
+}
+
 int DetourDetachFunctions()
 {
 	if (g_directoryTable.m_memory)
@@ -734,7 +765,7 @@ void PreInit(const DetoursPayload& payload)
 		if (!payload.logFile.IsEmpty())
 		{
 			g_logName.Append(payload.logFile);
-			HANDLE debugFile = CreateFileW(payload.logFile.data, GENERIC_WRITE, FILE_SHARE_READ, NULL, CREATE_ALWAYS, FILE_ATTRIBUTE_NORMAL | FILE_FLAG_SEQUENTIAL_SCAN, NULL);
+			HANDLE debugFile = CreateFileW(payload.logFile.data, GENERIC_WRITE, FILE_SHARE_READ, NULL, CREATE_ALWAYS, FILE_ATTRIBUTE_NORMAL, NULL);
 			#if UBA_DEBUG_LOG_ENABLED
 			g_debugFile = (FileHandle)(u64)debugFile;
 			#else
@@ -807,15 +838,7 @@ void Init(const DetoursPayload& payload, u64 startTime)
 			});
 	}
 
-	DetourTransactionBegin();
-	DetourUpdateThread(GetCurrentThread());
 	DetourAttachFunctions(g_runningRemote);
-	LONG error = DetourTransactionCommit();
-	if (error != NO_ERROR)
-	{
-		printf("Error detouring: %ld\n", error);
-		ExitProcess(1343);
-	}
 
 	if (g_isDetachedProcess)
 	{
@@ -1033,6 +1056,10 @@ extern "C"
 
 	UBA_DETOURED_API bool UbaRequestNextProcess(u32 prevExitCode, wchar_t* outArguments, u32 outArgumentsCapacity)
 	{
+		#if UBA_DEBUG_LOG_ENABLED
+		FlushDebugLog();
+		#endif
+
 		*outArguments = 0;
 		bool newProcess;
 		{
@@ -1048,12 +1075,27 @@ extern "C"
 			newProcess = reader.ReadBool();
 			if (newProcess)
 			{
-				g_stats = {};
 				reader.ReadString(outArguments, outArgumentsCapacity);
-				//writer.SkipString(workingDir);
-				//writer.SkipString(description);
+				reader.SkipString(); // workingDir
+				reader.SkipString(); // description
+				reader.ReadString(g_logName.Clear());
 			}
 		}
+
+		if (newProcess)
+		{
+			g_stats = {};
+
+			#if UBA_DEBUG_LOG_ENABLED
+			SuppressCreateFileDetourScope scope;
+			HANDLE debugFile = (HANDLE)g_debugFile;
+			g_debugFile = InvalidFileHandle;
+			CloseHandle(debugFile);
+			debugFile = CreateFileW(g_logName.data, GENERIC_WRITE, FILE_SHARE_READ, NULL, CREATE_ALWAYS, FILE_ATTRIBUTE_NORMAL, NULL);
+			g_debugFile = (FileHandle)(u64)debugFile;
+			#endif
+		}
+
 		Rpc_UpdateTables();
 		return newProcess;
 	}
diff --git a/Engine/Source/Programs/UnrealBuildAccelerator/Host/Private/UbaExports.cpp b/Engine/Source/Programs/UnrealBuildAccelerator/Host/Private/UbaExports.cpp
--- a/Engine/Source/Programs/UnrealBuildAccelerator/Host/Private/UbaExports.cpp
+++ b/Engine/Source/Programs/UnrealBuildAccelerator/Host/Private/UbaExports.cpp
@@ -265,7 +265,7 @@ uba::StorageClient* CreateStorageClient(uba::NetworkClient& client, const uba::t
 		return handle->GetExecutingHost();
 	}
 
-	const uba::tchar* ProcessHandle_GetLogLine(uba::ProcessHandle* handle, uba::u32 index)
+	const uba::tchar* ProcessHandle_GetLogLine(const uba::ProcessHandle* handle, uba::u32 index)
 	{
 		const auto& lines = handle->GetLogLines();
 		if (index >= lines.size()) return nullptr;
@@ -489,7 +489,7 @@ uba::StorageClient* CreateStorageClient(uba::NetworkClient& client, const uba::t
 	void Uba_FindImports(const uba::tchar* binary, ImportFunc* func, void* userData)
 	{
 #if PLATFORM_WINDOWS
-		uba::FindImports(binary, [&](const uba::tchar* importName) { func(importName, userData); });
+		uba::FindImports(binary, [&](const uba::tchar* importName, bool isKnown) { func(importName, userData); });
 #endif
 	}
 }
diff --git a/Engine/Source/Programs/UnrealBuildAccelerator/Host/Public/UbaExports.h b/Engine/Source/Programs/UnrealBuildAccelerator/Host/Public/UbaExports.h
--- a/Engine/Source/Programs/UnrealBuildAccelerator/Host/Public/UbaExports.h
+++ b/Engine/Source/Programs/UnrealBuildAccelerator/Host/Public/UbaExports.h
@@ -68,7 +68,7 @@ extern "C"
 
 	UBA_API uba::u32 ProcessHandle_GetExitCode(uba::ProcessHandle* handle);
 	UBA_API const uba::tchar* ProcessHandle_GetExecutingHost(uba::ProcessHandle* handle);
-	UBA_API const uba::tchar* ProcessHandle_GetLogLine(uba::ProcessHandle* handle, uba::u32 index);
+	UBA_API const uba::tchar* ProcessHandle_GetLogLine(const uba::ProcessHandle* handle, uba::u32 index);
 	UBA_API uba::u64 ProcessHandle_GetHash(uba::ProcessHandle* handle);
 	UBA_API uba::u64 ProcessHandle_GetTotalProcessorTime(uba::ProcessHandle* handle);
 	UBA_API uba::u64 ProcessHandle_GetTotalWallTime(uba::ProcessHandle* handle);
