[UBA]
* Fixed so -replay works again when viewing files

diff --git a/Engine/Source/Programs/UnrealBuildAccelerator/Visualizer/Private/UbaTraceReader.cpp b/Engine/Source/Programs/UnrealBuildAccelerator/Visualizer/Private/UbaTraceReader.cpp
--- a/Engine/Source/Programs/UnrealBuildAccelerator/Visualizer/Private/UbaTraceReader.cpp
+++ b/Engine/Source/Programs/UnrealBuildAccelerator/Visualizer/Private/UbaTraceReader.cpp
@@ -319,6 +319,8 @@ namespace uba
 			TraceView::Session& session = GetSession(out, sessionIndex);
 			session.name = sessionName.data;
 			session.clientUid = clientUid;
+
+			++out.activeSessionCount;
 			break;
 		}
 		case TraceType_SessionUpdate:
@@ -421,6 +423,8 @@ namespace uba
 			for (auto it = session.storedFiles.begin(), end = session.storedFiles.end(); it != end; ++it)
 				if (it->stop == ~u64(0))
 					it->stop = time;
+
+			--out.activeSessionCount;
 			break;
 		}
 		case TraceType_SessionSummary:
@@ -477,6 +481,7 @@ namespace uba
 			m_activeProcesses.try_emplace(id, TraceView::ProcessLocation{sessionIndex, processorIndex, u32(processor->processes.size() - 1)});
 			
 			++session.processActiveCount;
+			++out.totalProcessActiveCount;
 
 			process.id = id;
 			process.description = desc.data;
@@ -509,6 +514,9 @@ namespace uba
 			++session.processExitedCount;
 			--session.processActiveCount;
 
+			++out.totalProcessExitedCount;
+			--out.totalProcessActiveCount;
+
 			TraceView::Process& process = session.processors[active.processorIndex].processes[active.processIndex];
 			process.processStats.Read(reader, out.version);
 			bool isRemote = active.sessionIndex != 0;
@@ -576,6 +584,7 @@ namespace uba
 
 			auto& session = GetSession(out, active.sessionIndex);
 			--session.processActiveCount;
+			--out.totalProcessActiveCount;
 
 			TraceView::Process& process = session.processors[active.processorIndex].processes[active.processIndex];
 			process.exitCode = 0;
diff --git a/Engine/Source/Programs/UnrealBuildAccelerator/Visualizer/Private/UbaTraceReader.h b/Engine/Source/Programs/UnrealBuildAccelerator/Visualizer/Private/UbaTraceReader.h
--- a/Engine/Source/Programs/UnrealBuildAccelerator/Visualizer/Private/UbaTraceReader.h
+++ b/Engine/Source/Programs/UnrealBuildAccelerator/Visualizer/Private/UbaTraceReader.h
@@ -91,8 +91,8 @@ namespace uba
 			u64 prevSend = 0;
 			u64 prevRecv = 0;
 			u64 memTotal = 0;
-			u64 processActiveCount = 0;
-			u64 processExitedCount = 0;
+			u32 processActiveCount = 0;
+			u32 processExitedCount = 0;
 
 			TString proxyName;
 			bool proxyCreated = false;
@@ -114,6 +114,9 @@ namespace uba
 		Vector<TString> strings;
 		u64 startTime = 0;
 		u64 frequency = 0;
+		u32 totalProcessActiveCount = 0;
+		u32 totalProcessExitedCount = 0;
+		u32 activeSessionCount = 0;
 		u32 version = 0;
 		bool finished = true;
 	};
diff --git a/Engine/Source/Programs/UnrealBuildAccelerator/Visualizer/Private/UbaVisualizer.cpp b/Engine/Source/Programs/UnrealBuildAccelerator/Visualizer/Private/UbaVisualizer.cpp
--- a/Engine/Source/Programs/UnrealBuildAccelerator/Visualizer/Private/UbaVisualizer.cpp
+++ b/Engine/Source/Programs/UnrealBuildAccelerator/Visualizer/Private/UbaVisualizer.cpp
@@ -100,10 +100,8 @@ namespace uba
 		m_thread.Start([this]() { ThreadLoop(); return 0;});
 
 		while (!m_hwnd)
-		{
 			if (m_thread.Wait(10))
 				return true;
-		}
 
 		StringBuffer<256> traceName;
 		while (m_hwnd)
@@ -144,10 +142,8 @@ namespace uba
 		m_thread.Start([this]() { ThreadLoop(); return 0; });
 
 		while (!m_hwnd)
-		{
 			if (m_thread.Wait(10))
 				return true;
-		}
 
 		wchar_t dots[] = TC("....");
 		u32 dotsCounter = 0;
@@ -184,12 +180,16 @@ namespace uba
 
 	bool Visualizer::ShowUsingFile(const wchar_t* fileName, u32 replay)
 	{
-		if (!m_trace.ReadFile(m_traceView, fileName, replay != 0))
-			return false;
-		m_fileName.Append(fileName);
-		m_replay = replay;
 		m_looping = true;
+		m_autoScroll = false;
 		m_thread.Start([this]() { ThreadLoop(); return 0;});
+
+		while (!m_hwnd)
+			if (m_thread.Wait(10))
+				return true;
+		m_fileName.Append(fileName);
+		m_replay = replay;
+		PostMessage(m_hwnd, WM_NEWTRACE, 0, 0);
 		return true;
 	}
 
@@ -229,7 +229,7 @@ namespace uba
 		//m_zoomValue = 0.75f;
 		//m_horizontalScaleValue = 1.0f;
 
-		m_replay = 0;
+		//m_replay = 0;
 		m_startTime = GetTime();
 		m_pauseTime = 0;
 	}
@@ -504,6 +504,8 @@ namespace uba
 			if (!isFirst)
 				posY += 3;
 
+			bool isRemote = processLocation.sessionIndex != 0;
+
 			if (posY + stepY >= progressRect.top && posY <= progressRect.bottom)
 			{
 				SelectObject(hdc, m_separatorPen);
@@ -566,6 +568,7 @@ namespace uba
 					int prevMemY = 0;
 					double sendScale = double(session.highestSendPerS) / (double(GraphHeight) - 2);
 					double recvScale = double(session.highestRecvPerS) / (double(GraphHeight) - 2);
+
 					for (auto& update : session.updates)
 					{
 						float cpuLoad = update.cpuLoad;
@@ -574,6 +577,9 @@ namespace uba
 						else
 							prevCupLoad = cpuLoad;
 
+						auto updateSend = update.send;
+						auto updateRecv = update.recv;
+
 						int x = int(posX + TimeToS(update.time) * scaleX);
 						int sendY = graphBaseY;
 						int recvY = graphBaseY;
@@ -583,22 +589,26 @@ namespace uba
 						double duration = TimeToS(update.time - prevTime);
 						if (update.time == 0)
 							isFirstUpdate = true;
-						else if (prevSend > update.send || prevRecv > update.recv)
+						else if (prevSend > updateSend || prevRecv > updateRecv)
 							isFirstUpdate = true;
 
 						if (double sendInvScaleY = duration * sendScale)
-							sendY = graphBaseY - int(double(update.send - prevSend) / sendInvScaleY);
+							sendY = graphBaseY - int(double(updateSend - prevSend) / sendInvScaleY);
 						if (double recvInvScaleY = duration * recvScale)
-							recvY = graphBaseY - int(double(update.recv - prevRecv) / recvInvScaleY) - 1;
+							recvY = graphBaseY - int(double(updateRecv - prevRecv) / recvInvScaleY) - 1;
 
 						if (!isFirstUpdate && x > clientRect.left && prevX <= clientRect.right)
 						{
-							if (m_visibleComponents[ComponentType_SendRecv] && update.send != 0 && update.recv != 0)
+							if (m_visibleComponents[ComponentType_SendRecv] && updateSend != 0 && updateRecv != 0)
 							{
-								SelectObject(hdc, m_sendPen);
+								auto sendPen = m_sendPen;
+								auto recvPen = m_recvPen;
+								if (!isRemote)
+									std::swap(sendPen, recvPen);
+								SelectObject(hdc, sendPen);
 								MoveToEx(hdc, prevX, prevSendY, NULL);
 								LineTo(hdc, x, sendY);
-								SelectObject(hdc, m_recvPen);
+								SelectObject(hdc, recvPen);
 								MoveToEx(hdc, prevX, prevRecvY, NULL);
 								LineTo(hdc, x, recvY);
 							}
@@ -619,8 +629,8 @@ namespace uba
 						prevCpuY = cpuY;
 						prevMemY = memY;
 						prevTime = update.time;
-						prevSend = update.send;
-						prevRecv = update.recv;
+						prevSend = updateSend;
+						prevRecv = updateRecv;
 					}
 				}
 				posY += GraphHeight;
@@ -637,7 +647,7 @@ namespace uba
 						if (selected)
 							SetBkMode(hdc, TRANSPARENT);
 					};
-				PaintDetailedStats(posY, progressRect, session, processLocation.sessionIndex != 0, playTime, drawText);
+				PaintDetailedStats(posY, progressRect, session, isRemote, playTime, drawText);
 			}
 
 			if (m_visibleComponents[ComponentType_Bars])
@@ -1201,8 +1211,8 @@ namespace uba
 			DrawTextLogger logger(hdc, r, m_popupFontHeight);
 			logger.SetColor(m_cpuColor).Info(L"  Cpu: %.1f%%", m_stats.cpuLoad * 100.0f);
 			logger.SetColor(m_memColor).Info(L"  Mem: %ls/%ls", BytesToText(m_stats.memTotal - m_stats.memAvail).str, BytesToText(m_stats.memTotal).str);
-			logger.SetColor(m_sendColor).Info(L"  Recv: %ls/s", BytesToText(m_stats.recvBytesPerSecond).str);
-			logger.SetColor(m_recvColor).Info(L"  Send: %ls/s", BytesToText(m_stats.sendBytesPerSecond).str);
+			logger.SetColor(m_recvColor).Info(L"  Recv: %ls/s", BytesToText(m_stats.recvBytesPerSecond).str);
+			logger.SetColor(m_sendColor).Info(L"  Send: %ls/s", BytesToText(m_stats.sendBytesPerSecond).str);
 			if (m_stats.ping)
 				logger.Info(L"  Ping: %ls", TimeToText(m_stats.ping, false, m_traceView.frequency).str);
 		}
@@ -1502,11 +1512,11 @@ namespace uba
 				drawTextFunc(str, textRect);
 			};
 
-		drawText(L"Finished: %u", session.processExitedCount);
-		drawText(L"Active: %u", session.processActiveCount);
-
 		if (isRemote)
 		{
+			drawText(L"Finished Processes: %u", session.processExitedCount);
+			drawText(L"Active Processes: %u", session.processActiveCount);
+
 			if (!session.updates.empty())
 			{
 				auto& u = session.updates.back();
@@ -1518,8 +1528,8 @@ namespace uba
 					recvPerS = u64((u.recv - session.prevRecv) / duration);
 				}
 				drawText(L"ClientId: %u  TcpCount: %u", session.clientUid.data1, u.connectionCount);
-				drawText(L"Recv: %ls (%s/s)", BytesToText(u.send), BytesToText(sendPerS));
-				drawText(L"Send: %ls (%s/s)", BytesToText(u.recv), BytesToText(recvPerS));
+				drawText(L"Recv: %ls (%s/s)", BytesToText(u.recv), BytesToText(recvPerS));
+				drawText(L"Send: %ls (%s/s)", BytesToText(u.send), BytesToText(sendPerS));
 			}
 
 			if (session.disconnectTime == ~u64(0))
@@ -1566,6 +1576,29 @@ namespace uba
 			drawFiles(L"Stored", session.storedFiles, session.storedFilesBytes, session.maxVisibleFiles);
 			posY = Max(posY, Max(posY1, posY2));
 		}
+		else
+		{
+			drawText(L"Finished Processes: %u (local: %u)", m_traceView.totalProcessExitedCount, session.processExitedCount);
+			drawText(L"Active Processes: %u (local: %u)", m_traceView.totalProcessActiveCount, session.processActiveCount);
+			drawText(L"Active Helpers: %u", m_traceView.activeSessionCount - 1);
+
+			if (!session.updates.empty())
+			{
+				auto& u = session.updates.back();
+				if (u.send || u.recv)
+				{
+					u64 sendPerS = 0;
+					u64 recvPerS = 0;
+					if (float duration = TimeToS(u.time - session.prevUpdateTime))
+					{
+						sendPerS = u64((u.send - session.prevSend) / duration);
+						recvPerS = u64((u.recv - session.prevRecv) / duration);
+					}
+					drawText(L"Recv: %ls (%s/s)", BytesToText(u.send), BytesToText(sendPerS));
+					drawText(L"Send: %ls (%s/s)", BytesToText(u.recv), BytesToText(recvPerS));
+				}
+			}
+		}
 	}
 
 	void Visualizer::HitTest(HitTestResult& outResult, const POINT& pos)
@@ -1627,6 +1660,8 @@ namespace uba
 			if (!isFirst)
 				posY += 3;
 
+			bool isRemote = sessionIndex != 0;
+
 			if (pos.y >= posY && pos.y < posY + SessionStepY)
 			{
 				if (pos.x < 500)
@@ -1664,8 +1699,14 @@ namespace uba
 						if (pos.x + hitOffset >= prevX && pos.x + hitOffset <= x)
 						{
 							double duration = TimeToS(update.time - prevTime);
-							outResult.stats.recvBytesPerSecond = u64((update.send - prevSend) / duration);
-							outResult.stats.sendBytesPerSecond = u64((update.recv - prevRecv) / duration);
+
+							auto recvBytesPerSecond = u64((update.send - prevSend) / duration);
+							auto sendBytesPerSecond = u64((update.recv - prevRecv) / duration);
+							if (isRemote)
+								std::swap(recvBytesPerSecond, sendBytesPerSecond);
+
+							outResult.stats.recvBytesPerSecond = recvBytesPerSecond;
+							outResult.stats.sendBytesPerSecond = sendBytesPerSecond;
 							outResult.stats.ping = update.ping;
 							outResult.stats.memAvail = update.memAvail;
 							outResult.stats.cpuLoad = update.cpuLoad;
@@ -1925,6 +1966,15 @@ namespace uba
 					return false;
 				m_namedTrace.Clear().Append(m_newTraceName);
 				title.Appendf(L"Connected to host");
+				m_traceView.finished = false;
+			}
+			else if (!m_fileName.IsEmpty())
+			{
+				if (!m_trace.ReadFile(m_traceView, m_fileName.data, m_replay != 0))
+					return false;
+				RedrawWindow(hWnd, NULL, NULL, RDW_INVALIDATE | RDW_UPDATENOW);
+				title.Append(m_fileName);
+				m_traceView.finished = m_replay == 0;
 			}
 			else
 			{
@@ -1932,9 +1982,9 @@ namespace uba
 					return false;
 				m_namedTrace.Clear().Append(m_newTraceName);
 				title.Appendf(L"%s (Listening for new sessions on channel '%s')", m_namedTrace.data, m_listenChannel.data);
+				m_traceView.finished = false;
 			}
 
-			m_traceView.finished = false;
 			SetWindowTextW(m_hwnd, title.data);
 			SendMessage(m_hwnd, WM_TIMER, 0, 0);
 			UpdateScrollbars(true);
