UBA: Fixes various small bugs.

diff --git a/Engine/Source/Programs/UnrealBuildAccelerator/Core/Public/UbaTimer.h b/Engine/Source/Programs/UnrealBuildAccelerator/Core/Public/UbaTimer.h
--- a/Engine/Source/Programs/UnrealBuildAccelerator/Core/Public/UbaTimer.h
+++ b/Engine/Source/Programs/UnrealBuildAccelerator/Core/Public/UbaTimer.h
@@ -28,7 +28,8 @@ namespace uba
 
 	struct Timer
 	{
-		Timer() = default;
+        Timer(Timer &o) { time.store(o.time.load()); count.store(o.count.load());}
+        Timer() = default;
 		Timer(u64 t, u32 c) : time(t), count(c) {}
 		Timer(Timer&& o) noexcept : time(o.time.load()), count(o.count.load()) {}
 		void operator=(const Timer& o) { time.store(o.time.load()); count.store(o.count.load()); }
diff --git a/Engine/Source/Programs/UnrealBuildAccelerator/Visualizer/Private/UbaTraceReader.cpp b/Engine/Source/Programs/UnrealBuildAccelerator/Visualizer/Private/UbaTraceReader.cpp
--- a/Engine/Source/Programs/UnrealBuildAccelerator/Visualizer/Private/UbaTraceReader.cpp
+++ b/Engine/Source/Programs/UnrealBuildAccelerator/Visualizer/Private/UbaTraceReader.cpp
@@ -59,8 +59,10 @@ namespace uba
 		u32 traceSize = reader.ReadU32(); (void)traceSize;
 		u32 version = reader.ReadU32();
 		if (version < TraceReadCompatibilityVersion || version > TraceVersion)
-			return m_logger.Error(L"Incompatible trace version (%u). Current executable supports version %u to %u.", version, TraceReadCompatibilityVersion, TraceVersion);
-
+		{
+			m_logger.Error(L"Incompatible trace version (%u). Current executable supports version %u to %u.", version, TraceReadCompatibilityVersion, TraceVersion);
+			return false;
+		}
 		out.version = version;
 		reader.ReadU32(); // ProcessId
 		u64 traceSystemStartTimeUs = 0;
@@ -911,4 +913,4 @@ namespace uba
 		return nullptr;
 	}
 #endif // PLATFORM_WINDOWS
-}
\ No newline at end of file
+}
