[UBA]
* Added so in-process proxy can use storage of client
* Fixed reset of trace reader
* Fixed annoying warnings when closing invalid sockets

diff --git a/Engine/Source/Programs/UnrealBuildAccelerator/Agent/Private/UbaAgent.cpp b/Engine/Source/Programs/UnrealBuildAccelerator/Agent/Private/UbaAgent.cpp
--- a/Engine/Source/Programs/UnrealBuildAccelerator/Agent/Private/UbaAgent.cpp
+++ b/Engine/Source/Programs/UnrealBuildAccelerator/Agent/Private/UbaAgent.cpp
@@ -1014,6 +1014,7 @@ namespace uba
 				u32& targetTcpConnectionCount;
 				NetworkServer* server = nullptr;
 				StorageProxy* storage = nullptr;
+				StorageClient* storageClient = nullptr;
 				TString serverPrefix;
 			} proxy { g_consoleLogWriter, client, wakeupSessionWait, maxTcpConnectionCount, targetTcpConnectionCount };
 			auto psg = MakeGuard([&]() { delete proxy.server; });
@@ -1037,7 +1038,7 @@ namespace uba
 						delete proxy.server;
 						return false;
 					}
-					proxy.storage = new StorageProxy(*proxy.server, *proxy.client, storageServerUid, TC("Wooohoo"));
+					proxy.storage = new StorageProxy(*proxy.server, *proxy.client, storageServerUid, TC("Wooohoo"), proxy.storageClient);
 					proxy.server->StartListen(proxy.client->GetTcpBackend(), proxyPort);
 					proxy.targetTcpConnectionCount = proxy.maxTcpConnectionCount;
 					proxy.wakeupSessionWait.Set();
@@ -1064,6 +1065,8 @@ namespace uba
 			if (!storageClient->PopulateCasFromDirs(populateCasDirs, maxProcessCount))
 				return -1;
 
+			proxy.storageClient = storageClient;
+
 			SessionClient* sessionClient = nullptr;
 
 			CasKey keys[2];
diff --git a/Engine/Source/Programs/UnrealBuildAccelerator/Common/Private/UbaNetworkBackendTcp.cpp b/Engine/Source/Programs/UnrealBuildAccelerator/Common/Private/UbaNetworkBackendTcp.cpp
--- a/Engine/Source/Programs/UnrealBuildAccelerator/Common/Private/UbaNetworkBackendTcp.cpp
+++ b/Engine/Source/Programs/UnrealBuildAccelerator/Common/Private/UbaNetworkBackendTcp.cpp
@@ -101,8 +101,9 @@ namespace uba
 
 	bool CloseSocket(Logger& logger, SOCKET s)
 	{
-		if (closesocket(s) == SOCKET_ERROR)
-			return logger.Error(TC("failed to close socket (%s)"), LastErrorToText(WSAGetLastError()).data);
+		if (s != INVALID_SOCKET)
+			if (closesocket(s) == SOCKET_ERROR)
+				return logger.Error(TC("failed to close socket (%s)"), LastErrorToText(WSAGetLastError()).data);
 		return true;
 	}
 
diff --git a/Engine/Source/Programs/UnrealBuildAccelerator/Common/Private/UbaStorageClient.cpp b/Engine/Source/Programs/UnrealBuildAccelerator/Common/Private/UbaStorageClient.cpp
--- a/Engine/Source/Programs/UnrealBuildAccelerator/Common/Private/UbaStorageClient.cpp
+++ b/Engine/Source/Programs/UnrealBuildAccelerator/Common/Private/UbaStorageClient.cpp
@@ -298,6 +298,7 @@ namespace uba
 				{
 					StringBuffer<> proxyHost;
 					u16 proxyPort;
+					bool isInProcessClient = false;
 					if (reader.ReadBool())
 					{
 						proxyPort = reader.ReadU16();
@@ -306,6 +307,7 @@ namespace uba
 							// TODO: Tell server we failed
 						}
 						proxyHost.Append(TC("127.0.0.1"));
+						isInProcessClient = true;
 					}
 					else
 					{
@@ -358,6 +360,7 @@ namespace uba
 
 						proxyWriter.WriteString(TC("ProxyClient"));
 						proxyWriter.WriteU32(StorageNetworkVersion);
+						proxyWriter.WriteBool(isInProcessClient);
 						StackBinaryReader<256> proxyReader;
 						if (!proxyMsg.Send(proxyReader))
 							continue;
diff --git a/Engine/Source/Programs/UnrealBuildAccelerator/Common/Private/UbaStorageProxy.cpp b/Engine/Source/Programs/UnrealBuildAccelerator/Common/Private/UbaStorageProxy.cpp
--- a/Engine/Source/Programs/UnrealBuildAccelerator/Common/Private/UbaStorageProxy.cpp
+++ b/Engine/Source/Programs/UnrealBuildAccelerator/Common/Private/UbaStorageProxy.cpp
@@ -1,6 +1,7 @@
 // Copyright Epic Games, Inc. All Rights Reserved.
 
 #include "UbaStorageProxy.h"
+#include "UbaFileAccessor.h"
 #include "UbaNetworkClient.h"
 #include "UbaNetworkMessage.h"
 #include "UbaNetworkServer.h"
@@ -8,9 +9,10 @@
 
 namespace uba
 {
-	StorageProxy::StorageProxy(NetworkServer& server, NetworkClient& client, const Guid& storageServerUid, const tchar* name)
+	StorageProxy::StorageProxy(NetworkServer& server, NetworkClient& client, const Guid& storageServerUid, const tchar* name, StorageImpl* localStorage)
 	:	m_server(server)
 	,	m_client(client)
+	,	m_localStorage(localStorage)
 	,	m_logger(client.GetLogWriter(), TC("StorageProxy"))
 	,	m_storageServerUid(storageServerUid)
 	,	m_name(name)
@@ -71,6 +73,9 @@ namespace uba
 					m_logger.Error(TC("Different network versions. Client: %u, Server: %u. Disconnecting"), clientVersion, StorageNetworkVersion);
 					return false;
 				}
+				bool isInProcessClient = reader.ReadBool();
+				if (isInProcessClient)
+					m_inProcessClientId = connectionInfo.GetId();
 
 				//m_logger.Info(TC("%s connected"), clientName.data);
 				writer.WriteGuid(m_storageServerUid);
@@ -95,70 +100,109 @@ namespace uba
 
 				ScopedWriteLock fileLock(file.lock);
 
-				if (file.view.size != 0)
+				while (true)
 				{
-					fileLock.Leave();
-					view = m_fileMapping.MapView(file.view.handle, file.view.offset, file.view.size, hint.data);
-				}
-				else
-				{
-					NetworkMessage msg(m_client, ServiceId, messageType, writer2);
-					writer2.WriteBool(false); // Wants proxy
-					writer2.WriteCasKey(casKey);
-					writer2.WriteString(hint);
-					writer2.WriteBytes(reader.GetPositionData(), reader.GetLeft());
-
-					if (!msg.Send(reader2))
-						return false;
-
-					BinaryReader tempReader(reader2.GetPositionData(), 0, reader2.GetLeft());
-					u32 sizeOfFirstMessage = u32(reader2.GetLeft());
-					u16 fetchId = tempReader.ReadU16();
-					if (fetchId == 0)
+					if (file.view.size != 0)
 					{
-						m_logger.Error(TC("FetchBegin failed for cas file %s (%s)"), CasKeyString(casKey).str, hint.data);
-						writer.WriteU16(0);
-						return true;
+						fileLock.Leave();
+						view = m_fileMapping.MapView(file.view.handle, file.view.offset, file.view.size, hint.data);
+						break;
 					}
-					u64 fileSize = tempReader.Read7BitEncoded();
-					view = m_fileMapping.AllocAndMapView(MappedView_Transient, fileSize, 1, hint.data);
-
-					u8 flags = tempReader.ReadByte();
-					bool storeCompressed = (flags >> 0) & 1;
-					bool sendEnd = (flags >> 1) & 1;
-					u64 fetchedSize = tempReader.GetLeft();
 
-					u32 responseSize = u32(fetchedSize);
+					bool storeCompressed = true;
 
-					u64 left = fileSize;
-					u8* readBuffer = view.memory;
+					if (m_localStorage && IsCompressed(casKey) && m_inProcessClientId && connectionInfo.GetId() != m_inProcessClientId)
+					{
+						// We need to leave this lock here since the in-process storage client might be asking for this file too and then we can end up in a deadlock
+						fileLock.Leave();
+
+						bool hasCas = m_localStorage->EnsureCasFile(casKey, nullptr);
+
+						// Enter lock again, and also check if another thread might have already handled this file while we looked if it existed in local storage
+						fileLock.Enter();
+						if (file.view.size != 0)
+							continue;
+
+						if (hasCas)
+						{
+
+							StringBuffer<> casFile;
+							if (m_localStorage->GetCasFileName(casFile, casKey))
+							{
+								FileAccessor sourceFile(m_logger, casFile.data);
+								if (sourceFile.OpenMemoryRead())
+								{
+									u64 fileSize = sourceFile.GetSize();
+									view = m_fileMapping.AllocAndMapView(MappedView_Transient, fileSize, 1, hint.data);
+									if (!view.memory)
+										return false;
+									memcpy(view.memory, sourceFile.GetData(), fileSize);
+								}
+							}
+						}
+					}
 
-					memcpy(view.memory, tempReader.GetPositionData(), responseSize);
-					left -= responseSize;
-					readBuffer += responseSize;
+					if (!view.handle.IsValid())
+					{
+						NetworkMessage msg(m_client, ServiceId, messageType, writer2);
+						writer2.WriteBool(false); // Wants proxy
+						writer2.WriteCasKey(casKey);
+						writer2.WriteString(hint);
+						writer2.WriteBytes(reader.GetPositionData(), reader.GetLeft());
 
-					if (!StorageClient::SendAllSegments(m_client, fetchId, readBuffer, left, sizeOfFirstMessage))
-						return false;
+						if (!msg.Send(reader2))
+							return false;
 
-					if (sendEnd)
-					{
-						StackBinaryWriter<128> writer3;
-						NetworkMessage msg2(m_client, ServiceId, StorageMessageType_FetchEnd, writer3);
-						writer3.WriteCasKey(casKey);
-						if (!msg2.Send())
+						BinaryReader tempReader(reader2.GetPositionData(), 0, reader2.GetLeft());
+						u32 sizeOfFirstMessage = u32(reader2.GetLeft());
+						u16 fetchId = tempReader.ReadU16();
+						if (fetchId == 0)
+						{
+							m_logger.Error(TC("FetchBegin failed for cas file %s (%s)"), CasKeyString(casKey).str, hint.data);
+							writer.WriteU16(0);
+							return true;
+						}
+						u64 fileSize = tempReader.Read7BitEncoded();
+						view = m_fileMapping.AllocAndMapView(MappedView_Transient, fileSize, 1, hint.data);
+
+						u8 flags = tempReader.ReadByte();
+						storeCompressed = (flags >> 0) & 1;
+						bool sendEnd = (flags >> 1) & 1;
+						u64 fetchedSize = tempReader.GetLeft();
+
+						u32 responseSize = u32(fetchedSize);
+
+						u64 left = fileSize;
+						u8* readBuffer = view.memory;
+
+						memcpy(view.memory, tempReader.GetPositionData(), responseSize);
+						left -= responseSize;
+						readBuffer += responseSize;
+
+						if (!StorageClient::SendAllSegments(m_client, fetchId, readBuffer, left, sizeOfFirstMessage))
 							return false;
-					}
 
+						if (sendEnd)
+						{
+							StackBinaryWriter<128> writer3;
+							NetworkMessage msg2(m_client, ServiceId, StorageMessageType_FetchEnd, writer3);
+							writer3.WriteCasKey(casKey);
+							if (!msg2.Send())
+								return false;
+						}
+					}
 					file.storeCompressed = storeCompressed;
-					file.hasSegments = fetchId != u16(~0);
 					file.view = view;
 					fileLock.Leave();
+					break;
 				}
 
 				u16 fetchId = u16(~0);
-				u64 fetchedSize = view.size;
 
-				if (file.hasSegments)
+				u64 headerSize = sizeof(u16) + Get7BitEncodedCount(file.view.size) + sizeof(u8);
+				u64 fetchedSize = Min(view.size, m_client.GetMessageMaxSize() - m_client.GetMessageReceiveHeaderSize() - headerSize);
+
+				if (fetchedSize < view.size)
 				{
 					fetchId = PopId();
 					ScopedWriteLock lock(m_activeFetchesLock);
@@ -169,9 +213,6 @@ namespace uba
 					ActiveFetch& fetch = res.first->second;
 					lock.Leave();
 
-					u64 headerSize = sizeof(u16) + Get7BitEncodedCount(file.view.size) + sizeof(u8);
-					fetchedSize = Min(view.size, m_client.GetMessageMaxSize() - m_client.GetMessageReceiveHeaderSize() - headerSize);
-
 					fetch.fetchedSize = fetchedSize;
 					fetch.file = &file;
 				}
diff --git a/Engine/Source/Programs/UnrealBuildAccelerator/Common/Public/UbaNetwork.h b/Engine/Source/Programs/UnrealBuildAccelerator/Common/Public/UbaNetwork.h
--- a/Engine/Source/Programs/UnrealBuildAccelerator/Common/Public/UbaNetwork.h
+++ b/Engine/Source/Programs/UnrealBuildAccelerator/Common/Public/UbaNetwork.h
@@ -18,7 +18,7 @@ namespace uba
 	
 	static constexpr u32 SystemNetworkVersion = 1338;
 	static constexpr u32 StorageNetworkVersion = 3;
-	static constexpr u32 SessionNetworkVersion = 26;
+	static constexpr u32 SessionNetworkVersion = 27;
 
 
 	// Messages used over network between client/server (system, storage and session)
diff --git a/Engine/Source/Programs/UnrealBuildAccelerator/Common/Public/UbaStorageProxy.h b/Engine/Source/Programs/UnrealBuildAccelerator/Common/Public/UbaStorageProxy.h
--- a/Engine/Source/Programs/UnrealBuildAccelerator/Common/Public/UbaStorageProxy.h
+++ b/Engine/Source/Programs/UnrealBuildAccelerator/Common/Public/UbaStorageProxy.h
@@ -9,6 +9,7 @@ namespace uba
 {
     class NetworkClient;
 	class NetworkServer;
+	class StorageImpl;
 	struct BinaryReader;
 	struct BinaryWriter;
 	struct ConnectionInfo;
@@ -16,7 +17,7 @@ namespace uba
 	class StorageProxy
 	{
 	public:
-		StorageProxy(NetworkServer& server, NetworkClient& client, const Guid& storageServerUid, const tchar* name);
+		StorageProxy(NetworkServer& server, NetworkClient& client, const Guid& storageServerUid, const tchar* name, StorageImpl* localStorage = nullptr);
 
 		bool Disconnect(u32 timeoutMs); // Use a timeout to try to make more graceful disconnect
 		void PrintSummary();
@@ -30,6 +31,7 @@ namespace uba
 
 		NetworkServer& m_server;
 		NetworkClient& m_client;
+		StorageImpl* m_localStorage;
 
 		LoggerWithWriter m_logger;
 
@@ -37,7 +39,9 @@ namespace uba
 
 		TString m_name;
 
-		struct FileEntry { ReaderWriterLock lock; MappedView view; bool hasSegments = false; bool storeCompressed = false; };
+		u32 m_inProcessClientId = 0;
+
+		struct FileEntry { ReaderWriterLock lock; MappedView view; bool storeCompressed = false; };
 		ReaderWriterLock m_filesLock;
 		UnorderedMap<CasKey, FileEntry> m_files;
 
diff --git a/Engine/Source/Programs/UnrealBuildAccelerator/Visualizer/Private/UbaTraceReader.cpp b/Engine/Source/Programs/UnrealBuildAccelerator/Visualizer/Private/UbaTraceReader.cpp
--- a/Engine/Source/Programs/UnrealBuildAccelerator/Visualizer/Private/UbaTraceReader.cpp
+++ b/Engine/Source/Programs/UnrealBuildAccelerator/Visualizer/Private/UbaTraceReader.cpp
@@ -59,10 +59,8 @@ namespace uba
 		u32 traceSize = reader.ReadU32(); (void)traceSize;
 		u32 version = reader.ReadU32();
 		if (version < TraceReadCompatibilityVersion || version > TraceVersion)
-		{
-			m_logger.Error(L"Incompatible trace version (%u). Current executable supports version %u to %u.", version, TraceReadCompatibilityVersion, TraceVersion);
-			return false;
-		}
+			return m_logger.Error(L"Incompatible trace version (%u). Current executable supports version %u to %u.", version, TraceReadCompatibilityVersion, TraceVersion);
+
 		out.version = version;
 		reader.ReadU32(); // ProcessId
 		u64 traceSystemStartTimeUs = 0;
@@ -564,6 +562,9 @@ namespace uba
 			process2.start = time;
 			process2.stop = ~u64(0);
 			process2.exitCode = ~0u;
+
+			++session.processExitedCount;
+			++out.totalProcessExitedCount;
 			break;
 		}
 		case TraceType_ProcessReturned:
@@ -913,4 +914,4 @@ namespace uba
 		return nullptr;
 	}
 #endif // PLATFORM_WINDOWS
-}
+}
\ No newline at end of file
diff --git a/Engine/Source/Programs/UnrealBuildAccelerator/Visualizer/Private/UbaTraceReader.h b/Engine/Source/Programs/UnrealBuildAccelerator/Visualizer/Private/UbaTraceReader.h
--- a/Engine/Source/Programs/UnrealBuildAccelerator/Visualizer/Private/UbaTraceReader.h
+++ b/Engine/Source/Programs/UnrealBuildAccelerator/Visualizer/Private/UbaTraceReader.h
@@ -107,7 +107,7 @@ namespace uba
 		};
 
 		Process* GetProcess(const ProcessLocation& loc) { return &(sessions[loc.sessionIndex].processors[loc.processorIndex].processes[loc.processIndex]); }
-		void Clear() { sessions.clear(); workTracks.clear(); strings.clear(); startTime = 0; finished = true; };
+		void Clear() { sessions.clear(); workTracks.clear(); strings.clear(); startTime = 0; finished = true; totalProcessActiveCount = 0; totalProcessExitedCount = 0; activeSessionCount = 0; };
 
 		Vector<Session> sessions;
 		Vector<WorkTrack> workTracks;
