[UBA]
* Added log entries to trace file

diff --git a/Engine/Source/Programs/UnrealBuildAccelerator/Common/Private/UbaSession.cpp b/Engine/Source/Programs/UnrealBuildAccelerator/Common/Private/UbaSession.cpp
--- a/Engine/Source/Programs/UnrealBuildAccelerator/Common/Private/UbaSession.cpp
+++ b/Engine/Source/Programs/UnrealBuildAccelerator/Common/Private/UbaSession.cpp
@@ -1322,7 +1322,10 @@ namespace uba
 		{
 			StackBinaryWriter<1024> writer;
 			process.m_processStats.Write(writer);
-			m_trace.ProcessExited(id, process.GetExitCode(), writer.GetData(), writer.GetPosition());
+			u32 exitCode = process.GetExitCode();
+			Vector<ProcessLogLine> emptyLines;
+			auto& logLines = (exitCode != 0 || m_detailedTrace) ? process.m_logLines : emptyLines;
+			m_trace.ProcessExited(id, exitCode, writer.GetData(), writer.GetPosition(), logLines);
 			ScopedWriteLock lock(m_processStatsLock);
 			m_processStats.Add(process.m_processStats);
 		}
diff --git a/Engine/Source/Programs/UnrealBuildAccelerator/Common/Private/UbaSessionServer.cpp b/Engine/Source/Programs/UnrealBuildAccelerator/Common/Private/UbaSessionServer.cpp
--- a/Engine/Source/Programs/UnrealBuildAccelerator/Common/Private/UbaSessionServer.cpp
+++ b/Engine/Source/Programs/UnrealBuildAccelerator/Common/Private/UbaSessionServer.cpp
@@ -425,7 +425,7 @@ namespace uba
 		StackBinaryWriter<1024> statsWriter;
 		ProcessStats processStats;
 		processStats.Write(statsWriter);
-		m_trace.ProcessExited(id, exitCode, statsWriter.GetData(), statsWriter.GetPosition());
+		m_trace.ProcessExited(id, exitCode, statsWriter.GetData(), statsWriter.GetPosition(), Vector<ProcessLogLine>());
 	}
 
 	NetworkServer& SessionServer::GetServer()
@@ -1023,7 +1023,9 @@ namespace uba
 						func(process.startInfo.logLineUserData, line.text.c_str(), u32(line.text.size()), line.type);
 
 				u32 id = process.m_processId;
-				m_trace.ProcessExited(id, exitCode, reader.GetPositionData(), reader.GetLeft());
+				Vector<ProcessLogLine> emptyLines;
+				auto& logLines = (exitCode != 0 || m_detailedTrace) ? process.m_logLines : emptyLines;
+				m_trace.ProcessExited(id, exitCode, reader.GetPositionData(), reader.GetLeft(), logLines);
 
 				ProcessStats processStats;
 				processStats.Read(reader, ~0u);
diff --git a/Engine/Source/Programs/UnrealBuildAccelerator/Common/Private/UbaTrace.cpp b/Engine/Source/Programs/UnrealBuildAccelerator/Common/Private/UbaTrace.cpp
--- a/Engine/Source/Programs/UnrealBuildAccelerator/Common/Private/UbaTrace.cpp
+++ b/Engine/Source/Programs/UnrealBuildAccelerator/Common/Private/UbaTrace.cpp
@@ -2,6 +2,7 @@
 
 #include "UbaTrace.h"
 #include "UbaFileAccessor.h"
+#include "UbaProcessHandle.h"
 
 namespace uba
 {
@@ -187,12 +188,18 @@ namespace uba
 		writer.WriteBytes(data, dataSize);
 	}
 
-	void Trace::ProcessExited(u32 processId, u32 exitCode, const u8* data, u64 dataSize)
+	void Trace::ProcessExited(u32 processId, u32 exitCode, const u8* data, u64 dataSize, const Vector<ProcessLogLine>& logLines)
 	{
 		BEGIN_TRACE_ENTRY(TraceType_ProcessExited);
 		writer.WriteU32(processId);
 		writer.WriteU32(exitCode);
 		writer.WriteBytes(data, dataSize);
+		writer.Write7BitEncoded(u64(logLines.size()));
+		for (auto& line : logLines)
+		{
+			writer.WriteByte(line.type);
+			writer.WriteString(line.text);
+		}
 	}
 
 	void Trace::ProcessReturned(u32 processId)
diff --git a/Engine/Source/Programs/UnrealBuildAccelerator/Common/Public/UbaTrace.h b/Engine/Source/Programs/UnrealBuildAccelerator/Common/Public/UbaTrace.h
--- a/Engine/Source/Programs/UnrealBuildAccelerator/Common/Public/UbaTrace.h
+++ b/Engine/Source/Programs/UnrealBuildAccelerator/Common/Public/UbaTrace.h
@@ -7,6 +7,7 @@
 
 namespace uba
 {
+	struct ProcessLogLine;
 	struct ProcessStats;
 
 	class TraceChannel
@@ -49,7 +50,7 @@ namespace uba
 		TraceType_FileStoreLight,
 	};
 
-	static constexpr u32 TraceVersion = 19;
+	static constexpr u32 TraceVersion = 20;
 	static constexpr u32 TraceReadCompatibilityVersion = 6;
 
 	class Trace
@@ -67,7 +68,7 @@ namespace uba
 		void SessionDisconnect(u32 sessionId);
 		void ProcessAdded(u32 sessionId, u32 processId, const tchar* description);
 		void ProcessEnvironmentUpdated(u32 processId, const tchar* reason, const u8* data, u64 dataSize);
-		void ProcessExited(u32 processId, u32 exitCode, const u8* data, u64 dataSize);
+		void ProcessExited(u32 processId, u32 exitCode, const u8* data, u64 dataSize, const Vector<ProcessLogLine>& logLines);
 		void ProcessReturned(u32 processId);
 		void ProxyCreated(u32 clientId, const tchar* proxyName);
 		void ProxyUsed(u32 clientId, const tchar* proxyName);
diff --git a/Engine/Source/Programs/UnrealBuildAccelerator/Visualizer/Private/UbaTraceReader.cpp b/Engine/Source/Programs/UnrealBuildAccelerator/Visualizer/Private/UbaTraceReader.cpp
--- a/Engine/Source/Programs/UnrealBuildAccelerator/Visualizer/Private/UbaTraceReader.cpp
+++ b/Engine/Source/Programs/UnrealBuildAccelerator/Visualizer/Private/UbaTraceReader.cpp
@@ -526,6 +526,18 @@ namespace uba
 				process.storageStats.Read(reader);
 				process.systemStats.Read(reader);
 			}
+
+			if (out.version >= 20)
+			{
+				u64 logLineCount = reader.Read7BitEncoded();
+				process.logLines.reserve(logLineCount);
+				while (logLineCount--)
+				{
+					auto type = (LogEntryType)reader.ReadByte();
+					process.logLines.emplace_back(reader.ReadString(), type);
+				}
+			}
+
 			process.exitCode = exitCode;
 			process.stop = time;
 			process.bitmapDirty = true;
diff --git a/Engine/Source/Programs/UnrealBuildAccelerator/Visualizer/Private/UbaTraceReader.h b/Engine/Source/Programs/UnrealBuildAccelerator/Visualizer/Private/UbaTraceReader.h
--- a/Engine/Source/Programs/UnrealBuildAccelerator/Visualizer/Private/UbaTraceReader.h
+++ b/Engine/Source/Programs/UnrealBuildAccelerator/Visualizer/Private/UbaTraceReader.h
@@ -2,6 +2,7 @@
 
 #pragma once
 
+#include "UbaProcessHandle.h"
 #include "UbaStats.h"
 #include "UbaTrace.h"
 
@@ -27,6 +28,7 @@ namespace uba
 			SessionStats sessionStats;
 			StorageStats storageStats;
 			SystemStats systemStats;
+			Vector<ProcessLogLine> logLines;
 		};
 
 		struct Processor
diff --git a/Engine/Source/Programs/UnrealBuildAccelerator/Visualizer/Private/UbaVisualizer.cpp b/Engine/Source/Programs/UnrealBuildAccelerator/Visualizer/Private/UbaVisualizer.cpp
--- a/Engine/Source/Programs/UnrealBuildAccelerator/Visualizer/Private/UbaVisualizer.cpp
+++ b/Engine/Source/Programs/UnrealBuildAccelerator/Visualizer/Private/UbaVisualizer.cpp
@@ -13,6 +13,7 @@ enum
 {
 	Popup_CopySessionInfo = 3,
 	Popup_CopyProcessInfo,
+	Popup_CopyProcessLog,
 	Popup_SaveAs,
 	Popup_Quit,
 };
@@ -1062,6 +1063,11 @@ namespace uba
 			bool hasStorageStats = false;
 			u32 lineCount = 4;
 			u64 duration = 0;
+			
+			int width = 290;
+			Vector<TString> logLines;
+			u32 maxCharCount = 50u;
+
 			bool hasExited = process.stop != ~u64(0);
 			if (hasExited)
 			{
@@ -1078,13 +1084,34 @@ namespace uba
 				duration = process.stop - process.start;
 				if (process.exitCode != 0)
 					++lineCount;
+
+				if (!process.logLines.empty())
+				{
+					u32 lineMaxCount = 0;
+					for (auto& line : process.logLines)
+					{
+						u32 offset = 0;
+						u32 left = u32(line.text.size());
+						while (left)
+						{
+							u32 toCopy = Min(left, maxCharCount);
+							lineMaxCount = Max(lineMaxCount, toCopy);
+							logLines.push_back(line.text.substr(offset, toCopy));
+							offset += toCopy;
+							left -= toCopy;
+						}
+					}
+
+					width = Max(width, int(lineMaxCount) * 7 + 2*14);
+
+					lineCount += u32(logLines.size()) + 1;
+				}
 			}
 			else
 			{
 				duration = playTime - process.start;
 			}
 
-			int width = 290;
 			int height = lineCount*m_popupFontHeight;
 
 			POINT p;
@@ -1133,6 +1160,14 @@ namespace uba
 					logger.Info(L"  ----------- System stats ------------");
 					process.systemStats.Print(logger, false, m_traceView.frequency);
 				}
+
+				if (!logLines.empty())
+				{
+					logger.Info(L"  ---------------- Log ----------------");
+					logger.rect.left += 14;
+					for (auto& line : logLines)
+						logger.Log(LogEntryType_Info, line.c_str(), u32(line.size()));
+				}
 			}
 		}
 		else if (m_sessionSelectedIndex != ~0u)
@@ -2259,7 +2294,11 @@ namespace uba
 			}
 			else if (m_processSelected)
 			{
+				TraceView::Process& process = *m_traceView.GetProcess(m_processSelectedLocation);
+
 				AppendMenuW(hMenu, MF_STRING, Popup_CopyProcessInfo, L"&Copy Process Info");
+				if (!process.logLines.empty())
+					AppendMenuW(hMenu, MF_STRING, Popup_CopyProcessLog, L"&Copy Process Log");
 				AppendMenuW(hMenu, MF_SEPARATOR, 0, NULL);
 			}
 			AppendMenuW(hMenu, MF_STRING, Popup_SaveAs, L"&Save Trace");
@@ -2315,6 +2354,21 @@ namespace uba
 				CopyTextToClipboard(str);
 				break;
 			}
+			case Popup_CopyProcessLog:
+			{
+				TString str;
+				TraceView::Process& process = *m_traceView.GetProcess(m_processSelectedLocation);
+				bool isFirst = true;
+				for (auto line : process.logLines)
+				{
+					if (!isFirst)
+						str += '\n';
+					isFirst = false;
+					str += line.text;
+				}
+				CopyTextToClipboard(str);
+				break;
+			}
 			}
 
 			DestroyMenu(hMenu);
