UBA: Fix up tests on Mac.

* Need to grab the actual clang binary location and the SDK path.

diff --git a/Engine/Source/Programs/UnrealBuildAccelerator/Common/Public/UbaDirectoryIterator.h b/Engine/Source/Programs/UnrealBuildAccelerator/Common/Public/UbaDirectoryIterator.h
--- a/Engine/Source/Programs/UnrealBuildAccelerator/Common/Public/UbaDirectoryIterator.h
+++ b/Engine/Source/Programs/UnrealBuildAccelerator/Common/Public/UbaDirectoryIterator.h
@@ -203,7 +203,7 @@ namespace uba
 		while ((pDirent = readdir(dir)) != NULL)
 		{
 			char* fileName = pDirent->d_name;
-			if (fileName[0] == '.' && (fileName[1] == 0 || (fileName[1] == '.' && fileName[2] == 0)))
+			if (fileName[0] == '.' && (fileName[1] == 0 || (fileName[1] == '.' && fileName[2] == 0))) 
 				continue;
 
 			strcpy(dirPath2 + dirPathLen, fileName);
@@ -216,15 +216,11 @@ namespace uba
 				// just in case.
 				if (errno == ENOENT) {
 					res = lstat(dirPath2, &attr);
-
-					if (res == 0)
-						continue;
 				}
 
 				if (errno == ENODEV || errno == EACCES)
 					continue;
-				UBA_ASSERTF(false, "TraverseDir:stat error handling not added for file %s (error %s)", dirPath2, strerror(errno));
-				continue;
+				UBA_ASSERTF(res == 0, "TraverseDir:stat error handling not added for file %s (error %s)", dirPath2, strerror(errno));
 			}
 
 			DirectoryEntry entry;
diff --git a/Engine/Source/Programs/UnrealBuildAccelerator/Test/Private/UbaTestAll.h b/Engine/Source/Programs/UnrealBuildAccelerator/Test/Private/UbaTestAll.h
--- a/Engine/Source/Programs/UnrealBuildAccelerator/Test/Private/UbaTestAll.h
+++ b/Engine/Source/Programs/UnrealBuildAccelerator/Test/Private/UbaTestAll.h
@@ -10,6 +10,17 @@
 
 namespace uba
 {
+#if PLATFORM_MAC
+	#define UBA_EXTRA_TESTS
+#else
+	#define UBA_EXTRA_TESTS \
+		UBA_TEST(TestMultipleDetouredProcesses) \
+		UBA_TEST(TestLogLines) \
+		UBA_TEST(TestLocalSchedule) \
+		UBA_TEST(TestLocalScheduleReuse) \
+		UBA_TEST(TestRemoteScheduleReuse) \
+
+#endif 
 
 	#define UBA_TESTS \
 		UBA_TEST(TestTime) \
@@ -26,11 +37,8 @@ namespace uba
 		UBA_TEST(TestDetouredTouch) \
 		UBA_TEST(TestRemoteDetouredTestApp) \
 		UBA_TEST(TestCustomService) \
-		UBA_TEST(TestMultipleDetouredProcesses) \
-		UBA_TEST(TestLogLines) \
-		UBA_TEST(TestLocalSchedule) \
-		UBA_TEST(TestLocalScheduleReuse) \
-		UBA_TEST(TestRemoteScheduleReuse) \
+		UBA_EXTRA_TESTS
+
 
 	#define UBA_TEST(x) \
 		logger.Info(TC("Running %hs..."), #x); \
@@ -61,8 +69,8 @@ namespace uba
 
 		logger.Info(TC("Running tests (Test rootdir: %s)"), testRootDir.data);
 
-		UBA_TEST(TestRemoteScheduleReuse)
-		//UBA_TESTS
+		// UBA_TEST(TestRemoteDetouredTestApp)
+		UBA_TESTS
 
 		logger.Info(TC("Tests finished successfully!"));
 		Sleep(3000);
diff --git a/Engine/Source/Programs/UnrealBuildAccelerator/Test/Private/UbaTestSession.h b/Engine/Source/Programs/UnrealBuildAccelerator/Test/Private/UbaTestSession.h
--- a/Engine/Source/Programs/UnrealBuildAccelerator/Test/Private/UbaTestSession.h
+++ b/Engine/Source/Programs/UnrealBuildAccelerator/Test/Private/UbaTestSession.h
@@ -149,6 +149,24 @@ namespace uba
 
 		return true;
 	}
+
+	bool ExecuteCommand(LoggerWithWriter& logger, const tchar* command, StringBufferBase& commandOutput)
+	{
+		FILE* fpCommand = popen(command, "r");
+		if (fpCommand == nullptr || fgets(commandOutput.data, commandOutput.capacity, fpCommand) == nullptr || pclose(fpCommand) != 0)
+		{
+			logger.Error("Failed to get an Xcode from xcode-select");
+			return false;
+		}
+
+		commandOutput.count = strlen(commandOutput.data);
+		while (isspace(commandOutput.data[commandOutput.count-1]))
+		{
+			commandOutput.data[commandOutput.count-1] = 0;
+			commandOutput.count--;
+		}
+		return true;
+	}
 	
 	bool RunClang(LoggerWithWriter& logger, SessionServer& session, const tchar* workingDir, const RunProcessFunction& runProcess)
 	{
@@ -163,11 +181,27 @@ namespace uba
 		if (!codeFile.Close())
 			return false;
 
+#if PLATFORM_MAC
+		StringBuffer<512> xcodePath;
+		ExecuteCommand(logger, "/usr/bin/xcrun -f clang++", xcodePath);
+		const tchar* clangPath = xcodePath.data;
+#else
 		const tchar* clangPath = TC("/usr/bin/clang++");
+#endif
 
 		ProcessStartInfo processInfo;
 		processInfo.application = clangPath;
+#if PLATFORM_MAC
+		StringBuffer<512> xcodeSDKPath;
+		StringBuffer<> args;
+		ExecuteCommand(logger, "xcrun --show-sdk-path", xcodeSDKPath);
+		args.Append("-isysroot ");
+		args.Append(xcodeSDKPath.data);
+		args.Append(" -o code Code.cpp");
+		processInfo.arguments = args.data;
+#else
 		processInfo.arguments = TC("-o code Code.cpp");
+#endif
 		processInfo.workingDir = workingDir;
 		//processInfo.logFile = TC("[REDACTED_USER_HOME]/RunClang.log");
 		ProcessHandle process = runProcess(processInfo);
diff --git a/Engine/Source/Programs/UnrealBuildAccelerator/TestApp/Private/PosixOS/UbaTestAppPosix.cpp b/Engine/Source/Programs/UnrealBuildAccelerator/TestApp/Private/PosixOS/UbaTestAppPosix.cpp
--- a/Engine/Source/Programs/UnrealBuildAccelerator/TestApp/Private/PosixOS/UbaTestAppPosix.cpp
+++ b/Engine/Source/Programs/UnrealBuildAccelerator/TestApp/Private/PosixOS/UbaTestAppPosix.cpp
@@ -40,6 +40,7 @@ bool IsLibraryLoaded(const char* libraryToMatch)
 					return strstr(info->dlpi_name, (const char*)libraryToMatch) != 0 ? 1 : 0;
 				}, (void*)libraryToMatch);
 #elif PLATFORM_MAC
+
 		bool foundDetoursLib = false;
 		unsigned int count = _dyld_image_count();
 		for (int i = 0; i < count; i++)
@@ -62,10 +63,17 @@ bool IsLibraryLoaded(const char* libraryToMatch)
 // UbaTestSession.h
 int main()
 {
-	bool foundDetoursLib = IsLibraryLoaded(UBA_DETOURS_LIBRARY);
+	bool runningRemote = false;
+	char* tmp = getenv("UBA_REMOTE");
+	if (tmp && tmp[0] == '1')
+		runningRemote = true;
+
+	// Make the assumption that if we're running remote the detour lib will be there.
+	bool foundDetoursLib = IsLibraryLoaded(UBA_DETOURS_LIBRARY) || runningRemote;
+
 
 	if (!foundDetoursLib)
-		return LogError("libUbaDetours.so not loaded. This app is designed to only start from inside UnrealBuildAccelerator.");
+		return LogError("libUbaDetours not loaded. This app is designed to only start from inside UnrealBuildAccelerator.");
 
 	char cwd[1024];
 	if (!getcwd(cwd, sizeof(cwd)))
