UBA: Add Mac Menubar App for UbaAgent

diff --git a/Engine/Source/Programs/UnrealBuildAccelerator/UbaAgentHelperMac/UbaAgentHelper/Assets.xcassets/AccentColor.colorset/Contents.json b/Engine/Source/Programs/UnrealBuildAccelerator/UbaAgentHelperMac/UbaAgentHelper/Assets.xcassets/AccentColor.colorset/Contents.json
new file mode 100644
--- /dev/null
+++ b/Engine/Source/Programs/UnrealBuildAccelerator/UbaAgentHelperMac/UbaAgentHelper/Assets.xcassets/AccentColor.colorset/Contents.json
@@ -0,0 +1,11 @@
+{
+  "colors" : [
+    {
+      "idiom" : "universal"
+    }
+  ],
+  "info" : {
+    "author" : "xcode",
+    "version" : 1
+  }
+}
diff --git a/Engine/Source/Programs/UnrealBuildAccelerator/UbaAgentHelperMac/UbaAgentHelper/Assets.xcassets/AppIcon.appiconset/Contents.json b/Engine/Source/Programs/UnrealBuildAccelerator/UbaAgentHelperMac/UbaAgentHelper/Assets.xcassets/AppIcon.appiconset/Contents.json
new file mode 100644
--- /dev/null
+++ b/Engine/Source/Programs/UnrealBuildAccelerator/UbaAgentHelperMac/UbaAgentHelper/Assets.xcassets/AppIcon.appiconset/Contents.json
@@ -0,0 +1,58 @@
+{
+  "images" : [
+    {
+      "idiom" : "mac",
+      "scale" : "1x",
+      "size" : "16x16"
+    },
+    {
+      "idiom" : "mac",
+      "scale" : "2x",
+      "size" : "16x16"
+    },
+    {
+      "idiom" : "mac",
+      "scale" : "1x",
+      "size" : "32x32"
+    },
+    {
+      "idiom" : "mac",
+      "scale" : "2x",
+      "size" : "32x32"
+    },
+    {
+      "idiom" : "mac",
+      "scale" : "1x",
+      "size" : "128x128"
+    },
+    {
+      "idiom" : "mac",
+      "scale" : "2x",
+      "size" : "128x128"
+    },
+    {
+      "idiom" : "mac",
+      "scale" : "1x",
+      "size" : "256x256"
+    },
+    {
+      "idiom" : "mac",
+      "scale" : "2x",
+      "size" : "256x256"
+    },
+    {
+      "idiom" : "mac",
+      "scale" : "1x",
+      "size" : "512x512"
+    },
+    {
+      "idiom" : "mac",
+      "scale" : "2x",
+      "size" : "512x512"
+    }
+  ],
+  "info" : {
+    "author" : "xcode",
+    "version" : 1
+  }
+}
diff --git a/Engine/Source/Programs/UnrealBuildAccelerator/UbaAgentHelperMac/UbaAgentHelper/Assets.xcassets/Contents.json b/Engine/Source/Programs/UnrealBuildAccelerator/UbaAgentHelperMac/UbaAgentHelper/Assets.xcassets/Contents.json
new file mode 100644
--- /dev/null
+++ b/Engine/Source/Programs/UnrealBuildAccelerator/UbaAgentHelperMac/UbaAgentHelper/Assets.xcassets/Contents.json
@@ -0,0 +1,6 @@
+{
+  "info" : {
+    "author" : "xcode",
+    "version" : 1
+  }
+}
diff --git a/Engine/Source/Programs/UnrealBuildAccelerator/UbaAgentHelperMac/UbaAgentHelper/ContentView.swift b/Engine/Source/Programs/UnrealBuildAccelerator/UbaAgentHelperMac/UbaAgentHelper/ContentView.swift
new file mode 100644
--- /dev/null
+++ b/Engine/Source/Programs/UnrealBuildAccelerator/UbaAgentHelperMac/UbaAgentHelper/ContentView.swift
@@ -0,0 +1,160 @@
+//
+//  ContentView.swift
+//  UbaAgentHelper
+//
+//  Created by Zack Neyland on 1/10/24.
+//
+
+import SwiftUI
+
+// Implement this support: https://developer.apple.com/forums/thread/735862
+
+struct ContentView: View {
+    @State var activeTask : Process? = nil
+    @Binding var status: Bool
+    @State var pathToAgent : String = ""
+    @State var ipAddress = ""
+    @State var log : String = ""
+    @State var xcodePath: String = ""
+    var body: some View {
+        VStack {
+            GroupBox(label: Text("Uba Agent Configuration"), content: {
+                HStack(content: {
+                    Text("Agent Location:")
+                        .fontWeight(.bold)
+                    TextField(text: $pathToAgent, prompt: Text("e.g.: ..../Engine/Binaries/Mac/UnrealBuildAccelerator/UbaAgent"), axis: .vertical, label: {})
+                    Spacer()
+                })
+                HStack {
+                    Text("UbaHost IP:")
+                        .fontWeight(.bold)
+                    TextField(text: $ipAddress, prompt: Text("e.g.: 127.0.0.1"), label: {})
+                    Spacer()
+                }
+                HStack {
+                    Text("XCode Location:")
+                        .fontWeight(.bold)
+                    TextField(text: $xcodePath, prompt: Text("e.g.: /Applications/XCode"), axis: .vertical, label: {})
+                    Spacer()
+                    Button(action: locateXCode, label: {
+                        Text("Locate XCode")
+                    })
+                }
+            })
+            
+            GroupBox(label: Text("UbaAgent Output"), content: {
+                HStack {
+                    Spacer()
+                }
+                ScrollView {
+                    TextField(text: $log,  axis: .vertical, label: {})
+                }
+                
+                .defaultScrollAnchor(.bottom)
+            })
+            Divider()
+            HStack {
+                if let activeTask = activeTask, activeTask.isRunning {
+                    Button(action: killTask, label: {
+                        Text("Kill")
+                    })
+                    .buttonStyle(.borderedProminent)
+                    .tint(.red)
+                }
+                else {
+                    Button(action: {runCode()}, label: {
+                        Text("Start Agent")
+                            .bold()
+                    })
+                    .disabled(canRun())
+                }
+                
+                Spacer()
+                Button(action: { NSApplication.shared.terminate(nil)}, label: {
+                    Text("Quit")
+                })
+                .buttonStyle(.borderedProminent)
+                .tint(.red)
+            }
+        }
+        .padding()
+    }
+    func canRun() -> Bool {
+        return ipAddress.isEmpty || xcodePath.isEmpty || pathToAgent.isEmpty
+    }
+    
+    func locateXCode() {
+        let task = Process()
+
+        //the path to the external program you want to run
+        let executableURL = URL(fileURLWithPath: "/usr/bin/xcode-select")
+        task.executableURL = executableURL
+
+        //use pipe to get the execution program's output
+        let pipe = Pipe()
+        task.standardOutput = pipe
+
+        //all the arguments to the executable
+        let args = ["-p"]
+        task.arguments = args
+
+        try! task.run()
+        task.waitUntilExit()
+
+        //all this code helps you capture the output so you can, for e.g., show the user
+        let d = pipe.fileHandleForReading.readDataToEndOfFile()
+        xcodePath = (String(data: d, encoding: String.Encoding.utf8) ?? "").trimmingCharacters(in: CharacterSet.newlines)
+    }
+    
+    
+    func runCode() {
+        // This should be replaced with the recs from here:
+        // https://developer.apple.com/forums/thread/690310
+        if ((activeTask?.isRunning) != nil) {
+            killTask()
+        }
+        
+        log = ""
+        
+        let task = Process()
+        task.executableURL = URL(fileURLWithPath: pathToAgent)
+        task.arguments = ["-host=\(ipAddress)", "-log", "-populatecas=\(xcodePath)"]
+        let outputPipe = Pipe()
+        task.standardOutput = outputPipe
+        let outputHandle = outputPipe.fileHandleForReading
+        
+        outputHandle.readabilityHandler = { pipe in
+            if let ouput = String(data: pipe.availableData, encoding: .utf8) {
+                if !ouput.isEmpty {
+                    log += " " + ouput
+                }
+            } else {
+                print("Error decoding data: \(pipe.availableData)")
+            }
+        }
+        activeTask = task
+        task.launch()
+        
+        if task.isRunning {
+            status = true
+        }
+        //todo: Swap to this later
+//        Task {
+//            try! task.run()
+//            task.waitUntilExit()
+//        }
+    }
+    
+    func killTask() {
+        if let activeTask = activeTask, activeTask.isRunning {
+            activeTask.terminate()
+            activeTask.waitUntilExit()
+            log += "\n\nKILLED\n\n"
+            status = false
+        }
+    }
+}
+
+#Preview {
+    ContentView(status: .constant(true))
+}
diff --git a/Engine/Source/Programs/UnrealBuildAccelerator/UbaAgentHelperMac/UbaAgentHelper/Preview Content/Preview Assets.xcassets/Contents.json b/Engine/Source/Programs/UnrealBuildAccelerator/UbaAgentHelperMac/UbaAgentHelper/Preview Content/Preview Assets.xcassets/Contents.json
new file mode 100644
--- /dev/null
+++ b/Engine/Source/Programs/UnrealBuildAccelerator/UbaAgentHelperMac/UbaAgentHelper/Preview Content/Preview Assets.xcassets/Contents.json	
@@ -0,0 +1,6 @@
+{
+  "info" : {
+    "author" : "xcode",
+    "version" : 1
+  }
+}
diff --git a/Engine/Source/Programs/UnrealBuildAccelerator/UbaAgentHelperMac/UbaAgentHelper/UbaAgentHelperApp.swift b/Engine/Source/Programs/UnrealBuildAccelerator/UbaAgentHelperMac/UbaAgentHelper/UbaAgentHelperApp.swift
new file mode 100644
--- /dev/null
+++ b/Engine/Source/Programs/UnrealBuildAccelerator/UbaAgentHelperMac/UbaAgentHelper/UbaAgentHelperApp.swift
@@ -0,0 +1,52 @@
+//
+//  UbaAgentHelperApp.swift
+//  UbaAgentHelper
+//
+//  Created by Zack Neyland on 1/10/24.
+//
+
+import SwiftUI
+
+@main
+struct UbaAgentHelperApp: App {
+    @State var isRunning = false
+    var body: some Scene {
+        MenuBarExtra {
+            ContentView(status: $isRunning)
+                .frame(width: 500, height: 500)
+        } label: {
+            HStack {
+                let configuration = NSImage.SymbolConfiguration(pointSize: 16, weight: .light)
+                    .applying(.init(hierarchicalColor: isRunning ? .green : .red))
+                
+                let image = NSImage(systemSymbolName: "bolt.circle.fill", accessibilityDescription: nil)
+                let updateImage = image?.withSymbolConfiguration(configuration)
+                
+                Image(nsImage: updateImage!) // This works.
+                
+                //                    let font = NSFont.systemFont(ofSize: 16, weight: .light)
+                //                let color = NSColor.red // Attempting to make the text red.
+                //                    let attributes: [NSAttributedString.Key: Any] = [
+                //                        .font: font,
+                //                        .foregroundColor: color
+                //                    ]
+                
+                //                    let attributedString = NSAttributedString(string: "Hello, world!", attributes: attributes)
+                
+                //                    Text(AttributedString(attributedString)) // This doesn't work.
+            }
+        }
+        .menuBarExtraStyle(.window)
+        
+        //        MenuBarExtra("UbaAgent Helper")
+        //            ) {
+        //
+        //                }
+                WindowGroup {
+                    ContentView(status: $isRunning)
+                        .navigationTitle("UbaAgent Helper")
+        
+        
+                }
+    }
+}
