[UBA]
* Changed so com memory is not split up in read and write since they never overlap.. so both read and write can use all the memory

diff --git a/Engine/Source/Programs/UnrealBuildAccelerator/Common/Private/UbaProcess.cpp b/Engine/Source/Programs/UnrealBuildAccelerator/Common/Private/UbaProcess.cpp
--- a/Engine/Source/Programs/UnrealBuildAccelerator/Common/Private/UbaProcess.cpp
+++ b/Engine/Source/Programs/UnrealBuildAccelerator/Common/Private/UbaProcess.cpp
@@ -329,7 +329,7 @@ namespace uba
 		#endif
 
 		u8* writeMemory = comMemory;
-		u8* readMemory = comMemory + CommunicationMemSize / 2;
+		u8* readMemory = comMemory;// +CommunicationMemSize / 2;
 
 		u32 retryCount = 0; // Do not allow retry
 
@@ -346,7 +346,7 @@ namespace uba
 				u64 startTime = GetTime();
 				BinaryReader reader(readMemory);
 				BinaryWriter writer(writeMemory);
-				loop = HandleMessage(reader, writer, readMemory);
+				loop = HandleMessage(reader, writer);
 				SetWritten();
 				m_processStats.hostTotalTime += GetTime() - startTime;
 				++m_messageCount;
@@ -442,7 +442,7 @@ namespace uba
 		}
 	}
 
-	bool ProcessImpl::HandleMessage(BinaryReader& reader, BinaryWriter& writer, void* readStream)
+	bool ProcessImpl::HandleMessage(BinaryReader& reader, BinaryWriter& writer)
 	{
 		MessageType messageType = (MessageType)reader.ReadByte();
 		switch (messageType)
diff --git a/Engine/Source/Programs/UnrealBuildAccelerator/Common/Public/UbaProcess.h b/Engine/Source/Programs/UnrealBuildAccelerator/Common/Public/UbaProcess.h
--- a/Engine/Source/Programs/UnrealBuildAccelerator/Common/Public/UbaProcess.h
+++ b/Engine/Source/Programs/UnrealBuildAccelerator/Common/Public/UbaProcess.h
@@ -62,7 +62,7 @@ namespace uba
 		void SetWritten();
 
 		void ThreadRun(bool runningRemote, void* environment);
-		bool HandleMessage(BinaryReader& reader, BinaryWriter& writer, void* readStream);
+		bool HandleMessage(BinaryReader& reader, BinaryWriter& writer);
 		void LogLine(bool printInSession, TString&& line, LogEntryType logType);
 		bool CreateTempFile(BinaryReader& reader, ProcHandle nativeProcessHandle, const tchar* application);
 		bool OpenTempFile(BinaryReader& reader, BinaryWriter& writer, const tchar* application);
diff --git a/Engine/Source/Programs/UnrealBuildAccelerator/Detours/Private/PosixOS/UbaDetoursMainPosix.cpp b/Engine/Source/Programs/UnrealBuildAccelerator/Detours/Private/PosixOS/UbaDetoursMainPosix.cpp
--- a/Engine/Source/Programs/UnrealBuildAccelerator/Detours/Private/PosixOS/UbaDetoursMainPosix.cpp
+++ b/Engine/Source/Programs/UnrealBuildAccelerator/Detours/Private/PosixOS/UbaDetoursMainPosix.cpp
@@ -109,7 +109,7 @@ static void __attribute__((constructor(102))) PreInitCtor()
 	u8* comMem = rptr + sizeof(Event) * 3;
 
 	g_readMessageMappingMem = comMem;
-	g_writeMessageMappingMem = comMem + bytesToMap / 2;
+	g_writeMessageMappingMem = comMem;// +bytesToMap / 2;
 
 	#if 0 // TODO: For some reason this does not work
 	struct sigaction sa;
@@ -165,7 +165,7 @@ namespace uba
 	{
 		m_begin = g_writeMessageMappingMem;
 		m_pos = m_begin;
-		m_end = m_begin + CommunicationMemSize/2;
+		m_end = m_begin + CommunicationMemSize;// / 2;
 	}
 
 	void BinaryWriter::Flush(bool waitOnResponse)
@@ -194,6 +194,6 @@ namespace uba
 	{
 		m_begin = g_readMessageMappingMem;
 		m_pos = m_begin;
-		m_end = m_begin + CommunicationMemSize/2;
+		m_end = m_begin + CommunicationMemSize;// / 2;
 	}
 }
diff --git a/Engine/Source/Programs/UnrealBuildAccelerator/Detours/Private/Windows/UbaDetoursMainWin.cpp b/Engine/Source/Programs/UnrealBuildAccelerator/Detours/Private/Windows/UbaDetoursMainWin.cpp
--- a/Engine/Source/Programs/UnrealBuildAccelerator/Detours/Private/Windows/UbaDetoursMainWin.cpp
+++ b/Engine/Source/Programs/UnrealBuildAccelerator/Detours/Private/Windows/UbaDetoursMainWin.cpp
@@ -49,7 +49,7 @@ namespace uba
 
 		u8* mem = (u8*)::MapViewOfFile(g_messageMappingHandle, FILE_MAP_READ | FILE_MAP_WRITE, ToHigh(payload.communicationOffset), ToLow(payload.communicationOffset), CommunicationMemSize);
 		g_readMessageMappingMem = mem;
-		g_writeMessageMappingMem = mem + CommunicationMemSize / 2;
+		g_writeMessageMappingMem = mem;// +CommunicationMemSize / 2;
 
 		return 0;
 	}
@@ -78,7 +78,7 @@ namespace uba
 	{
 		m_begin = g_writeMessageMappingMem;
 		m_pos = m_begin;
-		m_end = m_begin + CommunicationMemSize / 2;
+		m_end = m_begin + CommunicationMemSize;// / 2;
 	}
 
 	void BinaryWriter::Flush(bool waitOnResponse)
@@ -122,7 +122,7 @@ namespace uba
 	{
 		m_begin = g_readMessageMappingMem;
 		m_pos = m_begin;
-		m_end = m_begin + CommunicationMemSize / 2;
+		m_end = m_begin + CommunicationMemSize;// / 2;
 	}
 }
 
diff --git a/Engine/Source/Programs/UnrealBuildAccelerator/Test/Private/UbaTestSession.h b/Engine/Source/Programs/UnrealBuildAccelerator/Test/Private/UbaTestSession.h
--- a/Engine/Source/Programs/UnrealBuildAccelerator/Test/Private/UbaTestSession.h
+++ b/Engine/Source/Programs/UnrealBuildAccelerator/Test/Private/UbaTestSession.h
@@ -150,6 +150,7 @@ namespace uba
 		return true;
 	}
 
+#if PLATFORM_MAC
 	bool ExecuteCommand(LoggerWithWriter& logger, const tchar* command, StringBufferBase& commandOutput)
 	{
 		FILE* fpCommand = popen(command, "r");
@@ -167,7 +168,8 @@ namespace uba
 		}
 		return true;
 	}
-	
+#endif
+
 	bool RunClang(LoggerWithWriter& logger, SessionServer& session, const tchar* workingDir, const RunProcessFunction& runProcess)
 	{
 		StringBuffer<> sourceFile;
