[Uba]
* Added more info to test error output

diff --git a/Engine/Source/Programs/UnrealBuildAccelerator/Test/Private/UbaTestStorage.h b/Engine/Source/Programs/UnrealBuildAccelerator/Test/Private/UbaTestStorage.h
--- a/Engine/Source/Programs/UnrealBuildAccelerator/Test/Private/UbaTestStorage.h
+++ b/Engine/Source/Programs/UnrealBuildAccelerator/Test/Private/UbaTestStorage.h
@@ -50,7 +50,7 @@ namespace uba
 		if (!GetFileSizeEx(copySize, copy))
 			return logger.Error(TC("Failed to get size of %s"), detoursLibCopy.data);
 		if (originalSize != copySize)
-			return logger.Error(TC("Sizes mismatch between %s and %s"), detoursLib.data, detoursLibCopy.data);
+			return logger.Error(TC("Sizes mismatch between %s and %s (%llu vs %llu)"), detoursLib.data, detoursLibCopy.data, originalSize, copySize);
 
 		u8 originalBuffer[64 * 1024];
 		u8 copyBuffer[64 * 1024];
