[UBA]
* Improved error description when failing to create pipe on posix

diff --git a/Engine/Source/Programs/UnrealBuildAccelerator/Common/Private/UbaProcess.cpp b/Engine/Source/Programs/UnrealBuildAccelerator/Common/Private/UbaProcess.cpp
--- a/Engine/Source/Programs/UnrealBuildAccelerator/Common/Private/UbaProcess.cpp
+++ b/Engine/Source/Programs/UnrealBuildAccelerator/Common/Private/UbaProcess.cpp
@@ -1475,7 +1475,7 @@ namespace uba
 			{
 				if (pipe(outPipe) || pipe(errPipe))
 				{
-					logger.Error("pipe failed");
+					logger.Error("pipe failed (%s)", strerror(errno));
 					return UBA_EXIT_CODE(18);
 				}
 
