[UBA]
* Added message on all asserts in posix detouring code since we can't rely on line number in assert on macos

diff --git a/Engine/Source/Programs/UnrealBuildAccelerator/Detours/Private/PosixOS/UbaDetoursFunctionsPosix.cpp b/Engine/Source/Programs/UnrealBuildAccelerator/Detours/Private/PosixOS/UbaDetoursFunctionsPosix.cpp
--- a/Engine/Source/Programs/UnrealBuildAccelerator/Detours/Private/PosixOS/UbaDetoursFunctionsPosix.cpp
+++ b/Engine/Source/Programs/UnrealBuildAccelerator/Detours/Private/PosixOS/UbaDetoursFunctionsPosix.cpp
@@ -397,7 +397,7 @@ int Shared_open(const char* funcName, const char* file, int flags, int mode, con
 			int fd = TRUE_WRAPPER(open)("/dev/null", O_RDONLY);
 			ScopedWriteLock lock(g_fileHandlesLock);
 			auto insres2 = g_fileHandles.insert({ fd, DetouredHandle() });
-			UBA_ASSERT(insres2.second);
+			UBA_ASSERTF(insres2.second, "File handle for directory already added");
 			DetouredHandle& h = insres2.first->second;
 			auto fo = new FileObject();
 			fo->closeId = closeId;
@@ -408,21 +408,21 @@ int Shared_open(const char* funcName, const char* file, int flags, int mode, con
 		}
 
 		DEBUG_LOG_DETOURED(funcName, "FAILED %s (%s)", fileName.data, realFileName);
-		UBA_ASSERT(false);
+		UBA_ASSERTF(false, "unsupported filename %s", realFileName);
 		return -1;
 	}
 
 	if (realFileName[0] == '^')
 	{
 		DEBUG_LOG_DETOURED(funcName, "FAILED %s (%s)", fileName.data, realFileName);
-		UBA_ASSERT(false);
+		UBA_ASSERTF(false, "^ filenames not implemented");
 		return -1;
 	}
 
 	if (keepInMemory)// || info.memoryFile)
 	{
 		DEBUG_LOG_DETOURED(funcName, "FAILED %s (%s)", fileName.data, realFileName);
-		UBA_ASSERT(false);
+		UBA_ASSERTF(false, "keepInMemory not implemented");
 		return -1;
 	}
 
@@ -440,7 +440,7 @@ int Shared_open(const char* funcName, const char* file, int flags, int mode, con
 
 	ScopedWriteLock lock(g_fileHandlesLock);
 	auto insres2 = g_fileHandles.insert({ fd, DetouredHandle() });
-	UBA_ASSERT(insres2.second);
+	UBA_ASSERTF(insres2.second, "File handle already added");
 	DetouredHandle& h = insres2.first->second;
 	auto fo = new FileObject();
 	fo->closeId = closeId;
@@ -507,7 +507,7 @@ void Shared_close(int fd, const TrueClose& trueClose)
 
 	DetouredHandle& h = findIt->second;
 	FileObject* fo = h.fileObject;
-	UBA_ASSERT(fo->refCount >= 1);
+	UBA_ASSERTF(fo->refCount >= 1, "FileObject needs to have ref count when closed");
 	g_fileHandles.erase(findIt);
 	lock.Leave();
 
@@ -583,16 +583,16 @@ int Shared_fstat(const char* funcName, int fd, struct stat* attr, const True_fst
 		if (res != -1)
 		{
 			bool isDir = S_ISDIR(attr->st_mode);
-			UBA_ASSERT(isDir == S_ISDIR(attr2.st_mode));
+			UBA_ASSERTF(isDir == S_ISDIR(attr2.st_mode), "fstat: isDir not matching");
 			//UBA_ASSERT(attr->st_mode == attr2.st_mode);
 			//UBA_ASSERT(attr->st_dev == attr2.st_dev)
 			UBA_ASSERTF(attr->st_ino == attr2.st_ino, "fstat: st_ino mismatch for %s (%llu vs %llu)", fi.originalName, attr->st_ino, attr2.st_ino);
-			UBA_ASSERT(isDir || attr->st_size == attr2.st_size);
+			UBA_ASSERTF(isDir || attr->st_size == attr2.st_size, "fstat: size not matching");
 			UBA_ASSERTF(isDir || FromTimeSpec(attr->st_mtimespec) == FromTimeSpec(attr2.st_mtimespec), "fstat: st_mtim mismatch for %s (%llu vs %llu)", fi.originalName, FromTimeSpec(attr->st_mtimespec), FromTimeSpec(attr2.st_mtimespec));
 		}
 		else
 		{
-			UBA_ASSERT(fileAttr.lastError == errno);
+			UBA_ASSERTF(fileAttr.lastError == errno, "fstat: error not matching");
 		}
 	}
 	#endif
@@ -699,28 +699,28 @@ UBA_EXPORT int UBA_WRAPPER(_NSGetExecutablePath)(char* buf, uint32_t* bufsize)
 UBA_EXPORT int UBA_WRAPPER(chdir)(const char* path)
 {
 	UBA_INIT_DETOUR(chdir, path);
-	UBA_ASSERT(false);
+	UBA_ASSERTF(false, "chdir not implemented");
 	return TRUE_WRAPPER(chdir)(path);
 }
 
 UBA_EXPORT int UBA_WRAPPER(fchdir)(int fd)
 {
 	UBA_INIT_DETOUR(fchdir, fd);
-	UBA_ASSERT(false);
+	UBA_ASSERTF(false, "fchdir not implemented");
 	return TRUE_WRAPPER(fchdir)(fd);
 }
 
 UBA_EXPORT int UBA_WRAPPER(chroot)(const char* path)
 {
 	UBA_INIT_DETOUR(chroot, path);
-	UBA_ASSERT(false);
+	UBA_ASSERTF(false, "chroot not implemented");
 	return TRUE_WRAPPER(chroot)(path);
 }
 
 UBA_EXPORT char* UBA_WRAPPER(getcwd)(char* buf, size_t size)
 {
 	UBA_INIT_DETOUR(getcwd, buf, size);
-	UBA_ASSERT(g_virtualWorkingDir.count < size);
+	UBA_ASSERTF(g_virtualWorkingDir.count < size, "getcwd with size smaller than path not implemented");
 	memcpy(buf, g_virtualWorkingDir.data, g_virtualWorkingDir.count+1);
 	DEBUG_LOG_DETOURED("getcwd", "%s -> %p", buf, buf);
 	return buf;
@@ -778,7 +778,7 @@ UBA_EXPORT ssize_t UBA_WRAPPER(readlink)(const char* pathname, char* buf, size_t
 
 	if (Equals(pathname, "/proc/self/exe"))
 	{
-		UBA_ASSERT(g_virtualApplication.count < bufsiz);
+		UBA_ASSERTF(g_virtualApplication.count < bufsiz, "readLink: buffer size smaller than path not implemented");
 		memcpy(buf, g_virtualApplication.data, g_virtualApplication.count + 1);
 		DEBUG_LOG_DETOURED("readlink", "(%s) (%s) -> %u", pathname, buf, g_virtualApplication.count);
 		return g_virtualApplication.count;
@@ -789,7 +789,7 @@ UBA_EXPORT ssize_t UBA_WRAPPER(readlink)(const char* pathname, char* buf, size_t
 		fdStr.Append(pathname + 14);
 		u32 fd;
 		if (!fdStr.Parse(fd))
-			UBA_ASSERT(false);
+			UBA_ASSERTF(false, "Failed to parse /proc/self/fd");
 		ScopedReadLock lock(g_fileHandlesLock);
 		auto findIt = g_fileHandles.find(fd);
 		if (findIt != g_fileHandles.end())
@@ -798,7 +798,7 @@ UBA_EXPORT ssize_t UBA_WRAPPER(readlink)(const char* pathname, char* buf, size_t
 			FileObject* fo = h.fileObject;
 			FileInfo& info = *fo->fileInfo;
 			u32 len = TStrlen(info.originalName);
-			UBA_ASSERT(len < bufsiz);
+			UBA_ASSERTF(len < bufsiz, "buffer size is smaller than length of name");
 			memcpy(buf, info.originalName, len + 1);
 			DEBUG_LOG_DETOURED("readlink", "(%s) (%s) -> %u", pathname, buf, len);
 			return len;
@@ -813,7 +813,7 @@ ssize_t UBA_WRAPPER(readlinkat)(int dirfd, const char* pathname, char* buf, size
 {
 	UBA_INIT_DETOUR(readlinkat, dirfd, pathname, buf, bufsiz);
 	DEBUG_LOG_TRUE("readlinkat","(%s)", pathname);
-	UBA_ASSERT(false);
+	UBA_ASSERTF(false, "readlinkat not implemented");
 	return TRUE_WRAPPER(readlinkat)(dirfd, pathname, buf, bufsiz);
 }
 
@@ -1022,7 +1022,7 @@ UBA_EXPORT dirent* UBA_WRAPPER(readdir)(DIR* dirp)
 UBA_EXPORT void UBA_WRAPPER(rewinddir)(DIR* dirp)
 {
 	UBA_INIT_DETOUR(rewinddir, dirp);
-	UBA_ASSERT(false);
+	UBA_ASSERTF(false, "rewinddir");
 	DEBUG_LOG_TRUE("rewinddir", "(%p)", dirp);
 	return TRUE_WRAPPER(rewinddir)(dirp);
 }
@@ -1030,14 +1030,14 @@ UBA_EXPORT void UBA_WRAPPER(rewinddir)(DIR* dirp)
 UBA_EXPORT int UBA_WRAPPER(scandir)(const char* dirp, dirent*** namelist, int (*filter)(const dirent*), int (*compar)(const dirent**, const dirent**))
 {
 	UBA_INIT_DETOUR(scandir, dirp, namelist, filter, compar);
-	UBA_ASSERT(false);
+	UBA_ASSERTF(false, "scandir");
 	DEBUG_LOG_TRUE("scandir", "(%p)", dirp);
 	return TRUE_WRAPPER(scandir)(dirp, namelist, filter, compar);
 }
 UBA_EXPORT void UBA_WRAPPER(seekdir)(DIR* dirp, long loc)
 {
 	UBA_INIT_DETOUR(seekdir, dirp, loc);
-	UBA_ASSERT(false);
+	UBA_ASSERTF(false, "seekdir");
 	DEBUG_LOG_TRUE("seekdir", "(%p)", dirp);
 	return TRUE_WRAPPER(seekdir)(dirp, loc);
 }
@@ -1045,7 +1045,7 @@ UBA_EXPORT void UBA_WRAPPER(seekdir)(DIR* dirp, long loc)
 UBA_EXPORT long UBA_WRAPPER(telldir)(DIR* dirp)
 {
 	UBA_INIT_DETOUR(telldir, dirp);
-	UBA_ASSERT(false);
+	UBA_ASSERTF(false, "telldir");
 	DEBUG_LOG_TRUE("telldir", "(%p)", dirp);
 	return TRUE_WRAPPER(telldir)(dirp);
 }
@@ -1053,7 +1053,7 @@ UBA_EXPORT long UBA_WRAPPER(telldir)(DIR* dirp)
 UBA_EXPORT DIR* UBA_WRAPPER(fdopendir)(int fd)
 {
 	UBA_INIT_DETOUR(fdopendir, fd);
-	UBA_ASSERT(false);
+	UBA_ASSERTF(false, "fdopendir");
 	DEBUG_LOG_TRUE("fdopendir", "(%i)", fd);
 	return TRUE_WRAPPER(fdopendir)(fd);
 }
@@ -1099,7 +1099,7 @@ UBA_EXPORT int UBA_WRAPPER(stat)(const char* file, struct stat* attr)
 UBA_EXPORT int UBA_WRAPPER(truncate)(const char* path, off_t length)
 {
 	UBA_INIT_DETOUR(truncate, path, length);
-	UBA_ASSERT(false); // TODO: Implement this if it is ever called
+	UBA_ASSERTF(false, "truncate"); // TODO: Implement this if it is ever called
 	return TRUE_WRAPPER(truncate)(path, length);
 }
 
@@ -1214,7 +1214,7 @@ UBA_EXPORT int UBA_WRAPPER(rename)(const char* oldpath, const char* newpath)
 			FileInfo& fi = *fo.fileInfo;
 			if (fi.fileNameKey == oldKey)
 			{
-				UBA_ASSERT(fo.desiredAccess & AccessFlag_Write);
+				UBA_ASSERTF(fo.desiredAccess & AccessFlag_Write, "Unsupported access flags");
 				fo.newName = fixedNewPath.data;
 				if (!fo.closeId)
 				{
@@ -1232,7 +1232,7 @@ UBA_EXPORT int UBA_WRAPPER(rename)(const char* oldpath, const char* newpath)
 					return 0;
 				}
 				bool isTempFile = fixedOldPath.StartsWith(g_systemTemp.data);
-				UBA_ASSERT(wasTempFile == isTempFile);
+				UBA_ASSERTF(wasTempFile == isTempFile, "File changing from temp to not or vice versa not implemented");
 
 				int res = TRUE_WRAPPER(rename)(oldpath, newpath);
 				DEBUG_LOG_DETOURED("rename", "IS_OPEN (%i) (from %s to %s) -> %i (%s)", kv.first, fixedOldPath.data, fixedNewPath.data, 0, StrError(0, errno));
@@ -1451,7 +1451,7 @@ UBA_EXPORT int UBA_WRAPPER(posix_spawn)(pid_t* pid, const char* path, const posi
 
 		BinaryReader reader;
 		processId = reader.ReadU32();
-		UBA_ASSERT(processId > 0);
+		UBA_ASSERTF(processId > 0, "Process id was zero");
 
 		rulesStr.Append("UBA_RULES=").AppendValue(reader.ReadU32());
 
@@ -1549,7 +1549,7 @@ UBA_EXPORT int UBA_WRAPPER(waitid)(idtype_t idtype, id_t id, siginfo_t* infop, i
 {
 	UBA_INIT_DETOUR(waitid, idtype, id, infop, options);
 	DEBUG_LOG_TRUE("waitid", "");
-	UBA_ASSERT(!t_inVfork);
+	UBA_ASSERTF(!t_inVfork, "waitid: is in fork");
 	// TODO: Should probably report id to session
 	return TRUE_WRAPPER(waitid)(idtype, id, infop, options);
 }
@@ -1558,16 +1558,16 @@ UBA_EXPORT pid_t UBA_WRAPPER(wait3)(int* status, int options, struct rusage* rus
 {
 	UBA_INIT_DETOUR(wait3, status, options, rusage);
 	DEBUG_LOG_TRUE("wait3", "");
-	UBA_ASSERT(!t_inVfork);
+	UBA_ASSERTF(!t_inVfork, "wait3: is in fork");
 	pid_t res =TRUE_WRAPPER(wait3)(status, options, rusage);
-	UBA_ASSERT(false);
+	UBA_ASSERTF(false, "wait3");
 	return res;
 }
 
 UBA_EXPORT pid_t UBA_WRAPPER(wait4)(pid_t pid, int* status, int options, struct rusage* rusage)
 {
 	UBA_INIT_DETOUR(wait4, pid, status, options, rusage);
-	UBA_ASSERT(!t_inVfork);
+	UBA_ASSERTF(!t_inVfork, "wait4: is in fork");
 	pid_t res = TRUE_WRAPPER(wait4)(pid, status, options, rusage);
 	DEBUG_LOG_TRUE("wait4", "");
 	return res;
@@ -1602,8 +1602,8 @@ int Internal_execve(const char* pathname, char* const _Nullable argv[], char* co
 	{
 		res = TRUE_WRAPPER(waitpid)(pid, &status, WUNTRACED | WCONTINUED);
 		DEBUG_LOG_TRUE("waitpid", "(%i) -> %i", pid, res);
-		UBA_ASSERT(res == pid);
-		UBA_ASSERT(WIFEXITED(status));
+		UBA_ASSERTF(res == pid, "execve: wait result was not same as pid");
+		UBA_ASSERTF(WIFEXITED(status), "execve: Unsupported status from waitpid");
 	}
 
 	{
@@ -1745,7 +1745,7 @@ UBA_EXPORT int openat64(int __fd, const char* __file, int __oflag, mode_t mode,
 
 UBA_EXPORT char* UBA_WRAPPER(get_current_dir_name)(void)
 {
-	UBA_ASSERT(false);
+	UBA_ASSERTF(false, "get_current_dir_name");
 	return TRUE_WRAPPER(get_current_dir_name)();
 }
 
@@ -1844,9 +1844,9 @@ namespace uba
 
 		StringBuffer<> exePath;
 		exePath.count = GetProcessExecutablePath(exePath.data, exePath.capacity);
-		UBA_ASSERT(exePath.count > 0);
+		UBA_ASSERTF(exePath.count > 0, "exePath.count == 0");
 		char* lastSlash = strrchr(exePath.data, '/');
-		UBA_ASSERT(lastSlash);
+		UBA_ASSERTF(lastSlash, "no slash found in %s", exePath.data);
 		exePath.Resize(lastSlash - exePath.data);
 		FixPath(g_exeDir, exePath.data);
 		g_exeDir.EnsureEndsWithSlash();
@@ -1854,7 +1854,7 @@ namespace uba
 
 	void Init()
 	{
-		UBA_ASSERT(!g_isInitialized);
+		UBA_ASSERTF(!g_isInitialized, "Already initialized");
 		g_isInitialized = true;
 
 		u64 directoryTableHandle;
@@ -1886,7 +1886,7 @@ namespace uba
 			DEBUG_LOG_PIPE(L"Init", L"");
 		}
 
-		UBA_ASSERT(g_virtualApplicationDir.capacity > 0);
+		UBA_ASSERTF(g_virtualApplicationDir.capacity > 0, "g_virtualApplicationDir.capacity > 0");
 
 		const char* lastSlash = strrchr(g_virtualApplication.data, '/');
 		UBA_ASSERTF(lastSlash, "Need fullpath for application (%s)", g_virtualApplication.data);
@@ -1898,17 +1898,17 @@ namespace uba
 		StringBuffer<128> mappedFileTableUidName;
 		GetMappingHandleName(mappedFileTableUidName, mappedFileTableHandle);
 		int mappedFileTableFd = shm_open(mappedFileTableUidName.data, O_RDONLY, S_IRUSR | S_IWUSR);
-		UBA_ASSERT(mappedFileTableFd != -1);
+		UBA_ASSERTF(mappedFileTableFd != -1, "mappedFileTableFd == %i", mappedFileTableFd);
 		u8* mappedFileTableMem = (u8*)mmap(NULL, FileMappingTableMemSize, PROT_READ, MAP_SHARED, mappedFileTableFd, 0);
-		UBA_ASSERT(mappedFileTableMem != MAP_FAILED);
+		UBA_ASSERTF(mappedFileTableMem != MAP_FAILED, "mmap failed (%s)", strerror(errno));
 		g_mappedFileTable.Init(mappedFileTableMem, mappedFileTableCount, mappedFileTableSize);
 
 		StringBuffer<128> dirTableUidName;
 		GetMappingHandleName(dirTableUidName, directoryTableHandle);
 		int dirTableFd = shm_open(dirTableUidName.data, O_RDONLY, S_IRUSR | S_IWUSR);
-		UBA_ASSERT(dirTableFd != -1);
+		UBA_ASSERTF(dirTableFd != -1, "shm_open failed (%s)", strerror(errno));
 		u8* dirTableMem = (u8*)mmap(NULL, DirTableMemSize, PROT_READ, MAP_SHARED, dirTableFd, 0);
-		UBA_ASSERT(dirTableMem != MAP_FAILED);
+		UBA_ASSERTF(dirTableMem != MAP_FAILED, "mmap for dirtable mem failed (%s)", strerror(errno));
 		g_directoryTable.Init(dirTableMem, directoryTableCount, directoryTableSize);
 
 		g_isDetouring = true;
