[Uba]
* Changed so detoured truncate only asserts for remote build

diff --git a/Engine/Source/Programs/UnrealBuildAccelerator/Detours/Private/PosixOS/UbaDetoursFunctionsPosix.cpp b/Engine/Source/Programs/UnrealBuildAccelerator/Detours/Private/PosixOS/UbaDetoursFunctionsPosix.cpp
--- a/Engine/Source/Programs/UnrealBuildAccelerator/Detours/Private/PosixOS/UbaDetoursFunctionsPosix.cpp
+++ b/Engine/Source/Programs/UnrealBuildAccelerator/Detours/Private/PosixOS/UbaDetoursFunctionsPosix.cpp
@@ -1099,7 +1099,7 @@ UBA_EXPORT int UBA_WRAPPER(stat)(const char* file, struct stat* attr)
 UBA_EXPORT int UBA_WRAPPER(truncate)(const char* path, off_t length)
 {
 	UBA_INIT_DETOUR(truncate, path, length);
-	UBA_ASSERTF(false, "truncate"); // TODO: Implement this if it is ever called
+	UBA_ASSERTF(!g_runningRemote, "truncate not implemented for remote execution (path: %s)", path); // TODO: Implement this if it is ever called
 	return TRUE_WRAPPER(truncate)(path, length);
 }
 
