[Uba]
* Fixed linux unit test errors

diff --git a/Engine/Source/Programs/UnrealBuildAccelerator/Common/Private/UbaFileMapping.cpp b/Engine/Source/Programs/UnrealBuildAccelerator/Common/Private/UbaFileMapping.cpp
--- a/Engine/Source/Programs/UnrealBuildAccelerator/Common/Private/UbaFileMapping.cpp
+++ b/Engine/Source/Programs/UnrealBuildAccelerator/Common/Private/UbaFileMapping.cpp
@@ -171,9 +171,19 @@ namespace uba
 		return { InternalCreateFileMappingW(asHANDLE(file), flProtect, (DWORD)ToHigh(maxSize), ToLow(maxSize), NULL) };
 #else
 		FileMappingHandle h;
-		h.shmFd = asFileDescriptor(file);
-		lseek(h.shmFd, maxSize-1, SEEK_SET);
-		write(h.shmFd, "", 1);
+		int fd = asFileDescriptor(file);
+		if (lseek(fd, maxSize - 1, SEEK_SET) == -1)
+		{
+			UBA_ASSERTF(false, "lseek to %llu failed: %s\n", maxSize - 1, strerror(errno));
+			return h;
+		}
+		if (write(fd, "", 1) != 1)
+		{
+			UBA_ASSERTF(false, "write one byte failed: %s\n", strerror(errno));
+			return h;
+		}
+
+		h.shmFd = fd;
 		return h;
 #endif
 	}
diff --git a/Engine/Source/Programs/UnrealBuildAccelerator/Test/Private/UbaTestStorage.h b/Engine/Source/Programs/UnrealBuildAccelerator/Test/Private/UbaTestStorage.h
--- a/Engine/Source/Programs/UnrealBuildAccelerator/Test/Private/UbaTestStorage.h
+++ b/Engine/Source/Programs/UnrealBuildAccelerator/Test/Private/UbaTestStorage.h
@@ -50,7 +50,7 @@ namespace uba
 		if (!GetFileSizeEx(copySize, copy))
 			return logger.Error(TC("Failed to get size of %s"), detoursLibCopy.data);
 		if (originalSize != copySize)
-			return logger.Error(TC("Sizes mismatch between %s and %s (%llu vs %llu)"), detoursLib.data, detoursLibCopy.data, originalSize, copySize);
+			return logger.Error(TC("Size mismatch between %s and %s (%llu vs %llu)"), detoursLib.data, detoursLibCopy.data, originalSize, copySize);
 
 		u8 originalBuffer[64 * 1024];
 		u8 copyBuffer[64 * 1024];
