[UBA]
* Fixed so if setpriority fails on linux because of not having permissions we don't fail the build

diff --git a/Engine/Source/Programs/UnrealBuildAccelerator/Common/Private/UbaProcess.cpp b/Engine/Source/Programs/UnrealBuildAccelerator/Common/Private/UbaProcess.cpp
--- a/Engine/Source/Programs/UnrealBuildAccelerator/Common/Private/UbaProcess.cpp
+++ b/Engine/Source/Programs/UnrealBuildAccelerator/Common/Private/UbaProcess.cpp
@@ -1514,10 +1514,16 @@ namespace uba
 				return UBA_EXIT_CODE(12);
 			}
 			
-			int prio = getpriority(PRIO_PROCESS, processID);
 			errno = 0;
-			res = setpriority(PRIO_PROCESS, processID, prio + 2);
-			UBA_ASSERTF(res == 0, TC("setpriority (%s)"), strerror(errno));
+			int prio = getpriority(PRIO_PROCESS, processID);
+			if (prio != -1)
+			{
+				errno = 0;
+				if (setpriority(PRIO_PROCESS, processID, prio + 2) == -1)
+				{
+					UBA_ASSERTF(errno == EPERM, TC("setpriority %i (%s)"), prio + 2, strerror(errno));
+				}
+			}
 
 			m_nativeProcessHandle = (ProcHandle)1;
 			m_nativeProcessId = u32(processID);
