[UBA]
* Improved error handling for setpriority even more. There is a risk that a non-detoured process have exited already when we get to this function

diff --git a/Engine/Source/Programs/UnrealBuildAccelerator/Common/Private/UbaProcess.cpp b/Engine/Source/Programs/UnrealBuildAccelerator/Common/Private/UbaProcess.cpp
--- a/Engine/Source/Programs/UnrealBuildAccelerator/Common/Private/UbaProcess.cpp
+++ b/Engine/Source/Programs/UnrealBuildAccelerator/Common/Private/UbaProcess.cpp
@@ -1521,7 +1521,7 @@ namespace uba
 				errno = 0;
 				if (setpriority(PRIO_PROCESS, processID, prio + 2) == -1)
 				{
-					UBA_ASSERTF(errno == EPERM, TC("setpriority %i (%s)"), prio + 2, strerror(errno));
+					UBA_ASSERTF(errno == ESRCH || errno == EPERM, TC("setpriority failed: %s. pid: %i prio: %i (%s)"), m_realApplication.c_str(), processID, prio + 2, strerror(errno));
 				}
 			}
 
