UBA: Bump cmdline limit to 64k

diff --git a/Engine/Source/Programs/UnrealBuildAccelerator/Common/Private/UbaProcess.cpp b/Engine/Source/Programs/UnrealBuildAccelerator/Common/Private/UbaProcess.cpp
--- a/Engine/Source/Programs/UnrealBuildAccelerator/Common/Private/UbaProcess.cpp
+++ b/Engine/Source/Programs/UnrealBuildAccelerator/Common/Private/UbaProcess.cpp
@@ -605,7 +605,7 @@ namespace uba
 
 					StringBuffer<> application;
 					reader.ReadString(application);
-					StringBuffer<32*1024> fullCommandLine;
+					StringBuffer<64*1024> fullCommandLine;
 					reader.ReadString(fullCommandLine);
 					StringBuffer<> currentDir;
 					reader.ReadString(currentDir);
