[Uba]
* Added retry and better error description when writing byte to create size of file mapping on posix

diff --git a/Engine/Source/Programs/UnrealBuildAccelerator/Common/Private/UbaFileMapping.cpp b/Engine/Source/Programs/UnrealBuildAccelerator/Common/Private/UbaFileMapping.cpp
--- a/Engine/Source/Programs/UnrealBuildAccelerator/Common/Private/UbaFileMapping.cpp
+++ b/Engine/Source/Programs/UnrealBuildAccelerator/Common/Private/UbaFileMapping.cpp
@@ -177,9 +177,14 @@ namespace uba
 			UBA_ASSERTF(false, "lseek to %llu failed: %s\n", maxSize - 1, strerror(errno));
 			return h;
 		}
-		if (write(fd, "", 1) != 1)
+
+		errno = 0;
+		int res = write(fd, "", 1);
+		if (res == 0)
+			res = write(fd, "", 1);
+		if (res != 1)
 		{
-			UBA_ASSERTF(false, "write one byte failed: %s\n", strerror(errno));
+			UBA_ASSERTF(false, "write one byte on fd %i failed (res: %i): %s\n", fd, res, strerror(errno));
 			return h;
 		}
 
