Fixing working dir on Mac

diff --git a/Engine/Source/Programs/UnrealBuildAccelerator/Common/Private/UbaProcess.cpp b/Engine/Source/Programs/UnrealBuildAccelerator/Common/Private/UbaProcess.cpp
--- a/Engine/Source/Programs/UnrealBuildAccelerator/Common/Private/UbaProcess.cpp
+++ b/Engine/Source/Programs/UnrealBuildAccelerator/Common/Private/UbaProcess.cpp
@@ -1385,8 +1385,8 @@ namespace uba
 			res = posix_spawn_file_actions_init(&fileActions);
 			UBA_ASSERTF(res == 0, TC("posix_spawn_file_actions_init (%s)"), strerror(errno));
 			auto actionsGuard = MakeGuard([&]() { posix_spawn_file_actions_destroy(&fileActions); });
-
-			if (!*m_realWorkingDir)
+			
+			if (*m_realWorkingDir)
 			{
 				#if PLATFORM_MAC
 				posix_spawn_file_actions_addchdir_np(&fileActions, m_realWorkingDir);
