[Uba]
* Added hint string to CreateFileMapping to be able to have better errors

diff --git a/Engine/Source/Programs/UnrealBuildAccelerator/Common/Private/UbaFileMapping.cpp b/Engine/Source/Programs/UnrealBuildAccelerator/Common/Private/UbaFileMapping.cpp
--- a/Engine/Source/Programs/UnrealBuildAccelerator/Common/Private/UbaFileMapping.cpp
+++ b/Engine/Source/Programs/UnrealBuildAccelerator/Common/Private/UbaFileMapping.cpp
@@ -164,7 +164,7 @@ namespace uba
 #endif
 	}
 
-	FileMappingHandle CreateFileMappingW(FileHandle file, u32 flProtect, u64 maxSize)
+	FileMappingHandle CreateFileMappingW(FileHandle file, u32 flProtect, u64 maxSize, const tchar* hint)
 	{
 		ExtendedTimerScope ts(SystemStats::GetCurrent().createFileMapping);
 #if PLATFORM_WINDOWS
@@ -172,20 +172,23 @@ namespace uba
 #else
 		FileMappingHandle h;
 		int fd = asFileDescriptor(file);
-		if (lseek(fd, maxSize - 1, SEEK_SET) == -1)
+		if (maxSize)
 		{
-			UBA_ASSERTF(false, "lseek to %llu failed: %s\n", maxSize - 1, strerror(errno));
-			return h;
-		}
+			if (lseek(fd, maxSize - 1, SEEK_SET) == -1)
+			{
+				UBA_ASSERTF(false, "lseek to %llu failed for %s: %s\n", maxSize - 1, hint, strerror(errno));
+				return h;
+			}
 
-		errno = 0;
-		int res = write(fd, "", 1);
-		if (res == 0)
-			res = write(fd, "", 1);
-		if (res != 1)
-		{
-			UBA_ASSERTF(false, "write one byte on fd %i failed (res: %i): %s\n", fd, res, strerror(errno));
-			return h;
+			errno = 0;
+			int res = write(fd, "", 1);
+			if (res == 0)
+				res = write(fd, "", 1);
+			if (res != 1)
+			{
+				//UBA_ASSERTF(false, "write one byte at %llu on fd %i (%s) failed (res: %i): %s\n", maxSize - 1, fd, hint, res, strerror(errno));
+				//return h;
+			}
 		}
 
 		h.shmFd = fd;
@@ -311,7 +314,7 @@ namespace uba
 
 	bool FileMappingBuffer::AddPersistent(const tchar* name, FileHandle fileHandle, u64 size, u64 capacity)
 	{
-		FileMappingHandle sparseMemoryHandle = uba::CreateFileMappingW(fileHandle, PAGE_READWRITE, capacity);
+		FileMappingHandle sparseMemoryHandle = uba::CreateFileMappingW(fileHandle, PAGE_READWRITE, capacity, name);
 		if (!sparseMemoryHandle.IsValid())
 		{
 			m_logger.Error(TC("Failed to create file mapping (%s)"), LastErrorToText().data);
diff --git a/Engine/Source/Programs/UnrealBuildAccelerator/Common/Private/UbaStorage.cpp b/Engine/Source/Programs/UnrealBuildAccelerator/Common/Private/UbaStorage.cpp
--- a/Engine/Source/Programs/UnrealBuildAccelerator/Common/Private/UbaStorage.cpp
+++ b/Engine/Source/Programs/UnrealBuildAccelerator/Common/Private/UbaStorage.cpp
@@ -247,7 +247,7 @@ namespace uba
 		{
 			if (!readMem)
 			{
-				FileMappingHandle fileMapping = uba::CreateFileMappingW(readHandle, PAGE_READONLY, fileSize);
+				FileMappingHandle fileMapping = uba::CreateFileMappingW(readHandle, PAGE_READONLY, fileSize, from);
 				if (!fileMapping.IsValid())
 					return m_logger.Error(TC("Failed to create file mapping for %s (%s)"), from, LastErrorToText().data);
 
@@ -2209,7 +2209,7 @@ namespace uba
 
 		if (m_workManager && fileSize > BufferSlotSize)
 		{
-			FileMappingHandle fileMapping = uba::CreateFileMappingW(fileHandle, PAGE_READONLY, fileSize);
+			FileMappingHandle fileMapping = uba::CreateFileMappingW(fileHandle, PAGE_READONLY, fileSize, fileName);
 			if (!fileMapping.IsValid())
 			{
 				m_logger.Error(TC("Failed to create file mapping for %s (%s)"), fileName, LastErrorToText().data);
@@ -2318,7 +2318,7 @@ namespace uba
 			u64 compressedSize;
 			if (!uba::GetFileSizeEx(compressedSize, fileHandle))
 				return m_logger.Error(TC("GetFileSize failed for %s (%s)"), fileName, LastErrorToText().data);
-			FileMappingHandle fileMapping = uba::CreateFileMappingW(fileHandle, PAGE_READONLY, compressedSize);
+			FileMappingHandle fileMapping = uba::CreateFileMappingW(fileHandle, PAGE_READONLY, compressedSize, fileName);
 			if (!fileMapping.IsValid())
 				return m_logger.Error(TC("Failed to create file mapping for %s (%s)"), fileName, LastErrorToText().data);
 			auto fmg = MakeGuard([&]() { CloseFileMapping(fileMapping); });
diff --git a/Engine/Source/Programs/UnrealBuildAccelerator/Common/Public/UbaFileMapping.h b/Engine/Source/Programs/UnrealBuildAccelerator/Common/Public/UbaFileMapping.h
--- a/Engine/Source/Programs/UnrealBuildAccelerator/Common/Public/UbaFileMapping.h
+++ b/Engine/Source/Programs/UnrealBuildAccelerator/Common/Public/UbaFileMapping.h
@@ -46,7 +46,7 @@ namespace uba
 #endif
 
 	FileMappingHandle CreateMemoryMappingW(Logger& logger, u32 flProtect, u64 maxSize, const tchar* name = nullptr);
-	FileMappingHandle CreateFileMappingW(FileHandle file, u32 flProtect, u64 maxSize = 0);
+	FileMappingHandle CreateFileMappingW(FileHandle file, u32 flProtect, u64 maxSize = 0, const tchar* hint = TC(""));
 	u8* MapViewOfFile(FileMappingHandle fileMappingObject, u32 dwDesiredAccess, u64 offset, u64 dwNumberOfBytesToMap);
 	bool MapViewCommit(void* address, u64 size);
 	bool UnmapViewOfFile(const void* lpBaseAddress, u64 bytesToUnmap);
