[Uba]
* Fixed file descriptor leak when running native processes on posix platforms
* Removed hangup test on poll for pipe file descriptors since it exited to early and missed reads

diff --git a/Engine/Source/Programs/UnrealBuildAccelerator/Common/Private/UbaProcess.cpp b/Engine/Source/Programs/UnrealBuildAccelerator/Common/Private/UbaProcess.cpp
--- a/Engine/Source/Programs/UnrealBuildAccelerator/Common/Private/UbaProcess.cpp
+++ b/Engine/Source/Programs/UnrealBuildAccelerator/Common/Private/UbaProcess.cpp
@@ -1413,10 +1413,10 @@ namespace uba
 				it += TStrlen(s) + 1;
 			}
 
-			int outPipe[2] = { 0 };
-			int errPipe[2] = { 0 };
-			auto pipeGuard0 = MakeGuard([&]() { if (outPipe[0]) close(outPipe[0]); if (errPipe[0]) close(errPipe[0]); });
-			auto pipeGuard1 = MakeGuard([&]() { if (outPipe[1]) close(outPipe[1]); if (errPipe[1]) close(errPipe[1]); });
+			int outPipe[2] = { -1, -1 };
+			int errPipe[2] = { -1, -1 };
+			auto pipeGuard0 = MakeGuard([&]() { if (outPipe[0] != -1) close(outPipe[0]); if (errPipe[0] != -1) close(errPipe[0]); });
+			auto pipeGuard1 = MakeGuard([&]() { if (outPipe[1] != -1) close(outPipe[1]); if (errPipe[1] != -1) close(errPipe[1]); });
 
 			if (m_detourEnabled)
 			{
@@ -1479,15 +1479,23 @@ namespace uba
 					return UBA_EXIT_CODE(18);
 				}
 
-				pipeGuard0.Cancel(); // TODO Should this be here? If process fails, will the below actions execute?
+				res = posix_spawn_file_actions_addclose(&fileActions, outPipe[0]);
+				UBA_ASSERTF(!res, "posix_spawn_file_actions_addclose outPipe[0] failed: %i", res);
 
-				posix_spawn_file_actions_addclose(&fileActions, outPipe[0]);
-				posix_spawn_file_actions_addclose(&fileActions, errPipe[0]);
-				posix_spawn_file_actions_adddup2(&fileActions, outPipe[1], 1);
-				posix_spawn_file_actions_adddup2(&fileActions, errPipe[1], 2);
+				res = posix_spawn_file_actions_addclose(&fileActions, errPipe[0]);
+				UBA_ASSERTF(!res, "posix_spawn_file_actions_addclose errPipe[1] failed: %i", res);
 
-				posix_spawn_file_actions_addclose(&fileActions, outPipe[1]);
-				posix_spawn_file_actions_addclose(&fileActions, errPipe[1]);
+				res = posix_spawn_file_actions_adddup2(&fileActions, outPipe[1], 1);
+				UBA_ASSERTF(!res, "posix_spawn_file_actions_adddup2 outPipe[1] failed: %i", res);
+
+				res = posix_spawn_file_actions_adddup2(&fileActions, errPipe[1], 2);
+				UBA_ASSERTF(!res, "posix_spawn_file_actions_adddup2 errPipe[1] failed: %i", res);
+
+				res = posix_spawn_file_actions_addclose(&fileActions, outPipe[1]);
+				UBA_ASSERTF(!res, "posix_spawn_file_actions_addclose outPipe[1] failed: %i", res);
+
+				res = posix_spawn_file_actions_addclose(&fileActions, errPipe[1]);
+				UBA_ASSERTF(!res, "posix_spawn_file_actions_addclose errPipe[1] failed: %i", res);
 			}
 
 			envvars.push_back(nullptr);
@@ -1539,9 +1547,6 @@ namespace uba
 
 				for (int rval; (rval = poll(plist, sizeof_array(plist), -1)) > 0;)
 				{
-					if (plist[0].revents & POLLHUP && plist[1].revents & POLLHUP) // If they both have hung up we hang up
-						break;
-
 					if (plist[0].revents & POLLERR || plist[1].revents & POLLERR) // If there is an error on any of them we hang up
 					{
 						logger.Error(TC("pipe polling error"));
@@ -1564,6 +1569,10 @@ namespace uba
 
 					char buffer[1024];
 					int bytesRead = read(fd, buffer, sizeof_array(buffer) - 1);
+
+					if (bytesRead == 0 && plist[0].revents & POLLHUP && plist[1].revents & POLLHUP) // If they both have hung up we hang up
+						break;
+
 					buffer[bytesRead] = 0;
 					pipeReader->ReadData(buffer, bytesRead);
 				}
