[UBA]
* Removed realpath test for linux because test machines do not have clang installed

diff --git a/Engine/Source/Programs/UnrealBuildAccelerator/Common/Private/UbaStorage.cpp b/Engine/Source/Programs/UnrealBuildAccelerator/Common/Private/UbaStorage.cpp
--- a/Engine/Source/Programs/UnrealBuildAccelerator/Common/Private/UbaStorage.cpp
+++ b/Engine/Source/Programs/UnrealBuildAccelerator/Common/Private/UbaStorage.cpp
@@ -1816,6 +1816,7 @@ namespace uba
 
 	bool StorageImpl::CopyOrLink(const CasKey& casKey, const tchar* destination, u32 fileAttributes)
 	{
+		UBA_ASSERT(casKey != CasKeyZero);
 		UBA_ASSERT(fileAttributes);
 
 		auto& stats = Stats();
diff --git a/Engine/Source/Programs/UnrealBuildAccelerator/TestApp/Private/PosixOS/UbaTestAppPosix.cpp b/Engine/Source/Programs/UnrealBuildAccelerator/TestApp/Private/PosixOS/UbaTestAppPosix.cpp
--- a/Engine/Source/Programs/UnrealBuildAccelerator/TestApp/Private/PosixOS/UbaTestAppPosix.cpp
+++ b/Engine/Source/Programs/UnrealBuildAccelerator/TestApp/Private/PosixOS/UbaTestAppPosix.cpp
@@ -144,9 +144,11 @@ int main()
 	if (stat("FileWF", &attrWF) == -1)
 		return LogError("stat for FileW failed");
 
+#if !PLATFORM_LINUX // Farm linux machines do not have clang installed... need to revisit this
 	char fullPath[PATH_MAX];
 	if (realpath("/usr/bin/clang", fullPath) == nullptr)
 		return LogError("realpath for 'clang' failed");
+#endif
 
 	struct stat attrRoot;
 	if (stat("/", &attrRoot) != 0)
