[UBA]
* Added lock around socket send for windows to see if it will resolve some of the crazy socket send times on the farm

diff --git a/Engine/Source/Programs/UnrealBuildAccelerator/Common/Private/UbaNetworkBackendTcp.cpp b/Engine/Source/Programs/UnrealBuildAccelerator/Common/Private/UbaNetworkBackendTcp.cpp
--- a/Engine/Source/Programs/UnrealBuildAccelerator/Common/Private/UbaNetworkBackendTcp.cpp
+++ b/Engine/Source/Programs/UnrealBuildAccelerator/Common/Private/UbaNetworkBackendTcp.cpp
@@ -38,6 +38,8 @@
 #define WSAPoll poll
 #endif
 
+#define UBA_LOCK_AROUND_SEND 1  // !PLATFORM_WINDOWS
+
 namespace uba
 {
 	struct NetworkBackendTcp::ListenEntry
@@ -69,7 +71,7 @@ namespace uba
 		void* disconnectContext = nullptr;
 		DisconnectCallback* disconnectCallback = nullptr;
 
-		#if !PLATFORM_WINDOWS
+		#if UBA_LOCK_AROUND_SEND
 		ReaderWriterLock sendLock;
 		#endif
 
@@ -166,13 +168,13 @@ namespace uba
 		auto& conn = *(Connection*)connection;
 		sendContext.isUsed = true;
 
-		#if !PLATFORM_WINDOWS
+		#if UBA_LOCK_AROUND_SEND
 		ScopedWriteLock lock(conn.sendLock);
 		#endif
 
 		bool res = SendSocket(logger, conn.socket, data, dataSize);
 
-		#if !PLATFORM_WINDOWS
+		#if UBA_LOCK_AROUND_SEND
 		lock.Leave();
 		#endif
 
