UBA: Close fds when done with polling to stop leaking on MacOS.

diff --git a/Engine/Source/Programs/UnrealBuildAccelerator/Common/Private/UbaProcess.cpp b/Engine/Source/Programs/UnrealBuildAccelerator/Common/Private/UbaProcess.cpp
--- a/Engine/Source/Programs/UnrealBuildAccelerator/Common/Private/UbaProcess.cpp
+++ b/Engine/Source/Programs/UnrealBuildAccelerator/Common/Private/UbaProcess.cpp
@@ -1544,9 +1544,13 @@ namespace uba
 				PipeReader errReader(*this, LogEntryType_Error);
 
 				pollfd plist[] = { {outPipe[0],POLLIN, 0}, {errPipe[0],POLLIN, 0} };
-
-				for (int rval; (rval = poll(plist, sizeof_array(plist), -1)) > 0;)
+				int rval = 0;
+				for (; (rval = poll(plist, sizeof_array(plist), -1)) > 0;)
 				{
+					if (plist[0].revents & POLLHUP && plist[1].revents & POLLHUP) // If they both have hung up we hang up
+					{
+						break;
+					}
 					if (plist[0].revents & POLLERR || plist[1].revents & POLLERR) // If there is an error on any of them we hang up
 					{
 						logger.Error(TC("pipe polling error"));
@@ -1570,12 +1574,16 @@ namespace uba
 					char buffer[1024];
 					int bytesRead = read(fd, buffer, sizeof_array(buffer) - 1);
 
-					if (bytesRead == 0 && plist[0].revents & POLLHUP && plist[1].revents & POLLHUP) // If they both have hung up we hang up
-						break;
-
 					buffer[bytesRead] = 0;
 					pipeReader->ReadData(buffer, bytesRead);
 				}
+				if (rval == -1)
+				{
+					logger.Error(TC("pipe polling error with -1"));
+				}
+				// Now that we're done polling, close the fds.
+				close(errPipe[0]);
+				close(outPipe[0]);
 			}
 		}
 		else
