[UBA]
* Removed ExitProcess call when showing assert info in CreateProcess call where process tries to spawn winedbg

diff --git a/Engine/Source/Programs/UnrealBuildAccelerator/Detours/Private/Windows/UbaDetoursFunctionsKernelBase.inl b/Engine/Source/Programs/UnrealBuildAccelerator/Detours/Private/Windows/UbaDetoursFunctionsKernelBase.inl
--- a/Engine/Source/Programs/UnrealBuildAccelerator/Detours/Private/Windows/UbaDetoursFunctionsKernelBase.inl
+++ b/Engine/Source/Programs/UnrealBuildAccelerator/Detours/Private/Windows/UbaDetoursFunctionsKernelBase.inl
@@ -2512,7 +2512,7 @@ BOOL Detoured_CreateProcessW(LPCWSTR lpApplicationName, LPWSTR lpCommandLine, LP
 	{
 		if (g_runningRemote)
 		{
-			UbaAssert(L"Suppress debugger startup and try to report issue instead. This message is here to hopefully see callstack", "", 0, "", 342343);
+			UbaAssert(L"Suppress debugger startup and try to report issue instead. This message is here to hopefully see callstack", "", 0, "", 0, false);
 			return false;
 		}
 		else
