UBA: Add if some defs around POLLUP for platform differences

diff --git a/Engine/Source/Programs/UnrealBuildAccelerator/Common/Private/UbaProcess.cpp b/Engine/Source/Programs/UnrealBuildAccelerator/Common/Private/UbaProcess.cpp
--- a/Engine/Source/Programs/UnrealBuildAccelerator/Common/Private/UbaProcess.cpp
+++ b/Engine/Source/Programs/UnrealBuildAccelerator/Common/Private/UbaProcess.cpp
@@ -1547,10 +1547,6 @@ namespace uba
 				int rval = 0;
 				for (; (rval = poll(plist, sizeof_array(plist), -1)) > 0;)
 				{
-					if (plist[0].revents & POLLHUP && plist[1].revents & POLLHUP) // If they both have hung up we hang up
-					{
-						break;
-					}
 					if (plist[0].revents & POLLERR || plist[1].revents & POLLERR) // If there is an error on any of them we hang up
 					{
 						logger.Error(TC("pipe polling error"));
@@ -1576,14 +1572,26 @@ namespace uba
 
 					buffer[bytesRead] = 0;
 					pipeReader->ReadData(buffer, bytesRead);
+
+					// If they both have hung up we hang up
+					// However, pollhup means slightly different things on MacOS vs Linux
+					// so we need to check slightly different things
+					if (plist[0].revents & POLLHUP && plist[1].revents & POLLHUP)
+					{
+						#if PLATFORM_LINUX
+						if (bytesRead == 0)
+						#endif
+							break;
+					}
 				}
 				if (rval == -1)
 				{
 					logger.Error(TC("pipe polling error with -1"));
 				}
+				#if PLATFORM_MAC
 				// Now that we're done polling, close the fds.
-				close(errPipe[0]);
-				close(outPipe[0]);
+				pipeGuard0.Execute();
+				#endif
 			}
 		}
 		else
