[Uba]
* Skip clang test if clang does not exists on machine

diff --git a/Engine/Source/Programs/UnrealBuildAccelerator/Test/Private/UbaTestSession.h b/Engine/Source/Programs/UnrealBuildAccelerator/Test/Private/UbaTestSession.h
--- a/Engine/Source/Programs/UnrealBuildAccelerator/Test/Private/UbaTestSession.h
+++ b/Engine/Source/Programs/UnrealBuildAccelerator/Test/Private/UbaTestSession.h
@@ -191,6 +191,9 @@ namespace uba
 		const tchar* clangPath = TC("/usr/bin/clang++");
 #endif
 
+		if (!FileExists(logger, clangPath)) // Skipping if clang is not installed.
+			return true;
+
 		ProcessStartInfo processInfo;
 		processInfo.application = clangPath;
 #if PLATFORM_MAC
@@ -355,6 +358,7 @@ namespace uba
 				processInfo.application = GetPingApplication();
 				processInfo.workingDir = workingDir;
 				processInfo.arguments = IsWindows ? TC("-n 2 localhost") : TC("-c 2 localhost");
+				//processInfo.logFile = "[REDACTED_USER_HOME]/LogFile.log";
 				Vector<ProcessHandle> processes;
 
 				for (u32 i=0; i!=50; ++i)
