[UBA]
* Removed error on linux for failing poll call

diff --git a/Engine/Source/Programs/UnrealBuildAccelerator/Common/Private/UbaProcess.cpp b/Engine/Source/Programs/UnrealBuildAccelerator/Common/Private/UbaProcess.cpp
--- a/Engine/Source/Programs/UnrealBuildAccelerator/Common/Private/UbaProcess.cpp
+++ b/Engine/Source/Programs/UnrealBuildAccelerator/Common/Private/UbaProcess.cpp
@@ -1584,11 +1584,10 @@ namespace uba
 							break;
 					}
 				}
-				if (rval == -1)
-				{
-					logger.Error(TC("pipe polling error with -1"));
-				}
+
 				#if PLATFORM_MAC
+				if (rval == -1)
+					logger.Error(TC("pipe polling error with -1 (%s)"), strerror(errno));
 				// Now that we're done polling, close the fds.
 				pipeGuard0.Execute();
 				#endif
