[UBA]
* Added more context menu options in visualizer to navigate trace

diff --git a/Engine/Source/Programs/UnrealBuildAccelerator/Visualizer/Private/UbaVisualizer.cpp b/Engine/Source/Programs/UnrealBuildAccelerator/Visualizer/Private/UbaVisualizer.cpp
--- a/Engine/Source/Programs/UnrealBuildAccelerator/Visualizer/Private/UbaVisualizer.cpp
+++ b/Engine/Source/Programs/UnrealBuildAccelerator/Visualizer/Private/UbaVisualizer.cpp
@@ -14,6 +14,10 @@ enum
 	Popup_CopySessionInfo = 3,
 	Popup_CopyProcessInfo,
 	Popup_CopyProcessLog,
+	Popup_Replay,
+	Popup_Pause,
+	Popup_Play,
+	Popup_JumpToEnd,
 	Popup_SaveAs,
 	Popup_Quit,
 };
@@ -2298,9 +2302,24 @@ namespace uba
 
 				AppendMenuW(hMenu, MF_STRING, Popup_CopyProcessInfo, L"&Copy Process Info");
 				if (!process.logLines.empty())
-					AppendMenuW(hMenu, MF_STRING, Popup_CopyProcessLog, L"&Copy Process Log");
+					AppendMenuW(hMenu, MF_STRING, Popup_CopyProcessLog, L"Copy Process &Log");
 				AppendMenuW(hMenu, MF_SEPARATOR, 0, NULL);
 			}
+			if (m_fileName.data)
+			{
+				if (!m_replay)
+					AppendMenuW(hMenu, MF_STRING, Popup_Replay, L"&Replay Trace");
+				else
+				{
+					if (m_paused)
+						AppendMenuW(hMenu, MF_STRING, Popup_Play, L"&Play");
+					else
+						AppendMenuW(hMenu, MF_STRING, Popup_Pause, L"&Pause");
+					AppendMenuW(hMenu, MF_STRING, Popup_JumpToEnd, L"&Jump To End");
+				}
+			}
+
+
 			AppendMenuW(hMenu, MF_STRING, Popup_SaveAs, L"&Save Trace");
 			AppendMenuW(hMenu, MF_SEPARATOR, 0, NULL);
 			AppendMenuW(hMenu, MF_STRING, Popup_Quit, L"&Quit");
@@ -2329,6 +2348,26 @@ namespace uba
 					m_trace.SaveAs(ofn.lpstrFile);
 				break;
 			}
+			case Popup_Replay:
+				m_replay = 1;
+				PostMessage(m_hwnd, WM_NEWTRACE, 0, 0);
+				break;
+
+			case Popup_Play:
+				m_paused = false;
+				m_pauseTime += GetTime() - m_pauseStart;
+				break;
+
+			case Popup_Pause:
+				m_paused = true;
+				m_pauseStart = GetTime();
+				break;
+
+			case Popup_JumpToEnd:
+				m_traceView.finished = true;
+				m_replay = 0;
+				PostMessage(m_hwnd, WM_NEWTRACE, 0, 0);
+				break;
 
 			case Popup_Quit: // Quit
 				m_looping = false;
