UBA: Improve the Mac Helper App.

* Adds a settings menu
* Install from P4
* Pick a binary
* Launch at Startup
* Uses App Storage for Settings

diff --git a/Engine/Source/Programs/UnrealBuildAccelerator/UbaAgentHelperMac/UbaAgentHelper/Assets.xcassets/AccentColor.colorset/Contents.json b/Engine/Source/Programs/UnrealBuildAccelerator/UbaAgentHelperMac/UbaAgentHelper/Assets.xcassets/AccentColor.colorset/Contents.json
--- a/Engine/Source/Programs/UnrealBuildAccelerator/UbaAgentHelperMac/UbaAgentHelper/Assets.xcassets/AccentColor.colorset/Contents.json
+++ b/Engine/Source/Programs/UnrealBuildAccelerator/UbaAgentHelperMac/UbaAgentHelper/Assets.xcassets/AccentColor.colorset/Contents.json
@@ -1,6 +1,33 @@
 {
   "colors" : [
     {
+      "color" : {
+        "color-space" : "srgb",
+        "components" : {
+          "alpha" : "1.000",
+          "blue" : "1.000",
+          "green" : "0.750",
+          "red" : "0.432"
+        }
+      },
+      "idiom" : "universal"
+    },
+    {
+      "appearances" : [
+        {
+          "appearance" : "luminosity",
+          "value" : "dark"
+        }
+      ],
+      "color" : {
+        "color-space" : "srgb",
+        "components" : {
+          "alpha" : "1.000",
+          "blue" : "1.000",
+          "green" : "1.000",
+          "red" : "1.000"
+        }
+      },
       "idiom" : "universal"
     }
   ],
diff --git a/Engine/Source/Programs/UnrealBuildAccelerator/UbaAgentHelperMac/UbaAgentHelper/Assets.xcassets/AppIcon.appiconset/Contents.json b/Engine/Source/Programs/UnrealBuildAccelerator/UbaAgentHelperMac/UbaAgentHelper/Assets.xcassets/AppIcon.appiconset/Contents.json
--- a/Engine/Source/Programs/UnrealBuildAccelerator/UbaAgentHelperMac/UbaAgentHelper/Assets.xcassets/AppIcon.appiconset/Contents.json
+++ b/Engine/Source/Programs/UnrealBuildAccelerator/UbaAgentHelperMac/UbaAgentHelper/Assets.xcassets/AppIcon.appiconset/Contents.json
@@ -1,51 +1,61 @@
 {
   "images" : [
     {
+      "filename" : "icon_16x16.png",
       "idiom" : "mac",
       "scale" : "1x",
       "size" : "16x16"
     },
     {
+      "filename" : "[REDACTED_EMAIL]",
       "idiom" : "mac",
       "scale" : "2x",
       "size" : "16x16"
     },
     {
+      "filename" : "icon_32x32.png",
       "idiom" : "mac",
       "scale" : "1x",
       "size" : "32x32"
     },
     {
+      "filename" : "[REDACTED_EMAIL]",
       "idiom" : "mac",
       "scale" : "2x",
       "size" : "32x32"
     },
     {
+      "filename" : "icon_128x128.png",
       "idiom" : "mac",
       "scale" : "1x",
       "size" : "128x128"
     },
     {
+      "filename" : "[REDACTED_EMAIL]",
       "idiom" : "mac",
       "scale" : "2x",
       "size" : "128x128"
     },
     {
+      "filename" : "icon_256x256.png",
       "idiom" : "mac",
       "scale" : "1x",
       "size" : "256x256"
     },
     {
+      "filename" : "[REDACTED_EMAIL]",
       "idiom" : "mac",
       "scale" : "2x",
       "size" : "256x256"
     },
     {
+      "filename" : "icon_512x512.png",
       "idiom" : "mac",
       "scale" : "1x",
       "size" : "512x512"
     },
     {
+      "filename" : "[REDACTED_EMAIL]",
       "idiom" : "mac",
       "scale" : "2x",
       "size" : "512x512"
diff --git a/Engine/Source/Programs/UnrealBuildAccelerator/UbaAgentHelperMac/UbaAgentHelper/Constants.swift b/Engine/Source/Programs/UnrealBuildAccelerator/UbaAgentHelperMac/UbaAgentHelper/Constants.swift
new file mode 100644
--- /dev/null
+++ b/Engine/Source/Programs/UnrealBuildAccelerator/UbaAgentHelperMac/UbaAgentHelper/Constants.swift
@@ -0,0 +1,14 @@
+//
+//  Constants.swift
+//  UbaAgentHelper
+//
+//  Created by Zack Neyland on 1/21/24.
+//
+
+import Foundation
+
+enum Constants {
+    enum BundleIds {
+        static let main = "com.epicgames.UbaAgentHelperApp"
+    }
+}
diff --git a/Engine/Source/Programs/UnrealBuildAccelerator/UbaAgentHelperMac/UbaAgentHelper/ContentView.swift b/Engine/Source/Programs/UnrealBuildAccelerator/UbaAgentHelperMac/UbaAgentHelper/ContentView.swift
--- a/Engine/Source/Programs/UnrealBuildAccelerator/UbaAgentHelperMac/UbaAgentHelper/ContentView.swift
+++ b/Engine/Source/Programs/UnrealBuildAccelerator/UbaAgentHelperMac/UbaAgentHelper/ContentView.swift
@@ -10,35 +10,26 @@ import SwiftUI
 // Implement this support: https://developer.apple.com/forums/thread/735862
 
 struct ContentView: View {
+    
+    @AppStorage(StorageKeys.XCodePath.id)
+    private var xcodePath: String = StorageKeys.XCodePath.default
+    
+    @AppStorage(StorageKeys.UbaAgentPath.id)
+    private var ubaAgentPath: String = StorageKeys.UbaAgentPath.default
+    
     @State var activeTask : Process? = nil
     @Binding var status: Bool
-    @State var pathToAgent : String = ""
+    //    @State var pathToAgent : String = ""
     @State var ipAddress = ""
     @State var log : String = ""
-    @State var xcodePath: String = ""
+    
     var body: some View {
         VStack {
             GroupBox(label: Text("Uba Agent Configuration"), content: {
-                HStack(content: {
-                    Text("Agent Location:")
-                        .fontWeight(.bold)
-                    TextField(text: $pathToAgent, prompt: Text("e.g.: ..../Engine/Binaries/Mac/UnrealBuildAccelerator/UbaAgent"), axis: .vertical, label: {})
-                    Spacer()
-                })
                 HStack {
                     Text("UbaHost IP:")
-                        .fontWeight(.bold)
+                    //                        .fontWeight(.bold)
                     TextField(text: $ipAddress, prompt: Text("e.g.: 127.0.0.1"), label: {})
-                    Spacer()
-                }
-                HStack {
-                    Text("XCode Location:")
-                        .fontWeight(.bold)
-                    TextField(text: $xcodePath, prompt: Text("e.g.: /Applications/XCode"), axis: .vertical, label: {})
-                    Spacer()
-                    Button(action: locateXCode, label: {
-                        Text("Locate XCode")
-                    })
                 }
             })
             
@@ -47,9 +38,8 @@ struct ContentView: View {
                     Spacer()
                 }
                 ScrollView {
-                    TextField(text: $log,  axis: .vertical, label: {})
+                    TextEditor(text: $log)
                 }
-                
                 .defaultScrollAnchor(.bottom)
             })
             Divider()
@@ -62,7 +52,9 @@ struct ContentView: View {
                     .tint(.red)
                 }
                 else {
-                    Button(action: {runCode()}, label: {
+                    Button(action: {
+                        runUbaAgent()
+                    }, label: {
                         Text("Start Agent")
                             .bold()
                     })
@@ -70,7 +62,10 @@ struct ContentView: View {
                 }
                 
                 Spacer()
-                Button(action: { NSApplication.shared.terminate(nil)}, label: {
+                SettingsLink()
+                Button(action: {
+                    NSApplication.shared.terminate(nil)
+                }, label: {
                     Text("Quit")
                 })
                 .buttonStyle(.borderedProminent)
@@ -78,36 +73,15 @@ struct ContentView: View {
             }
         }
         .padding()
+        .background(Color.black.opacity(0.4))
     }
-    func canRun() -> Bool {
-        return ipAddress.isEmpty || xcodePath.isEmpty || pathToAgent.isEmpty
-    }
-    
-    func locateXCode() {
-        let task = Process()
-
-        //the path to the external program you want to run
-        let executableURL = URL(fileURLWithPath: "/usr/bin/xcode-select")
-        task.executableURL = executableURL
-
-        //use pipe to get the execution program's output
-        let pipe = Pipe()
-        task.standardOutput = pipe
 
-        //all the arguments to the executable
-        let args = ["-p"]
-        task.arguments = args
-
-        try! task.run()
-        task.waitUntilExit()
-
-        //all this code helps you capture the output so you can, for e.g., show the user
-        let d = pipe.fileHandleForReading.readDataToEndOfFile()
-        xcodePath = (String(data: d, encoding: String.Encoding.utf8) ?? "").trimmingCharacters(in: CharacterSet.newlines)
-    }
     
+    func canRun() -> Bool {
+        return ipAddress.isEmpty || xcodePath.isEmpty || ubaAgentPath.isEmpty
+    }
     
-    func runCode() {
+    func runUbaAgent() {
         // This should be replaced with the recs from here:
         // https://developer.apple.com/forums/thread/690310
         if ((activeTask?.isRunning) != nil) {
@@ -117,7 +91,7 @@ struct ContentView: View {
         log = ""
         
         let task = Process()
-        task.executableURL = URL(fileURLWithPath: pathToAgent)
+        task.executableURL = URL(fileURLWithPath: ubaAgentPath)
         task.arguments = ["-host=\(ipAddress)", "-log", "-populatecas=\(xcodePath)"]
         let outputPipe = Pipe()
         task.standardOutput = outputPipe
@@ -139,10 +113,10 @@ struct ContentView: View {
             status = true
         }
         //todo: Swap to this later
-//        Task {
-//            try! task.run()
-//            task.waitUntilExit()
-//        }
+        //        Task {
+        //            try! task.run()
+        //            task.waitUntilExit()
+        //        }
     }
     
     func killTask() {
diff --git a/Engine/Source/Programs/UnrealBuildAccelerator/UbaAgentHelperMac/UbaAgentHelper/Extensions/Color.swift b/Engine/Source/Programs/UnrealBuildAccelerator/UbaAgentHelperMac/UbaAgentHelper/Extensions/Color.swift
new file mode 100644
--- /dev/null
+++ b/Engine/Source/Programs/UnrealBuildAccelerator/UbaAgentHelperMac/UbaAgentHelper/Extensions/Color.swift
@@ -0,0 +1,37 @@
+//
+//  Color.swift
+//  UbaAgentHelper
+//
+//  Created by Zack Neyland on 1/21/24.
+//
+
+import Foundation
+import SwiftUI
+
+extension Color: RawRepresentable {
+    public init?(rawValue: String) {
+        guard let data = Data(base64Encoded: rawValue) else {
+            self = .black
+            return
+        }
+
+        do {
+            let color = try NSKeyedUnarchiver.unarchivedObject(ofClass: NSColor.self, from: data) ?? .labelColor
+            self = Color(color)
+        } catch {
+            self = .black
+        }
+    }
+
+    public var rawValue: String {
+        do {
+            let data = try NSKeyedArchiver.archivedData(
+                withRootObject: NSColor(self),
+                requiringSecureCoding: false
+            ) as Data
+            return data.base64EncodedString()
+        } catch {
+            return ""
+        }
+    }
+}
diff --git a/Engine/Source/Programs/UnrealBuildAccelerator/UbaAgentHelperMac/UbaAgentHelper/LoginLaunch.swift b/Engine/Source/Programs/UnrealBuildAccelerator/UbaAgentHelperMac/UbaAgentHelper/LoginLaunch.swift
new file mode 100644
--- /dev/null
+++ b/Engine/Source/Programs/UnrealBuildAccelerator/UbaAgentHelperMac/UbaAgentHelper/LoginLaunch.swift
@@ -0,0 +1,82 @@
+//
+//  LaunchItem.swift
+//  UbaAgentHelper
+//
+//  Created by Zack Neyland on 1/21/24.
+//
+
+import OSLog
+import ServiceManagement
+import SwiftUI
+
+public enum LoginLaunch {
+    private static let logger = Logger(subsystem: Constants.BundleIds.main, category: "login")
+    fileprivate static let observable = Observable()
+
+    public static var isEnabled: Bool {
+        get { SMAppService.mainApp.status == .enabled }
+        set {
+            observable.objectWillChange.send()
+
+            do {
+                if newValue {
+                    if SMAppService.mainApp.status == .enabled {
+                        try? SMAppService.mainApp.unregister()
+                    }
+
+                    try SMAppService.mainApp.register()
+                } else {
+                    try SMAppService.mainApp.unregister()
+                }
+            } catch {
+                logger.error("Failed to \(newValue ? "enable" : "disable"): \(error.localizedDescription)")
+            }
+        }
+    }
+
+    public static var wasLaunchedAtLogin: Bool {
+        let event = NSAppleEventManager.shared().currentAppleEvent
+        return event?.eventID == kAEOpenApplication
+            && event?.paramDescriptor(forKeyword: keyAEPropData)?.enumCodeValue == keyAELaunchedAsLogInItem
+    }
+}
+
+extension LoginLaunch {
+    final class Observable: ObservableObject {
+        var isEnabled: Bool {
+            get { LoginLaunch.isEnabled }
+            set {
+                LoginLaunch.isEnabled = newValue
+            }
+        }
+    }
+}
+
+extension LoginLaunch {
+    public struct Toggle<Label: View>: View {
+        @ObservedObject private var launchAtLogin = LoginLaunch.observable
+        private let label: Label
+
+        public init(@ViewBuilder label: () -> Label) {
+            self.label = label()
+        }
+
+        public var body: some View {
+            SwiftUI.Toggle(isOn: $launchAtLogin.isEnabled) { label }
+        }
+    }
+}
+
+extension LoginLaunch.Toggle<Text> {
+    public init(_ titleKey: LocalizedStringKey) {
+        label = Text(titleKey)
+    }
+
+    public init(_ title: some StringProtocol) {
+        label = Text(title)
+    }
+
+    public init() {
+        self.init("Launch At Login")
+    }
+}
diff --git a/Engine/Source/Programs/UnrealBuildAccelerator/UbaAgentHelperMac/UbaAgentHelper/Preferences.swift b/Engine/Source/Programs/UnrealBuildAccelerator/UbaAgentHelperMac/UbaAgentHelper/Preferences.swift
new file mode 100644
--- /dev/null
+++ b/Engine/Source/Programs/UnrealBuildAccelerator/UbaAgentHelperMac/UbaAgentHelper/Preferences.swift
@@ -0,0 +1,21 @@
+//
+//  Preferences.swift
+//  UbaAgentHelper
+//
+//  Created by Zack Neyland on 1/21/24.
+//
+
+import Foundation
+
+enum StorageKeys {
+    
+    struct Key<Value: Any> {
+        let id: String
+        let `default`: Value
+    }
+    
+    static let XCodePath = Key(id: "xcodePath", default: "")
+    static let UbaAgentPath = Key(id: "ubaAgentPath", default: "")
+    static let P4DepotRoot = Key(id: "p4DepotRoot", default: "")
+    static let AgentFromRepo = Key(id: "agentFromRepo", default: false)
+}
diff --git a/Engine/Source/Programs/UnrealBuildAccelerator/UbaAgentHelperMac/UbaAgentHelper/SettingsView.swift b/Engine/Source/Programs/UnrealBuildAccelerator/UbaAgentHelperMac/UbaAgentHelper/SettingsView.swift
new file mode 100644
--- /dev/null
+++ b/Engine/Source/Programs/UnrealBuildAccelerator/UbaAgentHelperMac/UbaAgentHelper/SettingsView.swift
@@ -0,0 +1,272 @@
+//
+//  SettingsView.swift
+//  UbaAgentHelper
+//
+//  Created by Zack Neyland on 1/21/24.
+//
+
+import SwiftUI
+
+struct SettingsView: View {
+    
+    @AppStorage(StorageKeys.XCodePath.id)
+    private var xcodePath: String = StorageKeys.XCodePath.default
+    
+    @AppStorage(StorageKeys.UbaAgentPath.id)
+    private var ubaAgentPath: String = StorageKeys.UbaAgentPath.default
+    
+    @AppStorage(StorageKeys.P4DepotRoot.id)
+    private var p4DepotRoot : String = StorageKeys.P4DepotRoot.default
+    
+    
+    @AppStorage(StorageKeys.AgentFromRepo.id)
+    private var agentFromP4 : Bool = StorageKeys.AgentFromRepo.default
+    
+    private enum Tabs: Hashable {
+        case general
+        case ubaAgent
+    }
+    
+    @State private var isImporting: Bool = false
+    @State private var p4BinaryPath : URL? = nil
+    
+    @State private var selection: Tabs = .general
+    
+    var body: some View {
+        TabView(selection: $selection) {
+            generalTab
+            ubaAgentTab
+        }
+        .frame(width: 375, height: 350)
+        .onAppear(perform: {
+            findP4Binary()
+        })
+    }
+    
+    private var generalTab: some View {
+        Form {
+            Section("Paths") {
+                VStack(alignment: .leading) {
+                    Text("XCode Path")
+                    Text("\(xcodePath.isEmpty ? "Please locate XCode" : xcodePath)")
+                        .font(.caption)
+                        .foregroundStyle(.secondary)
+                }
+                HStack {
+                    Spacer()
+                    Button("Locate XCode") {
+                        locateXCode()
+                    }
+                    Spacer()
+                }
+            }
+            // TODO: Revisit launching at login at a later date
+            Section {
+                VStack(alignment: .leading) {
+                    LoginLaunch.Toggle()
+                }
+            }
+            Section("DEBUG") {
+                VStack(alignment: .center)
+                {
+                    Button(action:
+                            {
+                        if let bundleID = Bundle.main.bundleIdentifier {
+                            UserDefaults.standard.removePersistentDomain(forName: bundleID)
+                        }
+                    }, label: {
+                        Text("Reset Settings to Default")
+                    })
+                }
+            }
+        }
+        .formStyle(.grouped)
+        .tabItem {
+            Label("General", systemImage: Symbols.gear())
+        }
+        .tag(Tabs.general)
+    }
+    
+    private var ubaAgentTab : some View {
+        Form {
+            Section("Uba Agent") {
+                if let _ = p4BinaryPath {
+                    Toggle(isOn: $agentFromP4, label: {
+                        Text("Use Perforce")
+                    })
+                }
+                
+                if agentFromP4 {
+                    HStack() {
+                        VStack(alignment: .leading) {
+                            Text("P4 Depot Root")
+                            Text("e.g. //UE5/Main")
+                                .font(.caption)
+                                .foregroundStyle(.secondary)
+                        }
+                        TextField(text: $p4DepotRoot, prompt: Text("Click here"), axis: .vertical, label: {})
+                    }
+                    HStack {
+                        Spacer()
+                        Button("Install Agent") {
+                            InstallUbaAgent()
+                        }
+                        .disabled(p4DepotRoot.isEmpty)
+                        Spacer()
+                    }
+                    HStack {
+                        VStack(alignment: .leading) {
+                            Text("UbaAgentPath")
+                            Text(ubaAgentPath.isEmpty ? "Please select an UbaAgent" : ubaAgentPath)
+                                .font(.caption)
+                                .foregroundStyle(.secondary)
+                        }
+                        Spacer()
+                    }
+                }
+                else {
+                    HStack() {
+                        VStack(alignment: .leading) {
+                            Text("UbaAgentPath")
+                            Text(ubaAgentPath.isEmpty ? "Please select an UbaAgent" : ubaAgentPath)
+                                .font(.caption)
+                                .foregroundStyle(.secondary)
+                        }
+                        Spacer()
+                        Button(action: {
+                            isImporting = true
+                        }, label: {
+                            Text("Locate Agent")
+                        })
+                        .fileImporter(isPresented: $isImporting,
+                                      allowedContentTypes: [.executable],
+                                      onCompletion: { result in
+                            
+                            switch result {
+                            case .success(let url):
+                                // url contains the URL of the chosen file.
+                                ubaAgentPath = url.path()
+                            case .failure(let error):
+                                print(error)
+                            }
+                        })
+                    }
+                }
+            }
+        }
+        .formStyle(.grouped)
+        .tabItem {
+            Label("UbaAgent Config", systemImage: Symbols.wandAndStarsInverse())
+        }
+        .tag(Tabs.ubaAgent)
+    }
+    
+    func findP4Binary() {
+        let p4Locations = [
+            "/usr/local/bin/p4",
+            "/usr/bin/p4",
+            "/opt/homebrew/bin/p4"
+        ]
+        
+        for loc in p4Locations {
+            let url = URL(fileURLWithPath: loc)
+            do {
+                if try url.checkResourceIsReachable() {
+                    p4BinaryPath = url
+                    break
+                }
+            }
+            catch {
+                
+            }
+        }
+    }
+    
+    func InstallUbaAgent() {
+        let task = Process()
+        
+        let executableURL = p4BinaryPath
+        task.executableURL = executableURL
+        
+        let pipe = Pipe()
+        task.standardOutput = pipe
+        
+        if (p4DepotRoot.last == "/")
+        {
+            p4DepotRoot = String(p4DepotRoot.dropLast())
+        }
+        let args = ["print", "-q", (p4DepotRoot + "/Engine/Binaries/Mac/UnrealBuildAccelerator/UbaAgent")]
+        task.arguments = args
+        task.launch()
+        
+        //        try! task.run()
+        //        task.waitUntilExit()
+        
+        let data = pipe.fileHandleForReading.readDataToEndOfFile()
+        
+        if !data.isEmpty {
+            let ubaAgentUrl = FileManager.default.homeDirectoryForCurrentUser.appendingPathComponent("UbaAgent/UbaAgent", isDirectory: false)
+            
+            FileManager.default.createFile(atPath: ubaAgentUrl.path, contents: data, attributes: [FileAttributeKey.posixPermissions: NSNumber(value: 0o775 as UInt16)])
+            do {
+                if try ubaAgentUrl.checkResourceIsReachable() {
+                    ubaAgentPath = ubaAgentUrl.path
+                }
+            }
+            catch {
+            }
+        }
+        
+    }
+    
+    func locateXCode() {
+        let task = Process()
+        
+        let executableURL = URL(fileURLWithPath: "/usr/bin/xcode-select")
+        task.executableURL = executableURL
+        let pipe = Pipe()
+        task.standardOutput = pipe
+        
+        let args = ["-p"]
+        task.arguments = args
+        
+        try! task.run()
+        task.waitUntilExit()
+        
+        let data = pipe.fileHandleForReading.readDataToEndOfFile()
+        xcodePath = (String(data: data, encoding: String.Encoding.utf8) ?? "").trimmingCharacters(in: CharacterSet.newlines)
+    }
+    
+    func writeToFile(data: Data, fileName: String){
+        // get path of directory
+        guard let directory = FileManager.default.urls(for: .documentDirectory, in: .userDomainMask).last else {
+            return
+        }
+        // create file url
+        let fileurl =  directory.appendingPathComponent("\(fileName).txt")
+        // if file exists then write data
+        if FileManager.default.fileExists(atPath: fileurl.path) {
+            if let fileHandle = FileHandle(forWritingAtPath: fileurl.path) {
+                // seekToEndOfFile, writes data at the last of file(appends not override)
+                fileHandle.seekToEndOfFile()
+                fileHandle.write(data)
+                fileHandle.closeFile()
+            }
+            else {
+                print("Can't open file to write.")
+            }
+        }
+        else {
+            // if file does not exist write data for the first time
+            do{
+                try data.write(to: fileurl, options: .atomic)
+            }catch {
+                print("Unable to write in new file.")
+            }
+        }
+    }
+}
+
+#Preview {
+    SettingsView()
+}
diff --git a/Engine/Source/Programs/UnrealBuildAccelerator/UbaAgentHelperMac/UbaAgentHelper/Symbols.swift b/Engine/Source/Programs/UnrealBuildAccelerator/UbaAgentHelperMac/UbaAgentHelper/Symbols.swift
new file mode 100644
--- /dev/null
+++ b/Engine/Source/Programs/UnrealBuildAccelerator/UbaAgentHelperMac/UbaAgentHelper/Symbols.swift
@@ -0,0 +1,29 @@
+//
+//  Symbols.swift
+//  UbaAgentHelper
+//
+//  Created by Zack Neyland on 1/21/24.
+//
+
+import SwiftUI
+
+enum Symbols: String {
+    case gear
+    case wandAndStarsInverse = "wand.and.stars.inverse"
+
+    var image: Image {
+        Image(systemName: self.rawValue)
+    }
+
+    var name: String {
+        self.rawValue
+    }
+
+    func nsImage(accessibilityDescription: String? = nil) -> NSImage? {
+        return NSImage(systemSymbolName: self.rawValue, accessibilityDescription: accessibilityDescription)
+    }
+
+    func callAsFunction() -> String {
+        return self.rawValue
+    }
+}
diff --git a/Engine/Source/Programs/UnrealBuildAccelerator/UbaAgentHelperMac/UbaAgentHelper/UbaAgentHelperApp.swift b/Engine/Source/Programs/UnrealBuildAccelerator/UbaAgentHelperMac/UbaAgentHelper/UbaAgentHelperApp.swift
--- a/Engine/Source/Programs/UnrealBuildAccelerator/UbaAgentHelperMac/UbaAgentHelper/UbaAgentHelperApp.swift
+++ b/Engine/Source/Programs/UnrealBuildAccelerator/UbaAgentHelperMac/UbaAgentHelper/UbaAgentHelperApp.swift
@@ -9,11 +9,23 @@ import SwiftUI
 
 @main
 struct UbaAgentHelperApp: App {
+    
     @State var isRunning = false
+    
     var body: some Scene {
+        
+
+        
         MenuBarExtra {
+            
             ContentView(status: $isRunning)
                 .frame(width: 500, height: 500)
+//                .onAppear(perform: {
+//                    if let bundleID = Bundle.main.bundleIdentifier {
+//                        UserDefaults.standard.removePersistentDomain(forName: bundleID)
+//                    }
+//                })
+                
         } label: {
             HStack {
                 let configuration = NSImage.SymbolConfiguration(pointSize: 16, weight: .light)
@@ -23,12 +35,12 @@ struct UbaAgentHelperApp: App {
                 let updateImage = image?.withSymbolConfiguration(configuration)
                 
                 Image(nsImage: updateImage!) // This works.
-                
+                // TODO: If we want custom text in the menu bar this will cover it along with colors
                 //                    let font = NSFont.systemFont(ofSize: 16, weight: .light)
                 //                let color = NSColor.red // Attempting to make the text red.
                 //                    let attributes: [NSAttributedString.Key: Any] = [
                 //                        .font: font,
-                //                        .foregroundColor: color
+                //                        .foregroundColor: colorx
                 //                    ]
                 
                 //                    let attributedString = NSAttributedString(string: "Hello, world!", attributes: attributes)
@@ -38,15 +50,13 @@ struct UbaAgentHelperApp: App {
         }
         .menuBarExtraStyle(.window)
         
-        //        MenuBarExtra("UbaAgent Helper")
-        //            ) {
-        //
-        //                }
-                WindowGroup {
-                    ContentView(status: $isRunning)
-                        .navigationTitle("UbaAgent Helper")
-        
-        
-                }
+        //        WindowGroup {
+        //            ContentView(status: $isRunning)
+        //                .navigationTitle("UbaAgent Helper")
+        //        }
+        Settings {
+            SettingsView()
+            //                .environmentObject(appDelegate.utility)
+        }
     }
 }
