[UBA]
* Fixed so stdout/err is redirected to pipe on detoured processes on linux/mac.. this should fix so errors are visualized in the right place
* Added -nostdout to UbaCli to be able to check that redirected stdout/err works
* Added version for process messages to catch issues where ppl have old ubaagent.exe but new ubadetours.dll
* Added capacity parameter to FixPath and added asserts to check that we never write outside capacity
* Fixed code creating g_virtualApplication buffer.
* Disabled asserts for mac non debug builds
* Removed detoured write/fwrite for posix now when we instead redirect stdout/err
* Enabled some unit tests for linux now when they work as intended regarding redirected stdout/err

diff --git a/Engine/Source/Programs/UnrealBuildAccelerator/Cli/Private/UbaCli.cpp b/Engine/Source/Programs/UnrealBuildAccelerator/Cli/Private/UbaCli.cpp
--- a/Engine/Source/Programs/UnrealBuildAccelerator/Cli/Private/UbaCli.cpp
+++ b/Engine/Source/Programs/UnrealBuildAccelerator/Cli/Private/UbaCli.cpp
@@ -61,6 +61,7 @@ namespace uba
 		logger.Info(TC("   -getcas                 Will print hash of application"));
 		logger.Info(TC("   -summary                Print summary at the end of a session"));
 		logger.Info(TC("   -nocustomalloc          Disable custom allocator for processes. If you see odd crashes this can be tested"));
+		logger.Info(TC("   -nostdout               Disable stdout from process."));
 		logger.Info(TC("   -storeraw               Disable compression of storage. This will use more storage and might improve performance"));
 		if (IsWindows)
 			logger.Info(TC("   -visualizer             Spawn a visualizer that visualizes progress"));
@@ -133,6 +134,7 @@ namespace uba
 		bool checkCas = false;
 		bool checkAws = false;
 		bool getCas = false;
+		bool enableStdOut = true;
 		bool printSummary = false;
 		u32 loopCount = 1;
 
@@ -237,6 +239,10 @@ namespace uba
 			{
 				disableCustomAllocator = true;
 			}
+			else if (name.Equals(TC("-nostdout")))
+			{
+				enableStdOut = false;
+			}
 			else if (name.Equals(TC("-checkcas")))
 			{
 				checkCas = true;
@@ -426,14 +432,15 @@ namespace uba
 			pinfo.workingDir = workDir.data;
 			pinfo.logFile = logFile.data;
 			pinfo.logLineUserData = &logger;
-			pinfo.logLineFunc = [](void* userData, const tchar* line, u32 length, LogEntryType type) { ((Logger*)userData)->Log(type, line, length); };
+			if (enableStdOut)
+				pinfo.logLineFunc = [](void* userData, const tchar* line, u32 length, LogEntryType type) { ((Logger*)userData)->Log(type, line, length); };
 			pinfo.trackInputs = trackInputs;
 			logger.Info(TC("Running %s %s"), app.c_str(), arg.c_str());
 			ProcessHandle process = session->RunProcess(pinfo, false, enableDetour);
 			if (process.GetExitCode() != 0)
 				return logger.Error(TC("Error exit code: %u"), process.GetExitCode());
 			u64 time = GetTime() - start;
-			logger.Info(TC("Boxed run took %s"), TimeToText(time).str);
+			logger.Info(TC("%s run took %s"), (enableDetour ? TC("Boxed") : TC("Native")), TimeToText(time).str);
 			return true;
 		};
 
@@ -446,7 +453,8 @@ namespace uba
 			pinfo.workingDir = workDir.data;
 			pinfo.logFile = logFile.data;
 			pinfo.logLineUserData = &logger;
-			pinfo.logLineFunc = [](void* userData, const tchar* line, u32 length, LogEntryType type) { ((Logger*)userData)->Log(type, line, length); };
+			if (enableStdOut)
+				pinfo.logLineFunc = [](void* userData, const tchar* line, u32 length, LogEntryType type) { ((Logger*)userData)->Log(type, line, length); };
 			logger.Info(TC("Running %s %s"), app.c_str(), arg.c_str());
 			ProcessHandle process = session->RunProcessRemote(pinfo);
 			process.WaitForExit(~0u);
diff --git a/Engine/Source/Programs/UnrealBuildAccelerator/Common/Private/UbaProcess.cpp b/Engine/Source/Programs/UnrealBuildAccelerator/Common/Private/UbaProcess.cpp
--- a/Engine/Source/Programs/UnrealBuildAccelerator/Common/Private/UbaProcess.cpp
+++ b/Engine/Source/Programs/UnrealBuildAccelerator/Common/Private/UbaProcess.cpp
@@ -49,6 +49,41 @@ namespace uba
 			delete this;
 	}
 
+	struct ProcessImpl::PipeReader
+	{
+		PipeReader(ProcessImpl& p, LogEntryType lt) : process(p), logType(lt) {}
+		~PipeReader()
+		{
+			if (!currentString.empty())
+				process.LogLine(false, TString(currentString), logType);
+		}
+
+		void ReadData(char* buf, u32 readCount)
+		{
+			char* startPos = buf;
+			while (true)
+			{
+				char* endOfLine = strchr(startPos, '\n');
+				if (!endOfLine)
+				{
+					currentString.append(TString(startPos, startPos + strlen(startPos)));
+					return;
+				}
+				char* newStart = endOfLine + 1;
+				if (endOfLine > buf && endOfLine[-1] == '\r')
+					--endOfLine;
+				currentString.append(TString(startPos, endOfLine));
+				process.LogLine(false, TString(currentString), logType);
+				currentString.clear();
+				startPos = newStart;
+			}
+		}
+
+		ProcessImpl& process;
+		LogEntryType logType;
+		TString currentString;
+	};
+
 	ProcessImpl::ProcessImpl(Session& session, u32 id, ProcessImpl* parent)
 	:	m_session(session)
 	,	m_parentProcess(parent)
@@ -297,12 +332,17 @@ namespace uba
 		#endif
 	}
 
-	bool ProcessImpl::WaitForRead()
+	bool ProcessImpl::WaitForRead(PipeReader& outReader, PipeReader& errReader)
 	{
 		while (true)
 		{
 			if (m_readEvent.IsSet(1000))
 				break;
+
+			#if !PLATFORM_WINDOWS
+			PollStdPipes(outReader, errReader, 0);
+			#endif
+
 			if (!IsActive())
 				return false;
 			if (IsCancelled())
@@ -339,17 +379,27 @@ namespace uba
 		{
 			exitCode = InternalCreateProcess(runningRemote, environment, m_comMemory.handle, m_comMemory.offset);
 
-			bool loop = exitCode == 0 && m_detourEnabled;
-
-			while (loop && WaitForRead())
+			if (exitCode == 0)
 			{
-				u64 startTime = GetTime();
-				BinaryReader reader(readMemory);
-				BinaryWriter writer(writeMemory);
-				loop = HandleMessage(reader, writer);
-				SetWritten();
-				m_processStats.hostTotalTime += GetTime() - startTime;
-				++m_messageCount;
+				PipeReader outReader(*this, LogEntryType_Info);
+				PipeReader errReader(*this, LogEntryType_Error);
+
+				bool loop = m_detourEnabled;
+				while (loop && WaitForRead(outReader, errReader))
+				{
+					u64 startTime = GetTime();
+					BinaryReader reader(readMemory);
+					BinaryWriter writer(writeMemory);
+					loop = HandleMessage(reader, writer);
+					SetWritten();
+					m_processStats.hostTotalTime += GetTime() - startTime;
+					++m_messageCount;
+				}
+
+				#if !PLATFORM_WINDOWS
+				while (PollStdPipes(outReader, errReader))
+					continue;
+				#endif
 			}
 
 			u64 exitStartTime = GetTime();
@@ -1047,41 +1097,6 @@ namespace uba
 		return temp.data;
 	}
 
-	struct ProcessImpl::PipeReader
-	{
-		PipeReader(ProcessImpl& p, LogEntryType lt) : process(p), logType(lt) {}
-		~PipeReader()
-		{
-			if (!currentString.empty())
-				process.LogLine(false, TString(currentString), logType);
-		}
-
-		void ReadData(char* buf, u32 readCount)
-		{
-			char* startPos = buf;
-			while (true)
-			{
-				char* endOfLine = strchr(startPos, '\n');
-				if (!endOfLine)
-				{
-					currentString.append(TString(startPos, startPos + strlen(startPos)));
-					return;
-				}
-				char* newStart = endOfLine + 1;
-				if (endOfLine > buf && endOfLine[-1] == '\r')
-					--endOfLine;
-				currentString.append(TString(startPos, endOfLine));
-				process.LogLine(false, TString(currentString), logType);
-				currentString.clear();
-				startPos = newStart;
-			}
-		}
-
-		ProcessImpl& process;
-		LogEntryType logType;
-		TString currentString;
-	};
-
 	u32 ProcessImpl::InternalCreateProcess(bool runningRemote, void* environment, FileMappingHandle communicationHandle, u64 communicationOffset)
 	{
 		ScopedWriteLock initLock(m_initLock);
@@ -1290,6 +1305,7 @@ namespace uba
 			payload.communicationHandle = communicationHandle.handle;
 			payload.communicationOffset = communicationOffset;
 			payload.rulesIndex = m_rulesIndex;
+			payload.version = ProcessMessageVersion;
 			payload.runningRemote = runningRemote;
 			payload.isChild = m_parentProcess != nullptr;
 			payload.trackInputs = m_startInfo.trackInputs;
@@ -1473,7 +1489,10 @@ namespace uba
 				if (!logFile.IsEmpty())
 					envvars.push_back(logFile.data);
 			}
-			else
+
+			envvars.push_back(nullptr);
+
+			if (true)
 			{
 				if (pipe(outPipe) || pipe(errPipe))
 				{
@@ -1500,8 +1519,6 @@ namespace uba
 				UBA_ASSERTF(!res, "posix_spawn_file_actions_addclose errPipe[1] failed: %i", res);
 			}
 
-			envvars.push_back(nullptr);
-
 			u32 retryCount = 0;
 			pid_t processID;
 			while (true)
@@ -1538,62 +1555,11 @@ namespace uba
 			m_nativeProcessHandle = (ProcHandle)1;
 			m_nativeProcessId = u32(processID);
 
-			if (!m_detourEnabled)
-			{
-				pipeGuard1.Execute();
-
-				PipeReader outReader(*this, LogEntryType_Info);
-				PipeReader errReader(*this, LogEntryType_Error);
-
-				pollfd plist[] = { {outPipe[0],POLLIN, 0}, {errPipe[0],POLLIN, 0} };
-				int rval = 0;
-				for (; (rval = poll(plist, sizeof_array(plist), -1)) > 0;)
-				{
-					if (plist[0].revents & POLLERR || plist[1].revents & POLLERR) // If there is an error on any of them we hang up
-					{
-						logger.Error(TC("pipe polling error"));
-						break;
-					}
-					int fd = 0;
-					PipeReader* pipeReader = nullptr;
-					if (plist[0].revents & POLLIN)
-					{
-						fd = outPipe[0];
-						pipeReader = &outReader;
-					}
-					else if (plist[1].revents & POLLIN)
-					{
-						fd = errPipe[0];
-						pipeReader = &errReader;
-					}
-					else
-						break; // nothing left to read
-
-					char buffer[1024];
-					int bytesRead = read(fd, buffer, sizeof_array(buffer) - 1);
-
-					buffer[bytesRead] = 0;
-					pipeReader->ReadData(buffer, bytesRead);
+			pipeGuard1.Execute();
 
-					// If they both have hung up we hang up
-					// However, pollhup means slightly different things on MacOS vs Linux
-					// so we need to check slightly different things
-					if (plist[0].revents & POLLHUP && plist[1].revents & POLLHUP)
-					{
-						#if PLATFORM_LINUX
-						if (bytesRead == 0)
-						#endif
-							break;
-					}
-				}
-
-				#if PLATFORM_MAC
-				if (rval == -1)
-					logger.Error(TC("pipe polling error with -1 (%s)"), strerror(errno));
-				// Now that we're done polling, close the fds.
-				pipeGuard0.Execute();
-				#endif
-			}
+			m_stdOutPipe = outPipe[0];
+			m_stdErrPipe = errPipe[0];
+			pipeGuard0.Cancel();
 		}
 		else
 		{
@@ -1700,6 +1666,11 @@ namespace uba
 		return UBA_EXIT_CODE(14);
 #else
 
+		if (m_stdOutPipe != -1)
+			close(m_stdOutPipe);
+		if (m_stdErrPipe != -1)
+			close(m_stdErrPipe);
+
 		auto g = MakeGuard([this]() { m_nativeProcessId = 0; });
 
 		if (cancel)
@@ -1740,6 +1711,52 @@ namespace uba
 #endif
 	}
 
+#if !PLATFORM_WINDOWS
+	bool ProcessImpl::PollStdPipes(PipeReader& outReader, PipeReader& errReader, int timeoutMs)
+	{
+		if (m_stdOutPipe == -1)
+			return false;
+
+		auto pipeGuard = MakeGuard([&]() { close(m_stdOutPipe); m_stdOutPipe = -1; close(m_stdErrPipe); m_stdErrPipe = -1; });
+
+		PipeReader* pipeReaders[] = { &outReader, &errReader };
+
+		pollfd plist[] = { {m_stdOutPipe,POLLIN, 0}, {m_stdErrPipe,POLLIN, 0} };
+		int rval = poll(plist, sizeof_array(plist), timeoutMs);
+		if (rval < 0)
+		{
+			#if PLATFORM_MAC
+			m_session.m_logger.Error(TC("pipe polling error with %i (%s)"), rval, strerror(errno));
+			#endif
+			return false;
+		}
+
+		bool hasRead = false;
+		for (int i=0;i!=2;++i)
+		{
+			if (plist[i].revents & POLLERR) // If there is an error on any of them we hang up
+				return m_session.m_logger.Error(TC("pipe polling error"));
+
+			if (!(plist[i].revents & POLLIN))
+				continue;
+
+			char buffer[1024];
+			if (int bytesRead = read(plist[i].fd, buffer, sizeof_array(buffer) - 1))
+			{
+				hasRead = true;
+				buffer[bytesRead] = 0;
+				pipeReaders[i]->ReadData(buffer, bytesRead);
+			}
+		}
+
+		if (!hasRead)
+			if (plist[0].revents & POLLHUP && plist[1].revents & POLLHUP)
+				return false;
+		pipeGuard.Cancel();
+		return true;
+	}
+#endif
+
 	void ProcessImpl::ClearTempFiles()
 	{
 		for (auto& pair : m_tempFiles)
diff --git a/Engine/Source/Programs/UnrealBuildAccelerator/Common/Private/UbaProcessStartInfoHolder.h b/Engine/Source/Programs/UnrealBuildAccelerator/Common/Private/UbaProcessStartInfoHolder.h
--- a/Engine/Source/Programs/UnrealBuildAccelerator/Common/Private/UbaProcessStartInfoHolder.h
+++ b/Engine/Source/Programs/UnrealBuildAccelerator/Common/Private/UbaProcessStartInfoHolder.h
@@ -16,7 +16,7 @@ namespace uba
 		u64 workingDirLen = 0;
 		if (workingDir)
 			workingDirLen = TStrlen(workingDir);
-		FixPath2(fileName, workingDir, workingDirLen, buffer, &charLen);
+		FixPath2(fileName, workingDir, workingDirLen, buffer, sizeof_array(buffer), &charLen);
 		out.Append(buffer);
 	}
 
diff --git a/Engine/Source/Programs/UnrealBuildAccelerator/Common/Private/UbaSession.cpp b/Engine/Source/Programs/UnrealBuildAccelerator/Common/Private/UbaSession.cpp
--- a/Engine/Source/Programs/UnrealBuildAccelerator/Common/Private/UbaSession.cpp
+++ b/Engine/Source/Programs/UnrealBuildAccelerator/Common/Private/UbaSession.cpp
@@ -1216,7 +1216,7 @@ namespace uba
 		UBA_ASSERT(!m_runningRemote);
 
 		StringBuffer<> dirPath;
-		FixPath2(dirName, nullptr, 0, dirPath.data, &dirPath.count);
+		FixPath2(dirName, nullptr, 0, dirPath.data, dirPath.capacity, &dirPath.count);
 		StringKey dirKey = CaseInsensitiveFs ? ToStringKeyLower(dirPath) : ToStringKey(dirPath);
 
 		auto& dirTable = m_directoryTable;
@@ -1246,7 +1246,7 @@ namespace uba
 
 	StringKey GetKeyAndFixedName(StringBuffer<>& fixedFilePath, const tchar* filePath)
 	{
-		FixPath2(filePath, nullptr, 0, fixedFilePath.data, &fixedFilePath.count);
+		FixPath2(filePath, nullptr, 0, fixedFilePath.data, fixedFilePath.capacity, &fixedFilePath.count);
 
 		StringKey dirKey;
 		StringBuffer<> dirNameForHash;
diff --git a/Engine/Source/Programs/UnrealBuildAccelerator/Common/Public/UbaProcess.h b/Engine/Source/Programs/UnrealBuildAccelerator/Common/Public/UbaProcess.h
--- a/Engine/Source/Programs/UnrealBuildAccelerator/Common/Public/UbaProcess.h
+++ b/Engine/Source/Programs/UnrealBuildAccelerator/Common/Public/UbaProcess.h
@@ -53,12 +53,14 @@ namespace uba
 		ProcessImpl(Session& session, u32 id, ProcessImpl* parent);
 		~ProcessImpl();
 
+		struct PipeReader;
+
 		void Start(const ProcessStartInfo& startInfo, TString&& realApplication, const tchar* realWorkingDir, bool runningRemote, void* environment, bool async, bool enableDetour);
 
 		bool IsActive();
 		bool IsCancelled();
 
-		bool WaitForRead();
+		bool WaitForRead(PipeReader& outReader, PipeReader& errReader);
 		void SetWritten();
 
 		void ThreadRun(bool runningRemote, void* environment);
@@ -71,8 +73,8 @@ namespace uba
 		void SetRulesIndex(const ProcessStartInfo& si);
 		const tchar* InternalGetChildLogFile(StringBufferBase& temp);
 		
-		#if PLATFORM_MAC
-		int WaitForProcessGroup(pid_t pgid);
+		#if !PLATFORM_WINDOWS
+		bool PollStdPipes(PipeReader& outReader, PipeReader& errReader, int timeoutMs = -1);
 		#endif
 
 		u32 InternalCreateProcess(bool runningRemote, void* environment, FileMappingHandle communicationHandle, u64 communicationOffset);
@@ -96,6 +98,10 @@ namespace uba
 		Event& m_cancelEvent;
 		Event& m_writeEvent;
 		Event& m_readEvent;
+
+		int m_stdOutPipe = -1;
+		int m_stdErrPipe = -1;
+
 	#endif
 
 		ProcHandle m_nativeProcessHandle = InvalidProcHandle;
@@ -143,7 +149,6 @@ namespace uba
 
 		ProcessImpl(const ProcessImpl&) = delete;
 		void operator=(const ProcessImpl&) = delete;
-		struct PipeReader;
 	};
 
 	bool ParseArguments(Vector<TString>& outArguments, const tchar* argumentString);
diff --git a/Engine/Source/Programs/UnrealBuildAccelerator/Core/Public/UbaDetoursPayload.h b/Engine/Source/Programs/UnrealBuildAccelerator/Core/Public/UbaDetoursPayload.h
--- a/Engine/Source/Programs/UnrealBuildAccelerator/Core/Public/UbaDetoursPayload.h
+++ b/Engine/Source/Programs/UnrealBuildAccelerator/Core/Public/UbaDetoursPayload.h
@@ -16,6 +16,7 @@ namespace uba
 		HANDLE readEvent = 0;
 		HANDLE communicationHandle = 0;
 		u64 communicationOffset = 0;
+		u32 version = 0;
 		u32 rulesIndex = 0;
 		bool runningRemote = false;
 		bool trackInputs = false;
diff --git a/Engine/Source/Programs/UnrealBuildAccelerator/Core/Public/UbaPathUtils.h b/Engine/Source/Programs/UnrealBuildAccelerator/Core/Public/UbaPathUtils.h
--- a/Engine/Source/Programs/UnrealBuildAccelerator/Core/Public/UbaPathUtils.h
+++ b/Engine/Source/Programs/UnrealBuildAccelerator/Core/Public/UbaPathUtils.h
@@ -7,7 +7,7 @@
 
 namespace uba
 {
-	inline bool FixPath2(const tchar* fileName, const tchar* workingDir, u64 workingDirCharLen, tchar* buffer, u32* bufferCharLen)
+	inline bool FixPath2(const tchar* fileName, const tchar* workingDir, u64 workingDirCharLen, tchar* buffer, u64 bufferCharCapacity, u32* outBufferCharLen)
 	{
 		buffer[1] = 0;
 		UBA_ASSERTF(workingDir == nullptr || !*workingDir || workingDir[workingDirCharLen-1] == PathSeparator, TC("WorkingDir needs to end with path separator"));
@@ -38,8 +38,8 @@ namespace uba
 			u32 len = Local_GetLongPathNameW(fullName, fullName, sizeof_array(fullName));
 			if (!len)
 			{
-				if (bufferCharLen)
-					*bufferCharLen = 0;
+				if (outBufferCharLen)
+					*outBufferCharLen = 0;
 				return false;
 			}
 			if (backslash)
@@ -91,7 +91,8 @@ namespace uba
 
 		if (lastChar == '.' && lastLastChar == 0) // Sometimes path is '.'
 		{
-			UBA_ASSERT(workingDir && *workingDir);
+			UBA_ASSERTF(workingDir && *workingDir, TC("Working dir is null or empty"));
+			UBA_ASSERTF(workingDirCharLen < bufferCharCapacity, TC("%llu < %llu"), workingDirCharLen, bufferCharCapacity);
 			memcpy(buffer, workingDir, workingDirCharLen*sizeof(tchar));
 			buffer[workingDirCharLen - 1] = 0;
 			charLen = workingDirCharLen;
@@ -111,9 +112,11 @@ namespace uba
 
 			UBA_ASSERTF(workingDir && *workingDir, TC("No working dir provided but path is relative (%s)"), buffer);
 			tchar temp2[1024];
+			UBA_ASSERTF(workingDirCharLen + charLen < sizeof_array(temp2), TC("%llu + %llu < %llu"), workingDirCharLen, charLen, sizeof_array(temp2));
 			memcpy(temp2, workingDir, workingDirCharLen*sizeof(tchar));
 			memcpy(temp2 + workingDirCharLen, copyFrom, charLen*sizeof(tchar));
 			charLen += workingDirCharLen;
+			UBA_ASSERTF(charLen+1 <= bufferCharCapacity, TC("%llu+1 <= %llu"), charLen, bufferCharCapacity);
 			memcpy(buffer, temp2, (charLen+1)*sizeof(tchar));
 		}
 		else if (lastChar == '.' && charLen == 4) // X:.  .. this expands to X:\ unless working dir matches drive, then it becomes working dir
@@ -199,8 +202,10 @@ namespace uba
 			++charLen;
 		}
 
-		if (bufferCharLen)
-			*bufferCharLen = u32(charLen - 1); // Remove terminator
+		UBA_ASSERT(charLen <= bufferCharCapacity);
+
+		if (outBufferCharLen)
+			*outBufferCharLen = u32(charLen - 1); // Remove terminator
 		return true;
 #else
 		StringBuffer<MaxPath> tmp;
@@ -287,14 +292,14 @@ namespace uba
 		}
 		buffer[memPos] = 0;
 
-		if (bufferCharLen)
-			*bufferCharLen = memPos;
+		if (outBufferCharLen)
+			*outBufferCharLen = memPos;
 		return true;
 #endif
 	}
 
 	inline void FixPath(const tchar* fileName, const tchar* workingDir, u64 workingDirCharLen, StringBufferBase& buffer)
 	{
-		FixPath2(fileName, workingDir, workingDirCharLen, buffer.data, &buffer.count);
+		FixPath2(fileName, workingDir, workingDirCharLen, buffer.data, buffer.capacity, &buffer.count);
 	}
 }
\ No newline at end of file
diff --git a/Engine/Source/Programs/UnrealBuildAccelerator/Core/Public/UbaPlatform.h b/Engine/Source/Programs/UnrealBuildAccelerator/Core/Public/UbaPlatform.h
--- a/Engine/Source/Programs/UnrealBuildAccelerator/Core/Public/UbaPlatform.h
+++ b/Engine/Source/Programs/UnrealBuildAccelerator/Core/Public/UbaPlatform.h
@@ -56,7 +56,7 @@ namespace uba
 	#define	sizeof_array(array) int(sizeof(array)/sizeof(array[0]))
 	#endif
 
-	#if UBA_DEBUG || PLATFORM_MAC
+	#if UBA_DEBUG
 	#define UBA_ASSERT(x) do { if (x) break; uba::UbaAssert(TC(""), __FILE__, __LINE__, #x, 543221); } while(false)
 	#define UBA_ASSERTF(x, ...) do { if (x) break; uba::StringBuffer<> _buf; _buf.Appendf(__VA_ARGS__); uba::UbaAssert(_buf.data, __FILE__, __LINE__, #x, 543221); } while(false)
 	#define UBA_ASSERT_MESSAGEBOX 0
diff --git a/Engine/Source/Programs/UnrealBuildAccelerator/Core/Public/UbaProtocol.h b/Engine/Source/Programs/UnrealBuildAccelerator/Core/Public/UbaProtocol.h
--- a/Engine/Source/Programs/UnrealBuildAccelerator/Core/Public/UbaProtocol.h
+++ b/Engine/Source/Programs/UnrealBuildAccelerator/Core/Public/UbaProtocol.h
@@ -35,6 +35,8 @@ namespace uba
 		MessageType_Custom,
 	};
 
+	inline constexpr u32 ProcessMessageVersion = 1339;
+
 	inline constexpr u32 CommunicationMemSize = 64*1024;
 
 	inline constexpr u32 FileMappingTableMemSize = 16 * 1024 * 1024;
diff --git a/Engine/Source/Programs/UnrealBuildAccelerator/Detours/Private/PosixOS/UbaDetoursFunctionsPosix.cpp b/Engine/Source/Programs/UnrealBuildAccelerator/Detours/Private/PosixOS/UbaDetoursFunctionsPosix.cpp
--- a/Engine/Source/Programs/UnrealBuildAccelerator/Detours/Private/PosixOS/UbaDetoursFunctionsPosix.cpp
+++ b/Engine/Source/Programs/UnrealBuildAccelerator/Detours/Private/PosixOS/UbaDetoursFunctionsPosix.cpp
@@ -46,7 +46,6 @@ using namespace uba;
 	DETOURED_FUNCTION(read) \
 	DETOURED_FUNCTION(pread) \
 	DETOURED_FUNCTION(open) \
-	DETOURED_FUNCTION(write) \
 	DETOURED_FUNCTION(dup) \
 	DETOURED_FUNCTION(dup2) \
 	DETOURED_FUNCTION(close) \
@@ -54,7 +53,6 @@ using namespace uba;
 	DETOURED_FUNCTION(fdopen) \
 	DETOURED_FUNCTION(fchmod) \
 	DETOURED_FUNCTION(fchmodat) \
-	DETOURED_FUNCTION(fwrite) \
 	DETOURED_FUNCTION(fstat) \
 	DETOURED_FUNCTION(futimens) \
 	DETOURED_FUNCTION(fclose) \
@@ -98,6 +96,9 @@ using namespace uba;
 	DETOURED_FUNCTION_LINUX \
 	DETOURED_FUNCTION_MACOS \
 
+	//DETOURED_FUNCTION(write)
+	//DETOURED_FUNCTION(fwrite)
+
 #if UBA_DEBUG && PLATFORM_LINUX
 #define DETOURED_FUNCTION_DEBUG \
 
@@ -857,17 +858,17 @@ UBA_EXPORT int UBA_WRAPPER(open)(const char* file, int flags, ...)
 	return Shared_open("open", file, flags, mode, [](const char* realFile, int flags, int mode) { return TRUE_WRAPPER(open)(realFile, flags, mode); });
 }
 
-UBA_EXPORT ssize_t UBA_WRAPPER(write)(int fd, const void* buf, size_t count)
-{
-	UBA_INIT_DETOUR(write, fd, buf, count);
-	if (isatty(fd)) // stdout and stderr
-	{
-		Shared_WriteConsole((const char*)buf, count, fd == fileno(stderr));
-		return count;
-	}
-	//DEBUG_LOG_TRUE("write", "(%i size: %llu)", fd, count);
-	return TRUE_WRAPPER(write)(fd, buf, count);
-}
+//UBA_EXPORT ssize_t UBA_WRAPPER(write)(int fd, const void* buf, size_t count)
+//{
+//	UBA_INIT_DETOUR(write, fd, buf, count);
+//	if (isatty(fd)) // stdout and stderr
+//	{
+//		Shared_WriteConsole((const char*)buf, count, fd == fileno(stderr));
+//		return count;
+//	}
+//	//DEBUG_LOG_TRUE("write", "(%i size: %llu)", fd, count);
+//	return TRUE_WRAPPER(write)(fd, buf, count);
+//}
 
 //UBA_EXPORT int fcntl(int __fd, int __cmd, ...)
 //{
@@ -952,14 +953,14 @@ UBA_EXPORT int UBA_WRAPPER(fchmodat)(int dirfd, const char* pathname, mode_t mod
 	return TRUE_WRAPPER(fchmodat)(dirfd, pathname, mode, flags);
 }
 
-UBA_EXPORT size_t UBA_WRAPPER(fwrite)(const void* ptr, size_t size, size_t nitems, FILE* stream)
-{
-	UBA_INIT_DETOUR(fwrite, ptr, size, nitems, stream);
-	auto res = TRUE_WRAPPER(fwrite)(ptr, size, nitems, stream);
-	//if (stream != stdout)
-	//	DEBUG_LOG_TRUE("fwrite", "(%i) %i", fileno(stream), int(res));
-	return res;
-}
+//UBA_EXPORT size_t UBA_WRAPPER(fwrite)(const void* ptr, size_t size, size_t nitems, FILE* stream)
+//{
+//	UBA_INIT_DETOUR(fwrite, ptr, size, nitems, stream);
+//	auto res = TRUE_WRAPPER(fwrite)(ptr, size, nitems, stream);
+//	//if (stream != stdout)
+//	//	DEBUG_LOG_TRUE("fwrite", "(%i) %i", fileno(stream), int(res));
+//	return res;
+//}
 
 UBA_EXPORT int UBA_WRAPPER(fstat)(int fd, struct stat* buf)
 {
@@ -1974,7 +1975,8 @@ namespace uba
 	void WriteDebug(const char* str, u32 strLen)
 	{
 		int t = errno;
-		TRUE_WRAPPER(write)(g_debugFile, str, strLen);
+		//TRUE_WRAPPER(write)(g_debugFile, str, strLen);
+		write(g_debugFile, str, strLen);
 		//fsync(g_debugFile);
 		errno = t;
 	}
diff --git a/Engine/Source/Programs/UnrealBuildAccelerator/Detours/Private/UbaDetoursShared.cpp b/Engine/Source/Programs/UnrealBuildAccelerator/Detours/Private/UbaDetoursShared.cpp
--- a/Engine/Source/Programs/UnrealBuildAccelerator/Detours/Private/UbaDetoursShared.cpp
+++ b/Engine/Source/Programs/UnrealBuildAccelerator/Detours/Private/UbaDetoursShared.cpp
@@ -122,7 +122,7 @@ namespace uba
 
 	bool FixPath(StringBufferBase& out, const tchar* path)
 	{
-		return FixPath2(path, g_virtualWorkingDir.data, g_virtualWorkingDir.count, out.data, &out.count);
+		return FixPath2(path, g_virtualWorkingDir.data, g_virtualWorkingDir.count, out.data, out.capacity, &out.count);
 	}
 
 	const tchar* GetApplicationShortName()
diff --git a/Engine/Source/Programs/UnrealBuildAccelerator/Detours/Private/Windows/UbaDetoursFunctionsKernelBase.inl b/Engine/Source/Programs/UnrealBuildAccelerator/Detours/Private/Windows/UbaDetoursFunctionsKernelBase.inl
--- a/Engine/Source/Programs/UnrealBuildAccelerator/Detours/Private/Windows/UbaDetoursFunctionsKernelBase.inl
+++ b/Engine/Source/Programs/UnrealBuildAccelerator/Detours/Private/Windows/UbaDetoursFunctionsKernelBase.inl
@@ -81,7 +81,7 @@ void Shared_SetCurrentDirectory(const wchar_t* workingDirBuffer)
 {
 	u32 charLen = 0;
 	wchar_t temp[256];
-	FixPath2(workingDirBuffer, g_virtualWorkingDir.data, g_virtualWorkingDir.count, temp, &charLen);
+	FixPath2(workingDirBuffer, g_virtualWorkingDir.data, g_virtualWorkingDir.count, temp, sizeof_array(temp), &charLen);
 	g_virtualWorkingDir.Clear().Append(temp).Append('\\');
 }
 
@@ -2211,7 +2211,7 @@ DWORD Detoured_GetFinalPathNameByHandleW(HANDLE hFile, LPTSTR lpszFilePath, DWOR
 			if (!fo->newName.empty())
 				fileName = fo->newName.c_str();
 			StringBuffer<> buffer;
-			FixPath2(fileName, g_virtualWorkingDir.data, g_virtualWorkingDir.count, buffer.data, &buffer.count);
+			FixPath2(fileName, g_virtualWorkingDir.data, g_virtualWorkingDir.count, buffer.data, buffer.capacity, &buffer.count);
 
 			if (cchFilePath <= buffer.count)
 			{
diff --git a/Engine/Source/Programs/UnrealBuildAccelerator/Detours/Private/Windows/UbaDetoursFunctionsUcrtBase.inl b/Engine/Source/Programs/UnrealBuildAccelerator/Detours/Private/Windows/UbaDetoursFunctionsUcrtBase.inl
--- a/Engine/Source/Programs/UnrealBuildAccelerator/Detours/Private/Windows/UbaDetoursFunctionsUcrtBase.inl
+++ b/Engine/Source/Programs/UnrealBuildAccelerator/Detours/Private/Windows/UbaDetoursFunctionsUcrtBase.inl
@@ -35,12 +35,12 @@ wchar_t* Detoured__wfullpath(wchar_t* absPath, const wchar_t* relPath, size_t ma
 			res = (wchar_t*)malloc(maxLength * sizeof(WCHAR)); // Same as wine implementation
 		else
 			res = (wchar_t*)_malloc_base(maxLength * sizeof(WCHAR));
-		FixPath2(relPath, g_virtualWorkingDir.data, g_virtualWorkingDir.count, res, nullptr);
+		FixPath2(relPath, g_virtualWorkingDir.data, g_virtualWorkingDir.count, res, maxLength, nullptr);
 		UBA_ASSERT(res && wcslen(res) < maxLength - 1);
 	}
 	else if ((relPath && relPath[1] == ':' && relLen < maxLength) || (relLen + g_virtualWorkingDir.count <= maxLength))
 	{
-		FixPath2(relPath, g_virtualWorkingDir.data, g_virtualWorkingDir.count, absPath, nullptr);
+		FixPath2(relPath, g_virtualWorkingDir.data, g_virtualWorkingDir.count, absPath, maxLength, nullptr);
 		UBA_ASSERT(wcslen(absPath) < maxLength);
 		res = absPath;
 	}
diff --git a/Engine/Source/Programs/UnrealBuildAccelerator/Detours/Private/Windows/UbaDetoursFunctionsWin.cpp b/Engine/Source/Programs/UnrealBuildAccelerator/Detours/Private/Windows/UbaDetoursFunctionsWin.cpp
--- a/Engine/Source/Programs/UnrealBuildAccelerator/Detours/Private/Windows/UbaDetoursFunctionsWin.cpp
+++ b/Engine/Source/Programs/UnrealBuildAccelerator/Detours/Private/Windows/UbaDetoursFunctionsWin.cpp
@@ -866,8 +866,8 @@ void Init(const DetoursPayload& payload, u64 startTime)
 	GetEnvironmentVariableW(L"TEMP", systemTemp, 256);
 	FixPath(g_systemTemp, systemTemp);
 
-	StringBuffer<256> applicationBuffer;
-	StringBuffer<256> workingDirBuffer;
+	StringBuffer<512> applicationBuffer;
+	StringBuffer<512> workingDirBuffer;
 
 	HANDLE directoryTableHandle;
 	u32 directoryTableSize;
@@ -905,11 +905,10 @@ void Init(const DetoursPayload& payload, u64 startTime)
 	{
 		FixPath(applicationBuffer.data, g_virtualWorkingDir.data, g_virtualWorkingDir.count, g_virtualApplication);
 
-		static wchar_t applicationDir[MaxPath];
-		const wchar_t* lastBackslash = wcsrchr(g_virtualApplication.data, '\\') + 1;
-		memcpy(applicationDir, g_virtualApplication.data, (lastBackslash - g_virtualApplication.data)*2);
-		applicationDir[lastBackslash - g_virtualApplication.data] = 0;
-		g_virtualApplicationDir.Append(applicationDir);
+		if (const wchar_t* lastBackslash = g_virtualApplication.Last('\\'))
+			g_virtualApplicationDir.Append(g_virtualApplication.data, (lastBackslash + 1 - g_virtualApplication.data));
+		else
+			FatalError(4444, L"What the heck: %s", g_virtualApplication.data);
 	}
 
 	const wchar_t* cmdLine = True_GetCommandLineW();
diff --git a/Engine/Source/Programs/UnrealBuildAccelerator/Detours/Private/Windows/UbaDetoursMainWin.cpp b/Engine/Source/Programs/UnrealBuildAccelerator/Detours/Private/Windows/UbaDetoursMainWin.cpp
--- a/Engine/Source/Programs/UnrealBuildAccelerator/Detours/Private/Windows/UbaDetoursMainWin.cpp
+++ b/Engine/Source/Programs/UnrealBuildAccelerator/Detours/Private/Windows/UbaDetoursMainWin.cpp
@@ -143,6 +143,8 @@ BOOL WINAPI DllMain(HINSTANCE hinst, DWORD dwReason, LPVOID reserved)
 		DetoursPayload* payload = (DetoursPayload*)DetourFindPayloadEx(DetoursPayloadGuid, nullptr);
 		if (!payload)
 			TerminateCurrentProcess(1342);
+		if (payload->version != ProcessMessageVersion)
+			TerminateCurrentProcess(1354);
 
 		PreInit(*payload);
 
diff --git a/Engine/Source/Programs/UnrealBuildAccelerator/Test/Private/UbaTestAll.h b/Engine/Source/Programs/UnrealBuildAccelerator/Test/Private/UbaTestAll.h
--- a/Engine/Source/Programs/UnrealBuildAccelerator/Test/Private/UbaTestAll.h
+++ b/Engine/Source/Programs/UnrealBuildAccelerator/Test/Private/UbaTestAll.h
@@ -16,6 +16,7 @@ namespace uba
 	#define UBA_EXTRA_TESTS \
 		UBA_TEST(TestMultipleDetouredProcesses) \
 		UBA_TEST(TestLogLines) \
+		UBA_TEST(TestLogLinesNoDetour) \
 		UBA_TEST(TestLocalSchedule) \
 		UBA_TEST(TestLocalScheduleReuse) \
 		UBA_TEST(TestDetouredTouch) \
@@ -69,7 +70,7 @@ namespace uba
 
 		logger.Info(TC("Running tests (Test rootdir: %s)"), testRootDir.data);
 
-		// UBA_TEST(TestRemoteDetouredTestApp)
+		//UBA_TEST(TestLogLines)
 		UBA_TESTS
 
 		logger.Info(TC("Tests finished successfully!"));
diff --git a/Engine/Source/Programs/UnrealBuildAccelerator/Test/Private/UbaTestBasics.h b/Engine/Source/Programs/UnrealBuildAccelerator/Test/Private/UbaTestBasics.h
--- a/Engine/Source/Programs/UnrealBuildAccelerator/Test/Private/UbaTestBasics.h
+++ b/Engine/Source/Programs/UnrealBuildAccelerator/Test/Private/UbaTestBasics.h
@@ -76,10 +76,10 @@ namespace uba
 		tchar buffer[1024];
 		u32 lengthResult;
 
-		auto TestPath = [&](const tchar* path) { return FixPath2(path, workingDir, TStrlen(workingDir), buffer, &lengthResult); };
+		auto TestPath = [&](const tchar* path) { return FixPath2(path, workingDir, TStrlen(workingDir), buffer, sizeof_array(buffer), &lengthResult); };
 
 #if PLATFORM_WINDOWS
-		if (!FixPath2(TC("\"e:\\temp\""), workingDir, TStrlen(workingDir), buffer, &lengthResult))
+		if (!FixPath2(TC("\"e:\\temp\""), workingDir, TStrlen(workingDir), buffer, sizeof_array(buffer), &lengthResult))
 			return logger.Error(TC("FixPath2 (1) failed"));
 #else
 
@@ -87,24 +87,24 @@ namespace uba
 			return logger.Error(TC("FixPath2 should have failed"));
 		UBA_TEST_CHECK(Equals(buffer, TC("/")), "Should not contain ..");
 
-		if (!FixPath2(TC("/../Foo"), workingDir, TStrlen(workingDir), buffer, &lengthResult))
+		if (!FixPath2(TC("/../Foo"), workingDir, TStrlen(workingDir), buffer, sizeof_array(buffer), &lengthResult))
 			return logger.Error(TC("FixPath2 should have failed"));
 		UBA_TEST_CHECK(Equals(buffer, TC("/Foo")), "Should not contain ..");
 
-		if (!FixPath2(TC("/usr/bin//clang++"), workingDir, TStrlen(workingDir), buffer, &lengthResult))
+		if (!FixPath2(TC("/usr/bin//clang++"), workingDir, TStrlen(workingDir), buffer, sizeof_array(buffer), &lengthResult))
 			return logger.Error(TC("FixPath2 should have failed"));
 		UBA_TEST_CHECK(!Contains(buffer, TC("//")), "Should not contain //");
 #endif
 
-		if (!FixPath2(TC("../Foo"), workingDir, TStrlen(workingDir), buffer, &lengthResult))
+		if (!FixPath2(TC("../Foo"), workingDir, TStrlen(workingDir), buffer, sizeof_array(buffer), &lengthResult))
 			return logger.Error(TC("FixPath2 (1) failed"));
 		UBA_TEST_CHECK(!Contains(buffer, TC("..")), "Should not contain ..");
 
-		if (!FixPath2(TC("@../Foo"), workingDir, TStrlen(workingDir), buffer, &lengthResult))
+		if (!FixPath2(TC("@../Foo"), workingDir, TStrlen(workingDir), buffer, sizeof_array(buffer), &lengthResult))
 			return logger.Error(TC("FixPath2 (1) failed"));
 		UBA_TEST_CHECK(Contains(buffer, TC("..")), "Should contain ..");
 
-		if (!FixPath2(TC("..@/Foo"), workingDir, TStrlen(workingDir), buffer, &lengthResult))
+		if (!FixPath2(TC("..@/Foo"), workingDir, TStrlen(workingDir), buffer, sizeof_array(buffer), &lengthResult))
 			return logger.Error(TC("FixPath2 (1) failed"));
 		UBA_TEST_CHECK(Contains(buffer, TC("..")), "Should contain ..");
 
diff --git a/Engine/Source/Programs/UnrealBuildAccelerator/Test/Private/UbaTestScheduler.h b/Engine/Source/Programs/UnrealBuildAccelerator/Test/Private/UbaTestScheduler.h
--- a/Engine/Source/Programs/UnrealBuildAccelerator/Test/Private/UbaTestScheduler.h
+++ b/Engine/Source/Programs/UnrealBuildAccelerator/Test/Private/UbaTestScheduler.h
@@ -9,9 +9,6 @@ namespace uba
 {
 	bool TestLocalSchedule(LoggerWithWriter& logger, const StringBufferBase& testRootDir)
 	{
-		if (!IsWindows) // TODO: Remove
-			return true;
-
 		return RunLocal(logger, testRootDir, [](LoggerWithWriter& logger, SessionServer& session, const tchar* workingDir, const RunProcessFunction& runProcess)
 			{
 
diff --git a/Engine/Source/Programs/UnrealBuildAccelerator/Test/Private/UbaTestSession.h b/Engine/Source/Programs/UnrealBuildAccelerator/Test/Private/UbaTestSession.h
--- a/Engine/Source/Programs/UnrealBuildAccelerator/Test/Private/UbaTestSession.h
+++ b/Engine/Source/Programs/UnrealBuildAccelerator/Test/Private/UbaTestSession.h
@@ -355,7 +355,7 @@ namespace uba
 		#if PLATFORM_WINDOWS
 		return TC("-n 1 localhost");
 		#else
-		return TC("");
+		return TC("-help");
 		#endif
 	}
 
@@ -370,9 +370,6 @@ namespace uba
 
 	bool TestMultipleDetouredProcesses(LoggerWithWriter& logger, const StringBufferBase& testRootDir)
 	{
-		if (!IsWindows) // TODO: Remove
-			return true;
-
 		return RunLocal(logger, testRootDir, [](LoggerWithWriter& logger, SessionServer& session, const tchar* workingDir, const RunProcessFunction& runProcess)
 			{
 				ProcessStartInfo processInfo;
@@ -398,35 +395,39 @@ namespace uba
 			});
 	}
 
-	bool TestLogLines(LoggerWithWriter& logger, const StringBufferBase& testRootDir)
+	bool RunSystemApplicationAndLookForLog(LoggerWithWriter& logger, SessionServer& session, const tchar* workingDir, const RunProcessFunction& runProcess)
 	{
-		if (!IsWindows) // TODO: Remove
-			return true;
+		ProcessStartInfo processInfo;
+		processInfo.application = GetSystemApplication();
+		processInfo.workingDir = workingDir;
+		processInfo.arguments = GetSystemArguments();
 
-		return RunLocal(logger, testRootDir, [](LoggerWithWriter& logger, SessionServer& session, const tchar* workingDir, const RunProcessFunction& runProcess)
+		bool foundPingString = false;
+		processInfo.logLineUserData = &foundPingString;
+		processInfo.logLineFunc = [](void* userData, const tchar* line, u32 length, LogEntryType type)
 			{
-				ProcessStartInfo processInfo;
-				processInfo.application = GetSystemApplication();
-				processInfo.workingDir = workingDir;
-				processInfo.arguments = GetSystemArguments();
+				*(bool*)userData |= Contains(line, GetSystemExpectedLogLine());
+			};
 
-				bool foundPingString = false;
-				processInfo.logLineUserData = &foundPingString;
-				processInfo.logLineFunc = [](void* userData, const tchar* line, u32 length, LogEntryType type)
-					{
-						*(bool*)userData |= Contains(line, GetSystemExpectedLogLine());
-					};
+		ProcessHandle process = runProcess(processInfo);
 
-				ProcessHandle process = runProcess(processInfo);
+		if (!process.WaitForExit(10000))
+			return logger.Error(TC("UbaTestApp did not exit in 10 seconds"));
+		u32 exitCode = process.GetExitCode();
+		if (exitCode != 0)
+			return logger.Error(TC("Got exit code %u"), exitCode);
+		if (!foundPingString)
+			return logger.Error(TC("Did not log string containing \"%s\""), GetSystemExpectedLogLine());
+		return true;
+	}
 
-				if (!process.WaitForExit(10000000))
-					return logger.Error(TC("UbaTestApp did not exit in 10 seconds"));
-				u32 exitCode = process.GetExitCode();
-				if (exitCode != 0)
-					return logger.Error(TC("Got exit code %u"), exitCode);
-				if (!foundPingString)
-					return logger.Error(TC("Did not log ping string"));
-				return true;
-			}, false);
+	bool TestLogLines(LoggerWithWriter& logger, const StringBufferBase& testRootDir)
+	{
+		return RunLocal(logger, testRootDir, RunSystemApplicationAndLookForLog);
+	}
+
+	bool TestLogLinesNoDetour(LoggerWithWriter& logger, const StringBufferBase& testRootDir)
+	{
+		return RunLocal(logger, testRootDir, RunSystemApplicationAndLookForLog, false);
 	}
 }
diff --git a/Engine/Source/Programs/UnrealBuildAccelerator/Test/Private/UbaTestStorage.h b/Engine/Source/Programs/UnrealBuildAccelerator/Test/Private/UbaTestStorage.h
--- a/Engine/Source/Programs/UnrealBuildAccelerator/Test/Private/UbaTestStorage.h
+++ b/Engine/Source/Programs/UnrealBuildAccelerator/Test/Private/UbaTestStorage.h
@@ -10,7 +10,7 @@ namespace uba
 	bool TestStorage(LoggerWithWriter& logger, const StringBufferBase& testRootDir)
 	{
 		#if PLATFORM_LINUX
-		if (true)
+		if (true) // TODO: Revisit this... fails on farm but works locally
 			return true;
 		#endif
 
