[UBA]
* Fixed so renaming from castemp to real destination handle cross-device

diff --git a/Engine/Source/Programs/UnrealBuildAccelerator/Common/Private/UbaFileAccessor.cpp b/Engine/Source/Programs/UnrealBuildAccelerator/Common/Private/UbaFileAccessor.cpp
--- a/Engine/Source/Programs/UnrealBuildAccelerator/Common/Private/UbaFileAccessor.cpp
+++ b/Engine/Source/Programs/UnrealBuildAccelerator/Common/Private/UbaFileAccessor.cpp
@@ -109,7 +109,7 @@ namespace uba
 		if (!CreateWrite(allowRead, flagsAndAttributes, size, tempPath))
 			return false;
 
-		m_mappingHandle = uba::CreateFileMappingW(m_fileHandle, PAGE_READWRITE, size);
+		m_mappingHandle = uba::CreateFileMappingW(m_fileHandle, PAGE_READWRITE, size, m_fileName);
 		if (!m_mappingHandle.IsValid())
 			return m_logger.Error(TC("Failed to create memory map %s (%s)"), m_fileName, LastErrorToText().data);
 
@@ -286,9 +286,9 @@ namespace uba
 		m_size = info.size;
 #if PLATFORM_WINDOWS
 		if (m_size)
-			m_mappingHandle = uba::CreateFileMappingW(m_fileHandle, PAGE_READONLY, m_size);
+			m_mappingHandle = uba::CreateFileMappingW(m_fileHandle, PAGE_READONLY, m_size, m_fileName);
 		else
-			m_mappingHandle = uba::CreateFileMappingW(InvalidFileHandle, PAGE_READONLY, 1);
+			m_mappingHandle = uba::CreateFileMappingW(InvalidFileHandle, PAGE_READONLY, 1, m_fileName);
 
 		if (!m_mappingHandle.IsValid())
 			return m_logger.Error(TC("Failed to create mapping handle for %s (%s)"), m_fileName, LastErrorToText().data);
@@ -351,12 +351,15 @@ namespace uba
 							return m_logger.Error(TC("Failed to create file %s for move from temporary file %s (%s)"), m_fileName, realFileName, strerror(errno));
 						
 						#if PLATFORM_MAC
-						if (fcopyfile(asFileDescriptor(m_fileHandle), targetFd, 0, COPYFILE_ALL) == -1) {
+						if (fcopyfile(asFileDescriptor(m_fileHandle), targetFd, 0, COPYFILE_ALL) == -1)
+							return m_logger.Error(TC("Failed to do fcopyfile from temporary %s to file %s (%s)"), realFileName, m_fileName, strerror(errno));
 						#else
-						if (sendfile(targetFd, asFileDescriptor(m_fileHandle), NULL, m_size) == -1) {
-						#endif
+						int sourceFd = asFileDescriptor(m_fileHandle);
+						if (lseek(sourceFd, 0, SEEK_SET) == -1)
+							return m_logger.Error(TC("Failed to do lseek to beginning for sendfile (%s)"), strerror(errno));
+						if (sendfile(targetFd, sourceFd, NULL, m_size) != m_size)
 							return m_logger.Error(TC("Failed to do sendfile from temporary %s to file %s (%s)"), realFileName, m_fileName, strerror(errno));
-						}
+						#endif
 
 						remove(realFileName); // Remove real file now when we have copied it over
 					}
