[UBA]
* Added more error handling of failing VirtualAlloc

diff --git a/Engine/Source/Programs/UnrealBuildAccelerator/Detours/Private/Windows/UbaDetoursFunctionsWin.cpp b/Engine/Source/Programs/UnrealBuildAccelerator/Detours/Private/Windows/UbaDetoursFunctionsWin.cpp
--- a/Engine/Source/Programs/UnrealBuildAccelerator/Detours/Private/Windows/UbaDetoursFunctionsWin.cpp
+++ b/Engine/Source/Programs/UnrealBuildAccelerator/Detours/Private/Windows/UbaDetoursFunctionsWin.cpp
@@ -141,6 +141,8 @@ struct MemoryFile
 		if (isLocalOnly)
 		{
 			baseAddress = (u8*)VirtualAlloc(NULL, reserveSize, MEM_RESERVE, PAGE_READWRITE);
+			if (!baseAddress)
+				FatalError(1354, L"VirtualAlloc failed trying to reserve %llu. (Error code: %u)", reserveSize, GetLastError());
 			mappedSize = reserveSize;
 		}
 		else
