[UBA]
* Added support for custom text in trace/visualizer... can be used to report things like horde status etc

diff --git a/Engine/Source/Programs/UnrealBuildAccelerator/Common/Private/UbaSessionServer.cpp b/Engine/Source/Programs/UnrealBuildAccelerator/Common/Private/UbaSessionServer.cpp
--- a/Engine/Source/Programs/UnrealBuildAccelerator/Common/Private/UbaSessionServer.cpp
+++ b/Engine/Source/Programs/UnrealBuildAccelerator/Common/Private/UbaSessionServer.cpp
@@ -428,6 +428,11 @@ namespace uba
 		m_trace.ProcessExited(id, exitCode, statsWriter.GetData(), statsWriter.GetPosition(), Vector<ProcessLogLine>());
 	}
 
+	void SessionServer::UpdateStatus(u32 statusIndex, u32 statusNameIndent, const tchar* statusName, u32 statusTextIndent, const tchar* statusText, LogEntryType statusType)
+	{
+		m_trace.StatusUpdate(statusIndex, statusNameIndent, statusName, statusTextIndent, statusText, statusType);
+	}
+
 	NetworkServer& SessionServer::GetServer()
 	{
 		return m_server;
diff --git a/Engine/Source/Programs/UnrealBuildAccelerator/Common/Private/UbaTrace.cpp b/Engine/Source/Programs/UnrealBuildAccelerator/Common/Private/UbaTrace.cpp
--- a/Engine/Source/Programs/UnrealBuildAccelerator/Common/Private/UbaTrace.cpp
+++ b/Engine/Source/Programs/UnrealBuildAccelerator/Common/Private/UbaTrace.cpp
@@ -288,6 +288,17 @@ namespace uba
 		writer.Write7BitEncoded(workIndex);
 	}
 
+	void Trace::StatusUpdate(u32 statusIndex, u32 statusNameIndent, const tchar* statusName, u32 statusTextIndent, const tchar* statusText, LogEntryType statusType)
+	{
+		BEGIN_TRACE_ENTRY(TraceType_StatusUpdate);
+		writer.Write7BitEncoded(statusIndex);
+		writer.Write7BitEncoded(statusNameIndent);
+		writer.WriteString(statusName);
+		writer.Write7BitEncoded(statusTextIndent);
+		writer.WriteString(statusText);
+		writer.WriteByte(statusType);
+	}
+
 	TraceChannel::TraceChannel(Logger& logger) : m_logger(logger)
 	{
 	}
diff --git a/Engine/Source/Programs/UnrealBuildAccelerator/Common/Public/UbaSessionServer.h b/Engine/Source/Programs/UnrealBuildAccelerator/Common/Public/UbaSessionServer.h
--- a/Engine/Source/Programs/UnrealBuildAccelerator/Common/Public/UbaSessionServer.h
+++ b/Engine/Source/Programs/UnrealBuildAccelerator/Common/Public/UbaSessionServer.h
@@ -56,6 +56,9 @@ namespace uba
 		// End external process.
 		void EndExternalProcess(u32 id, u32 exitCode);
 
+		// Add external status information to trace stream. Will show in visualizer
+		void UpdateStatus(u32 statusIndex, u32 statusNameIndent, const tchar* statusName, u32 statusTextIndent, const tchar* statusText, LogEntryType statusType);
+
 		// Get the network server used by this session
 		NetworkServer& GetServer();
 
diff --git a/Engine/Source/Programs/UnrealBuildAccelerator/Common/Public/UbaTrace.h b/Engine/Source/Programs/UnrealBuildAccelerator/Common/Public/UbaTrace.h
--- a/Engine/Source/Programs/UnrealBuildAccelerator/Common/Public/UbaTrace.h
+++ b/Engine/Source/Programs/UnrealBuildAccelerator/Common/Public/UbaTrace.h
@@ -48,9 +48,10 @@ namespace uba
 		TraceType_ProxyUsed,
 		TraceType_FileFetchLight,
 		TraceType_FileStoreLight,
+		TraceType_StatusUpdate,
 	};
 
-	static constexpr u32 TraceVersion = 20;
+	static constexpr u32 TraceVersion = 21;
 	static constexpr u32 TraceReadCompatibilityVersion = 6;
 
 	class Trace
@@ -78,6 +79,7 @@ namespace uba
 		void FileEndStore(u32 clientId, const CasKey& key);
 		void BeginWork(u32 workIndex, const tchar* desc);
 		void EndWork(u32 workIndex);
+		void StatusUpdate(u32 statusIndex, u32 statusNameIndent, const tchar* statusName, u32 statusTextIndent, const tchar* statusText, LogEntryType statusType);
 
 		bool StopWrite(const tchar* writeFileName);
 
diff --git a/Engine/Source/Programs/UnrealBuildAccelerator/Host/Private/UbaExports.cpp b/Engine/Source/Programs/UnrealBuildAccelerator/Host/Private/UbaExports.cpp
--- a/Engine/Source/Programs/UnrealBuildAccelerator/Host/Private/UbaExports.cpp
+++ b/Engine/Source/Programs/UnrealBuildAccelerator/Host/Private/UbaExports.cpp
@@ -406,6 +406,11 @@ uba::StorageClient* CreateStorageClient(uba::NetworkClient& client, const uba::t
 		server->EndExternalProcess(id, exitCode);
 	}
 
+	void SessionServer_UpdateStatus(uba::SessionServer* server, uba::u32 statusIndex, uba::u32 statusNameIndent, const uba::tchar* statusName, uba::u32 statusTextIndent, const uba::tchar* statusText, uba::LogEntryType statusType)
+	{
+		server->UpdateStatus(statusIndex, statusNameIndent, statusName, statusTextIndent, statusText, statusType);
+	}
+
 	void SessionServer_RegisterCustomService(uba::SessionServer* server, SessionServer_CustomServiceFunction* function, void* userData)
 	{
 		server->RegisterCustomService([function, userData](uba::Process& process, const void* recv, uba::u32 recvSize, void* send, uba::u32 sendCapacity)
diff --git a/Engine/Source/Programs/UnrealBuildAccelerator/Host/Public/UbaExports.h b/Engine/Source/Programs/UnrealBuildAccelerator/Host/Public/UbaExports.h
--- a/Engine/Source/Programs/UnrealBuildAccelerator/Host/Public/UbaExports.h
+++ b/Engine/Source/Programs/UnrealBuildAccelerator/Host/Public/UbaExports.h
@@ -100,6 +100,7 @@ extern "C"
 	UBA_API void SessionServer_SetCustomCasKeyFromTrackedInputs(uba::SessionServer* server, uba::ProcessHandle* handle, const uba::tchar* fileName, const uba::tchar* workingDir);
 	UBA_API uba::u32 SessionServer_BeginExternalProcess(uba::SessionServer* server, const uba::tchar* description);
 	UBA_API void SessionServer_EndExternalProcess(uba::SessionServer* server, uba::u32 id, uba::u32 exitCode);
+	UBA_API void SessionServer_UpdateStatus(uba::SessionServer* server, uba::u32 statusIndex, uba::u32 statusNameIndent, const uba::tchar* statusName, uba::u32 statusTextIndent, const uba::tchar* statusText, uba::LogEntryType statusType);
 	UBA_API void SessionServer_RegisterCustomService(uba::SessionServer* server, SessionServer_CustomServiceFunction* function, void* userData = nullptr);
 	UBA_API void DestroySessionServer(uba::SessionServer* server);
 
diff --git a/Engine/Source/Programs/UnrealBuildAccelerator/Visualizer/Private/UbaTraceReader.cpp b/Engine/Source/Programs/UnrealBuildAccelerator/Visualizer/Private/UbaTraceReader.cpp
--- a/Engine/Source/Programs/UnrealBuildAccelerator/Visualizer/Private/UbaTraceReader.cpp
+++ b/Engine/Source/Programs/UnrealBuildAccelerator/Visualizer/Private/UbaTraceReader.cpp
@@ -862,6 +862,17 @@ namespace uba
 				record.stop = time;
 			break;
 		}
+		case TraceType_StatusUpdate:
+		{
+			u32 statusIndex = u32(reader.Read7BitEncoded());
+			auto& status = out.statusMap[statusIndex];
+			status.nameIndent = u32(reader.Read7BitEncoded());
+			status.name = reader.ReadString();
+			status.textIndent = u32(reader.Read7BitEncoded());
+			status.text = reader.ReadString();
+			status.type = (LogEntryType)reader.ReadByte();
+			break;
+		}
 		case TraceType_String:
 		{
 			out.strings.push_back(reader.ReadString());
diff --git a/Engine/Source/Programs/UnrealBuildAccelerator/Visualizer/Private/UbaTraceReader.h b/Engine/Source/Programs/UnrealBuildAccelerator/Visualizer/Private/UbaTraceReader.h
--- a/Engine/Source/Programs/UnrealBuildAccelerator/Visualizer/Private/UbaTraceReader.h
+++ b/Engine/Source/Programs/UnrealBuildAccelerator/Visualizer/Private/UbaTraceReader.h
@@ -71,6 +71,15 @@ namespace uba
 			u64 stop;
 		};
 
+		struct StatusUpdate
+		{
+			TString name;
+			TString text;
+			u32 nameIndent;
+			u32 textIndent;
+			LogEntryType type;
+		};
+
 		struct Session
 		{
 			TString name;
@@ -109,11 +118,12 @@ namespace uba
 		};
 
 		Process* GetProcess(const ProcessLocation& loc) { return &(sessions[loc.sessionIndex].processors[loc.processorIndex].processes[loc.processIndex]); }
-		void Clear() { sessions.clear(); workTracks.clear(); strings.clear(); startTime = 0; finished = true; totalProcessActiveCount = 0; totalProcessExitedCount = 0; activeSessionCount = 0; };
+		void Clear() { sessions.clear(); workTracks.clear(); strings.clear(); statusMap.clear(); startTime = 0; finished = true; totalProcessActiveCount = 0; totalProcessExitedCount = 0; activeSessionCount = 0; };
 
 		Vector<Session> sessions;
 		Vector<WorkTrack> workTracks;
 		Vector<TString> strings;
+		Map<u32, StatusUpdate> statusMap;
 		u64 startTime = 0;
 		u64 frequency = 0;
 		u32 totalProcessActiveCount = 0;
diff --git a/Engine/Source/Programs/UnrealBuildAccelerator/Visualizer/Private/UbaVisualizer.cpp b/Engine/Source/Programs/UnrealBuildAccelerator/Visualizer/Private/UbaVisualizer.cpp
--- a/Engine/Source/Programs/UnrealBuildAccelerator/Visualizer/Private/UbaVisualizer.cpp
+++ b/Engine/Source/Programs/UnrealBuildAccelerator/Visualizer/Private/UbaVisualizer.cpp
@@ -251,6 +251,10 @@ namespace uba
 
 		if (m_useDarkMode)
 		{
+			m_textColor = RGB(190, 190, 190);
+			m_textWarningColor = RGB(190, 190, 0);
+			m_textErrorColor = RGB(190, 0, 0);
+
 			m_processBrushes[0].inProgress = CreateSolidBrush(RGB(70, 70, 70));
 			m_processBrushes[1].inProgress = CreateSolidBrush(RGB(130, 130, 130));
 
@@ -269,7 +273,6 @@ namespace uba
 
 			m_workBrush = CreateSolidBrush(RGB(70, 70, 100));
 
-			m_textColor = RGB(170, 170, 170);
 			m_backgroundBrush = CreateSolidBrush(0x00252526);
 			m_separatorPen = CreatePen(PS_SOLID, 1, RGB(50, 50, 50));
 			m_tooltipBackgroundBrush = CreateSolidBrush(0x00404040);
@@ -282,6 +285,10 @@ namespace uba
 		}
 		else
 		{
+			m_textColor = GetSysColor(COLOR_INFOTEXT);
+			m_textWarningColor = RGB(170, 130, 0);
+			m_textErrorColor = RGB(190, 0, 0);
+
 			m_processBrushes[0].inProgress = CreateSolidBrush(RGB(150, 150, 150));
 			m_processBrushes[1].inProgress = CreateSolidBrush(RGB(180, 180, 180));
 
@@ -300,7 +307,6 @@ namespace uba
 
 			m_workBrush = CreateSolidBrush(RGB(150, 150, 200));
 
-			m_textColor = GetSysColor(COLOR_INFOTEXT);
 			m_backgroundBrush = GetSysColorBrush(0);
 			m_separatorPen = CreatePen(PS_SOLID, 1, RGB(180, 180, 180));
 			m_tooltipBackgroundBrush = GetSysColorBrush(COLOR_INFOBK);
@@ -490,6 +496,36 @@ namespace uba
 
 		u64 lastStop = 0;
 
+		if (!m_traceView.statusMap.empty())
+		{
+			posY += 4;
+
+			auto drawStatusText = [&](const TString& text, LogEntryType type, int left, bool moveY)
+				{
+					RECT rect;
+					rect.left = left;
+					rect.right = clientRect.right;
+					rect.top = posY;
+					rect.bottom = posY + FontHeight + 2;
+					SetTextColor(hdc, type == LogEntryType_Info ? m_textColor : (type == LogEntryType_Error ? m_textErrorColor : m_textWarningColor));
+					ExtTextOutW(hdc, left, posY, ETO_CLIPPED, &rect, text.c_str(), u32(text.size()), NULL);
+					if (moveY)
+						posY = rect.bottom;
+				};
+
+			for (auto& kv : m_traceView.statusMap)
+			{
+				auto& status = kv.second;
+				if (status.name.empty() && status.text.empty())
+					continue;
+				drawStatusText(status.name, status.type, 5 + status.nameIndent*15, false);
+				drawStatusText(status.text, status.type, 5 + status.textIndent*15, true);
+			}
+			SetTextColor(hdc, m_textColor);
+			posY += 4;
+		}
+
+
 		struct SessionRec { TraceView::Session* session; u32 index; };
 		SessionRec sortedSessions[256];
 		for (u32 i = 0, e = u32(m_traceView.sessions.size()); i != e; ++i)
@@ -1678,6 +1714,17 @@ namespace uba
 
 		u64 lastStop = 0;
 
+		if (!m_traceView.statusMap.empty())
+		{
+			posY += 4;
+			auto drawStatusText = [&]() { posY = posY + FontHeight + 2; };
+			drawStatusText();
+			for (auto& kv : m_traceView.statusMap)
+				if (!kv.second.text.empty())
+					drawStatusText();
+			posY += 4;
+		}
+
 		TraceView::ProcessLocation& outLocation = outResult.processLocation;
 
 		struct SessionRec { TraceView::Session* session; u32 index; };
diff --git a/Engine/Source/Programs/UnrealBuildAccelerator/Visualizer/Private/UbaVisualizer.h b/Engine/Source/Programs/UnrealBuildAccelerator/Visualizer/Private/UbaVisualizer.h
--- a/Engine/Source/Programs/UnrealBuildAccelerator/Visualizer/Private/UbaVisualizer.h
+++ b/Engine/Source/Programs/UnrealBuildAccelerator/Visualizer/Private/UbaVisualizer.h
@@ -86,6 +86,8 @@ namespace uba
 		Atomic<bool> m_looping;
 		HWND m_hwnd = 0;
 		COLORREF m_textColor = {};
+		COLORREF m_textWarningColor = {};
+		COLORREF m_textErrorColor = {};
 		COLORREF m_sendColor = {};
 		COLORREF m_recvColor = {};
 		COLORREF m_cpuColor = {};
