[Uba]
* Fixed bug in storage client code that directly decompress network data to disk/mem. It could write outside buffer by 8 bytes
* Added UbaCli -checkcas2  .. this codepath creates a storage server and client and networks all cas files over.. just to check so code paths properly work
* Fixed UbaCli -checkaws bug

diff --git a/Engine/Source/Programs/UnrealBuildAccelerator/Cli/Private/UbaCli.cpp b/Engine/Source/Programs/UnrealBuildAccelerator/Cli/Private/UbaCli.cpp
--- a/Engine/Source/Programs/UnrealBuildAccelerator/Cli/Private/UbaCli.cpp
+++ b/Engine/Source/Programs/UnrealBuildAccelerator/Cli/Private/UbaCli.cpp
@@ -3,6 +3,7 @@
 #include "UbaNetworkBackendTcp.h"
 #include "UbaFileAccessor.h"
 #include "UbaSessionServer.h"
+#include "UbaStorageClient.h"
 #include "UbaStorageServer.h"
 #include "UbaPlatform.h"
 #include "UbaVersion.h"
@@ -132,6 +133,7 @@ namespace uba
 		bool disableCustomAllocator = false;
 		bool quiet = false;
 		bool checkCas = false;
+		bool checkCas2 = false;
 		bool checkAws = false;
 		bool getCas = false;
 		bool enableStdOut = true;
@@ -247,6 +249,10 @@ namespace uba
 			{
 				checkCas = true;
 			}
+			else if (name.Equals(TC("-checkcas2")))
+			{
+				checkCas2 = true;
+			}
 			else if (name.Equals(TC("-checkaws")))
 			{
 				checkAws = true;
@@ -294,12 +300,51 @@ namespace uba
 			bool success = storage.CheckCasContent(DefaultProcessorCount);
 			return success ? 0 : -1;
 		}
+
+		if (checkCas2) // Creates a storage server and storage client and transfer _all_ cas files over network
+		{
+			NetworkBackendTcp networkBackend(logWriter);
+			NetworkServerCreateInfo nsci(logWriter);
+			bool ctorSuccess = true;
+			NetworkServer server(ctorSuccess, nsci);
+			StorageServerCreateInfo storageInfo(server, g_rootDir.data, logWriter);
+			storageInfo.casCapacityBytes = 0;
+			storageInfo.storeCompressed = storeCompressed;
+			StorageServer storageServer(storageInfo);
+			NetworkClientCreateInfo ncci;
+			NetworkClient client(ctorSuccess, ncci);
+			StringBuffer<> rootDir2(g_rootDir.data);
+			rootDir2.Append("_CHECKCAS2");
+			StorageClientCreateInfo scci(client, rootDir2.data);
+			StorageClient storageClient(scci);
+			auto g = MakeGuard([&]() { server.StopAll(); });
+			if (!server.StartListen(networkBackend, 1347, TC("127.0.0.1")))
+				return false;
+			if (!client.Connect(networkBackend, TC("127.0.0.1"), 1347))
+				return false;
+			bool success = true;
+			WorkManagerImpl workManager(DefaultProcessorCount);
+			storageServer.TraverseAllCasFiles([&](const CasKey& casKey)
+				{
+					workManager.AddWork([&, casKey]()
+						{
+							Storage::RetrieveResult res;
+							storageServer.EnsureCasFile(casKey, TC("Dummy"));
+							CasKey casKey2 = AsCompressed(casKey, false);
+							if (!storageClient.RetrieveCasFile(res, casKey2, TC("")))
+								success = false;
+						}, 1, TC(""));
+				});
+			workManager.FlushWork();
+			return success;
+		}
+
 #if UBA_USE_AWS
 		if (checkAws)
 		{
 			AWS aws;
 			StringBuffer<> info;
-			if (aws.QueryInformation(logger, info, TC("UbaCli")))
+			if (aws.QueryInformation(logger, info, g_rootDir.data))
 			{
 				logger.Info(TC("We are inside AWS: %s (%s)"), info.data, aws.GetAvailabilityZone());
 				
diff --git a/Engine/Source/Programs/UnrealBuildAccelerator/Common/Private/UbaStorage.cpp b/Engine/Source/Programs/UnrealBuildAccelerator/Common/Private/UbaStorage.cpp
--- a/Engine/Source/Programs/UnrealBuildAccelerator/Common/Private/UbaStorage.cpp
+++ b/Engine/Source/Programs/UnrealBuildAccelerator/Common/Private/UbaStorage.cpp
@@ -796,6 +796,16 @@ namespace uba
 			});
 	}
 
+	void StorageImpl::TraverseAllCasFiles(const Function<void(const CasKey& key)>& func)
+	{
+		StringBuffer<> casRoot;
+		casRoot.Append(m_rootDir.data, m_rootDir.count - 1);
+		TraverseAllCasFiles(casRoot.data, 0, [&](const StringBufferBase& fullPath, const DirectoryEntry& e)
+			{
+				func(CasKeyFromString(e.name));
+			});
+	}
+
 	void StorageImpl::CheckAllCasFiles()
 	{
 #if !UBA_USE_SPARSEFILE
diff --git a/Engine/Source/Programs/UnrealBuildAccelerator/Common/Private/UbaStorageClient.cpp b/Engine/Source/Programs/UnrealBuildAccelerator/Common/Private/UbaStorageClient.cpp
--- a/Engine/Source/Programs/UnrealBuildAccelerator/Common/Private/UbaStorageClient.cpp
+++ b/Engine/Source/Programs/UnrealBuildAccelerator/Common/Private/UbaStorageClient.cpp
@@ -481,6 +481,7 @@ namespace uba
 				bool sendSegmentMessage = responseSize == 0;
 				u64 leftUncompressed = actualSize;
 				readBuffer += sizeof(u64); // Size is stored first
+				u64 maxReadSize = BufferSlotHalfSize - sizeof(u64);
 
 				if (actualSize)
 				{
@@ -501,7 +502,7 @@ namespace uba
 							{
 								if (fetchId == u16(~0))
 									return m_logger.Error(TC("Cas content error (2). Server believes %s was only one segment but client sees more. UncompressedSize: %llu LeftUncompressed: %llu Size: %llu Left to read: %llu ResponseSize: %u. (%s)"), hint, actualSize, leftUncompressed, fileSize, left, responseSize, casFile.data);
-								if (!SendBatchMessages(*client, fetchId, readPosition, BufferSlotHalfSize - u32(readPosition - readBuffer), leftCompressed, sizeOfFirstMessage, readIndex, responseSize))
+								if (!SendBatchMessages(*client, fetchId, readPosition, maxReadSize - u32(readPosition - readBuffer), leftCompressed, sizeOfFirstMessage, readIndex, responseSize))
 								{
 									if (proxy)
 									{
diff --git a/Engine/Source/Programs/UnrealBuildAccelerator/Common/Public/UbaStorage.h b/Engine/Source/Programs/UnrealBuildAccelerator/Common/Public/UbaStorage.h
--- a/Engine/Source/Programs/UnrealBuildAccelerator/Common/Public/UbaStorage.h
+++ b/Engine/Source/Programs/UnrealBuildAccelerator/Common/Public/UbaStorage.h
@@ -143,6 +143,7 @@ namespace uba
 		void AttachEntry(CasEntry& entry);
 		void DetachEntry(CasEntry& entry);
 		void TraverseAllCasFiles(const tchar* dir, u32 recursion, const Function<void(const StringBufferBase& fullPath, const DirectoryEntry& e)>& func);
+		void TraverseAllCasFiles(const Function<void(const CasKey& key)>& func);
 		void CheckAllCasFiles();
 		void HandleOverflow();
 		bool OpenCasDataFile(u32 index, u64 size);
