[UBA]
* Proper fix of bug where code could write 8 bytes outside memory allocation

diff --git a/Engine/Source/Programs/UnrealBuildAccelerator/Common/Private/UbaStorageClient.cpp b/Engine/Source/Programs/UnrealBuildAccelerator/Common/Private/UbaStorageClient.cpp
--- a/Engine/Source/Programs/UnrealBuildAccelerator/Common/Private/UbaStorageClient.cpp
+++ b/Engine/Source/Programs/UnrealBuildAccelerator/Common/Private/UbaStorageClient.cpp
@@ -527,6 +527,7 @@ namespace uba
 								compressedSize = blockSize[0];
 								uncompressedSize = blockSize[1];
 								readBuffer += sizeof(u32) * 2;
+								maxReadSize = BufferSlotHalfSize - sizeof(u32) * 2;
 								u32 read = (responseSize + u32(readPosition - readBuffer));
 								//UBA_ASSERTF(read <= compressedSize, TC("Error in datastream fetching cas. Read size: %u CompressedSize: %u %s (%s)"), read, compressedSize, casFile.data, hint);
 								if (read > compressedSize)
@@ -591,6 +592,7 @@ namespace uba
 
 						// Move overflow back to the beginning of the buffer and start the next block (if there is one)
 						readBuffer = slot;
+						maxReadSize = BufferSlotHalfSize;
 						UBA_ASSERTF(readPosition - overflow >= readBuffer, TC("ReadPosition - overflow is before beginning of buffer (overflow: %u) for file %s"), overflow, hint);
 						UBA_ASSERTF(readPosition <= readBuffer + BufferSlotHalfSize, TC("ReadPosition is outside readBuffer size (pos: %llu, overflow: %u) for file %s"), readPosition - readBuffer, overflow, hint);
 						memmove(readBuffer, readPosition - overflow, overflow);
