[UBA]
* Fixed bug where error messages written to trace stream could be bigger than max size for a write scope. Previous number was 4k, now changed to 256k instead (I don't expect an individual error message to be bigger than 256k at least!)

diff --git a/Engine/Source/Programs/UnrealBuildAccelerator/Common/Private/UbaTrace.cpp b/Engine/Source/Programs/UnrealBuildAccelerator/Common/Private/UbaTrace.cpp
--- a/Engine/Source/Programs/UnrealBuildAccelerator/Common/Private/UbaTrace.cpp
+++ b/Engine/Source/Programs/UnrealBuildAccelerator/Common/Private/UbaTrace.cpp
@@ -6,6 +6,8 @@
 
 namespace uba
 {
+	constexpr u64 TraceMessageMaxSize = 256 * 1024;
+
 	Trace::Trace(LogWriter& logWriter)
 	:	m_logger(logWriter)
 	,	m_channel(m_logger)
@@ -24,7 +26,7 @@ namespace uba
 	{
 		WriterScope(Trace& trace) : ScopedWriteLock(trace.m_memoryLock), BinaryWriter(trace.m_memoryBegin, trace.m_memoryPos, trace.m_memoryCapacity), m_trace(trace)
 		{
-			u64 committedMemoryNeeded = AlignUp(trace.m_memoryPos + 4095, 256 * 1024);
+			u64 committedMemoryNeeded = AlignUp(trace.m_memoryPos + TraceMessageMaxSize, TraceMessageMaxSize);
 			if (trace.m_memoryCommitted >= committedMemoryNeeded)
 				return;
 			if (!MapViewCommit(trace.m_memoryBegin + trace.m_memoryCommitted, committedMemoryNeeded - trace.m_memoryCommitted))
