[UBA]
* Fixed up replay controls in visualizer

diff --git a/Engine/Source/Programs/UnrealBuildAccelerator/Visualizer/Private/UbaVisualizer.cpp b/Engine/Source/Programs/UnrealBuildAccelerator/Visualizer/Private/UbaVisualizer.cpp
--- a/Engine/Source/Programs/UnrealBuildAccelerator/Visualizer/Private/UbaVisualizer.cpp
+++ b/Engine/Source/Programs/UnrealBuildAccelerator/Visualizer/Private/UbaVisualizer.cpp
@@ -425,6 +425,25 @@ namespace uba
 		}
 	}
 
+	void Visualizer::Pause(bool pause)
+	{
+		if (m_paused == pause)
+			return;
+
+		m_paused = pause;
+		if (pause)
+		{
+			m_pauseStart = GetTime();
+		}
+		else
+		{
+			m_replay = 1;
+			m_pauseTime += GetTime() - m_pauseStart;
+			m_traceView.finished = false;
+			SetTimer(m_hwnd, 0, 200, NULL);
+		}
+	}
+
 	void Visualizer::PaintClient(const Function<void(HDC hdc, HDC memDC, RECT& clientRect)>& paintFunc)
 	{
 		HDC hdc = GetDC(m_hwnd);
@@ -2307,17 +2326,25 @@ namespace uba
 			}
 			else if (m_timelineSelected)
 			{
-				if (!m_fileName.IsEmpty() && !m_replay)
+				if (!m_fileName.IsEmpty()) // Only works for files right now
 				{
 					Reset();
 					if (!m_trace.ReadFile(m_traceView, m_fileName.data, true))
 						return false;
 					bool changed;
 					u64 time = MsToTime(u64(m_timelineSelected * 1000.0));
+					m_traceView.finished = false;
 					m_trace.UpdateReadFile(m_traceView, time, changed);
 					m_pauseStart = m_startTime + time;
-					m_paused = true;
-					m_traceView.finished = true;
+					if (m_paused || !m_replay)
+					{
+						m_pauseTime = 0;
+						m_paused = true;
+					}
+					else
+					{
+						m_pauseTime = GetTime() - m_pauseStart;
+					}
 
 					HitTestResult res;
 					HitTest(res, { -1, -1 });
@@ -2401,13 +2428,11 @@ namespace uba
 				break;
 
 			case Popup_Play:
-				m_paused = false;
-				m_pauseTime += GetTime() - m_pauseStart;
+				Pause(false);
 				break;
 
 			case Popup_Pause:
-				m_paused = true;
-				m_pauseStart = GetTime();
+				Pause(true);
 				break;
 
 			case Popup_JumpToEnd:
@@ -2475,18 +2500,7 @@ namespace uba
 		case WM_KEYDOWN:
 		{
 			if (wParam == VK_SPACE)
-			{
-				if (!m_paused)
-				{
-					m_paused = true;
-					m_pauseStart = GetTime();
-				}
-				else
-				{
-					m_paused = false;
-					m_pauseTime += GetTime() - m_pauseStart;
-				}
-			}
+				Pause(!m_paused);
 			break;
 		}
 		case WM_VSCROLL:
diff --git a/Engine/Source/Programs/UnrealBuildAccelerator/Visualizer/Private/UbaVisualizer.h b/Engine/Source/Programs/UnrealBuildAccelerator/Visualizer/Private/UbaVisualizer.h
--- a/Engine/Source/Programs/UnrealBuildAccelerator/Visualizer/Private/UbaVisualizer.h
+++ b/Engine/Source/Programs/UnrealBuildAccelerator/Visualizer/Private/UbaVisualizer.h
@@ -64,6 +64,7 @@ namespace uba
 		void UpdateScrollbars(bool redraw);
 		void GetTitlePrefix(StringBufferBase& out);
 		void ThreadLoop();
+		void Pause(bool pause);
 
 		StringBuffer<256> m_namedTrace;
 		StringBuffer<256> m_fileName;
