[UBA]
* Fixed logic calculating max page file size. System-managed page file logic was wrong giving calculations that were way to high

diff --git a/Engine/Source/Programs/UnrealBuildAccelerator/Common/Private/UbaSession.cpp b/Engine/Source/Programs/UnrealBuildAccelerator/Common/Private/UbaSession.cpp
--- a/Engine/Source/Programs/UnrealBuildAccelerator/Common/Private/UbaSession.cpp
+++ b/Engine/Source/Programs/UnrealBuildAccelerator/Common/Private/UbaSession.cpp
@@ -2160,23 +2160,46 @@ namespace uba
 		if (m_maxPageSize == ~u64(0))
 		{
 			m_maxPageSize = 0;
-			bool foundEntry = false;
 			wchar_t str[1024];
 			DWORD strBytes = sizeof(str);
 			LSTATUS res = RegGetValueW(HKEY_LOCAL_MACHINE, TC("SYSTEM\\CurrentControlSet\\Control\\Session Manager\\Memory Management"), TC("PagingFiles"), RRF_RT_REG_MULTI_SZ, NULL, str, &strBytes);
 			if (res == ERROR_SUCCESS)
 			{
 				wchar_t* line = str;
-				while (size_t len = wcslen(line))
+				for (; size_t len = wcslen(line); line += len + 1)
 				{
-					foundEntry = true;
-					if (const wchar_t* maxSize = wcsrchr(line, ' '))
-						m_maxPageSize += wcstoull(maxSize + 1, nullptr, 10) * 1024 * 1024;
-					line += len + 1;
+					const wchar_t* maxSizeStr = wcsrchr(line, ' ');
+					u64 maxSizeMb = 0;
+					if (!maxSizeStr || !StringBuffer<32>(maxSizeStr + 1).Parse(maxSizeMb))
+					{
+						m_logger.Warning(TC("Unrecognized page file information format (please report): %s"), line);
+						continue;
+					}
+
+					if (maxSizeMb) // Custom set page file size
+					{
+						m_maxPageSize += maxSizeMb * 1024 * 1024;
+						continue;
+					}
+
+					// System-managed page file. Get drive root path
+					StringBuffer<8> drive;
+					drive.Append(line, 3);
+
+					// Max possible system-managed page file
+					maxSizeMb = Max(u64(memStatus.ullTotalPhys) * 3, 4ull * 1024 * 1024 * 1024);
+
+					// Check if disk is limiting factor of system-managed page file
+					// Page file can be max 1/8 of volume size and ofc not more than free space
+					ULARGE_INTEGER totalNumberOfBytes;
+					ULARGE_INTEGER totalNumberOfFreeBytes;
+					if (!GetDiskFreeSpaceExW(drive.data, NULL, &totalNumberOfBytes, &totalNumberOfFreeBytes))
+						return m_logger.Error(TC("GetDiskFreeSpaceExW failed to get information about %s (%s)"), drive.data, LastErrorToText().data);
+
+					u64 maxDiskPageFileSize = Min(totalNumberOfBytes.QuadPart / 8, totalNumberOfFreeBytes.QuadPart);
+					m_maxPageSize += Min(maxDiskPageFileSize, maxSizeMb);
 				}
 			}
-			if (foundEntry && m_maxPageSize == 0)
-				m_maxPageSize = Max(u64(memStatus.ullTotalPhys) * 3, 4ull * 1024 * 1024 * 1024); // System managed page file. Equation taken from msdn
 		}
 
 		u64 currentPageSize = memStatus.ullTotalPageFile - memStatus.ullTotalPhys;
