[UBA] Disable power throttling in created processes to ensure P-Cores are preferred over E-Cores on Intel Hybrid architectures platforms

diff --git a/Engine/Source/Programs/UnrealBuildAccelerator/Common/Private/UbaProcess.cpp b/Engine/Source/Programs/UnrealBuildAccelerator/Common/Private/UbaProcess.cpp
--- a/Engine/Source/Programs/UnrealBuildAccelerator/Common/Private/UbaProcess.cpp
+++ b/Engine/Source/Programs/UnrealBuildAccelerator/Common/Private/UbaProcess.cpp
@@ -1097,6 +1097,24 @@ namespace uba
 		return temp.data;
 	}
 
+#if PLATFORM_WINDOWS
+	/** Disables Power Throttling in the provided process, to ensure P-Cores are preferred over E Cores on hybrid architecture intel platforms */
+	void DisableProcessPowerThrottling(HANDLE ProcessHandle)
+	{
+		PROCESS_POWER_THROTTLING_STATE PowerThrottling;
+		RtlZeroMemory(&PowerThrottling, sizeof(PowerThrottling));
+
+		// Enable PowerThrottling policies for the process
+		// and disable power throttling by setting the state mask to 0
+		PowerThrottling.Version = PROCESS_POWER_THROTTLING_CURRENT_VERSION;
+		PowerThrottling.ControlMask = PROCESS_POWER_THROTTLING_EXECUTION_SPEED;
+		PowerThrottling.StateMask = 0;
+
+		SetProcessInformation(ProcessHandle, ProcessPowerThrottling, &PowerThrottling, sizeof(PowerThrottling));
+	}
+#endif
+
+
 	u32 ProcessImpl::InternalCreateProcess(bool runningRemote, void* environment, FileMappingHandle communicationHandle, u64 communicationOffset)
 	{
 		ScopedWriteLock initLock(m_initLock);
@@ -1188,7 +1206,10 @@ namespace uba
 				if (m_detourEnabled)
 				{
 					if (DetourCreateProcessWithDlls(NULL, (tchar*)commandLine.c_str(), NULL, NULL, inheritHandles, creationFlags, environment, workingDir, &si, &processInfo, sizeof_array(dlls), dlls, NULL))
+					{
+						DisableProcessPowerThrottling(processInfo.hProcess);
 						break;
+					}
 				}
 				else
 				{
@@ -1218,7 +1239,11 @@ namespace uba
 					si.dwFlags |= STARTF_USESTDHANDLES;
 
 					if (CreateProcessW(NULL, (tchar*)commandLine.c_str(), NULL, NULL, TRUE, creationFlags, environment, workingDir, &si, &processInfo))
+					{
+						DisableProcessPowerThrottling(processInfo.hProcess);
 						break;
+					}
+						
 				}
 
 				DWORD error = GetLastError();
