[UBA]
* Added support for PagingFiles registry key containing drive letter '?'. Not sure exactly what this mean except that it can exist if system automatically manage page files for all drives. We resort to use first entry in ExistingPageFiles registry key to figure out drive

diff --git a/Engine/Source/Programs/UnrealBuildAccelerator/Common/Private/UbaSession.cpp b/Engine/Source/Programs/UnrealBuildAccelerator/Common/Private/UbaSession.cpp
--- a/Engine/Source/Programs/UnrealBuildAccelerator/Common/Private/UbaSession.cpp
+++ b/Engine/Source/Programs/UnrealBuildAccelerator/Common/Private/UbaSession.cpp
@@ -2166,26 +2166,49 @@ namespace uba
 			if (res == ERROR_SUCCESS)
 			{
 				wchar_t* line = str;
-				for (; size_t len = wcslen(line); line += len + 1)
+				for (; size_t lineLen = wcslen(line); line += lineLen + 1)
 				{
-					const wchar_t* maxSizeStr = wcsrchr(line, ' ');
+					if (lineLen < 3)
+						continue;
+
 					u64 maxSizeMb = 0;
-					if (!maxSizeStr || !StringBuffer<32>(maxSizeStr + 1).Parse(maxSizeMb))
+
+					StringBuffer<8> drive;
+					drive.Append(line, 3); // Get drive root path
+
+					if (drive[0] == '?') // Drive '?' can exist when "Automatically manage paging file size for all drives".. 
 					{
-						m_logger.Warning(TC("Unrecognized page file information format (please report): %s"), line);
-						continue;
+						// We can use ExistingPageFiles registry key to figure out which drive...
+						// This key can contain multiple page files normally.. have no idea if it can contain multiple when drive is '?'.. but for now, just use the first
+						wchar_t str2[1024];
+						DWORD str2Bytes = sizeof(str2);
+						res = RegGetValueW(HKEY_LOCAL_MACHINE, TC("SYSTEM\\CurrentControlSet\\Control\\Session Manager\\Memory Management"), TC("ExistingPageFiles"), RRF_RT_REG_MULTI_SZ, NULL, str2, &str2Bytes);
+						if (res != ERROR_SUCCESS)
+							continue;
+
+						auto colon = wcschr(str2, ':'); // Path is something like \??\C:\pagefile.sys or similar.. let's search for : and use character in front of it.
+						if (colon == nullptr || colon == str2)
+							continue;
+
+						drive[0] = colon[-1];
 					}
-
-					if (maxSizeMb) // Custom set page file size
+					else
 					{
-						m_maxPageSize += maxSizeMb * 1024 * 1024;
-						continue;
+						const wchar_t* maxSizeStr = wcsrchr(line, ' ');
+						
+						if (!maxSizeStr || !StringBuffer<32>(maxSizeStr + 1).Parse(maxSizeMb))
+						{
+							m_logger.Warning(TC("Unrecognized page file information format (please report): %s"), line);
+							continue;
+						}
+
+						if (maxSizeMb) // Custom set page file size
+						{
+							m_maxPageSize += maxSizeMb * 1024 * 1024;
+							continue;
+						}
 					}
 
-					// System-managed page file. Get drive root path
-					StringBuffer<8> drive;
-					drive.Append(line, 3);
-
 					// Max possible system-managed page file
 					maxSizeMb = Max(u64(memStatus.ullTotalPhys) * 3, 4ull * 1024 * 1024 * 1024);
 
