[UBA]
* Fixed so include-what-you-use json output is working properly with uba. "puts" implementation is wrong but don't remember why it was added first place so didn't dare changing it too much.. current change only affect remote helpers running on wine
* Fixed ubavisualizer hittest logic when status lines are added
* Enabled remote log transfer by default in UbaCli

#changelist validated
#virtualized

diff --git a/Engine/Source/Programs/UnrealBuildAccelerator/Cli/Private/UbaCli.cpp b/Engine/Source/Programs/UnrealBuildAccelerator/Cli/Private/UbaCli.cpp
--- a/Engine/Source/Programs/UnrealBuildAccelerator/Cli/Private/UbaCli.cpp
+++ b/Engine/Source/Programs/UnrealBuildAccelerator/Cli/Private/UbaCli.cpp
@@ -455,6 +455,7 @@ namespace uba
 		//info.shouldWriteToDisk = shouldWriteToDisk;
 		info.rootDir = g_rootDir.data;
 		//info.traceName.Append(TC("TESTTRACE"));
+		info.remoteLogEnabled = true;
 		info.deleteSessionsOlderThanSeconds = 1;
 		auto session = new SessionServer(info);
 		auto destroySession = MakeGuard([&]() { delete session; });
diff --git a/Engine/Source/Programs/UnrealBuildAccelerator/Detours/Private/Windows/UbaDetoursFunctionsUcrtBase.inl b/Engine/Source/Programs/UnrealBuildAccelerator/Detours/Private/Windows/UbaDetoursFunctionsUcrtBase.inl
--- a/Engine/Source/Programs/UnrealBuildAccelerator/Detours/Private/Windows/UbaDetoursFunctionsUcrtBase.inl
+++ b/Engine/Source/Programs/UnrealBuildAccelerator/Detours/Private/Windows/UbaDetoursFunctionsUcrtBase.inl
@@ -257,12 +257,18 @@ int Detoured__write(int fd, const void* buffer, unsigned int count)
 
 int Detoured_fputs(const char* str, FILE* stream)
 {
-	int fno = _fileno(stream); (void)fno;
-	int errfno = _fileno(stderr);
-	if (fno == errfno || fno == _fileno(stdout))
+	// This code is not working properly on wine helpers.
+	// When using iwyu on remote wine helper we get a fno that matches stderr but it is not a stderr.. it is supposed to write to a file
+	// Unfortunately I can't remember the initial reason this code was added so I can't change it and know if I broke something or not
+	if (!g_runningRemote || !g_isRunningWine)
 	{
-		Shared_WriteConsole(str, u32(strlen(str)), fno == errfno);
-		return 1;
+		int fno = _fileno(stream); (void)fno;
+		int errfno = _fileno(stderr);
+		if (fno == errfno || fno == _fileno(stdout))
+		{
+			Shared_WriteConsole(str, u32(strlen(str)), fno == errfno);
+			return 1;
+		}
 	}
 	return True_fputs(str, stream);
 }
diff --git a/Engine/Source/Programs/UnrealBuildAccelerator/Visualizer/Private/UbaVisualizer.cpp b/Engine/Source/Programs/UnrealBuildAccelerator/Visualizer/Private/UbaVisualizer.cpp
--- a/Engine/Source/Programs/UnrealBuildAccelerator/Visualizer/Private/UbaVisualizer.cpp
+++ b/Engine/Source/Programs/UnrealBuildAccelerator/Visualizer/Private/UbaVisualizer.cpp
@@ -1737,9 +1737,8 @@ namespace uba
 		{
 			posY += 4;
 			auto drawStatusText = [&]() { posY = posY + FontHeight + 2; };
-			drawStatusText();
 			for (auto& kv : m_traceView.statusMap)
-				if (!kv.second.text.empty())
+				if (!kv.second.name.empty() || !kv.second.text.empty())
 					drawStatusText();
 			posY += 4;
 		}
