[UBA]
* Removed logging entries that was only added to make horde not timeout.. since this was implemented horde has added ping logic so this is not needed anymore
* Fixed bug in callstack logging on posix
* Changed a bunch of %hs to %s
* Added allocator debug code to try to figure out asan issues

diff --git a/Engine/Source/Programs/UnrealBuildAccelerator/Agent/Private/UbaAgent.cpp b/Engine/Source/Programs/UnrealBuildAccelerator/Agent/Private/UbaAgent.cpp
--- a/Engine/Source/Programs/UnrealBuildAccelerator/Agent/Private/UbaAgent.cpp
+++ b/Engine/Source/Programs/UnrealBuildAccelerator/Agent/Private/UbaAgent.cpp
@@ -36,7 +36,7 @@
 
 namespace uba
 {
-	const char*		Version = GetVersionString();
+	const tchar*		Version = GetVersionString();
 	constexpr u32	DefaultCapacityGb = 20;
 	constexpr u32	DefaultListenTimeout = 5;
 	const tchar*	DefaultRootDir = [](){
@@ -61,7 +61,7 @@ namespace uba
 		}
 		logger.Info(TC(""));
 		logger.Info(TC("-------------------------------------------"));
-		logger.Info(TC("   UbaAgent v%hs"), Version);
+		logger.Info(TC("   UbaAgent v%s"), Version);
 		logger.Info(TC("-------------------------------------------"));
 		logger.Info(TC(""));
 		logger.Info(TC("  When started UbaAgent will keep trying to connect to provided host address."));
@@ -685,7 +685,7 @@ namespace uba
 			char release[128];
 			char url[512];
 			size_t urlLen;
-			sprintf_s(release, sizeof_array(release), "BoxAgent@%s", Version);
+			sprintf_s(release, sizeof_array(release), "BoxAgent@%ls", Version);
 			wcstombs_s(&urlLen, url, sizeof_array(url), sentryUrl.data, sizeof_array(url) - 1);
 			sentry_options_t* options = sentry_options_new();
 			sentry_options_set_dsn(options, url);
@@ -728,7 +728,7 @@ namespace uba
 		#if UBA_DEBUG
 		dbgStr = TC(" (DEBUG)");
 		#endif
-		logger.Info(TC("UbaAgent v%hs%s (Cpu: %u, MaxCon: %u, Dir: \"%s\", StoreCapacity: %uGb%s)"), Version, dbgStr, maxProcessCount, maxTcpConnectionCount, g_rootDir.data, storageCapacityGb, extraInfo.data);
+		logger.Info(TC("UbaAgent v%s%s (Cpu: %u, MaxCon: %u, Dir: \"%s\", StoreCapacity: %uGb%s)"), Version, dbgStr, maxProcessCount, maxTcpConnectionCount, g_rootDir.data, storageCapacityGb, extraInfo.data);
 		if (!eventFile.IsEmpty())
 			logger.Info(TC("  Will poll for external events in file %s"), eventFile.data);
 		logger.Info(TC(""));
@@ -737,7 +737,7 @@ namespace uba
 #if PLATFORM_WINDOWS
 		{
 			StringBuffer<256> consoleTitle;
-			consoleTitle.Appendf(TC("UbaAgent v%hs%s"), Version, dbgStr);
+			consoleTitle.Appendf(TC("UbaAgent v%s%s"), Version, dbgStr);
 			SetConsoleTitleW(consoleTitle.data);
 		}
 #endif
@@ -1152,7 +1152,9 @@ namespace uba
 			else
 				logger.Info(TC("----------- Session %s started -----------"), sessionClient->GetId());
 
+			#if 0 
 			u64 lastLogTime = GetTime();
+			#endif
 
 			u32 tcpConnectionCount = 1;
 
@@ -1202,7 +1204,7 @@ namespace uba
 					}
 				}
 
-
+				#if 0 // Not needed anymore I think.. horde is now pinging in the background for itself
 				if (quiet)
 				{
 					u64 time = GetTime();
@@ -1213,6 +1215,7 @@ namespace uba
 						lastLogTime = time;
 					}
 				}
+				#endif
 			}
 
 			disconnectAndStopLoggingThread.Execute();
diff --git a/Engine/Source/Programs/UnrealBuildAccelerator/Cli/Private/UbaCli.cpp b/Engine/Source/Programs/UnrealBuildAccelerator/Cli/Private/UbaCli.cpp
--- a/Engine/Source/Programs/UnrealBuildAccelerator/Cli/Private/UbaCli.cpp
+++ b/Engine/Source/Programs/UnrealBuildAccelerator/Cli/Private/UbaCli.cpp
@@ -19,7 +19,7 @@
 namespace uba
 {
 
-	const char*		Version = GetVersionString();
+	const tchar*	Version = GetVersionString();
 	u32				DefaultCapacityGb = 20;
 	const tchar*	DefaultRootDir = []() {
 		static tchar buf[256];
@@ -41,7 +41,7 @@ namespace uba
 		}
 		logger.Info(TC(""));
 		logger.Info(TC("-------------------------------------------"));
-		logger.Info(TC("   UbaCli v%hs"), Version);
+		logger.Info(TC("   UbaCli v%s"), Version);
 		logger.Info(TC("-------------------------------------------"));
 		logger.Info(TC(""));
 		logger.Info(TC("  UbaCli.exe [options...] <commandtype> <executable> [arguments...]"));
@@ -414,7 +414,7 @@ namespace uba
 		#if UBA_DEBUG
 		dbgStr = TC(" (DEBUG)");
 		#endif
-		logger.Info(TC("UbaCli v%hs%s (Rootdir: \"%s\", StoreCapacity: %uGb)\n"), Version, dbgStr, g_rootDir.data, storageCapacityGb);
+		logger.Info(TC("UbaCli v%s%s (Rootdir: \"%s\", StoreCapacity: %uGb)\n"), Version, dbgStr, g_rootDir.data, storageCapacityGb);
 
 		u64 storageCapacity = u64(storageCapacityGb)*1000*1000*1000;
 
@@ -425,8 +425,10 @@ namespace uba
 
 
 		StringBuffer<> logFile;
+		#if UBA_DEBUG
 		logFile.count = GetFullPathNameW(g_rootDir.data, logFile.capacity, logFile.data, nullptr);
 		logFile.EnsureEndsWithSlash().Append(TC("DebugLog.log"));
+		#endif
 
 		#if PLATFORM_WINDOWS
 		SetConsoleCtrlHandler(ConsoleHandler, TRUE);
@@ -450,6 +452,7 @@ namespace uba
 
 		SessionServerCreateInfo info(*storage, *server);
 		info.useUniqueId = false;
+		info.traceEnabled = true;
 		info.launchVisualizer = launchVisualizer;
 		info.disableCustomAllocator = disableCustomAllocator;
 		//info.shouldWriteToDisk = shouldWriteToDisk;
@@ -473,6 +476,7 @@ namespace uba
 		{
 			u64 start = GetTime();
 			ProcessStartInfo pinfo;
+			pinfo.description = app.c_str();
 			pinfo.application = app.c_str();
 			pinfo.arguments = arg.c_str();
 			pinfo.workingDir = workDir.data;
@@ -494,6 +498,7 @@ namespace uba
 		{
 			u64 start = GetTime();
 			ProcessStartInfo pinfo;
+			pinfo.description = app.c_str();
 			pinfo.application = app.c_str();
 			pinfo.arguments = arg.c_str();
 			pinfo.workingDir = workDir.data;
diff --git a/Engine/Source/Programs/UnrealBuildAccelerator/Common/Private/UbaProcess.cpp b/Engine/Source/Programs/UnrealBuildAccelerator/Common/Private/UbaProcess.cpp
--- a/Engine/Source/Programs/UnrealBuildAccelerator/Common/Private/UbaProcess.cpp
+++ b/Engine/Source/Programs/UnrealBuildAccelerator/Common/Private/UbaProcess.cpp
@@ -1715,7 +1715,7 @@ namespace uba
 			int res = waitpid((pid_t)m_nativeProcessId, &status, 0);
 			if (res == -1)
 			{
-				logger.Error(TC("waitpid failed on %u (%hs)"), m_nativeProcessId, strerror(errno));
+				logger.Error(TC("waitpid failed on %u (%s)"), m_nativeProcessId, strerror(errno));
 				return UBA_EXIT_CODE(15);
 			}
 			if (WIFEXITED(status))
diff --git a/Engine/Source/Programs/UnrealBuildAccelerator/Common/Private/UbaSession.cpp b/Engine/Source/Programs/UnrealBuildAccelerator/Common/Private/UbaSession.cpp
--- a/Engine/Source/Programs/UnrealBuildAccelerator/Common/Private/UbaSession.cpp
+++ b/Engine/Source/Programs/UnrealBuildAccelerator/Common/Private/UbaSession.cpp
@@ -2304,7 +2304,7 @@ namespace uba
 		if (fd != -1)
 		{
 			char buffer[512];
-			int size = read(fd, buffer, 512);
+			int size = read(fd, buffer, sizeof_array(buffer) - 1);
 			if (size != -1)
 			{
 				buffer[size] = 0;
diff --git a/Engine/Source/Programs/UnrealBuildAccelerator/Core/Private/UbaPlatform.cpp b/Engine/Source/Programs/UnrealBuildAccelerator/Core/Private/UbaPlatform.cpp
--- a/Engine/Source/Programs/UnrealBuildAccelerator/Core/Private/UbaPlatform.cpp
+++ b/Engine/Source/Programs/UnrealBuildAccelerator/Core/Private/UbaPlatform.cpp
@@ -185,10 +185,10 @@ namespace uba
 		}
 #else
 		out.Append(TC("ASSERT: ")).Append(*text ? text : TC("Unknown")).Append('\n');
-		if (strcmp(expr, "false") != 0)
-			out.Appendf(TC(" EXPR: %hs\n"), expr);
+		if (expr && strcmp(expr, "false") != 0)
+			out.Appendf(TC(" EXPR: %s\n"), expr);
 		if (file && *file)
-			out.Appendf(TC("LOCATION: %hs:%u"), file, line);
+			out.Appendf(TC("LOCATION: %s:%u"), file, line);
 
 		constexpr u32 maxCallers = 100;
 		void* buffer[maxCallers];
@@ -232,7 +232,7 @@ namespace uba
 							fileName = lastSlash + 1;
 						if (strncmp(fileName, "function.h", 10) == 0 || strncmp(fileName, "invoke.h:", 9) == 0)
 							continue;
-						out.Appendf(TC("  %hs"), fileName);
+						out.Appendf(TC("  %s"), fileName);
 					}
 					fclose(fp);
 					if (countBeforeCallstack != out.count)
@@ -248,7 +248,7 @@ namespace uba
 			return;
 		}
 		for (int i = 0; i < nptrs; ++i)
-			out.Appendf(TC("%hs\n"), strings[i]);
+			out.Appendf(TC("%s\n"), strings[i]);
 		free(strings);
 #endif
 	}
diff --git a/Engine/Source/Programs/UnrealBuildAccelerator/Core/Private/UbaSynchronization.cpp b/Engine/Source/Programs/UnrealBuildAccelerator/Core/Private/UbaSynchronization.cpp
--- a/Engine/Source/Programs/UnrealBuildAccelerator/Core/Private/UbaSynchronization.cpp
+++ b/Engine/Source/Programs/UnrealBuildAccelerator/Core/Private/UbaSynchronization.cpp
@@ -9,9 +9,11 @@ namespace uba
 	CriticalSection::CriticalSection()
 	{
 		#if PLATFORM_WINDOWS
+		static_assert(alignof(CRITICAL_SECTION) == alignof(CriticalSection));
 		static_assert(sizeof(data) >= sizeof(CRITICAL_SECTION));
 		InitializeCriticalSection((CRITICAL_SECTION*)&data);
 		#else
+		static_assert(alignof(pthread_mutex_t) == alignof(CriticalSection));
 		static_assert(sizeof(data) >= sizeof(pthread_mutex_t));
 		pthread_mutexattr_t attr;
 		pthread_mutexattr_init(&attr);
@@ -51,8 +53,10 @@ namespace uba
 	{
 		#if PLATFORM_WINDOWS
 		static_assert(sizeof(data) >= sizeof(SRWLOCK));
+		static_assert(alignof(SRWLOCK) == alignof(ReaderWriterLock));
 		InitializeSRWLock((SRWLOCK*)&data);
 		#else
+		static_assert(alignof(pthread_rwlock_t) == alignof(ReaderWriterLock));
 		static_assert(sizeof(data) >= sizeof(pthread_rwlock_t));
 		pthread_rwlock_init((pthread_rwlock_t*)data, NULL);
 		#endif
diff --git a/Engine/Source/Programs/UnrealBuildAccelerator/Core/Public/UbaHash.h b/Engine/Source/Programs/UnrealBuildAccelerator/Core/Public/UbaHash.h
--- a/Engine/Source/Programs/UnrealBuildAccelerator/Core/Public/UbaHash.h
+++ b/Engine/Source/Programs/UnrealBuildAccelerator/Core/Public/UbaHash.h
@@ -15,6 +15,7 @@ namespace uba
 		u64 b;
 		bool operator==(const StringKey& o) const { return a == o.a && b == o.b; }
 		bool operator!=(const StringKey& o) const { return a != o.a || b != o.b; }
+		bool operator<(const StringKey& o) const { if (a != o.a) return a < o.a; return b < o.b; }
 	};
 	constexpr StringKey StringKeyZero;
 
@@ -29,6 +30,7 @@ namespace uba
 		u32 c;
 		bool operator==(const CasKey& o) const { return a == o.a && b == o.b && c == o.c; }
 		bool operator!=(const CasKey& o) const { return a != o.a || b != o.b || c != o.c; }
+		bool operator<(const CasKey& o) const { if (a != o.a) return a < o.a; if (b != o.b) return b < o.b; return c < o.c; }
 	};
 #pragma pack(pop)
 
diff --git a/Engine/Source/Programs/UnrealBuildAccelerator/Core/Public/UbaMemory.h b/Engine/Source/Programs/UnrealBuildAccelerator/Core/Public/UbaMemory.h
--- a/Engine/Source/Programs/UnrealBuildAccelerator/Core/Public/UbaMemory.h
+++ b/Engine/Source/Programs/UnrealBuildAccelerator/Core/Public/UbaMemory.h
@@ -40,6 +40,23 @@ namespace uba
 		void deallocate(value_type* p, u64 n) { mi_free_size_aligned(p, sizeof(value_type)*n, alignof(value_type)); }
 		u64 max_size() const { return static_cast<size_t>(-1) / sizeof(value_type); }
 	};
+	#elif 0
+	template<typename Type>
+	class Allocator {
+	public:
+		using value_type = Type;
+
+		Allocator() {}
+		Allocator(const Allocator& o) {}
+		Allocator(Allocator&& o) noexcept {}
+		constexpr bool operator==(const Allocator&) const noexcept { return true; }
+		template <class _Other>
+		constexpr Allocator(const Allocator<_Other>&) noexcept {}
+
+		value_type* allocate(u64 n) { return (value_type*)aligned_alloc(16, AlignUp(sizeof(value_type)*n, 16)); }
+		void deallocate(value_type* p, u64 n) { free(p); }
+		u64 max_size() const { return static_cast<size_t>(-1) / sizeof(value_type); }
+	};
 	#else
 	template<typename Type>
 	using Allocator = std::allocator<Type>;
diff --git a/Engine/Source/Programs/UnrealBuildAccelerator/StorageProxy/Private/UbaStorageProxyApp.cpp b/Engine/Source/Programs/UnrealBuildAccelerator/StorageProxy/Private/UbaStorageProxyApp.cpp
--- a/Engine/Source/Programs/UnrealBuildAccelerator/StorageProxy/Private/UbaStorageProxyApp.cpp
+++ b/Engine/Source/Programs/UnrealBuildAccelerator/StorageProxy/Private/UbaStorageProxyApp.cpp
@@ -10,7 +10,7 @@
 
 namespace uba
 {
-	const char*		Version = GetVersionString();
+	const tchar*	Version = GetVersionString();
 	u16				DefaultListenPort = DefaultPort + 1;
 	const wchar_t*	DefaultName = []() { static wchar_t buf[256]; GetComputerNameW(buf, sizeof_array(buf)); return buf; }();
 
@@ -25,7 +25,7 @@ namespace uba
 
 		logger.Info(L"");
 		logger.Info(L"------------------------");
-		logger.Info(L"   UbaProxy v%hs", Version);
+		logger.Info(L"   UbaProxy v%s", Version);
 		logger.Info(L"------------------------");
 		logger.Info(L"");
 		logger.Info(L"  When started UbaProxy will keep trying to connect to provided host address.");
@@ -120,7 +120,7 @@ int wmain(int argc, wchar_t *argv[])
 	#if UBA_DEBUG
 	dbgStr = L" (DEBUG)";
 	#endif
-	logger.Info(L"UbaStorageProxy v%hs%s", Version, dbgStr);
+	logger.Info(L"UbaStorageProxy v%s%s", Version, dbgStr);
 	logger.Info(L"");
 
 	while (true)
diff --git a/Engine/Source/Programs/UnrealBuildAccelerator/Version/Private/UbaVersion.cpp b/Engine/Source/Programs/UnrealBuildAccelerator/Version/Private/UbaVersion.cpp
--- a/Engine/Source/Programs/UnrealBuildAccelerator/Version/Private/UbaVersion.cpp
+++ b/Engine/Source/Programs/UnrealBuildAccelerator/Version/Private/UbaVersion.cpp
@@ -4,7 +4,7 @@
 
 namespace uba
 {
-	const char* GetVersionString()
+	const tchar* GetVersionString()
 	{
 		return UBA_VERSION;
 	}
diff --git a/Engine/Source/Programs/UnrealBuildAccelerator/Version/Public/UbaVersion.h b/Engine/Source/Programs/UnrealBuildAccelerator/Version/Public/UbaVersion.h
--- a/Engine/Source/Programs/UnrealBuildAccelerator/Version/Public/UbaVersion.h
+++ b/Engine/Source/Programs/UnrealBuildAccelerator/Version/Public/UbaVersion.h
@@ -4,5 +4,11 @@
 
 namespace uba
 {
-	const char* GetVersionString();
+	#if PLATFORM_WINDOWS
+	using tchar = wchar_t;
+	#else
+	using tchar = char;
+	#endif
+
+	const tchar* GetVersionString();
 }
\ No newline at end of file
diff --git a/Engine/Source/Programs/UnrealBuildAccelerator/Version/UbaVersion.Build.cs b/Engine/Source/Programs/UnrealBuildAccelerator/Version/UbaVersion.Build.cs
--- a/Engine/Source/Programs/UnrealBuildAccelerator/Version/UbaVersion.Build.cs
+++ b/Engine/Source/Programs/UnrealBuildAccelerator/Version/UbaVersion.Build.cs
@@ -10,6 +10,7 @@ public class UbaVersion : ModuleRules
 		UnsafeTypeCastWarningLevel = WarningLevel.Error;
 		PCHUsage = PCHUsageMode.NoPCHs;
 
-		PrivateDefinitions.Add($"UBA_VERSION=\"{Target.Version.MajorVersion}.{Target.Version.MinorVersion}.{Target.Version.PatchVersion}-{Target.BuildVersion}\"");
+		string StringPrefix = Target.Platform.IsInGroup(UnrealPlatformGroup.Windows) ? "L" : "";
+		PrivateDefinitions.Add($"UBA_VERSION={StringPrefix}\"{Target.Version.MajorVersion}.{Target.Version.MinorVersion}.{Target.Version.PatchVersion}-{Target.BuildVersion}\"");
 	}
 }
diff --git a/Engine/Source/Programs/UnrealBuildAccelerator/Visualizer/Private/UbaVisualizerApp.cpp b/Engine/Source/Programs/UnrealBuildAccelerator/Visualizer/Private/UbaVisualizerApp.cpp
--- a/Engine/Source/Programs/UnrealBuildAccelerator/Visualizer/Private/UbaVisualizerApp.cpp
+++ b/Engine/Source/Programs/UnrealBuildAccelerator/Visualizer/Private/UbaVisualizerApp.cpp
@@ -6,8 +6,6 @@
 
 using namespace uba;
 
-const char* Version = GetVersionString();
-
 int PrintHelp(const tchar* message)
 {
 	StringBuffer<64*1024> s;
@@ -16,7 +14,7 @@ int PrintHelp(const tchar* message)
 
 	s.Appendf(TC("\r\n"));
 	s.Appendf(TC("------------------------\r\n"));
-	s.Appendf(TC("   UbaVisualizer v%hs\r\n"), Version);
+	s.Appendf(TC("   UbaVisualizer v%s\r\n"), GetVersionString());
 	s.Appendf(TC("------------------------\r\n"));
 	s.Appendf(TC("\r\n"));
 	s.Appendf(TC("  When started UbaVisualizer will keep trying to connect to provided host address or named memory buffer.\r\n"));
