[Uba]
* Added support for environment variable UBA_ZONE

diff --git a/Engine/Source/Programs/UnrealBuildAccelerator/Agent/Private/UbaAgent.cpp b/Engine/Source/Programs/UnrealBuildAccelerator/Agent/Private/UbaAgent.cpp
--- a/Engine/Source/Programs/UnrealBuildAccelerator/Agent/Private/UbaAgent.cpp
+++ b/Engine/Source/Programs/UnrealBuildAccelerator/Agent/Private/UbaAgent.cpp
@@ -712,6 +712,9 @@ namespace uba
 		}
 		#endif
 
+		if (!zone.count)
+			zone.count = GetEnvironmentVariableW(TC("UBA_ZONE"), zone.data, zone.capacity);
+
 		if (zone.count)
 			extraInfo.Append(TC(", ")).Append(zone);
 
diff --git a/Engine/Source/Programs/UnrealBuildAccelerator/Host/Private/UbaExports.cpp b/Engine/Source/Programs/UnrealBuildAccelerator/Host/Private/UbaExports.cpp
--- a/Engine/Source/Programs/UnrealBuildAccelerator/Host/Private/UbaExports.cpp
+++ b/Engine/Source/Programs/UnrealBuildAccelerator/Host/Private/UbaExports.cpp
@@ -216,6 +216,13 @@ extern "C"
 		}
 		#endif
 
+		StringBuffer<256> zoneTemp;
+		if (!zone || !*zone)
+		{
+			zoneTemp.count = GetEnvironmentVariableW(TC("UBA_ZONE"), zoneTemp.data, zoneTemp.capacity);
+			zone = zoneTemp.data;
+		}
+
 		StorageServerCreateInfo info(server, rootDir, writer);
 		info.casCapacityBytes = casCapacityBytes;
 		info.storeCompressed = storeCompressed;
