[Uba]
* Fixed so uba uses virtual file index numbers of "real" index is (u64)-1. This can be used for systems like ProjFS where we don't know the file index until it is hydrated but we need a proper number in order for applications like clang to work properly (clang uses file index as unique key for files)
* Added some more debug code

diff --git a/Engine/Source/Programs/UnrealBuildAccelerator/Common/Private/UbaSession.cpp b/Engine/Source/Programs/UnrealBuildAccelerator/Common/Private/UbaSession.cpp
--- a/Engine/Source/Programs/UnrealBuildAccelerator/Common/Private/UbaSession.cpp
+++ b/Engine/Source/Programs/UnrealBuildAccelerator/Common/Private/UbaSession.cpp
@@ -233,11 +233,15 @@ namespace uba
 				g_debugLogger.Info(TC("    %s (Size: %llu, Key: %s, Id: %llu)\n"), e.name, e.size, KeyToString(fileKey).data, e.id);
 				#endif
 
+				u64 id = e.id;
+				if (id == 0xffffffffffffffffllu) // When using projfs we might not have the file yet and in that case we need to make this up.
+					id = ++m_fileIndexCounter;
+
 				res.first->second = u32(memoryWriter.GetPosition()); // Temporary offset that will be used further down to calculate the real offset
 				memoryWriter.WriteU64(e.lastWritten);
 				memoryWriter.WriteU32(e.attributes);
 				memoryWriter.WriteU32(e.volumeSerial);
-				memoryWriter.WriteU64(e.id);
+				memoryWriter.WriteU64(id);
 				memoryWriter.WriteU64(e.size);
 
 				FileEntryAdded(res.first->first, e.lastWritten, e.size);
diff --git a/Engine/Source/Programs/UnrealBuildAccelerator/Detours/Private/Windows/UbaDetoursFunctionsKernelBase.inl b/Engine/Source/Programs/UnrealBuildAccelerator/Detours/Private/Windows/UbaDetoursFunctionsKernelBase.inl
--- a/Engine/Source/Programs/UnrealBuildAccelerator/Detours/Private/Windows/UbaDetoursFunctionsKernelBase.inl
+++ b/Engine/Source/Programs/UnrealBuildAccelerator/Detours/Private/Windows/UbaDetoursFunctionsKernelBase.inl
@@ -1821,7 +1821,20 @@ BOOL Detoured_GetFileInformationByHandle(HANDLE hFile, LPBY_HANDLE_FILE_INFORMAT
 			{
 				SuppressDetourScope _;
 				WIN32_FILE_ATTRIBUTE_DATA data;
-				if (True_GetFileAttributesExW(fi.originalName, GetFileExInfoStandard, &data))
+				if (dh.trueHandle != INVALID_HANDLE_VALUE)
+				{
+					BY_HANDLE_FILE_INFORMATION bhfi;
+					auto res2 = True_GetFileInformationByHandle(dh.trueHandle, &bhfi);
+					UBA_ASSERT(res2 == TRUE);
+					u64 size = ToLargeInteger(bhfi.nFileSizeHigh, bhfi.nFileSizeLow).QuadPart; (void)size;
+					u64 fileIndex = ToLargeInteger(bhfi.nFileIndexHigh, bhfi.nFileIndexLow).QuadPart; (void)fileIndex;
+					UBA_ASSERTF(fileSize == size, L"File size used: %llu Actual file size: %llu (%s)", fileSize, size, fi.originalName);
+					//UBA_ASSERTF(entryInfo.attributes == bhfi.dwFileAttributes, L"Attributes used: 0x%x Actual: 0x%x (%s)", entryInfo.attributes, bhfi.dwFileAttributes, fi.originalName);
+					UBA_ASSERTF(entryInfo.volumeSerial == bhfi.dwVolumeSerialNumber, L"VolumeSerial used: %u Actual: %u (%s)", entryInfo.volumeSerial, bhfi.dwVolumeSerialNumber, fi.originalName);
+					UBA_ASSERTF(entryInfo.fileIndex == fileIndex, L"FileIndex used: %llu Actual: %llu (%s)", entryInfo.fileIndex, fileIndex, fi.originalName);
+					UBA_ASSERTF(bhfi.nNumberOfLinks == 1, L"Links used: %llu Actual: %llu (%s)", 1, bhfi.nNumberOfLinks, fi.originalName);
+				}
+				else if (True_GetFileAttributesExW(fi.originalName, GetFileExInfoStandard, &data))
 				{
 					u64 size = ToLargeInteger(data.nFileSizeHigh, data.nFileSizeLow).QuadPart; (void)size;
 					UBA_ASSERTF(fileSize == size, L"File size used: %llu Actual file size: %llu (%s)", fileSize, size, fi.originalName);
