[UBA] Fixed NtCreateFile filename string capacity
Fixed the filename string capacity in NtCreateFile before the call to True_NtCreateFile to be +2 of the string length to accomodate the string termination chars

diff --git a/Engine/Source/Programs/UnrealBuildAccelerator/Detours/Private/Windows/UbaDetoursFunctionsNtDll.inl b/Engine/Source/Programs/UnrealBuildAccelerator/Detours/Private/Windows/UbaDetoursFunctionsNtDll.inl
--- a/Engine/Source/Programs/UnrealBuildAccelerator/Detours/Private/Windows/UbaDetoursFunctionsNtDll.inl
+++ b/Engine/Source/Programs/UnrealBuildAccelerator/Detours/Private/Windows/UbaDetoursFunctionsNtDll.inl
@@ -1035,7 +1035,7 @@ NTSTATUS NTAPI Shared_NtCreateFile(bool IsCreateFunc, PHANDLE hFileHandle, ACCES
 	UNICODE_STRING str;
 	str.Buffer = temp.data;
 	str.Length = u16(temp.count * 2);
-	str.MaximumLength = str.Length;
+	str.MaximumLength = str.Length + 2;
 	ObjectAttributes->ObjectName = &str;
 	// TODO!!! THIS NEEDS TO set the ObjectAttributes->ObjectName->Buffer and ObjectAttributes->ObjectName->Length;
 	//wcscpy_s(ObjectAttributes->ObjectName->Buffer + 4, ObjectAttributes->ObjectName->MaximumLength/2 - 8, lpFileName);
