[UBA]
* Added hint to unmapfileview to be able to add useful information to asserts on fails
* Added error handling for vsnprintf returning error
* Fixed memory stomp in unit test
* Removed usage of %hs on non-windows since asan complains about it
* Changed UBA_USE_MIMALLOC to always be defined and be 0 or 1 instead of using ifdef
* Enabled mimalloc on linux

diff --git a/Engine/Source/Programs/UnrealBuildAccelerator/Agent/Private/UbaAgent.cpp b/Engine/Source/Programs/UnrealBuildAccelerator/Agent/Private/UbaAgent.cpp
--- a/Engine/Source/Programs/UnrealBuildAccelerator/Agent/Private/UbaAgent.cpp
+++ b/Engine/Source/Programs/UnrealBuildAccelerator/Agent/Private/UbaAgent.cpp
@@ -349,7 +349,7 @@ namespace uba
 				terminateTimeMsUtc = strtoull(lineBegin, nullptr, 10);
 				break;
 			case 3: // reason
-				outReason.Appendf(TC("%hs"), lineBegin);
+				outReason.Appendf(PERCENT_HS, lineBegin);
 				break;
 			}
 
@@ -1127,6 +1127,7 @@ namespace uba
 				if (proxy.server)
 					proxy.server->StopAll();
 				client->Disconnect();
+				sessionClient->Stop();
 				storageClient->StopProxy();
 				loopLogging = false;
 				logLinesAvailable.Set();
diff --git a/Engine/Source/Programs/UnrealBuildAccelerator/Common/Private/UbaFile.cpp b/Engine/Source/Programs/UnrealBuildAccelerator/Common/Private/UbaFile.cpp
--- a/Engine/Source/Programs/UnrealBuildAccelerator/Common/Private/UbaFile.cpp
+++ b/Engine/Source/Programs/UnrealBuildAccelerator/Common/Private/UbaFile.cpp
@@ -94,13 +94,8 @@ namespace uba
 #if PLATFORM_WINDOWS
 			DWORD wasRead = 0;
 			if (!::ReadFile(asHANDLE(fileHandle), buffer, toRead, &wasRead, NULL))
-			{
 				if (GetLastError() != ERROR_IO_PENDING)
-				{
-					logger.Error(TC("ERROR reading file %s (error: %s)"), fileName, LastErrorToText().data);
-					return false;
-				}
-			}
+					return logger.Error(TC("ERROR reading file %s (error: %s)"), fileName, LastErrorToText().data);
 #else
 			ssize_t wasRead = read(asFileDescriptor(fileHandle), buffer, toRead);
 			if (wasRead == -1)
@@ -697,7 +692,11 @@ namespace uba
 		}
 
 		char fullPath[1024];
-		realpath(fileName, fullPath);
+		if (!realpath(fileName, fullPath))
+		{
+			UBA_ASSERTF(false, TC("realpath error handling not implemented for path %s (%s)"), fileName, strerror(errno));
+			return 0;
+		}
 		u32 len = TStrlen(fullPath);
 		UBA_ASSERT(len < sizeof_array(fullPath));
 		UBA_ASSERT(len < nBufferLength);
diff --git a/Engine/Source/Programs/UnrealBuildAccelerator/Common/Private/UbaFileAccessor.cpp b/Engine/Source/Programs/UnrealBuildAccelerator/Common/Private/UbaFileAccessor.cpp
--- a/Engine/Source/Programs/UnrealBuildAccelerator/Common/Private/UbaFileAccessor.cpp
+++ b/Engine/Source/Programs/UnrealBuildAccelerator/Common/Private/UbaFileAccessor.cpp
@@ -309,7 +309,7 @@ namespace uba
 	{
 		if (m_data)
 		{
-			if (!UnmapViewOfFile(m_data, m_size))
+			if (!UnmapViewOfFile(m_data, m_size, m_fileName))
 				return m_logger.Error(TC("Failed to unmap memory for %s (%s)"), m_fileName, LastErrorToText().data);
 			m_data = nullptr;
 		}
diff --git a/Engine/Source/Programs/UnrealBuildAccelerator/Common/Private/UbaFileMapping.cpp b/Engine/Source/Programs/UnrealBuildAccelerator/Common/Private/UbaFileMapping.cpp
--- a/Engine/Source/Programs/UnrealBuildAccelerator/Common/Private/UbaFileMapping.cpp
+++ b/Engine/Source/Programs/UnrealBuildAccelerator/Common/Private/UbaFileMapping.cpp
@@ -227,13 +227,13 @@ namespace uba
 #endif
 	}
 
-	bool UnmapViewOfFile(const void* lpBaseAddress, u64 bytesToUnmap)
+	bool UnmapViewOfFile(const void* lpBaseAddress, u64 bytesToUnmap, const tchar* hint)
 	{
 		ExtendedTimerScope ts(SystemStats::GetCurrent().unmapViewOfFile);
 #if PLATFORM_WINDOWS
 		(void)bytesToUnmap; return ::UnmapViewOfFile(lpBaseAddress);
 #else
-		UBA_ASSERT(bytesToUnmap);
+		UBA_ASSERTF(bytesToUnmap, TC("bytesToUnmap is zero unmapping %p (%s)"), lpBaseAddress, hint);
 		if (munmap((void*)lpBaseAddress, bytesToUnmap) == 0)
 			return true;
 		UBA_ASSERT(false);
@@ -386,7 +386,7 @@ namespace uba
 			u64 commitSize = commitedAfter - committedBefore;
 			if (!MapViewCommit(data + commitStart, commitSize))
 			{
-				UnmapViewOfFile(data, mapSize);
+				UnmapViewOfFile(data, mapSize, hint);
 				m_logger.Error(TC("%s - Failed to allocate memory for %s (%s)"), f.name, hint, LastErrorToText().data);
 				return res;
 			}
@@ -531,7 +531,7 @@ namespace uba
 
 		u8* memory = view.memory - (view.offset - alignedOffsetStart);
 		u64 mapSize = alignedOffsetEnd - alignedOffsetStart;
-		if (!UnmapViewOfFile(memory, mapSize))
+		if (!UnmapViewOfFile(memory, mapSize, hint))
 		{
 			m_logger.Error(TC("%s - Failed to unmap view on address %llx (offset %llu) - %s (%s)"), file.name, u64(memory), view.offset, hint, LastErrorToText().data);
 		}
@@ -643,7 +643,7 @@ namespace uba
 	void FileMappingAllocator::Free(Allocation allocation)
 	{
 		UBA_ASSERT(allocation.handle == m_mappingHandle);
-		if (!UnmapViewOfFile(allocation.memory, m_blockSize))
+		if (!UnmapViewOfFile(allocation.memory, m_blockSize, m_name))
 			m_logger.Error(TC("%s - Failed to unmap view of file (%s)"), m_name, LastErrorToText().data);
 		u64 index = allocation.offset / m_blockSize;
 		ScopedWriteLock lock(m_mappingLock);
diff --git a/Engine/Source/Programs/UnrealBuildAccelerator/Common/Private/UbaNetworkBackendTcp.cpp b/Engine/Source/Programs/UnrealBuildAccelerator/Common/Private/UbaNetworkBackendTcp.cpp
--- a/Engine/Source/Programs/UnrealBuildAccelerator/Common/Private/UbaNetworkBackendTcp.cpp
+++ b/Engine/Source/Programs/UnrealBuildAccelerator/Common/Private/UbaNetworkBackendTcp.cpp
@@ -700,9 +700,9 @@ namespace uba
 			int sent = (int)send(socket, (char*)b, u32(bufferLen), 0);
 			if (sent == SOCKET_ERROR)
 			{
-				#if UBA_DEBUG
-				logger.Warning(TC("ERROR sending socket (error: %s)"), LastErrorToText(WSAGetLastError()).data);
-				#endif
+				//#if UBA_DEBUG
+				//logger.Warning(TC("ERROR sending socket (error: %s)"), LastErrorToText(WSAGetLastError()).data);
+				//#endif
 				return false;
 			}
 
diff --git a/Engine/Source/Programs/UnrealBuildAccelerator/Common/Private/UbaSession.cpp b/Engine/Source/Programs/UnrealBuildAccelerator/Common/Private/UbaSession.cpp
--- a/Engine/Source/Programs/UnrealBuildAccelerator/Common/Private/UbaSession.cpp
+++ b/Engine/Source/Programs/UnrealBuildAccelerator/Common/Private/UbaSession.cpp
@@ -1103,10 +1103,10 @@ namespace uba
 		//for (auto& i : m_fileMappingTableLookup)
 		//	CloseHandle(i.second.mapping);
 
-		UnmapViewOfFile(m_fileMappingTableMem, FileMappingTableMemSize);
+		UnmapViewOfFile(m_fileMappingTableMem, FileMappingTableMemSize, TC("FileMappingTable"));
 		CloseFileMapping(m_fileMappingTableHandle);
 
-		UnmapViewOfFile(m_directoryTableMem, DirTableMemSize);
+		UnmapViewOfFile(m_directoryTableMem, DirTableMemSize, TC("DirectoryTable"));
 		CloseFileMapping(m_directoryTableHandle);
 
 		//#if !UBA_DEBUG
@@ -1918,7 +1918,7 @@ namespace uba
 			if (!mem)
 				return m_logger.Error(TC("Failed to map view of filehandle for read %s (%s)"), file.name.c_str(), LastErrorToText().data);
 
-			auto memClose = MakeGuard([&](){ UnmapViewOfFile(mem, fileSize); });
+			auto memClose = MakeGuard([&](){ UnmapViewOfFile(mem, fileSize, file.name.c_str()); });
 
 			// Seems like best combo (for windows at least) is to use writes with overlap and max 16 at the same time.
 			// On one machine we get twice as fast without overlap if no bottleneck. On another machine (ntfs compression on) we get twice as slow without overlap
diff --git a/Engine/Source/Programs/UnrealBuildAccelerator/Common/Private/UbaSessionClient.cpp b/Engine/Source/Programs/UnrealBuildAccelerator/Common/Private/UbaSessionClient.cpp
--- a/Engine/Source/Programs/UnrealBuildAccelerator/Common/Private/UbaSessionClient.cpp
+++ b/Engine/Source/Programs/UnrealBuildAccelerator/Common/Private/UbaSessionClient.cpp
@@ -60,6 +60,11 @@ namespace uba
 	}
 
 	SessionClient::~SessionClient()
+	{
+		Stop();
+	}
+
+	void SessionClient::Stop()
 	{
 		CancelAllProcessesAndWait();
 		m_loop = false;
diff --git a/Engine/Source/Programs/UnrealBuildAccelerator/Common/Private/UbaStats.cpp b/Engine/Source/Programs/UnrealBuildAccelerator/Common/Private/UbaStats.cpp
--- a/Engine/Source/Programs/UnrealBuildAccelerator/Common/Private/UbaStats.cpp
+++ b/Engine/Source/Programs/UnrealBuildAccelerator/Common/Private/UbaStats.cpp
@@ -23,7 +23,7 @@ namespace uba
 		const tchar empty[] = TC("                   ");
 		for (Stat& s : stats)
 			if (s.timer.count)
-				logger.Info(TC("  %c%hs%s %8u %9s"), ToUpper(s.name[0]), s.name+1, empty + s.nameLen, s.timer.count.load(), TimeToText(s.timer.time, false, frequency).str);
+				logger.Info(TC("  %c") PERCENT_HS TC("%s %8u %9s"), ToUpper(s.name[0]), s.name + 1, empty + s.nameLen, s.timer.count.load(), TimeToText(s.timer.time, false, frequency).str);
 
 		logger.Info(TC(""));
 
@@ -50,7 +50,7 @@ namespace uba
 		const tchar empty[] = TC("                   ");
 		for (Stat& s : stats)
 			if (s.timer.count)
-				logger.Info(TC("  %c%hs%s %8u %9s"), ToUpper(s.name[0]), s.name+1, empty + s.nameLen, s.timer.count.load(), TimeToText(s.timer.time, false, frequency).str);
+				logger.Info(TC("  %c") PERCENT_HS TC("%s %8u %9s"), ToUpper(s.name[0]), s.name+1, empty + s.nameLen, s.timer.count.load(), TimeToText(s.timer.time, false, frequency).str);
 
 		if (writeHeader)
 			logger.Info(TC(""));
@@ -179,7 +179,7 @@ namespace uba
 		if (!timer.count)
 			return;
 		const tchar empty[] = TC("                   ");
-		logger.Info(TC("  %c%hs%s %8u %9s"), ToUpper(name[0]), name+1, empty + strlen(name)+1, timer.count.load(), TimeToText(timer.time, false, frequency).str);
+		logger.Info(TC("  %c") PERCENT_HS TC("%s %8u %9s"), ToUpper(name[0]), name+1, empty + strlen(name)+1, timer.count.load(), TimeToText(timer.time, false, frequency).str);
 	}
 
 	void SessionStats::Print(Logger& logger, u64 frequency)
diff --git a/Engine/Source/Programs/UnrealBuildAccelerator/Common/Private/UbaStorage.cpp b/Engine/Source/Programs/UnrealBuildAccelerator/Common/Private/UbaStorage.cpp
--- a/Engine/Source/Programs/UnrealBuildAccelerator/Common/Private/UbaStorage.cpp
+++ b/Engine/Source/Programs/UnrealBuildAccelerator/Common/Private/UbaStorage.cpp
@@ -1,15 +1,12 @@
 // Copyright Epic Games, Inc. All Rights Reserved.
 
 #include "UbaStorage.h"
+#include "UbaBottleneck.h"
 #include "UbaFileAccessor.h"
 #include "UbaBinaryReaderWriter.h"
 #include "UbaDirectoryIterator.h"
 #include "UbaWorkManager.h"
 
-#if defined(UBA_USE_MIMALLOC)
-#include <mimalloc.h>
-#endif
-
 namespace uba
 {
 	constexpr u32 CasTableVersion = IsWindows ? 32 : 34;
@@ -256,7 +253,7 @@ namespace uba
 				if (!uncompressedData)
 					return m_logger.Error(TC("Failed to map view of file mapping for %s (%s)"), from, LastErrorToText().data);
 
-				auto udg = MakeGuard([&]() { UnmapViewOfFile(uncompressedData, fileSize); });
+				auto udg = MakeGuard([&]() { UnmapViewOfFile(uncompressedData, fileSize, from); });
 
 				if (!WriteMemToCompressedFile(destinationFile, workCount, uncompressedData, fileSize, maxUncompressedBlock, totalWritten))
 					return false;
@@ -965,7 +962,7 @@ namespace uba
 		return true;
 	}
 
-	#if defined(UBA_USE_MIMALLOC)
+	#if UBA_USE_MIMALLOC
 	void* Oodle_MallocAligned(OO_SINTa bytes, OO_S32 alignment)
 	{
 		return mi_malloc_aligned(bytes, alignment);
@@ -999,7 +996,7 @@ namespace uba
 
 		m_maxParallelCopyOrLink = info.maxParallelCopyOrLink;
 
-		#if defined(UBA_USE_MIMALLOC)
+		#if UBA_USE_MIMALLOC
 		OodleCore_Plugins_SetAllocators(Oodle_MallocAligned, Oodle_Free);
 		#endif
 	}
@@ -1931,6 +1928,13 @@ namespace uba
 				ScopedWriteLock entryLock(entry.lock);
 				entry.verified = false;
 
+				// This is to reduce number of active CreateFiles.. seems like machines don't like tons of CreateFile at the same time
+				#if PLATFORM_WINDOWS
+				constexpr u32 bottleneckMax = 16;
+				static Bottleneck bottleneck(bottleneckMax);
+				BottleneckScope scope(bottleneck);
+				#endif
+
 				if (writeDirectlyToFile || !decompressedSize)
 				{
 					if (!destinationFile.CreateWrite(allowRead, writeFlags | fileAttributes, decompressedSize, m_tempPath.data))
@@ -2233,7 +2237,7 @@ namespace uba
 				m_logger.Error(TC("Failed to map view of file mapping for %s (%s)"), fileName, LastErrorToText().data);
 				return CasKeyZero;
 			}
-			auto udg = MakeGuard([&]() { UnmapViewOfFile(fileData, fileSize); });
+			auto udg = MakeGuard([&]() { UnmapViewOfFile(fileData, fileSize, fileName); });
 
 			struct WorkRec
 			{
@@ -2336,7 +2340,7 @@ namespace uba
 			u8* fileData = MapViewOfFile(fileMapping, FILE_MAP_READ, 0, compressedSize);
 			if (!fileData)
 				return m_logger.Error(TC("Failed to map view of file mapping for %s (%s)"), fileName, LastErrorToText().data);
-			auto udg = MakeGuard([&]() { UnmapViewOfFile(fileData, compressedSize); });
+			auto udg = MakeGuard([&]() { UnmapViewOfFile(fileData, compressedSize, fileName); });
 			
 			if (!DecompressMemoryToMemory(fileData + 8, dest, decompressedSize))
 				return false;
diff --git a/Engine/Source/Programs/UnrealBuildAccelerator/Common/Private/UbaStorageClient.cpp b/Engine/Source/Programs/UnrealBuildAccelerator/Common/Private/UbaStorageClient.cpp
--- a/Engine/Source/Programs/UnrealBuildAccelerator/Common/Private/UbaStorageClient.cpp
+++ b/Engine/Source/Programs/UnrealBuildAccelerator/Common/Private/UbaStorageClient.cpp
@@ -693,7 +693,7 @@ namespace uba
 			isPersistentMapping = true;
 		}
 
-		auto unmapGuard = MakeGuard([&](){ if (isPersistentMapping) UnmapViewOfFile(fileMem, fileSize); });
+		auto unmapGuard = MakeGuard([&](){ if (isPersistentMapping) UnmapViewOfFile(fileMem, fileSize, fileName); });
 
 		CasKey casKey;
 
diff --git a/Engine/Source/Programs/UnrealBuildAccelerator/Common/Private/UbaStorageServer.cpp b/Engine/Source/Programs/UnrealBuildAccelerator/Common/Private/UbaStorageServer.cpp
--- a/Engine/Source/Programs/UnrealBuildAccelerator/Common/Private/UbaStorageServer.cpp
+++ b/Engine/Source/Programs/UnrealBuildAccelerator/Common/Private/UbaStorageServer.cpp
@@ -171,7 +171,7 @@ namespace uba
 	{
 		u8* fileMem = MapViewOfFile(mappingHandle, FILE_MAP_READ, mappingOffset, fileSize);
 		UBA_ASSERT(fileMem);
-		auto memClose = MakeGuard([&](){ UnmapViewOfFile(fileMem, fileSize); });
+		auto memClose = MakeGuard([&](){ UnmapViewOfFile(fileMem, fileSize, fileName); });
 
 		bool storeCompressed = true;
 		out = CalculateCasKey(fileMem, fileSize, storeCompressed);
@@ -217,7 +217,7 @@ namespace uba
 		lock.Leave();
 		u8* fileMem = MapViewOfFile(mapping.mappingHandle, FILE_MAP_READ, mapping.mappingOffset, mapping.fileSize);
 		UBA_ASSERT(fileMem);
-		auto memClose = MakeGuard([&](){ UnmapViewOfFile(fileMem, mapping.fileSize); });
+		auto memClose = MakeGuard([&](){ UnmapViewOfFile(fileMem, mapping.fileSize, from); });
 		return StorageImpl::WriteCompressed(out, from, InvalidFileHandle, fileMem, mapping.fileSize, toFile);
 	}
 
diff --git a/Engine/Source/Programs/UnrealBuildAccelerator/Common/Private/UbaTrace.cpp b/Engine/Source/Programs/UnrealBuildAccelerator/Common/Private/UbaTrace.cpp
--- a/Engine/Source/Programs/UnrealBuildAccelerator/Common/Private/UbaTrace.cpp
+++ b/Engine/Source/Programs/UnrealBuildAccelerator/Common/Private/UbaTrace.cpp
@@ -17,7 +17,7 @@ namespace uba
 	Trace::~Trace()
 	{
 		if (m_memoryBegin)
-			UnmapViewOfFile(m_memoryBegin, 0);
+			UnmapViewOfFile(m_memoryBegin, m_memoryCapacity, TC("Trace"));
 		if (m_memoryHandle.IsValid())
 			CloseFileMapping(m_memoryHandle);
 	}
@@ -108,7 +108,7 @@ namespace uba
 			return false;
 		m_logger.Info(TC("Trace file written to %s with size %s"), writeFileName, BytesToText(fileSize).str);
 		
-		UnmapViewOfFile(m_memoryBegin, m_memoryCapacity);
+		UnmapViewOfFile(m_memoryBegin, m_memoryCapacity, TC("Trace"));
 		m_memoryBegin = nullptr;
 		return true;
 	}
@@ -330,7 +330,7 @@ namespace uba
 		if (isCreator)
 			*(tchar*)m_mem = 0;
 
-		auto mg = MakeGuard([&]() { UnmapViewOfFile(m_mem, 256); m_mem = nullptr; });
+		auto mg = MakeGuard([&]() { UnmapViewOfFile(m_mem, 256, channelMutex.data); m_mem = nullptr; });
 
 		channelMutex.Append(channelName).Append(TC("Mutex"));
 		m_mutex = CreateMutexW(false, channelMutex.data);
diff --git a/Engine/Source/Programs/UnrealBuildAccelerator/Common/Public/UbaFileMapping.h b/Engine/Source/Programs/UnrealBuildAccelerator/Common/Public/UbaFileMapping.h
--- a/Engine/Source/Programs/UnrealBuildAccelerator/Common/Public/UbaFileMapping.h
+++ b/Engine/Source/Programs/UnrealBuildAccelerator/Common/Public/UbaFileMapping.h
@@ -49,7 +49,7 @@ namespace uba
 	FileMappingHandle CreateFileMappingW(FileHandle file, u32 flProtect, u64 maxSize = 0, const tchar* hint = TC(""));
 	u8* MapViewOfFile(FileMappingHandle fileMappingObject, u32 dwDesiredAccess, u64 offset, u64 dwNumberOfBytesToMap);
 	bool MapViewCommit(void* address, u64 size);
-	bool UnmapViewOfFile(const void* lpBaseAddress, u64 bytesToUnmap);
+	bool UnmapViewOfFile(const void* lpBaseAddress, u64 bytesToUnmap, const tchar* hint);
 	bool CloseFileMapping(FileMappingHandle h);
 	bool DuplicateFileMapping(ProcHandle hSourceProcessHandle, FileMappingHandle hSourceHandle, ProcHandle hTargetProcessHandle, FileMappingHandle* lpTargetHandle, u32 dwDesiredAccess, bool bInheritHandle, u32 dwOptions);
 
diff --git a/Engine/Source/Programs/UnrealBuildAccelerator/Common/Public/UbaSessionClient.h b/Engine/Source/Programs/UnrealBuildAccelerator/Common/Public/UbaSessionClient.h
--- a/Engine/Source/Programs/UnrealBuildAccelerator/Common/Public/UbaSessionClient.h
+++ b/Engine/Source/Programs/UnrealBuildAccelerator/Common/Public/UbaSessionClient.h
@@ -34,6 +34,7 @@ namespace uba
 		SessionClient(const SessionClientCreateInfo& info);
 		~SessionClient();
 
+		void Stop();
 		bool Wait(u32 milliseconds = 0xFFFFFFFF, Event* wakeupEvent = nullptr);
 		void SetIsTerminating(const tchar* reason = TC("Terminating"), u64 delayMs = 0); // Session stores pointer directly. Can't be temporary
 
diff --git a/Engine/Source/Programs/UnrealBuildAccelerator/Core/Private/UbaStringBuffer.cpp b/Engine/Source/Programs/UnrealBuildAccelerator/Core/Private/UbaStringBuffer.cpp
--- a/Engine/Source/Programs/UnrealBuildAccelerator/Core/Private/UbaStringBuffer.cpp
+++ b/Engine/Source/Programs/UnrealBuildAccelerator/Core/Private/UbaStringBuffer.cpp
@@ -193,33 +193,38 @@ namespace uba
 		return Append(buf);
 	}
 
+	StringBufferBase& ReplaceEnd(StringBufferBase& sb, const tchar* str)
+	{
+		auto len = TStrlen(str);
+		if (sb.count > sb.capacity - len - 1)
+			sb.count = sb.capacity - len - 1;
+		return sb.Append(str, len);
+	}
+
 	StringBufferBase& StringBufferBase::Append(const tchar* format, va_list& args)
 	{
 		#if PLATFORM_WINDOWS
-		u32 len = _vscwprintf(format, args);
+		int len = _vscwprintf(format, args);
 		va_list& args2 = args;
 		#else
 		va_list args2;
 		va_copy(args2, args);
 		auto g = MakeGuard([&]() { va_end(args2); });
 		int len = vsnprintf(0, 0, format, args);
-		UBA_ASSERT(len >= 0);
 		#endif
 
-		if (u32(len) >= capacity - count)
-		{
-			// 16 chars needed to fit buffer overflow string
-			if (count > capacity - 16)
-				count = capacity - 16;
-			Append(TC("BUFFEROVERFLOW!"));
-			return *this;
-		}
+		if (len < 0)
+			return ReplaceEnd(*this, TC("PRINTF ERROR!"));
+
+		if (len >= int(capacity - count) - 1)
+			return ReplaceEnd(*this, TC("BUFFEROVERFLOW!"));
+
 		int res = Tvsprintf_s(data + count, capacity - count, format, args2);
 
 		if (res > 0)
 			count += u32(res);
 		else
-			Append(TC("SPRINTF_ERROR!"));
+			return ReplaceEnd(*this, TC("SPRINTF_ERROR!"));
 
 		return *this;
 	}
diff --git a/Engine/Source/Programs/UnrealBuildAccelerator/Core/Public/UbaBase.h b/Engine/Source/Programs/UnrealBuildAccelerator/Core/Public/UbaBase.h
--- a/Engine/Source/Programs/UnrealBuildAccelerator/Core/Public/UbaBase.h
+++ b/Engine/Source/Programs/UnrealBuildAccelerator/Core/Public/UbaBase.h
@@ -23,10 +23,12 @@ namespace uba
 	inline constexpr bool IsWindows = true;
 	using tchar = wchar_t;
 	#define TC(x) L##x
+	#define PERCENT_HS L"%hs"
 	#else
 	inline constexpr bool IsWindows = false;
 	using tchar = char;
 	#define TC(x) x
+	#define PERCENT_HS "%s"
 	#endif
 }
 
diff --git a/Engine/Source/Programs/UnrealBuildAccelerator/Core/Public/UbaCorePch.h b/Engine/Source/Programs/UnrealBuildAccelerator/Core/Public/UbaCorePch.h
--- a/Engine/Source/Programs/UnrealBuildAccelerator/Core/Public/UbaCorePch.h
+++ b/Engine/Source/Programs/UnrealBuildAccelerator/Core/Public/UbaCorePch.h
@@ -2,7 +2,9 @@
 
 #pragma once
 
-#if defined(UBA_USE_MIMALLOC)
+#if UBA_USE_MIMALLOC
 #include <mimalloc.h>
 #include <mimalloc-override.h>
 #endif
+
+#include "UbaPlatform.h"
diff --git a/Engine/Source/Programs/UnrealBuildAccelerator/Core/Public/UbaMemory.h b/Engine/Source/Programs/UnrealBuildAccelerator/Core/Public/UbaMemory.h
--- a/Engine/Source/Programs/UnrealBuildAccelerator/Core/Public/UbaMemory.h
+++ b/Engine/Source/Programs/UnrealBuildAccelerator/Core/Public/UbaMemory.h
@@ -12,7 +12,7 @@
 #include <vector>
 #include <functional>
 
-#if defined(UBA_USE_MIMALLOC)
+#if UBA_USE_MIMALLOC
 #include <mimalloc.h>
 #endif
 
@@ -22,7 +22,7 @@ namespace uba
 	constexpr inline T AlignUp(T arg, uintptr_t alignment) { return T(((uintptr_t)arg+(alignment-1)) & ~(alignment-1)); }
 
 
-	#if defined(UBA_USE_MIMALLOC)
+	#if UBA_USE_MIMALLOC
 	template<typename Type>
 	class Allocator {
 
diff --git a/Engine/Source/Programs/UnrealBuildAccelerator/Detours/Private/Windows/UbaDetoursFunctionsMiMalloc.inl b/Engine/Source/Programs/UnrealBuildAccelerator/Detours/Private/Windows/UbaDetoursFunctionsMiMalloc.inl
--- a/Engine/Source/Programs/UnrealBuildAccelerator/Detours/Private/Windows/UbaDetoursFunctionsMiMalloc.inl
+++ b/Engine/Source/Programs/UnrealBuildAccelerator/Detours/Private/Windows/UbaDetoursFunctionsMiMalloc.inl
@@ -1,6 +1,6 @@
 // Copyright Epic Games, Inc. All Rights Reserved.
 
-#if defined(UBA_USE_MIMALLOC)
+#if UBA_USE_MIMALLOC
 
 #define DETOURED_CALL_MEM(x) DETOURED_CALL(x)
 
diff --git a/Engine/Source/Programs/UnrealBuildAccelerator/Detours/Private/Windows/UbaDetoursFunctionsNtDll.inl b/Engine/Source/Programs/UnrealBuildAccelerator/Detours/Private/Windows/UbaDetoursFunctionsNtDll.inl
--- a/Engine/Source/Programs/UnrealBuildAccelerator/Detours/Private/Windows/UbaDetoursFunctionsNtDll.inl
+++ b/Engine/Source/Programs/UnrealBuildAccelerator/Detours/Private/Windows/UbaDetoursFunctionsNtDll.inl
@@ -412,7 +412,7 @@ bool g_checkRtlHeap = true;
 
 SIZE_T Detoured_RtlSizeHeap(HANDLE HeapPtr, ULONG Flags, PVOID Ptr)
 {
-#if defined(UBA_USE_MIMALLOC)
+#if UBA_USE_MIMALLOC
 	if (g_checkRtlHeap && IsInMiMalloc(Ptr))
 		return mi_usable_size(Ptr);
 #endif
@@ -421,7 +421,7 @@ SIZE_T Detoured_RtlSizeHeap(HANDLE HeapPtr, ULONG Flags, PVOID Ptr)
 
 BOOLEAN Detoured_RtlFreeHeap(PVOID HeapHandle, ULONG Flags, PVOID BaseAddress)
 {
-#if defined(UBA_USE_MIMALLOC)
+#if UBA_USE_MIMALLOC
 	if (g_checkRtlHeap && IsInMiMalloc(BaseAddress))
 	{
 		mi_free(BaseAddress);
@@ -433,7 +433,7 @@ BOOLEAN Detoured_RtlFreeHeap(PVOID HeapHandle, ULONG Flags, PVOID BaseAddress)
 
 NTSTATUS Detoured_RtlAnsiStringToUnicodeString(PUNICODE_STRING DestinationString, PCANSI_STRING SourceString, BOOLEAN AllocateDestinationString)
 {
-#if defined(UBA_USE_MIMALLOC)
+#if UBA_USE_MIMALLOC
 	if (AllocateDestinationString)
 	{
 		DestinationString->MaximumLength = SourceString->MaximumLength * 2;
@@ -447,7 +447,7 @@ NTSTATUS Detoured_RtlAnsiStringToUnicodeString(PUNICODE_STRING DestinationString
 
 NTSTATUS Detoured_RtlUnicodeStringToAnsiString(PANSI_STRING DestinationString, PCUNICODE_STRING SourceString, BOOLEAN AllocateDestinationString)
 {
-#if defined(UBA_USE_MIMALLOC)
+#if UBA_USE_MIMALLOC
 	if (AllocateDestinationString)
 	{
 		DestinationString->MaximumLength = SourceString->MaximumLength / 2;
@@ -1417,7 +1417,7 @@ PVOID Detoured_RtlAllocateHeap(PVOID HeapHandle, ULONG Flags, SIZE_T Size)
 
 PVOID Detoured_RtlReAllocateHeap(PVOID HeapHandle, ULONG Flags, PVOID BaseAddress, SIZE_T Size)
 {
-#if defined(UBA_USE_MIMALLOC)
+#if UBA_USE_MIMALLOC
 	if (g_checkRtlHeap && IsInMiMalloc(BaseAddress))
 	{
 		Rpc_WriteLogf(L"ERROR: RtlReAllocateHeap - This is not implemented");
diff --git a/Engine/Source/Programs/UnrealBuildAccelerator/Detours/Private/Windows/UbaDetoursFunctionsWin.cpp b/Engine/Source/Programs/UnrealBuildAccelerator/Detours/Private/Windows/UbaDetoursFunctionsWin.cpp
--- a/Engine/Source/Programs/UnrealBuildAccelerator/Detours/Private/Windows/UbaDetoursFunctionsWin.cpp
+++ b/Engine/Source/Programs/UnrealBuildAccelerator/Detours/Private/Windows/UbaDetoursFunctionsWin.cpp
@@ -737,7 +737,7 @@ int DetourDetachFunctions()
 
 void PreInit(const DetoursPayload& payload)
 {
-	#if defined(UBA_USE_MIMALLOC)
+	#if UBA_USE_MIMALLOC
 	//mi_option_enable(mi_option_large_os_pages);
 	mi_option_disable(mi_option_abandoned_page_reset);
 	#endif
diff --git a/Engine/Source/Programs/UnrealBuildAccelerator/Detours/Private/Windows/UbaDetoursFunctionsWin.h b/Engine/Source/Programs/UnrealBuildAccelerator/Detours/Private/Windows/UbaDetoursFunctionsWin.h
--- a/Engine/Source/Programs/UnrealBuildAccelerator/Detours/Private/Windows/UbaDetoursFunctionsWin.h
+++ b/Engine/Source/Programs/UnrealBuildAccelerator/Detours/Private/Windows/UbaDetoursFunctionsWin.h
@@ -152,7 +152,7 @@
 	DETOURED_FUNCTION(_wsplitpath_s) \
 	DETOURED_FUNCTIONS_UCRTBASE_DEBUG \
 
-#if defined(UBA_USE_MIMALLOC)
+#if UBA_USE_MIMALLOC
 #define DETOURED_FUNCTIONS_MEMORY \
 	DETOURED_FUNCTION(malloc) \
 	DETOURED_FUNCTION(calloc) \
@@ -313,7 +313,7 @@
 	//DETOURED_FUNCTION(_wsopen_s) \
 	//DETOURED_FUNCTION(_fileno) \
 
-#if defined(UBA_USE_MIMALLOC)
+#if UBA_USE_MIMALLOC
 #define DETOURED_FUNCTIONS_MEMORY_DEBUG \
 	DETOURED_FUNCTION(_aligned_msize) \
 	//DETOURED_FUNCTION(_free_dbg) \
diff --git a/Engine/Source/Programs/UnrealBuildAccelerator/Mimalloc/Private/UbaMimalloc.c b/Engine/Source/Programs/UnrealBuildAccelerator/Mimalloc/Private/UbaMimalloc.c
--- a/Engine/Source/Programs/UnrealBuildAccelerator/Mimalloc/Private/UbaMimalloc.c
+++ b/Engine/Source/Programs/UnrealBuildAccelerator/Mimalloc/Private/UbaMimalloc.c
@@ -1,5 +1,5 @@
 // Copyright Epic Games, Inc. All Rights Reserved.
 
-#if defined(UBA_USE_MIMALLOC)
+#if UBA_USE_MIMALLOC
 #include "static.c"
 #endif
diff --git a/Engine/Source/Programs/UnrealBuildAccelerator/Mimalloc/Private/UbaMimallocOverride.cpp b/Engine/Source/Programs/UnrealBuildAccelerator/Mimalloc/Private/UbaMimallocOverride.cpp
--- a/Engine/Source/Programs/UnrealBuildAccelerator/Mimalloc/Private/UbaMimallocOverride.cpp
+++ b/Engine/Source/Programs/UnrealBuildAccelerator/Mimalloc/Private/UbaMimallocOverride.cpp
@@ -3,7 +3,7 @@
 // Compile this file to override new and delete
 // To override malloc/free, include mimalloc-override.h
 
-#if defined(UBA_USE_MIMALLOC)
+#if UBA_USE_MIMALLOC
 // mimalloc-new-delete.h should be included in only one source file!
 #include <mimalloc-new-delete.h>
 #endif
diff --git a/Engine/Source/Programs/UnrealBuildAccelerator/Mimalloc/UbaMimalloc.Build.cs b/Engine/Source/Programs/UnrealBuildAccelerator/Mimalloc/UbaMimalloc.Build.cs
--- a/Engine/Source/Programs/UnrealBuildAccelerator/Mimalloc/UbaMimalloc.Build.cs
+++ b/Engine/Source/Programs/UnrealBuildAccelerator/Mimalloc/UbaMimalloc.Build.cs
@@ -11,7 +11,7 @@ public class UbaMimalloc : ModuleRules
 	{
 		PCHUsage = PCHUsageMode.NoPCHs;
 
-		if (Target.Platform == UnrealTargetPlatform.Win64) {
+		if (Target.Platform == UnrealTargetPlatform.Win64 || Target.Platform == UnrealTargetPlatform.Linux) {
 			PublicDependencyModuleNames.AddRange(new string[] {
 				"mimalloc212",
 			});
diff --git a/Engine/Source/Programs/UnrealBuildAccelerator/Test/Private/UbaTestAll.h b/Engine/Source/Programs/UnrealBuildAccelerator/Test/Private/UbaTestAll.h
--- a/Engine/Source/Programs/UnrealBuildAccelerator/Test/Private/UbaTestAll.h
+++ b/Engine/Source/Programs/UnrealBuildAccelerator/Test/Private/UbaTestAll.h
@@ -42,10 +42,10 @@ namespace uba
 
 
 	#define UBA_TEST(x) \
-		logger.Info(TC("Running %hs..."), #x); \
+		logger.Info(TC("Running %s..."), TC(#x)); \
 		if (!x(testLogger, testRootDir)) \
-			return logger.Error(TC("  %hs failed"), #x); \
-		logger.Info(TC("  %hs success!"), #x);
+			return logger.Error(TC("  %s failed"), TC(#x)); \
+		logger.Info(TC("  %s success!"),  TC(#x));
 
 
 	bool RunAllTests()
diff --git a/Engine/Source/Programs/UnrealBuildAccelerator/Test/Private/UbaTestSession.h b/Engine/Source/Programs/UnrealBuildAccelerator/Test/Private/UbaTestSession.h
--- a/Engine/Source/Programs/UnrealBuildAccelerator/Test/Private/UbaTestSession.h
+++ b/Engine/Source/Programs/UnrealBuildAccelerator/Test/Private/UbaTestSession.h
@@ -67,7 +67,7 @@ namespace uba
 		sessionServerInfo.rootDir = rootDir.data;
 		SessionServer sessionServer(sessionServerInfo);
 
-		auto g = MakeGuard([&]() { server.StopAll(); });
+		auto sg = MakeGuard([&]() { server.StopAll(); });
 
 		rootDir.Append(TC("Client"));
 		if (!DeleteAllFiles(logger, rootDir.data))
@@ -80,6 +80,8 @@ namespace uba
 		sessionClientInfo.rootDir = rootDir.data;
 		SessionClient sessionClient(sessionClientInfo);
 
+		auto cg = MakeGuard([&]() { sessionClient.Stop(); client.Disconnect(); });
+
 		StringBuffer<> workingDir;
 		workingDir.Append(testRootDir).Append(TC("WorkingDir"));
 		if (!DeleteAllFiles(logger, workingDir.data))
diff --git a/Engine/Source/Programs/UnrealBuildAccelerator/UbaAgent.Target.cs b/Engine/Source/Programs/UnrealBuildAccelerator/UbaAgent.Target.cs
--- a/Engine/Source/Programs/UnrealBuildAccelerator/UbaAgent.Target.cs
+++ b/Engine/Source/Programs/UnrealBuildAccelerator/UbaAgent.Target.cs
@@ -104,13 +104,20 @@ public class UbaAgentTarget : TargetRules
 			Rules.GlobalDefinitions.Add("PLATFORM_MAC=0");
 		}
 
-		if (Target.Platform.IsInGroup(UnrealPlatformGroup.Windows))
+		int useMiMalloc = 0;
+		if (Target.Platform.IsInGroup(UnrealPlatformGroup.Windows) || Target.Platform.IsInGroup(UnrealPlatformGroup.Linux))
 		{
-			Rules.AdditionalCompilerArguments = "/wd4100 "; // -- C4100: unreferenced formal parameter
 			if (Rules.bUsePCHFiles)
 			{
-				Rules.GlobalDefinitions.Add("UBA_USE_MIMALLOC=1");
+				useMiMalloc = 1;
 			}
+		}
+		Rules.GlobalDefinitions.Add($"UBA_USE_MIMALLOC={useMiMalloc}");
+
+
+		if (Target.Platform.IsInGroup(UnrealPlatformGroup.Windows))
+		{
+			Rules.AdditionalCompilerArguments = "/wd4100 "; // -- C4100: unreferenced formal parameter
 			if (ShouldExport)
 			{
 				Rules.GlobalDefinitions.Add("UBA_API=__declspec(dllexport)");
diff --git a/Engine/Source/Programs/UnrealBuildAccelerator/Visualizer/Private/UbaTraceReader.cpp b/Engine/Source/Programs/UnrealBuildAccelerator/Visualizer/Private/UbaTraceReader.cpp
--- a/Engine/Source/Programs/UnrealBuildAccelerator/Visualizer/Private/UbaTraceReader.cpp
+++ b/Engine/Source/Programs/UnrealBuildAccelerator/Visualizer/Private/UbaTraceReader.cpp
@@ -19,7 +19,7 @@ namespace uba
 		if (m_hostProcess)
 			CloseHandle(m_hostProcess);
 		if (m_memoryBegin)
-			UnmapViewOfFile(m_memoryBegin, 0);
+			UnmapViewOfFile(m_memoryBegin, 0, TC("TraceReader"));
 		if (m_memoryHandle.IsValid())
 			CloseFileMapping(m_memoryHandle);
 	}
@@ -49,7 +49,7 @@ namespace uba
 		m_memoryBegin = MapViewOfFile(m_memoryHandle, FILE_MAP_READ, 0, 0);
 		if (!m_memoryBegin)
 			return false;
-		auto closeView = MakeGuard([&]() { UnmapViewOfFile(m_memoryBegin, 0); m_memoryBegin = nullptr; m_memoryPos = nullptr; m_memoryEnd = nullptr; });
+		auto closeView = MakeGuard([&]() { UnmapViewOfFile(m_memoryBegin, 0, TC("TraceReader")); m_memoryBegin = nullptr; m_memoryPos = nullptr; m_memoryEnd = nullptr; });
 
 		m_memoryPos = m_memoryBegin;
 		m_memoryEnd = m_memoryBegin + fileSize;
