[UBA]
* Fixed memory stomp in code writing memory to compressed file. Found by asan on linux

diff --git a/Engine/Source/Programs/UnrealBuildAccelerator/Common/Private/UbaStorage.cpp b/Engine/Source/Programs/UnrealBuildAccelerator/Common/Private/UbaStorage.cpp
--- a/Engine/Source/Programs/UnrealBuildAccelerator/Common/Private/UbaStorage.cpp
+++ b/Engine/Source/Programs/UnrealBuildAccelerator/Common/Private/UbaStorage.cpp
@@ -326,6 +326,8 @@ namespace uba
 
 		struct WorkRec
 		{
+			WorkRec() = delete;
+			WorkRec(const WorkRec&) = delete;
 			WorkRec(u32 wc) { workCount = wc; events = new Event[workCount]; }
 			~WorkRec() { delete[] events; }
 			Atomic<u64> refCount;
@@ -405,7 +407,8 @@ namespace uba
 				rec->memPos += writeBytes;
 #endif
 				rec->written += writeBytes;
-				rec->events[index].Set();
+				if (index < rec->workCount)
+					rec->events[index].Set();
 			}
 		};
 
