[UBA]
* Added more information for when process dies without sending exit message

diff --git a/Engine/Source/Programs/UnrealBuildAccelerator/Common/Private/UbaProcess.cpp b/Engine/Source/Programs/UnrealBuildAccelerator/Common/Private/UbaProcess.cpp
--- a/Engine/Source/Programs/UnrealBuildAccelerator/Common/Private/UbaProcess.cpp
+++ b/Engine/Source/Programs/UnrealBuildAccelerator/Common/Private/UbaProcess.cpp
@@ -245,7 +245,7 @@ namespace uba
 			}
 
 			if (err.IsEmpty())
-				err.Appendf(TC("ERROR: Process %llu (%s) not active but did not get exit message. Received %u messages (GetExitCodeProcess returned %u)"), u64(m_nativeProcessHandle), m_realApplication.c_str(), m_messageCount, exitCode);
+				err.Appendf(TC("ERROR: Process %llu %s (%s) not active but did not get exit message. Received %u messages (GetExitCodeProcess returned %u)"), u64(m_nativeProcessHandle), m_description.c_str(), m_realApplication.c_str(), m_messageCount, exitCode);
 			LogLine(false, err.data, LogEntryType_Error);
 			m_nativeProcessExitCode = UBA_EXIT_CODE(666);
 		}
