[UBA]
* Fixed compile error for linux arm target

diff --git a/Engine/Source/Programs/UnrealBuildAccelerator/Mimalloc/UbaMimalloc.Build.cs b/Engine/Source/Programs/UnrealBuildAccelerator/Mimalloc/UbaMimalloc.Build.cs
--- a/Engine/Source/Programs/UnrealBuildAccelerator/Mimalloc/UbaMimalloc.Build.cs
+++ b/Engine/Source/Programs/UnrealBuildAccelerator/Mimalloc/UbaMimalloc.Build.cs
@@ -11,7 +11,7 @@ public class UbaMimalloc : ModuleRules
 	{
 		PCHUsage = PCHUsageMode.NoPCHs;
 
-		if (Target.Platform == UnrealTargetPlatform.Win64 || Target.Platform == UnrealTargetPlatform.Linux) {
+		if (Target.Platform.IsInGroup(UnrealPlatformGroup.Windows) || Target.Platform.IsInGroup(UnrealPlatformGroup.Linux)) {
 			PublicDependencyModuleNames.AddRange(new string[] {
 				"mimalloc212",
 			});
