[UBA]
* Fixed so GetFullPathNameW for linux works for non-existing paths (we just want it to fix it up)

diff --git a/Engine/Source/Programs/UnrealBuildAccelerator/Common/Private/UbaFile.cpp b/Engine/Source/Programs/UnrealBuildAccelerator/Common/Private/UbaFile.cpp
--- a/Engine/Source/Programs/UnrealBuildAccelerator/Common/Private/UbaFile.cpp
+++ b/Engine/Source/Programs/UnrealBuildAccelerator/Common/Private/UbaFile.cpp
@@ -691,6 +691,7 @@ namespace uba
 			return strlen(strcpy(lpBuffer, source));
 		}
 
+#if 0
 		char fullPath[1024];
 		if (!realpath(fileName, fullPath))
 		{
@@ -699,6 +700,25 @@ namespace uba
 		}
 		u32 len = TStrlen(fullPath);
 		UBA_ASSERT(len < sizeof_array(fullPath));
+#elif 0
+		const char* fullPath = canonicalize_file_name(fileName);
+		u32 len = TStrlen(fullPath);
+#else
+		char cwd[PATH_MAX];
+		if (!getcwd(cwd, PATH_MAX))
+		{
+			UBA_ASSERT(false);
+			return 0;
+		}
+		u32 cwdlen = u32(strlen(cwd));
+		cwd[cwdlen++] = '/';
+		cwd[cwdlen] = 0;
+
+		char fullPath[1024];
+		u32 len;
+		if (!FixPath2(fileName, cwd, cwdlen, fullPath, sizeof(fullPath), &len))
+			return 0;
+#endif
 		UBA_ASSERT(len < nBufferLength);
 		memcpy(lpBuffer, fullPath, len + 1);
 		return len;
