[UBA]
* Added support for most features in UbaScheduler. Dependencies, weights etc etc.
* Added support for loading yaml file with processes for UbaScheduler.
* Added so UbaCli interprets yaml file as a file with processes use it to populate a scheduler which is then executed
* Fixed so all threads inside same process spawned by uba ends up in the same thread group.
* Fixed linux crash where process comunication memory was deleted when cancel event was called (added lock around code)
* Fixed deadlock that could happen if flush dead processes were called at the after lock but before processhandle dtor in Session::ProcessExited
* Changed new[]/delete[] to aligned_alloc/free because for some reason new/delete trigger asan on linux and don't know why.

diff --git a/Engine/Source/Programs/UnrealBuildAccelerator/Cli/Private/UbaCli.cpp b/Engine/Source/Programs/UnrealBuildAccelerator/Cli/Private/UbaCli.cpp
--- a/Engine/Source/Programs/UnrealBuildAccelerator/Cli/Private/UbaCli.cpp
+++ b/Engine/Source/Programs/UnrealBuildAccelerator/Cli/Private/UbaCli.cpp
@@ -2,6 +2,7 @@
 
 #include "UbaNetworkBackendTcp.h"
 #include "UbaFileAccessor.h"
+#include "UbaScheduler.h"
 #include "UbaSessionServer.h"
 #include "UbaStorageClient.h"
 #include "UbaStorageServer.h"
@@ -146,7 +147,6 @@ namespace uba
 			CommandType_Local,
 			CommandType_Remote,
 			CommandType_Native,
-
 		};
 
 		CommandType commandType = CommandType_NotSet;
@@ -470,7 +470,7 @@ namespace uba
 			if (!server->StartListen(networkBackend, port, listenIp.data))
 				return -1;
 		}
-		auto g = MakeGuard([&]() { server->StopAll(); });
+		auto stopServer = MakeGuard([&]() { server->StopAll(); });
 
 		auto RunLocal = [&](const TString& app, const TString& arg, bool enableDetour, bool trackInputs = false)
 		{
@@ -516,20 +516,77 @@ namespace uba
 			return true;
 		};
 
+		auto RunScheduler = [&](const tchar* yamlFile)
+		{
+			SchedulerCreateInfo info(*session);
+			info.forceRemote = commandType == CommandType_Remote;
+			Scheduler scheduler(info);
+
+			if (!scheduler.EnqueueFromFile(yamlFile))
+				return false;
+
+			u32 queued, activeLocal, activeRemote, outFinished;
+			scheduler.GetStats(queued, activeLocal, activeRemote, outFinished);
+
+			bool success = true;
+			Atomic<u32> counter;
+			Event finished(true);
+			scheduler.SetProcessFinishedCallback([&](const ProcessHandle& ph)
+				{
+					if (ph.GetExitCode() != 0)
+						success = false;
+					u32 c = ++counter;
+					logger.BeginScope();
+					const tchar* desc = ph.GetStartInfo().description;
+					StringBuffer<128> extra;
+					if (ph.IsRemote())
+						extra.Append(TC(" [RemoteExecutor: ")).Append(ph.GetExecutingHost()).Append(']');
+					else if (!ph.IsDetoured())
+						extra.Append(TC(" (Not detoured)"));
+					logger.Info(TC("[%u/%u] %s%s"), c, queued, desc, extra.data);
+					for (auto& line : ph.GetLogLines())
+						if (line.text != desc && !StartsWith(line.text.c_str(), TC("   Creating library")))
+							logger.Log(line.type, line.text.c_str(), u32(line.text.size()));
+					logger.EndScope();
+
+					if (c == queued)
+						finished.Set();
+				});
+
+			logger.Info(TC("Running Scheduler with %u processes"), queued);
+			u64 start = GetTime();
+			scheduler.Start();
+			if (!finished.IsSet())
+				return false;
+			u64 time = GetTime() - start;
+			logger.Info(TC("Scheduler run took %s"), TimeToText(time).str);
+			logger.Info(TC(""));
+			stopServer.Execute();
+			return success;
+		};
+
 		for (u32 i=0; i!=loopCount; ++i)
 		{
 			bool success = false;
-			switch (commandType)
+
+			if (EndsWith(application.c_str(), application.size(), TC(".yaml")))
 			{
-			case CommandType_Native:
-				success = RunLocal(application, arguments, false);
-				break;
-			case CommandType_Local:
-				success = RunLocal(application, arguments, true);
-				break;
-			case CommandType_Remote:
-				success = RunRemote(application, arguments);
-				break;
+				success = RunScheduler(application.c_str());
+			}
+			else
+			{
+				switch (commandType)
+				{
+				case CommandType_Native:
+					success = RunLocal(application, arguments, false);
+					break;
+				case CommandType_Local:
+					success = RunLocal(application, arguments, true);
+					break;
+				case CommandType_Remote:
+					success = RunRemote(application, arguments);
+					break;
+				}
 			}
 			if (!success)
 				return -1;
diff --git a/Engine/Source/Programs/UnrealBuildAccelerator/Common/Private/UbaProcess.cpp b/Engine/Source/Programs/UnrealBuildAccelerator/Common/Private/UbaProcess.cpp
--- a/Engine/Source/Programs/UnrealBuildAccelerator/Common/Private/UbaProcess.cpp
+++ b/Engine/Source/Programs/UnrealBuildAccelerator/Common/Private/UbaProcess.cpp
@@ -108,8 +108,13 @@ namespace uba
 
 	ProcessImpl::~ProcessImpl()
 	{
-		if (m_comMemory.memory)
-			m_cancelEvent.Set();
+		{
+			#if !PLATFORM_WINDOWS
+			ScopedWriteLock lock(m_comMemoryLock);
+			#endif
+			if (m_comMemory.memory)
+				m_cancelEvent.Set();
+		}
 
 		m_messageThread.Wait();
 
@@ -473,8 +478,13 @@ namespace uba
 		m_readEvent.~Event();
 		#endif
 
-		m_session.m_processCommunicationAllocator.Free(m_comMemory);
-		m_comMemory = {};
+		{
+			#if !PLATFORM_WINDOWS
+			ScopedWriteLock lock(m_comMemoryLock);
+			#endif
+			m_session.m_processCommunicationAllocator.Free(m_comMemory);
+			m_comMemory = {};
+		}
 
 		if (!m_parentProcess)
 			ClearTempFiles();
@@ -1352,7 +1362,7 @@ namespace uba
 			}
 		}
 
-		if (!AlternateGroupAffinity(m_nativeThreadHandle))
+		if (!AlternateThreadGroupAffinity(m_nativeThreadHandle))
 		{
 			logger.Error(TC("Failed to set thread group affinity to process"));//% ls. (% ls)"), commandLine.c_str(), LastErrorToText().data);
 			return UBA_EXIT_CODE(10);
diff --git a/Engine/Source/Programs/UnrealBuildAccelerator/Common/Private/UbaScheduler.cpp b/Engine/Source/Programs/UnrealBuildAccelerator/Common/Private/UbaScheduler.cpp
--- a/Engine/Source/Programs/UnrealBuildAccelerator/Common/Private/UbaScheduler.cpp
+++ b/Engine/Source/Programs/UnrealBuildAccelerator/Common/Private/UbaScheduler.cpp
@@ -8,42 +8,60 @@
 
 namespace uba
 {
-	struct Scheduler::QueuedProcess : ProcessStartInfoHolder
+	struct Scheduler::ProcessStartInfo2 : ProcessStartInfoHolder
 	{
-		QueuedProcess(const ProcessStartInfo& si, const u8* ki, u32 kic)
+		ProcessStartInfo2(const ProcessStartInfo& si, const u8* ki, u32 kic)
 		:	ProcessStartInfoHolder(si)
 		, knownInputs(ki)
 		, knownInputsCount(kic)
 		{
-		};
+		}
+
+		~ProcessStartInfo2()
+		{
+			delete[] knownInputs;
+		}
 
 		const u8* knownInputs;
 		u32 knownInputsCount;
+		float weight = 1.0f;
 	};
 
+
 	struct Scheduler::ExitProcessInfo
 	{
 		Scheduler* scheduler = nullptr;
-		ProcessStartInfo::ExitedCallback* originalExitedFunc = nullptr;
-		void* originalUserData = nullptr;
-		const u8* knownInputs = nullptr;
-		u32 knownInputsCount = 0;
+		ProcessStartInfo2* startInfo;
+		u32 processIndex = ~0u;
 		bool wasReturned = false;
-		bool isRemote = false;
+		bool isLocal = true;
+	};
+
 
-		// If reusing processes then these are the actual info of the process
-		TString lastArguments;
-		TString lastWorkingDir;
-		TString lastDescription;
+	class SkippedProcess : public Process
+	{
+	public:
+		SkippedProcess(const ProcessStartInfo& i) : holder(i) {}
+		virtual u32 GetExitCode() override { return ProcessCancelExitCode; }
+		virtual bool HasExited() override { return true; }
+		virtual bool WaitForExit(u32 millisecondsTimeout) override{ return true; }
+		virtual const ProcessStartInfo& GetStartInfo() override { return holder.startInfo; }
+		virtual const Vector<ProcessLogLine>& GetLogLines() override { static Vector<ProcessLogLine> v{ProcessLogLine{TC("Skipped"), LogEntryType_Warning}}; return v; }
+		virtual const Vector<u8>& GetTrackedInputs() override { static Vector<u8> v; return v;}
+		virtual bool IsRemote() const override { return false; }
+		virtual bool IsDetoured() const { return false; }
+		ProcessStartInfoHolder holder;
 	};
 
-	Scheduler::Scheduler(SessionServer& session, u32 maxLocalProcessors, bool enableProcessReuse)
-	:	m_session(session)
-	,	m_maxLocalProcessors(maxLocalProcessors != ~0u ? maxLocalProcessors : GetLogicalProcessorCount())
+	
+	Scheduler::Scheduler(const SchedulerCreateInfo& info)
+	:	m_session(info.session)
+	,	m_maxLocalProcessors(info.maxLocalProcessors != ~0u ? info.maxLocalProcessors : GetLogicalProcessorCount())
 	,	m_updateThreadLoop(false)
-	,	m_enableProcessReuse(enableProcessReuse)
+	,	m_enableProcessReuse(info.enableProcessReuse)
+	,	m_forceRemote(info.forceRemote)
 	{
-		session.RegisterGetNextProcess([this](Process& process, NextProcessInfo& outNextProcess, u32 prevExitCode)
+		m_session.RegisterGetNextProcess([this](Process& process, NextProcessInfo& outNextProcess, u32 prevExitCode)
 			{
 				return HandleReuseMessage(process, outNextProcess, prevExitCode);
 			});
@@ -69,39 +87,67 @@ namespace uba
 		m_updateThreadLoop.Set();
 		m_thread.Wait();
 		m_session.WaitOnAllTasks();
+
+		for (auto& entry : m_processEntries)
+		{
+			delete[] entry.dependencies;
+			delete entry.info;
+		}
 	}
 
 	void Scheduler::SetMaxLocalProcessors(u32 maxLocalProcessors)
 	{
 		m_maxLocalProcessors = maxLocalProcessors;
+		m_updateThreadLoop.Set();
 	}
 
-	u32 Scheduler::EnqueueProcess(const ProcessStartInfo& info, float weight, const void* knownInputs, u32 knownInputsBytes, u32 knownInputsCount)
+	u32 Scheduler::EnqueueProcess(const EnqueueProcessInfo& info)
 	{
-		UBA_ASSERTF(weight == 1.0f, TC("weight not implemented"));
+		u8* ki = nullptr;
+		if (info.knownInputsCount)
+		{
+			ki = new u8[info.knownInputsBytes];
+			memcpy(ki, info.knownInputs, info.knownInputsBytes);
+		}
 
-		if (knownInputsCount)
+		u32* dep = nullptr;
+		if (info.dependencyCount)
 		{
-			void* ki = new u8[knownInputsBytes];
-			memcpy(ki, knownInputs, knownInputsBytes);
-			knownInputs = ki;
+			dep = new u32[info.dependencyCount];
+			memcpy(dep, info.dependencies, info.dependencyCount*sizeof(u32));
 		}
-		else
-			knownInputs = nullptr;
 
-		ScopedWriteLock lock(m_queuedProcessesLock);
-		m_queuedProcesses.emplace_back(info, (const u8*)knownInputs, knownInputsCount);
+		auto info2 = new ProcessStartInfo2(info.info, ki, info.knownInputsCount);
+		info2->weight = info.weight;
+
+		ScopedWriteLock lock(m_processEntriesLock);
+		u32 index = u32(m_processEntries.size());
+		auto& entry = m_processEntries.emplace_back();
+		entry.info = info2;
+		entry.dependencies = dep;
+		entry.dependencyCount = info.dependencyCount;
+		entry.status = ProcessStatus_Queued;
+		entry.canDetour = info.canDetour;
+		entry.canExecuteRemotely = info.canExecuteRemotely && info.canDetour;
+		lock.Leave();
+
+		UpdateQueueCounter(1);
+
 		m_updateThreadLoop.Set();
-		return 0;
+		return index;
 	}
 
 	void Scheduler::GetStats(u32& outQueued, u32& outActiveLocal, u32& outActiveRemote, u32& outFinished)
 	{
-		ScopedReadLock lock(m_queuedProcessesLock);
 		outActiveLocal = m_activeLocalProcesses;
 		outActiveRemote = m_activeRemoteProcesses;
 		outFinished = m_finishedProcesses;
-		outQueued = u32(m_queuedProcesses.size());
+		outQueued = m_queuedProcesses;
+	}
+
+	void Scheduler::SetProcessFinishedCallback(const Function<void(const ProcessHandle&)>& processFinished)
+	{
+		m_processFinished = processFinished;
 	}
 
 	void Scheduler::ThreadLoop()
@@ -111,9 +157,8 @@ namespace uba
 			if (!m_updateThreadLoop.IsSet())
 				break;
 
-			while (m_activeLocalProcesses < m_maxLocalProcessors)
-				if (!RunQueuedProcess(true))
-					break;
+			while (RunQueuedProcess(true))
+				;
 		}
 	}
 
@@ -121,24 +166,23 @@ namespace uba
 	{
 		auto& ei = *(ExitProcessInfo*)process.GetStartInfo().userData;
 
-		ScopedWriteLock lock(m_queuedProcessesLock);
-		ProcessStartInfo info = process.GetStartInfo();
-		info.exitedFunc = ei.originalExitedFunc;
-		info.userData = ei.originalUserData;
-		if (!ei.lastArguments.empty())
-		{
-			info.arguments = ei.lastArguments.c_str();
-			info.workingDir = ei.lastWorkingDir.c_str();
-			info.description = ei.lastDescription.c_str();
-		}
-		m_queuedProcesses.emplace_front(info, ei.knownInputs, ei.knownInputsCount);
-		ei.knownInputsCount = 0;
-		ei.knownInputs = nullptr;
 		ei.wasReturned = true;
-		lock.Leave();
+		u32 processIndex = ei.processIndex;
 
 		process.Cancel(true); // Cancel will call ProcessExited
-		--m_activeRemoteProcesses;
+		
+		if (processIndex == ~0u)
+			return;
+
+		ScopedWriteLock lock(m_processEntriesLock);
+		if (m_processEntries[processIndex].status != ProcessStatus_Running)
+			return;
+		m_processEntries[processIndex].status = ProcessStatus_Queued;
+		m_processEntriesStart = Min(m_processEntriesStart, processIndex);
+		lock.Leave();
+
+		UpdateQueueCounter(1);
+		UpdateActiveProcessCounter(false, -1);
 		m_updateThreadLoop.Set();
 	}
 
@@ -149,62 +193,129 @@ namespace uba
 
 	void Scheduler::ProcessExited(ExitProcessInfo* info, const ProcessHandle& handle)
 	{
-		auto ig = MakeGuard([info]() { delete[] info->knownInputs; delete info; });
+		auto ig = MakeGuard([info]() { delete info; });
 
 		if (info->wasReturned)
 			return;
 
-		if (auto func = info->originalExitedFunc)
-			func(info->originalUserData, handle);
-
-		++m_finishedProcesses;
-
-		if (info->isRemote)
-		{
-			--m_activeRemoteProcesses;
-		}
-		else
+		auto si = info->startInfo;
+		if (!si) // Can be a process that was reused but didn't get a new process
 		{
-			--m_activeLocalProcesses;
+			UBA_ASSERT(info->processIndex == ~0u);
+			return;
 		}
-		m_updateThreadLoop.Set();
+
+		ExitProcess(*info, *handle.m_process, handle.m_process->GetExitCode());
 	}
 
-	bool Scheduler::RunQueuedProcess(bool runLocal)
+	u32 Scheduler::PopProcess(bool isLocal)
 	{
-		ScopedWriteLock lock(m_queuedProcessesLock);
-		if (m_queuedProcesses.empty())
-			return false;
-		auto si = m_queuedProcesses.front();
-		m_queuedProcesses.pop_front();
-		if (runLocal)
-			++m_activeLocalProcesses;
-		lock.Leave();
+		if (isLocal)
+			if (m_activeLocalProcessWeight >= float(m_maxLocalProcessors))
+				return ~0u;
 
-		auto exitInfo = new ExitProcessInfo();
-		exitInfo->scheduler = this;
-		exitInfo->originalExitedFunc = si.startInfo.exitedFunc;
-		exitInfo->originalUserData = si.startInfo.userData;
-		exitInfo->knownInputs = si.knownInputs;
-		exitInfo->knownInputsCount = si.knownInputsCount;
-		exitInfo->isRemote = !runLocal;
-		si.startInfo.userData = exitInfo;
-		si.startInfo.exitedFunc = [](void* userData, const ProcessHandle& handle)
-			{
-				auto ei = (ExitProcessInfo*)userData;
-				ei->scheduler->ProcessExited(ei, handle);
-			};
+		auto processEntries = m_processEntries.data();
+		bool allFinished = true;
 
-		if (runLocal)
+		for (u32 i=m_processEntriesStart, e=u32(m_processEntries.size()); i!=e; ++i)
 		{
-			m_session.RunProcess(si.startInfo);
+			auto& entry = processEntries[i];
+			auto status = entry.status;
+			if (status != ProcessStatus_Queued)
+			{
+				if (allFinished)
+				{
+					if (status != ProcessStatus_Running)
+						m_processEntriesStart = i;
+					else
+						allFinished = false;
+				}
+				continue;
+			}
+			allFinished = false;
+
+			if (!isLocal && !entry.canExecuteRemotely)
+				continue;
+
+			if (isLocal && m_forceRemote && entry.canExecuteRemotely)
+				continue;
+
+			bool canRun = true;
+			for (u32 j=0, je=entry.dependencyCount; j!=je; ++j)
+			{
+				UBA_ASSERT(j < m_processEntries.size());
+				auto depIndex = entry.dependencies[j];
+				auto depStatus = processEntries[depIndex].status;
+				if (depStatus == ProcessStatus_Failed || depStatus == ProcessStatus_Skipped)
+				{
+					entry.status = ProcessStatus_Skipped;
+					return i;
+				}
+				if (depStatus != ProcessStatus_Success)
+				{
+					canRun = false;
+					break;
+				}
+			}
+
+			if (!canRun)
+				continue;
+
+			if (isLocal)
+				m_activeLocalProcessWeight += entry.info->weight;
+
+			entry.status = ProcessStatus_Running;
+			return i;
 		}
-		else
+		return ~0u;
+	}
+
+
+	bool Scheduler::RunQueuedProcess(bool isLocal)
+	{
+		while (true)
 		{
-			++m_activeRemoteProcesses;
-			m_session.RunProcessRemote(si.startInfo, 1.0f, si.knownInputs, si.knownInputsCount);
+			ScopedWriteLock lock(m_processEntriesLock);
+			u32 indexToRun = PopProcess(isLocal);
+			if (indexToRun == ~0u)
+				return false;
+
+			auto& processEntry = m_processEntries[indexToRun];
+			auto info = processEntry.info;
+			bool canDetour = processEntry.canDetour;
+			bool wasSkipped = processEntry.status == ProcessStatus_Skipped;
+			lock.Leave();
+
+			UpdateQueueCounter(-1);
+
+			if (wasSkipped)
+			{
+				SkipProcess(*info);
+				continue;
+			}
+
+			UpdateActiveProcessCounter(isLocal, 1);
+	
+			auto exitInfo = new ExitProcessInfo();
+			exitInfo->scheduler = this;
+			exitInfo->startInfo = info;
+			exitInfo->isLocal = isLocal;
+			exitInfo->processIndex = indexToRun;
+
+			ProcessStartInfo si = info->startInfo;
+			si.userData = exitInfo;
+			si.exitedFunc = [](void* userData, const ProcessHandle& handle)
+				{
+					auto ei = (ExitProcessInfo*)userData;
+					ei->scheduler->ProcessExited(ei, handle);
+				};
+
+			if (isLocal)
+				m_session.RunProcess(si, true, canDetour);
+			else
+				m_session.RunProcessRemote(si, 1.0f, info->knownInputs, info->knownInputsCount);
+			return true;
 		}
-		return true;
 	}
 
 	bool Scheduler::HandleReuseMessage(Process& process, NextProcessInfo& outNextProcess, u32 prevExitCode)
@@ -213,59 +324,276 @@ namespace uba
 			return false;
 
 		auto& currentStartInfo = process.GetStartInfo();
-		auto info = (ExitProcessInfo*)currentStartInfo.userData;
-		if (!info) // If null, process has already exited from some other thread
+		auto ei = (ExitProcessInfo*)currentStartInfo.userData;
+		if (!ei) // If null, process has already exited from some other thread
+			return false;
+
+		ExitProcess(*ei, process, prevExitCode);
+
+		ei->startInfo = nullptr;
+		ei->processIndex = ~0u;
+		if (ei->wasReturned)
 			return false;
 
-		// Call ExitedFunc and cleanup
-		if (auto func = info->originalExitedFunc)
+		bool isLocal = !process.IsRemote();
+
+		while (true)
 		{
-			UBA_ASSERT(!info->wasReturned);
-			UBA_ASSERT(prevExitCode == 0); // This is not handled for anything else
-			ProcessHandle h;
-			h.m_process = &process;
-			func(info->originalUserData, h);
-			h.m_process = nullptr;
+			ScopedWriteLock lock(m_processEntriesLock);
+			u32 indexToRun = PopProcess(isLocal);
+			if (indexToRun == ~0u)
+				return false;
+			auto& processEntry = m_processEntries[indexToRun];
+			auto newInfo = processEntry.info;
+			bool wasSkipped = processEntry.status == ProcessStatus_Skipped;
+			lock.Leave();
+
+			UpdateQueueCounter(-1);
+
+			if (wasSkipped)
+			{
+				SkipProcess(*newInfo);
+				continue;
+			}
+
+			UpdateActiveProcessCounter(isLocal, 1);
+
+			ei->startInfo = newInfo;
+			ei->processIndex = indexToRun;
+
+			auto& si = newInfo->startInfo;
+			UBA_ASSERT(Equals(currentStartInfo.application, si.application));
+			outNextProcess.arguments = si.arguments;
+			outNextProcess.workingDir = si.workingDir;
+			outNextProcess.description = si.description;
+			outNextProcess.logFile = si.logFile;
+			return true;
 		}
-		info->knownInputsCount = 0;
-		delete[] info->knownInputs;
-		info->knownInputs = nullptr;
-		info->originalExitedFunc = nullptr;
-		info->originalUserData = nullptr;
+	}
 
-		info->lastArguments.clear();
-		info->lastWorkingDir.clear();
-		info->lastDescription.clear();
+	void Scheduler::ExitProcess(ExitProcessInfo& info, Process& process, u32 exitCode)
+	{
+		ProcessHandle ph;
+		ph.m_process = &process;
+
+		auto si = info.startInfo;
+		if (auto func = si->startInfo.exitedFunc)
+			func(si->startInfo.userData, ph);
+
+		ScopedWriteLock lock(m_processEntriesLock);
+		auto& entry = m_processEntries[info.processIndex];
+		u32* dependencies = entry.dependencies;
+		entry.status = exitCode == 0 ? ProcessStatus_Success : ProcessStatus_Failed;
+		entry.info = nullptr;
+		entry.dependencies = nullptr;
+		if (info.isLocal)
+			m_activeLocalProcessWeight -= si->weight;
+		lock.Leave();
 
-		if (m_activeLocalProcesses > m_maxLocalProcessors)
-			return false;
+		UpdateActiveProcessCounter(info.isLocal, -1);
+		FinishProcess(ph);
+		m_updateThreadLoop.Set();
+		delete[] dependencies;
+		delete si;
 
-		// Try to get queued process to send back
-		ScopedWriteLock lock(m_queuedProcessesLock);
-		if (info->wasReturned)
-			return false;
-		if (m_queuedProcesses.empty())
-			return false;
-		auto qp = m_queuedProcesses.front();
-		m_queuedProcesses.pop_front();
-		lock.Leave();
+		ph.m_process = nullptr;
+	}
 
-		delete[] qp.knownInputs;
-		auto& si = qp.startInfo;
+	void Scheduler::SkipProcess(ProcessStartInfo2& info)
+	{
+		ProcessHandle ph(new SkippedProcess(info.startInfo));
+		if (auto func = info.startInfo.exitedFunc)
+			func(info.startInfo.userData, ph);
+		FinishProcess(ph);
+	}
 
-		UBA_ASSERT(Equals(currentStartInfo.application, si.application));
+	void Scheduler::UpdateQueueCounter(int offset)
+	{
+		m_queuedProcesses += u32(offset);
+		m_session.UpdateStatus(1, 1, TC("Queue"), 4, StringBuffer<32>().Appendf(TC("%u"), m_queuedProcesses.load()).data, LogEntryType_Info);
+	}
 
-		// Move over exited func and user data for the queued process
-		info->originalExitedFunc = si.exitedFunc;
-		info->originalUserData = si.userData;
-		info->lastArguments = si.arguments;
-		info->lastWorkingDir = si.workingDir;
-		info->lastDescription = si.description;
+	void Scheduler::UpdateActiveProcessCounter(bool isLocal, int offset)
+	{
+		if (isLocal)
+			m_activeLocalProcesses += u32(offset);
+		else
+			m_activeRemoteProcesses += u32(offset);
+		m_session.UpdateStatus(2, 1, TC("Active"), 4, StringBuffer<32>().Appendf(TC("Local %u Remote %u"), m_activeLocalProcesses.load(), m_activeRemoteProcesses.load()).data, LogEntryType_Info);
+	}
+
+	void Scheduler::FinishProcess(const ProcessHandle& handle)
+	{
+		++m_finishedProcesses;
+		if (m_processFinished)
+			m_processFinished(handle);
+		m_session.UpdateStatus(3, 1, TC("Finished"), 4, StringBuffer<32>().Appendf(TC("%u"), m_finishedProcesses.load()).data, LogEntryType_Info);
+	}
+
+	template<typename LineFunc>
+	bool ReadLines(Logger& logger, const tchar* file, const LineFunc& lineFunc)
+	{
+		FileHandle handle;
+		if (!OpenFileSequentialRead(logger, file, handle))
+			return logger.Error(TC("Failed to open file %s"), file);
+		auto fg = MakeGuard([&]() { CloseFile(file, handle); });
+		u64 fileSize = 0;
+		if (!GetFileSizeEx(fileSize, handle))
+			return logger.Error(TC("Failed to get size of file %s"), file);
+		char buffer[512];
+		u64 left = fileSize;
+		std::string line;
+		while (left)
+		{
+			u64 toRead = Min(left, u64(sizeof(buffer)));
+			left -= toRead;
+			if (!ReadFile(logger, file, handle, buffer, toRead))
+				return false;
+
+			u64 start = 0;
+			for (u64 i=0;i!=toRead;++i)
+			{
+				if (buffer[i] != '\n')
+					continue;
+				u64 end = i;
+				if (i > 0 && buffer[i-1] == '\r')
+					--end;
+				line.append(buffer + start, end - start);
+				if (!line.empty())
+					if (!lineFunc(TString(line.begin(), line.end())))
+						return false;
+				line.clear();
+				start = i + 1;
+			}
+			if (toRead && buffer[toRead - 1] == '\r')
+				--toRead;
+			line.append(buffer + start, toRead - start);
+		}
+		if (!line.empty())
+			if (!lineFunc(TString(line.begin(), line.end())))
+				return false;
+		return true;
+	}
+
+	bool Scheduler::EnqueueFromFile(const tchar* yamlFilename)
+	{
+		auto& logger = m_session.GetLogger();
+
+		TString app;
+		TString arg;
+		TString dir;
+		TString desc;
+		bool allowDetour = true;
+		bool allowRemote = true;
+		float weight = 1.0f;
+		Vector<u32> deps;
+
+		ProcessStartInfo si;
+
+		auto enqueueProcess = [&]()
+			{
+				si.application = app.c_str();
+				si.arguments = arg.c_str();
+				si.workingDir = dir.c_str();
+				si.description = desc.c_str();
+				EnqueueProcessInfo info { si };
+				info.dependencies = deps.data();
+				info.dependencyCount = u32(deps.size());
+				info.canDetour = allowDetour;
+				info.canExecuteRemotely = allowRemote;
+				info.weight = weight;
+				EnqueueProcess(info);
+				app.clear();
+				arg.clear();
+				dir.clear();
+				desc.clear();
+				deps.clear();
+				allowDetour = true;
+				allowRemote = true;
+				weight = 1.0f;
+			};
+
+		auto readLine = [&](const TString& line)
+			{
+				const tchar* keyStart = line.c_str();
+				while (*keyStart && *keyStart == ' ')
+					++keyStart;
+				if (!*keyStart)
+					return true;
+				const tchar* colon = TStrchr(keyStart, ':');
+				if (!colon)
+					return false;
+				if (*keyStart == '-')
+				{
+					keyStart += 2;
+					if (!app.empty())
+						enqueueProcess();
+				}
+
+				StringBuffer<32> key;
+				key.Append(keyStart, colon - keyStart);
+				const tchar* valueStart = colon + 1;
+				while (*valueStart && *valueStart == ' ')
+					++valueStart;
+
+				if (key.Equals(TC("environment")))
+				{
+					#if PLATFORM_WINDOWS
+					SetEnvironmentVariable(TC("PATH"), valueStart);
+					#endif
+				}
+				else if (key.Equals(TC("processes")))
+					return true;
+				else if (key.Equals(TC("app")))
+					app = valueStart;
+				else if (key.Equals(TC("arg")))
+					arg = valueStart;
+				else if (key.Equals(TC("dir")))
+					dir = valueStart;
+				else if (key.Equals(TC("desc")))
+					desc = valueStart;
+				else if (key.Equals(TC("detour")))
+					allowDetour = !Equals(valueStart, TC("false"));
+				else if (key.Equals(TC("remote")))
+					allowRemote = !Equals(valueStart, TC("false"));
+				else if (key.Equals(TC("weight")))
+					StringBuffer<32>(valueStart).Parse(weight);
+				else if (key.Equals(TC("dep")))
+				{
+					const tchar* depStart = TStrchr(valueStart, '[');
+					if (!depStart)
+						return false;
+					++depStart;
+					StringBuffer<32> depStr;
+					for (const tchar* it = depStart; *it; ++it)
+					{
+						if (*it != ']' && *it != ',')
+						{
+							if (*it != ' ')
+								depStr.Append(*it);
+							continue;
+						}
+						u32 depIndex;
+						if (!depStr.Parse(depIndex))
+							return false;
+						depStr.Clear();
+						deps.push_back(depIndex);
+
+						if (!*it)
+							break;
+						depStart = it + 1;
+					}
+				}
+				return true;
+			};
+
+		if (!ReadLines(logger, yamlFilename, readLine))
+			return false;
+
+		if (!app.empty())
+			enqueueProcess();
 
-		outNextProcess.arguments = si.arguments;
-		outNextProcess.workingDir = si.workingDir;
-		outNextProcess.description = si.description;
-		outNextProcess.logFile = si.logFile;
 		return true;
 	}
+
 }
diff --git a/Engine/Source/Programs/UnrealBuildAccelerator/Common/Private/UbaSession.cpp b/Engine/Source/Programs/UnrealBuildAccelerator/Common/Private/UbaSession.cpp
--- a/Engine/Source/Programs/UnrealBuildAccelerator/Common/Private/UbaSession.cpp
+++ b/Engine/Source/Programs/UnrealBuildAccelerator/Common/Private/UbaSession.cpp
@@ -162,6 +162,11 @@ namespace uba
 		UBA_ASSERT(m_process);
 		return m_process->IsRemote();
 	}
+	bool ProcessHandle::IsDetoured() const
+	{
+		UBA_ASSERT(m_process);
+		return m_process->IsDetoured();
+	}
 	ProcessHandle::ProcessHandle(Process* process)
 	{
 		m_process = process;
@@ -1334,10 +1339,8 @@ namespace uba
 			m_processStats.Add(process.m_processStats);
 		}
 
-		ProcessHandle h(&process);
-
 		ScopedWriteLock lock(m_processesLock);
-		m_deadProcesses.push_back(h); // Here to prevent Process thread call trigger a delete of Process which causes a deadlock
+		m_deadProcesses.emplace_back(&process); // Here to prevent Process thread call trigger a delete of Process which causes a deadlock
 		auto& stats = m_applicationStats[applicationName.data];
 		stats.count++;
 		stats.time += executionTime;
diff --git a/Engine/Source/Programs/UnrealBuildAccelerator/Common/Private/UbaSessionServer.cpp b/Engine/Source/Programs/UnrealBuildAccelerator/Common/Private/UbaSessionServer.cpp
--- a/Engine/Source/Programs/UnrealBuildAccelerator/Common/Private/UbaSessionServer.cpp
+++ b/Engine/Source/Programs/UnrealBuildAccelerator/Common/Private/UbaSessionServer.cpp
@@ -53,6 +53,7 @@ namespace uba
 
 		virtual const tchar* GetExecutingHost() const override { return m_executingHost.c_str(); }
 		virtual bool IsRemote() const override { return true; }
+		virtual bool IsDetoured() const { return true; }
 		virtual bool IsChild() override { return false; }
 
 		void CallProcessExit(ProcessHandle& h)
diff --git a/Engine/Source/Programs/UnrealBuildAccelerator/Common/Private/UbaStorage.cpp b/Engine/Source/Programs/UnrealBuildAccelerator/Common/Private/UbaStorage.cpp
--- a/Engine/Source/Programs/UnrealBuildAccelerator/Common/Private/UbaStorage.cpp
+++ b/Engine/Source/Programs/UnrealBuildAccelerator/Common/Private/UbaStorage.cpp
@@ -7,6 +7,13 @@
 #include "UbaDirectoryIterator.h"
 #include "UbaWorkManager.h"
 
+#if PLATFORM_WINDOWS && !defined(aligned_alloc)
+#define aligned_alloc(a, s) _aligned_malloc(s, a)
+#define aligned_free(p) _aligned_free(p)
+#else
+#define aligned_free(p) free(p)
+#endif
+
 namespace uba
 {
 	constexpr u32 CasTableVersion = IsWindows ? 32 : 34;
@@ -328,8 +335,24 @@ namespace uba
 		{
 			WorkRec() = delete;
 			WorkRec(const WorkRec&) = delete;
-			WorkRec(u32 wc) { workCount = wc; events = new Event[workCount]; }
-			~WorkRec() { delete[] events; }
+			WorkRec(u32 wc)
+			{
+				workCount = wc;
+
+				// For some very unknown reason ASAN on linux triggers on the "new[]" call when doing delete[]
+				// while doing it manually works properly. Will stop investigating this and move on
+				//events = new Event[workCount];
+				events = (Event*)aligned_alloc(alignof(Event), sizeof(Event)*workCount);
+				for (auto i=0;i!=workCount; ++i)
+					new (events + i) Event();
+			}
+			~WorkRec()
+			{
+				//delete[] events;
+				for (auto i=0;i!=workCount; ++i)
+					events[i].~Event();
+				aligned_free(events);
+			}
 			Atomic<u64> refCount;
 			Atomic<u64> compressCounter;
 			Event* events;
diff --git a/Engine/Source/Programs/UnrealBuildAccelerator/Common/Private/UbaThread.cpp b/Engine/Source/Programs/UnrealBuildAccelerator/Common/Private/UbaThread.cpp
--- a/Engine/Source/Programs/UnrealBuildAccelerator/Common/Private/UbaThread.cpp
+++ b/Engine/Source/Programs/UnrealBuildAccelerator/Common/Private/UbaThread.cpp
@@ -5,7 +5,7 @@
 
 namespace uba
 {
-	bool AlternateGroupAffinity(void* nativeThreadHandle)
+	bool AlternateThreadGroupAffinity(void* nativeThreadHandle)
 	{
 #if PLATFORM_WINDOWS
 		static int processorGroupCount = GetActiveProcessorGroupCount();
@@ -44,7 +44,7 @@ namespace uba
 		m_func = std::move(f);
 #if PLATFORM_WINDOWS
 		m_handle = CreateThread(NULL, 0, [](LPVOID p) -> DWORD { return ((Thread*)p)->m_func(); }, this, 0, NULL);
-		AlternateGroupAffinity(m_handle);
+		//AlternateThreadGroupAffinity(m_handle);
 #else
 		int err = 0;
 
diff --git a/Engine/Source/Programs/UnrealBuildAccelerator/Common/Public/UbaProcess.h b/Engine/Source/Programs/UnrealBuildAccelerator/Common/Public/UbaProcess.h
--- a/Engine/Source/Programs/UnrealBuildAccelerator/Common/Public/UbaProcess.h
+++ b/Engine/Source/Programs/UnrealBuildAccelerator/Common/Public/UbaProcess.h
@@ -20,9 +20,10 @@ namespace uba
 		virtual const Vector<ProcessLogLine>& GetLogLines() = 0;
 		virtual const Vector<u8>& GetTrackedInputs() = 0;
 		virtual void Cancel(bool terminate) { UBA_ASSERT(false); }
-		virtual const tchar* GetExecutingHost() const = 0;
-		virtual bool IsRemote() const = 0;
-		virtual bool IsChild() = 0;
+		virtual const tchar* GetExecutingHost() const { UBA_ASSERT(false); return nullptr; }
+		virtual bool IsRemote() const { UBA_ASSERT(false); return false; }
+		virtual bool IsDetoured() const { UBA_ASSERT(false); return false; }
+		virtual bool IsChild() { UBA_ASSERT(false); return false; }
 
 	protected:
 		void AddRef();
@@ -48,6 +49,7 @@ namespace uba
 		virtual void Cancel(bool terminate) override;
 		virtual const tchar* GetExecutingHost() const override { return TC(""); }
 		virtual bool IsRemote() const override { return false; }
+		virtual bool IsDetoured() const { return m_detourEnabled; }
 		virtual bool IsChild() override { return m_parentProcess != nullptr; }
 
 		ProcessImpl(Session& session, u32 id, ProcessImpl* parent);
@@ -94,6 +96,7 @@ namespace uba
 		Event m_writeEvent;
 		Event m_readEvent;
 	#else
+		ReaderWriterLock m_comMemoryLock;
 		bool m_cancelled = false;
 		Event& m_cancelEvent;
 		Event& m_writeEvent;
diff --git a/Engine/Source/Programs/UnrealBuildAccelerator/Common/Public/UbaProcessHandle.h b/Engine/Source/Programs/UnrealBuildAccelerator/Common/Public/UbaProcessHandle.h
--- a/Engine/Source/Programs/UnrealBuildAccelerator/Common/Public/UbaProcessHandle.h
+++ b/Engine/Source/Programs/UnrealBuildAccelerator/Common/Public/UbaProcessHandle.h
@@ -38,6 +38,7 @@ namespace uba
 		void Cancel(bool terminate) const;					// Request to cancel process. Set terminate to true to quickly kill process
 		const tchar* GetExecutingHost() const;				// Host that is executing process.
 		bool IsRemote() const;								// Returns true if process is a remote process
+		bool IsDetoured() const;							// Returns true if process is detoured
 
 		ProcessHandle();
 		ProcessHandle(const ProcessHandle& o);
diff --git a/Engine/Source/Programs/UnrealBuildAccelerator/Common/Public/UbaScheduler.h b/Engine/Source/Programs/UnrealBuildAccelerator/Common/Public/UbaScheduler.h
--- a/Engine/Source/Programs/UnrealBuildAccelerator/Common/Public/UbaScheduler.h
+++ b/Engine/Source/Programs/UnrealBuildAccelerator/Common/Public/UbaScheduler.h
@@ -13,42 +13,103 @@ namespace uba
 	struct NextProcessInfo;
 	struct ProcessStartInfoHolder;
 
-	// This is a simple scheduler that does not handle any dependencies. Can be added based on need
+	struct SchedulerCreateInfo
+	{
+		SessionServer& session;
+		u32 maxLocalProcessors = ~0u; // Max local processors to use. ~0u means it will use all processors
+		bool enableProcessReuse = false; // If this is true, the system will allow processes to be reused when they're asking for it.
+		bool forceRemote = false; // Force all processes that can run remotely to run remotely.
+	};
+
+	struct EnqueueProcessInfo
+	{
+		const ProcessStartInfo& info;
+		
+		float weight = 1.0f; // Weight of process. This is used towards max local processors. If a process is multithreaded it is likely it's weight should be more than 1.0
+		bool canDetour = true; // If true, uba will detour the process. If false it will just create pipes for std out and then run the process as-is.
+		bool canExecuteRemotely = true; // If true, this process can run on other machines, if false it will always be executed locally
+
+		const void* knownInputs = nullptr; // knownInputs is a memory block with null terminated tchar strings followed by an empty null terminated string to end.
+		u32 knownInputsBytes = 0; // knownInputsBytes is the total size in bytes of knownInputs
+		u32 knownInputsCount = 0; // knownInputsCount is the number of strings in the memory block
+
+		const u32* dependencies = nullptr; // An array of u32 holding indicies to processes this process depends on. Index is a rolling number returned by EnqueueProcess
+		u32 dependencyCount = 0; // Number of elements in dependencies
+	};
 
 	class Scheduler
 	{
 	public:
-		Scheduler(SessionServer& session, u32 maxLocalProcessors = ~0u, bool enableProcessReuse = false);
+		Scheduler(const SchedulerCreateInfo& info);
 		~Scheduler();
 
-		void Start();
-		void Stop();
-		void SetMaxLocalProcessors(u32 maxLocalProcessors);
+		void Start(); // Start scheduler thread. Should be called before server starts listen to connections if using remote help
+		void Stop(); // Will wait on all active processes and then exit.
+		void SetMaxLocalProcessors(u32 maxLocalProcessors); // Set max local processes
 
-		u32 EnqueueProcess(const ProcessStartInfo& info, float weight = 1.0f, const void* knownInputs = nullptr, u32 knownInputsBytes = 0, u32 knownInputsCount = 0);
+		u32 EnqueueProcess(const EnqueueProcessInfo& info); // Returns index of process. Index is a rolling number
 
 		void GetStats(u32& outQueued, u32& outActiveLocal, u32& outActiveRemote, u32& outFinished);
+
+		bool EnqueueFromFile(const tchar* yamlFilename); // Enqueue actions from file. Example of format of file is at the end of this file
+
+		void SetProcessFinishedCallback(const Function<void(const ProcessHandle&)>& processFinished); // Set callback 
+
 	private:
 		struct ExitProcessInfo;
+		struct ProcessStartInfo2;
 
 		void ThreadLoop();
 		void RemoteProcessReturned(Process& process);
 		void RemoteSlotAvailable();
 		void ProcessExited(ExitProcessInfo* info, const ProcessHandle& handle);
-		bool RunQueuedProcess(bool runLocal);
+		u32 PopProcess(bool isLocal);
+		bool RunQueuedProcess(bool isLocal);
 		bool HandleReuseMessage(Process& process, NextProcessInfo& outNextProcess, u32 prevExitCode);
+		void ExitProcess(ExitProcessInfo& info, Process& process, u32 exitCode);
+		void SkipProcess(ProcessStartInfo2& info);
+		void UpdateQueueCounter(int offset);
+		void UpdateActiveProcessCounter(bool isLocal, int offset);
+		void FinishProcess(const ProcessHandle& handle);
 
 		SessionServer& m_session;
 		u32 m_maxLocalProcessors;
+	
+		enum ProcessStatus : u8
+		{
+			ProcessStatus_Queued,
+			ProcessStatus_Running,
+			ProcessStatus_Success,
+			ProcessStatus_Failed,
+			ProcessStatus_Skipped,
+		};
+
+		
+		struct ProcessEntry
+		{
+			ProcessStartInfo2* info;
+			u32* dependencies;
+			u32 dependencyCount;
+			ProcessStatus status;
+			u8 canDetour;
+			u8 canExecuteRemotely;
+		};
+
+		ReaderWriterLock m_processEntriesLock;
+		Vector<ProcessEntry> m_processEntries;
+		u32 m_processEntriesStart = 0;
+
+		Function<void(const ProcessHandle&)> m_processFinished;
 
-		ReaderWriterLock m_queuedProcessesLock;
-		struct QueuedProcess;
-		List<QueuedProcess> m_queuedProcesses;
 		Event m_updateThreadLoop;
 		Thread m_thread;
 		bool m_loop = false;
 		bool m_enableProcessReuse;
+		bool m_forceRemote;
 
+		float m_activeLocalProcessWeight = 0.0f;
+
+		Atomic<u32> m_queuedProcesses;
 		Atomic<u32> m_activeLocalProcesses;
 		Atomic<u32> m_activeRemoteProcesses;
 		Atomic<u32> m_finishedProcesses;
@@ -57,3 +118,42 @@ namespace uba
 		void operator=(const Scheduler&) = delete;
 	};
 }
+
+#if 0
+// Example of yaml file with processes that can be queued up in scheduler
+// id - Not used right now. Number does not matter because id will be a rolling number
+// app - Application to execute
+// arg - Arguments to application
+// dir - Working directory
+// desc - Description of process
+// weight - How much cpu this process is using. Optional. Defaults to 1.0
+// remote - Decides if this process can execute remotely. Optional. Defaults to true
+// detour - Decides if this process can be detoured. Optional. Defaults to true.
+// dep - Dependencies. An array of indices to other processes
+
+processes:
+  - id: 0
+    app: C:\Program Files\Microsoft Visual Studio\2022\Professional\VC\Tools\MSVC\14.38.33130\bin\Hostx64\x64\cl.exe
+    arg: @"..\Intermediate\Build\Win64\x64\UnrealPak\Development\Core\SharedPCH.Core.Cpp20.h.obj.rsp"
+    dir: E:\dev\fn\Engine\Source
+    desc: SharedPCH.Core.Cpp20.cpp
+    weight: 1.25
+    remote: false
+
+  - id: 44
+    app: C:\Program Files\Microsoft Visual Studio\2022\Professional\VC\Tools\MSVC\14.38.33130\bin\Hostx64\x64\cl.exe
+    arg: @"..\Intermediate\Build\Win64\x64\UnrealPak\Development\Json\Module.Json.cpp.obj.rsp"
+    dir: E:\dev\fn\Engine\Source
+    desc: Module.Json.cpp
+    weight: 1.5
+    dep: [0]
+
+  - id: 337
+    app: E:\dev\fn\Engine\Binaries\ThirdParty\DotNet\6.0.302\windows\dotnet.exe
+    arg: "E:\dev\fn\Engine\Binaries\DotNET\UnrealBuildTool\UnrealBuildTool.dll" -Mode=WriteMetadata -Input="E:\dev\fn\Engine\Intermediate\Build\Win64\x64\UnrealPak\Development\TargetMetadata.dat" -Version=2
+    dir: E:\dev\fn\Engine\Source
+    desc: UnrealPak.target
+    detour: false
+    dep: [336, 0]
+
+#endif
\ No newline at end of file
diff --git a/Engine/Source/Programs/UnrealBuildAccelerator/Common/Public/UbaThread.h b/Engine/Source/Programs/UnrealBuildAccelerator/Common/Public/UbaThread.h
--- a/Engine/Source/Programs/UnrealBuildAccelerator/Common/Public/UbaThread.h
+++ b/Engine/Source/Programs/UnrealBuildAccelerator/Common/Public/UbaThread.h
@@ -30,5 +30,5 @@ namespace uba
 	};
 
 
-	bool AlternateGroupAffinity(void* nativeThreadHandle);
+	bool AlternateThreadGroupAffinity(void* nativeThreadHandle);
 }
diff --git a/Engine/Source/Programs/UnrealBuildAccelerator/Core/Public/UbaCorePch.h b/Engine/Source/Programs/UnrealBuildAccelerator/Core/Public/UbaCorePch.h
--- a/Engine/Source/Programs/UnrealBuildAccelerator/Core/Public/UbaCorePch.h
+++ b/Engine/Source/Programs/UnrealBuildAccelerator/Core/Public/UbaCorePch.h
@@ -3,7 +3,6 @@
 #pragma once
 
 #if UBA_USE_MIMALLOC
-#include <mimalloc.h>
 #include <mimalloc-override.h>
 #endif
 
diff --git a/Engine/Source/Programs/UnrealBuildAccelerator/Core/Public/UbaMemory.h b/Engine/Source/Programs/UnrealBuildAccelerator/Core/Public/UbaMemory.h
--- a/Engine/Source/Programs/UnrealBuildAccelerator/Core/Public/UbaMemory.h
+++ b/Engine/Source/Programs/UnrealBuildAccelerator/Core/Public/UbaMemory.h
@@ -13,7 +13,7 @@
 #include <functional>
 
 #if UBA_USE_MIMALLOC
-#include <mimalloc.h>
+#include <mimalloc-override.h>
 #endif
 
 namespace uba
diff --git a/Engine/Source/Programs/UnrealBuildAccelerator/Detours/Private/PosixOS/UbaDetoursFunctionsPosix.cpp b/Engine/Source/Programs/UnrealBuildAccelerator/Detours/Private/PosixOS/UbaDetoursFunctionsPosix.cpp
--- a/Engine/Source/Programs/UnrealBuildAccelerator/Detours/Private/PosixOS/UbaDetoursFunctionsPosix.cpp
+++ b/Engine/Source/Programs/UnrealBuildAccelerator/Detours/Private/PosixOS/UbaDetoursFunctionsPosix.cpp
@@ -14,6 +14,10 @@
 #include <dlfcn.h>
 #include <sys/stat.h>
 
+// undefs for mimalloc
+#undef realpath
+#undef malloc
+
 #if PLATFORM_LINUX
 #include <sys/prctl.h>
 #else
diff --git a/Engine/Source/Programs/UnrealBuildAccelerator/Detours/Private/Windows/UbaDetoursMainWin.cpp b/Engine/Source/Programs/UnrealBuildAccelerator/Detours/Private/Windows/UbaDetoursMainWin.cpp
--- a/Engine/Source/Programs/UnrealBuildAccelerator/Detours/Private/Windows/UbaDetoursMainWin.cpp
+++ b/Engine/Source/Programs/UnrealBuildAccelerator/Detours/Private/Windows/UbaDetoursMainWin.cpp
@@ -130,6 +130,8 @@ BOOL WINAPI DllMain(HINSTANCE hinst, DWORD dwReason, LPVOID reserved)
 {
 	using namespace uba;
 
+	static GROUP_AFFINITY GroupAffinity;
+
 	u64 startTime = GetTime();
 
 	if (DetourIsHelperProcess())
@@ -137,6 +139,8 @@ BOOL WINAPI DllMain(HINSTANCE hinst, DWORD dwReason, LPVOID reserved)
 
 	if (dwReason == DLL_PROCESS_ATTACH)
 	{
+		GetThreadGroupAffinity(GetCurrentThread(), &GroupAffinity);
+
 		if (!DetourRestoreAfterWith())
 			TerminateCurrentProcess(1344);
 
@@ -157,6 +161,10 @@ BOOL WINAPI DllMain(HINSTANCE hinst, DWORD dwReason, LPVOID reserved)
 
 		Init(*payload, startTime);
 	}
+	else if (dwReason == DLL_THREAD_ATTACH)
+	{
+		SetThreadGroupAffinity(GetCurrentThread(), &GroupAffinity, NULL);
+	}
 	else if (dwReason == DLL_PROCESS_DETACH)
 	{
 		Deinit(startTime);
diff --git a/Engine/Source/Programs/UnrealBuildAccelerator/Host/Private/UbaExports.cpp b/Engine/Source/Programs/UnrealBuildAccelerator/Host/Private/UbaExports.cpp
--- a/Engine/Source/Programs/UnrealBuildAccelerator/Host/Private/UbaExports.cpp
+++ b/Engine/Source/Programs/UnrealBuildAccelerator/Host/Private/UbaExports.cpp
@@ -460,7 +460,10 @@ uba::StorageClient* CreateStorageClient(uba::NetworkClient& client, const uba::t
 
 	uba::Scheduler* Scheduler_Create(uba::SessionServer* session, uba::u32 maxLocalProcessors, bool enableProcessReuse)
 	{
-		return new uba::Scheduler(*session, maxLocalProcessors, enableProcessReuse);
+		uba::SchedulerCreateInfo info{*session};
+		info.maxLocalProcessors = maxLocalProcessors;
+		info.enableProcessReuse = enableProcessReuse;
+		return new uba::Scheduler(info);
 	}
 
 	void Scheduler_Start(uba::Scheduler* scheduler)
@@ -470,7 +473,12 @@ uba::StorageClient* CreateStorageClient(uba::NetworkClient& client, const uba::t
 
 	void Scheduler_EnqueueProcess(uba::Scheduler* scheduler, const uba::ProcessStartInfo& info, float weight, const void* knownInputs, uba::u32 knownInputsBytes, uba::u32 knownInputsCount)
 	{
-		scheduler->EnqueueProcess(info, weight, knownInputs, knownInputsBytes, knownInputsCount);
+		uba::EnqueueProcessInfo epi(info);
+		epi.weight = weight;
+		epi.knownInputs = knownInputs;
+		epi.knownInputsBytes = knownInputsBytes;
+		epi.knownInputsCount = knownInputsCount;
+		scheduler->EnqueueProcess(epi);
 	}
 
 	void Scheduler_SetMaxLocalProcessors(uba::Scheduler* scheduler, uba::u32 maxLocalProcessors)
