[UBA]
* Turned off mimalloc again for linux because of new errors not seen before
* Improved some error descriptions and bumped error stack buffer a bit

diff --git a/Engine/Source/Programs/UnrealBuildAccelerator/Detours/Private/PosixOS/UbaDetoursFunctionsPosix.cpp b/Engine/Source/Programs/UnrealBuildAccelerator/Detours/Private/PosixOS/UbaDetoursFunctionsPosix.cpp
--- a/Engine/Source/Programs/UnrealBuildAccelerator/Detours/Private/PosixOS/UbaDetoursFunctionsPosix.cpp
+++ b/Engine/Source/Programs/UnrealBuildAccelerator/Detours/Private/PosixOS/UbaDetoursFunctionsPosix.cpp
@@ -659,7 +659,8 @@ int Shared_stat(const char* funcName, const char* file, struct stat* attr, const
 		}
 		else
 		{
-			UBA_ASSERT(fileAttr.lastError == errno);
+			UBA_ASSERTF(fileAttr.lastError == errno || ((fileAttr.lastError == ENOTDIR || fileAttr.lastError == ENOENT) && (errno == ENOTDIR || errno == ENOENT))
+				, "Detoured stat returned a different error. Returned %i (%s) but should return %i (%s)", fileAttr.lastError, strerror(fileAttr.lastError), errno, strerror(errno));
 		}
 	}
 	#endif
@@ -1989,7 +1990,7 @@ namespace uba
 	ANALYSIS_NORETURN void UbaAssert(const tchar* text, const char* file, u32 line, const char* expr, u32 terminateCode)
 	{
 		SuppressDetourScope s;
-		StringBuffer<4096> b;
+		StringBuffer<8*1024> b;
 		WriteAssertInfo(b, text, file, line, expr, 1);
 		Rpc_WriteLog(b.data, b.count, true, true);
 
diff --git a/Engine/Source/Programs/UnrealBuildAccelerator/Mimalloc/UbaMimalloc.Build.cs b/Engine/Source/Programs/UnrealBuildAccelerator/Mimalloc/UbaMimalloc.Build.cs
--- a/Engine/Source/Programs/UnrealBuildAccelerator/Mimalloc/UbaMimalloc.Build.cs
+++ b/Engine/Source/Programs/UnrealBuildAccelerator/Mimalloc/UbaMimalloc.Build.cs
@@ -11,7 +11,8 @@ public class UbaMimalloc : ModuleRules
 	{
 		PCHUsage = PCHUsageMode.NoPCHs;
 
-		if (Target.Platform.IsInGroup(UnrealPlatformGroup.Windows) || Target.Platform.IsInGroup(UnrealPlatformGroup.Linux)) {
+		if (Target.Platform.IsInGroup(UnrealPlatformGroup.Windows))// || Target.Platform.IsInGroup(UnrealPlatformGroup.Linux))
+		{
 			PublicDependencyModuleNames.AddRange(new string[] {
 				"mimalloc212",
 			});
diff --git a/Engine/Source/Programs/UnrealBuildAccelerator/UbaAgent.Target.cs b/Engine/Source/Programs/UnrealBuildAccelerator/UbaAgent.Target.cs
--- a/Engine/Source/Programs/UnrealBuildAccelerator/UbaAgent.Target.cs
+++ b/Engine/Source/Programs/UnrealBuildAccelerator/UbaAgent.Target.cs
@@ -105,7 +105,7 @@ public class UbaAgentTarget : TargetRules
 		}
 
 		int useMiMalloc = 0;
-		if (Target.Platform.IsInGroup(UnrealPlatformGroup.Windows) || Target.Platform.IsInGroup(UnrealPlatformGroup.Linux))
+		if (Target.Platform.IsInGroup(UnrealPlatformGroup.Windows))// || Target.Platform.IsInGroup(UnrealPlatformGroup.Linux))
 		{
 			if (Rules.bUsePCHFiles)
 			{
