[UBA]
* Fixed unit test compile errors
* Fixed race condition in unit test shutdown

diff --git a/Engine/Source/Programs/UnrealBuildAccelerator/Common/Private/UbaScheduler.cpp b/Engine/Source/Programs/UnrealBuildAccelerator/Common/Private/UbaScheduler.cpp
--- a/Engine/Source/Programs/UnrealBuildAccelerator/Common/Private/UbaScheduler.cpp
+++ b/Engine/Source/Programs/UnrealBuildAccelerator/Common/Private/UbaScheduler.cpp
@@ -88,11 +88,14 @@ namespace uba
 		m_thread.Wait();
 		m_session.WaitOnAllTasks();
 
+		ScopedWriteLock lock(m_processEntriesLock);
 		for (auto& entry : m_processEntries)
 		{
+			UBA_ASSERT(entry.status !=ProcessStatus_Running);
 			delete[] entry.dependencies;
 			delete entry.info;
 		}
+		m_processEntries.clear();
 	}
 
 	void Scheduler::SetMaxLocalProcessors(u32 maxLocalProcessors)
diff --git a/Engine/Source/Programs/UnrealBuildAccelerator/Test/Private/UbaTestScheduler.h b/Engine/Source/Programs/UnrealBuildAccelerator/Test/Private/UbaTestScheduler.h
--- a/Engine/Source/Programs/UnrealBuildAccelerator/Test/Private/UbaTestScheduler.h
+++ b/Engine/Source/Programs/UnrealBuildAccelerator/Test/Private/UbaTestScheduler.h
@@ -11,19 +11,20 @@ namespace uba
 	{
 		return RunLocal(logger, testRootDir, [](LoggerWithWriter& logger, SessionServer& session, const tchar* workingDir, const RunProcessFunction& runProcess)
 			{
-
-				Scheduler scheduler(session, true);
+				SchedulerCreateInfo info(session);
+				Scheduler scheduler(info);
 
 				ProcessStartInfo processInfo;
 				processInfo.application = GetSystemApplication();
 				processInfo.workingDir = workingDir;
 				processInfo.arguments = GetSystemArguments();
 
-				scheduler.EnqueueProcess(processInfo);
+				EnqueueProcessInfo epi(processInfo);
+				scheduler.EnqueueProcess(epi);
 				scheduler.Start();
 
 				u32 queued, activeLocal, activeRemote, finished;
-				do { scheduler.GetStats(queued, activeLocal, activeRemote, finished); } while ((queued+activeLocal+activeRemote) != 0);
+				do { scheduler.GetStats(queued, activeLocal, activeRemote, finished); } while (finished != 1);
 
 				scheduler.Stop();
 				return true;
@@ -37,8 +38,9 @@ namespace uba
 
 		return RunLocal(logger, testRootDir, [](LoggerWithWriter& logger, SessionServer& session, const tchar* workingDir, const RunProcessFunction& runProcess)
 			{
-
-				Scheduler scheduler(session, ~0u, true);
+				SchedulerCreateInfo info(session);
+				info.enableProcessReuse = true;
+				Scheduler scheduler(info);
 
 				StringBuffer<> testApp;
 				GetTestAppPath(logger, testApp);
@@ -48,11 +50,12 @@ namespace uba
 				processInfo.workingDir = workingDir;
 				processInfo.arguments = TC("-reuse");
 
-				scheduler.EnqueueProcess(processInfo);
+				EnqueueProcessInfo epi(processInfo);
+				scheduler.EnqueueProcess(epi);
 				scheduler.Start();
 
 				u32 queued, activeLocal, activeRemote, finished;
-				do { scheduler.GetStats(queued, activeLocal, activeRemote, finished); } while ((queued + activeLocal + activeRemote) != 0);
+				do { scheduler.GetStats(queued, activeLocal, activeRemote, finished); } while (finished != 1);
 
 				scheduler.Stop();
 				return true;
@@ -66,8 +69,10 @@ namespace uba
 
 		return RunRemote(logger, testRootDir, [](LoggerWithWriter& logger, SessionServer& session, const tchar* workingDir, const RunProcessFunction& runProcess)
 			{
-
-				Scheduler scheduler(session, 0u, true);
+				SchedulerCreateInfo info(session);
+				info.enableProcessReuse = true;
+				info.maxLocalProcessors = 0;
+				Scheduler scheduler(info);
 
 				StringBuffer<> testApp;
 				GetTestAppPath(logger, testApp);
@@ -77,11 +82,12 @@ namespace uba
 				processInfo.workingDir = workingDir;
 				processInfo.arguments = TC("-reuse");
 
-				scheduler.EnqueueProcess(processInfo);
+				EnqueueProcessInfo epi(processInfo);
+				scheduler.EnqueueProcess(epi);
 				scheduler.Start();
 
 				u32 queued, activeLocal, activeRemote, finished;
-				do { scheduler.GetStats(queued, activeLocal, activeRemote, finished); } while ((queued + activeLocal + activeRemote) != 0);
+				do { scheduler.GetStats(queued, activeLocal, activeRemote, finished); } while (finished != 1);
 
 				scheduler.Stop();
 				return true;
