[Uba]
* Fixed crash in trace system when log lines are insanely long (over 256k)

diff --git a/Engine/Source/Programs/UnrealBuildAccelerator/Common/Private/UbaTrace.cpp b/Engine/Source/Programs/UnrealBuildAccelerator/Common/Private/UbaTrace.cpp
--- a/Engine/Source/Programs/UnrealBuildAccelerator/Common/Private/UbaTrace.cpp
+++ b/Engine/Source/Programs/UnrealBuildAccelerator/Common/Private/UbaTrace.cpp
@@ -26,16 +26,7 @@ namespace uba
 	{
 		WriterScope(Trace& trace) : ScopedWriteLock(trace.m_memoryLock), BinaryWriter(trace.m_memoryBegin, trace.m_memoryPos, trace.m_memoryCapacity), m_trace(trace)
 		{
-			u64 committedMemoryNeeded = AlignUp(trace.m_memoryPos + TraceMessageMaxSize, TraceMessageMaxSize);
-			if (trace.m_memoryCommitted >= committedMemoryNeeded)
-				return;
-			if (!MapViewCommit(trace.m_memoryBegin + trace.m_memoryCommitted, committedMemoryNeeded - trace.m_memoryCommitted))
-			{
-				trace.m_logger.Error(TC("Failed to commit memory for trace (Pos: %llu Capacity: %llu, Already Committed: %llu, Needed: %llu): %s"), trace.m_memoryPos, trace.m_memoryCapacity, trace.m_memoryCommitted, committedMemoryNeeded, LastErrorToText().data);
-				return;
-			}
-
-			trace.m_memoryCommitted = committedMemoryNeeded;
+			EnsureMemory(TraceMessageMaxSize);
 		}
 
 		~WriterScope()
@@ -44,6 +35,20 @@ namespace uba
 			*(u32*)m_trace.m_memoryBegin = u32(m_trace.m_memoryPos);
 		}
 
+		void EnsureMemory(u64 size)
+		{
+			u64 committedMemoryNeeded = AlignUp(m_trace.m_memoryPos + size, size);
+			if (m_trace.m_memoryCommitted >= committedMemoryNeeded)
+				return;
+			if (!MapViewCommit(m_trace.m_memoryBegin + m_trace.m_memoryCommitted, committedMemoryNeeded - m_trace.m_memoryCommitted))
+			{
+				m_trace.m_logger.Error(TC("Failed to commit memory for trace (Pos: %llu Capacity: %llu, Already Committed: %llu, Needed: %llu): %s"), m_trace.m_memoryPos, m_trace.m_memoryCapacity, m_trace.m_memoryCommitted, committedMemoryNeeded, LastErrorToText().data);
+				return;
+			}
+
+			m_trace.m_memoryCommitted = committedMemoryNeeded;
+		}
+
 		Trace& m_trace;
 	};
 
@@ -199,6 +204,7 @@ namespace uba
 		writer.Write7BitEncoded(u64(logLines.size()));
 		for (auto& line : logLines)
 		{
+			writer.EnsureMemory(line.text.size()*sizeof(tchar));
 			writer.WriteByte(line.type);
 			writer.WriteString(line.text);
 		}
