[UBA]
* Fixed mac compile errors by adding explicit ctor

diff --git a/Engine/Source/Programs/UnrealBuildAccelerator/Common/Public/UbaScheduler.h b/Engine/Source/Programs/UnrealBuildAccelerator/Common/Public/UbaScheduler.h
--- a/Engine/Source/Programs/UnrealBuildAccelerator/Common/Public/UbaScheduler.h
+++ b/Engine/Source/Programs/UnrealBuildAccelerator/Common/Public/UbaScheduler.h
@@ -15,6 +15,8 @@ namespace uba
 
 	struct SchedulerCreateInfo
 	{
+		SchedulerCreateInfo(SessionServer& s) : session(s) {}
+
 		SessionServer& session;
 		u32 maxLocalProcessors = ~0u; // Max local processors to use. ~0u means it will use all processors
 		bool enableProcessReuse = false; // If this is true, the system will allow processes to be reused when they're asking for it.
@@ -23,6 +25,8 @@ namespace uba
 
 	struct EnqueueProcessInfo
 	{
+		EnqueueProcessInfo(const ProcessStartInfo& i) : info(i) {}
+
 		const ProcessStartInfo& info;
 		
 		float weight = 1.0f; // Weight of process. This is used towards max local processors. If a process is multithreaded it is likely it's weight should be more than 1.0
