[UBA]
* Added code to silence STATUS_NO_MEDIA_IN_DEVICE errors when caching folders.

diff --git a/Engine/Source/Programs/UnrealBuildAccelerator/Common/Public/UbaDirectoryIterator.h b/Engine/Source/Programs/UnrealBuildAccelerator/Common/Public/UbaDirectoryIterator.h
--- a/Engine/Source/Programs/UnrealBuildAccelerator/Common/Public/UbaDirectoryIterator.h
+++ b/Engine/Source/Programs/UnrealBuildAccelerator/Common/Public/UbaDirectoryIterator.h
@@ -91,7 +91,7 @@ namespace uba
 		NTSTATUS res = NtCreateFile(&handle, FILE_READ_ATTRIBUTES | FILE_LIST_DIRECTORY | FILE_ATTRIBUTE_UNPINNED, &ObjectAttributes, &IoStatusBlock, 0, FILE_ATTRIBUTE_NORMAL, ShareAccess, FILE_OPEN, FILE_DIRECTORY_FILE | FILE_SYNCHRONOUS_IO_ALERT, NULL, 0);
 		if (res == STATUS_OBJECT_NAME_NOT_FOUND || res == STATUS_OBJECT_PATH_NOT_FOUND)
 			return !errorOnNotFound;
-		if (res == STATUS_NOT_A_DIRECTORY || res == STATUS_ACCESS_DENIED)
+		if (res == STATUS_NOT_A_DIRECTORY || res == STATUS_ACCESS_DENIED || res == STATUS_NO_MEDIA_IN_DEVICE)
 			return false;
 		if (res != STATUS_SUCCESS)
 			return logger.Error(L"NtCreateFile for TraverseDir on %s failed with error code %x", dirPath, (DWORD)res);
