[uba]
* Fixed naming of some help information and documentation after feedback from external user

diff --git a/Engine/Source/Programs/UnrealBuildAccelerator/Agent/Private/UbaAgent.cpp b/Engine/Source/Programs/UnrealBuildAccelerator/Agent/Private/UbaAgent.cpp
--- a/Engine/Source/Programs/UnrealBuildAccelerator/Agent/Private/UbaAgent.cpp
+++ b/Engine/Source/Programs/UnrealBuildAccelerator/Agent/Private/UbaAgent.cpp
@@ -49,7 +49,7 @@ namespace uba
 	}();
 	u32				DefaultProcessorCount = []() { return GetLogicalProcessorCount(); }();
 	const tchar*	DefaultAgentName = []() { static tchar buf[256]; GetComputerNameW(buf, sizeof_array(buf)); return buf; }();
-	u32				DefaultMaxTcpCount = 4;
+	u32				DefaultMaxConnectionCount = 4;
 
 	int PrintHelp(const tchar* message)
 	{
@@ -74,7 +74,7 @@ namespace uba
 		logger.Info(TC("  -proxyport=<port>       Which port that agent will use if being assigned to be proxy for other agents (default: %u)"), DefaultStorageProxyPort);
 		logger.Info(TC("  -maxcpu=<number>        Max number of processes that can be started. Defaults to \"%u\" on this machine"), DefaultProcessorCount);
 		logger.Info(TC("  -mulcpu=<number>        This value multiplies with number of cpu to figure out max cpu. Defaults to 1.0"));
-		logger.Info(TC("  -maxtcp=<number>        Max number of tcp connections that can be started. Defaults to \"%u\" (amount up to max will depend on ping)"), DefaultMaxTcpCount);
+		logger.Info(TC("  -maxcon=<number>        Max number of connections that can be started by agent. Defaults to \"%u\" (amount up to max will depend on ping)"), DefaultMaxConnectionCount);
 		logger.Info(TC("  -capacity=<gigaby>      Capacity of local store. Defaults to %u gigabytes"), DefaultCapacityGb);
 		logger.Info(TC("  -quic                   Use Quic instead of tcp backend."));
 		logger.Info(TC("  -name=<name>            The identifier of this agent. Defaults to \"%s\" on this machine"), DefaultAgentName);
@@ -375,7 +375,7 @@ namespace uba
 
 		u32 maxProcessCount = DefaultProcessorCount;
 		float mulProcessValue = 1.0f;
-		u32 maxTcpConnectionCount = DefaultMaxTcpCount;
+		u32 maxConnectionCount = DefaultMaxConnectionCount;
 		u32 outputStatsThresholdMs = 0;
 		u32 storageCapacityGb = DefaultCapacityGb;
 		StringBuffer<256> host;
@@ -453,10 +453,10 @@ namespace uba
 				if (!value.Parse(mulProcessValue))
 					return PrintHelp(TC("Invalid value for -mulcpu"));
 			}
-			else if (name.Equals(TC("-maxtcp")))
+			else if (name.Equals(TC("-maxcon")) || name.Equals(TC("-maxtcp")))
 			{
-				if (!value.Parse(maxTcpConnectionCount) || maxTcpConnectionCount == 0)
-					return PrintHelp(TC("Invalid value for -maxtcp"));
+				if (!value.Parse(maxConnectionCount) || maxConnectionCount == 0)
+					return PrintHelp(TC("Invalid value for -maxcon"));
 			}
 			else if (name.Equals(TC("-capacity")))
 			{
@@ -731,7 +731,7 @@ namespace uba
 		#if UBA_DEBUG
 		dbgStr = TC(" (DEBUG)");
 		#endif
-		logger.Info(TC("UbaAgent v%s%s (Cpu: %u, MaxCon: %u, Dir: \"%s\", StoreCapacity: %uGb%s)"), Version, dbgStr, maxProcessCount, maxTcpConnectionCount, g_rootDir.data, storageCapacityGb, extraInfo.data);
+		logger.Info(TC("UbaAgent v%s%s (Cpu: %u, MaxCon: %u, Dir: \"%s\", StoreCapacity: %uGb%s)"), Version, dbgStr, maxProcessCount, maxConnectionCount, g_rootDir.data, storageCapacityGb, extraInfo.data);
 		if (!eventFile.IsEmpty())
 			logger.Info(TC("  Will poll for external events in file %s"), eventFile.data);
 		logger.Info(TC(""));
@@ -1006,20 +1006,20 @@ namespace uba
 				return 0;
 
 			Event wakeupSessionWait(false);
-			u32 targetTcpConnectionCount = 1;
+			u32 targetConnectionCount = 1;
 
 			struct Proxy
 			{
 				LogWriter& logWriter;
 				NetworkClient* client;
 				Event& wakeupSessionWait;
-				u32& maxTcpConnectionCount;
-				u32& targetTcpConnectionCount;
+				u32& maxConnectionCount;
+				u32& targetConnectionCount;
 				NetworkServer* server = nullptr;
 				StorageProxy* storage = nullptr;
 				StorageClient* storageClient = nullptr;
 				TString serverPrefix;
-			} proxy { g_consoleLogWriter, client, wakeupSessionWait, maxTcpConnectionCount, targetTcpConnectionCount };
+			} proxy { g_consoleLogWriter, client, wakeupSessionWait, maxConnectionCount, targetConnectionCount };
 			auto psg = MakeGuard([&]() { delete proxy.server; });
 			auto pg = MakeGuard([&]() { delete proxy.storage; });
 
@@ -1043,7 +1043,7 @@ namespace uba
 					}
 					proxy.storage = new StorageProxy(*proxy.server, *proxy.client, storageServerUid, TC("Wooohoo"), proxy.storageClient);
 					proxy.server->StartListen(proxy.client->GetTcpBackend(), proxyPort);
-					proxy.targetTcpConnectionCount = proxy.maxTcpConnectionCount;
+					proxy.targetConnectionCount = proxy.maxConnectionCount;
 					proxy.wakeupSessionWait.Set();
 					return true;
 				};
@@ -1160,25 +1160,25 @@ namespace uba
 			u64 lastLogTime = GetTime();
 			#endif
 
-			u32 tcpConnectionCount = 1;
+			u32 connectionCount = 1;
 
 			while (true)
 			{
 				if (useListen)
 				{
-					if (tcpConnectionCount != targetTcpConnectionCount)
+					if (connectionCount != targetConnectionCount)
 					{
-						client->SetConnectionCount(targetTcpConnectionCount);
-						tcpConnectionCount = targetTcpConnectionCount;
+						client->SetConnectionCount(targetConnectionCount);
+						connectionCount = targetConnectionCount;
 					}
 				}
 				else
 				{
-					while (tcpConnectionCount < targetTcpConnectionCount)
+					while (connectionCount < targetConnectionCount)
 					{
 						bool timedOut = false;
 						if (client->Connect(*networkBackend, host.data, port, &timedOut))
-							++tcpConnectionCount;
+							++connectionCount;
 					}
 				}
 
@@ -1186,10 +1186,10 @@ namespace uba
 					break;
 
 				// This is an estimation based on tcp limitations (ack and sliding windows).
-				// For every 15ms latency on "best ping") we increase targetTcpConnectionCount up to maxTcpConnectionCount
+				// For every 15ms latency on "best ping") we increase targetConnectionCount up to maxConnectionCount
 				if (!storageClient->IsUsingProxy())
 					if (u64 bestPing = sessionClient->GetBestPing())
-						targetTcpConnectionCount = Min(u32(TimeToMs(bestPing) / 15), maxTcpConnectionCount);
+						targetConnectionCount = Min(u32(TimeToMs(bestPing) / 15), maxConnectionCount);
 
 				if (!isTerminating)
 				{
