[Uba]
* Added code to validate that connections from server to client are all from the same server. Solved by giving servers a unique id and made sure it is part of connection handshake with client. Client will just discard connections that is providing a different server uid than the first connection

This will hopefully fix the bug seen on our farm a couple days ago. Theory is that two servers managed to connect to the same helper. First server is just about to start connecting to a helper when helper is brought down and restarted to help another server. Some sort of stall happens on first server and once it starts the actual tcp connection the helper has been brought up to help the other server. The helper allows more than one tcp connection for performance reasons so will accept tcp connection from both old and new server. In this messed up scenario some messages would go to one server and some to the other.. causing all kinds of weird things. It is critical that all connections go to the same server since all messages are just round robin the tcp connections.

diff --git a/Engine/Source/Programs/UnrealBuildAccelerator/Common/Private/UbaNetworkBackendTcp.cpp b/Engine/Source/Programs/UnrealBuildAccelerator/Common/Private/UbaNetworkBackendTcp.cpp
--- a/Engine/Source/Programs/UnrealBuildAccelerator/Common/Private/UbaNetworkBackendTcp.cpp
+++ b/Engine/Source/Programs/UnrealBuildAccelerator/Common/Private/UbaNetworkBackendTcp.cpp
@@ -42,6 +42,8 @@
 
 namespace uba
 {
+	constexpr u32 MaxHeaderSize = 24;
+
 	struct NetworkBackendTcp::ListenEntry
 	{
 		StringBuffer<128> ip;
@@ -194,7 +196,7 @@ namespace uba
 	void NetworkBackendTcp::SetRecvCallbacks(void* connection, void* context, u32 headerSize, RecvHeaderCallback* h, RecvBodyCallback* b, const tchar* recvHint)
 	{
 		UBA_ASSERT(h);
-		UBA_ASSERT(headerSize <= 16);
+		UBA_ASSERT(headerSize <= MaxHeaderSize);
 		auto& conn = *(Connection*)connection;
 		conn.recvContext = context;
 		conn.headerSize = headerSize;
@@ -428,7 +430,7 @@ namespace uba
 				u8* bodyData = nullptr;
 				u32 bodySize = 0;
 
-				u8 headerData[16];
+				u8 headerData[MaxHeaderSize];
 				if (!RecvSocket(m_logger, connection.socket, headerData, connection.headerSize, connection.recvTimeoutMs, connectionUid, connection.recvHint, TC(""), isFirst))
 					break;
 				isFirst = false;
diff --git a/Engine/Source/Programs/UnrealBuildAccelerator/Common/Private/UbaNetworkClient.cpp b/Engine/Source/Programs/UnrealBuildAccelerator/Common/Private/UbaNetworkClient.cpp
--- a/Engine/Source/Programs/UnrealBuildAccelerator/Common/Private/UbaNetworkClient.cpp
+++ b/Engine/Source/Programs/UnrealBuildAccelerator/Common/Private/UbaNetworkClient.cpp
@@ -86,12 +86,24 @@ namespace uba
 				c.recvEvent.Set();
 			});
 
-		backend.SetRecvCallbacks(backendConnection, &rc, 1, [](void* context, u8* headerData, void*& outBodyContext, u8*& outBodyData, u32& outBodySize)
+		backend.SetRecvCallbacks(backendConnection, &rc, 1 + sizeof(Guid), [](void* context, u8* headerData, void*& outBodyContext, u8*& outBodyData, u32& outBodySize)
 			{
 				auto& c = *(RecvContext*)context;
 				c.error = *headerData;
+				Guid serverUid = *(Guid*)(headerData+1);
+
+				if (!c.error)
+				{
+					ScopedWriteLock lock(c.client.m_serverUidLock);
+					if (c.client.m_serverUid == Guid())
+						c.client.m_serverUid = serverUid;
+					else if (c.client.m_serverUid != serverUid) // Seems like two different servers tried to connect to this client.. keep the first one and ignore the others
+						c.error = 5;
+				}
+
 				if (!c.error)
 					c.client.ConnectedCallback(c.backend, c.backendConnection);
+
 				c.recvEvent.Set();
 				return true;
 			}, nullptr, TC("Connecting"));
@@ -151,6 +163,13 @@ namespace uba
 			Sleep(1000); // Kind of ugly, but we want the retry-clients to keep retrying so we pretend it is a timeout
 			return false;
 		}
+
+		if (rc.error == 5)
+		{
+			m_logger.Warning(TC("A connection from a server with different uid was requested. Ignore"));
+			return false;
+		}
+
 		if (m_connectionCount.fetch_add(1) != 0)
 			return true;
 
diff --git a/Engine/Source/Programs/UnrealBuildAccelerator/Common/Private/UbaNetworkServer.cpp b/Engine/Source/Programs/UnrealBuildAccelerator/Common/Private/UbaNetworkServer.cpp
--- a/Engine/Source/Programs/UnrealBuildAccelerator/Common/Private/UbaNetworkServer.cpp
+++ b/Engine/Source/Programs/UnrealBuildAccelerator/Common/Private/UbaNetworkServer.cpp
@@ -119,10 +119,13 @@ namespace uba
 			return true;
 		}
 
-		bool SendAsync(u8 value)
+		bool SendInitialResponse(u8 value)
 		{
+			u8 data[32];
+			*data = value;
+			*(Guid*)(data+1) = m_server.m_uid;
 			NetworkBackend::SendContext context(NetworkBackend::SendFlags_Async);
-			return m_backend.Send(m_server.m_logger, m_backendConnection, &value, 1, context);
+			return m_backend.Send(m_server.m_logger, m_backendConnection, data, 1 + sizeof(Guid), context);
 		}
 
 		static bool ReceiveHandshakeHeader(void* context, u8* headerData, void*& outBodyContext, u8*& outBodyData, u32& outBodySize)
@@ -156,7 +159,7 @@ namespace uba
 			u32 clientVersion = *(u32*)headerData;
 			if (clientVersion != SystemNetworkVersion)
 			{
-				conn.SendAsync(1);
+				conn.SendInitialResponse(1);
 				return false;
 			}
 
@@ -180,7 +183,7 @@ namespace uba
 					found |= kv.second.uid == clientUid;
 				if (!found)
 				{
-					conn.SendAsync(3);
+					conn.SendInitialResponse(3);
 					return false;
 				}
 			}
@@ -188,7 +191,7 @@ namespace uba
 			constexpr u32 HeaderSize = 6;
 			conn.m_backend.SetRecvCallbacks(conn.m_backendConnection, &conn, HeaderSize, ReceiveMessageHeader, ReceiveMessageBody, TC("ReceiveMessage"));
 
-			if (!conn.SendAsync(0))
+			if (!conn.SendInitialResponse(0))
 				return false;
 			
 			ScopedWriteLock clientsLock(server.m_clientsLock);
@@ -407,14 +410,14 @@ namespace uba
 
 			if (!rec.func)
 			{
-				server.m_logger.Error(TC("WORKER FUNCTION NOT FOUND. id: %u, serviceid: %u type: %s"), m_id, m_serviceId, rec.toString(m_messageType));
+				server.m_logger.Error(TC("WORKER FUNCTION NOT FOUND. id: %u, serviceid: %u type: %s, client: %s"), m_id, m_serviceId, rec.toString(m_messageType), GuidToString(m_connection->m_client->uid).str);
 				m_connection->SetShouldDisconnect();
 				size = ErrorSize;
 			}
 			else if (!rec.func({m_connection}, m_messageType, reader, writer))
 			{
 				if (m_connection->SetShouldDisconnect())
-					server.m_logger.Error(TC("WORKER FUNCTION FAILED. id: %u, serviceid: %u type: %s"), m_id, m_serviceId, rec.toString(m_messageType));
+					server.m_logger.Error(TC("WORKER FUNCTION FAILED. id: %u, serviceid: %u type: %s, client: %s"), m_id, m_serviceId, rec.toString(m_messageType), GuidToString(m_connection->m_client->uid).str);
 				size = ErrorSize;
 			}
 			else
@@ -534,6 +537,9 @@ namespace uba
 			{
 				return HandleSystemMessage(connectionInfo, messageType, reader, writer);
 			};
+
+		if (!CreateGuid(m_uid))
+			outCtorSuccess = false;
 	}
 
 	NetworkServer::~NetworkServer()
diff --git a/Engine/Source/Programs/UnrealBuildAccelerator/Common/Public/UbaNetwork.h b/Engine/Source/Programs/UnrealBuildAccelerator/Common/Public/UbaNetwork.h
--- a/Engine/Source/Programs/UnrealBuildAccelerator/Common/Public/UbaNetwork.h
+++ b/Engine/Source/Programs/UnrealBuildAccelerator/Common/Public/UbaNetwork.h
@@ -16,7 +16,7 @@ namespace uba
 	static constexpr u8 StorageServiceId = 1;
 	static constexpr u8 SessionServiceId = 2;
 	
-	static constexpr u32 SystemNetworkVersion = 1338;
+	static constexpr u32 SystemNetworkVersion = 1339;
 	static constexpr u32 StorageNetworkVersion = 3;
 	static constexpr u32 SessionNetworkVersion = 27;
 
diff --git a/Engine/Source/Programs/UnrealBuildAccelerator/Common/Public/UbaNetworkClient.h b/Engine/Source/Programs/UnrealBuildAccelerator/Common/Public/UbaNetworkClient.h
--- a/Engine/Source/Programs/UnrealBuildAccelerator/Common/Public/UbaNetworkClient.h
+++ b/Engine/Source/Programs/UnrealBuildAccelerator/Common/Public/UbaNetworkClient.h
@@ -94,6 +94,9 @@ namespace uba
 		u32 m_sendSize;
 		u32 m_receiveTimeoutSeconds;
 
+		ReaderWriterLock m_serverUidLock;
+		Guid m_serverUid;
+
 		Event m_isConnected;
 		Event m_isOrWasConnected;
 		Atomic<u32> m_connectionCount;
diff --git a/Engine/Source/Programs/UnrealBuildAccelerator/Common/Public/UbaNetworkServer.h b/Engine/Source/Programs/UnrealBuildAccelerator/Common/Public/UbaNetworkServer.h
--- a/Engine/Source/Programs/UnrealBuildAccelerator/Common/Public/UbaNetworkServer.h
+++ b/Engine/Source/Programs/UnrealBuildAccelerator/Common/Public/UbaNetworkServer.h
@@ -106,7 +106,7 @@ namespace uba
 
 		NetworkBackend* m_listenBackend = nullptr;
 		CryptoKey m_listenCrypto = InvalidCryptoKey;
-
+		Guid m_uid;
 		bool m_allowNewClients = true;
 
 		struct WorkerRec { WorkerFunction func; TypeToNameFunction* toString = nullptr; };
diff --git a/Engine/Source/Programs/UnrealBuildAccelerator/Test/Private/UbaTestNetwork.h b/Engine/Source/Programs/UnrealBuildAccelerator/Test/Private/UbaTestNetwork.h
--- a/Engine/Source/Programs/UnrealBuildAccelerator/Test/Private/UbaTestNetwork.h
+++ b/Engine/Source/Programs/UnrealBuildAccelerator/Test/Private/UbaTestNetwork.h
@@ -115,6 +115,8 @@ namespace uba
 		Sleep(100);
 		if (!server.AddClient(serverTcp, TC("127.0.0.1"), 1239))
 			return logger.Error(TC("Server failed to connect"));
+		if (!server.AddClient(serverTcp, TC("127.0.0.1"), 1239))
+			return logger.Error(TC("Server failed to connect second"));
 
 		u64 time = GetTime();
 		while (!client.GetConnectionCount())
