[UBA]
* Added code to handling uba::ReadFile trying to read file with less bytes than requested. It will try for 3 seconds and then fail with a message saying it read 0 bytes for 3 seconds
* Added more information to error messages

diff --git a/Engine/Source/Programs/UnrealBuildAccelerator/Common/Private/UbaFile.cpp b/Engine/Source/Programs/UnrealBuildAccelerator/Common/Private/UbaFile.cpp
--- a/Engine/Source/Programs/UnrealBuildAccelerator/Common/Private/UbaFile.cpp
+++ b/Engine/Source/Programs/UnrealBuildAccelerator/Common/Private/UbaFile.cpp
@@ -87,7 +87,8 @@ namespace uba
 		ExtendedTimerScope ts(SystemStats::GetCurrent().readFile);
 		u8* buffer = (u8*)b;
 		u64 readLeft = bufferLen;
-		
+		u64 firstZeroReadTime = 0;
+
 		while (readLeft)
 		{
 			u32 toRead = u32(Min(readLeft, u64(~u32(0)) - 1));
@@ -104,6 +105,14 @@ namespace uba
 				return false;
 			}
 #endif
+			if (wasRead == 0)
+			{
+				if (firstZeroReadTime == 0)
+					firstZeroReadTime = GetTime();
+				else if (TimeToMs(GetTime() - firstZeroReadTime) > 3*1000)
+					return logger.Error(TC("ERROR reading file %s trying to read %u bytes from offset %llu but ReadFile returns 0 bytes read.. Is the file big enough?"), fileName, toRead, bufferLen - readLeft);
+			}
+
 			readLeft -= wasRead;
 			buffer += wasRead;
 		}
diff --git a/Engine/Source/Programs/UnrealBuildAccelerator/Common/Private/UbaFileAccessor.cpp b/Engine/Source/Programs/UnrealBuildAccelerator/Common/Private/UbaFileAccessor.cpp
--- a/Engine/Source/Programs/UnrealBuildAccelerator/Common/Private/UbaFileAccessor.cpp
+++ b/Engine/Source/Programs/UnrealBuildAccelerator/Common/Private/UbaFileAccessor.cpp
@@ -148,14 +148,27 @@ namespace uba
 			u8* pos = (u8*)data;
 			u64 i = 0;
 
+			auto WaitAndCheckError = [&](u64 index)
+				{
+					if (!ev[index].IsSet())
+						return m_logger.Error(L"Overlapped I/O WriteFile FAILED on waiting for event!");
+					u32 error = u32(ol[index].Internal);
+					if (error != ERROR_SUCCESS)
+						return m_logger.Error(L"Overlapped I/O WriteFile FAILED!: %s", LastErrorToText(error).data);
+					return true;
+				};
+
 			auto eg = MakeGuard([&]()
 				{
 					u64 index = i % BlockCount;
 					if (i > BlockCount)
 						for (u64 j = index; j != BlockCount; ++j)
-							ev[j].IsSet();
+							if (!WaitAndCheckError(j))
+								return false;
 					for (u64 j = 0; j != index; ++j)
-						ev[j].IsSet();
+						if (!WaitAndCheckError(j))
+							return false;
+					return true;
 				});
 
 			while (writeLeft)
@@ -165,7 +178,10 @@ namespace uba
 				if (i < BlockCount)
 					ev[i].Create(false);
 				else
-					ev[index].IsSet();
+				{
+					if (!WaitAndCheckError(index))
+						return false;
+				}
 
 				u64 toWrite = Min(writeLeft, BlockSize);
 				u64 toActuallyWrite = toWrite;
@@ -197,18 +213,21 @@ namespace uba
 				{
 					u32 lastError = GetLastError();
 					if (lastError != ERROR_IO_PENDING)
-						return m_logger.Error(L"FAILED!: %ls", LastErrorToText(lastError).data);
+						return m_logger.Error(L"FAILED!: %s", LastErrorToText(lastError).data);
 				}
 				++i;
 				pos += toWrite;
 				writeLeft -= toWrite;
 			}
 
-#if 0
-			eg.Execute();
+			if (!eg.Execute())
+				return false;
+
+			#if 0
 			if (setFileSize)
 				SetEndOfFile(logger, fileName, m_fileHandle, bufferLen);
-#endif
+			#endif
+
 			return true;
 		}
 #endif
diff --git a/Engine/Source/Programs/UnrealBuildAccelerator/Common/Private/UbaSessionServer.cpp b/Engine/Source/Programs/UnrealBuildAccelerator/Common/Private/UbaSessionServer.cpp
--- a/Engine/Source/Programs/UnrealBuildAccelerator/Common/Private/UbaSessionServer.cpp
+++ b/Engine/Source/Programs/UnrealBuildAccelerator/Common/Private/UbaSessionServer.cpp
@@ -1059,7 +1059,7 @@ namespace uba
 				RemoteProcess* process = (RemoteProcess*)h.m_process;
 				if (!process)
 				{
-					m_logger.Warning(TC("Client returned process that is not found on server (%s)"), reason.data);
+					m_logger.Warning(TC("Client %s returned process %u that is not found on server (%s)"), GuidToString(connectionInfo.GetUid()).str, processId, reason.data);
 					return true;
 				}
 
@@ -1067,7 +1067,7 @@ namespace uba
 				if (!m_activeRemoteProcesses.erase(process))
 				{
 					cs2.Leave();
-					m_logger.Warning(TC("Got returned process but process was not in active remote processes. Was there a disconnect happening directly after but executed before?"));
+					m_logger.Warning(TC("Got returned process %u from client %s but process was not in active remote processes. Was there a disconnect happening directly after but executed before?"), processId, GuidToString(connectionInfo.GetUid()).str);
 					return true;
 				}
 				u32 sessionIndex = process->m_sessionId - 1;
diff --git a/Engine/Source/Programs/UnrealBuildAccelerator/Common/Private/UbaStorageProxy.cpp b/Engine/Source/Programs/UnrealBuildAccelerator/Common/Private/UbaStorageProxy.cpp
--- a/Engine/Source/Programs/UnrealBuildAccelerator/Common/Private/UbaStorageProxy.cpp
+++ b/Engine/Source/Programs/UnrealBuildAccelerator/Common/Private/UbaStorageProxy.cpp
@@ -151,14 +151,14 @@ namespace uba
 						writer2.WriteBytes(reader.GetPositionData(), reader.GetLeft());
 
 						if (!msg.Send(reader2))
-							return false;
+							return m_logger.Error(TC("FetchBegin failed for cas file %s (%s). Requested by %s"), CasKeyString(casKey).str, hint.data, GuidToString(connectionInfo.GetUid()).str);
 
 						BinaryReader tempReader(reader2.GetPositionData(), 0, reader2.GetLeft());
 						u32 sizeOfFirstMessage = u32(reader2.GetLeft());
 						u16 fetchId = tempReader.ReadU16();
 						if (fetchId == 0)
 						{
-							m_logger.Error(TC("FetchBegin failed for cas file %s (%s)"), CasKeyString(casKey).str, hint.data);
+							m_logger.Error(TC("FetchBegin failed for cas file %s (%s). Requested by %s"), CasKeyString(casKey).str, hint.data, GuidToString(connectionInfo.GetUid()).str);
 							writer.WriteU16(0);
 							return true;
 						}
diff --git a/Engine/Source/Programs/UnrealBuildAccelerator/Common/Private/UbaStorageServer.cpp b/Engine/Source/Programs/UnrealBuildAccelerator/Common/Private/UbaStorageServer.cpp
--- a/Engine/Source/Programs/UnrealBuildAccelerator/Common/Private/UbaStorageServer.cpp
+++ b/Engine/Source/Programs/UnrealBuildAccelerator/Common/Private/UbaStorageServer.cpp
@@ -330,7 +330,7 @@ namespace uba
 								proxy.port = info.proxyPort;
 								proxy.zone = proxyName.data;
 
-								m_logger.Info(TC("%s:%u is assigned as proxy for zone %s"), proxy.host.c_str(), proxy.port, proxy.zone.c_str());
+								m_logger.Info(TC("%s:%u (%s) is assigned as proxy for zone %s"), proxy.host.c_str(), proxy.port, GuidToString(connectionInfo.GetUid()).str, proxy.zone.c_str());
 
 								writer.WriteBool(true);
 								writer.WriteU16(info.proxyPort);
@@ -383,7 +383,7 @@ namespace uba
 						m_logger.Info(TC("Server did not find cas for %s in file table lookup. Recalculating cas key"), hint.data);
 						if (!CalculateCasKey(checkedCasKey, hint.data))
 						{
-							m_logger.Error(TC("FetchBegin failed for cas file %s (%s). Can't calculate cas key for file"), CasKeyString(casKey).str, hint.data);
+							m_logger.Error(TC("FetchBegin failed for cas file %s (%s) requested by %s. Can't calculate cas key for file"), CasKeyString(casKey).str, hint.data, GuidToString(connectionInfo.GetUid()).str);
 							writer.WriteU16(0);
 							return false;
 						}
