[UBA]
* Fixed code to handle extremely long error log entries

diff --git a/Engine/Source/Programs/UnrealBuildAccelerator/Common/Private/UbaTrace.cpp b/Engine/Source/Programs/UnrealBuildAccelerator/Common/Private/UbaTrace.cpp
--- a/Engine/Source/Programs/UnrealBuildAccelerator/Common/Private/UbaTrace.cpp
+++ b/Engine/Source/Programs/UnrealBuildAccelerator/Common/Private/UbaTrace.cpp
@@ -202,9 +202,17 @@ namespace uba
 		writer.WriteU32(exitCode);
 		writer.WriteBytes(data, dataSize);
 		writer.Write7BitEncoded(u64(logLines.size()));
+		u32 lineCounter = 0;
 		for (auto& line : logLines)
 		{
-			writer.EnsureMemory(line.text.size()*sizeof(tchar));
+			if (lineCounter++ == 100) // We don't want to write the entire error in the trace stream to blow the entire buffer
+			{
+				writer.EnsureMemory(100);
+				writer.WriteByte(LogEntryType_Info);
+				writer.WriteString(TC("Error is cut-off. Look in normal log to see full error"));
+				break;
+			}
+			writer.EnsureMemory(1 + (line.text.size()+1)*sizeof(tchar));
 			writer.WriteByte(line.type);
 			writer.WriteString(line.text);
 		}
