[UBA]
* Fixed crash in UbaVisualizer when jumping back and forth in traces hitting edge cases with multiple sessions. Code was not cleaning up after itself properly

diff --git a/Engine/Source/Programs/UnrealBuildAccelerator/Visualizer/Private/UbaTraceReader.cpp b/Engine/Source/Programs/UnrealBuildAccelerator/Visualizer/Private/UbaTraceReader.cpp
--- a/Engine/Source/Programs/UnrealBuildAccelerator/Visualizer/Private/UbaTraceReader.cpp
+++ b/Engine/Source/Programs/UnrealBuildAccelerator/Visualizer/Private/UbaTraceReader.cpp
@@ -33,6 +33,8 @@ namespace uba
 
 	bool TraceReader::ReadFile(TraceView& out, const tchar* fileName, bool replay)
 	{
+		Reset();
+
 		FileHandle readHandle;
 		if (!OpenFileSequentialRead(m_logger, fileName, readHandle))
 			return false;
@@ -119,6 +121,8 @@ namespace uba
 
 	bool TraceReader::StartReadClient(TraceView& out, NetworkClient& client)
 	{
+		Reset();
+
 		u32 traceMemSize = 128 * 1024 * 1024;
 		m_memoryHandle = uba::CreateMemoryMappingW(m_logger, PAGE_READWRITE, traceMemSize);
 		if (!m_memoryHandle.IsValid())
@@ -177,6 +181,8 @@ namespace uba
 
 	bool TraceReader::StartReadNamed(TraceView& out, const tchar* namedTrace, bool silentFail)
 	{
+		Reset();
+
 		m_memoryHandle.handle = ::OpenFileMappingW(PAGE_READWRITE, false, namedTrace);
 		if (!m_memoryHandle.IsValid())
 		{
@@ -897,6 +903,13 @@ namespace uba
 		out.finished = true;
 	}
 
+	void TraceReader::Reset()
+	{
+		m_activeProcesses.clear();
+		m_activeWorkRecords.clear();
+		m_sessionIndexToSession.clear();
+	}
+
 	bool TraceReader::SaveAs(const tchar* fileName)
 	{
 		if (!m_memoryBegin)
diff --git a/Engine/Source/Programs/UnrealBuildAccelerator/Visualizer/Private/UbaTraceReader.h b/Engine/Source/Programs/UnrealBuildAccelerator/Visualizer/Private/UbaTraceReader.h
--- a/Engine/Source/Programs/UnrealBuildAccelerator/Visualizer/Private/UbaTraceReader.h
+++ b/Engine/Source/Programs/UnrealBuildAccelerator/Visualizer/Private/UbaTraceReader.h
@@ -157,6 +157,7 @@ namespace uba
 		bool ReadMemory(TraceView& out, bool trackHost);
 		bool ReadTrace(TraceView& out, BinaryReader& reader, u64 maxTime);
 		void StopAllActive(TraceView& out, u64 stopTime);
+		void Reset();
 
 		bool SaveAs(const tchar* fileName);
 
