[UBA]
* Reverting some allocations to use new[] on windows. There has been weird hangs on the farm since last binaries and this is one of the changes that I can imagine can make a difference.

diff --git a/Engine/Source/Programs/UnrealBuildAccelerator/Common/Private/UbaStorage.cpp b/Engine/Source/Programs/UnrealBuildAccelerator/Common/Private/UbaStorage.cpp
--- a/Engine/Source/Programs/UnrealBuildAccelerator/Common/Private/UbaStorage.cpp
+++ b/Engine/Source/Programs/UnrealBuildAccelerator/Common/Private/UbaStorage.cpp
@@ -341,17 +341,23 @@ namespace uba
 
 				// For some very unknown reason ASAN on linux triggers on the "new[]" call when doing delete[]
 				// while doing it manually works properly. Will stop investigating this and move on
-				//events = new Event[workCount];
+				#if PLATFORM_WINDOWS
+				events = new Event[workCount];
+				#else
 				events = (Event*)aligned_alloc(alignof(Event), sizeof(Event)*workCount);
 				for (auto i=0;i!=workCount; ++i)
 					new (events + i) Event();
+				#endif
 			}
 			~WorkRec()
 			{
-				//delete[] events;
+				#if PLATFORM_WINDOWS
+				delete[] events;
+				#else
 				for (auto i=0;i!=workCount; ++i)
 					events[i].~Event();
 				aligned_free(events);
+				#endif
 			}
 			Atomic<u64> refCount;
 			Atomic<u64> compressCounter;
