[UBA]
* Fixed allocation to make sure it honors alignments
* Added error handling in trace code. It will stop writing when running out of shared mem

diff --git a/Engine/Source/Programs/UnrealBuildAccelerator/Agent/UbaAgent.Build.cs b/Engine/Source/Programs/UnrealBuildAccelerator/Agent/UbaAgent.Build.cs
--- a/Engine/Source/Programs/UnrealBuildAccelerator/Agent/UbaAgent.Build.cs
+++ b/Engine/Source/Programs/UnrealBuildAccelerator/Agent/UbaAgent.Build.cs
@@ -7,6 +7,8 @@ public class UbaAgent : ModuleRules
 {
 	public UbaAgent(ReadOnlyTargetRules Target) : base(Target)
 	{
+		PrivatePCHHeaderFile = "../Core/Public/UbaCorePch.h";
+
 		UnsafeTypeCastWarningLevel = WarningLevel.Error;
 		StaticAnalyzerDisabledCheckers.Clear();
 
diff --git a/Engine/Source/Programs/UnrealBuildAccelerator/Common/Private/UbaStorage.cpp b/Engine/Source/Programs/UnrealBuildAccelerator/Common/Private/UbaStorage.cpp
--- a/Engine/Source/Programs/UnrealBuildAccelerator/Common/Private/UbaStorage.cpp
+++ b/Engine/Source/Programs/UnrealBuildAccelerator/Common/Private/UbaStorage.cpp
@@ -7,13 +7,6 @@
 #include "UbaDirectoryIterator.h"
 #include "UbaWorkManager.h"
 
-#if PLATFORM_WINDOWS && !defined(aligned_alloc)
-#define aligned_alloc(a, s) _aligned_malloc(s, a)
-#define aligned_free(p) _aligned_free(p)
-#else
-#define aligned_free(p) free(p)
-#endif
-
 namespace uba
 {
 	constexpr u32 CasTableVersion = IsWindows ? 32 : 34;
diff --git a/Engine/Source/Programs/UnrealBuildAccelerator/Common/Private/UbaTrace.cpp b/Engine/Source/Programs/UnrealBuildAccelerator/Common/Private/UbaTrace.cpp
--- a/Engine/Source/Programs/UnrealBuildAccelerator/Common/Private/UbaTrace.cpp
+++ b/Engine/Source/Programs/UnrealBuildAccelerator/Common/Private/UbaTrace.cpp
@@ -26,7 +26,7 @@ namespace uba
 	{
 		WriterScope(Trace& trace) : ScopedWriteLock(trace.m_memoryLock), BinaryWriter(trace.m_memoryBegin, trace.m_memoryPos, trace.m_memoryCapacity), m_trace(trace)
 		{
-			EnsureMemory(TraceMessageMaxSize);
+			isValid = EnsureMemory(TraceMessageMaxSize);
 		}
 
 		~WriterScope()
@@ -35,21 +35,22 @@ namespace uba
 			*(u32*)m_trace.m_memoryBegin = u32(m_trace.m_memoryPos);
 		}
 
-		void EnsureMemory(u64 size)
+		WriterScope(const WriterScope&) = delete;
+		void operator=(const WriterScope&) = delete;
+
+		bool EnsureMemory(u64 size)
 		{
 			u64 committedMemoryNeeded = AlignUp(m_trace.m_memoryPos + size, size);
 			if (m_trace.m_memoryCommitted >= committedMemoryNeeded)
-				return;
+				return true;
 			if (!MapViewCommit(m_trace.m_memoryBegin + m_trace.m_memoryCommitted, committedMemoryNeeded - m_trace.m_memoryCommitted))
-			{
-				m_trace.m_logger.Error(TC("Failed to commit memory for trace (Pos: %llu Capacity: %llu, Already Committed: %llu, Needed: %llu): %s"), m_trace.m_memoryPos, m_trace.m_memoryCapacity, m_trace.m_memoryCommitted, committedMemoryNeeded, LastErrorToText().data);
-				return;
-			}
-
+				return m_trace.m_logger.Error(TC("Failed to commit memory for trace (Pos: %llu Capacity: %llu, Already Committed: %llu, Needed: %llu): %s"), m_trace.m_memoryPos, m_trace.m_memoryCapacity, m_trace.m_memoryCommitted, committedMemoryNeeded, LastErrorToText().data);
 			m_trace.m_memoryCommitted = committedMemoryNeeded;
+			return true;
 		}
 
 		Trace& m_trace;
+		bool isValid;
 	};
 
 	bool Trace::StartWrite(const tchar* namedTrace, u64 traceMemCapacity)
@@ -74,6 +75,8 @@ namespace uba
 
 		{
 			WriterScope writer(*this);
+			if (!writer.isValid)
+				return false;
 			writer.AllocWrite(4);
 			writer.WriteU32(TraceVersion);
 			writer.WriteU32(GetCurrentProcessId());
@@ -96,6 +99,8 @@ namespace uba
 
 		{
 			WriterScope writer(*this);
+			if (!writer.isValid)
+				return false;
 			writer.WriteByte(TraceType_Summary);
 			writer.Write7BitEncoded(GetTime() - m_startTime);
 		}
@@ -131,6 +136,8 @@ namespace uba
 		{
 			insres.first->second = u32(m_strings.size() - 1);
 			WriterScope writer(*this);
+			if (!writer.isValid)
+				return 0;
 			writer.WriteByte(TraceType_String);
 			writer.WriteString(string, stringLen);
 		}
@@ -141,6 +148,8 @@ namespace uba
 		if (!m_memoryBegin) \
 			return; \
 		WriterScope writer(*this); \
+		if (!writer.isValid) \
+			return; \
 		writer.WriteByte(x); \
 		writer.Write7BitEncoded(GetTime() - m_startTime);
 
@@ -207,12 +216,14 @@ namespace uba
 		{
 			if (lineCounter++ == 100) // We don't want to write the entire error in the trace stream to blow the entire buffer
 			{
-				writer.EnsureMemory(100);
+				if (!writer.EnsureMemory(100))
+					return;
 				writer.WriteByte(LogEntryType_Info);
 				writer.WriteString(TC("Error is cut-off. Look in normal log to see full error"));
 				break;
 			}
-			writer.EnsureMemory(1 + (line.text.size()+1)*sizeof(tchar));
+			if (!writer.EnsureMemory(1 + (line.text.size()+1)*sizeof(tchar)))
+				return;
 			writer.WriteByte(line.type);
 			writer.WriteString(line.text);
 		}
diff --git a/Engine/Source/Programs/UnrealBuildAccelerator/Common/UbaCommon.Build.cs b/Engine/Source/Programs/UnrealBuildAccelerator/Common/UbaCommon.Build.cs
--- a/Engine/Source/Programs/UnrealBuildAccelerator/Common/UbaCommon.Build.cs
+++ b/Engine/Source/Programs/UnrealBuildAccelerator/Common/UbaCommon.Build.cs
@@ -7,6 +7,8 @@ public class UbaCommon : ModuleRules
 {
 	public UbaCommon(ReadOnlyTargetRules Target) : base(Target)
 	{
+		PrivatePCHHeaderFile = "../Core/Public/UbaCorePch.h";
+
 		UnsafeTypeCastWarningLevel = WarningLevel.Error;
 		StaticAnalyzerDisabledCheckers.Clear();
 
diff --git a/Engine/Source/Programs/UnrealBuildAccelerator/Core/Private/UbaMemory.cpp b/Engine/Source/Programs/UnrealBuildAccelerator/Core/Private/UbaMemory.cpp
--- a/Engine/Source/Programs/UnrealBuildAccelerator/Core/Private/UbaMemory.cpp
+++ b/Engine/Source/Programs/UnrealBuildAccelerator/Core/Private/UbaMemory.cpp
@@ -58,7 +58,7 @@ namespace uba
 	void* MemoryBlock::Allocate(u64 bytes, u64 alignment, const tchar* hint)
 	{
 		if (!memory)
-			return malloc(bytes); // TODO: alignment?
+			return aligned_alloc(alignment, bytes);
 
 		ScopedWriteLock l(lock);
 		return AllocateNoLock(bytes, alignment, hint);
@@ -91,7 +91,7 @@ namespace uba
 	void MemoryBlock::Free(void* p)
 	{
 		if (!memory)
-			free(p);
+			aligned_free(p);
 	}
 
 	tchar* MemoryBlock::Strdup(const tchar* str)
diff --git a/Engine/Source/Programs/UnrealBuildAccelerator/Core/Public/UbaMemory.h b/Engine/Source/Programs/UnrealBuildAccelerator/Core/Public/UbaMemory.h
--- a/Engine/Source/Programs/UnrealBuildAccelerator/Core/Public/UbaMemory.h
+++ b/Engine/Source/Programs/UnrealBuildAccelerator/Core/Public/UbaMemory.h
@@ -16,6 +16,13 @@
 #include <mimalloc-override.h>
 #endif
 
+#if PLATFORM_WINDOWS && !defined(aligned_alloc)
+#define aligned_alloc(a, s) _aligned_malloc(s, a)
+#define aligned_free(p) _aligned_free(p)
+#else
+#define aligned_free(p) free(p)
+#endif
+
 namespace uba
 {
 	template<typename T>
diff --git a/Engine/Source/Programs/UnrealBuildAccelerator/Host/UbaHost.Build.cs b/Engine/Source/Programs/UnrealBuildAccelerator/Host/UbaHost.Build.cs
--- a/Engine/Source/Programs/UnrealBuildAccelerator/Host/UbaHost.Build.cs
+++ b/Engine/Source/Programs/UnrealBuildAccelerator/Host/UbaHost.Build.cs
@@ -7,6 +7,8 @@ public class UbaHost : ModuleRules
 {
 	public UbaHost(ReadOnlyTargetRules Target) : base(Target)
 	{
+		PrivatePCHHeaderFile = "../Core/Public/UbaCorePch.h";
+
 		UnsafeTypeCastWarningLevel = WarningLevel.Error;
 		StaticAnalyzerDisabledCheckers.Clear();
 
diff --git a/Engine/Source/Programs/UnrealBuildAccelerator/StorageProxy/UbaStorageProxy.Build.cs b/Engine/Source/Programs/UnrealBuildAccelerator/StorageProxy/UbaStorageProxy.Build.cs
--- a/Engine/Source/Programs/UnrealBuildAccelerator/StorageProxy/UbaStorageProxy.Build.cs
+++ b/Engine/Source/Programs/UnrealBuildAccelerator/StorageProxy/UbaStorageProxy.Build.cs
@@ -7,6 +7,8 @@ public class UbaStorageProxy : ModuleRules
 {
 	public UbaStorageProxy(ReadOnlyTargetRules Target) : base(Target)
 	{
+		PrivatePCHHeaderFile = "../Core/Public/UbaCorePch.h";
+
 		UnsafeTypeCastWarningLevel = WarningLevel.Error;
 		StaticAnalyzerDisabledCheckers.Clear();
 
diff --git a/Engine/Source/Programs/UnrealBuildAccelerator/Test/UbaTest.Build.cs b/Engine/Source/Programs/UnrealBuildAccelerator/Test/UbaTest.Build.cs
--- a/Engine/Source/Programs/UnrealBuildAccelerator/Test/UbaTest.Build.cs
+++ b/Engine/Source/Programs/UnrealBuildAccelerator/Test/UbaTest.Build.cs
@@ -7,6 +7,8 @@ public class UbaTest : ModuleRules
 {
 	public UbaTest(ReadOnlyTargetRules Target) : base(Target)
 	{
+		PrivatePCHHeaderFile = "../Core/Public/UbaCorePch.h";
+
 		UnsafeTypeCastWarningLevel = WarningLevel.Error;
 
 		PrivateDependencyModuleNames.AddRange(new string[] {
diff --git a/Engine/Source/Programs/UnrealBuildAccelerator/Visualizer/UbaVisualizer.Build.cs b/Engine/Source/Programs/UnrealBuildAccelerator/Visualizer/UbaVisualizer.Build.cs
--- a/Engine/Source/Programs/UnrealBuildAccelerator/Visualizer/UbaVisualizer.Build.cs
+++ b/Engine/Source/Programs/UnrealBuildAccelerator/Visualizer/UbaVisualizer.Build.cs
@@ -7,6 +7,8 @@ public class UbaVisualizer : ModuleRules
 {
 	public UbaVisualizer(ReadOnlyTargetRules Target) : base(Target)
 	{
+		PrivatePCHHeaderFile = "../Core/Public/UbaCorePch.h";
+
 		UnsafeTypeCastWarningLevel = WarningLevel.Error;
 		StaticAnalyzerDisabledCheckers.Clear();
 
