[UBA]
* Removed force drop of cas file for remote log file since it caused problems if two log files were the same

diff --git a/Engine/Source/Programs/UnrealBuildAccelerator/Common/Private/UbaSessionServer.cpp b/Engine/Source/Programs/UnrealBuildAccelerator/Common/Private/UbaSessionServer.cpp
--- a/Engine/Source/Programs/UnrealBuildAccelerator/Common/Private/UbaSessionServer.cpp
+++ b/Engine/Source/Programs/UnrealBuildAccelerator/Common/Private/UbaSessionServer.cpp
@@ -775,7 +775,7 @@ namespace uba
 						StringBuffer<> logPath;
 						logPath.Append(m_sessionLogDir).Append(destination.data + 5);
 						m_storage.CopyOrLink(casKey, logPath.data, attributes);
-						m_storage.DropCasFile(casKey, true);
+						m_storage.DropCasFile(casKey, false);
 						writer.WriteBool(true);
 						return true;
 					}
