[UBA]
* Removed includes if mimalloc-override.h to test if that is the cause of hangs we've seen on the farm (long shot)

diff --git a/Engine/Source/Programs/UnrealBuildAccelerator/Core/Public/UbaCorePch.h b/Engine/Source/Programs/UnrealBuildAccelerator/Core/Public/UbaCorePch.h
--- a/Engine/Source/Programs/UnrealBuildAccelerator/Core/Public/UbaCorePch.h
+++ b/Engine/Source/Programs/UnrealBuildAccelerator/Core/Public/UbaCorePch.h
@@ -3,7 +3,8 @@
 #pragma once
 
 #if UBA_USE_MIMALLOC
-#include <mimalloc-override.h>
+#include <mimalloc.h>
+//#include <mimalloc-override.h>
 #endif
 
 #include "UbaPlatform.h"
diff --git a/Engine/Source/Programs/UnrealBuildAccelerator/Core/Public/UbaMemory.h b/Engine/Source/Programs/UnrealBuildAccelerator/Core/Public/UbaMemory.h
--- a/Engine/Source/Programs/UnrealBuildAccelerator/Core/Public/UbaMemory.h
+++ b/Engine/Source/Programs/UnrealBuildAccelerator/Core/Public/UbaMemory.h
@@ -13,7 +13,8 @@
 #include <functional>
 
 #if UBA_USE_MIMALLOC
-#include <mimalloc-override.h>
+#include <mimalloc.h>
+//#include <mimalloc-override.h>
 #endif
 
 #if PLATFORM_WINDOWS && !defined(aligned_alloc)
