[UBA]
* Fixed so server side code sending cas to client does not have to send messages in order. If file is larger than 16mb we use file mapping instead with random access.

diff --git a/Engine/Source/Programs/UnrealBuildAccelerator/Common/Private/UbaStorageServer.cpp b/Engine/Source/Programs/UnrealBuildAccelerator/Common/Private/UbaStorageServer.cpp
--- a/Engine/Source/Programs/UnrealBuildAccelerator/Common/Private/UbaStorageServer.cpp
+++ b/Engine/Source/Programs/UnrealBuildAccelerator/Common/Private/UbaStorageServer.cpp
@@ -100,6 +100,23 @@ namespace uba
 		return waitEntry.Success;
 	}
 
+	void StorageServer::ActiveFetch::Release(StorageServer& server, const tchar* reason)
+	{
+		if (mappedView.handle.IsValid())
+		{
+			if (ownsMapping)
+			{
+				UnmapViewOfFile(memoryBegin, mappedView.size, TC(""));
+				CloseFileMapping(mappedView.handle);
+				CloseFile(nullptr, readFileHandle);
+			}
+			else
+				server.m_casDataBuffer.UnmapView(mappedView, TC("OnDisconnected"));
+		}
+		else
+			server.PushBufferSlot(memoryBegin);
+	}
+
 	void StorageServer::OnDisconnected(u32 clientId)
 	{
 		{
@@ -147,16 +164,8 @@ namespace uba
 					++it;
 					continue;
 				}
-				if (fetch.mappedView.handle.IsValid())
-					m_casDataBuffer.UnmapView(fetch.mappedView, TC("OnDisconnected"));
-				else
-					PushBufferSlot(fetch.memoryBegin);
 
-				{
-					ScopedWriteLock lock1(fetch.readIndexLock);
-					for (ActiveFetchItem* i = fetch.firstItem; i; i = i->next)
-						i->event.Set();
-				}
+				fetch.Release(*this, TC("OnDisconnected"));
 
 				if (m_traceFetch)
 					m_trace->FileEndFetch(clientId, AsCompressed(fetch.casKey, m_storeCompressed));
@@ -415,11 +424,13 @@ namespace uba
 					return true;
 				}
 				
+				StringBuffer<512> casFile;
 				FileHandle readFileHandle = InvalidFileHandle;
 				auto rfg = MakeGuard([&]() { CloseFile(nullptr, readFileHandle); });
 				u64 fileSize;
 				u8* memoryBegin = nullptr;
 				u8* memoryPos = nullptr;
+				bool ownsMapping = false;
 
 				MappedView mappedView;
 				auto mvg = MakeGuard([&](){ m_casDataBuffer.UnmapView(mappedView, TC("FetchBegin")); });
@@ -437,7 +448,6 @@ namespace uba
 				else
 				{
 	#if !UBA_USE_SPARSEFILE
-					StringBuffer<> casFile;
 					GetCasFileName(casFile, casKey);
 					if (!OpenFileSequentialRead(m_logger, casFile.data, readFileHandle))
 					{
@@ -450,6 +460,20 @@ namespace uba
 	#else
 					UBA_ASSERT(false);
 	#endif
+					if (fileSize > BufferSlotSize)
+					{
+						mappedView.handle = CreateFileMappingW(readFileHandle, PAGE_READONLY, fileSize, TC(""));
+						if (!mappedView.handle.IsValid())
+							return m_logger.Error(TC("Failed to create file mapping of %s (%s)"), casFile.data, LastErrorToText().data);
+						u64 offset = memoryPos - memoryBegin;
+						mappedView.memory = MapViewOfFile(mappedView.handle, FILE_MAP_READ, 0, fileSize);
+						if (!mappedView.memory)
+							return m_logger.Error(TC("Failed to map memory of %s (%s)"), casFile.data, LastErrorToText().data);
+						memoryBegin = mappedView.memory;
+						memoryPos = memoryBegin + offset;
+						ownsMapping = true;
+						useFileMapping = true;
+					}
 				}
 
 				if (m_trace)
@@ -477,7 +501,7 @@ namespace uba
 				}
 				else if (toWrite == left)
 				{
-					if (!ReadFile(m_logger, CasKeyString(casKey).str, readFileHandle, writeBuffer, toWrite))
+					if (!ReadFile(m_logger, casFile.data, readFileHandle, writeBuffer, toWrite))
 					{
 						UBA_ASSERT(false); // Implement
 						return false;
@@ -488,13 +512,16 @@ namespace uba
 					memoryBegin = PopBufferSlot();
 					memoryPos = memoryBegin;
 					u32 toRead = u32(Min(left, BufferSlotSize));
- 					if (!ReadFile(m_logger, CasKeyString(casKey).str, readFileHandle, memoryBegin, toRead))
+ 					if (!ReadFile(m_logger, casFile.data, readFileHandle, memoryBegin, toRead))
 					{
 						UBA_ASSERT(false); // Implement
 						return false;
 					}
 					memcpy(writeBuffer, memoryPos, toWrite);
 					memoryPos += toWrite;
+
+					CloseFile(casFile.data, readFileHandle);
+					readFileHandle = InvalidFileHandle;
 				}
 
 				u64 actualSize = fileSize;
@@ -524,15 +551,21 @@ namespace uba
 				ScopedWriteLock lock(m_activeFetchesLock);
 				auto insres = m_activeFetches.try_emplace(*fetchId);
 				UBA_ASSERT(insres.second);
+				lock.Leave();
+
+				mappedView.size = fileSize;
+
 				ActiveFetch& fetch = insres.first->second;
 				fetch.clientId = connectionInfo.GetId();
 				fetch.readFileHandle = readFileHandle;
 				fetch.mappedView = mappedView;
+				fetch.ownsMapping = ownsMapping;
 				fetch.memoryBegin = memoryBegin;
 				fetch.memoryPos = memoryPos;
 				fetch.left = left;
 				fetch.casKey = casKey;
 				fetch.sendCasTime = GetTime() - start;
+
 				return true;
 			}
 			case StorageMessageType_FetchSegment:
@@ -544,112 +577,26 @@ namespace uba
 				ScopedReadLock lock(m_activeFetchesLock);
 				auto findIt = m_activeFetches.find(fetchId);
 				if (findIt == m_activeFetches.end())
-					return m_logger.Error(TC("Can't find active fetch %u, disconnected client?"), fetchId);
-
+					return m_logger.Error(TC("Can't find active fetch %u, disconnected client? (index %u)"), fetchId, fetchIndex);
 				ActiveFetch& fetch = findIt->second;
 				lock.Leave();
 
-				if (fetchIndex != 0)
-				{
-					ScopedWriteLock lock1(fetch.readIndexLock);
-					if (fetch.readIndexHandled != fetchIndex - 1)
-					{
-						ActiveFetchItem item;
-						item.next = fetch.firstItem;
-						if (item.next)
-							item.next->prev = &item;
-						item.index = fetchIndex;
-						fetch.firstItem = &item;
-						item.event.Create(true);
-						lock1.Leave();
-
-						bool shouldExit = false;
-						while (!item.event.IsSet(5000))
-						{
-							// There is a rare chance that connection got disconnected while we have received message with later index but not earlier (can happen when client has multiple tcp connections)
-							// In that case this will lock up forever unless we check the connection to see if we are still connected
-							if (!connectionInfo.ShouldDisconnect())
-								continue;
-							shouldExit = true;
-							break;
-						}
-
-						ScopedWriteLock lock2(fetch.readIndexLock);
-						if (item.prev)
-							item.prev->next = item.next;
-						else
-							fetch.firstItem = item.next;
-						if (item.next)
-							item.next->prev = item.prev;
-
-						if (shouldExit)
-							return false;
-					}
-				}
-
-				u32 toWrite = u32(Min(fetch.left, writer.GetCapacityLeft()));
-				if (fetch.mappedView.handle.IsValid())
-				{
-					memcpy(writer.AllocWrite(toWrite), fetch.memoryPos, toWrite);
-					fetch.memoryPos += toWrite;
-					fetch.left -= toWrite;
-				}
-				else
-				{
-					u8* writeBuf = writer.AllocWrite(toWrite);
-					u64 offset = u64(fetch.memoryPos - fetch.memoryBegin);
-					u64 leftInBlock = BufferSlotSize - offset;
-
-					if (leftInBlock < toWrite)
-					{
-						memcpy(writeBuf, fetch.memoryPos, leftInBlock);
-						writeBuf += leftInBlock;
-						fetch.left -= leftInBlock;
-						toWrite -= u32(leftInBlock);
-
-						u64 toRead = Min(fetch.left, BufferSlotSize);
-						if (!ReadFile(m_logger, CasKeyString(fetch.casKey).str, fetch.readFileHandle, fetch.memoryBegin, toRead))
-						{
-							UBA_ASSERT(false); // Implement
-							return false;
-						}
-						fetch.memoryPos = fetch.memoryBegin;
-					}
+				UBA_ASSERT(fetchIndex);
+				const u8* pos = fetch.memoryPos + (fetchIndex-1) * writer.GetCapacityLeft();
+				u64 toWrite = writer.GetCapacityLeft();
+				if ((pos - fetch.memoryBegin) + toWrite > fetch.mappedView.size)
+					toWrite = fetch.mappedView.size - (pos - fetch.memoryBegin);
 
-					memcpy(writeBuf, fetch.memoryPos, toWrite);
-					fetch.memoryPos += toWrite;
-
-					fetch.left -= toWrite;
-				}
-
-				bool isDone = fetch.left == 0;
-
-				if (fetchIndex != 0)
-				{
-					ScopedWriteLock lock1(fetch.readIndexLock);
-					fetch.readIndexHandled = fetchIndex;
-					for (ActiveFetchItem* i = fetch.firstItem; i; i = i->next)
-					{
-						if (i->index != fetchIndex + 1)
-							continue;
-						i->event.Set();
-						break;
-					}
-				}
+				memcpy(writer.AllocWrite(toWrite), pos, toWrite);
 
+				bool isDone = fetch.left.fetch_sub(toWrite) == toWrite;
 				if (!isDone)
 				{
 					fetch.sendCasTime += GetTime() - start;
 					return true;
 				}
 
-				if (fetch.mappedView.handle.IsValid())
-					m_casDataBuffer.UnmapView(fetch.mappedView, TC("FetchDone"));
-				else
-				{
-					PushBufferSlot(fetch.memoryBegin);
-					CloseFile(nullptr, fetch.readFileHandle);
-				}
+				fetch.Release(*this, TC("FetchDone"));
 
 				u64 sendCasTime = fetch.sendCasTime;
 				ScopedWriteLock activeLock(m_activeFetchesLock);
diff --git a/Engine/Source/Programs/UnrealBuildAccelerator/Common/Public/UbaStorageServer.h b/Engine/Source/Programs/UnrealBuildAccelerator/Common/Public/UbaStorageServer.h
--- a/Engine/Source/Programs/UnrealBuildAccelerator/Common/Public/UbaStorageServer.h
+++ b/Engine/Source/Programs/UnrealBuildAccelerator/Common/Public/UbaStorageServer.h
@@ -71,29 +71,20 @@ namespace uba
 		UnorderedMap<u16, ActiveStore> m_activeStores;
 
 
-		struct ActiveFetchItem
-		{
-			ActiveFetchItem* next = nullptr;
-			u32 index = 0;
-			Event event;
-			ActiveFetchItem* prev = nullptr;
-		};
-
 		struct ActiveFetch
 		{
 			u32 clientId = ~0u;
-			u64 left = 0;
-			FileHandle readFileHandle = InvalidFileHandle;
+			Atomic<u64> left;
 			CasKey casKey;
 			Atomic<u64> sendCasTime;
 
+			FileHandle readFileHandle = InvalidFileHandle;
+			MappedView mappedView;
 			u8* memoryBegin = nullptr;
 			u8* memoryPos = nullptr;
-			MappedView mappedView;
+			bool ownsMapping = false;
 
-			ReaderWriterLock readIndexLock;
-			ActiveFetchItem* firstItem = nullptr;
-			u32 readIndexHandled = 0;
+			void Release(StorageServer& server, const tchar* reason);
 		};
 		ReaderWriterLock m_activeFetchesLock;
 		UnorderedMap<u16, ActiveFetch> m_activeFetches;
