[UBA]
* Added support for workers being able to push current work and pick up new work

diff --git a/Engine/Source/Programs/UnrealBuildAccelerator/Common/Private/UbaNetworkServer.cpp b/Engine/Source/Programs/UnrealBuildAccelerator/Common/Private/UbaNetworkServer.cpp
--- a/Engine/Source/Programs/UnrealBuildAccelerator/Common/Private/UbaNetworkServer.cpp
+++ b/Engine/Source/Programs/UnrealBuildAccelerator/Common/Private/UbaNetworkServer.cpp
@@ -7,41 +7,82 @@
 
 namespace uba
 {
+	struct NetworkServer::WorkerContext
+	{
+		WorkerContext(NetworkServer& s) : server(s), workAvailable(false)
+		{
+			writeMemSize = server.m_sendSize;
+			writeMem = new u8[writeMemSize];
+		}
+
+		~WorkerContext()
+		{
+			delete[] writeMem;
+		}
+
+		NetworkServer& server;
+		Event workAvailable;
+		u8* writeMem = nullptr;
+		u32 writeMemSize = 0;
+
+		Vector<u8> buffer;
+		Connection* connection = nullptr;
+		u32 dataSize = 0;
+		u8 serviceId = 0;
+		u8 messageType = 0;
+		u16 id = 0;
+	};
+
 	class NetworkServer::Worker
 	{
 	public:
-		Worker() : m_workAvailable(false) {}
+		Worker() {}
 		~Worker()
 		{
+			UBA_ASSERT(!m_inUse);
+			m_context->connection = nullptr;
 			m_loop = false;
-			m_workAvailable.Set();
+			m_context->workAvailable.Set();
+			m_thread.Wait();
+			delete m_context;
+			m_context = nullptr;
 		}
 
 		void Start(NetworkServer& server)
 		{
+			m_context = new WorkerContext(server);
 			m_loop = true;
 			m_thread.Start([&]() { ThreadWorker(server); return 0; });
 		}
 
+		void Stop(NetworkServer& server)
+		{
+			m_loop = false;
+			ScopedWriteLock lock(server.m_availableWorkersLock);
+			while (m_inUse)
+			{
+				lock.Leave();
+				Sleep(1);
+				lock.Enter();
+			}
+		}
+
 		void ThreadWorker(NetworkServer& server);
+		void Update(WorkerContext& context, bool signalAvailable);
 		void DoAdditionalWorkAndSignalAvailable(NetworkServer& server);
 
 		Worker* m_nextWorker = nullptr;
 		Worker* m_prevWorker = nullptr;
 
-		Vector<u8> m_buffer;
-		Event m_workAvailable;
-		Connection* m_connection = nullptr;
-		u32 m_dataSize = 0;
-		u8 m_serviceId = 0;
-		u8 m_messageType = 0;
-		u16 m_id = 0;
+		WorkerContext* m_context = nullptr;
+
 		bool m_loop = false;
+		bool m_inUse = false;
 		Thread m_thread;
 
 		Worker(const Worker&) = delete;
 	};
-
+	thread_local NetworkServer::Worker* t_worker;
 
 	class NetworkServer::Connection
 	{
@@ -234,15 +275,18 @@ namespace uba
 
 			//m_logger.Debug(TC("Recv: %u, %u, %u, %u"), serviceId, messageType, id, size);
 			Worker* worker = conn.m_server.PopWorker();
-			worker->m_id = messageId;
-			worker->m_serviceId = serviceId;
-			worker->m_messageType = messageType;
-			worker->m_dataSize = messageSize;
-			worker->m_connection = &conn;// this;
-			if (worker->m_buffer.size() < messageSize)
-				worker->m_buffer.resize(size_t(Min(messageSize + 1024u, SendMaxSize)));
+			if (!worker)
+				return false;
+			auto& wc = *worker->m_context;
+			wc.id = messageId;
+			wc.serviceId = serviceId;
+			wc.messageType = messageType;
+			wc.dataSize = messageSize;
+			wc.connection = &conn;// this;
+			if (wc.buffer.size() < messageSize)
+				wc.buffer.resize(size_t(Min(messageSize + 1024u, SendMaxSize)));
 			outBodyContext = worker;
-			outBodyData = worker->m_buffer.data();
+			outBodyData = wc.buffer.data();
 			outBodySize = messageSize;
 			return true;
 		}
@@ -257,13 +301,14 @@ namespace uba
 				conn.m_server.PushWorker(worker);
 				return false;
 			}
+			auto& wc = *worker->m_context;
 
-			conn.m_client->sendBytes += worker->m_dataSize;
-			conn.m_server.m_recvBytes += worker->m_dataSize;
+			conn.m_client->sendBytes += wc.dataSize;
+			conn.m_server.m_recvBytes += wc.dataSize;
 			++conn.m_server.m_recvCount;
 
 			++conn.m_activeWorkerCount;
-			worker->m_workAvailable.Set();
+			wc.workAvailable.Set();
 			return true;
 		}
 
@@ -356,109 +401,109 @@ namespace uba
 		return ((NetworkServer::Connection*)internalData)->m_shouldDisconnect;
 	}
 
-	void NetworkServer::Worker::ThreadWorker(NetworkServer& server)
+	void NetworkServer::Worker::Update(WorkerContext& context, bool signalAvailable)
 	{
-		u32 writeMemSize = server.m_sendSize;
-		u8* writeMem = new u8[writeMemSize];
-		auto memGuard = MakeGuard([writeMem]() { delete[] writeMem; });
-
-		while (true)
-		{
-			if (!m_workAvailable.IsSet())
-				break;
-			if (!m_loop)
-				break;
+		auto& server = context.server;
 
-			// This is only additional work
-			if (!m_connection)
+		auto sg = MakeGuard([&]()
 			{
-				DoAdditionalWorkAndSignalAvailable(server);
-				continue;
-			}
+				if (signalAvailable)
+					DoAdditionalWorkAndSignalAvailable(server);
+			});
+
+		// This is only additional work
+		if (!context.connection)
+			return;
+		auto& connection = *context.connection;
+		context.connection = nullptr;
 
-			CryptoKey cryptoKey = m_connection->m_cryptoKey;
-			if (cryptoKey)
+		CryptoKey cryptoKey = connection.m_cryptoKey;
+		if (cryptoKey)
+		{
+			TimerScope ts(connection.m_decryptTimer);
+			if (!Crypto::Decrypt(server.m_logger, cryptoKey, context.buffer.data(), context.dataSize))
 			{
-				TimerScope ts(m_connection->m_decryptTimer);
-				if (!Crypto::Decrypt(server.m_logger, cryptoKey, m_buffer.data(), m_dataSize))
-				{
-					m_connection->SetShouldDisconnect();
-					m_connection->Release();
-					DoAdditionalWorkAndSignalAvailable(server);
-					continue;
-				}
+				connection.SetShouldDisconnect();
+				connection.Release();
+				return;
 			}
+		}
 
-			BinaryReader reader(m_buffer.data(), 0, m_dataSize);
+		BinaryReader reader(context.buffer.data(), 0, context.dataSize);
 
 
-			constexpr u32 HeaderSize = 5; // 2 byte id, 3 bytes size
-			constexpr u32 ErrorSize = 0xffffff;
+		constexpr u32 HeaderSize = 5; // 2 byte id, 3 bytes size
+		constexpr u32 ErrorSize = 0xffffff;
 
-			BinaryWriter writer(writeMem, 0, writeMemSize);
-			u8* idAndSizePtr = writer.AllocWrite(HeaderSize);
+		BinaryWriter writer(context.writeMem, 0, context.writeMemSize);
+		u8* idAndSizePtr = writer.AllocWrite(HeaderSize);
 			
-			u32 size;
-			WorkerRec& rec = server.m_workerFunctions[m_serviceId];
+		u32 size;
+		WorkerRec& rec = server.m_workerFunctions[context.serviceId];
 
-			u32 workIndex = 0;
-			if (server.m_trackWork)
-			{
-				workIndex = server.m_workCounter++;
-				server.m_startWork(workIndex, rec.toString(m_messageType));
-			}
+		u32 workIndex = 0;
+		if (server.m_trackWork)
+		{
+			workIndex = server.m_workCounter++;
+			server.m_startWork(workIndex, rec.toString(context.messageType));
+		}
 
-			if (!rec.func)
-			{
-				server.m_logger.Error(TC("WORKER FUNCTION NOT FOUND. id: %u, serviceid: %u type: %s, client: %s"), m_id, m_serviceId, rec.toString(m_messageType), GuidToString(m_connection->m_client->uid).str);
-				m_connection->SetShouldDisconnect();
-				size = ErrorSize;
-			}
-			else if (!rec.func({m_connection}, m_messageType, reader, writer))
-			{
-				if (m_connection->SetShouldDisconnect())
-					server.m_logger.Error(TC("WORKER FUNCTION FAILED. id: %u, serviceid: %u type: %s, client: %s"), m_id, m_serviceId, rec.toString(m_messageType), GuidToString(m_connection->m_client->uid).str);
-				size = ErrorSize;
-			}
-			else
-			{
-				size = u32(writer.GetPosition());
-			}
+		if (!rec.func)
+		{
+			server.m_logger.Error(TC("WORKER FUNCTION NOT FOUND. id: %u, serviceid: %u type: %s, client: %s"), context.id, context.serviceId, rec.toString(context.messageType), GuidToString(connection.m_client->uid).str);
+			connection.SetShouldDisconnect();
+			size = ErrorSize;
+		}
+		else if (!rec.func({&connection}, context.messageType, reader, writer))
+		{
+			if (connection.SetShouldDisconnect())
+				server.m_logger.Error(TC("WORKER FUNCTION FAILED. id: %u, serviceid: %u type: %s, client: %s"), context.id, context.serviceId, rec.toString(context.messageType), GuidToString(connection.m_client->uid).str);
+			size = ErrorSize;
+		}
+		else
+		{
+			size = u32(writer.GetPosition());
+		}
 
-			if (server.m_trackWork)
-				server.m_endWork(workIndex);
+		if (server.m_trackWork)
+			server.m_endWork(workIndex);
 
-			if (m_id)
-			{
-				UBA_ASSERT(size < (1 << 24));
+		if (context.id)
+		{
+			UBA_ASSERT(size < (1 << 24));
 				
-				u32 bodySize = u32(size - HeaderSize);
-				if (cryptoKey && size != ErrorSize && bodySize)
+			u32 bodySize = u32(size - HeaderSize);
+			if (cryptoKey && size != ErrorSize && bodySize)
+			{
+				TimerScope ts(connection.m_encryptTimer);
+				u8* bodyData = writer.GetData() + HeaderSize;
+				if (!Crypto::Encrypt(server.m_logger, cryptoKey, bodyData, bodySize))
 				{
-					TimerScope ts(m_connection->m_encryptTimer);
-					u8* bodyData = writer.GetData() + HeaderSize;
-					if (!Crypto::Encrypt(server.m_logger, cryptoKey, bodyData, bodySize))
-					{
-						m_connection->SetShouldDisconnect();
-						size = ErrorSize;
-						bodySize = u32(size - HeaderSize);
-					}
+					connection.SetShouldDisconnect();
+					size = ErrorSize;
+					bodySize = u32(size - HeaderSize);
 				}
+			}
 
-				idAndSizePtr[0] = m_id >> 8;
-				*(u32*)(idAndSizePtr + 1) = bodySize | u32(m_id << 24);
+			idAndSizePtr[0] = context.id >> 8;
+			*(u32*)(idAndSizePtr + 1) = bodySize | u32(context.id << 24);
 
-				// This can happen for proxy servers in a valid situation
-				//if (size == ErrorSize)
-				//	UBA_ASSERT(false);
+			// This can happen for proxy servers in a valid situation
+			//if (size == ErrorSize)
+			//	UBA_ASSERT(false);
 
-				m_connection->Send(writer.GetData(), size == ErrorSize ? HeaderSize : size);
-			}
+			connection.Send(writer.GetData(), size == ErrorSize ? HeaderSize : size);
+		}
 			
-			m_connection->Release();
+		connection.Release();
+	}
 
-			DoAdditionalWorkAndSignalAvailable(server);
-		}
+	void NetworkServer::Worker::ThreadWorker(NetworkServer& server)
+	{
+		t_worker = this;
+		while (m_context->workAvailable.IsSet(~0u) && m_loop)
+			Update(*m_context, true);
+		t_worker = nullptr;
 	}
 
 	void NetworkServer::Worker::DoAdditionalWorkAndSignalAvailable(NetworkServer& server)
@@ -514,12 +559,12 @@ namespace uba
 
 	NetworkServer::NetworkServer(bool& outCtorSuccess, const NetworkServerCreateInfo& info, const tchar* name)
 	:	m_logger(info.logWriter, name)
-	,	m_workerAvailable(false)
+	,	m_workerAvailable(true)
 	{
 		outCtorSuccess = true;
 
 		u32 workerCount = Min(Max(info.workerCount, (u32)(1u)), (u32)(1024u));
-		m_workerCount = workerCount;
+		m_maxWorkerCount = workerCount;
 
 		#if UBA_DEBUG
 		m_logger.Info(TC("Created in DEBUG"));
@@ -588,6 +633,11 @@ namespace uba
 	{
 		StopListen();
 
+		{
+			ScopedWriteLock lock(m_availableWorkersLock);
+			m_workersEnabled = false;
+			m_workerAvailable.Set();
+		}
 		{
 			ScopedWriteLock lock(m_addConnectionsLock);
 			m_addConnections.clear();
@@ -612,19 +662,23 @@ namespace uba
 				abort(); // TODO: Does this produce core dump on windows?
 		}
 
-		auto deleteWorkers = [](Worker*& start)
+
+		ScopedWriteLock lock(m_availableWorkersLock);
+		while (auto worker = m_firstActiveWorker)
 		{
-			Worker* worker = start;
-			while (worker)
-			{
-				Worker* temp = worker;
-				worker = worker->m_nextWorker;
-				delete temp;
-			}
-			start = nullptr;
-		};
-		deleteWorkers(m_firstAvailableWorker);
-		deleteWorkers(m_firstActiveWorker);
+			lock.Leave();
+			worker->Stop(*this);
+			lock.Enter();
+		}
+
+		auto worker = m_firstAvailableWorker;
+		while (worker)
+		{
+			auto temp = worker;
+			worker = worker->m_nextWorker;
+			delete temp;
+		}
+		m_firstAvailableWorker = nullptr;
 	}
 
 	bool NetworkServer::AddClient(NetworkBackend& backend, const tchar* ip, u16 port, const u8* cryptoKey128)
@@ -672,7 +726,7 @@ namespace uba
 			return;
 
 		StringBuffer<> workers;
-		workers.Appendf(TC("%u/%u"), m_createdWorkerCount, m_workerCount);
+		workers.Appendf(TC("%u/%u"), m_createdWorkerCount, m_maxWorkerCount);
 
 		logger.Info(TC("  ----- Uba server stats summary ------"));
 		logger.Info(TC("  MaxActiveConnections           %6u"), m_maxActiveConnections);
@@ -752,17 +806,20 @@ namespace uba
 		lock.Leave();
 
 		ScopedWriteLock lock2(m_availableWorkersLock);
-		while (count-- && m_createdWorkerCount < m_workerCount)
+		while (count-- && m_createdWorkerCount < m_maxWorkerCount)
 		{
 			Worker* worker = PopWorkerNoLock();
-			worker->m_connection = nullptr;
-			worker->m_workAvailable.Set();
+			if (!worker)
+				break;
+			UBA_ASSERT(worker->m_inUse);
+			worker->m_context->connection = nullptr;
+			worker->m_context->workAvailable.Set();
 		}
 	}
 
 	u32 NetworkServer::GetWorkerCount()
 	{
-		return m_workerCount;
+		return m_maxWorkerCount;
 	}
 
 	u64 NetworkServer::GetTotalSentBytes()
@@ -792,7 +849,50 @@ namespace uba
 		AdditionalWork work;
 		ScopedWriteLock lock(m_additionalWorkLock);
 		if (m_additionalWork.empty())
-			return false;
+		{
+			lock.Leave();
+
+			ScopedWriteLock lock2(m_availableWorkersLock);
+			if (m_createdWorkerCount != m_maxWorkerCount)
+				return false;
+			lock2.Leave();
+			auto worker = t_worker;
+			UBA_ASSERT(worker);
+			auto oldContext = worker->m_context;
+			WorkerContext context(*this);
+			worker->m_context = &context;
+
+			PushWorker(worker);
+			bool workAvail = context.workAvailable.IsSet(10);
+			lock2.Enter();
+			if (worker->m_inUse)
+			{
+				lock2.Leave();
+				if (!workAvail)
+					context.workAvailable.IsSet(~0u);
+				worker->Update(context, false);
+				UBA_ASSERT(worker->m_inUse);
+			}
+			else
+			{
+				// Take worker back from free list
+				if (m_firstAvailableWorker == worker)
+					m_firstAvailableWorker = worker->m_nextWorker;
+				else
+					worker->m_prevWorker->m_nextWorker = worker->m_nextWorker;
+				if (worker->m_nextWorker)
+					worker->m_nextWorker->m_prevWorker = worker->m_prevWorker;
+				worker->m_prevWorker = nullptr;
+				worker->m_nextWorker = m_firstActiveWorker;
+				if (m_firstActiveWorker)
+					m_firstActiveWorker->m_prevWorker = worker;
+				m_firstActiveWorker = worker;
+				worker->m_inUse = true;
+			}
+			
+			worker->m_context = oldContext;
+			return true;
+		}
 		work = m_additionalWork.front();
 		m_additionalWork.pop_front();
 		lock.Leave();
@@ -827,8 +927,17 @@ namespace uba
 
 	NetworkServer::Worker* NetworkServer::PopWorker()
 	{
-		ScopedWriteLock lock(m_availableWorkersLock);
-		return PopWorkerNoLock();
+		while (true)
+		{
+			ScopedWriteLock lock(m_availableWorkersLock);
+			if (!m_workersEnabled)
+				return nullptr;
+			if (auto worker = PopWorkerNoLock())
+				return worker;
+			m_workerAvailable.Reset();
+			lock.Leave();
+			m_workerAvailable.IsSet();
+		}
 	}
 
 	NetworkServer::Worker* NetworkServer::PopWorkerNoLock()
@@ -842,6 +951,9 @@ namespace uba
 		}
 		else
 		{
+			if (m_createdWorkerCount == m_maxWorkerCount)
+				return nullptr;
+
 			worker = new Worker();
 			worker->Start(*this);
 			++m_createdWorkerCount;
@@ -851,6 +963,7 @@ namespace uba
 			m_firstActiveWorker->m_prevWorker = worker;
 		worker->m_nextWorker = m_firstActiveWorker;
 		m_firstActiveWorker = worker;
+		worker->m_inUse = true;
 
 		return worker;
 	}
@@ -863,6 +976,8 @@ namespace uba
 
 	void NetworkServer::PushWorkerNoLock(Worker* worker)
 	{
+		UBA_ASSERT(worker->m_inUse);
+
 		if (worker->m_prevWorker)
 			worker->m_prevWorker->m_nextWorker = worker->m_nextWorker;
 		else
@@ -874,6 +989,7 @@ namespace uba
 			m_firstAvailableWorker->m_prevWorker = worker;
 		worker->m_prevWorker = nullptr;
 		worker->m_nextWorker = m_firstAvailableWorker;
+		worker->m_inUse = false;
 		m_firstAvailableWorker = worker;
 		m_workerAvailable.Set();
 	}
diff --git a/Engine/Source/Programs/UnrealBuildAccelerator/Common/Public/UbaNetworkServer.h b/Engine/Source/Programs/UnrealBuildAccelerator/Common/Public/UbaNetworkServer.h
--- a/Engine/Source/Programs/UnrealBuildAccelerator/Common/Public/UbaNetworkServer.h
+++ b/Engine/Source/Programs/UnrealBuildAccelerator/Common/Public/UbaNetworkServer.h
@@ -87,8 +87,9 @@ namespace uba
 		void SetWorkListener(const WorkBeginFunction& start, const WorkEndFunction& end);
 		void ResetWorkListener();
 
-	private:
 		class Worker;
+	private:
+		struct WorkerContext;
 		class Connection;
 		friend ConnectionInfo;
 
@@ -116,7 +117,7 @@ namespace uba
 		struct OnDisconnectEntry { u8 id; OnDisconnectFunction function; };
 		List<OnDisconnectEntry> m_onDisconnectFunctions;
 
-		u32 m_workerCount = 0;
+		u32 m_maxWorkerCount = 0;
 
 		ReaderWriterLock m_additionalWorkLock;
 		struct AdditionalWork { Function<void()> func; TString desc; };
@@ -129,6 +130,7 @@ namespace uba
 		u32 m_sendSize = 0;
 		u32 m_receiveTimeoutMs = 0;
 		u32 m_createdWorkerCount = 0;
+		bool m_workersEnabled = true;
 
 		ReaderWriterLock m_addConnectionsLock;
 		List<Thread> m_addConnections;
