[UbaCli]
* Added -deleteCas
* Added command mode "agent" that can be used to test things with local agent

diff --git a/Engine/Source/Programs/UnrealBuildAccelerator/Cli/Private/UbaCli.cpp b/Engine/Source/Programs/UnrealBuildAccelerator/Cli/Private/UbaCli.cpp
--- a/Engine/Source/Programs/UnrealBuildAccelerator/Cli/Private/UbaCli.cpp
+++ b/Engine/Source/Programs/UnrealBuildAccelerator/Cli/Private/UbaCli.cpp
@@ -3,6 +3,7 @@
 #include "UbaNetworkBackendTcp.h"
 #include "UbaFileAccessor.h"
 #include "UbaScheduler.h"
+#include "UbaSessionClient.h"
 #include "UbaSessionServer.h"
 #include "UbaStorageClient.h"
 #include "UbaStorageServer.h"
@@ -50,6 +51,7 @@ namespace uba
 		logger.Info(TC("  CommandTypes:"));
 		logger.Info(TC("   local                   Will run executable locally using detoured paths"));
 		logger.Info(TC("   remote                  Will wait for available agent and then run executable remotely"));
+		logger.Info(TC("   agent                   Will run executable against agent spawned in process"));
 		logger.Info(TC("   native                  Will run executable in a normal way"));
 		logger.Info(TC(""));
 		logger.Info(TC("  Options:"));
@@ -60,6 +62,7 @@ namespace uba
 		logger.Info(TC("   -workdir=<dir>          Working directory"));
 		logger.Info(TC("   -checkcas               Check so all cas entries are correct"));
 		logger.Info(TC("   -checkaws               Check if we are inside aws and output information about aws"));
+		logger.Info(TC("   -deletecas              Deletes the casdb"));
 		logger.Info(TC("   -getcas                 Will print hash of application"));
 		logger.Info(TC("   -summary                Print summary at the end of a session"));
 		logger.Info(TC("   -nocustomalloc          Disable custom allocator for processes. If you see odd crashes this can be tested"));
@@ -137,6 +140,7 @@ namespace uba
 		bool checkCas2 = false;
 		bool checkAws = false;
 		bool getCas = false;
+		bool deleteCas = false;
 		bool enableStdOut = true;
 		bool printSummary = false;
 		u32 loopCount = 1;
@@ -147,6 +151,7 @@ namespace uba
 			CommandType_Local,
 			CommandType_Remote,
 			CommandType_Native,
+			CommandType_Agent
 		};
 
 		CommandType commandType = CommandType_NotSet;
@@ -197,6 +202,10 @@ namespace uba
 			{
 				commandType = CommandType_Native;
 			}
+			else if (name.Equals(TC("agent")))
+			{
+				commandType = CommandType_Agent;
+			}
 			else if (IsWindows && name.Equals(TC("-visualizer")))
 			{
 				launchVisualizer = true;
@@ -257,6 +266,10 @@ namespace uba
 			{
 				checkAws = true;
 			}
+			else if (name.Equals(TC("-deletecas")))
+			{
+				deleteCas = true;
+			}
 			else if (name.Equals(TC("-getcas")))
 			{
 				getCas = true;
@@ -291,6 +304,14 @@ namespace uba
 		FilteredLogWriter logWriter(g_consoleLogWriter, quiet ? LogEntryType_Info : LogEntryType_Detail);
 		LoggerWithWriter logger(logWriter, TC(""));
 
+		if (deleteCas)
+		{
+			StorageImpl(StorageCreateInfo(g_rootDir.data, logWriter)).DeleteAllCas();
+			StringBuffer<> clientRootDir;
+			clientRootDir.Append(g_rootDir).Append("Agent");
+			StorageImpl(StorageCreateInfo(clientRootDir.data, logWriter)).DeleteAllCas();
+		}
+
 		if (checkCas)
 		{
 			StorageCreateInfo storageInfo(g_rootDir.data, logWriter);
@@ -438,12 +459,16 @@ namespace uba
 
 		NetworkBackendTcp networkBackend(logWriter);
 		NetworkServerCreateInfo nsci(logWriter);
+		//nsci.workerCount = 4;
 		bool ctorSuccess = true;
 		NetworkServer* server = new NetworkServer(ctorSuccess, nsci);
 		auto destroyServer = MakeGuard([&]() { delete server; });
 		if (!ctorSuccess)
 			return -1;
 
+		bool isRemote = commandType == CommandType_Remote || commandType == CommandType_Agent;
+		bool useScheduler = EndsWith(application.c_str(), application.size(), TC(".yaml"));
+
 		StorageServerCreateInfo storageInfo(*server, g_rootDir.data, logWriter);
 		storageInfo.casCapacityBytes = storageCapacity;
 		storageInfo.storeCompressed = storeCompressed;
@@ -451,8 +476,9 @@ namespace uba
 		auto destroyStorage = MakeGuard([&]() { delete storage; });
 
 		SessionServerCreateInfo info(*storage, *server);
-		info.useUniqueId = false;
+		info.useUniqueId = useScheduler;
 		info.traceEnabled = true;
+		//info.detailedTrace = true;
 		info.launchVisualizer = launchVisualizer;
 		info.disableCustomAllocator = disableCustomAllocator;
 		//info.shouldWriteToDisk = shouldWriteToDisk;
@@ -463,7 +489,7 @@ namespace uba
 		auto session = new SessionServer(info);
 		auto destroySession = MakeGuard([&]() { delete session; });
 
-		if (commandType == CommandType_Remote)
+		if (isRemote)
 		{
 			if (!storage->LoadCasTable(true))
 				return -1;
@@ -516,10 +542,39 @@ namespace uba
 			return true;
 		};
 
+		auto RunWithClient = [&](const Function<bool()>& func)
+			{
+				NetworkClient client(ctorSuccess, { logWriter });
+
+				StringBuffer<> clientRootDir;
+				clientRootDir.Append(g_rootDir).Append("Agent");
+				StorageClientCreateInfo storageClientInfo(client, clientRootDir.data);
+				StorageClient storageClient(storageClientInfo);
+
+				SessionClientCreateInfo sessionClientInfo(storageClient, client, logWriter);
+				sessionClientInfo.maxProcessCount = DefaultProcessorCount;
+				sessionClientInfo.rootDir = clientRootDir.data;
+				sessionClientInfo.deleteSessionsOlderThanSeconds = 1;
+				SessionClient sessionClient(sessionClientInfo);
+
+				//for (u32 i=0; i!=4; ++i)
+					if (!client.Connect(networkBackend, TC("127.0.0.1"), port))
+						return logger.Error(TC("Failed to connect"));
+
+				auto cg = MakeGuard([&]() { sessionClient.Stop(); client.Disconnect(); });
+
+				return func();
+			};
+
+		auto RunAgent = [&](const TString& app, const TString& arg)
+		{
+			return RunWithClient([&]() { return RunRemote(app, arg); });
+		};
+
 		auto RunScheduler = [&](const tchar* yamlFile)
 		{
 			SchedulerCreateInfo info(*session);
-			info.forceRemote = commandType == CommandType_Remote;
+			info.forceRemote = isRemote;
 			Scheduler scheduler(info);
 
 			if (!scheduler.EnqueueFromFile(yamlFile))
@@ -553,23 +608,31 @@ namespace uba
 						finished.Set();
 				});
 
-			logger.Info(TC("Running Scheduler with %u processes"), queued);
-			u64 start = GetTime();
-			scheduler.Start();
-			if (!finished.IsSet())
-				return false;
-			u64 time = GetTime() - start;
-			logger.Info(TC("Scheduler run took %s"), TimeToText(time).str);
-			logger.Info(TC(""));
-			stopServer.Execute();
-			return success;
+			auto RunQueue = [&]()
+				{
+					logger.Info(TC("Running Scheduler with %u processes"), queued);
+					u64 start = GetTime();
+					scheduler.Start();
+					if (!finished.IsSet())
+						return false;
+					u64 time = GetTime() - start;
+					logger.Info(TC("Scheduler run took %s"), TimeToText(time).str);
+					logger.Info(TC(""));
+					stopServer.Execute();
+					return success;
+				};
+
+			if (commandType == CommandType_Agent)
+				return RunWithClient([&]() { return RunQueue(); });
+			else
+				return RunQueue();
 		};
 
 		for (u32 i=0; i!=loopCount; ++i)
 		{
 			bool success = false;
 
-			if (EndsWith(application.c_str(), application.size(), TC(".yaml")))
+			if (useScheduler)
 			{
 				success = RunScheduler(application.c_str());
 			}
@@ -586,6 +649,8 @@ namespace uba
 				case CommandType_Remote:
 					success = RunRemote(application, arguments);
 					break;
+				case CommandType_Agent:
+					success = RunAgent(application, arguments);
 				}
 			}
 			if (!success)
