[UBA]
* Added so session information contains information about number of processor groups
* Added so detoured process matches uba message thread's processor group

diff --git a/Engine/Source/Programs/UnrealBuildAccelerator/Common/Private/UbaProcess.cpp b/Engine/Source/Programs/UnrealBuildAccelerator/Common/Private/UbaProcess.cpp
--- a/Engine/Source/Programs/UnrealBuildAccelerator/Common/Private/UbaProcess.cpp
+++ b/Engine/Source/Programs/UnrealBuildAccelerator/Common/Private/UbaProcess.cpp
@@ -1362,10 +1362,21 @@ namespace uba
 			}
 		}
 
-		if (!AlternateThreadGroupAffinity(m_nativeThreadHandle))
+		bool affinitySet = false;
+		if (!m_messageThread.Wait(0))
 		{
-			logger.Error(TC("Failed to set thread group affinity to process"));//% ls. (% ls)"), commandLine.c_str(), LastErrorToText().data);
-			return UBA_EXIT_CODE(10);
+			GroupAffinity aff;
+			if (m_messageThread.GetGroupAffinity(aff))
+				affinitySet = SetThreadGroupAffinity(m_nativeThreadHandle, aff);
+		}
+		
+		if (!affinitySet)
+		{
+			if (!AlternateThreadGroupAffinity(m_nativeThreadHandle))
+			{
+				logger.Error(TC("Failed to set thread group affinity to process"));//% ls. (% ls)"), commandLine.c_str(), LastErrorToText().data);
+				return UBA_EXIT_CODE(10);
+			}
 		}
 
 		m_processStats.startupTime = GetTime() - m_startTime;
diff --git a/Engine/Source/Programs/UnrealBuildAccelerator/Common/Private/UbaSession.cpp b/Engine/Source/Programs/UnrealBuildAccelerator/Common/Private/UbaSession.cpp
--- a/Engine/Source/Programs/UnrealBuildAccelerator/Common/Private/UbaSession.cpp
+++ b/Engine/Source/Programs/UnrealBuildAccelerator/Common/Private/UbaSession.cpp
@@ -2093,7 +2093,13 @@ namespace uba
 
 	void Session::GetSystemInfo(StringBufferBase& out)
 	{
-		u32 numCPU = GetLogicalProcessorCount();
+		u32 cpuCount = GetLogicalProcessorCount();
+		u32 cpuGroupCount = GetProcessorGroupCount();
+
+		StringBuffer<128> cpuCountStr;
+		if (cpuGroupCount != 1)
+			cpuCountStr.AppendValue(cpuGroupCount).Append('x');
+		cpuCountStr.AppendValue(cpuCount/cpuGroupCount);
 
 		u64 totalMemoryInKilobytes = 0;
 		StringBuffer<128> hzStr;
@@ -2102,8 +2108,8 @@ namespace uba
 		GetPhysicallyInstalledSystemMemory(&totalMemoryInKilobytes);
 
 		Vector<PROCESSOR_POWER_INFORMATION> procInfos;
-		procInfos.resize(numCPU);
-		if (CallNtPowerInformation(ProcessorInformation, NULL, 0, procInfos.data(), numCPU*sizeof(PROCESSOR_POWER_INFORMATION)) == STATUS_SUCCESS)
+		procInfos.resize(cpuCount);
+		if (CallNtPowerInformation(ProcessorInformation, NULL, 0, procInfos.data(), cpuCount*sizeof(PROCESSOR_POWER_INFORMATION)) == STATUS_SUCCESS)
 			hzStr.Appendf(TC(" @ %.1fGHz"), float(procInfos[0].MaxMhz) / 1000.0f);
 		#else
 		u64 throwAway;
@@ -2118,7 +2124,7 @@ namespace uba
 		BytesToText temp(capacity);
 		if (capacity)
 			capacityStr = temp.str;
-		out.Appendf(TC("CPU:%u%s Mem:%ugb Cas:%s"), numCPU, hzStr.data, u32(totalMemoryInKilobytes/(1024*1024)), capacityStr);
+		out.Appendf(TC("CPU:%s%s Mem:%ugb Cas:%s"), cpuCountStr.data, hzStr.data, u32(totalMemoryInKilobytes/(1024*1024)), capacityStr);
 
 		StringBuffer<128> zone;
 		if (m_storage.GetZone(zone))
diff --git a/Engine/Source/Programs/UnrealBuildAccelerator/Common/Private/UbaThread.cpp b/Engine/Source/Programs/UnrealBuildAccelerator/Common/Private/UbaThread.cpp
--- a/Engine/Source/Programs/UnrealBuildAccelerator/Common/Private/UbaThread.cpp
+++ b/Engine/Source/Programs/UnrealBuildAccelerator/Common/Private/UbaThread.cpp
@@ -8,8 +8,8 @@ namespace uba
 	bool AlternateThreadGroupAffinity(void* nativeThreadHandle)
 	{
 #if PLATFORM_WINDOWS
-		static int processorGroupCount = GetActiveProcessorGroupCount();
-		if (processorGroupCount == 1)
+		int processorGroupCount = GetProcessorGroupCount();
+		if (processorGroupCount <= 1)
 			return true;
 		static Atomic<int> processorGroupCounter;
 		u16 processorGroup = u16((processorGroupCounter++) % processorGroupCount);
@@ -19,12 +19,27 @@ namespace uba
 		GROUP_AFFINITY groupAffinity = {};
 		groupAffinity.Mask = ~0ull >> (int)(64 - groupProcessorCount);
 		groupAffinity.Group = processorGroup;
-		return SetThreadGroupAffinity(nativeThreadHandle, &groupAffinity, NULL);
+		return ::SetThreadGroupAffinity(nativeThreadHandle, &groupAffinity, NULL);
 #else
 		return true;
 #endif
 	}
 
+	bool SetThreadGroupAffinity(void* nativeThreadHandle, const GroupAffinity& affinity)
+	{
+#if PLATFORM_WINDOWS
+		if (GetProcessorGroupCount() <= 1)
+			return true;
+		GROUP_AFFINITY groupAffinity = {};
+		groupAffinity.Mask = affinity.mask;
+		groupAffinity.Group = affinity.group;
+		return ::SetThreadGroupAffinity(nativeThreadHandle, &groupAffinity, NULL);
+#else
+		return false;
+#endif
+	}
+
+
 	Thread::Thread()
 	{
 	}
@@ -44,7 +59,7 @@ namespace uba
 		m_func = std::move(f);
 #if PLATFORM_WINDOWS
 		m_handle = CreateThread(NULL, 0, [](LPVOID p) -> DWORD { return ((Thread*)p)->m_func(); }, this, 0, NULL);
-		//AlternateThreadGroupAffinity(m_handle);
+		AlternateThreadGroupAffinity(m_handle);
 #else
 		int err = 0;
 
@@ -107,4 +122,21 @@ namespace uba
 		m_handle = nullptr;
 		return true;
 	}
+
+	bool Thread::GetGroupAffinity(GroupAffinity& out)
+	{
+#if PLATFORM_WINDOWS
+		if (GetProcessorGroupCount() <= 1)
+			return true;
+		GROUP_AFFINITY aff;
+		if (!::GetThreadGroupAffinity(m_handle, &aff))
+			return false;
+		out.mask = aff.Mask;
+		out.group = aff.Group;
+		return true;
+#else
+		return false;
+#endif
+	}
+
 }
diff --git a/Engine/Source/Programs/UnrealBuildAccelerator/Common/Public/UbaThread.h b/Engine/Source/Programs/UnrealBuildAccelerator/Common/Public/UbaThread.h
--- a/Engine/Source/Programs/UnrealBuildAccelerator/Common/Public/UbaThread.h
+++ b/Engine/Source/Programs/UnrealBuildAccelerator/Common/Public/UbaThread.h
@@ -8,6 +8,8 @@
 namespace uba
 {
 	class Event;
+	struct GroupAffinity;
+
 
 	class Thread
 	{
@@ -17,6 +19,8 @@ namespace uba
 		~Thread();
 		void Start(Function<u32()>&& func);
 		bool Wait(u32 milliseconds = ~0u, Event* wakeupEvent = nullptr);
+		bool GetGroupAffinity(GroupAffinity& out);
+
 	private:
 		Function<u32()>	m_func;
 		void* m_handle = nullptr;
@@ -30,5 +34,11 @@ namespace uba
 	};
 
 
+	struct GroupAffinity
+	{
+		u64 mask = 0;
+		u16 group = 0;
+	};
+	bool SetThreadGroupAffinity(void* nativeThreadHandle, const GroupAffinity& affinity);
 	bool AlternateThreadGroupAffinity(void* nativeThreadHandle);
 }
diff --git a/Engine/Source/Programs/UnrealBuildAccelerator/Core/Private/UbaPlatform.cpp b/Engine/Source/Programs/UnrealBuildAccelerator/Core/Private/UbaPlatform.cpp
--- a/Engine/Source/Programs/UnrealBuildAccelerator/Core/Private/UbaPlatform.cpp
+++ b/Engine/Source/Programs/UnrealBuildAccelerator/Core/Private/UbaPlatform.cpp
@@ -331,6 +331,17 @@ namespace uba
 		#endif
 	}
 
+	u32 GetProcessorGroupCount()
+	{
+		#if PLATFORM_WINDOWS
+		static u32 s_processorGroupCount = u32(GetActiveProcessorGroupCount());
+		if (s_processorGroupCount)
+			return s_processorGroupCount;
+		#endif
+		return 1u;
+	}
+
+
 #if !PLATFORM_WINDOWS
 	void GetMappingHandleName(StringBufferBase& out, u64 uid)
 	{
diff --git a/Engine/Source/Programs/UnrealBuildAccelerator/Core/Public/UbaPlatform.h b/Engine/Source/Programs/UnrealBuildAccelerator/Core/Public/UbaPlatform.h
--- a/Engine/Source/Programs/UnrealBuildAccelerator/Core/Public/UbaPlatform.h
+++ b/Engine/Source/Programs/UnrealBuildAccelerator/Core/Public/UbaPlatform.h
@@ -96,6 +96,7 @@ namespace uba
 	u32 GetEnvironmentVariableW(const tchar* name, tchar* buffer, u32 nSize);
 	u32 ExpandEnvironmentStringsW(const tchar* lpSrc, tchar* lpDst, u32 nSize);
 	u32 GetLogicalProcessorCount();
+	u32 GetProcessorGroupCount();
 
 #if PLATFORM_WINDOWS
 	inline constexpr bool CaseInsensitiveFs = true;
