[UBA]
* Fixed so "access" calls inside /proc/ is never detoured

diff --git a/Engine/Source/Programs/UnrealBuildAccelerator/Detours/Private/PosixOS/UbaDetoursFunctionsPosix.cpp b/Engine/Source/Programs/UnrealBuildAccelerator/Detours/Private/PosixOS/UbaDetoursFunctionsPosix.cpp
--- a/Engine/Source/Programs/UnrealBuildAccelerator/Detours/Private/PosixOS/UbaDetoursFunctionsPosix.cpp
+++ b/Engine/Source/Programs/UnrealBuildAccelerator/Detours/Private/PosixOS/UbaDetoursFunctionsPosix.cpp
@@ -1114,7 +1114,7 @@ UBA_EXPORT int UBA_WRAPPER(access)(const char* pathname, int mode)
 	UBA_INIT_DETOUR(access, pathname, mode);
 
 	StringBuffer<> fixedPath;
-	if (!FixPath(fixedPath, pathname) || fixedPath.StartsWith(g_systemTemp.data))
+	if (!FixPath(fixedPath, pathname) || fixedPath.StartsWith(g_systemTemp.data) || fixedPath.StartsWith("/proc"))
 	{
 		return TRUE_WRAPPER(access)(pathname, mode);
 	}
