[UBA]
* Fixed write-after-free bug in code writing cas file to disk
* Fixed write-after-free bug in code cleaning up finished processes in session client code
* Added some asserts
* UbaCli - disabled remote logging in non-debug configs

diff --git a/Engine/Source/Programs/UnrealBuildAccelerator/Cli/Private/UbaCli.cpp b/Engine/Source/Programs/UnrealBuildAccelerator/Cli/Private/UbaCli.cpp
--- a/Engine/Source/Programs/UnrealBuildAccelerator/Cli/Private/UbaCli.cpp
+++ b/Engine/Source/Programs/UnrealBuildAccelerator/Cli/Private/UbaCli.cpp
@@ -2,6 +2,7 @@
 
 #include "UbaNetworkBackendTcp.h"
 #include "UbaFileAccessor.h"
+#include "UbaProtocol.h"
 #include "UbaScheduler.h"
 #include "UbaSessionClient.h"
 #include "UbaSessionServer.h"
@@ -484,7 +485,10 @@ namespace uba
 		//info.shouldWriteToDisk = shouldWriteToDisk;
 		info.rootDir = g_rootDir.data;
 		//info.traceName.Append(TC("TESTTRACE"));
+		#if UBA_DEBUG_LOG_ENABLED
 		info.remoteLogEnabled = true;
+		#endif
+
 		info.deleteSessionsOlderThanSeconds = 1;
 		auto session = new SessionServer(info);
 		auto destroySession = MakeGuard([&]() { delete session; });
diff --git a/Engine/Source/Programs/UnrealBuildAccelerator/Common/Private/UbaSessionClient.cpp b/Engine/Source/Programs/UnrealBuildAccelerator/Common/Private/UbaSessionClient.cpp
--- a/Engine/Source/Programs/UnrealBuildAccelerator/Common/Private/UbaSessionClient.cpp
+++ b/Engine/Source/Programs/UnrealBuildAccelerator/Common/Private/UbaSessionClient.cpp
@@ -1496,26 +1496,25 @@ namespace uba
 						if (session.m_shouldSendLogToServer)
 							session.SendLogFileToServer(*(ProcessImpl*)h.m_process);
 
+						ScopedWriteLock lock(rec->lock);
+
 						auto decreaseWeight = MakeGuard([&]()
 							{
 								ScopedWriteLock weightLock(activeWeightLock);
 								activeWeight -= rec->weight;
+								rec->isDone = true;
 								session.m_waitToSendEvent.Set();
 							});
 
-						ScopedWriteLock lock(rec->lock);
 						if (rec->isKilled)
-						{
-							rec->isDone = true;
 							return;
-						}
+
 						auto& process = *(ProcessImpl*)h.m_process;
 
 						if (session.m_killRandomIndex != ~0u && session.m_killRandomCounter++ == session.m_killRandomIndex)
 						{
 							session.m_loop = false;
 							session.m_logger.Info(TC("Killed random process (%s)"), process.m_description.c_str());
-							rec->isDone = true;
 							return;
 						}
 
@@ -1527,7 +1526,6 @@ namespace uba
 							{
 								if (session.m_loop)
 									session.SendReturnProcess(rec->handle.GetId(), session.m_terminationReason);
-								rec->isDone = true;
 								return;
 							}
 						}
@@ -1543,7 +1541,6 @@ namespace uba
 								session.m_logger.Error(desc);
 								if (session.m_loop)
 									session.SendReturnProcess(rec->handle.GetId(), desc);
-								rec->isDone = true;
 								return;
 							}
 						}
@@ -1554,7 +1551,6 @@ namespace uba
 						{
 							if (session.m_loop)
 								session.SendReturnProcess(rec->handle.GetId(), TC("Cancelled"));
-							rec->isDone = true;
 							return;
 						}
 
@@ -1581,9 +1577,6 @@ namespace uba
 
 						if (session.m_processFinished)
 							session.m_processFinished(&process);
-
-						rec->isDone = true;
-						session.m_waitToSendEvent.Set();
 					};
 
 					process->Start(startInfo, realApplication.data, m_processWorkingDir.data, true, env, true, true);
diff --git a/Engine/Source/Programs/UnrealBuildAccelerator/Common/Private/UbaStorage.cpp b/Engine/Source/Programs/UnrealBuildAccelerator/Common/Private/UbaStorage.cpp
--- a/Engine/Source/Programs/UnrealBuildAccelerator/Common/Private/UbaStorage.cpp
+++ b/Engine/Source/Programs/UnrealBuildAccelerator/Common/Private/UbaStorage.cpp
@@ -532,6 +532,9 @@ namespace uba
 
 	bool StorageImpl::DecompressMemoryToMemory(u8* compressedData, u8* writeData, u64 decompressedSize)
 	{
+		UBA_ASSERT(compressedData);
+		UBA_ASSERT(writeData);
+
 		StorageStats& stats = Stats();
 
 		if (decompressedSize > BufferSlotSize * 4) // Arbitrary size threshold. We want to at least catch the pch here
@@ -615,7 +618,7 @@ namespace uba
 				if (!compressedBlockSize)
 					break;
 				u32 decompressedBlockSize = ((u32*)readPos)[1];
-
+				UBA_ASSERT(decompressedBlockSize <= left);
 				readPos += sizeof(u32) * 2;
 
 				TimerScope ts(stats.decompressToMem);
diff --git a/Engine/Source/Programs/UnrealBuildAccelerator/Common/Private/UbaStorageServer.cpp b/Engine/Source/Programs/UnrealBuildAccelerator/Common/Private/UbaStorageServer.cpp
--- a/Engine/Source/Programs/UnrealBuildAccelerator/Common/Private/UbaStorageServer.cpp
+++ b/Engine/Source/Programs/UnrealBuildAccelerator/Common/Private/UbaStorageServer.cpp
@@ -756,6 +756,9 @@ namespace uba
 				u64 toRead = reader.GetLeft();
 				reader.ReadBytes(activeStore.mappedView.memory + memOffset, toRead);
 				
+				u64 time2 = GetTime();
+				activeStore.recvCasTime += time2 - start;
+
 				u64 totalWritten = activeStore.totalWritten.fetch_add(toRead) + toRead;
 				if (totalWritten == activeStore.fileSize)
 				{
@@ -775,7 +778,7 @@ namespace uba
 					if (isPersistentStore)
 						CasEntryWritten(*activeStore.casEntry, totalWritten);
 
-					activeStore.recvCasTime += GetTime() - start;
+					activeStore.recvCasTime += GetTime() - time2;
 
 					StorageStats& stats = Stats();
 					stats.recvCas.Add(Timer{activeStore.recvCasTime, 1});
@@ -808,8 +811,6 @@ namespace uba
 					return true;
 				}
 
-				activeStore.recvCasTime += GetTime() - start;
-
 				if (firstStore)
 				{
 					storeId = PopId();
