[Uba]
* Added O_CLOEXEC to all open to prevent file descriptors from leaking to child processes. This seems to really fix the ETXTBUSY
* Added read hint for DecompressMemoryToMemory to be able to show better error messages

diff --git a/Engine/Source/Programs/UnrealBuildAccelerator/Common/Private/UbaFile.cpp b/Engine/Source/Programs/UnrealBuildAccelerator/Common/Private/UbaFile.cpp
--- a/Engine/Source/Programs/UnrealBuildAccelerator/Common/Private/UbaFile.cpp
+++ b/Engine/Source/Programs/UnrealBuildAccelerator/Common/Private/UbaFile.cpp
@@ -395,7 +395,7 @@ namespace uba
 	#if PLATFORM_WINDOWS
 		return (FileHandle)(u64)::CreateFileW(fileName, desiredAccess, shareMode, NULL, createDisp, flagsAndAttributes, NULL);
 	#else
-		int flags = 0;
+		int flags = O_CLOEXEC;
 		if (createDisp == CREATE_ALWAYS)
 			flags |= O_CREAT | O_TRUNC;
 		else if (createDisp == OPEN_EXISTING)
diff --git a/Engine/Source/Programs/UnrealBuildAccelerator/Common/Private/UbaFileAccessor.cpp b/Engine/Source/Programs/UnrealBuildAccelerator/Common/Private/UbaFileAccessor.cpp
--- a/Engine/Source/Programs/UnrealBuildAccelerator/Common/Private/UbaFileAccessor.cpp
+++ b/Engine/Source/Programs/UnrealBuildAccelerator/Common/Private/UbaFileAccessor.cpp
@@ -364,7 +364,7 @@ namespace uba
 							return m_logger.Error(TC("Failed to rename temporary file %s to %s (%s)"), realFileName, m_fileName, strerror(errno));
 
 						// Need to copy, can't rename over devices
-						int targetFd = open(m_fileName, O_CREAT | O_TRUNC | O_WRONLY);
+						int targetFd = open(m_fileName, O_CREAT | O_TRUNC | O_WRONLY | O_CLOEXEC);
 						auto g = MakeGuard([targetFd]() { close(targetFd); });
 						if (targetFd == -1)
 							return m_logger.Error(TC("Failed to create file %s for move from temporary file %s (%s)"), m_fileName, realFileName, strerror(errno));
diff --git a/Engine/Source/Programs/UnrealBuildAccelerator/Common/Private/UbaFileMapping.cpp b/Engine/Source/Programs/UnrealBuildAccelerator/Common/Private/UbaFileMapping.cpp
--- a/Engine/Source/Programs/UnrealBuildAccelerator/Common/Private/UbaFileMapping.cpp
+++ b/Engine/Source/Programs/UnrealBuildAccelerator/Common/Private/UbaFileMapping.cpp
@@ -70,7 +70,7 @@ namespace uba
 
 					StringBuffer<128> lockFile;
 					lockFile.Append(lockDir).EnsureEndsWithSlash().Append(e.name);
-					int lockFd = open(lockFile.data, O_RDWR, S_IRUSR | S_IWUSR);
+					int lockFd = open(lockFile.data, O_RDWR, S_IRUSR | S_IWUSR | O_CLOEXEC);
 					if (lockFd == -1)
 					{
 						if (errno == EPERM)
@@ -113,7 +113,7 @@ namespace uba
 
 			lockFile.Clear().Append(lockDir).EnsureEndsWithSlash().AppendValue(uid);
 
-			lockFd = open(lockFile.data, O_CREAT | O_RDWR | O_NOFOLLOW | O_EXCL, S_IRUSR | S_IWUSR);
+			lockFd = open(lockFile.data, O_CREAT | O_RDWR | O_NOFOLLOW | O_EXCL, S_IRUSR | S_IWUSR | O_CLOEXEC);
 			if (lockFd == -1)
 			{
 				if (errno == EEXIST)
diff --git a/Engine/Source/Programs/UnrealBuildAccelerator/Common/Private/UbaProcess.cpp b/Engine/Source/Programs/UnrealBuildAccelerator/Common/Private/UbaProcess.cpp
--- a/Engine/Source/Programs/UnrealBuildAccelerator/Common/Private/UbaProcess.cpp
+++ b/Engine/Source/Programs/UnrealBuildAccelerator/Common/Private/UbaProcess.cpp
@@ -1569,8 +1569,6 @@ namespace uba
 			pid_t processID;
 			while (true)
 			{
-				ScopedReadLock lock(m_session.m_hackToPreventETXTBUSY);
-
 				res = posix_spawnp(&processID, m_realApplication.c_str(), &fileActions, &attr, (char**)argsArray, (char**)envvars.data());
 				if (res == 0)
 					break;
diff --git a/Engine/Source/Programs/UnrealBuildAccelerator/Common/Private/UbaSession.cpp b/Engine/Source/Programs/UnrealBuildAccelerator/Common/Private/UbaSession.cpp
--- a/Engine/Source/Programs/UnrealBuildAccelerator/Common/Private/UbaSession.cpp
+++ b/Engine/Source/Programs/UnrealBuildAccelerator/Common/Private/UbaSession.cpp
@@ -583,7 +583,7 @@ namespace uba
 				auto mappedViewWrite = m_fileMappingBuffer.AllocAndMapView(MappedView_Transient, size, alignment, fileName);
 				auto unmapGuard = MakeGuard([&](){ m_fileMappingBuffer.UnmapView(mappedViewWrite, fileName); });
 
-				if (!m_storage.DecompressMemoryToMemory(readMemory, mappedViewWrite.memory, size))
+				if (!m_storage.DecompressMemoryToMemory(readMemory, mappedViewWrite.memory, size, fileName))
 					return false;
 				unmapGuard.Execute();
 
@@ -2313,7 +2313,7 @@ namespace uba
 			return m_cpuLoad;
 		totalTime = kernelTime + userTime;
 #elif PLATFORM_LINUX
-		int fd = open("/proc/stat", O_RDONLY);
+		int fd = open("/proc/stat", O_RDONLY | O_CLOEXEC);
 		if (fd != -1)
 		{
 			char buffer[512];
diff --git a/Engine/Source/Programs/UnrealBuildAccelerator/Common/Private/UbaSessionClient.cpp b/Engine/Source/Programs/UnrealBuildAccelerator/Common/Private/UbaSessionClient.cpp
--- a/Engine/Source/Programs/UnrealBuildAccelerator/Common/Private/UbaSessionClient.cpp
+++ b/Engine/Source/Programs/UnrealBuildAccelerator/Common/Private/UbaSessionClient.cpp
@@ -346,10 +346,6 @@ namespace uba
 				return false;
 		}
 
-		#if !PLATFORM_WINDOWS
-		ScopedWriteLock hackLock(m_hackToPreventETXTBUSY);
-		#endif
-
 		return m_storage.CopyOrLink(casKey, out.data, fileAttributes);
 	}
 
diff --git a/Engine/Source/Programs/UnrealBuildAccelerator/Common/Private/UbaStorage.cpp b/Engine/Source/Programs/UnrealBuildAccelerator/Common/Private/UbaStorage.cpp
--- a/Engine/Source/Programs/UnrealBuildAccelerator/Common/Private/UbaStorage.cpp
+++ b/Engine/Source/Programs/UnrealBuildAccelerator/Common/Private/UbaStorage.cpp
@@ -530,7 +530,7 @@ namespace uba
 		return false;
 	}
 
-	bool StorageImpl::DecompressMemoryToMemory(u8* compressedData, u8* writeData, u64 decompressedSize)
+	bool StorageImpl::DecompressMemoryToMemory(u8* compressedData, u8* writeData, u64 decompressedSize, const tchar* readHint)
 	{
 		UBA_ASSERT(compressedData);
 		UBA_ASSERT(writeData);
@@ -1447,7 +1447,7 @@ namespace uba
 								mem += sizeof(u64);
 								u8* dest = new u8[decompressedSize];
 								auto g = MakeGuard([dest]() { delete[] dest; });
-								if (!DecompressMemoryToMemory(mem, dest, decompressedSize))
+								if (!DecompressMemoryToMemory(mem, dest, decompressedSize, TC("")))
 									return;
 								checkedKey = CalculateCasKey(dest, decompressedSize, true);
 							}
@@ -1917,21 +1917,22 @@ namespace uba
 				MappedView mappedView;
 				auto mapViewGuard = MakeGuard([&](){ m_casDataBuffer.UnmapView(mappedView, destination); });
 
+				StringBuffer<512> casFile;
 				u64 decompressedSize;
 
 				if (casEntry->mappingHandle.IsValid())
 				{
-					mappedView = m_casDataBuffer.MapView(casEntry->mappingHandle, casEntry->mappingOffset, casEntry->mappingSize, CasKeyString(actualKey).str);
+					casFile.Append(CasKeyString(actualKey).str);
+					mappedView = m_casDataBuffer.MapView(casEntry->mappingHandle, casEntry->mappingOffset, casEntry->mappingSize, casFile.data);
 					compressedData = mappedView.memory;
 					if (!compressedData)
-						return m_logger.Error(TC("Failed to map view of mapping %s (%s)"), CasKeyString(casKey).str, LastErrorToText().data);
+						return m_logger.Error(TC("Failed to map view of mapping %s (%s)"), casFile.data, LastErrorToText().data);
 
 					decompressedSize = *(u64*)compressedData;
 					readData = compressedData + sizeof(u64);
 				}
 				else
 				{
-					StringBuffer<> casFile;
 #if !UBA_USE_SPARSEFILE
 					if (!StorageImpl::GetCasFileName(casFile, actualKey))
 						return false;
@@ -1978,7 +1979,7 @@ namespace uba
 
 					if (casEntry->mappingHandle.IsValid())
 					{
-						if (!DecompressMemoryToMemory(readData, destinationFile.GetData(), decompressedSize))
+						if (!DecompressMemoryToMemory(readData, destinationFile.GetData(), decompressedSize, casFile.data))
 							return false;
 					}
 					else
@@ -2370,7 +2371,7 @@ namespace uba
 				return m_logger.Error(TC("Failed to map view of file mapping for %s (%s)"), fileName, LastErrorToText().data);
 			auto udg = MakeGuard([&]() { UnmapViewOfFile(fileData, compressedSize, fileName); });
 			
-			if (!DecompressMemoryToMemory(fileData + 8, dest, decompressedSize))
+			if (!DecompressMemoryToMemory(fileData + 8, dest, decompressedSize, fileName))
 				return false;
 		}
 		else
diff --git a/Engine/Source/Programs/UnrealBuildAccelerator/Common/Public/UbaSession.h b/Engine/Source/Programs/UnrealBuildAccelerator/Common/Public/UbaSession.h
--- a/Engine/Source/Programs/UnrealBuildAccelerator/Common/Public/UbaSession.h
+++ b/Engine/Source/Programs/UnrealBuildAccelerator/Common/Public/UbaSession.h
@@ -229,12 +229,6 @@ namespace uba
 		GetNextProcessFunction m_getNextProcessFunction;
 		CustomServiceFunction m_customServiceFunction;
 
-		// This is a horrible hack to try to prevent ETXTBUSY. Our theory is that if posix_spawn happens with the exact same timing as symlink has an open handle we end up with the child process cloning that handle
-		// Causing posix_spawn to fail on the executable that was created with the symlink
-		#if !PLATFORM_WINDOWS
-		ReaderWriterLock m_hackToPreventETXTBUSY;
-		#endif
-
 		friend class ProcessImpl;
 	};
 
diff --git a/Engine/Source/Programs/UnrealBuildAccelerator/Common/Public/UbaStorage.h b/Engine/Source/Programs/UnrealBuildAccelerator/Common/Public/UbaStorage.h
--- a/Engine/Source/Programs/UnrealBuildAccelerator/Common/Public/UbaStorage.h
+++ b/Engine/Source/Programs/UnrealBuildAccelerator/Common/Public/UbaStorage.h
@@ -33,7 +33,7 @@ namespace uba
 		virtual bool GetZone(StringBufferBase& out) = 0;
 
 		virtual bool DecompressFileToMemory(const tchar* fileName, FileHandle fileHandle, u8* dest, u64 decompressedSize) = 0;
-		virtual bool DecompressMemoryToMemory(u8* compressedData, u8* writeData, u64 decompressedSize) = 0;
+		virtual bool DecompressMemoryToMemory(u8* compressedData, u8* writeData, u64 decompressedSize, const tchar* readHint) = 0;
 		virtual bool CreateDirectory(const tchar* dir) = 0;
 		virtual bool DeleteCasForFile(const tchar* file) = 0;
 
@@ -134,7 +134,7 @@ namespace uba
 		bool AddCasFile(const tchar* fileName, const CasKey& casKey, bool deferCreation);
 		void CasEntryAccessed(const CasKey& casKey);
 		virtual bool IsDisallowedPath(const tchar* fileName);
-		virtual bool DecompressMemoryToMemory(u8* compressedData, u8* writeData, u64 decompressedSize) override;
+		virtual bool DecompressMemoryToMemory(u8* compressedData, u8* writeData, u64 decompressedSize, const tchar* readHint) override;
 		bool DecompressMemoryToFile(u8* compressedData, FileAccessor& destination, u64 decompressedSize, bool useNoBuffering);
 
 		void CasEntryAccessed(CasEntry& entry);
diff --git a/Engine/Source/Programs/UnrealBuildAccelerator/Core/Private/UbaPlatform.cpp b/Engine/Source/Programs/UnrealBuildAccelerator/Core/Private/UbaPlatform.cpp
--- a/Engine/Source/Programs/UnrealBuildAccelerator/Core/Private/UbaPlatform.cpp
+++ b/Engine/Source/Programs/UnrealBuildAccelerator/Core/Private/UbaPlatform.cpp
@@ -45,7 +45,7 @@ namespace uba
 		arc4random_buf(&out, 16);
 		return true;
 		#else
-		const int f = open("/dev/urandom", O_RDONLY);
+		const int f = open("/dev/urandom", O_RDONLY | O_CLOEXEC);
 		if (f == -1)
 			return false;
 		size_t bytesRead = read(f, &out, 16);
