[Uba]
* Changed so ClientSession struct is allocated using aligned_alloc instead of new.. we can't figure out why this is failing with new on the farm and noone can repro so this is a very nasty workaround until we do understand how it can go wrong.
* Fixed potential race condition in scheduler
* Reduced lock scope in StorageClient to remove chance of deadlock
* Enabled mimalloc on linux

diff --git a/Engine/Source/Programs/UnrealBuildAccelerator/Common/Private/UbaSessionServer.cpp b/Engine/Source/Programs/UnrealBuildAccelerator/Common/Private/UbaSessionServer.cpp
--- a/Engine/Source/Programs/UnrealBuildAccelerator/Common/Private/UbaSessionServer.cpp
+++ b/Engine/Source/Programs/UnrealBuildAccelerator/Common/Private/UbaSessionServer.cpp
@@ -220,7 +220,10 @@ namespace uba
 		}
 
 		for (auto s : m_clientSessions)
-			delete s;
+		{
+			s->~ClientSession();
+			aligned_free(s);
+		}
 		m_clientSessions.clear();
 
 		#if 0
@@ -591,10 +594,12 @@ namespace uba
 				StringBuffer<256> info;
 				reader.ReadString(info);
 
+				// I have no explanation for this. On linux we get a shutdown crash when running through UBT if session is allocated with normal new
+				// For now we will work around it by using aligned_alloc which seems to be working on all platforms
+				auto& session = *new (aligned_alloc(alignof(ClientSession), sizeof(ClientSession))) ClientSession();
 				ScopedCriticalSection lock(m_remoteProcessAndSessionLock);
-				m_clientSessions.push_back(new ClientSession());
+				m_clientSessions.push_back(&session);
 				u32 sessionId = u32(m_clientSessions.size());
-				auto& session = *m_clientSessions.back();
 				session.name = name.data;
 				session.id = connectionInfo.GetId();
 				session.processSlotCount = processSlotCount;
@@ -852,6 +857,7 @@ namespace uba
 				u32 sessionId = reader.ReadU32();
 				u32 sessionIndex = sessionId - 1;
 				ScopedCriticalSection lock(m_remoteProcessAndSessionLock);
+				UBA_ASSERT(sessionIndex < m_clientSessions.size());
 				ClientSession& session = *m_clientSessions[sessionIndex];
 				lock.Leave();
 
@@ -869,6 +875,7 @@ namespace uba
 				u32 sessionId = reader.ReadU32();
 				u32 sessionIndex = sessionId - 1;
 				ScopedCriticalSection lock(m_remoteProcessAndSessionLock);
+				UBA_ASSERT(sessionIndex < m_clientSessions.size());
 				ClientSession& session = *m_clientSessions[sessionIndex];
 				lock.Leave();
 				WriteDirectoryTable(session, reader, writer);
@@ -896,6 +903,7 @@ namespace uba
 				u32 sessionIndex = sessionId - 1;
 
 				ScopedCriticalSection sessionsLock(m_remoteProcessAndSessionLock);
+				UBA_ASSERT(sessionIndex < m_clientSessions.size());
 				ClientSession& session = *m_clientSessions[sessionIndex];
 				sessionsLock.Leave();
 
@@ -1004,6 +1012,7 @@ namespace uba
 					return true;
 				}
 				u32 sessionIndex = process.m_sessionId - 1;
+				UBA_ASSERT(sessionIndex < m_clientSessions.size());
 				auto& session = *m_clientSessions[sessionIndex];
 				++m_finishedRemoteProcessCount;
 				--session.usedSlotCount;
@@ -1071,6 +1080,7 @@ namespace uba
 					return true;
 				}
 				u32 sessionIndex = process->m_sessionId - 1;
+				UBA_ASSERT(sessionIndex < m_clientSessions.size());
 				auto& session = *m_clientSessions[sessionIndex];
 				--session.usedSlotCount;
 				if (session.enabled)
@@ -1117,6 +1127,7 @@ namespace uba
 
 				u32 sessionIndex = sessionId - 1;
 				ScopedCriticalSection lock(m_remoteProcessAndSessionLock);
+				UBA_ASSERT(sessionIndex < m_clientSessions.size());
 				auto& session = *m_clientSessions[sessionIndex];
 				session.lastPing = lastPing;
 				session.memAvail = memAvail;
@@ -1403,6 +1414,7 @@ namespace uba
 		while (true)
 		{
 			ScopedCriticalSection queueLock(m_remoteProcessAndSessionLock);
+			UBA_ASSERT(sessionIndex < m_clientSessions.size());
 			auto& session = *m_clientSessions[sessionIndex];
 
 			while (!m_queuedRemoteProcesses.empty())
diff --git a/Engine/Source/Programs/UnrealBuildAccelerator/Common/Private/UbaStorageClient.cpp b/Engine/Source/Programs/UnrealBuildAccelerator/Common/Private/UbaStorageClient.cpp
--- a/Engine/Source/Programs/UnrealBuildAccelerator/Common/Private/UbaStorageClient.cpp
+++ b/Engine/Source/Programs/UnrealBuildAccelerator/Common/Private/UbaStorageClient.cpp
@@ -657,7 +657,7 @@ namespace uba
 	bool StorageClient::HasCasFile(const CasKey& casKey, CasEntry** out)
 	{
 		CasKey localKey = AsCompressed(casKey, false);
-		ScopedReadLock tempLock(m_localStorageFilesLock);
+		ScopedReadLock lock(m_localStorageFilesLock);
 		auto findIt = m_localStorageFiles.find(localKey);
 		if (findIt != m_localStorageFiles.end())
 		{
@@ -665,6 +665,7 @@ namespace uba
 				*out = &findIt->second.casEntry;
 			return true;
 		}
+		lock.Leave();
 		return StorageImpl::HasCasFile(casKey, out);
 	}
 
diff --git a/Engine/Source/Programs/UnrealBuildAccelerator/Common/Public/UbaSessionServer.h b/Engine/Source/Programs/UnrealBuildAccelerator/Common/Public/UbaSessionServer.h
--- a/Engine/Source/Programs/UnrealBuildAccelerator/Common/Public/UbaSessionServer.h
+++ b/Engine/Source/Programs/UnrealBuildAccelerator/Common/Public/UbaSessionServer.h
@@ -118,6 +118,9 @@ namespace uba
 		struct ClientSession
 		{
 			TString name;
+			UnorderedSet<CasKey> sentKeys;
+			ReaderWriterLock dirTablePosLock;
+			u32 dirTablePos = 0;
 			u32 id = ~0u;
 			u32 processSlotCount = 1;
 			u32 usedSlotCount = 0;
@@ -127,12 +130,6 @@ namespace uba
 			float cpuLoad = 0;
 			bool enabled = true;
 			bool dedicated = false;
-
-			ReaderWriterLock dirTablePosLock;
-			u32 dirTablePos = 0;
-
-			UnorderedSet<CasKey> sentKeys;
-
 		};
 		Vector<ClientSession*> m_clientSessions;
 
diff --git a/Engine/Source/Programs/UnrealBuildAccelerator/Mimalloc/UbaMimalloc.Build.cs b/Engine/Source/Programs/UnrealBuildAccelerator/Mimalloc/UbaMimalloc.Build.cs
--- a/Engine/Source/Programs/UnrealBuildAccelerator/Mimalloc/UbaMimalloc.Build.cs
+++ b/Engine/Source/Programs/UnrealBuildAccelerator/Mimalloc/UbaMimalloc.Build.cs
@@ -11,7 +11,7 @@ public class UbaMimalloc : ModuleRules
 	{
 		PCHUsage = PCHUsageMode.NoPCHs;
 
-		if (Target.Platform.IsInGroup(UnrealPlatformGroup.Windows))// || Target.Platform.IsInGroup(UnrealPlatformGroup.Linux))
+		if (Target.Platform.IsInGroup(UnrealPlatformGroup.Windows) || Target.Platform.IsInGroup(UnrealPlatformGroup.Linux))
 		{
 			PublicDependencyModuleNames.AddRange(new string[] {
 				"mimalloc212",
diff --git a/Engine/Source/Programs/UnrealBuildAccelerator/UbaAgent.Target.cs b/Engine/Source/Programs/UnrealBuildAccelerator/UbaAgent.Target.cs
--- a/Engine/Source/Programs/UnrealBuildAccelerator/UbaAgent.Target.cs
+++ b/Engine/Source/Programs/UnrealBuildAccelerator/UbaAgent.Target.cs
@@ -105,7 +105,7 @@ public class UbaAgentTarget : TargetRules
 		}
 
 		int useMiMalloc = 0;
-		if (Target.Platform.IsInGroup(UnrealPlatformGroup.Windows))// || Target.Platform.IsInGroup(UnrealPlatformGroup.Linux))
+		if (Target.Platform.IsInGroup(UnrealPlatformGroup.Windows) || Target.Platform.IsInGroup(UnrealPlatformGroup.Linux))
 		{
 			if (Rules.bUsePCHFiles)
 			{
