[Uba]
* Fixed deadlock in SessionClient related to where ProcessHandle is destroyed.

diff --git a/Engine/Source/Programs/UnrealBuildAccelerator/Common/Private/UbaSessionClient.cpp b/Engine/Source/Programs/UnrealBuildAccelerator/Common/Private/UbaSessionClient.cpp
--- a/Engine/Source/Programs/UnrealBuildAccelerator/Common/Private/UbaSessionClient.cpp
+++ b/Engine/Source/Programs/UnrealBuildAccelerator/Common/Private/UbaSessionClient.cpp
@@ -1282,8 +1282,8 @@ namespace uba
 			ProcessRec(ProcessImpl* impl) : handle(impl) {}
 			ProcessHandle handle;
 			ReaderWriterLock lock;
-			Atomic<bool> isKilled;
-			Atomic<bool> isDone;
+			bool isKilled = false;
+			bool isDone = false;
 			float weight = 1.0f;
 		};
 		List<ProcessRec> activeProcesses;
@@ -1314,11 +1314,13 @@ namespace uba
 			for (auto it=activeProcesses.begin();it!=activeProcesses.end();)
 			{
 				ProcessRec& r = *it;
+				ScopedWriteLock lock(r.lock);
 				if (!r.isDone)
 				{
 					++it;
 					continue;
 				}
+				lock.Leave();
 				it = activeProcesses.erase(it);
 			}
 
