[Uba]
* Added assert when pthread_rwlock_init or pthread_rwlock_destroy fails

diff --git a/Engine/Source/Programs/UnrealBuildAccelerator/Core/Private/UbaSynchronization.cpp b/Engine/Source/Programs/UnrealBuildAccelerator/Core/Private/UbaSynchronization.cpp
--- a/Engine/Source/Programs/UnrealBuildAccelerator/Core/Private/UbaSynchronization.cpp
+++ b/Engine/Source/Programs/UnrealBuildAccelerator/Core/Private/UbaSynchronization.cpp
@@ -2,6 +2,7 @@
 
 #include "UbaSynchronization.h"
 #include "UbaPlatform.h"
+#include "UbaStringBuffer.h"
 
 namespace uba
 {
@@ -18,7 +19,8 @@ namespace uba
 		pthread_mutexattr_t attr;
 		pthread_mutexattr_init(&attr);
 		pthread_mutexattr_settype(&attr, PTHREAD_MUTEX_RECURSIVE);
-		pthread_mutex_init((pthread_mutex_t*)data, &attr);
+		int res = pthread_mutex_init((pthread_mutex_t*)data, &attr);(void)res;
+		UBA_ASSERTF(res == 0, TC("pthread_mutex_init failed: %i"), res);
 		#endif
 	}
 
@@ -27,7 +29,8 @@ namespace uba
 		#if PLATFORM_WINDOWS
 		DeleteCriticalSection((CRITICAL_SECTION*)&data);
 		#else
-		pthread_mutex_destroy((pthread_mutex_t*)data);
+		int res = pthread_mutex_destroy((pthread_mutex_t*)data);(void)res;
+		UBA_ASSERTF(res == 0, TC("pthread_mutex_destroy failed: %i"), res);
 		#endif
 	}
 
@@ -58,14 +61,16 @@ namespace uba
 		#else
 		static_assert(alignof(pthread_rwlock_t) == alignof(ReaderWriterLock));
 		static_assert(sizeof(data) >= sizeof(pthread_rwlock_t));
-		pthread_rwlock_init((pthread_rwlock_t*)data, NULL);
+		int res = pthread_rwlock_init((pthread_rwlock_t*)data, NULL);(void)res;
+		UBA_ASSERTF(res == 0, TC("pthread_rwlock_init failed: %i"), res);
 		#endif
 	}
 
 	ReaderWriterLock::~ReaderWriterLock()
 	{
 		#if !PLATFORM_WINDOWS
-		pthread_rwlock_destroy((pthread_rwlock_t*)data);
+		int res = pthread_rwlock_destroy((pthread_rwlock_t*)data);(void)res;
+		UBA_ASSERTF(res == 0, TC("pthread_rwlock_destroy failed: %i"), res);
 		#endif
 	}
 
