[Uba]
* Proper fix of UbaAgent shutdown deadlock. Missed that scope guard was explicitly called earlier.. added separate scope guard

diff --git a/Engine/Source/Programs/UnrealBuildAccelerator/Common/Private/UbaSessionClient.cpp b/Engine/Source/Programs/UnrealBuildAccelerator/Common/Private/UbaSessionClient.cpp
--- a/Engine/Source/Programs/UnrealBuildAccelerator/Common/Private/UbaSessionClient.cpp
+++ b/Engine/Source/Programs/UnrealBuildAccelerator/Common/Private/UbaSessionClient.cpp
@@ -1494,16 +1494,17 @@ namespace uba
 						if (session.m_shouldSendLogToServer)
 							session.SendLogFileToServer(*(ProcessImpl*)h.m_process);
 
-						ScopedWriteLock lock(rec->lock);
-
+						float weight = rec->weight;
 						auto decreaseWeight = MakeGuard([&]()
 							{
 								ScopedWriteLock weightLock(activeWeightLock);
-								activeWeight -= rec->weight;
-								rec->isDone = true;
+								activeWeight -= weight;
 								session.m_waitToSendEvent.Set();
 							});
 
+						ScopedWriteLock lock(rec->lock);
+						auto doneGuard = MakeGuard([&]() { rec->isDone = true; session.m_waitToSendEvent.Set(); });
+
 						if (rec->isKilled)
 							return;
 
