[UBA]
* Fixed read-after-free bug found by TSAN in StorageServer.cpp
* Fixed write-after-free bug found in UbaProcess.cpp (m_cancelEvent.Set() could be called after memory was freed)
* Lots of minor fixes found using TSAN. Most are harmless but still nice to cleanup
* Disabled mimalloc for linux again.. seems like tsan does not like it so maybe there are bugs in it

diff --git a/Engine/Source/Programs/UnrealBuildAccelerator/Agent/Private/UbaAgent.cpp b/Engine/Source/Programs/UnrealBuildAccelerator/Agent/Private/UbaAgent.cpp
--- a/Engine/Source/Programs/UnrealBuildAccelerator/Agent/Private/UbaAgent.cpp
+++ b/Engine/Source/Programs/UnrealBuildAccelerator/Agent/Private/UbaAgent.cpp
@@ -1098,6 +1098,8 @@ namespace uba
 			info.memWaitLoadPercent = u8(memWaitLoadPercent);
 			info.memKillLoadPercent = u8(memKillLoadPercent);
 
+			if (!quiet)
+				info.processFinished = processFinished;
 
 			sessionClient = new SessionClient(info);
 			auto secsg = MakeGuard([&]() { delete sessionClient; });
@@ -1134,9 +1136,6 @@ namespace uba
 				loggingThread.Wait();
 			});
 
-			if (!quiet)
-				sessionClient->m_processFinished = processFinished;
-
 			// We got version mismatch and have the cas keys for the needed Agent/Detours binaries
 			if (keys[0] != CasKeyZero)
 			{
diff --git a/Engine/Source/Programs/UnrealBuildAccelerator/Common/Private/UbaNetworkBackendTcp.cpp b/Engine/Source/Programs/UnrealBuildAccelerator/Common/Private/UbaNetworkBackendTcp.cpp
--- a/Engine/Source/Programs/UnrealBuildAccelerator/Common/Private/UbaNetworkBackendTcp.cpp
+++ b/Engine/Source/Programs/UnrealBuildAccelerator/Common/Private/UbaNetworkBackendTcp.cpp
@@ -50,7 +50,7 @@ namespace uba
 		u16 port;
 		Thread thread;
 		Event listening;
-		SOCKET socket = INVALID_SOCKET;
+		Atomic<SOCKET> socket = INVALID_SOCKET;
 	};
 
 	struct NetworkBackendTcp::Connection
diff --git a/Engine/Source/Programs/UnrealBuildAccelerator/Common/Private/UbaProcess.cpp b/Engine/Source/Programs/UnrealBuildAccelerator/Common/Private/UbaProcess.cpp
--- a/Engine/Source/Programs/UnrealBuildAccelerator/Common/Private/UbaProcess.cpp
+++ b/Engine/Source/Programs/UnrealBuildAccelerator/Common/Private/UbaProcess.cpp
@@ -472,15 +472,12 @@ namespace uba
 
 		m_hasExited = true;
 
-		#if !PLATFORM_WINDOWS
-		m_cancelEvent.~Event();
-		m_writeEvent.~Event();
-		m_readEvent.~Event();
-		#endif
-
 		{
 			#if !PLATFORM_WINDOWS
 			ScopedWriteLock lock(m_comMemoryLock);
+			m_cancelEvent.~Event();
+			m_writeEvent.~Event();
+			m_readEvent.~Event();
 			#endif
 			m_session.m_processCommunicationAllocator.Free(m_comMemory);
 			m_comMemory = {};
diff --git a/Engine/Source/Programs/UnrealBuildAccelerator/Common/Private/UbaSession.cpp b/Engine/Source/Programs/UnrealBuildAccelerator/Common/Private/UbaSession.cpp
--- a/Engine/Source/Programs/UnrealBuildAccelerator/Common/Private/UbaSession.cpp
+++ b/Engine/Source/Programs/UnrealBuildAccelerator/Common/Private/UbaSession.cpp
@@ -1738,7 +1738,6 @@ namespace uba
 			ScopedWriteLock writtenLock(msg.process.m_writtenFilesLock);
 			auto insres = msg.process.m_writtenFiles.try_emplace(name);
 			WrittenFile& writtenFile = insres.first->second;
-			writtenLock.Leave();
 
 			UBA_ASSERT(writtenFile.owner == nullptr || writtenFile.owner == &msg.process);
 			writtenFile.owner = &msg.process;
diff --git a/Engine/Source/Programs/UnrealBuildAccelerator/Common/Private/UbaSessionClient.cpp b/Engine/Source/Programs/UnrealBuildAccelerator/Common/Private/UbaSessionClient.cpp
--- a/Engine/Source/Programs/UnrealBuildAccelerator/Common/Private/UbaSessionClient.cpp
+++ b/Engine/Source/Programs/UnrealBuildAccelerator/Common/Private/UbaSessionClient.cpp
@@ -27,6 +27,7 @@ namespace uba
 		m_useBinariesAsVersion = info.useBinariesAsVersion;
 		m_memWaitLoadPercent = info.memWaitLoadPercent;
 		m_memKillLoadPercent = info.memKillLoadPercent;
+		m_processFinished = info.processFinished;
 
 		if (m_name.IsEmpty())
 		{
@@ -46,17 +47,17 @@ namespace uba
 			m_killRandomIndex = 10 + g.data1 % 30;
 		}
 
-		m_client.RegisterOnConnected([this]()
-			{
-				Connect();
-			});
+		m_nameToHashTableMem.Init(NameToHashMemSize);
 
 		m_client.RegisterOnDisconnected([this]()
 			{
 				m_loop = false;
 			});
 
-		m_nameToHashTableMem.Init(NameToHashMemSize);
+		m_client.RegisterOnConnected([this]()
+			{
+				Connect();
+			});
 	}
 
 	SessionClient::~SessionClient()
diff --git a/Engine/Source/Programs/UnrealBuildAccelerator/Common/Private/UbaSessionServer.cpp b/Engine/Source/Programs/UnrealBuildAccelerator/Common/Private/UbaSessionServer.cpp
--- a/Engine/Source/Programs/UnrealBuildAccelerator/Common/Private/UbaSessionServer.cpp
+++ b/Engine/Source/Programs/UnrealBuildAccelerator/Common/Private/UbaSessionServer.cpp
@@ -299,6 +299,8 @@ namespace uba
 		
 		ScopedCriticalSection lock(m_remoteProcessAndSessionLock);
 		m_queuedRemoteProcesses.push_back(remoteProcess);
+
+		ScopedReadLock lock2(m_remoteProcessReturnedEventLock);
 		if (m_remoteProcessReturnedEvent)
 		{
 			if (!m_remoteExecutionEnabled)
@@ -392,11 +394,13 @@ namespace uba
 
 	void SessionServer::SetRemoteProcessSlotAvailableEvent(const Function<void()>& remoteProcessSlotAvailableEvent)
 	{
+		ScopedWriteLock lock(m_remoteProcessSlotAvailableEventLock);
 		m_remoteProcessSlotAvailableEvent = remoteProcessSlotAvailableEvent;
 	}
 
 	void SessionServer::SetRemoteProcessReturnedEvent(const Function<void(Process&)>& remoteProcessReturnedEvent)
 	{
+		ScopedWriteLock lock(m_remoteProcessReturnedEventLock);
 		m_remoteProcessReturnedEvent = remoteProcessReturnedEvent;
 	}
 
@@ -1408,6 +1412,7 @@ namespace uba
 
 	SessionServer::RemoteProcess* SessionServer::DequeueProcess(u32 sessionId, u32 clientId)
 	{
+		ScopedReadLock lock(m_remoteProcessSlotAvailableEventLock);
 		bool hasCalledCallback = !m_remoteProcessSlotAvailableEvent;
 		u32 sessionIndex = sessionId - 1;
 
diff --git a/Engine/Source/Programs/UnrealBuildAccelerator/Common/Private/UbaStorage.cpp b/Engine/Source/Programs/UnrealBuildAccelerator/Common/Private/UbaStorage.cpp
--- a/Engine/Source/Programs/UnrealBuildAccelerator/Common/Private/UbaStorage.cpp
+++ b/Engine/Source/Programs/UnrealBuildAccelerator/Common/Private/UbaStorage.cpp
@@ -25,7 +25,8 @@ namespace uba
 
 	void StorageImpl::CasEntryAccessed(CasEntry& entry)
 	{
-		if (entry.mappingHandle.IsValid())
+		bool hasMapping = entry.lock.ScopedRead([&](){ return entry.mappingHandle.IsValid(); });
+		if (hasMapping)
 			return;
 
 		ScopedWriteLock lock(m_accessLock);
@@ -571,6 +572,7 @@ namespace uba
 						{
 							if (rec->written == rec->decompressedSize)
 								rec->done.Set();
+							lock.Leave();
 							if (!--rec->refCount)
 								delete rec;
 							return 0;
diff --git a/Engine/Source/Programs/UnrealBuildAccelerator/Common/Private/UbaStorageServer.cpp b/Engine/Source/Programs/UnrealBuildAccelerator/Common/Private/UbaStorageServer.cpp
--- a/Engine/Source/Programs/UnrealBuildAccelerator/Common/Private/UbaStorageServer.cpp
+++ b/Engine/Source/Programs/UnrealBuildAccelerator/Common/Private/UbaStorageServer.cpp
@@ -759,8 +759,9 @@ namespace uba
 				u64 time2 = GetTime();
 				activeStore.recvCasTime += time2 - start;
 
+				u64 fileSize = activeStore.fileSize;
 				u64 totalWritten = activeStore.totalWritten.fetch_add(toRead) + toRead;
-				if (totalWritten == activeStore.fileSize)
+				if (totalWritten == fileSize)
 				{
 					m_casDataBuffer.UnmapView(activeStore.mappedView, TC("StoreDone"));
 
diff --git a/Engine/Source/Programs/UnrealBuildAccelerator/Common/Public/UbaLogger.h b/Engine/Source/Programs/UnrealBuildAccelerator/Common/Public/UbaLogger.h
--- a/Engine/Source/Programs/UnrealBuildAccelerator/Common/Public/UbaLogger.h
+++ b/Engine/Source/Programs/UnrealBuildAccelerator/Common/Public/UbaLogger.h
@@ -65,14 +65,6 @@ namespace uba
 		u32 m_prefixLen;
 	};
 
-	struct MutableLogger : public LoggerWithWriter
-	{
-		MutableLogger(LogWriter& writer, const tchar* prefix) : LoggerWithWriter(writer, prefix) {}
-		virtual void Log(LogEntryType type, const tchar* str, u32 strLen) override { if (!isMuted) LoggerWithWriter::Log(type, str, strLen); }
-		bool isMuted = false;
-	};
-
-
 	class FilteredLogWriter : public LogWriter
 	{
 	public:
diff --git a/Engine/Source/Programs/UnrealBuildAccelerator/Common/Public/UbaNetworkClient.h b/Engine/Source/Programs/UnrealBuildAccelerator/Common/Public/UbaNetworkClient.h
--- a/Engine/Source/Programs/UnrealBuildAccelerator/Common/Public/UbaNetworkClient.h
+++ b/Engine/Source/Programs/UnrealBuildAccelerator/Common/Public/UbaNetworkClient.h
@@ -72,7 +72,7 @@ namespace uba
 			void* backendConnection = nullptr;
 			Atomic<u32> connected;
 			Timer sendTimer;
-			u64 sendBytes = 0;
+			Atomic<u64> sendBytes;
 			u64 recvBytes = 0;
 			u32 recvCount = 0;
 			NetworkBackend* backend = nullptr;
diff --git a/Engine/Source/Programs/UnrealBuildAccelerator/Common/Public/UbaProcess.h b/Engine/Source/Programs/UnrealBuildAccelerator/Common/Public/UbaProcess.h
--- a/Engine/Source/Programs/UnrealBuildAccelerator/Common/Public/UbaProcess.h
+++ b/Engine/Source/Programs/UnrealBuildAccelerator/Common/Public/UbaProcess.h
@@ -119,8 +119,8 @@ namespace uba
 		u32 m_exitCode = ~0u;
 		u32 m_rulesIndex = 0;
 		u32 m_messageCount = 0;
+		Atomic<bool> m_hasExited;
 		bool m_messageSuccess = true;
-		bool m_hasExited = false;
 		bool m_echoOn = true;
 		bool m_gotExitMessage = false;
 		bool m_parentReportedExit = false;
diff --git a/Engine/Source/Programs/UnrealBuildAccelerator/Common/Public/UbaScheduler.h b/Engine/Source/Programs/UnrealBuildAccelerator/Common/Public/UbaScheduler.h
--- a/Engine/Source/Programs/UnrealBuildAccelerator/Common/Public/UbaScheduler.h
+++ b/Engine/Source/Programs/UnrealBuildAccelerator/Common/Public/UbaScheduler.h
@@ -107,7 +107,7 @@ namespace uba
 
 		Event m_updateThreadLoop;
 		Thread m_thread;
-		bool m_loop = false;
+		Atomic<bool> m_loop;
 		bool m_enableProcessReuse;
 		bool m_forceRemote;
 
diff --git a/Engine/Source/Programs/UnrealBuildAccelerator/Common/Public/UbaSession.h b/Engine/Source/Programs/UnrealBuildAccelerator/Common/Public/UbaSession.h
--- a/Engine/Source/Programs/UnrealBuildAccelerator/Common/Public/UbaSession.h
+++ b/Engine/Source/Programs/UnrealBuildAccelerator/Common/Public/UbaSession.h
@@ -139,6 +139,14 @@ namespace uba
 		void ThreadTraceLoop();
 		virtual void TraceSessionUpdate();
 
+		struct MutableLogger : public LoggerWithWriter
+		{
+			MutableLogger(LogWriter& writer, const tchar* prefix) : LoggerWithWriter(writer, prefix) {}
+			virtual void Log(LogEntryType type, const tchar* str, u32 strLen) override { if (!isMuted) LoggerWithWriter::Log(type, str, strLen); }
+			Atomic<bool> isMuted;
+		};
+
+
 		Storage& m_storage;
 		MutableLogger m_logger;
 
diff --git a/Engine/Source/Programs/UnrealBuildAccelerator/Common/Public/UbaSessionClient.h b/Engine/Source/Programs/UnrealBuildAccelerator/Common/Public/UbaSessionClient.h
--- a/Engine/Source/Programs/UnrealBuildAccelerator/Common/Public/UbaSessionClient.h
+++ b/Engine/Source/Programs/UnrealBuildAccelerator/Common/Public/UbaSessionClient.h
@@ -26,6 +26,7 @@ namespace uba
 		bool useBinariesAsVersion = false;
 		bool killRandom = false;
 		bool useStorage = true;
+		Function<void(const ProcessHandle&)> processFinished;
 	};
 
 	class SessionClient final : public Session
@@ -40,6 +41,7 @@ namespace uba
 
 		u64 GetBestPing();
 
+	private:
 		bool RetrieveCasFile(CasKey& outNewKey, u64& outSize, const CasKey& casKey, const tchar* hint, bool storeUncompressed, bool allowProxy = true);
 
 		virtual bool PrepareProcess(const ProcessStartInfo& startInfo, bool isChild, StringBufferBase& outRealApplication, const tchar*& outRealWorkingDir) override;
@@ -139,7 +141,7 @@ namespace uba
 
 		SessionSummaryStats m_stats;
 
-		u64 m_bestPing = 0;
+		Atomic<u64> m_bestPing;
 		u64 m_lastPing = 0;
 		u64 m_lastPingSendTime = 0;
 	};
diff --git a/Engine/Source/Programs/UnrealBuildAccelerator/Common/Public/UbaSessionServer.h b/Engine/Source/Programs/UnrealBuildAccelerator/Common/Public/UbaSessionServer.h
--- a/Engine/Source/Programs/UnrealBuildAccelerator/Common/Public/UbaSessionServer.h
+++ b/Engine/Source/Programs/UnrealBuildAccelerator/Common/Public/UbaSessionServer.h
@@ -100,7 +100,10 @@ namespace uba
 
 		static constexpr u8 ServiceId = SessionServiceId;
 
+		ReaderWriterLock m_remoteProcessSlotAvailableEventLock;
 		Function<void()> m_remoteProcessSlotAvailableEvent;
+
+		ReaderWriterLock m_remoteProcessReturnedEventLock;
 		Function<void(Process&)> m_remoteProcessReturnedEvent;
 
 		CriticalSection m_remoteProcessAndSessionLock; // Can be re-entrant.
diff --git a/Engine/Source/Programs/UnrealBuildAccelerator/Core/Private/UbaSynchronization.cpp b/Engine/Source/Programs/UnrealBuildAccelerator/Core/Private/UbaSynchronization.cpp
--- a/Engine/Source/Programs/UnrealBuildAccelerator/Core/Private/UbaSynchronization.cpp
+++ b/Engine/Source/Programs/UnrealBuildAccelerator/Core/Private/UbaSynchronization.cpp
@@ -27,6 +27,12 @@ namespace uba
 	CriticalSection::~CriticalSection()
 	{
 		#if PLATFORM_WINDOWS
+		#if UBA_DEBUG
+		if (TryEnterCriticalSection((CRITICAL_SECTION*)&data))
+			LeaveCriticalSection((CRITICAL_SECTION*)&data);
+		else
+			UBA_ASSERT(false);
+		#endif
 		DeleteCriticalSection((CRITICAL_SECTION*)&data);
 		#else
 		int res = pthread_mutex_destroy((pthread_mutex_t*)data);(void)res;
@@ -68,7 +74,14 @@ namespace uba
 
 	ReaderWriterLock::~ReaderWriterLock()
 	{
-		#if !PLATFORM_WINDOWS
+		#if PLATFORM_WINDOWS
+		#if UBA_DEBUG
+		if (TryAcquireSRWLockExclusive((SRWLOCK*)&data))
+			ReleaseSRWLockExclusive((SRWLOCK*)&data);
+		else
+			UBA_ASSERT(false);
+		#endif
+		#else
 		int res = pthread_rwlock_destroy((pthread_rwlock_t*)data);(void)res;
 		UBA_ASSERTF(res == 0, TC("pthread_rwlock_destroy failed: %i"), res);
 		#endif
diff --git a/Engine/Source/Programs/UnrealBuildAccelerator/Core/Public/UbaProcessStats.h b/Engine/Source/Programs/UnrealBuildAccelerator/Core/Public/UbaProcessStats.h
--- a/Engine/Source/Programs/UnrealBuildAccelerator/Core/Public/UbaProcessStats.h
+++ b/Engine/Source/Programs/UnrealBuildAccelerator/Core/Public/UbaProcessStats.h
@@ -46,16 +46,16 @@ namespace uba
 		UBA_PROCESS_STATS
 		#undef UBA_PROCESS_STAT
 
-		u64 startupTime = 0;
-		u64 exitTime = 0;
+		AtomicU64 startupTime = 0;
+		AtomicU64 exitTime = 0;
 
 		// Don't add in GetTotalTime()
-		u64 wallTime = 0;
-		u64 cpuTime = 0;
+		AtomicU64 wallTime = 0;
+		AtomicU64 cpuTime = 0;
 
-		u32 usedMemory = 0;
+		AtomicU64 usedMemory = 0;
 
-		u64 hostTotalTime = 0;
+		AtomicU64 hostTotalTime;
 
 		void Print(Logger& logger, u64 frequency = GetFrequency());
 
@@ -105,7 +105,7 @@ namespace uba
 			writer.WriteU64(exitTime);
 			writer.WriteU64(wallTime);
 			writer.WriteU64(cpuTime);
-			writer.WriteU32(usedMemory);
+			writer.WriteU32(u32(usedMemory));
 			writer.WriteU64(hostTotalTime);
 		}
 
diff --git a/Engine/Source/Programs/UnrealBuildAccelerator/Mimalloc/UbaMimalloc.Build.cs b/Engine/Source/Programs/UnrealBuildAccelerator/Mimalloc/UbaMimalloc.Build.cs
--- a/Engine/Source/Programs/UnrealBuildAccelerator/Mimalloc/UbaMimalloc.Build.cs
+++ b/Engine/Source/Programs/UnrealBuildAccelerator/Mimalloc/UbaMimalloc.Build.cs
@@ -11,7 +11,7 @@ public class UbaMimalloc : ModuleRules
 	{
 		PCHUsage = PCHUsageMode.NoPCHs;
 
-		if (Target.Platform.IsInGroup(UnrealPlatformGroup.Windows) || Target.Platform.IsInGroup(UnrealPlatformGroup.Linux))
+		if (Target.Platform.IsInGroup(UnrealPlatformGroup.Windows))// || Target.Platform.IsInGroup(UnrealPlatformGroup.Linux))
 		{
 			PublicDependencyModuleNames.AddRange(new string[] {
 				"mimalloc212",
diff --git a/Engine/Source/Programs/UnrealBuildAccelerator/UbaAgent.Target.cs b/Engine/Source/Programs/UnrealBuildAccelerator/UbaAgent.Target.cs
--- a/Engine/Source/Programs/UnrealBuildAccelerator/UbaAgent.Target.cs
+++ b/Engine/Source/Programs/UnrealBuildAccelerator/UbaAgent.Target.cs
@@ -105,7 +105,7 @@ public class UbaAgentTarget : TargetRules
 		}
 
 		int useMiMalloc = 0;
-		if (Target.Platform.IsInGroup(UnrealPlatformGroup.Windows) || Target.Platform.IsInGroup(UnrealPlatformGroup.Linux))
+		if (Target.Platform.IsInGroup(UnrealPlatformGroup.Windows))// || Target.Platform.IsInGroup(UnrealPlatformGroup.Linux))
 		{
 			if (Rules.bUsePCHFiles)
 			{
