[Uba]
* Fixed a bug in UbaWorkManager shutdown where intrusive linked list contains bad instances while deleting workers
* Fixed more tsan errors when doing ctrl-c on host

diff --git a/Engine/Source/Programs/UnrealBuildAccelerator/Common/Private/UbaNetworkBackendTcp.cpp b/Engine/Source/Programs/UnrealBuildAccelerator/Common/Private/UbaNetworkBackendTcp.cpp
--- a/Engine/Source/Programs/UnrealBuildAccelerator/Common/Private/UbaNetworkBackendTcp.cpp
+++ b/Engine/Source/Programs/UnrealBuildAccelerator/Common/Private/UbaNetworkBackendTcp.cpp
@@ -57,7 +57,7 @@ namespace uba
 	{
 		Connection(SOCKET s) : socket(s), ready(true) {}
 
-		SOCKET socket;
+		Atomic<SOCKET> socket;
 		Event ready;
 		u32 headerSize = 0;
 		u32 recvTimeoutMs = 0;
diff --git a/Engine/Source/Programs/UnrealBuildAccelerator/Common/Private/UbaProcess.cpp b/Engine/Source/Programs/UnrealBuildAccelerator/Common/Private/UbaProcess.cpp
--- a/Engine/Source/Programs/UnrealBuildAccelerator/Common/Private/UbaProcess.cpp
+++ b/Engine/Source/Programs/UnrealBuildAccelerator/Common/Private/UbaProcess.cpp
@@ -332,6 +332,7 @@ namespace uba
 		m_cancelEvent.Set();
 		#else
 		m_cancelled = true;
+		ScopedWriteLock lock(m_comMemoryLock);
 		if (m_comMemory.memory)
 			m_cancelEvent.Set();
 		#endif
diff --git a/Engine/Source/Programs/UnrealBuildAccelerator/Common/Private/UbaWorkManager.cpp b/Engine/Source/Programs/UnrealBuildAccelerator/Common/Private/UbaWorkManager.cpp
--- a/Engine/Source/Programs/UnrealBuildAccelerator/Common/Private/UbaWorkManager.cpp
+++ b/Engine/Source/Programs/UnrealBuildAccelerator/Common/Private/UbaWorkManager.cpp
@@ -15,9 +15,14 @@ namespace uba
 			m_thread.Start([&]() { ThreadWorker(manager); return 0; });
 		}
 		~Worker()
+		{
+		}
+
+		void Stop()
 		{
 			m_loop = false;
 			m_workAvailable.Set();
+			m_thread.Wait();
 		}
 
 		void ThreadWorker(WorkManagerImpl& manager)
@@ -57,9 +62,9 @@ namespace uba
 
 		Worker* m_nextWorker = nullptr;
 		Worker* m_prevWorker = nullptr;
+		Atomic<bool> m_loop;
 		Event m_workAvailable;
 		Thread m_thread;
-		bool m_loop = false;
 	};
 
 	WorkManagerImpl::WorkManagerImpl(u32 workerCount)
@@ -72,8 +77,10 @@ namespace uba
 
 	WorkManagerImpl::~WorkManagerImpl()
 	{
-		for (u32 i = 0; i != m_workers.size(); ++i)
-			delete m_workers[i];
+		for (auto worker : m_workers)
+			worker->Stop();
+		for (auto worker : m_workers)
+			delete worker;
 	}
 
 
diff --git a/Engine/Source/Programs/UnrealBuildAccelerator/Common/Public/UbaNetworkMessage.h b/Engine/Source/Programs/UnrealBuildAccelerator/Common/Public/UbaNetworkMessage.h
--- a/Engine/Source/Programs/UnrealBuildAccelerator/Common/Public/UbaNetworkMessage.h
+++ b/Engine/Source/Programs/UnrealBuildAccelerator/Common/Public/UbaNetworkMessage.h
@@ -28,7 +28,7 @@ namespace uba
 		u32 m_responseSize = 0;
 		u32 m_responseCapacity = 0;
 		u16 m_id = 0;
-		bool m_error = false;
+		Atomic<bool> m_error;
 		NetworkClient::Connection* m_connection = nullptr;
 		friend NetworkClient;
 		NetworkMessage(const NetworkMessage&) = delete;
diff --git a/Engine/Source/Programs/UnrealBuildAccelerator/Common/Public/UbaProcess.h b/Engine/Source/Programs/UnrealBuildAccelerator/Common/Public/UbaProcess.h
--- a/Engine/Source/Programs/UnrealBuildAccelerator/Common/Public/UbaProcess.h
+++ b/Engine/Source/Programs/UnrealBuildAccelerator/Common/Public/UbaProcess.h
@@ -97,7 +97,7 @@ namespace uba
 		Event m_readEvent;
 	#else
 		ReaderWriterLock m_comMemoryLock;
-		bool m_cancelled = false;
+		Atomic<bool> m_cancelled;
 		Event& m_cancelEvent;
 		Event& m_writeEvent;
 		Event& m_readEvent;
