[Uba]
* Fixed a UbaAgent shutdown hang that can happen with certain timings where workers go from used to non-used at the same time as they are being stopped

diff --git a/Engine/Source/Programs/UnrealBuildAccelerator/Agent/Private/UbaAgent.cpp b/Engine/Source/Programs/UnrealBuildAccelerator/Agent/Private/UbaAgent.cpp
--- a/Engine/Source/Programs/UnrealBuildAccelerator/Agent/Private/UbaAgent.cpp
+++ b/Engine/Source/Programs/UnrealBuildAccelerator/Agent/Private/UbaAgent.cpp
@@ -1006,7 +1006,7 @@ namespace uba
 				return 0;
 
 			Event wakeupSessionWait(false);
-			u32 targetConnectionCount = 1;
+			Atomic<u32> targetConnectionCount = 1;
 
 			struct Proxy
 			{
@@ -1014,7 +1014,7 @@ namespace uba
 				NetworkClient* client;
 				Event& wakeupSessionWait;
 				u32& maxConnectionCount;
-				u32& targetConnectionCount;
+				Atomic<u32>& targetConnectionCount;
 				NetworkServer* server = nullptr;
 				StorageProxy* storage = nullptr;
 				StorageClient* storageClient = nullptr;
@@ -1128,6 +1128,7 @@ namespace uba
 						LoggerWithWriter(g_consoleLogWriter, TC("")).Info(TC("Proxy timed out waiting for zero active fetches"));
 				if (proxy.server)
 					proxy.server->StopAll();
+				client->StopListen();
 				client->Disconnect();
 				sessionClient->Stop();
 				storageClient->StopProxy();
diff --git a/Engine/Source/Programs/UnrealBuildAccelerator/Common/Private/UbaNetworkServer.cpp b/Engine/Source/Programs/UnrealBuildAccelerator/Common/Private/UbaNetworkServer.cpp
--- a/Engine/Source/Programs/UnrealBuildAccelerator/Common/Private/UbaNetworkServer.cpp
+++ b/Engine/Source/Programs/UnrealBuildAccelerator/Common/Private/UbaNetworkServer.cpp
@@ -61,8 +61,10 @@ namespace uba
 			ScopedWriteLock lock(server.m_availableWorkersLock);
 			while (m_inUse)
 			{
+				m_context->workAvailable.Set();
 				lock.Leave();
-				Sleep(1);
+				if (m_thread.Wait(5))
+					break;
 				lock.Enter();
 			}
 		}
@@ -76,8 +78,8 @@ namespace uba
 
 		WorkerContext* m_context = nullptr;
 
-		bool m_loop = false;
-		bool m_inUse = false;
+		Atomic<bool> m_loop;
+		Atomic<bool> m_inUse;
 		Thread m_thread;
 
 		Worker(const Worker&) = delete;
@@ -504,6 +506,9 @@ namespace uba
 		while (m_context->workAvailable.IsSet(~0u) && m_loop)
 			Update(*m_context, true);
 		t_worker = nullptr;
+
+		if (m_inUse) // I have no idea how this can happen.. should not be possible. There is a path somewhere where it can leave while still being in use
+			server.PushWorker(this);
 	}
 
 	void NetworkServer::Worker::DoAdditionalWorkAndSignalAvailable(NetworkServer& server)
diff --git a/Engine/Source/Programs/UnrealBuildAccelerator/Common/Private/UbaSessionClient.cpp b/Engine/Source/Programs/UnrealBuildAccelerator/Common/Private/UbaSessionClient.cpp
--- a/Engine/Source/Programs/UnrealBuildAccelerator/Common/Private/UbaSessionClient.cpp
+++ b/Engine/Source/Programs/UnrealBuildAccelerator/Common/Private/UbaSessionClient.cpp
@@ -14,6 +14,7 @@ namespace uba
 	: Session(info, TC("UbaSessionClient"), true)
 	,	m_client(info.client)
 	,	m_name(info.name.data)
+	,	m_terminationTime(~0ull)
 	,	m_waitToSendEvent(false)
 	,	m_loop(true)
 	{
@@ -1328,7 +1329,7 @@ namespace uba
 			if (m_remoteExecutionEnabled && m_terminationReason)
 			{
 				m_remoteExecutionEnabled = false;
-				m_logger.Info(TC("%s. Will stop scheduling processes and send failing processes back for retry"), m_terminationReason);
+				m_logger.Info(TC("%s. Will stop scheduling processes and send failing processes back for retry"), m_terminationReason.load());
 			}
 
 			if (!activeProcesses.empty())
diff --git a/Engine/Source/Programs/UnrealBuildAccelerator/Common/Public/UbaSessionClient.h b/Engine/Source/Programs/UnrealBuildAccelerator/Common/Public/UbaSessionClient.h
--- a/Engine/Source/Programs/UnrealBuildAccelerator/Common/Public/UbaSessionClient.h
+++ b/Engine/Source/Programs/UnrealBuildAccelerator/Common/Public/UbaSessionClient.h
@@ -109,8 +109,9 @@ namespace uba
 		bool m_useStorage = true;
 		bool m_shouldSendLogToServer = false;
 		bool m_remoteExecutionEnabled = true;
-		const tchar* m_terminationReason = nullptr;
-		u64 m_terminationTime = ~0ull;
+		
+		Atomic<const tchar*> m_terminationReason;
+		Atomic<u64> m_terminationTime;
 
 		ReaderWriterLock m_handledApplicationEnvironmentsLock;
 		UnorderedSet<TString> m_handledApplicationEnvironments;
