[Uba]
* Added more description when failing to drop cas db entry
* Added cas entry lock around code that copy or link file out
* Added more fixes for things reported by tsan/asan

diff --git a/Engine/Source/Programs/UnrealBuildAccelerator/Common/Private/UbaNetworkBackendTcp.cpp b/Engine/Source/Programs/UnrealBuildAccelerator/Common/Private/UbaNetworkBackendTcp.cpp
--- a/Engine/Source/Programs/UnrealBuildAccelerator/Common/Private/UbaNetworkBackendTcp.cpp
+++ b/Engine/Source/Programs/UnrealBuildAccelerator/Common/Private/UbaNetworkBackendTcp.cpp
@@ -651,9 +651,9 @@ namespace uba
 
 		//char* ip = inet_ntoa(((sockaddr_in*)const_cast<sockaddr*>(&remoteSocketAddr))->sin_addr);
 		if (nameHint)
-			logger.Info(TC("Connected to %s:%u"), nameHint, ((sockaddr_in&)remoteSocketAddr).sin_port);
+			logger.Detail(TC("Connected to %s:%u"), nameHint, ((sockaddr_in&)remoteSocketAddr).sin_port);
 		else
-			logger.Info(TC("Connected using sockaddr"));
+			logger.Detail(TC("Connected using sockaddr"));
 
 		socketClose.Cancel();
 
diff --git a/Engine/Source/Programs/UnrealBuildAccelerator/Common/Private/UbaNetworkClient.cpp b/Engine/Source/Programs/UnrealBuildAccelerator/Common/Private/UbaNetworkClient.cpp
--- a/Engine/Source/Programs/UnrealBuildAccelerator/Common/Private/UbaNetworkClient.cpp
+++ b/Engine/Source/Programs/UnrealBuildAccelerator/Common/Private/UbaNetworkClient.cpp
@@ -47,7 +47,7 @@ namespace uba
 			connection.backend->SetDisconnectCallback(connection.backendConnection, nullptr, nullptr);
 			if (connection.connected.exchange(0) == 0)
 				continue;
-			m_logger.Info(TC("Disconnected from server..."));
+			m_logger.Detail(TC("Disconnected from server..."));
 			connection.backend->Shutdown(connection.backendConnection);
 			connection.backend->Close(connection.backendConnection);
 		}
@@ -385,7 +385,7 @@ namespace uba
 	{
 		if (connection.connected.exchange(0) == 1)
 		{
-			m_logger.Info(TC("Disconnected from server..."));
+			m_logger.Detail(TC("Disconnected from server..."));
 
 			connection.backend->Shutdown(connection.backendConnection);
 
diff --git a/Engine/Source/Programs/UnrealBuildAccelerator/Common/Private/UbaNetworkServer.cpp b/Engine/Source/Programs/UnrealBuildAccelerator/Common/Private/UbaNetworkServer.cpp
--- a/Engine/Source/Programs/UnrealBuildAccelerator/Common/Private/UbaNetworkServer.cpp
+++ b/Engine/Source/Programs/UnrealBuildAccelerator/Common/Private/UbaNetworkServer.cpp
@@ -254,10 +254,10 @@ namespace uba
 			{
 				if (server.m_onConnectionFunction)
 					server.m_onConnectionFunction(clientUid, clientId);
-				server.m_logger.Info(TC("Client %s connected on connection %s"), GuidToString(clientUid).str, GuidToString(connectionUid).str);
+				server.m_logger.Detail(TC("Client %s connected on connection %s"), GuidToString(clientUid).str, GuidToString(connectionUid).str);
 			}
 			else
-				server.m_logger.Info(TC("Client %s additional connection %s connected"), GuidToString(clientUid).str, GuidToString(connectionUid).str);
+				server.m_logger.Detail(TC("Client %s additional connection %s connected"), GuidToString(clientUid).str, GuidToString(connectionUid).str);
 
 
 			return true;
@@ -362,8 +362,8 @@ namespace uba
 		CryptoKey m_cryptoKey;
 		Atomic<int> m_activeWorkerCount;
 		Atomic<int> m_disconnectCalled;
+		Atomic<bool> m_disconnected;
 		bool m_shouldDisconnect = false;
-		bool m_disconnected = false;
 		void* m_backendConnection = nullptr;
 
 		Timer m_sendTimer;
@@ -650,17 +650,16 @@ namespace uba
 
 		{
 			ScopedWriteLock lock(m_connectionsLock);
-			auto connections(std::move(m_connections));
-			lock.Leave();
-
 			bool success = true;
-			for (auto& c : connections)
+			for (auto& c : m_connections)
 			{
 				success = c.Stop() && success;
 				m_sendTimer.Add(c.m_sendTimer);
 				m_encryptTimer.Add(c.m_encryptTimer);
 				m_decryptTimer.Add(c.m_decryptTimer);
 			}
+			m_connections.clear();
+			lock.Leave();
 
 			// If stopping connections fail we need to abort because we will most likely run into a deadlock when deleting the workers.
 			if (!success)
diff --git a/Engine/Source/Programs/UnrealBuildAccelerator/Common/Private/UbaSessionServer.cpp b/Engine/Source/Programs/UnrealBuildAccelerator/Common/Private/UbaSessionServer.cpp
--- a/Engine/Source/Programs/UnrealBuildAccelerator/Common/Private/UbaSessionServer.cpp
+++ b/Engine/Source/Programs/UnrealBuildAccelerator/Common/Private/UbaSessionServer.cpp
@@ -784,7 +784,7 @@ namespace uba
 						StringBuffer<> logPath;
 						logPath.Append(m_sessionLogDir).Append(destination.data + 5);
 						m_storage.CopyOrLink(casKey, logPath.data, attributes);
-						m_storage.DropCasFile(casKey, false);
+						m_storage.DropCasFile(casKey, false, logPath.data);
 						writer.WriteBool(true);
 						return true;
 					}
@@ -808,7 +808,7 @@ namespace uba
 
 				if (success)
 				{
-					m_storage.DropCasFile(casKey, false);
+					m_storage.DropCasFile(casKey, false, destination.data);
 					RegisterCreateFileForWrite(StringKeyZero, destination.data, destination.count, true);
 				}
 				return true;
diff --git a/Engine/Source/Programs/UnrealBuildAccelerator/Common/Private/UbaStorage.cpp b/Engine/Source/Programs/UnrealBuildAccelerator/Common/Private/UbaStorage.cpp
--- a/Engine/Source/Programs/UnrealBuildAccelerator/Common/Private/UbaStorage.cpp
+++ b/Engine/Source/Programs/UnrealBuildAccelerator/Common/Private/UbaStorage.cpp
@@ -925,7 +925,7 @@ namespace uba
 				m_casTotalBytes = 0;
 				break;
 			}
-			DropCasFile(entry->key, true);
+			DropCasFile(entry->key, true, TC("HandleOverflow"));
 			DetachEntry(*entry);
 			m_casLookup.erase(entry->key);
 		}
@@ -1787,7 +1787,7 @@ namespace uba
 		m_casDataBuffer.UnmapView(view, hint);
 	}
 
-	bool StorageImpl::DropCasFile(const CasKey& casKey, bool forceDelete)
+	bool StorageImpl::DropCasFile(const CasKey& casKey, bool forceDelete, const tchar* hint)
 	{
 		ScopedReadLock lookupLock(m_casLookupLock);
 		auto findIt = m_casLookup.find(casKey);
@@ -1813,7 +1813,7 @@ namespace uba
 			{
 				u32 lastError = GetLastError();
 				if (lastError != ERROR_FILE_NOT_FOUND && lastError != ERROR_PATH_NOT_FOUND)
-					return m_logger.Error(TC("Failed to drop cas %s (%s)"), casFile.data, LastErrorToText(lastError).data);
+					return m_logger.Error(TC("Failed to drop cas %s (%s) (%s)"), casFile.data, hint, LastErrorToText(lastError).data);
 			}
 			else
 			{
@@ -1909,6 +1909,10 @@ namespace uba
 				continue;
 			}
 
+			ScopedReadLock casEntryLock(casEntry->lock);
+			UBA_ASSERT(casEntry->verified);
+			UBA_ASSERT(casEntry->exists);
+
 			if (IsCompressed(actualKey))
 			{
 				FileHandle readHandle = InvalidFileHandle;
@@ -2448,7 +2452,7 @@ namespace uba
 		ScopedWriteLock entryLock(fileEntry.lock);
 		fileEntry.verified = false;
 
-		return DropCasFile(fileEntry.casKey, true);
+		return DropCasFile(fileEntry.casKey, true, file);
 	}
 
 }
diff --git a/Engine/Source/Programs/UnrealBuildAccelerator/Common/Private/UbaStorageClient.cpp b/Engine/Source/Programs/UnrealBuildAccelerator/Common/Private/UbaStorageClient.cpp
--- a/Engine/Source/Programs/UnrealBuildAccelerator/Common/Private/UbaStorageClient.cpp
+++ b/Engine/Source/Programs/UnrealBuildAccelerator/Common/Private/UbaStorageClient.cpp
@@ -60,6 +60,7 @@ namespace uba
 
 	bool StorageClient::IsUsingProxy()
 	{
+		ScopedReadLock proxyLock(m_proxyClientLock);
 		return m_proxyClient != nullptr;
 	}
 
diff --git a/Engine/Source/Programs/UnrealBuildAccelerator/Common/Public/UbaStorage.h b/Engine/Source/Programs/UnrealBuildAccelerator/Common/Public/UbaStorage.h
--- a/Engine/Source/Programs/UnrealBuildAccelerator/Common/Public/UbaStorage.h
+++ b/Engine/Source/Programs/UnrealBuildAccelerator/Common/Public/UbaStorage.h
@@ -45,7 +45,7 @@ namespace uba
 
 		virtual bool StoreCasFile(CasKey& out, const tchar* fileName, const CasKey& casKeyOverride = CasKeyZero, bool deferCreation = false) = 0;
 		virtual bool StoreCasFile(CasKey& out, StringKey fileNameKey, const tchar* fileName, FileMappingHandle mappingHandle, u64 mappingOffset, u64 fileSize, const tchar* hint, bool deferCreation = false, bool keepMappingInMemory = false) = 0;
-		virtual bool DropCasFile(const CasKey& casKey, bool forceDelete) = 0;
+		virtual bool DropCasFile(const CasKey& casKey, bool forceDelete, const tchar* hint) = 0;
 		virtual bool CalculateCasKey(CasKey& out, const tchar* fileName) = 0;
 		virtual bool CopyOrLink(const CasKey& casKey, const tchar* destination, u32 fileAttributes) = 0;
 		virtual bool FakeCopy(const CasKey& casKey, const tchar* destination) = 0;
@@ -102,7 +102,7 @@ namespace uba
 		virtual bool VerifyAndGetCachedFileInfo(CachedFileInfo& out, StringKey fileNameKey, u64 verifiedLastWriteTime, u64 verifiedSize) override;
 		virtual bool StoreCasFile(CasKey& out, const tchar* fileName, const CasKey& casKeyOverride = CasKeyZero, bool deferCreation = false) override;
 		virtual bool StoreCasFile(CasKey& out, StringKey fileNameKey, const tchar* fileName, FileMappingHandle mappingHandle, u64 mappingOffset, u64 fileSize, const tchar* hint, bool deferCreation = false, bool keepMappingInMemory = false) override;
-		virtual bool DropCasFile(const CasKey& casKey, bool forceDelete) override;
+		virtual bool DropCasFile(const CasKey& casKey, bool forceDelete, const tchar* hint) override;
 		virtual bool CalculateCasKey(CasKey& out, const tchar* fileName) override;
 		virtual bool CopyOrLink(const CasKey& casKey, const tchar* destination, u32 fileAttributes) override;
 		virtual bool FakeCopy(const CasKey& casKey, const tchar* destination) override;
