[Uba]
* Fixed shutdown issues where client and server did not wait for connection to be closed before closing down
* Added error handling for trying to decompress bad cas files

diff --git a/Engine/Source/Programs/UnrealBuildAccelerator/Common/Private/UbaNetworkBackendTcp.cpp b/Engine/Source/Programs/UnrealBuildAccelerator/Common/Private/UbaNetworkBackendTcp.cpp
--- a/Engine/Source/Programs/UnrealBuildAccelerator/Common/Private/UbaNetworkBackendTcp.cpp
+++ b/Engine/Source/Programs/UnrealBuildAccelerator/Common/Private/UbaNetworkBackendTcp.cpp
@@ -198,6 +198,9 @@ namespace uba
 		UBA_ASSERT(h);
 		UBA_ASSERT(headerSize <= MaxHeaderSize);
 		auto& conn = *(Connection*)connection;
+
+		ScopedCriticalSection lock(conn.shutdownLock);
+		UBA_ASSERTF(conn.disconnectCallback, TC("SetDisconnectCallback must be called before SetRecvCallbacks"));
 		conn.recvContext = context;
 		conn.headerSize = headerSize;
 		conn.headerCallback = h;
@@ -387,12 +390,11 @@ namespace uba
 				break;
 			}
 
-			auto socketClose = MakeGuard([&]() { CloseSocket(logger, clientSocket); });
-
-			if (!DisableNagle(logger, clientSocket))
-				continue;
-			if (!SetKeepAlive(logger, clientSocket))
+			if (!DisableNagle(logger, clientSocket) || !SetKeepAlive(logger, clientSocket))
+			{
+				CloseSocket(logger, clientSocket);
 				continue;
+			}
 
 			ScopedWriteLock lock(m_connectionsLock);
 			auto it = m_connections.emplace(m_connections.end(), clientSocket);
@@ -402,15 +404,12 @@ namespace uba
 
 			if (!m_connectedFunc(&conn, remoteSockAddr))
 			{
-				conn.socket = INVALID_SOCKET;
-				socketClose.Execute();
+				shutdown(clientSocket, SD_BOTH);
 				conn.ready.Set();
 				ScopedWriteLock lock2(m_connectionsLock);
 				m_connections.erase(it);
 				continue;
 			}
-
-			socketClose.Cancel();
 		}
 
 		return true;
@@ -470,14 +469,14 @@ namespace uba
 		SOCKET s = connection.socket;
 		connection.socket = INVALID_SOCKET;
 
-		auto cb = connection.disconnectCallback;
-		if (cb)
+		if (auto cb = connection.disconnectCallback)
 		{
 			auto context = connection.disconnectContext;
 			connection.disconnectCallback = nullptr;
 			connection.disconnectContext = nullptr;
 			cb(context, &connection);
 		}
+
 		if (s == INVALID_SOCKET)
 			return;
 		shutdown(s, SD_BOTH);
diff --git a/Engine/Source/Programs/UnrealBuildAccelerator/Common/Private/UbaNetworkClient.cpp b/Engine/Source/Programs/UnrealBuildAccelerator/Common/Private/UbaNetworkClient.cpp
--- a/Engine/Source/Programs/UnrealBuildAccelerator/Common/Private/UbaNetworkClient.cpp
+++ b/Engine/Source/Programs/UnrealBuildAccelerator/Common/Private/UbaNetworkClient.cpp
@@ -41,16 +41,8 @@ namespace uba
 	{
 		StopListen();
 
+		Disconnect();
 		ScopedWriteLock lock(m_connectionsLock);
-		for (auto& connection : m_connections)
-		{
-			connection.backend->SetDisconnectCallback(connection.backendConnection, nullptr, nullptr);
-			if (connection.connected.exchange(0) == 0)
-				continue;
-			m_logger.Detail(TC("Disconnected from server..."));
-			connection.backend->Shutdown(connection.backendConnection);
-			connection.backend->Close(connection.backendConnection);
-		}
 		m_connections.clear(); // Before tcpBackend
 
 		delete m_tcpBackend;
@@ -71,11 +63,12 @@ namespace uba
 	{
 		struct RecvContext
 		{
-			RecvContext(NetworkClient& c, NetworkBackend& b, void* bc) : client(c), backend(b), backendConnection(bc), recvEvent(true) {}
+			RecvContext(NetworkClient& c, NetworkBackend& b, void* bc) : client(c), backend(b), backendConnection(bc), recvEvent(true), disconnectedEvent(true) {}
 			NetworkClient& client;
 			NetworkBackend& backend;
 			void* backendConnection;
 			Event recvEvent;
+			Event disconnectedEvent;
 			u8 error = 0;
 		} rc(*this, backend, backendConnection);
 
@@ -84,6 +77,13 @@ namespace uba
 				auto& c = *(RecvContext*)context;
 				c.error = 4;
 				c.recvEvent.Set();
+				c.disconnectedEvent.Set();
+			});
+
+		auto waitForDisconnect = MakeGuard([&]()
+			{
+				backend.Shutdown(backendConnection);
+				rc.disconnectedEvent.IsSet(~0u);
 			});
 
 		backend.SetRecvCallbacks(backendConnection, &rc, 1 + sizeof(Guid), [](void* context, u8* headerData, void*& outBodyContext, u8*& outBodyData, u32& outBodySize)
@@ -170,6 +170,8 @@ namespace uba
 			return false;
 		}
 
+		waitForDisconnect.Cancel();
+
 		if (m_connectionCount.fetch_add(1) != 0)
 			return true;
 
@@ -199,6 +201,7 @@ namespace uba
 		backend.SetDisconnectCallback(backendConnection, connection, [](void* context, void* connection)
 			{
 				auto& c = *(Connection*)context;
+				c.disconnectedEvent.Set();
 				c.owner.OnDisconnected(c, true);
 			});
 		backend.SetRecvCallbacks(backendConnection, connection, ReceiveHeaderSize, ReceiveResponseHeader, ReceiveResponseBody, TC("ReceiveMessageResponse"));
@@ -260,7 +263,10 @@ namespace uba
 	{
 		ScopedReadLock lock(m_connectionsLock);
 		for (auto& c : m_connections)
+		{
 			OnDisconnected(c, false);
+			c.disconnectedEvent.IsSet(~0u);
+		}
 	}
 
 	bool NetworkClient::StartListen(NetworkBackend& backend, u16 port)
@@ -314,6 +320,9 @@ namespace uba
 		u64 sendBytes = 0;
 		u64 recvBytes = 0;
 		u32 recvCount = 0;
+
+		ScopedReadLock lock(m_connectionsLock);
+		u32 connectionsCount = u32(m_connections.size());
 		for (auto& c : m_connections)
 		{
 			sendTimer.count += c.sendTimer.count;
@@ -322,6 +331,7 @@ namespace uba
 			recvBytes += c.recvBytes;
 			recvCount += c.recvCount;
 		}
+		lock.Leave();
 
 
 		logger.Info(TC("  ----- Uba client stats summary ------"));
@@ -334,7 +344,7 @@ namespace uba
 			logger.Info(TC("  DecryptTotal       %8u %9s"), m_decryptTimer.count.load(), TimeToText(m_decryptTimer.time).str);
 		}
 		logger.Info(TC("  MaxActiveMessages  %8u"), m_activeMessageIdMax);
-		logger.Info(TC("  Connections        %8u"), u32(m_connections.size()));
+		logger.Info(TC("  Connections        %8u"), connectionsCount);
 		logger.Info(TC("  SendSize Set/Max  %9s %9s"), BytesToText(m_sendSize).str, BytesToText(SendMaxSize).str);
 		logger.Info(TC(""));
 	}
diff --git a/Engine/Source/Programs/UnrealBuildAccelerator/Common/Private/UbaNetworkServer.cpp b/Engine/Source/Programs/UnrealBuildAccelerator/Common/Private/UbaNetworkServer.cpp
--- a/Engine/Source/Programs/UnrealBuildAccelerator/Common/Private/UbaNetworkServer.cpp
+++ b/Engine/Source/Programs/UnrealBuildAccelerator/Common/Private/UbaNetworkServer.cpp
@@ -94,23 +94,8 @@ namespace uba
 		,	m_backend(backend)
 		,	m_remoteSockAddr(remoteSockAddr)
 		,	m_cryptoKey(cryptoKey)
+		,	m_disconnectCallbackCalled(true)
 		,	m_backendConnection(backendConnection)
-		{
-		}
-
-		~Connection()
-		{
-			Stop();
-			if (m_backendConnection)
-			{
-				m_backend.SetDisconnectCallback(m_backendConnection, nullptr, nullptr);
-				m_backend.Close(m_backendConnection);
-			}
-			if (m_cryptoKey)
-				Crypto::DestroyKey(m_cryptoKey);
-		}
-
-		void Start()
 		{
 			m_activeWorkerCount = 1;
 
@@ -118,6 +103,7 @@ namespace uba
 				{
 					auto& conn = *(Connection*)context;
 					conn.Disconnect();
+					conn.m_disconnectCallbackCalled.Set();
 				});
 
 			m_backend.SetDataSentCallback(m_backendConnection, this, [](void* context, u32 bytes)
@@ -136,6 +122,13 @@ namespace uba
 				m_backend.SetRecvCallbacks(m_backendConnection, this, 4, ReceiveVersion, nullptr, TC("ReceiveVersion"));
 		}
 
+		~Connection()
+		{
+			Stop();
+			if (m_cryptoKey)
+				Crypto::DestroyKey(m_cryptoKey);
+		}
+
 		void Disconnect()
 		{
 			if (m_disconnectCalled.fetch_add(1) != 0)
@@ -148,7 +141,7 @@ namespace uba
 		bool Stop()
 		{
 			Disconnect();
-				
+
 			u64 startTimer = GetTime();
 			while (m_activeWorkerCount)
 			{
@@ -159,6 +152,10 @@ namespace uba
 				}
 				Sleep(1);
 			}
+
+			m_disconnectCallbackCalled.IsSet(~0u);
+			if (m_backendConnection)
+				m_backend.Close(m_backendConnection);
 			return true;
 		}
 
@@ -360,6 +357,7 @@ namespace uba
 		Client* m_client = nullptr;
 		sockaddr m_remoteSockAddr;
 		CryptoKey m_cryptoKey;
+		Event m_disconnectCallbackCalled;
 		Atomic<int> m_activeWorkerCount;
 		Atomic<int> m_disconnectCalled;
 		Atomic<bool> m_disconnected;
@@ -400,7 +398,9 @@ namespace uba
 
 	bool ConnectionInfo::ShouldDisconnect() const
 	{
-		return ((NetworkServer::Connection*)internalData)->m_shouldDisconnect;
+		auto& conn = *(NetworkServer::Connection*)internalData;
+		ScopedWriteLock lock(conn.m_shutdownLock);
+		return conn.m_shouldDisconnect;
 	}
 
 	void NetworkServer::Worker::Update(WorkerContext& context, bool signalAvailable)
@@ -1071,7 +1071,7 @@ namespace uba
 
 		RemoveDisconnectedConnections();
 
-		m_connections.emplace_back(*this, backend, backendConnection, remoteSocketAddr, cryptoKey).Start();
+		m_connections.emplace_back(*this, backend, backendConnection, remoteSocketAddr, cryptoKey);
 		m_maxActiveConnections = Max(m_maxActiveConnections, u32(m_connections.size()));
 		return true;
 	}
diff --git a/Engine/Source/Programs/UnrealBuildAccelerator/Common/Private/UbaStorage.cpp b/Engine/Source/Programs/UnrealBuildAccelerator/Common/Private/UbaStorage.cpp
--- a/Engine/Source/Programs/UnrealBuildAccelerator/Common/Private/UbaStorage.cpp
+++ b/Engine/Source/Programs/UnrealBuildAccelerator/Common/Private/UbaStorage.cpp
@@ -542,6 +542,8 @@ namespace uba
 		{
 			struct WorkRec
 			{
+				Logger* logger;
+				const tchar* hint;
 				Atomic<u64> refCount;
 				u8* readPos = nullptr;
 				u8* writePos = nullptr;
@@ -551,9 +553,12 @@ namespace uba
 				u64 decompressedLeft = 0;
 				u64 written = 0;
 				Event done;
+				Atomic<bool> error;
 			};
 
 			WorkRec* rec = new WorkRec();
+			rec->logger = &m_logger;
+			rec->hint = readHint;
 			rec->readPos = compressedData;
 			rec->writePos = writeData;
 			rec->decompressedSize = decompressedSize;
@@ -575,12 +580,24 @@ namespace uba
 							lock.Leave();
 							if (!--rec->refCount)
 								delete rec;
-							return 0;
+							return;
 						}
 						u8* readPos = rec->readPos;
 						u8* writePos = rec->writePos;
 						u32 compressedBlockSize = ((u32*)readPos)[0];
 						u32 decompressedBlockSize = ((u32*)readPos)[1];
+
+						if (decompressedBlockSize == 0 || decompressedBlockSize > rec->decompressedSize)
+						{
+							bool f = false;
+							if (rec->error.compare_exchange_strong(f, true))
+								rec->logger->Error(TC("Decompressed block size %u is invalid. Decompressed file is %u (%s)"), decompressedBlockSize, rec->decompressedSize, rec->hint);
+							if (!--rec->refCount)
+								delete rec;
+							rec->done.Set();
+							return;
+						}
+
 						readPos += sizeof(u32) * 2;
 						rec->decompressedLeft -= decompressedBlockSize;
 						rec->readPos = readPos + compressedBlockSize;
@@ -588,7 +605,16 @@ namespace uba
 						lock.Leave();
 
 						OO_SINTa decompLen = OodleLZ_Decompress(readPos, (OO_SINTa)compressedBlockSize, writePos, (OO_SINTa)decompressedBlockSize); (void)decompLen;
-						UBA_ASSERTF(decompLen == decompressedBlockSize, TC("Expecting to be able to decompress to %u bytes but got %llu"), decompressedBlockSize, decompLen);
+						if (decompLen != decompressedBlockSize)
+						{
+							bool f = false;
+							if (rec->error.compare_exchange_strong(f, true))
+								rec->logger->Error(TC("Expecting to be able to decompress to %u bytes but got %llu (%s)"), decompressedBlockSize, decompLen, rec->hint);
+							if (!--rec->refCount)
+								delete rec;
+							rec->done.Set();
+							return;
+						}
 						lastWritten = u64(decompLen);
 					}
 				};
@@ -605,8 +631,14 @@ namespace uba
 			TimerScope ts(stats.decompressToMem);
 			work();
 			rec->done.IsSet();
+			bool success = !rec->error;
+			if (!success)
+				while (rec->refCount > 1)
+					Sleep(10);
+
 			if (!--rec->refCount)
 				delete rec;
+			return success;
 		}
 		else
 		{
@@ -620,12 +652,14 @@ namespace uba
 				if (!compressedBlockSize)
 					break;
 				u32 decompressedBlockSize = ((u32*)readPos)[1];
-				UBA_ASSERT(decompressedBlockSize <= left);
+				if (decompressedBlockSize == 0 || decompressedBlockSize > left)
+					return m_logger.Error(TC("Decompressed block size %u is invalid. Decompressed file is %u (%s)"), decompressedBlockSize, decompressedSize, readHint);
 				readPos += sizeof(u32) * 2;
 
 				TimerScope ts(stats.decompressToMem);
 				OO_SINTa decompLen = OodleLZ_Decompress(readPos, (OO_SINTa)compressedBlockSize, writePos, (OO_SINTa)decompressedBlockSize); (void)decompLen;
-				UBA_ASSERTF(decompLen == decompressedBlockSize, TC("Expecting to be able to decompress to %u bytes but got %llu"), decompressedBlockSize, decompLen);
+				if (decompLen != decompressedBlockSize)
+					return m_logger.Error(TC("Expecting to be able to decompress to %u bytes but got %llu (%s)"), decompressedBlockSize, decompLen, readHint);
 				writePos += decompressedBlockSize;
 				readPos += compressedBlockSize;
 				left -= decompressedBlockSize;
diff --git a/Engine/Source/Programs/UnrealBuildAccelerator/Common/Public/UbaNetworkClient.h b/Engine/Source/Programs/UnrealBuildAccelerator/Common/Public/UbaNetworkClient.h
--- a/Engine/Source/Programs/UnrealBuildAccelerator/Common/Public/UbaNetworkClient.h
+++ b/Engine/Source/Programs/UnrealBuildAccelerator/Common/Public/UbaNetworkClient.h
@@ -67,10 +67,11 @@ namespace uba
 	private:
 		struct Connection
 		{
-			Connection(NetworkClient& o) : owner(o) {}
+			Connection(NetworkClient& o) : owner(o), disconnectedEvent(true) {}
 			NetworkClient& owner;
 			void* backendConnection = nullptr;
 			Atomic<u32> connected;
+			Event disconnectedEvent;
 			Timer sendTimer;
 			Atomic<u64> sendBytes;
 			u64 recvBytes = 0;
