[Uba]
* NetworkBackendTcp - Changed so when having failed connection we wait for socket thread before taking lock to delete it.
* NetworkBackendTcp - Removed check for 0 connections. That was a hack and now all code should properly wait
* Changed infinte wait of event to 60 seconds and output error if it times out..
* Removed NetworkBackend::Close. This should never be done on the outside. Outside should always call shutdown and then wait for disconnect callback
* Moved shutdown of socket to outside lock to prevent deadlock

diff --git a/Engine/Source/Programs/UnrealBuildAccelerator/Common/Private/UbaNetworkBackendTcp.cpp b/Engine/Source/Programs/UnrealBuildAccelerator/Common/Private/UbaNetworkBackendTcp.cpp
--- a/Engine/Source/Programs/UnrealBuildAccelerator/Common/Private/UbaNetworkBackendTcp.cpp
+++ b/Engine/Source/Programs/UnrealBuildAccelerator/Common/Private/UbaNetworkBackendTcp.cpp
@@ -152,19 +152,6 @@ namespace uba
 		shutdown(conn.socket, SD_BOTH);
 	}
 
-	void NetworkBackendTcp::Close(void* connection)
-	{
-		auto& conn = *(Connection*)connection;
-		ScopedCriticalSection lock(conn.shutdownLock);
-		if (conn.socket == INVALID_SOCKET)
-			return;
-		SOCKET s = conn.socket;
-		conn.socket = INVALID_SOCKET;
-		lock.Leave();
-		conn.recvThread.Wait();
-		CloseSocket(m_logger, s);
-	}
-
 	bool NetworkBackendTcp::Send(Logger& logger, void* connection, const void* data, u32 dataSize, SendContext& sendContext)
 	{
 		auto& conn = *(Connection*)connection;
@@ -217,11 +204,6 @@ namespace uba
 
 	void NetworkBackendTcp::SetDisconnectCallback(void* connection, void* context, DisconnectCallback* callback)
 	{
-		{
-			ScopedReadLock lock(m_connectionsLock);
-			if (m_connections.empty())
-				return;
-		}
 		auto& conn = *(Connection*)connection;
 		ScopedCriticalSection lock(conn.shutdownLock);
 		conn.disconnectCallback = callback;
@@ -406,6 +388,7 @@ namespace uba
 			{
 				shutdown(clientSocket, SD_BOTH);
 				conn.ready.Set();
+				conn.recvThread.Wait();
 				ScopedWriteLock lock2(m_connectionsLock);
 				m_connections.erase(it);
 				continue;
@@ -417,7 +400,7 @@ namespace uba
 
 	void NetworkBackendTcp::ThreadRecv(Connection& connection)
 	{
-		if (connection.ready.IsSet(2000))
+		if (connection.ready.IsSet(60000)) // This should never time out!
 		{
 			Guid connectionUid;
 			CreateGuid(connectionUid);
@@ -462,7 +445,7 @@ namespace uba
 		}
 		else
 		{
-			m_logger.Error(TC("Timed out waiting for recv thread to be ready"));
+			m_logger.Warning(TC("Timed out waiting for recv thread to be ready"));
 		}
 
 		ScopedCriticalSection lock2(connection.shutdownLock);
@@ -631,6 +614,8 @@ namespace uba
 		if (!SetKeepAlive(logger, socketFd))
 			return false;
 
+		socketClose.Cancel();
+
 		ScopedWriteLock lock(m_connectionsLock);
 		auto it = m_connections.emplace(m_connections.end(), socketFd);
 		auto& conn = *it;
@@ -639,10 +624,9 @@ namespace uba
 
 		if (!connectedFunc(&conn, remoteSocketAddr, timedOut))
 		{
-			socketFd = conn.socket;
-			conn.socket = INVALID_SOCKET;
-			socketClose.Execute();
+			shutdown(socketFd, SD_BOTH);
 			conn.ready.Set();
+			conn.recvThread.Wait();
 			ScopedWriteLock lock2(m_connectionsLock);
 			m_connections.erase(it);
 			return false;
@@ -654,8 +638,6 @@ namespace uba
 		else
 			logger.Detail(TC("Connected using sockaddr"));
 
-		socketClose.Cancel();
-
 		return true;
 	}
 
diff --git a/Engine/Source/Programs/UnrealBuildAccelerator/Common/Private/UbaNetworkServer.cpp b/Engine/Source/Programs/UnrealBuildAccelerator/Common/Private/UbaNetworkServer.cpp
--- a/Engine/Source/Programs/UnrealBuildAccelerator/Common/Private/UbaNetworkServer.cpp
+++ b/Engine/Source/Programs/UnrealBuildAccelerator/Common/Private/UbaNetworkServer.cpp
@@ -153,9 +153,8 @@ namespace uba
 				Sleep(1);
 			}
 
-			m_disconnectCallbackCalled.IsSet(~0u);
-			if (m_backendConnection)
-				m_backend.Close(m_backendConnection);
+			if (!m_disconnectCallbackCalled.IsSet(60000)) // This should never time out!
+				m_server.m_logger.Warning(TC("This should never happen!! Unknown conseqences"));
 			return true;
 		}
 
@@ -340,8 +339,8 @@ namespace uba
 				return;
 			if (m_disconnected)
 				return;
-			m_backend.Shutdown(m_backendConnection);
 			lock.Leave();
+			m_backend.Shutdown(m_backendConnection);
 			if (m_client && m_client->connectionCount.fetch_sub(1) == 1)
 			{
 				for (auto& entry : m_server.m_onDisconnectFunctions)
diff --git a/Engine/Source/Programs/UnrealBuildAccelerator/Common/Public/UbaNetworkBackend.h b/Engine/Source/Programs/UnrealBuildAccelerator/Common/Public/UbaNetworkBackend.h
--- a/Engine/Source/Programs/UnrealBuildAccelerator/Common/Public/UbaNetworkBackend.h
+++ b/Engine/Source/Programs/UnrealBuildAccelerator/Common/Public/UbaNetworkBackend.h
@@ -25,11 +25,9 @@ namespace uba
 		virtual ~NetworkBackend() {}
 
 		// Shutdown the connection (stops/finishes traffic)
+		// External code should wait for disconnect callback before assuming that no more recv callbacks can happen
 		virtual void Shutdown(void* connection) = 0;
 
-		// Close the connection
-		virtual void Close(void* connection) = 0;
-
 		enum : u8
 		{
 			SendFlags_Async = 1 << 0, 	// Data is sent async and call is returned before data has been transmitted. (can be used inside recv callbacks). Will make copy of data
diff --git a/Engine/Source/Programs/UnrealBuildAccelerator/Common/Public/UbaNetworkBackendTcp.h b/Engine/Source/Programs/UnrealBuildAccelerator/Common/Public/UbaNetworkBackendTcp.h
--- a/Engine/Source/Programs/UnrealBuildAccelerator/Common/Public/UbaNetworkBackendTcp.h
+++ b/Engine/Source/Programs/UnrealBuildAccelerator/Common/Public/UbaNetworkBackendTcp.h
@@ -14,7 +14,6 @@ namespace uba
 		NetworkBackendTcp(LogWriter& writer, const tchar* prefix = TC("NetworkBackendTcp"));
 		virtual ~NetworkBackendTcp();
 		virtual void Shutdown(void* connection) override;
-		virtual void Close(void* connection) override;
 		virtual bool Send(Logger& logger, void* connection, const void* data, u32 dataSize, SendContext& sendContext) override;
 		virtual void SetDataSentCallback(void* connection, void* context, DataSentCallback* callback) override;
 		virtual void SetRecvCallbacks(void* connection, void* context, u32 headerSize, RecvHeaderCallback* h, RecvBodyCallback* b, const tchar* recvHint) override;
