[UBA]
* Disabled custom signal handler in UbaHost in case this is the reason we get crashes in ubt
* Changed some logging from Info to Detail

diff --git a/Engine/Source/Programs/UnrealBuildAccelerator/Common/Private/UbaNetworkServer.cpp b/Engine/Source/Programs/UnrealBuildAccelerator/Common/Private/UbaNetworkServer.cpp
--- a/Engine/Source/Programs/UnrealBuildAccelerator/Common/Private/UbaNetworkServer.cpp
+++ b/Engine/Source/Programs/UnrealBuildAccelerator/Common/Private/UbaNetworkServer.cpp
@@ -345,7 +345,7 @@ namespace uba
 			{
 				for (auto& entry : m_server.m_onDisconnectFunctions)
 					entry.function(m_client->uid, m_client->id);
-				m_server.m_logger.Info(TC("Client %s disconnected"), GuidToString(m_client->uid).str);
+				m_server.m_logger.Detail(TC("Client %s disconnected"), GuidToString(m_client->uid).str);
 			}
 			m_disconnected = true;
 		}
diff --git a/Engine/Source/Programs/UnrealBuildAccelerator/Common/Private/UbaSessionServer.cpp b/Engine/Source/Programs/UnrealBuildAccelerator/Common/Private/UbaSessionServer.cpp
--- a/Engine/Source/Programs/UnrealBuildAccelerator/Common/Private/UbaSessionServer.cpp
+++ b/Engine/Source/Programs/UnrealBuildAccelerator/Common/Private/UbaSessionServer.cpp
@@ -540,7 +540,7 @@ namespace uba
 				reader.ReadString(name);
 				u32 clientVersion = reader.ReadU32();
 
-				m_logger.Info(TC("Client session %s connected (Id: %u, Uid: %s)"), name.data, connectionInfo.GetId(), GuidToString(connectionInfo.GetUid()).str);
+				m_logger.Detail(TC("Client session %s connected (Id: %u, Uid: %s)"), name.data, connectionInfo.GetId(), GuidToString(connectionInfo.GetUid()).str);
 
 				CasKey clientKeys[2];
 				clientKeys[0] = reader.ReadCasKey();
@@ -955,7 +955,7 @@ namespace uba
 					if (session.enabled)
 						m_availableRemoteSlotCount -= session.processSlotCount - session.usedSlotCount;
 					session.enabled = false;
-					m_logger.Info(TC("Disable remote execution on %s because remote execution has been disabled and queue is empty (will finish %u processes)"), session.name.c_str(), session.usedSlotCount);
+					m_logger.Detail(TC("Disable remote execution on %s because remote execution has been disabled and queue is empty (will finish %u processes)"), session.name.c_str(), session.usedSlotCount);
 				}
 
 				// If this client session has 0 active processes and m_maxRemoteProcessCount < total available compute - client session, then we can disconnect this client
@@ -1680,7 +1680,7 @@ namespace uba
 		u64 entryCount = m_nameToHashLookup.size();
 		lock2.Leave();
 
-		m_logger.Info(TC("Prepopulated NameToHash table with %u entries"), entryCount);
+		m_logger.Debug(TC("Prepopulated NameToHash table with %u entries"), entryCount);
 
 		return true;
 	}
diff --git a/Engine/Source/Programs/UnrealBuildAccelerator/Common/Private/UbaStorage.cpp b/Engine/Source/Programs/UnrealBuildAccelerator/Common/Private/UbaStorage.cpp
--- a/Engine/Source/Programs/UnrealBuildAccelerator/Common/Private/UbaStorage.cpp
+++ b/Engine/Source/Programs/UnrealBuildAccelerator/Common/Private/UbaStorage.cpp
@@ -1264,7 +1264,7 @@ namespace uba
 		if (logStats)
 		{
 			u64 duration = GetTime() - startTime;
-			m_logger.Info(TC("Database loaded from %s in %s"), fileName.data, TimeToText(duration).str);
+			m_logger.Detail(TC("Database loaded from %s in %s"), fileName.data, TimeToText(duration).str);
 		}
 
 		return true;
diff --git a/Engine/Source/Programs/UnrealBuildAccelerator/Common/Private/UbaStorageClient.cpp b/Engine/Source/Programs/UnrealBuildAccelerator/Common/Private/UbaStorageClient.cpp
--- a/Engine/Source/Programs/UnrealBuildAccelerator/Common/Private/UbaStorageClient.cpp
+++ b/Engine/Source/Programs/UnrealBuildAccelerator/Common/Private/UbaStorageClient.cpp
@@ -346,7 +346,7 @@ namespace uba
 					u64 startTime = GetTime();
 					if (!proxy->client.Connect(proxyBackend, proxyHost.data, proxyPort))
 					{
-						m_logger.Info(TC("Redirection to proxy %s:%u for cas %s download failed! (%s)"), proxyHost.data, proxyPort, casFile.data, hint);
+						m_logger.Detail(TC("Redirection to proxy %s:%u for cas %s download failed! (%s)"), proxyHost.data, proxyPort, casFile.data, hint);
 						continue;
 					}
 					
diff --git a/Engine/Source/Programs/UnrealBuildAccelerator/Common/Private/UbaStorageServer.cpp b/Engine/Source/Programs/UnrealBuildAccelerator/Common/Private/UbaStorageServer.cpp
--- a/Engine/Source/Programs/UnrealBuildAccelerator/Common/Private/UbaStorageServer.cpp
+++ b/Engine/Source/Programs/UnrealBuildAccelerator/Common/Private/UbaStorageServer.cpp
@@ -126,7 +126,7 @@ namespace uba
 				ProxyEntry& e = it->second;
 				if (e.clientId != clientId)
 					continue;
-				m_logger.Info(TC("Proxy %s:%u for zone %s removed"), e.host.c_str(), e.port, e.zone.c_str());
+				m_logger.Detail(TC("Proxy %s:%u for zone %s removed"), e.host.c_str(), e.port, e.zone.c_str());
 				m_proxies.erase(it);
 				break;
 			}
@@ -339,7 +339,7 @@ namespace uba
 								proxy.port = info.proxyPort;
 								proxy.zone = proxyName.data;
 
-								m_logger.Info(TC("%s:%u (%s) is assigned as proxy for zone %s"), proxy.host.c_str(), proxy.port, GuidToString(connectionInfo.GetUid()).str, proxy.zone.c_str());
+								m_logger.Detail(TC("%s:%u (%s) is assigned as proxy for zone %s"), proxy.host.c_str(), proxy.port, GuidToString(connectionInfo.GetUid()).str, proxy.zone.c_str());
 
 								writer.WriteBool(true);
 								writer.WriteU16(info.proxyPort);
diff --git a/Engine/Source/Programs/UnrealBuildAccelerator/Host/Private/UbaExports.cpp b/Engine/Source/Programs/UnrealBuildAccelerator/Host/Private/UbaExports.cpp
--- a/Engine/Source/Programs/UnrealBuildAccelerator/Host/Private/UbaExports.cpp
+++ b/Engine/Source/Programs/UnrealBuildAccelerator/Host/Private/UbaExports.cpp
@@ -53,7 +53,7 @@ namespace uba
 		NetworkBackend* networkBackend;
 	};
 
-	#define UBA_USE_SIGNALHANDLER PLATFORM_LINUX
+	#define UBA_USE_SIGNALHANDLER 0//PLATFORM_LINUX // It might be that we can't use signal handlers in c# processes.. so don't set this to 1
 
 	#if UBA_USE_SIGNALHANDLER
 	void SignalHandler(int sig, siginfo_t* si, void* unused)
