[Uba]
* Fixed deadlock in client that could happen if client connecting timed out at the same time as it managed to connect. There were lots of race conditions in this code around these scenarios. Current implementation should be properly fixed

diff --git a/Engine/Source/Programs/UnrealBuildAccelerator/Common/Private/UbaNetworkClient.cpp b/Engine/Source/Programs/UnrealBuildAccelerator/Common/Private/UbaNetworkClient.cpp
--- a/Engine/Source/Programs/UnrealBuildAccelerator/Common/Private/UbaNetworkClient.cpp
+++ b/Engine/Source/Programs/UnrealBuildAccelerator/Common/Private/UbaNetworkClient.cpp
@@ -63,49 +63,67 @@ namespace uba
 	{
 		struct RecvContext
 		{
-			RecvContext(NetworkClient& c, NetworkBackend& b, void* bc) : client(c), backend(b), backendConnection(bc), recvEvent(true), disconnectedEvent(true) {}
+			RecvContext(NetworkClient& c, NetworkBackend& b, void* bc) : client(c), backend(b), backendConnection(bc), recvEvent(true), exitScopeEvent(true)
+			{
+			}
+
+			~RecvContext()
+			{
+				if (error)
+					backend.Shutdown(backendConnection);
+				exitScopeEvent.IsSet(~0u);
+			}
+
 			NetworkClient& client;
 			NetworkBackend& backend;
 			void* backendConnection;
 			Event recvEvent;
-			Event disconnectedEvent;
-			u8 error = 0;
-		} rc(*this, backend, backendConnection);
+			Event exitScopeEvent;
+			u8 error = 255;
+		};
 
-		backend.SetDisconnectCallback(backendConnection, &rc, [](void* context, void* connection)
-			{
-				auto& c = *(RecvContext*)context;
-				c.error = 4;
-				c.recvEvent.Set();
-				c.disconnectedEvent.Set();
-			});
+		RecvContext rc(*this, backend, backendConnection);
 
-		auto waitForDisconnect = MakeGuard([&]()
+		// The only way out of this function is to get a call to one of the below callbacks since exitScopeEvent must be set.
+
+		backend.SetDisconnectCallback(backendConnection, &rc, [](void* context, void* connection)
 			{
-				backend.Shutdown(backendConnection);
-				rc.disconnectedEvent.IsSet(~0u);
+				auto& rc = *(RecvContext*)context;
+				rc.error = 4;
+				rc.recvEvent.Set();
+				rc.exitScopeEvent.Set();
 			});
 
 		backend.SetRecvCallbacks(backendConnection, &rc, 1 + sizeof(Guid), [](void* context, u8* headerData, void*& outBodyContext, u8*& outBodyData, u32& outBodySize)
 			{
-				auto& c = *(RecvContext*)context;
-				c.error = *headerData;
+				auto& rc = *(RecvContext*)context;
+				rc.error = *headerData;
 				Guid serverUid = *(Guid*)(headerData+1);
 
-				if (!c.error)
+				if (!rc.error)
 				{
-					ScopedWriteLock lock(c.client.m_serverUidLock);
-					if (c.client.m_serverUid == Guid())
-						c.client.m_serverUid = serverUid;
-					else if (c.client.m_serverUid != serverUid) // Seems like two different servers tried to connect to this client.. keep the first one and ignore the others
-						c.error = 5;
+					ScopedWriteLock lock(rc.client.m_serverUidLock);
+					if (rc.client.m_serverUid == Guid())
+						rc.client.m_serverUid = serverUid;
+					else if (rc.client.m_serverUid != serverUid) // Seems like two different servers tried to connect to this client.. keep the first one and ignore the others
+						rc.error = 5;
 				}
 
-				if (!c.error)
-					c.client.ConnectedCallback(c.backend, c.backendConnection);
+				if (!rc.error)
+				{
+					rc.client.ConnectedCallback(rc.backend, rc.backendConnection);
+				}
+				else
+				{
+					// Zero out callbacks since we will be leaving this scope when exiting current callback and we don't want the disconnect callback
+					rc.backend.SetRecvCallbacks(rc.backendConnection, nullptr, 0, [](void*, u8*, void*&, u8*&, u32&) { return false; }, nullptr, TC("Null"));
+					rc.backend.SetDisconnectCallback(rc.backendConnection, nullptr, nullptr);
+				}
 
-				c.recvEvent.Set();
+				rc.recvEvent.Set();
+				rc.exitScopeEvent.Set();
 				return true;
+
 			}, nullptr, TC("Connecting"));
 
 
@@ -133,9 +151,7 @@ namespace uba
 		if (!backend.Send(m_logger, backendConnection, &m_uid, sizeof(m_uid), uidContext))
 			return false;
 
-		u32 timeoutSeconds = 4; // Four seconds for server to respond
-
-		if (!rc.recvEvent.IsSet(timeoutSeconds * 1000))
+		if (!rc.recvEvent.IsSet(~0u)) // This can not happen. Since both callbacks are using rc we can't leave this function until we know we are not in the callbacks
 			return m_logger.Error(TC("Timed out waiting for connection response from server"));
 
 		m_isOrWasConnected.Set();
@@ -170,8 +186,6 @@ namespace uba
 			return false;
 		}
 
-		waitForDisconnect.Cancel();
-
 		if (m_connectionCount.fetch_add(1) != 0)
 			return true;
 
