[UBA]
* Changed aws code to use token for meta data queries. Some of our AWS machines require this and it doesn't hurt to use it everywhere

diff --git a/Engine/Source/Programs/UnrealBuildAccelerator/Common/Private/UbaNetworkBackendTcp.cpp b/Engine/Source/Programs/UnrealBuildAccelerator/Common/Private/UbaNetworkBackendTcp.cpp
--- a/Engine/Source/Programs/UnrealBuildAccelerator/Common/Private/UbaNetworkBackendTcp.cpp
+++ b/Engine/Source/Programs/UnrealBuildAccelerator/Common/Private/UbaNetworkBackendTcp.cpp
@@ -907,7 +907,7 @@ namespace uba
 		return true;
 	}
 
-	bool HttpConnection::Get(Logger& logger, StringBufferBase& outResponse, u32& outStatusCode, const char* host, const char* path)
+	bool HttpConnection::Query(Logger& logger, const char* type, StringBufferBase& outResponse, u32& outStatusCode, const char* host, const char* path, const char* header)
 	{
 		// TODO: Fix so we reuse socket connection for multiple queries
 		if (*m_host)// && _stricmp(m_host, host) != 0)
@@ -922,7 +922,7 @@ namespace uba
 				return false;
 
 		char request[512];
-		int requestLen = snprintf(request, 512, "GET /%s HTTP/1.1\r\nHost: %s\r\n\r\n", path, m_host);
+		int requestLen = snprintf(request, 512, "%s /%s HTTP/1.1\r\nHost: %s\r\n%s\r\n", type, path, m_host, header);
 
 		int totalBytesSent = 0;
 		while (totalBytesSent < requestLen) {
diff --git a/Engine/Source/Programs/UnrealBuildAccelerator/Common/Public/UbaAWS.h b/Engine/Source/Programs/UnrealBuildAccelerator/Common/Public/UbaAWS.h
--- a/Engine/Source/Programs/UnrealBuildAccelerator/Common/Public/UbaAWS.h
+++ b/Engine/Source/Programs/UnrealBuildAccelerator/Common/Public/UbaAWS.h
@@ -22,6 +22,7 @@ namespace uba
 		static constexpr char g_imdsAutoScalingLifeCycleState[] = "latest/meta-data/autoscaling/target-lifecycle-state";
 		static constexpr char g_imdsInstanceAvailabilityZone[]	= "latest/meta-data/placement/availability-zone";
 		static constexpr char g_imdsSpotInstanceAction[]		= "latest/meta-data/spot/instance-action";
+		static constexpr char g_imdsToken[]						= "latest/api/token";
 
 
 		bool QueryInformation(Logger& logger, StringBufferBase& outExtraInfo, const tchar* rootDir)
@@ -33,21 +34,36 @@ namespace uba
 
 			u32 statusCode = 0;
 
+			StringBuffer<128> token;
+			token.Append(TC("X-aws-ec2-metadata-token: "));
+			if (!http.Query(logger, "PUT", token, statusCode, g_imdsHost, g_imdsToken, "X-aws-ec2-metadata-token-ttl-seconds: 21600\r\n"))
+				return false;
+			token.Append("\r\n");
+
+			#if PLATFORM_WINDOWS
+			char tokenString[512];
+			size_t tokenStringLen;
+			wcstombs_s(&tokenStringLen, tokenString, sizeof_array(tokenString), token.data, _TRUNCATE);
+			m_tokenString = tokenString;
+			#else
+			m_tokenString = token.data;
+			#endif
+
 			StringBuffer<128> instanceId;
-			if (!http.Get(logger, instanceId, statusCode, g_imdsHost, g_imdsInstanceId))
+			if (!http.Query(logger, "GET", instanceId, statusCode, g_imdsHost, g_imdsInstanceId, m_tokenString.c_str()))
 				return WriteIsNotAws(logger, rootDir);
 
 			outExtraInfo.Append(TC(", AWS: ")).Append(instanceId);
 
 			StringBuffer<32> instanceLifeCycle;
-			if (http.Get(logger, instanceLifeCycle, statusCode, g_imdsHost, g_imdsInstanceLifeCycle))
+			if (http.Query(logger, "GET", instanceLifeCycle, statusCode, g_imdsHost, g_imdsInstanceLifeCycle, m_tokenString.c_str()))
 			{
 				outExtraInfo.Append(' ').Append(instanceLifeCycle);
 				m_isSpot = instanceLifeCycle.Contains(TC("spot"));
 			}
 
 			StringBuffer<32> autoscaling;
-			if (http.Get(logger, autoscaling, statusCode, g_imdsHost, g_imdsAutoScalingLifeCycleState) && statusCode == 200)
+			if (http.Query(logger, "GET", autoscaling, statusCode, g_imdsHost, g_imdsAutoScalingLifeCycleState, m_tokenString.c_str()) && statusCode == 200)
 			{
 				outExtraInfo.Append(m_isSpot ? '/' : ' ').Append(TC("autoscale"));
 				m_isAutoscaling = true;
@@ -68,7 +84,7 @@ namespace uba
 
 			StringBuffer<128> availabilityZone;
 			u32 statusCode = 0;
-			if (!http.Get(logger, availabilityZone, statusCode, g_imdsHost, g_imdsInstanceAvailabilityZone))
+			if (!http.Query(logger, "GET", availabilityZone, statusCode, g_imdsHost, g_imdsInstanceAvailabilityZone, m_tokenString.c_str()))
 			{
 				if (rootDir)
 					WriteIsNotAws(logger, rootDir);
@@ -109,7 +125,7 @@ namespace uba
 			{
 				StringBuffer<1024> content;
 				u32 statusCode = 0;
-				if (http.Get(logger, content, statusCode, g_imdsHost, g_imdsSpotInstanceAction) && statusCode == 200)
+				if (http.Query(logger, "GET", content, statusCode, g_imdsHost, g_imdsSpotInstanceAction, m_tokenString.c_str()) && statusCode == 200)
 				{
 					outReason.Append(TC("AWS spot instance interruption"));
 					return true;
@@ -120,7 +136,7 @@ namespace uba
 			{
 				StringBuffer<1024> content;
 				u32 statusCode = 0;
-				if (http.Get(logger, content, statusCode, g_imdsHost, g_imdsAutoScalingLifeCycleState) && statusCode == 200)
+				if (http.Query(logger, "GET", content, statusCode, g_imdsHost, g_imdsAutoScalingLifeCycleState, m_tokenString.c_str()) && statusCode == 200)
 				{
 					//if (!content.Equals(L"InService"))
 					//{
@@ -144,6 +160,7 @@ namespace uba
 		}
 
 		TString m_availabilityZone;
+		std::string m_tokenString;
 		bool m_isSpot = false;
 		bool m_isAutoscaling = false;
 	};
diff --git a/Engine/Source/Programs/UnrealBuildAccelerator/Common/Public/UbaNetworkBackend.h b/Engine/Source/Programs/UnrealBuildAccelerator/Common/Public/UbaNetworkBackend.h
--- a/Engine/Source/Programs/UnrealBuildAccelerator/Common/Public/UbaNetworkBackend.h
+++ b/Engine/Source/Programs/UnrealBuildAccelerator/Common/Public/UbaNetworkBackend.h
@@ -79,7 +79,7 @@ namespace uba
 	public:
 		HttpConnection();
 		~HttpConnection();
-		bool Get(Logger& logger, StringBufferBase& outResponse, u32& outStatusCode, const char* host, const char* path);
+		bool Query(Logger& logger, const char* type, StringBufferBase& outResponse, u32& outStatusCode, const char* host, const char* path, const char* header = "");
 
 	private:
 		bool Connect(Logger& logger, const char* host);
