[UBA]
* Changed so all ScopedReadLock and ScopedWriteLock are macros (SCOPED_READ_LOCK and SCOPED_WRITE_LOCK).. this to be able to add code around the lock to measure contention. Contention testing code can be enabled by setting UBA_TRACK_CONTENTION to 1

diff --git a/Engine/Source/Programs/UnrealBuildAccelerator/Agent/Private/UbaAgent.cpp b/Engine/Source/Programs/UnrealBuildAccelerator/Agent/Private/UbaAgent.cpp
--- a/Engine/Source/Programs/UnrealBuildAccelerator/Agent/Private/UbaAgent.cpp
+++ b/Engine/Source/Programs/UnrealBuildAccelerator/Agent/Private/UbaAgent.cpp
@@ -879,7 +879,7 @@ namespace uba
 			const Vector<ProcessLogLine>& processLogLines = process.GetLogLines();
 			if (!processLogLines.empty())
 			{
-				ScopedWriteLock lock(logLinesLock);
+				SCOPED_WRITE_LOCK(logLinesLock, lock);
 				for (auto& line : processLogLines)
 					logLines[logLinesIndex].push_back(line);
 				if (errorCode)
@@ -903,7 +903,7 @@ namespace uba
 					name.Appendf(TC(" (exit code: %u)"), errorCode);
 					entryType = LogEntryType_Error;
 				}
-				ScopedWriteLock lock(logLinesLock);
+				SCOPED_WRITE_LOCK(logLinesLock, lock);
 				logLines[logLinesIndex].push_back({ TString(name.data), entryType });
 			}
 
@@ -1111,7 +1111,11 @@ namespace uba
 					{
 						logLinesAvailable.IsSet();
 						u32 logLinesIndexPrev;
-						logLinesLock.ScopedWrite([&]() { logLinesIndexPrev = logLinesIndex; logLinesIndex = (logLinesIndex + 1) % 2; });
+						{
+							SCOPED_WRITE_LOCK(logLinesLock, l);
+							logLinesIndexPrev = logLinesIndex;
+							logLinesIndex = (logLinesIndex + 1) % 2;
+						}
 						logger.BeginScope();
 						for (auto& s : logLines[logLinesIndexPrev])
 							logger.Log(LogEntryType_Detail, s.text.c_str(), u32(s.text.size()));
@@ -1243,6 +1247,16 @@ namespace uba
 				logger.Info(TC(""));
 				logger.EndScope();
 			}
+
+			//if (proxy.storage)
+			//	proxy.storage->PrintSummary();
+			//if (proxy.server)
+			//	proxy.server->PrintSummary(logger);
+
+			#if UBA_TRACK_CONTENTION
+			LoggerWithWriter contLogger(g_consoleLogWriter, TC(""));
+			PrintContentionSummary(contLogger);
+			#endif
 		}
 		while (poll && !isTerminating);
 
diff --git a/Engine/Source/Programs/UnrealBuildAccelerator/Common/Private/UbaFile.cpp b/Engine/Source/Programs/UnrealBuildAccelerator/Common/Private/UbaFile.cpp
--- a/Engine/Source/Programs/UnrealBuildAccelerator/Common/Private/UbaFile.cpp
+++ b/Engine/Source/Programs/UnrealBuildAccelerator/Common/Private/UbaFile.cpp
@@ -796,10 +796,10 @@ namespace uba
 		TString key(dir, dir + dirLen);
 		dir = key.c_str();
 
-		ScopedWriteLock lock(m_createdDirsLock);
+		SCOPED_WRITE_LOCK(m_createdDirsLock, lock);
 		CreatedDir& cd = m_createdDirs.try_emplace(key).first->second;
 		lock.Leave();
-		ScopedWriteLock dirLock(cd.lock);
+		SCOPED_WRITE_LOCK(cd.lock, dirLock);
 		if (cd.handled)
 			return true;
 		cd.handled = true;
@@ -827,7 +827,7 @@ namespace uba
 
 	void DirectoryCache::Clear()
 	{
-		ScopedWriteLock lock(m_createdDirsLock);
+		SCOPED_WRITE_LOCK(m_createdDirsLock, lock);
 		m_createdDirs.clear();
 	}
 }
diff --git a/Engine/Source/Programs/UnrealBuildAccelerator/Common/Private/UbaFileMapping.cpp b/Engine/Source/Programs/UnrealBuildAccelerator/Common/Private/UbaFileMapping.cpp
--- a/Engine/Source/Programs/UnrealBuildAccelerator/Common/Private/UbaFileMapping.cpp
+++ b/Engine/Source/Programs/UnrealBuildAccelerator/Common/Private/UbaFileMapping.cpp
@@ -2,10 +2,11 @@
 
 #include "UbaFileMapping.h"
 #include "UbaDirectoryIterator.h"
+#include "UbaFile.h"
 #include "UbaPlatform.h"
 #include "UbaProcessStats.h"
 #include "UbaTimer.h"
-#include "UbaFile.h"
+#include "UbaWorkManager.h"
 
 #if !PLATFORM_WINDOWS
 #include <sys/file.h>
@@ -28,7 +29,7 @@ namespace uba
 			return ::CreateFileMappingW(hFile, NULL, flProtect, dwMaximumSizeHigh, dwMaximumSizeLow, lpName);
 
 		// Experiment to try to prevent lock happening on AWS servers when lots of helpers are sending back obj files.
-		ScopedWriteLock lock(g_createFileHandleLock);
+		SCOPED_WRITE_LOCK(g_createFileHandleLock, lock);
 		return ::CreateFileMappingW(hFile, NULL, flProtect, dwMaximumSizeHigh, dwMaximumSizeLow, lpName);
 	}
 #else
@@ -51,7 +52,7 @@ namespace uba
 		StringBuffer<64> lockDir;
 		lockDir.Append("/tmp/uba_shm_locks");
 
-		ScopedWriteLock lock(g_mappingUidCounterLock);
+		SCOPED_WRITE_LOCK(g_mappingUidCounterLock, lock);
 		if (!g_mappingUidCounter)
 		{
 			// Create dir
@@ -423,7 +424,7 @@ namespace uba
 	{
 		while (true)
 		{
-			ScopedWriteLock lock(storage.availableFilesLock);
+			SCOPED_WRITE_LOCK(storage.availableFilesLock, lock);
 			for (u32 i = storage.availableFilesCount; i!=0; --i)
 			{
 				u32 index = i - 1;
@@ -461,7 +462,7 @@ namespace uba
 
 	void FileMappingBuffer::PushFile(MappingStorage& storage, File* file)
 	{
-		ScopedWriteLock lock(storage.availableFilesLock);
+		SCOPED_WRITE_LOCK(storage.availableFilesLock, lock);
 		storage.availableFiles[storage.availableFilesCount++] = file;
 		storage.availableFilesEvent.Set();
 	}
@@ -471,7 +472,7 @@ namespace uba
 		u32 locksTaken = 0;
 		while (locksTaken < storage.fileCount)
 		{
-			ScopedWriteLock lock(storage.availableFilesLock);
+			SCOPED_WRITE_LOCK(storage.availableFilesLock, lock);
 			if (storage.availableFilesCount == 0)
 			{
 				lock.Leave();
@@ -601,7 +602,7 @@ namespace uba
 
 	FileMappingAllocator::Allocation FileMappingAllocator::Alloc(const tchar* hint)
 	{
-		ScopedWriteLock lock(m_mappingLock);
+		SCOPED_WRITE_LOCK(m_mappingLock, lock);
 
 		u64 index = m_mappingCount;
 		bool needCommit = false;
@@ -646,7 +647,7 @@ namespace uba
 		if (!UnmapViewOfFile(allocation.memory, m_blockSize, m_name))
 			m_logger.Error(TC("%s - Failed to unmap view of file (%s)"), m_name, LastErrorToText().data);
 		u64 index = allocation.offset / m_blockSize;
-		ScopedWriteLock lock(m_mappingLock);
+		SCOPED_WRITE_LOCK(m_mappingLock, lock);
 		m_availableBlocks.insert(index);
 	}
 }
diff --git a/Engine/Source/Programs/UnrealBuildAccelerator/Common/Private/UbaLogger.cpp b/Engine/Source/Programs/UnrealBuildAccelerator/Common/Private/UbaLogger.cpp
--- a/Engine/Source/Programs/UnrealBuildAccelerator/Common/Private/UbaLogger.cpp
+++ b/Engine/Source/Programs/UnrealBuildAccelerator/Common/Private/UbaLogger.cpp
@@ -5,6 +5,7 @@
 #include "UbaFileAccessor.h"
 #include "UbaPlatform.h"
 #include "UbaStringBuffer.h"
+#include "UbaTimer.h"
 
 #if PLATFORM_WINDOWS
 #include <io.h>
@@ -22,7 +23,7 @@ namespace uba
 	ANALYSIS_NORETURN void UbaAssert(const tchar* text, const char* file, u32 line, const char* expr, u32 terminateCode)
 	{
 		static ReaderWriterLock assertLock;
-		ScopedWriteLock lock(assertLock);
+		SCOPED_WRITE_LOCK(assertLock, lock);
 
 		StringBuffer<4096> b;
 		WriteAssertInfo(b, text, file, line, expr, 1);
@@ -227,7 +228,7 @@ namespace uba
 	{
 		if (t_consoleLogScopeCount)
 			return LogNoLock(type, str, strLen, prefix, prefixLen);
-		ScopedWriteLock lock(m_lock);
+		SCOPED_WRITE_LOCK(m_lock, lock);
 		LogNoLock(type, str, strLen, prefix, prefixLen);
 #if PLATFORM_WINDOWS
 		if (!m_stdout)
@@ -341,7 +342,7 @@ namespace uba
 				return;
 			if (t_debugLogScopeCount)
 				return LogNoLock(type, str, strLen, prefix, prefixLen);
-			ScopedWriteLock lock(m_logLock);
+			SCOPED_WRITE_LOCK(m_logLock, lock);
 			LogNoLock(type, str, strLen, prefix, prefixLen);
 		}
 
@@ -395,4 +396,35 @@ namespace uba
 
 	LoggerWithWriter g_debugLogger(g_debugLogWriter);
 #endif
+
+	#if UBA_TRACK_CONTENTION
+	List<ContentionTracker>& GetContentionTrackerList();
+	#endif
+
+
+	void PrintContentionSummary(Logger& logger)
+	{
+	#if UBA_TRACK_CONTENTION
+		logger.Info(TC("Contention summary:"));
+		List<ContentionTracker*> list;
+		for (auto& ct : GetContentionTrackerList())
+			if (TimeToMs(ct.time) > 1)
+				list.push_back(&ct);
+		list.sort([](const ContentionTracker* a, const ContentionTracker* b)
+			{
+				if (a->time != b->time)
+					return a->time > b->time;
+				return a < b;
+			});
+
+		for (auto& ct : list)
+		{
+			StringBuffer<512> fn;
+			fn.Append(ct->file);
+			StringBuffer<256> s;
+			s.Append(TC("  ")).AppendFileName(fn.data).Append(':').AppendValue(ct->line).Append(TC(" - ")).AppendValue(ct->count).Append(TC(" calls in ")).Append(TimeToText(ct->time).str);
+			logger.Info(s.data);
+		}
+	#endif
+	}
 }
diff --git a/Engine/Source/Programs/UnrealBuildAccelerator/Common/Private/UbaNetworkBackendQuic.cpp b/Engine/Source/Programs/UnrealBuildAccelerator/Common/Private/UbaNetworkBackendQuic.cpp
--- a/Engine/Source/Programs/UnrealBuildAccelerator/Common/Private/UbaNetworkBackendQuic.cpp
+++ b/Engine/Source/Programs/UnrealBuildAccelerator/Common/Private/UbaNetworkBackendQuic.cpp
@@ -57,7 +57,7 @@ namespace uba
 	void NetworkBackendQuic::Close(void* connection)
 	{
 		auto& conn = *(Connection*)connection;
-		ScopedWriteLock lock(m_connectionsLock);
+		SCOPED_WRITE_LOCK(m_connectionsLock, lock);
 		bool erased = m_connections.erase(conn.connectionHandle);
 		lock.Leave();
 		if (erased)
@@ -353,7 +353,7 @@ namespace uba
 
 	NetworkBackendQuic::~NetworkBackendQuic()
 	{
-		ScopedWriteLock lock(m_connectionsLock);
+		SCOPED_WRITE_LOCK(m_connectionsLock, lock);
 		for (auto& kv : m_connections)
 			kv.second->Release();
 		m_connections.clear();
@@ -433,7 +433,7 @@ namespace uba
 			nc.Info->RemoteAddress->Ipv4;
 
 			auto connPtr = new Connection{ *this, nc.Connection };
-			ScopedWriteLock lock(m_connectionsLock);
+			SCOPED_WRITE_LOCK(m_connectionsLock, lock);
 			auto insres = m_connections.try_emplace(nc.Connection, connPtr);
 			UBA_ASSERT(insres.second);
 			Connection& conn = *connPtr;
@@ -548,14 +548,14 @@ namespace uba
 		Connection& conn = *connPtr;
 		conn.connectionHandle = connection;
 
-		ScopedWriteLock lock(m_connectionsLock);
+		SCOPED_WRITE_LOCK(m_connectionsLock, lock);
 		auto insres = m_connections.try_emplace(connection, connPtr);
 		UBA_ASSERT(insres.second);
 		lock.Leave();
 
 		auto connectionGuard = MakeGuard([&]()
 			{
-				ScopedWriteLock lock(m_connectionsLock);
+				SCOPED_WRITE_LOCK(m_connectionsLock, lock);
 				if (m_connections.erase(connection))
 					conn.Release();
 			});
@@ -609,7 +609,7 @@ namespace uba
 		case QUIC_CONNECTION_EVENT_CONNECTED:
 		{
 			m_logger.Info(L"[conn][%p] Connected", connection); // The handshake has completed for the connection.
-			ScopedReadLock lock(m_connectionsLock);
+			SCOPED_READ_LOCK(m_connectionsLock, lock);
 			Connection& conn = *m_connections.find(connection)->second;
 			lock.Leave();
 
@@ -670,7 +670,7 @@ namespace uba
 
 	void NetworkBackendQuic::StreamCreated(HQUIC connection, HQUIC streamHandle)
 	{
-		ScopedWriteLock lock(m_connectionsLock);
+		SCOPED_WRITE_LOCK(m_connectionsLock, lock);
 		auto findIt = m_connections.find(connection);
 		UBA_ASSERT(findIt != m_connections.end());
 		Connection& conn = *findIt->second;
@@ -700,7 +700,7 @@ namespace uba
 
 	void NetworkBackendQuic::CloseConnection(HQUIC connection)
 	{
-		ScopedWriteLock lock(m_connectionsLock);
+		SCOPED_WRITE_LOCK(m_connectionsLock, lock);
 		auto findIt = m_connections.find(connection);
 		if (findIt == m_connections.end())
 			return;
diff --git a/Engine/Source/Programs/UnrealBuildAccelerator/Common/Private/UbaNetworkBackendTcp.cpp b/Engine/Source/Programs/UnrealBuildAccelerator/Common/Private/UbaNetworkBackendTcp.cpp
--- a/Engine/Source/Programs/UnrealBuildAccelerator/Common/Private/UbaNetworkBackendTcp.cpp
+++ b/Engine/Source/Programs/UnrealBuildAccelerator/Common/Private/UbaNetworkBackendTcp.cpp
@@ -122,7 +122,7 @@ namespace uba
 	{
 		StopListen();
 
-		ScopedWriteLock lock(m_connectionsLock);
+		SCOPED_WRITE_LOCK(m_connectionsLock, lock);
 		for (auto& conn : m_connections)
 		{
 			ScopedCriticalSection lock2(conn.shutdownLock);
@@ -158,7 +158,7 @@ namespace uba
 		sendContext.isUsed = true;
 
 		#if UBA_LOCK_AROUND_SEND
-		ScopedWriteLock lock(conn.sendLock);
+		SCOPED_WRITE_LOCK(conn.sendLock, lock);
 		#endif
 
 		bool res = SendSocket(logger, conn.socket, data, dataSize);
@@ -378,7 +378,7 @@ namespace uba
 				continue;
 			}
 
-			ScopedWriteLock lock(m_connectionsLock);
+			SCOPED_WRITE_LOCK(m_connectionsLock, lock);
 			auto it = m_connections.emplace(m_connections.end(), clientSocket);
 			auto& conn = *it;
 			conn.recvThread.Start([this, connPtr = &conn] { ThreadRecv(*connPtr); return 0; });
@@ -389,7 +389,7 @@ namespace uba
 				shutdown(clientSocket, SD_BOTH);
 				conn.ready.Set();
 				conn.recvThread.Wait();
-				ScopedWriteLock lock2(m_connectionsLock);
+				SCOPED_WRITE_LOCK(m_connectionsLock, lock2);
 				m_connections.erase(it);
 				continue;
 			}
@@ -616,7 +616,7 @@ namespace uba
 
 		socketClose.Cancel();
 
-		ScopedWriteLock lock(m_connectionsLock);
+		SCOPED_WRITE_LOCK(m_connectionsLock, lock);
 		auto it = m_connections.emplace(m_connections.end(), socketFd);
 		auto& conn = *it;
 		conn.recvThread.Start([this, connPtr = &conn] { ThreadRecv(*connPtr); return 0; });
@@ -627,7 +627,7 @@ namespace uba
 			shutdown(socketFd, SD_BOTH);
 			conn.ready.Set();
 			conn.recvThread.Wait();
-			ScopedWriteLock lock2(m_connectionsLock);
+			SCOPED_WRITE_LOCK(m_connectionsLock, lock2);
 			m_connections.erase(it);
 			return false;
 		}
diff --git a/Engine/Source/Programs/UnrealBuildAccelerator/Common/Private/UbaNetworkClient.cpp b/Engine/Source/Programs/UnrealBuildAccelerator/Common/Private/UbaNetworkClient.cpp
--- a/Engine/Source/Programs/UnrealBuildAccelerator/Common/Private/UbaNetworkClient.cpp
+++ b/Engine/Source/Programs/UnrealBuildAccelerator/Common/Private/UbaNetworkClient.cpp
@@ -42,7 +42,8 @@ namespace uba
 		StopListen();
 
 		Disconnect();
-		ScopedWriteLock lock(m_connectionsLock);
+
+		SCOPED_WRITE_LOCK(m_connectionsLock, lock);
 		m_connections.clear(); // Before tcpBackend
 
 		delete m_tcpBackend;
@@ -102,7 +103,7 @@ namespace uba
 
 				if (!rc.error)
 				{
-					ScopedWriteLock lock(rc.client.m_serverUidLock);
+					SCOPED_WRITE_LOCK(rc.client.m_serverUidLock, lock);
 					if (rc.client.m_serverUid == Guid())
 						rc.client.m_serverUid = serverUid;
 					else if (rc.client.m_serverUid != serverUid) // Seems like two different servers tried to connect to this client.. keep the first one and ignore the others
@@ -189,7 +190,7 @@ namespace uba
 		if (m_connectionCount.fetch_add(1) != 0)
 			return true;
 
-		ScopedWriteLock lock(m_onConnectedFunctionsLock);
+		SCOPED_WRITE_LOCK(m_onConnectedFunctionsLock, lock);
 		for (auto& f : m_onConnectedFunctions)
 			f();
 		lock.Leave();
@@ -203,13 +204,16 @@ namespace uba
 
 	void NetworkClient::ConnectedCallback(NetworkBackend& backend, void* backendConnection)
 	{
-		ScopedWriteLock lock(m_connectionsLock);
+		SCOPED_WRITE_LOCK(m_connectionsLock, lock);
 		m_connections.emplace_back(*this);
 		Connection* connection = &m_connections.back();
 		connection->backendConnection = backendConnection;
 		connection->connected = 1;
 		connection->backend = &backend;
-		m_connectionsItLock.ScopedWrite([&]() { m_connectionsIt = m_connections.begin(); });
+		{
+			SCOPED_WRITE_LOCK(m_connectionsItLock, l);
+			m_connectionsIt = m_connections.begin();
+		}
 		lock.Leave();
 
 		backend.SetDisconnectCallback(backendConnection, connection, [](void* context, void* connection)
@@ -229,7 +233,7 @@ namespace uba
 		u16 messageId = u16(headerData[0] << 8) | u16((*(u32*)(headerData + 1) & 0xff000000) >> 24);
 		u32 messageSize = *(u32*)(headerData + 1) & 0x00FFFFFF;
 
-		ScopedReadLock lock(client.m_activeMessagesLock);
+		SCOPED_READ_LOCK(client.m_activeMessagesLock, lock);
 		if (!connection.connected)
 			return false;
 		UBA_ASSERTF(messageId < client.m_activeMessages.size(), TC("Message id %u is higher than max %u"), messageId, u32(client.m_activeMessages.size()));
@@ -275,7 +279,7 @@ namespace uba
 
 	void NetworkClient::Disconnect()
 	{
-		ScopedReadLock lock(m_connectionsLock);
+		SCOPED_READ_LOCK(m_connectionsLock, lock);
 		for (auto& c : m_connections)
 		{
 			OnDisconnected(c, false);
@@ -335,7 +339,7 @@ namespace uba
 		u64 recvBytes = 0;
 		u32 recvCount = 0;
 
-		ScopedReadLock lock(m_connectionsLock);
+		SCOPED_READ_LOCK(m_connectionsLock, lock);
 		u32 connectionsCount = u32(m_connections.size());
 		for (auto& c : m_connections)
 		{
@@ -365,7 +369,7 @@ namespace uba
 
 	void NetworkClient::RegisterOnConnected(const OnConnectedFunction& function)
 	{
-		ScopedWriteLock lock(m_onConnectedFunctionsLock);
+		SCOPED_WRITE_LOCK(m_onConnectedFunctionsLock, lock);
 		m_onConnectedFunctions.push_back(function);
 		if (m_connectionCount.load() == 0)
 			return;
@@ -422,7 +426,7 @@ namespace uba
 		}
 
 		u16 messageId = 0;
-		ScopedWriteLock lock(m_activeMessagesLock);
+		SCOPED_WRITE_LOCK(m_activeMessagesLock, lock);
 		for (auto m : m_activeMessages)
 		{
 			if (m && m->m_connection == &connection)
@@ -437,8 +441,8 @@ namespace uba
 
 	bool NetworkClient::Send(NetworkMessage& message, void* response, u32 responseCapacity, bool async)
 	{
-		ScopedReadLock connectionLock(m_connectionsLock);
-		ScopedWriteLock connectionItLock(m_connectionsItLock);
+		SCOPED_READ_LOCK(m_connectionsLock, connectionLock);
+		SCOPED_WRITE_LOCK(m_connectionsItLock, connectionItLock);
 		if (m_connectionsIt == m_connections.end())
 			return false;
 
@@ -461,7 +465,7 @@ namespace uba
 		{
 			while (true)
 			{
-				ScopedWriteLock lock(m_activeMessagesLock);
+				SCOPED_WRITE_LOCK(m_activeMessagesLock, lock);
 				if (m_availableMessageIds.empty())
 				{
 					if (!connection.connected)
@@ -559,7 +563,7 @@ namespace uba
 
 	void NetworkClient::ReturnMessageId(u16 id)
 	{
-		ScopedWriteLock lock(m_activeMessagesLock);
+		SCOPED_WRITE_LOCK(m_activeMessagesLock, lock);
 		m_availableMessageIds.push_back(id);
 		m_activeMessages[id] = nullptr;
 	}
diff --git a/Engine/Source/Programs/UnrealBuildAccelerator/Common/Private/UbaNetworkServer.cpp b/Engine/Source/Programs/UnrealBuildAccelerator/Common/Private/UbaNetworkServer.cpp
--- a/Engine/Source/Programs/UnrealBuildAccelerator/Common/Private/UbaNetworkServer.cpp
+++ b/Engine/Source/Programs/UnrealBuildAccelerator/Common/Private/UbaNetworkServer.cpp
@@ -58,7 +58,7 @@ namespace uba
 		void Stop(NetworkServer& server)
 		{
 			m_loop = false;
-			ScopedWriteLock lock(server.m_availableWorkersLock);
+			SCOPED_WRITE_LOCK(server.m_availableWorkersLock, lock);
 			while (m_inUse)
 			{
 				m_context->workAvailable.Set();
@@ -216,7 +216,7 @@ namespace uba
 
 			if (!server.m_allowNewClients)
 			{
-				ScopedReadLock clientsLock(server.m_clientsLock);
+				SCOPED_READ_LOCK(server.m_clientsLock, clientsLock);
 				bool found = false;
 				for (auto& kv : server.m_clients)
 					found |= kv.second.uid == clientUid;
@@ -233,7 +233,7 @@ namespace uba
 			if (!conn.SendInitialResponse(0))
 				return false;
 			
-			ScopedWriteLock clientsLock(server.m_clientsLock);
+			SCOPED_WRITE_LOCK(server.m_clientsLock, clientsLock);
 			u32 clientId = u32(server.m_clients.size() + 1);
 			for (auto& kv : server.m_clients)
 				if (kv.second.uid == clientUid)
@@ -320,7 +320,7 @@ namespace uba
 
 		bool SetShouldDisconnect()
 		{
-			ScopedWriteLock lock(m_shutdownLock);
+			SCOPED_WRITE_LOCK(m_shutdownLock, lock);
 			bool isConnected = !m_shouldDisconnect;
 			m_shouldDisconnect = true;
 			return isConnected;
@@ -334,7 +334,7 @@ namespace uba
 
 		void TestDisconnect()
 		{
-			ScopedWriteLock lock(m_shutdownLock);
+			SCOPED_WRITE_LOCK(m_shutdownLock, lock);
 			if (!m_shouldDisconnect)
 				return;
 			if (m_disconnected)
@@ -398,7 +398,7 @@ namespace uba
 	bool ConnectionInfo::ShouldDisconnect() const
 	{
 		auto& conn = *(NetworkServer::Connection*)internalData;
-		ScopedWriteLock lock(conn.m_shutdownLock);
+		SCOPED_WRITE_LOCK(conn.m_shutdownLock, lock);
 		return conn.m_shouldDisconnect;
 	}
 
@@ -517,7 +517,7 @@ namespace uba
 			while (true)
 			{
 				AdditionalWork work;
-				ScopedWriteLock lock(server.m_additionalWorkLock);
+				SCOPED_WRITE_LOCK(server.m_additionalWorkLock, lock);
 				if (server.m_additionalWork.empty())
 					break;
 				work = server.m_additionalWork.front();
@@ -541,8 +541,8 @@ namespace uba
 			// is present before making ourself available to avoid
 			// a race where AddWork would not see this thread in the
 			// available list after adding some work.
-			ScopedWriteLock lock1(server.m_availableWorkersLock);
-			ScopedReadLock lock2(server.m_additionalWorkLock);
+			SCOPED_WRITE_LOCK(server.m_availableWorkersLock, lock1);
+			SCOPED_READ_LOCK(server.m_additionalWorkLock, lock2);
 			// Verify there is not additional work while we hold both lock
 			// and only add ourself as available if no additional work is present.
 			if (!server.m_additionalWork.empty())
@@ -638,17 +638,17 @@ namespace uba
 		StopListen();
 
 		{
-			ScopedWriteLock lock(m_availableWorkersLock);
+			SCOPED_WRITE_LOCK(m_availableWorkersLock, lock);
 			m_workersEnabled = false;
 			m_workerAvailable.Set();
 		}
 		{
-			ScopedWriteLock lock(m_addConnectionsLock);
+			SCOPED_WRITE_LOCK(m_addConnectionsLock, lock);
 			m_addConnections.clear();
 		}
 
 		{
-			ScopedWriteLock lock(m_connectionsLock);
+			SCOPED_WRITE_LOCK(m_connectionsLock, lock);
 			bool success = true;
 			for (auto& c : m_connections)
 			{
@@ -665,7 +665,7 @@ namespace uba
 		}
 
 
-		ScopedWriteLock lock(m_availableWorkersLock);
+		SCOPED_WRITE_LOCK(m_availableWorkersLock, lock);
 		while (auto worker = m_firstActiveWorker)
 		{
 			lock.Leave();
@@ -687,7 +687,7 @@ namespace uba
 
 	bool NetworkServer::AddClient(NetworkBackend& backend, const tchar* ip, u16 port, const u8* cryptoKey128)
 	{
-		ScopedWriteLock lock(m_addConnectionsLock);
+		SCOPED_WRITE_LOCK(m_addConnectionsLock, lock);
 		for (auto it = m_addConnections.begin(); it != m_addConnections.end();)
 		{
 			if (it->Wait(0))
@@ -760,7 +760,7 @@ namespace uba
 
 	void NetworkServer::UnregisterService(u8 serviceId)
 	{
-		ScopedWriteLock lock(m_connectionsLock);
+		SCOPED_WRITE_LOCK(m_connectionsLock, lock);
 		UBA_ASSERTF(!m_listenBackend, TC("Server still listens to new connections while service is unregistered"));
 		UBA_ASSERTF(m_connections.empty(), TC("Unregistering service while still having live connections"));
 		WorkerRec& rec = m_workerFunctions[serviceId];
@@ -776,7 +776,7 @@ namespace uba
 
 	void NetworkServer::UnregisterOnClientConnected(u8 id)
 	{
-		ScopedWriteLock lock(m_connectionsLock);
+		SCOPED_WRITE_LOCK(m_connectionsLock, lock);
 		UBA_ASSERT(!m_listenBackend);
 		UBA_ASSERT(m_connections.empty());
 		m_onConnectionFunction = {};
@@ -800,7 +800,7 @@ namespace uba
 
 	void NetworkServer::AddWork(const Function<void()>& work, u32 count, const tchar* desc)
 	{
-		ScopedWriteLock lock(m_additionalWorkLock);
+		SCOPED_WRITE_LOCK(m_additionalWorkLock, lock);
 		for (u32 i = 0; i != count; ++i)
 		{
 			m_additionalWork.push_back({ work });
@@ -809,7 +809,7 @@ namespace uba
 		}
 		lock.Leave();
 
-		ScopedWriteLock lock2(m_availableWorkersLock);
+		SCOPED_WRITE_LOCK(m_availableWorkersLock, lock2);
 		while (count-- && m_createdWorkerCount < m_maxWorkerCount)
 		{
 			Worker* worker = PopWorkerNoLock();
@@ -838,7 +838,7 @@ namespace uba
 
 	void NetworkServer::GetClientStats(ClientStats& out, u32 clientId)
 	{
-		ScopedReadLock lock(m_clientsLock);
+		SCOPED_READ_LOCK(m_clientsLock, lock);
 		auto findIt = m_clients.find(clientId);
 		if (findIt == m_clients.end())
 			return;
@@ -851,12 +851,12 @@ namespace uba
 	bool NetworkServer::DoAdditionalWork()
 	{
 		AdditionalWork work;
-		ScopedWriteLock lock(m_additionalWorkLock);
+		SCOPED_WRITE_LOCK(m_additionalWorkLock, lock);
 		if (m_additionalWork.empty())
 		{
 			lock.Leave();
 
-			ScopedWriteLock lock2(m_availableWorkersLock);
+			SCOPED_WRITE_LOCK(m_availableWorkersLock, lock2);
 			if (m_createdWorkerCount != m_maxWorkerCount)
 				return false;
 			lock2.Leave();
@@ -933,7 +933,7 @@ namespace uba
 	{
 		while (true)
 		{
-			ScopedWriteLock lock(m_availableWorkersLock);
+			SCOPED_WRITE_LOCK(m_availableWorkersLock, lock);
 			if (!m_workersEnabled)
 				return nullptr;
 			if (auto worker = PopWorkerNoLock())
@@ -974,7 +974,7 @@ namespace uba
 
 	void NetworkServer::PushWorker(Worker* worker)
 	{
-		ScopedWriteLock lock(m_availableWorkersLock);
+		SCOPED_WRITE_LOCK(m_availableWorkersLock, lock);
 		PushWorkerNoLock(worker);
 	}
 
@@ -1021,7 +1021,7 @@ namespace uba
 			{
 				u32 connectionCount = reader.ReadU32();
 
-				ScopedReadLock lock(m_clientsLock);
+				SCOPED_READ_LOCK(m_clientsLock, lock);
 				auto findIt = m_clients.find(connectionInfo.GetId());
 				if (findIt == m_clients.end())
 					return true;
@@ -1033,7 +1033,7 @@ namespace uba
 				u32 toAdd = connectionCount - c.connectionCount;
 
 				auto connPtr = (NetworkServer::Connection*)connectionInfo.internalData;
-				ScopedWriteLock lock2(m_addConnectionsLock);
+				SCOPED_WRITE_LOCK(m_addConnectionsLock, lock2);
 				for (u32 i = 0; i != toAdd; ++i)
 				{
 					m_addConnections.emplace_back([this, connPtr]()
@@ -1066,7 +1066,7 @@ namespace uba
 
 	bool NetworkServer::AddConnection(NetworkBackend& backend, void* backendConnection, const sockaddr& remoteSocketAddr, CryptoKey cryptoKey)
 	{
-		ScopedWriteLock lock(m_connectionsLock);
+		SCOPED_WRITE_LOCK(m_connectionsLock, lock);
 
 		RemoveDisconnectedConnections();
 
diff --git a/Engine/Source/Programs/UnrealBuildAccelerator/Common/Private/UbaProcess.cpp b/Engine/Source/Programs/UnrealBuildAccelerator/Common/Private/UbaProcess.cpp
--- a/Engine/Source/Programs/UnrealBuildAccelerator/Common/Private/UbaProcess.cpp
+++ b/Engine/Source/Programs/UnrealBuildAccelerator/Common/Private/UbaProcess.cpp
@@ -110,7 +110,7 @@ namespace uba
 	{
 		{
 			#if !PLATFORM_WINDOWS
-			ScopedWriteLock lock(m_comMemoryLock);
+			SCOPED_WRITE_LOCK(m_comMemoryLock, lock);
 			#endif
 			if (m_comMemory.memory)
 				m_cancelEvent.Set();
@@ -332,7 +332,7 @@ namespace uba
 		m_cancelEvent.Set();
 		#else
 		m_cancelled = true;
-		ScopedWriteLock lock(m_comMemoryLock);
+		SCOPED_WRITE_LOCK(m_comMemoryLock, lock);
 		if (m_comMemory.memory)
 			m_cancelEvent.Set();
 		#endif
@@ -475,7 +475,7 @@ namespace uba
 
 		{
 			#if !PLATFORM_WINDOWS
-			ScopedWriteLock lock(m_comMemoryLock);
+			SCOPED_WRITE_LOCK(m_comMemoryLock, lock);
 			m_cancelEvent.~Event();
 			m_writeEvent.~Event();
 			m_readEvent.~Event();
@@ -956,7 +956,7 @@ namespace uba
 			m_session.m_logger.Log(LogEntryType_Warning, line.c_str(), u32(line.size()));
 		if (m_startInfo.logLineFunc)
 			m_startInfo.logLineFunc(m_startInfo.logLineUserData, line.c_str(), u32(line.size()), logType);
-		ScopedWriteLock l(m_logLinesLock);
+		SCOPED_WRITE_LOCK(m_logLinesLock, l);
 		m_logLines.push_back({ std::move(line), logType });
 	}
 
@@ -977,7 +977,7 @@ namespace uba
 			return true;
 		}
 
-		ScopedWriteLock tempLock(m_tempFilesLock);
+		SCOPED_WRITE_LOCK(m_tempFilesLock, tempLock);
 		auto insres = m_tempFiles.try_emplace(key, WrittenFile{ nullptr, StringKeyZero, fileName.data, newHandle, mappingHandleSize, mappingHandle });
 		if (insres.second)
 			return true;
@@ -999,7 +999,7 @@ namespace uba
 		u64 mappingHandle = 0;
 		u64 mappingWritten = 0;
 
-		ScopedReadLock lock(m_tempFilesLock);
+		SCOPED_READ_LOCK(m_tempFilesLock, lock);
 		auto findIt = m_tempFiles.find(fileKey);
 		if (findIt != m_tempFiles.end())
 		{
@@ -1081,7 +1081,7 @@ namespace uba
 	bool ProcessImpl::WriteFilesToDisk()
 	{
 		TimerScope ts(m_processStats.writeFiles);
-		ScopedWriteLock lock(m_writtenFilesLock);
+		SCOPED_WRITE_LOCK(m_writtenFilesLock, lock);
 		for (auto& kv : m_writtenFiles)
 		{
 			if (kv.second.owner != this)
@@ -1125,7 +1125,7 @@ namespace uba
 
 	u32 ProcessImpl::InternalCreateProcess(bool runningRemote, void* environment, FileMappingHandle communicationHandle, u64 communicationOffset)
 	{
-		ScopedWriteLock initLock(m_initLock);
+		SCOPED_WRITE_LOCK(m_initLock, initLock);
 		Logger& logger = m_session.m_logger;
 
 #if PLATFORM_WINDOWS
@@ -1178,7 +1178,7 @@ namespace uba
 			creationFlags |= EXTENDED_STARTUPINFO_PRESENT;
 
 
-			ScopedReadLock jobObjectLock(m_session.m_processJobObjectLock);
+			SCOPED_READ_LOCK(m_session.m_processJobObjectLock, jobObjectLock);
 			if (!m_session.m_processJobObject)
 			{
 				m_cancelEvent.Set();
@@ -1629,7 +1629,7 @@ namespace uba
 
 	u32 ProcessImpl::InternalExitProcess(bool cancel)
 	{
-		ScopedWriteLock lock(m_initLock);
+		SCOPED_WRITE_LOCK(m_initLock, lock);
 		Logger& logger = m_session.m_logger;
 
 		ProcHandle handle = m_nativeProcessHandle;
diff --git a/Engine/Source/Programs/UnrealBuildAccelerator/Common/Private/UbaScheduler.cpp b/Engine/Source/Programs/UnrealBuildAccelerator/Common/Private/UbaScheduler.cpp
--- a/Engine/Source/Programs/UnrealBuildAccelerator/Common/Private/UbaScheduler.cpp
+++ b/Engine/Source/Programs/UnrealBuildAccelerator/Common/Private/UbaScheduler.cpp
@@ -88,7 +88,7 @@ namespace uba
 		m_thread.Wait();
 		m_session.WaitOnAllTasks();
 
-		ScopedWriteLock lock(m_processEntriesLock);
+		SCOPED_WRITE_LOCK(m_processEntriesLock, lock);
 		for (auto& entry : m_processEntries)
 		{
 			UBA_ASSERT(entry.status !=ProcessStatus_Running);
@@ -123,7 +123,7 @@ namespace uba
 		auto info2 = new ProcessStartInfo2(info.info, ki, info.knownInputsCount);
 		info2->weight = info.weight;
 
-		ScopedWriteLock lock(m_processEntriesLock);
+		SCOPED_WRITE_LOCK(m_processEntriesLock, lock);
 		u32 index = u32(m_processEntries.size());
 		auto& entry = m_processEntries.emplace_back();
 		entry.info = info2;
@@ -177,7 +177,7 @@ namespace uba
 		if (processIndex == ~0u)
 			return;
 
-		ScopedWriteLock lock(m_processEntriesLock);
+		SCOPED_WRITE_LOCK(m_processEntriesLock, lock);
 		if (m_processEntries[processIndex].status != ProcessStatus_Running)
 			return;
 		m_processEntries[processIndex].status = ProcessStatus_Queued;
@@ -278,7 +278,7 @@ namespace uba
 	{
 		while (true)
 		{
-			ScopedWriteLock lock(m_processEntriesLock);
+			SCOPED_WRITE_LOCK(m_processEntriesLock, lock);
 			u32 indexToRun = PopProcess(isLocal);
 			if (indexToRun == ~0u)
 				return false;
@@ -342,7 +342,7 @@ namespace uba
 
 		while (true)
 		{
-			ScopedWriteLock lock(m_processEntriesLock);
+			SCOPED_WRITE_LOCK(m_processEntriesLock, lock);
 			u32 indexToRun = PopProcess(isLocal);
 			if (indexToRun == ~0u)
 				return false;
@@ -383,7 +383,7 @@ namespace uba
 		if (auto func = si->startInfo.exitedFunc)
 			func(si->startInfo.userData, ph);
 
-		ScopedWriteLock lock(m_processEntriesLock);
+		SCOPED_WRITE_LOCK(m_processEntriesLock, lock);
 		auto& entry = m_processEntries[info.processIndex];
 		u32* dependencies = entry.dependencies;
 		entry.status = exitCode == 0 ? ProcessStatus_Success : ProcessStatus_Failed;
diff --git a/Engine/Source/Programs/UnrealBuildAccelerator/Common/Private/UbaSession.cpp b/Engine/Source/Programs/UnrealBuildAccelerator/Common/Private/UbaSession.cpp
--- a/Engine/Source/Programs/UnrealBuildAccelerator/Common/Private/UbaSession.cpp
+++ b/Engine/Source/Programs/UnrealBuildAccelerator/Common/Private/UbaSession.cpp
@@ -9,6 +9,7 @@
 #include "UbaApplicationRules.h"
 #include "UbaPathUtils.h"
 #include "UbaProtocol.h"
+#include "UbaWorkManager.h"
 
 #if PLATFORM_WINDOWS
 #include "UbaWinBinDependencyParser.h"
@@ -274,7 +275,7 @@ namespace uba
 
 		u32 tableOffset;
 
-		ScopedWriteLock memoryLock(dirTable.m_memoryLock);
+		SCOPED_WRITE_LOCK(dirTable.m_memoryLock, memoryLock);
 		u32 writePos = dirTable.m_memorySize;
 		BinaryWriter tableWriter(dirTable.m_memory + dirTable.m_memorySize);
 
@@ -319,12 +320,12 @@ namespace uba
 	void Session::WriteDirectoryEntriesRecursive(const StringKey& dirKey, tchar* dirPath, u32& outTableOffset)
 	{
 		auto& dirTable = m_directoryTable;
-		ScopedWriteLock lookupLock(dirTable.m_lookupLock);
+		SCOPED_WRITE_LOCK(dirTable.m_lookupLock, lookupLock);
 		auto res = dirTable.m_lookup.try_emplace(dirKey, dirTable.m_memoryBlock);
 		DirectoryTable::Directory& dir = res.first->second;
 		lookupLock.Leave();
 
-		ScopedWriteLock dirLock(dir.lock);
+		SCOPED_WRITE_LOCK(dir.lock, dirLock);
 
 		if (dir.parseOffset == 1)
 		{
@@ -382,24 +383,24 @@ namespace uba
 	{
 		auto& dirTable = m_directoryTable;
 		WriteDirectoryEntriesRecursive(dirKey, dirPath, outTableOffset);
-		ScopedReadLock memoryLock(dirTable.m_memoryLock);
+		SCOPED_READ_LOCK(dirTable.m_memoryLock, memoryLock);
 		return dirTable.m_memorySize;
 	}
 
 	u32 Session::AddFileMapping(StringKey fileNameKey, const tchar* fileName, const tchar* newFileName, u64 fileSize)
 	{
 		UBA_ASSERT(fileNameKey != StringKeyZero);
-		ScopedWriteLock lookupLock(m_fileMappingTableLookupLock);
+		SCOPED_WRITE_LOCK(m_fileMappingTableLookupLock, lookupLock);
 		auto insres = m_fileMappingTableLookup.try_emplace(fileNameKey);
 		FileMappingEntry& entry = insres.first->second;
 		lookupLock.Leave();
 
-		ScopedWriteLock entryCs(entry.lock);
+		SCOPED_WRITE_LOCK(entry.lock, entryCs);
 
 		if (entry.handled)
 		{
 			entryCs.Leave();
-			ScopedReadLock lookupCs2(m_fileMappingTableMemLock);
+			SCOPED_READ_LOCK(m_fileMappingTableMemLock, lookupCs2);
 			return entry.success ? m_fileMappingTableSize : 0;
 		}
 
@@ -409,7 +410,7 @@ namespace uba
 		entry.mapping = {};
 		entry.handled = true;
 
-		ScopedWriteLock lock(m_fileMappingTableMemLock);
+		SCOPED_WRITE_LOCK(m_fileMappingTableMemLock, lock);
 		BinaryWriter writer(m_fileMappingTableMem, m_fileMappingTableSize);
 		writer.WriteStringKey(fileNameKey);
 		writer.WriteString(newFileName);
@@ -423,12 +424,12 @@ namespace uba
 	{
 		TimerScope ts(Stats().waitMmapFromFile);
 
-		ScopedWriteLock lookupLock(m_fileMappingTableLookupLock);
+		SCOPED_WRITE_LOCK(m_fileMappingTableLookupLock, lookupLock);
 		auto insres = m_fileMappingTableLookup.try_emplace(fileNameKey);
 		FileMappingEntry& entry = insres.first->second;
 		lookupLock.Leave();
 
-		ScopedWriteLock entryCs(entry.lock);
+		SCOPED_WRITE_LOCK(entry.lock, entryCs);
 
 		if (entry.handled)
 		{
@@ -520,7 +521,7 @@ namespace uba
 		entry.success = true;
 
 		{
-			ScopedWriteLock lock(m_fileMappingTableMemLock);
+			SCOPED_WRITE_LOCK(m_fileMappingTableMemLock, lock);
 			BinaryWriter writer(m_fileMappingTableMem, m_fileMappingTableSize);
 			writer.WriteStringKey(fileNameKey);
 			writer.WriteString(out.name);
@@ -537,12 +538,12 @@ namespace uba
 
 	bool Session::CreateMemoryMapFromView(MemoryMap& out, StringKey fileNameKey, const tchar* fileName, const CasKey& casKey, u64 alignment)
 	{
-		ScopedWriteLock lookupLock(m_fileMappingTableLookupLock);
+		SCOPED_WRITE_LOCK(m_fileMappingTableLookupLock, lookupLock);
 		auto insres = m_fileMappingTableLookup.try_emplace(fileNameKey);
 		FileMappingEntry& entry = insres.first->second;
 		lookupLock.Leave();
 
-		ScopedWriteLock entryCs(entry.lock);
+		SCOPED_WRITE_LOCK(entry.lock, entryCs);
 
 		if (entry.handled)
 		{
@@ -604,7 +605,7 @@ namespace uba
 		entry.success = true;
 
 		{
-			ScopedWriteLock lock(m_fileMappingTableMemLock);
+			SCOPED_WRITE_LOCK(m_fileMappingTableMemLock, lock);
 			BinaryWriter writer(m_fileMappingTableMem, m_fileMappingTableSize);
 			writer.WriteStringKey(fileNameKey);
 			writer.WriteString(out.name);
@@ -649,13 +650,13 @@ namespace uba
 			
 		#if 0//_DEBUG  // Bring this back, turned off right now because a few lines above the call to this method we add a mapping
 		{
-			ScopedWriteLock lookupLock(m_fileMappingTableLookupLock);
+			SCOPED_WRITE_LOCK(m_fileMappingTableLookupLock, lookupLock);
 			auto findIt = m_fileMappingTableLookup.find(fileNameKey);
 			if (findIt != m_fileMappingTableLookup.end())
 			{
 				FileMappingEntry& entry = findIt->second;
 				lookupLock.Leave();
-				ScopedWriteLock entryCs(entry.lock);
+				SCOPED_WRITE_LOCK(entry.lock, entryCs);
 				UBA_ASSERT(!entry.mapping);
 			}
 		}
@@ -671,7 +672,7 @@ namespace uba
 			UBA_ASSERT(res); (void)res;
 		}
 
-		ScopedReadLock lookupCs(dirTable.m_lookupLock);
+		SCOPED_READ_LOCK(dirTable.m_lookupLock, lookupCs);
 		auto findIt = dirTable.m_lookup.find(dirKey);
 		if (findIt == dirTable.m_lookup.end())
 			return true;
@@ -679,7 +680,7 @@ namespace uba
 		DirectoryTable::Directory& dir = findIt->second;
 		lookupCs.Leave();
 
-		ScopedWriteLock dirLock(dir.lock);
+		SCOPED_WRITE_LOCK(dir.lock, dirLock);
 
 		// To prevent race where code creating dir manage to add to lookup but then got here later than this thread.
 		while (dir.parseOffset == 0)
@@ -781,7 +782,7 @@ namespace uba
 		#endif
 
 
-		ScopedWriteLock memoryLock(dirTable.m_memoryLock);
+		SCOPED_WRITE_LOCK(dirTable.m_memoryLock, memoryLock);
 		u8* startPos = dirTable.m_memory + dirTable.m_memorySize;
 		BinaryWriter writer(startPos);
 		writer.Write7BitEncoded(written); // Storage size
@@ -808,13 +809,13 @@ namespace uba
 		if (!GetDirKey(dirKey, dirName, lastSlash, fileName))
 			return InvalidTableOffset;
 
-		ScopedReadLock lookupCs(dirTable.m_lookupLock);
+		SCOPED_READ_LOCK(dirTable.m_lookupLock, lookupCs);
 		auto res = dirTable.m_lookup.find(dirKey);
 		if (res == dirTable.m_lookup.end())
 			return 0;
 		DirectoryTable::Directory& dir = res->second;
 		lookupCs.Leave();
-		ScopedWriteLock dirLock(dir.lock);
+		SCOPED_WRITE_LOCK(dir.lock, dirLock);
 
 		while (dir.parseOffset == 0)
 		{
@@ -857,7 +858,7 @@ namespace uba
 		g_debugLogger.Info(TC("TRACKDEL    %s (Key: %s)\n"), fileName, KeyToString(fileNameKey).data);
 		#endif
 
-		ScopedWriteLock memoryLock(dirTable.m_memoryLock);
+		SCOPED_WRITE_LOCK(dirTable.m_memoryLock, memoryLock);
 		u8* startPos = dirTable.m_memory + dirTable.m_memorySize;
 		BinaryWriter writer(startPos);
 		writer.Write7BitEncoded(written); // Storage size
@@ -1134,7 +1135,7 @@ namespace uba
 		{
 			Vector<ProcessHandle> processes;
 			{
-				ScopedWriteLock lock(m_processesLock);
+				SCOPED_WRITE_LOCK(m_processesLock, lock);
 				isEmpty = m_processes.empty();
 				processes.reserve(m_processes.size());
 				for (auto& pair : m_processes)
@@ -1154,7 +1155,7 @@ namespace uba
 			#if PLATFORM_WINDOWS
 			if (m_processJobObject != NULL)
 			{
-				ScopedWriteLock lock(m_processJobObjectLock);
+				SCOPED_WRITE_LOCK(m_processJobObjectLock, lock);
 				CloseHandle(m_processJobObject);
 				m_processJobObject = NULL;
 			}
@@ -1229,13 +1230,13 @@ namespace uba
 		StringKey dirKey = CaseInsensitiveFs ? ToStringKeyLower(dirPath) : ToStringKey(dirPath);
 
 		auto& dirTable = m_directoryTable;
-		ScopedReadLock lookupLock(dirTable.m_lookupLock);
+		SCOPED_READ_LOCK(dirTable.m_lookupLock, lookupLock);
 		auto res = dirTable.m_lookup.find(dirKey);
 		if (res == dirTable.m_lookup.end())
 			return;
 		DirectoryTable::Directory& dir = res->second;
 		lookupLock.Leave();
-		ScopedWriteLock dirLock(dir.lock);
+		SCOPED_WRITE_LOCK(dir.lock, dirLock);
 
 		while (dir.parseOffset == 0)
 		{
@@ -1313,7 +1314,7 @@ namespace uba
 		if (!process.IsChild())
 			m_trace.ProcessAdded(sessionId, processId, process.GetStartInfo().description);
 
-		ScopedWriteLock lock(m_processesLock);
+		SCOPED_WRITE_LOCK(m_processesLock, lock);
 		m_processes.try_emplace(processId, ProcessHandle(&process));
 	}
 
@@ -1335,11 +1336,11 @@ namespace uba
 			Vector<ProcessLogLine> emptyLines;
 			auto& logLines = (exitCode != 0 || m_detailedTrace) ? process.m_logLines : emptyLines;
 			m_trace.ProcessExited(id, exitCode, writer.GetData(), writer.GetPosition(), logLines);
-			ScopedWriteLock lock(m_processStatsLock);
+			SCOPED_WRITE_LOCK(m_processStatsLock, lock);
 			m_processStats.Add(process.m_processStats);
 		}
 
-		ScopedWriteLock lock(m_processesLock);
+		SCOPED_WRITE_LOCK(m_processesLock, lock);
 		m_deadProcesses.emplace_back(&process); // Here to prevent Process thread call trigger a delete of Process which causes a deadlock
 		auto& stats = m_applicationStats[applicationName.data];
 		stats.count++;
@@ -1349,7 +1350,7 @@ namespace uba
 
 	void Session::FlushDeadProcesses()
 	{
-		ScopedWriteLock lock(m_processesLock);
+		SCOPED_WRITE_LOCK(m_processesLock, lock);
 		Vector<ProcessHandle> deadProcesses;
 		deadProcesses.swap(m_deadProcesses);
 		lock.Leave();
@@ -1358,23 +1359,35 @@ namespace uba
 	bool Session::GetInitResponse(InitResponse& out, const InitMessage& msg)
 	{
 		out.directoryTableHandle = m_directoryTableHandle.ToU64();
-		out.directoryTableSize = m_directoryTable.m_memoryLock.ScopedRead([&]() { return (u32)m_directoryTable.m_memorySize; });
-		out.directoryTableCount = m_directoryTable.m_lookupLock.ScopedRead([&]() { return (u32)m_directoryTable.m_lookup.size(); });
+		{
+			SCOPED_READ_LOCK(m_directoryTable.m_memoryLock, l);
+			out.directoryTableSize = (u32)m_directoryTable.m_memorySize;
+		}
+		{
+			SCOPED_READ_LOCK(m_directoryTable.m_lookupLock, l);
+			out.directoryTableCount = (u32)m_directoryTable.m_lookup.size();
+		}
 		out.mappedFileTableHandle = m_fileMappingTableHandle.ToU64();
-		out.mappedFileTableSize = m_fileMappingTableMemLock.ScopedRead([&]() { return m_fileMappingTableSize; });
-		out.mappedFileTableCount = m_fileMappingTableLookupLock.ScopedRead([&]() { return (u32)m_fileMappingTableLookup.size(); });
+		{
+			SCOPED_READ_LOCK(m_fileMappingTableMemLock, l);
+			out.mappedFileTableSize = m_fileMappingTableSize;
+		}
+		{
+			SCOPED_READ_LOCK(m_fileMappingTableLookupLock, l);
+			out.mappedFileTableCount = (u32)m_fileMappingTableLookup.size();
+		}
 		return true;
 	}
 
 	u32 Session::GetDirectoryTableSize()
 	{
-		ScopedReadLock lock(m_directoryTable.m_memoryLock);
+		SCOPED_READ_LOCK(m_directoryTable.m_memoryLock, lock);
 		return m_directoryTable.m_memorySize;
 	}
 
 	u32 Session::GetFileMappingSize()
 	{
-		ScopedReadLock lock(m_fileMappingTableMemLock);
+		SCOPED_READ_LOCK(m_fileMappingTableMemLock, lock);
 		return m_fileMappingTableSize;
 	}
 
@@ -1387,7 +1400,7 @@ namespace uba
 
 	u32 Session::GetActiveProcessCount()
 	{
-		ScopedReadLock cs(m_processesLock);
+		SCOPED_READ_LOCK(m_processesLock, cs);
 		return u32(m_processes.size());
 	}
 
@@ -1514,7 +1527,7 @@ namespace uba
 
 	void* Session::GetProcessEnvironmentVariables()
 	{
-		ScopedWriteLock lock(m_environmentVariablesLock);
+		SCOPED_WRITE_LOCK(m_environmentVariablesLock, lock);
 		if (!m_environmentVariables.empty())
 			return m_environmentVariables.data();
 
@@ -1653,7 +1666,7 @@ namespace uba
 
 		if (m_runningRemote && !fileName.StartsWith(m_tempPath.data))
 		{
-			ScopedWriteLock lock(m_outputFilesLock);
+			SCOPED_WRITE_LOCK(m_outputFilesLock, lock);
 			auto insres = m_outputFiles.try_emplace(fileName.data);
 			if (insres.second)
 			{
@@ -1670,7 +1683,7 @@ namespace uba
 			out.fileName.Append(fileName);
 		}
 
-		ScopedWriteLock lock(m_activeFilesLock);
+		SCOPED_WRITE_LOCK(m_activeFilesLock, lock);
 		u32 wantsOnCloseId = m_wantsOnCloseIdCounter++;
 		out.closeId = wantsOnCloseId;
 		auto insres = m_activeFiles.try_emplace(wantsOnCloseId);
@@ -1684,7 +1697,7 @@ namespace uba
 
 	void RemoveWrittenFile(ProcessImpl& process, const TString& name)
 	{
-		ScopedWriteLock writtenLock(process.m_writtenFilesLock);
+		SCOPED_WRITE_LOCK(process.m_writtenFilesLock, writtenLock);
 		auto& writtenFiles = process.m_writtenFiles;
 		auto findIt = writtenFiles.find(name);
 		if (findIt == writtenFiles.end())
@@ -1699,7 +1712,7 @@ namespace uba
 
 	bool Session::CloseFile(CloseFileResponse& out, const CloseFileMessage& msg)
 	{
-		ScopedWriteLock lock(m_activeFilesLock);
+		SCOPED_WRITE_LOCK(m_activeFilesLock, lock);
 		auto findIt = m_activeFiles.find(msg.closeId);
 		if (findIt == m_activeFiles.end())
 			return m_logger.Error(TC("This should not happen. Got unknown closeId %u - %s"), msg.closeId, msg.fileName.data);
@@ -1735,7 +1748,7 @@ namespace uba
 					msgName = msg.newName.data;
 			}
 
-			ScopedWriteLock writtenLock(msg.process.m_writtenFilesLock);
+			SCOPED_WRITE_LOCK(msg.process.m_writtenFilesLock, writtenLock);
 			auto insres = msg.process.m_writtenFiles.try_emplace(name);
 			WrittenFile& writtenFile = insres.first->second;
 
@@ -1828,12 +1841,12 @@ namespace uba
 	{
 		if (msg.closeId != 0)
 		{
-			ScopedWriteLock lock(m_activeFilesLock);
+			SCOPED_WRITE_LOCK(m_activeFilesLock, lock);
 			m_activeFiles.erase(msg.closeId);
 		}
 
 		{
-			ScopedWriteLock lock(m_outputFilesLock);
+			SCOPED_WRITE_LOCK(m_outputFilesLock, lock);
 			m_outputFiles.erase(msg.fileName.data);
 		}
 
@@ -1850,7 +1863,7 @@ namespace uba
 		out.fromName.Append(msg.fromName);
 		out.toName.Append(msg.toName);
 
-		ScopedWriteLock lock(m_activeFilesLock);
+		SCOPED_WRITE_LOCK(m_activeFilesLock, lock);
 		u32 closeId = m_wantsOnCloseIdCounter++;
 		if (!m_activeFiles.try_emplace(closeId, ActiveFile{ msg.toName.data, msg.toKey }).second)
 		{
@@ -1977,11 +1990,11 @@ namespace uba
 		{
 			StringBuffer<> name;
 			Storage::GetMappingString(name, file.mappingHandle, 0);
-			ScopedWriteLock lookupLock(m_fileMappingTableLookupLock);
+			SCOPED_WRITE_LOCK(m_fileMappingTableLookupLock, lookupLock);
 			auto insres = m_fileMappingTableLookup.try_emplace(file.key);
 			FileMappingEntry& entry = insres.first->second;
 			lookupLock.Leave();
-			ScopedWriteLock entryCs(entry.lock);
+			SCOPED_WRITE_LOCK(entry.lock, entryCs);
 			entry.handled = true;
 			entry.mapping = file.mappingHandle;
 			entry.mappingOffset = 0;
@@ -1989,7 +2002,7 @@ namespace uba
 			entry.isDir = false;
 			entry.success = true;
 
-			ScopedWriteLock lock(m_fileMappingTableMemLock);
+			SCOPED_WRITE_LOCK(m_fileMappingTableMemLock, lock);
 			BinaryWriter writer(m_fileMappingTableMem, m_fileMappingTableSize);
 			writer.WriteStringKey(file.key);
 			writer.WriteString(name);
diff --git a/Engine/Source/Programs/UnrealBuildAccelerator/Common/Private/UbaSessionClient.cpp b/Engine/Source/Programs/UnrealBuildAccelerator/Common/Private/UbaSessionClient.cpp
--- a/Engine/Source/Programs/UnrealBuildAccelerator/Common/Private/UbaSessionClient.cpp
+++ b/Engine/Source/Programs/UnrealBuildAccelerator/Common/Private/UbaSessionClient.cpp
@@ -114,11 +114,11 @@ namespace uba
 	bool SessionClient::GetCasKeyForFile(CasKey& out, u32 processId, const StringBufferBase& fileName, const StringKey& fileNameKey)
 	{
 		TimerScope waitTimer(Stats().waitGetFileMsg);
-		ScopedWriteLock lock(m_nameToHashLookupLock);
+		SCOPED_WRITE_LOCK(m_nameToHashLookupLock, lock);
 		auto insres = m_nameToHashLookup.try_emplace(fileNameKey);
 		HashRec& rec = insres.first->second;
 		lock.Leave();
-		ScopedWriteLock lock2(rec.lock);
+		SCOPED_WRITE_LOCK(rec.lock, lock2);
 		if (rec.key == CasKeyZero)//!rec.serverTime)
 		{
 			waitTimer.Cancel();
@@ -218,7 +218,7 @@ namespace uba
 		KeyToString keyStr(ToStringKeyLower(applicationDir));
 
 		UBA_ASSERT(application && *application);
-		ScopedWriteLock lock(m_handledApplicationEnvironmentsLock);
+		SCOPED_WRITE_LOCK(m_handledApplicationEnvironmentsLock, lock);
 		auto insres = m_handledApplicationEnvironments.insert(application);
 		if (insres.second)
 		{
@@ -323,7 +323,7 @@ namespace uba
 		StringBuffer<> lower;
 		lower.Append(applicationDir).Append(PathSeparator).Append(binaryName);
 		lower.MakeLower();
-		ScopedWriteLock lock(m_binFileLock);
+		SCOPED_WRITE_LOCK(m_binFileLock, lock);
 
 		auto insres = m_writtenBinFiles.try_emplace(lower.data, casKey);
 		if (!insres.second)
@@ -416,12 +416,12 @@ namespace uba
 				StorageStats& stats = m_storage.Stats();
 				TimerScope ts(stats.ensureCas);
 
-				ScopedWriteLock lookupLock(m_fileMappingTableLookupLock);
+				SCOPED_WRITE_LOCK(m_fileMappingTableLookupLock, lookupLock);
 				auto insres = m_fileMappingTableLookup.try_emplace(fileNameKey);
 				FileMappingEntry& entry = insres.first->second;
 				lookupLock.Leave();
 
-				ScopedWriteLock entryCs(entry.lock);
+				SCOPED_WRITE_LOCK(entry.lock, entryCs);
 				ts.Leave();
 
 				if (entry.handled)
@@ -537,12 +537,12 @@ namespace uba
 		if (msg.closeId != 0)
 		{
 			UBA_ASSERTF(false, TC("This has not been tested properly"));
-			ScopedWriteLock lock(m_activeFilesLock);
+			SCOPED_WRITE_LOCK(m_activeFilesLock, lock);
 			sendDelete = m_activeFiles.erase(msg.closeId) == 0;
 		}
 
 		{
-			ScopedWriteLock lock(m_outputFilesLock);
+			SCOPED_WRITE_LOCK(m_outputFilesLock, lock);
 			sendDelete = m_outputFiles.erase(msg.fileName.data) == 0 && sendDelete;
 		}
 
@@ -586,7 +586,7 @@ namespace uba
 
 	bool SessionClient::CopyFile(CopyFileResponse& out, const CopyFileMessage& msg)
 	{
-		ScopedWriteLock lock(m_outputFilesLock);
+		SCOPED_WRITE_LOCK(m_outputFilesLock, lock);
 		auto findIt = m_outputFiles.find(msg.fromName.data);
 		if (findIt == m_outputFiles.end())
 		{
@@ -637,7 +637,7 @@ namespace uba
 		const tchar* toName = msg.toName.data;
 
 		{
-			ScopedWriteLock lock(msg.process.m_writtenFilesLock);
+			SCOPED_WRITE_LOCK(msg.process.m_writtenFilesLock, lock);
 			auto& writtenFiles = msg.process.m_writtenFiles;
 			auto findIt = writtenFiles.find(fromName);
 			if (findIt != writtenFiles.end())
@@ -652,7 +652,7 @@ namespace uba
 
 		bool sendMove = true;
 		{
-			ScopedWriteLock lock(m_outputFilesLock);
+			SCOPED_WRITE_LOCK(m_outputFilesLock, lock);
 			auto findIt = m_outputFiles.find(fromName);
 			if (findIt != m_outputFiles.end())
 			{
@@ -682,7 +682,7 @@ namespace uba
 		const tchar* fromName = msg.fileName.data;
 
 		{
-			ScopedWriteLock lock(msg.process.m_writtenFilesLock);
+			SCOPED_WRITE_LOCK(msg.process.m_writtenFilesLock, lock);
 			auto& writtenFiles = msg.process.m_writtenFiles;
 			auto findIt = writtenFiles.find(fromName);
 			if (findIt != writtenFiles.end())
@@ -722,11 +722,11 @@ namespace uba
 		// TODO: There is a potential risk here where two different applications asks for the full name of a file
 		// and they have different bin/working dir.. and there are two versions of this file.
 
-		ScopedWriteLock lock(m_nameToNameLookupLock);
+		SCOPED_WRITE_LOCK(m_nameToNameLookupLock, lock);
 		auto insres = m_nameToNameLookup.try_emplace(msg.fileName.data);
 		NameRec& rec = insres.first->second;
 		lock.Leave();
-		ScopedWriteLock lock2(rec.lock);
+		SCOPED_WRITE_LOCK(rec.lock, lock2);
 
 		if (rec.handled)
 		{
@@ -861,7 +861,7 @@ namespace uba
 			u8* pos = dirTable.m_memory + readPos;
 			u32 toRead = u32(reader.GetLeft());
 
-			ScopedWriteLock lock(m_directoryTableLock);
+			SCOPED_WRITE_LOCK(m_directoryTableLock, lock);
 
 			if (toRead == 0)
 			{
@@ -886,7 +886,7 @@ namespace uba
 			{
 				//dirTable.ParseDirectoryTable(m_directoryTableMemPos); // This is not needed.. we never read from the directory table in the client session
 				{
-					ScopedWriteLock lock2(dirTable.m_memoryLock);
+					SCOPED_WRITE_LOCK(dirTable.m_memoryLock, lock2);
 					dirTable.m_memorySize = m_directoryTableMemPos;
 				}
 				ActiveUpdateDirectoryEntry::UpdateReadPosLess(m_firstEmptyWait, m_directoryTableMemPos);
@@ -944,14 +944,14 @@ namespace uba
 
 		u32 addCount = 0;
 		BinaryReader r(m_nameToHashTableMem.memory, readStartPos, NameToHashMemSize);
-		ScopedWriteLock lock(m_nameToHashLookupLock);
+		SCOPED_WRITE_LOCK(m_nameToHashLookupLock, lock);
 		while (r.GetPosition() < localTableSize)
 		{
 			StringKey name = r.ReadStringKey();
 			CasKey hash = r.ReadCasKey();
 
 			HashRec& rec = m_nameToHashLookup[name];
-			ScopedWriteLock lock2(rec.lock);
+			SCOPED_WRITE_LOCK(rec.lock, lock2);
 			if (serverTime < rec.serverTime)
 				continue;
 			rec.key = hash;
@@ -1150,7 +1150,11 @@ namespace uba
 		}
 
 		// Always nice to update name-to-hash table since it can reduce number of messages while building.
-		u32 hashTableMemSize = m_nameToHashMemLock.ScopedRead([this]() { return u32(m_nameToHashTableMem.writtenSize); });
+		u32 hashTableMemSize;
+		{
+			SCOPED_READ_LOCK(m_nameToHashMemLock, l);
+			hashTableMemSize = u32(m_nameToHashTableMem.writtenSize);
+		}
 		if (neededHashTableSize > hashTableMemSize)
 		{
 			reader.Reset();
@@ -1190,7 +1194,7 @@ namespace uba
 		NetworkMessage msg(m_client, ServiceId, SessionMessageType_GetNameToHashFromServer, writer);
 		writer.WriteU32(~u32(0));
 
-		ScopedWriteLock lock(m_nameToHashMemLock);
+		SCOPED_WRITE_LOCK(m_nameToHashMemLock, lock);
 		writer.WriteU32(u32(m_nameToHashTableMem.writtenSize));
 
 		if (!msg.Send(reader, Stats().getHashesMsg))
@@ -1316,7 +1320,7 @@ namespace uba
 			for (auto it=activeProcesses.begin();it!=activeProcesses.end();)
 			{
 				ProcessRec& r = *it;
-				ScopedWriteLock lock(r.lock);
+				SCOPED_WRITE_LOCK(r.lock, lock);
 				if (!r.isDone)
 				{
 					++it;
@@ -1362,7 +1366,7 @@ namespace uba
 				for (auto it = activeProcesses.rbegin(); it != activeProcesses.rend(); ++it)
 				{
 					ProcessRec& rec = *it;
-					ScopedWriteLock lock(rec.lock);
+					SCOPED_WRITE_LOCK(rec.lock, lock);
 					if (rec.isKilled || rec.isDone)
 						continue;
 					rec.handle.Cancel(true);
@@ -1385,7 +1389,7 @@ namespace uba
 			{
 				float availableWeight;
 				{
-					ScopedReadLock lock(activeWeightLock);
+					SCOPED_READ_LOCK(activeWeightLock, lock);
 					if (activeWeight >= maxWeight)
 						break;
 					availableWeight = maxWeight - activeWeight;
@@ -1468,7 +1472,7 @@ namespace uba
 					rec->weight = info.weight;
 
 					{
-						ScopedWriteLock lock(activeWeightLock);
+						SCOPED_WRITE_LOCK(activeWeightLock, lock);
 						activeWeight += rec->weight;
 					}
 
@@ -1499,12 +1503,12 @@ namespace uba
 						float weight = rec->weight;
 						auto decreaseWeight = MakeGuard([&]()
 							{
-								ScopedWriteLock weightLock(activeWeightLock);
+								SCOPED_WRITE_LOCK(activeWeightLock, weightLock);
 								activeWeight -= weight;
 								session.m_waitToSendEvent.Set();
 							});
 
-						ScopedWriteLock lock(rec->lock);
+						SCOPED_WRITE_LOCK(rec->lock, lock);
 						auto doneGuard = MakeGuard([&]() { rec->isDone = true; session.m_waitToSendEvent.Set(); });
 
 						if (rec->isKilled)
@@ -1745,11 +1749,11 @@ namespace uba
 
 	bool SessionClient::FlushWrittenFiles(ProcessImpl& process)
 	{
-		ScopedWriteLock lock(process.m_writtenFilesLock);
+		SCOPED_WRITE_LOCK(process.m_writtenFilesLock, lock);
 		if (!SendFiles(process, process.m_processStats.sendFiles))
 			return false;
 		{
-			ScopedWriteLock lock2(m_outputFilesLock);
+			SCOPED_WRITE_LOCK(m_outputFilesLock, lock2);
 			for (auto& kv : process.m_writtenFiles)
 				m_outputFiles.erase(kv.first);
 		}
diff --git a/Engine/Source/Programs/UnrealBuildAccelerator/Common/Private/UbaSessionServer.cpp b/Engine/Source/Programs/UnrealBuildAccelerator/Common/Private/UbaSessionServer.cpp
--- a/Engine/Source/Programs/UnrealBuildAccelerator/Common/Private/UbaSessionServer.cpp
+++ b/Engine/Source/Programs/UnrealBuildAccelerator/Common/Private/UbaSessionServer.cpp
@@ -58,7 +58,7 @@ namespace uba
 
 		void CallProcessExit(ProcessHandle& h)
 		{
-			ScopedWriteLock lock(m_exitedLock);
+			SCOPED_WRITE_LOCK(m_exitedLock, lock);
 			if (!startInfo.exitedFunc)
 				return;
 			auto exitedFunc = startInfo.exitedFunc;
@@ -283,7 +283,7 @@ namespace uba
 				if (CaseInsensitiveFs)
 					fileName.MakeLower();
 				StringKey fileNameKey = ToStringKey(fileName);
-				ScopedWriteLock lock(m_nameToHashLookupLock);
+				SCOPED_WRITE_LOCK(m_nameToHashLookupLock, lock);
 				CasKey& lookupCasKey = m_nameToHashLookup[fileNameKey];
 				if (lookupCasKey != casKey)
 				{
@@ -300,7 +300,7 @@ namespace uba
 		ScopedCriticalSection lock(m_remoteProcessAndSessionLock);
 		m_queuedRemoteProcesses.push_back(remoteProcess);
 
-		ScopedReadLock lock2(m_remoteProcessReturnedEventLock);
+		SCOPED_READ_LOCK(m_remoteProcessReturnedEventLock, lock2);
 		if (m_remoteProcessReturnedEvent)
 		{
 			if (!m_remoteExecutionEnabled)
@@ -335,7 +335,7 @@ namespace uba
 		FixFileName(fileName, fileName_, workingDir.data);
 		StringKey fileNameKey = ToStringKey(fileName);
 		
-		ScopedWriteLock lock(m_customCasKeysLock);
+		SCOPED_WRITE_LOCK(m_customCasKeysLock, lock);
 		auto insres = m_customCasKeys.try_emplace(fileNameKey);
 		CustomCasKey& customKey = insres.first->second;
 		customKey.casKey = CasKeyZero;
@@ -394,13 +394,13 @@ namespace uba
 
 	void SessionServer::SetRemoteProcessSlotAvailableEvent(const Function<void()>& remoteProcessSlotAvailableEvent)
 	{
-		ScopedWriteLock lock(m_remoteProcessSlotAvailableEventLock);
+		SCOPED_WRITE_LOCK(m_remoteProcessSlotAvailableEventLock, lock);
 		m_remoteProcessSlotAvailableEvent = remoteProcessSlotAvailableEvent;
 	}
 
 	void SessionServer::SetRemoteProcessReturnedEvent(const Function<void(Process&)>& remoteProcessReturnedEvent)
 	{
-		ScopedWriteLock lock(m_remoteProcessReturnedEventLock);
+		SCOPED_WRITE_LOCK(m_remoteProcessReturnedEventLock, lock);
 		m_remoteProcessReturnedEvent = remoteProcessReturnedEvent;
 	}
 
@@ -548,7 +548,7 @@ namespace uba
 
 				bool binAsVersion = clientKeys[0] != CasKeyZero;
 				{
-					ScopedWriteLock lock(m_binKeysLock);
+					SCOPED_WRITE_LOCK(m_binKeysLock, lock);
 					if (m_detoursBinaryKey == CasKeyZero || (binAsVersion && m_agentBinaryKey == CasKeyZero))
 					{
 						StringBuffer<> dir;
@@ -642,12 +642,12 @@ namespace uba
 				lookupStr.MakeLower();
 				StringKey lookupKey = ToStringKeyNoCheck(lookupStr.data, lookupStr.count);
 
-				ScopedWriteLock lock(m_applicationDataLock);
+				SCOPED_WRITE_LOCK(m_applicationDataLock, lock);
 				auto insres = m_applicationData.try_emplace(lookupKey);
 				ApplicationData& data = insres.first->second;
 				lock.Leave();
 
-				ScopedWriteLock lock2(data.lock);
+				SCOPED_WRITE_LOCK(data.lock, lock2);
 				if (!data.bytes.empty())
 				{
 					writer.WriteBytes(data.bytes.data(), data.bytes.size());
@@ -681,12 +681,12 @@ namespace uba
 				reader.ReadString(applicationName);
 				StringKey applicationKey = ToStringKeyLower(applicationName);
 
-				ScopedWriteLock lock(m_applicationDataLock);
+				SCOPED_WRITE_LOCK(m_applicationDataLock, lock);
 				auto insres = m_applicationData.try_emplace(applicationKey);
 				ApplicationData& data = insres.first->second;
 				lock.Leave();
 				
-				ScopedWriteLock lock2(data.lock);
+				SCOPED_WRITE_LOCK(data.lock, lock2);
 				if (!data.bytes.empty())
 				{
 					writer.WriteBytes(data.bytes.data(), data.bytes.size());
@@ -744,7 +744,7 @@ namespace uba
 				u64 serverTime;
 				if (m_nameToHashInitialized && casKey != CasKeyIsDirectory)
 				{
-					ScopedWriteLock lock(m_nameToHashLookupLock);
+					SCOPED_WRITE_LOCK(m_nameToHashLookupLock, lock);
 					serverTime = GetTime();
 					CasKey& lookupCasKey = m_nameToHashLookup[fileNameKey];
 					if (lookupCasKey != casKey)
@@ -800,7 +800,7 @@ namespace uba
 						success = m_storage.FakeCopy(casKey, destination.data);
 						if (!success)
 							m_logger.Error(TC("Failed to fake copy cas from %s to %s"), CasKeyString(casKey).str, destination.data);
-						ScopedWriteLock lock(m_receivedFilesLock);
+						SCOPED_WRITE_LOCK(m_receivedFilesLock, lock);
 						m_receivedFiles.try_emplace(destinationKey, casKey);
 					}
 				}
@@ -889,7 +889,7 @@ namespace uba
 			{
 				u32 requestedSize = reader.ReadU32();
 
-				ScopedReadLock lock(m_nameToHashLookupLock);
+				SCOPED_READ_LOCK(m_nameToHashLookupLock, lock);
 				if (requestedSize == ~u32(0))
 				{
 					requestedSize = u32(m_nameToHashTableMem.writtenSize);
@@ -918,7 +918,7 @@ namespace uba
 
 				float weightLeft = availableWeight;
 				u32 addCount = 0;
-				ScopedWriteLock fillLock(m_fillUpOneAtTheTimeLock); // This is a lock to group files better (all clients connect at the same time)
+				SCOPED_WRITE_LOCK(m_fillUpOneAtTheTimeLock, fillLock); // This is a lock to group files better (all clients connect at the same time);
 				while (weightLeft > 0)
 				{
 					RemoteProcess* process = DequeueProcess(sessionId, connectionInfo.GetId());
@@ -943,7 +943,11 @@ namespace uba
 				fillLock.Leave();
 
 				u32 neededDirectoryTableSize = GetDirectoryTableSize();
-				u32 neededHashTableSize = m_nameToHashLookupLock.ScopedRead([this]() { return u32(m_nameToHashTableMem.writtenSize); });
+				u32 neededHashTableSize;
+				{
+					SCOPED_READ_LOCK(m_nameToHashLookupLock, l);
+					neededHashTableSize = u32(m_nameToHashTableMem.writtenSize);
+				}
 
 				sessionsLock.Enter();
 				//if (addCount)
@@ -1113,8 +1117,12 @@ namespace uba
 			case SessionMessageType_GetTraceInformation:
 			{
 				u32 remotePos = reader.ReadU32();
-				ScopedReadLock lock(m_trace.m_memoryLock);
-				u32 localPos = m_trace.m_memoryLock.ScopedRead([this]() { return u32(m_trace.m_memoryPos); });
+				SCOPED_READ_LOCK(m_trace.m_memoryLock, lock);
+				u32 localPos;
+				{
+					SCOPED_READ_LOCK(m_trace.m_memoryLock, l);
+					localPos = u32(m_trace.m_memoryPos);
+				}
 
 				writer.WriteU32(localPos);
 				u32 toWrite = Min(localPos - remotePos, u32(writer.GetCapacityLeft()));
@@ -1144,7 +1152,7 @@ namespace uba
 			{
 				u32 processId = reader.ReadU32();
 				u32 prevExitCode = reader.ReadU32();
-				ScopedWriteLock lock(m_processesLock);
+				SCOPED_WRITE_LOCK(m_processesLock, lock);
 				auto findIt = m_processes.find(processId);
 				if (findIt == m_processes.end())
 					return m_logger.Error(TC("Failed to find process for id %u when receiving custom message"), processId);
@@ -1152,7 +1160,7 @@ namespace uba
 				lock.Leave();
 
 				auto& remoteProcess = *(RemoteProcess*)h.m_process;
-				ScopedWriteLock exitedLock(remoteProcess.m_exitedLock);
+				SCOPED_WRITE_LOCK(remoteProcess.m_exitedLock, exitedLock);
 				NextProcessInfo nextProcess;
 				bool newProcess;
 				if (!GetNextProcess(remoteProcess, newProcess, nextProcess, prevExitCode, reader))
@@ -1171,7 +1179,7 @@ namespace uba
 			case SessionMessageType_Custom:
 			{
 				u32 processId = reader.ReadU32();
-				ScopedWriteLock lock(m_processesLock);
+				SCOPED_WRITE_LOCK(m_processesLock, lock);
 				auto findIt = m_processes.find(processId);
 				if (findIt == m_processes.end())
 					return m_logger.Error(TC("Failed to find process for id %u when receiving custom message"), processId);
@@ -1179,7 +1187,7 @@ namespace uba
 				lock.Leave();
 
 				auto& remoteProcess = *(RemoteProcess*)h.m_process;
-				ScopedWriteLock exitedLock(remoteProcess.m_exitedLock);
+				SCOPED_WRITE_LOCK(remoteProcess.m_exitedLock, exitedLock);
 				CustomMessage(remoteProcess, reader, writer);
 				return true;
 			}
@@ -1209,7 +1217,7 @@ namespace uba
 
 		bool deferCreation = true;
 		{
-			ScopedWriteLock lock(m_customCasKeysLock);
+			SCOPED_WRITE_LOCK(m_customCasKeysLock, lock);
 			auto findIt = m_customCasKeys.find(fileNameKey);
 			if (findIt != m_customCasKeys.end())
 			{
@@ -1234,7 +1242,7 @@ namespace uba
 		if (m_shouldWriteToDisk)
 			return true;
 
-		ScopedReadLock lookupLock(m_fileMappingTableLookupLock);
+		SCOPED_READ_LOCK(m_fileMappingTableLookupLock, lookupLock);
 		auto findIt = m_fileMappingTableLookup.find(fileNameKey);
 		if (findIt == m_fileMappingTableLookup.end())
 			return true;
@@ -1242,7 +1250,7 @@ namespace uba
 		UBA_ASSERT(casKeyOverride == CasKeyZero);
 		FileMappingEntry& entry = findIt->second;
 		lookupLock.Leave();
-		ScopedReadLock entryCs(entry.lock);
+		SCOPED_READ_LOCK(entry.lock, entryCs);
 		return m_storage.StoreCasFile(out, fileNameKey, fileName, entry.mapping, entry.mappingOffset, entry.size, fileName, deferCreation, true);
 	}
 
@@ -1250,7 +1258,7 @@ namespace uba
 	{
 		auto& dirTable = m_directoryTable;
 
-		ScopedWriteLock lock2(session.dirTablePosLock);
+		SCOPED_WRITE_LOCK(session.dirTablePosLock, lock2);
 
 		//m_logger.Info(TC("WritePos: %llu"), session.dirTablePos);
 		writer.WriteU32(session.dirTablePos); // We can figure out on the other side if everything was written based on if the message is full or not.
@@ -1317,7 +1325,7 @@ namespace uba
 			bool allGood = false;
 			while (memAvail >= m_memRequiredToSpawn)
 			{
-				ScopedWriteLock lock(m_waitingProcessesLock);
+				SCOPED_WRITE_LOCK(m_waitingProcessesLock, lock);
 				WaitingProcess* wp = m_oldestWaitingProcess;
 				if (!wp)
 				{
@@ -1339,7 +1347,7 @@ namespace uba
 			{
 				lastMessageTime = time;
 				u32 delayCount = 0;
-				ScopedWriteLock lock(m_waitingProcessesLock);
+				SCOPED_WRITE_LOCK(m_waitingProcessesLock, lock);
 				for (auto it = m_oldestWaitingProcess; it; it = it->next)
 					++delayCount;
 				lock.Leave();
@@ -1380,7 +1388,7 @@ namespace uba
 			{
 				u64 newestTime = 0;
 				ProcessImpl* newestProcess = nullptr;
-				ScopedWriteLock lock(m_processesLock);
+				SCOPED_WRITE_LOCK(m_processesLock, lock);
 				for (auto& kv : m_processes)
 				{
 					ProcessHandle& h = kv.second;
@@ -1403,7 +1411,7 @@ namespace uba
 			}
 		}
 
-		ScopedWriteLock lock(m_waitingProcessesLock);
+		SCOPED_WRITE_LOCK(m_waitingProcessesLock, lock);
 		for (auto it = m_oldestWaitingProcess; it; it = it->next)
 			it->event.Set();
 		m_oldestWaitingProcess = nullptr;
@@ -1412,7 +1420,7 @@ namespace uba
 
 	SessionServer::RemoteProcess* SessionServer::DequeueProcess(u32 sessionId, u32 clientId)
 	{
-		ScopedReadLock lock(m_remoteProcessSlotAvailableEventLock);
+		SCOPED_READ_LOCK(m_remoteProcessSlotAvailableEventLock, lock);
 		bool hasCalledCallback = !m_remoteProcessSlotAvailableEvent;
 		u32 sessionIndex = sessionId - 1;
 
@@ -1474,7 +1482,7 @@ namespace uba
 			m_logger.Warning(TC("Cancelling remote active processes has not been tested. Notify devs"));
 
 			{
-				ScopedWriteLock lock(m_processesLock);
+				SCOPED_WRITE_LOCK(m_processesLock, lock);
 				m_processes.erase(process->m_processId);
 			}
 		}
@@ -1485,7 +1493,7 @@ namespace uba
 
 	ProcessHandle SessionServer::ProcessRemoved(u32 processId)
 	{
-		ScopedWriteLock lock(m_processesLock);
+		SCOPED_WRITE_LOCK(m_processesLock, lock);
 		auto findIt = m_processes.find(processId);
 		if (findIt == m_processes.end())
 			return {};
@@ -1507,7 +1515,7 @@ namespace uba
 		WaitingProcess wp;
 		wp.event.Create(true);
 
-		ScopedWriteLock lock(m_waitingProcessesLock);
+		SCOPED_WRITE_LOCK(m_waitingProcessesLock, lock);
 		if (m_memoryThreadEvent.IsSet(0))
 			return false;
 
@@ -1530,7 +1538,7 @@ namespace uba
 	{
 		if (!m_shouldWriteToDisk && ((msg.access & FileAccess_Write) == 0))
 		{
-			ScopedReadLock lock(m_receivedFilesLock);
+			SCOPED_READ_LOCK(m_receivedFilesLock, lock);
 			auto findIt = m_receivedFiles.find(msg.fileNameKey);
 			if (findIt != m_receivedFiles.end())
 			{
@@ -1552,7 +1560,7 @@ namespace uba
 
 	void SessionServer::FileEntryAdded(StringKey fileNameKey, u64 lastWritten, u64 size)
 	{
-		ScopedWriteLock lock(m_nameToHashLookupLock);
+		SCOPED_WRITE_LOCK(m_nameToHashLookupLock, lock);
 
 		if (!m_nameToHashInitialized)
 			return;
@@ -1641,7 +1649,7 @@ namespace uba
 		if (!m_nameToHashTableEnabled || m_nameToHashInitialized)
 			return true;
 
-		ScopedWriteLock lock(m_nameToHashLookupLock);
+		SCOPED_WRITE_LOCK(m_nameToHashLookupLock, lock);
 		m_nameToHashTableMem.Init(NameToHashMemSize);
 		m_nameToHashInitialized = true;
 		lock.Leave();
@@ -1650,7 +1658,7 @@ namespace uba
 
 		{
 			Vector<DirectoryTable::Directory*> dirs;
-			ScopedReadLock dirsLock(dirTable.m_lookupLock);
+			SCOPED_READ_LOCK(dirTable.m_lookupLock, dirsLock);
 			dirs.reserve(dirTable.m_lookup.size());
 			for (auto& kv : dirTable.m_lookup)
 				dirs.push_back(&kv.second);
@@ -1659,7 +1667,7 @@ namespace uba
 			for (auto dirPtr : dirs)
 			{
 				DirectoryTable::Directory& dir = *dirPtr;
-				ScopedReadLock dirLock(dir.lock);
+				SCOPED_READ_LOCK(dir.lock, dirLock);
 				for (auto& fileKv : dir.files)
 				{
 					StringKey fileNameKey = fileKv.first;
@@ -1676,7 +1684,7 @@ namespace uba
 				}
 			}
 		}
-		ScopedWriteLock lock2(m_nameToHashLookupLock);
+		SCOPED_WRITE_LOCK(m_nameToHashLookupLock, lock2);
 		u64 entryCount = m_nameToHashLookup.size();
 		lock2.Leave();
 
diff --git a/Engine/Source/Programs/UnrealBuildAccelerator/Common/Private/UbaStorage.cpp b/Engine/Source/Programs/UnrealBuildAccelerator/Common/Private/UbaStorage.cpp
--- a/Engine/Source/Programs/UnrealBuildAccelerator/Common/Private/UbaStorage.cpp
+++ b/Engine/Source/Programs/UnrealBuildAccelerator/Common/Private/UbaStorage.cpp
@@ -25,11 +25,15 @@ namespace uba
 
 	void StorageImpl::CasEntryAccessed(CasEntry& entry)
 	{
-		bool hasMapping = entry.lock.ScopedRead([&](){ return entry.mappingHandle.IsValid(); });
+		bool hasMapping;
+		{
+			SCOPED_READ_LOCK(entry.lock, l); // Note, this lock is taken again outside CasEntryAccessed.. so if this takes a long time it won't help to remove this lock
+			hasMapping = entry.mappingHandle.IsValid();
+		}
 		if (hasMapping)
 			return;
 
-		ScopedWriteLock lock(m_accessLock);
+		SCOPED_WRITE_LOCK(m_accessLock, lock);
 
 		CasEntry* prevAccessed = entry.prevAccessed;
 		if (prevAccessed == nullptr)
@@ -56,7 +60,7 @@ namespace uba
 
 	void StorageImpl::CasEntryWritten(CasEntry& entry, u64 size)
 	{
-		ScopedWriteLock lock(m_accessLock);
+		SCOPED_WRITE_LOCK(m_accessLock, lock);
 
 		m_casTotalBytes += size - entry.size;
 		m_casMaxBytes = Max(m_casTotalBytes, m_casMaxBytes);
@@ -132,7 +136,7 @@ namespace uba
 
 					// TODO: Should this instead set some overflow
 					/*
-					ScopedWriteLock lock(m_accessLock);
+					SCOPED_WRITE_LOCK(m_accessLock, lock);
 					m_casEvictedBytes -= rec.size;
 					--m_casEvictedCount;
 					m_casTotalBytes += rec.size;
@@ -154,7 +158,7 @@ namespace uba
 
 	void StorageImpl::CasEntryDeleted(CasEntry& entry, u64 size)
 	{
-		ScopedWriteLock lock(m_accessLock);
+		SCOPED_WRITE_LOCK(m_accessLock, lock);
 		m_casTotalBytes -= size;
 		entry.size = 0;
 		DetachEntry(entry);
@@ -517,7 +521,7 @@ namespace uba
 
 	void StorageImpl::CasEntryAccessed(const CasKey& casKey)
 	{
-		ScopedReadLock lookupLock(m_casLookupLock);
+		SCOPED_READ_LOCK(m_casLookupLock, lookupLock);
 		auto findIt = m_casLookup.find(casKey);
 		if (findIt == m_casLookup.end())
 			return;
@@ -571,7 +575,7 @@ namespace uba
 					u64 lastWritten = 0;
 					while (true)
 					{
-						ScopedWriteLock lock(rec->lock);
+						SCOPED_WRITE_LOCK(rec->lock, lock);
 						rec->written += lastWritten;
 						if (!rec->decompressedLeft)
 						{
@@ -783,13 +787,13 @@ namespace uba
 	bool StorageImpl::AddCasFile(const tchar* fileName, const CasKey& casKey, bool deferCreation)
 	{
 		UBA_ASSERTF(IsCompressed(casKey) == m_storeCompressed, TC("CasKey compress mode must match storage compress mode (%s)"), fileName);
-		ScopedWriteLock lookupLock(m_casLookupLock);
+		SCOPED_WRITE_LOCK(m_casLookupLock, lookupLock);
 		auto insres = m_casLookup.try_emplace(casKey);
 		CasEntry& casEntry = insres.first->second;
 		lookupLock.Leave();
 		CasEntryAccessed(casEntry);
 
-		ScopedWriteLock entryLock(casEntry.lock);
+		SCOPED_WRITE_LOCK(casEntry.lock, entryLock);
 
 		if (casEntry.verified)
 			if (casEntry.exists)
@@ -810,7 +814,7 @@ namespace uba
 
 		if (deferCreation)
 		{
-			ScopedWriteLock deferredLock(m_deferredCasCreationLookupLock);
+			SCOPED_WRITE_LOCK(m_deferredCasCreationLookupLock, deferredLock);
 			auto res = m_deferredCasCreationLookup.try_emplace(casKey);
 			if (res.second)
 			{
@@ -1089,7 +1093,7 @@ namespace uba
 
 		CreateDirectory(m_rootDir.data);
 
-		ScopedWriteLock loadSaveLock(m_casTableLoadSaveLock);
+		SCOPED_WRITE_LOCK(m_casTableLoadSaveLock, loadSaveLock);
 
 		UBA_ASSERT(!m_casTableLoaded);
 		m_casTableLoaded = true;
@@ -1272,7 +1276,7 @@ namespace uba
 
 	bool StorageImpl::SaveCasTable(bool deleteIsRunningfile, bool deleteDropped)
 	{
-		ScopedWriteLock loadSaveLock(m_casTableLoadSaveLock);
+		SCOPED_WRITE_LOCK(m_casTableLoadSaveLock, loadSaveLock);
 		if (!m_casTableLoaded)
 			return true;
 
@@ -1311,9 +1315,9 @@ namespace uba
 			if (!tempFile.CreateWrite(false, DefaultAttributes(), 0, m_tempPath.data))
 				return false;
 			
-			ScopedReadLock fileTableLock(m_fileTableLookupLock);
-			ScopedReadLock casLookupLock(m_casLookupLock);
-			ScopedReadLock accessLock(m_accessLock);
+			SCOPED_READ_LOCK(m_fileTableLookupLock, fileTableLock);
+			SCOPED_READ_LOCK(m_casLookupLock, casLookupLock);
+			SCOPED_READ_LOCK(m_accessLock, accessLock);
 
 			u8 buffer[1024];
 			{
@@ -1464,7 +1468,7 @@ namespace uba
 
 						auto reportError = MakeGuard([&]()
 							{
-								ScopedWriteLock l(lock);
+								SCOPED_WRITE_LOCK(lock, l);
 								++errorCount;
 								if (lastWritten > newestWrittenError)
 									newestWrittenError = lastWritten;
@@ -1589,14 +1593,14 @@ namespace uba
 	bool StorageImpl::VerifyAndGetCachedFileInfo(CachedFileInfo& out, StringKey fileNameKey, u64 verifiedLastWriteTime, u64 verifiedSize)
 	{
 		out.casKey = CasKeyZero;
-		ScopedReadLock lookupLock(m_fileTableLookupLock);
+		SCOPED_READ_LOCK(m_fileTableLookupLock, lookupLock);
 		auto findIt = m_fileTableLookup.find(fileNameKey);
 		if (findIt == m_fileTableLookup.end())
 			return false;
 		FileEntry& fileEntry = findIt->second;
 		lookupLock.Leave();
 
-		ScopedWriteLock entryLock(fileEntry.lock);
+		SCOPED_WRITE_LOCK(fileEntry.lock, entryLock);
 		fileEntry.verified = fileEntry.lastWritten == verifiedLastWriteTime && fileEntry.size == verifiedSize;
 		if (!fileEntry.verified)
 			return false;
@@ -1612,12 +1616,12 @@ namespace uba
 			forKey.MakeLower();
 		StringKey fileNameKey = ToStringKey(forKey);
 
-		ScopedWriteLock lookupLock(m_fileTableLookupLock);
+		SCOPED_WRITE_LOCK(m_fileTableLookupLock, lookupLock);
 		auto insres = m_fileTableLookup.try_emplace(fileNameKey);
 		FileEntry& fileEntry = insres.first->second;
 		lookupLock.Leave();
 
-		ScopedWriteLock entryLock(fileEntry.lock);
+		SCOPED_WRITE_LOCK(fileEntry.lock, entryLock);
 	
 		if (fileEntry.verified)
 		{
@@ -1695,7 +1699,7 @@ namespace uba
 
 	bool StorageImpl::HasCasFile(const CasKey& casKey, CasEntry** out)
 	{
-		ScopedReadLock lookupLock(m_casLookupLock);
+		SCOPED_READ_LOCK(m_casLookupLock, lookupLock);
 		auto it = m_casLookup.find(casKey);
 		if (it == m_casLookup.end())
 			return false;
@@ -1706,11 +1710,11 @@ namespace uba
 		if (out)
 			*out = &casEntry;
 
-		ScopedWriteLock entryLock(casEntry.lock);
+		SCOPED_WRITE_LOCK(casEntry.lock, entryLock);
 
 		if (casEntry.verified)
 			return casEntry.exists;
-		ScopedWriteLock deferredLock(m_deferredCasCreationLookupLock);
+		SCOPED_WRITE_LOCK(m_deferredCasCreationLookupLock, deferredLock);
 		auto findIt = m_deferredCasCreationLookup.find(casKey);
 		if (findIt == m_deferredCasCreationLookup.end())
 			return false;
@@ -1737,13 +1741,13 @@ namespace uba
 
 	bool StorageImpl::EnsureCasFile(const CasKey& casKey, const tchar* fileName)
 	{
-		ScopedWriteLock lookupLock(m_casLookupLock);
+		SCOPED_WRITE_LOCK(m_casLookupLock, lookupLock);
 		auto insres = m_casLookup.try_emplace(casKey);
 		CasEntry& casEntry = insres.first->second;
 		lookupLock.Leave();
 		CasEntryAccessed(casEntry);
 
-		ScopedWriteLock entryLock(casEntry.lock);
+		SCOPED_WRITE_LOCK(casEntry.lock, entryLock);
 
 		if (casEntry.verified)
 		{
@@ -1789,7 +1793,7 @@ namespace uba
 
 	MappedView StorageImpl::MapView(const CasKey& casKey, const tchar* hint)
 	{
-		ScopedReadLock lookupLock(m_casLookupLock);
+		SCOPED_READ_LOCK(m_casLookupLock, lookupLock);
 		auto findIt = m_casLookup.find(casKey);
 		bool foundEntry = findIt != m_casLookup.end();
 		if (!foundEntry)
@@ -1799,7 +1803,7 @@ namespace uba
 		}
 		CasEntry& casEntry = findIt->second;
 		lookupLock.Leave();
-		ScopedWriteLock entryLock(casEntry.lock);
+		SCOPED_WRITE_LOCK(casEntry.lock, entryLock);
 		//if (!casEntry.verified)
 		//{
 		//	m_logger.Error(TC("Trying to use unverified mapping of %s (%s)"), CasKeyString(casKey).str, hint);
@@ -1823,7 +1827,7 @@ namespace uba
 
 	bool StorageImpl::DropCasFile(const CasKey& casKey, bool forceDelete, const tchar* hint)
 	{
-		ScopedReadLock lookupLock(m_casLookupLock);
+		SCOPED_READ_LOCK(m_casLookupLock, lookupLock);
 		auto findIt = m_casLookup.find(casKey);
 		bool foundEntry = findIt != m_casLookup.end();
 		if (!foundEntry)
@@ -1831,7 +1835,7 @@ namespace uba
 		CasEntry& casEntry = findIt->second;
 		lookupLock.Leave();
 	
-		ScopedWriteLock entryLock(casEntry.lock);
+		SCOPED_WRITE_LOCK(casEntry.lock, entryLock);
 
 		if (forceDelete)
 		{
@@ -1894,7 +1898,7 @@ namespace uba
 		
 		u32 maxParallelCopyOrLink = m_maxParallelCopyOrLink;
 		{
-			ScopedWriteLock lock(m_activeCopyOrLinkLock);
+			SCOPED_WRITE_LOCK(m_activeCopyOrLinkLock, lock);
 			if (m_activeCopyOrLink > maxParallelCopyOrLink)
 			{
 				TimerScope ts(stats.copyOrLinkWait);
@@ -1910,7 +1914,7 @@ namespace uba
 		
 		auto activeGuard = MakeGuard([&]()
 			{
-				ScopedWriteLock lock(m_activeCopyOrLinkLock);
+				SCOPED_WRITE_LOCK(m_activeCopyOrLinkLock, lock);
 				--m_activeCopyOrLink;
 				m_activeCopyOrLinkEvent.Set();
 			});
@@ -1920,7 +1924,7 @@ namespace uba
 		if (CaseInsensitiveFs)
 			forKey.MakeLower();
 		StringKey key = ToStringKey(forKey);
-		ScopedWriteLock lock(m_fileTableLookupLock);
+		SCOPED_WRITE_LOCK(m_fileTableLookupLock, lock);
 		auto insres = m_fileTableLookup.try_emplace(key);
 		FileEntry& entry = insres.first->second;
 		lock.Leave();
@@ -1943,7 +1947,7 @@ namespace uba
 				continue;
 			}
 
-			ScopedReadLock casEntryLock(casEntry->lock);
+			SCOPED_READ_LOCK(casEntry->lock, casEntryLock);
 			UBA_ASSERT(casEntry->verified);
 			UBA_ASSERT(casEntry->exists);
 
@@ -1994,7 +1998,7 @@ namespace uba
 
 				FileAccessor destinationFile(m_logger, destination);
 
-				ScopedWriteLock entryLock(entry.lock);
+				SCOPED_WRITE_LOCK(entry.lock, entryLock);
 				entry.verified = false;
 
 				// This is to reduce number of active CreateFiles.. seems like machines don't like tons of CreateFile at the same time
@@ -2051,7 +2055,7 @@ namespace uba
 			UBA_ASSERT(false);
 			#endif
 
-			ScopedWriteLock entryLock(entry.lock);
+			SCOPED_WRITE_LOCK(entry.lock, entryLock);
 			entry.verified = false;
 
 			bool firstTry = true;
@@ -2100,7 +2104,7 @@ namespace uba
 		if (CaseInsensitiveFs)
 			forKey.MakeLower();
 		StringKey key = ToStringKey(forKey);
-		ScopedWriteLock lock(m_fileTableLookupLock);
+		SCOPED_WRITE_LOCK(m_fileTableLookupLock, lock);
 		auto insres = m_fileTableLookup.try_emplace(key);
 		FileEntry& entry = insres.first->second;
 		entry.casKey = casKey;
@@ -2113,7 +2117,7 @@ namespace uba
 	void StorageImpl::ReportFileWrite(const tchar* fileName)
 	{
 		// If a defered cas creation is queued up while the source file is about to be modified we need to flush out the cas creation before modifying the file
-		ScopedReadLock deferredLock(m_deferredCasCreationLookupLock);
+		SCOPED_READ_LOCK(m_deferredCasCreationLookupLock, deferredLock);
 		auto findIt = m_deferredCasCreationLookupByName.find(fileName);
 		if (findIt == m_deferredCasCreationLookupByName.end())
 			return;
@@ -2445,7 +2449,7 @@ namespace uba
 
 	u8* StorageImpl::PopBufferSlot()
 	{
-		ScopedWriteLock lock(m_compSlotsLock);
+		SCOPED_WRITE_LOCK(m_compSlotsLock, lock);
 		if (!m_compSlots.empty())
 		{
 			auto back = m_compSlots.back();
@@ -2459,7 +2463,7 @@ namespace uba
 	{
 		if (!slot)
 			return;
-		ScopedWriteLock lock(m_compSlotsLock);
+		SCOPED_WRITE_LOCK(m_compSlotsLock, lock);
 		m_compSlots.push_back(slot);
 	}
 
@@ -2476,14 +2480,14 @@ namespace uba
 			forKey.MakeLower();
 		StringKey fileNameKey = ToStringKey(forKey);
 
-		ScopedReadLock lookupLock(m_fileTableLookupLock);
+		SCOPED_READ_LOCK(m_fileTableLookupLock, lookupLock);
 		auto findIt = m_fileTableLookup.find(fileNameKey);
 		if (findIt == m_fileTableLookup.end())
 			return false;
 		FileEntry& fileEntry = findIt->second;
 		lookupLock.Leave();
 
-		ScopedWriteLock entryLock(fileEntry.lock);
+		SCOPED_WRITE_LOCK(fileEntry.lock, entryLock);
 		fileEntry.verified = false;
 
 		return DropCasFile(fileEntry.casKey, true, file);
diff --git a/Engine/Source/Programs/UnrealBuildAccelerator/Common/Private/UbaStorageClient.cpp b/Engine/Source/Programs/UnrealBuildAccelerator/Common/Private/UbaStorageClient.cpp
--- a/Engine/Source/Programs/UnrealBuildAccelerator/Common/Private/UbaStorageClient.cpp
+++ b/Engine/Source/Programs/UnrealBuildAccelerator/Common/Private/UbaStorageClient.cpp
@@ -60,13 +60,13 @@ namespace uba
 
 	bool StorageClient::IsUsingProxy()
 	{
-		ScopedReadLock proxyLock(m_proxyClientLock);
+		SCOPED_READ_LOCK(m_proxyClientLock, proxyLock);
 		return m_proxyClient != nullptr;
 	}
 
 	void StorageClient::StopProxy()
 	{
-		ScopedWriteLock proxyLock(m_proxyClientLock);
+		SCOPED_WRITE_LOCK(m_proxyClientLock, proxyLock);
 		if (m_proxyClient)
 			m_proxyClient->client.Disconnect();
 	}
@@ -97,7 +97,7 @@ namespace uba
 #if !UBA_USE_SPARSEFILE
 	bool StorageClient::GetCasFileName(StringBufferBase& out, const CasKey& casKey)
 	{
-		ScopedReadLock tempLock(m_localStorageFilesLock);
+		SCOPED_READ_LOCK(m_localStorageFilesLock, tempLock);
 		auto findIt = m_localStorageFiles.find(AsCompressed(casKey, false));
 		if (findIt != m_localStorageFiles.end())
 		{
@@ -115,7 +115,7 @@ namespace uba
 
 	MappedView StorageClient::MapView(const CasKey& casKey, const tchar* hint)
 	{
-		ScopedReadLock tempLock(m_localStorageFilesLock);
+		SCOPED_READ_LOCK(m_localStorageFilesLock, tempLock);
 		auto findIt = m_localStorageFiles.find(AsCompressed(casKey, false));
 		bool isValid = findIt != m_localStorageFiles.end();
 		tempLock.Leave();
@@ -158,7 +158,7 @@ namespace uba
 
 		// Cas file might have been created by this client and in that case we can just reuse the file we just wrote. No need to fetch from server
 		// This needs to be first so it doesnt end up in the cas table even though it is not in there. (otherwise it might be garbage collected)
-		ScopedReadLock tempLock(m_localStorageFilesLock);
+		SCOPED_READ_LOCK(m_localStorageFilesLock, tempLock);
 		auto findIt = m_localStorageFiles.find(AsCompressed(casKey, false));
 		if (findIt != m_localStorageFiles.end())
 		{
@@ -182,7 +182,7 @@ namespace uba
 			if (EnsureCasFile(casKey, nullptr))
 				return true;
 
-			ScopedReadLock lock(m_casLookupLock);
+			SCOPED_READ_LOCK(m_casLookupLock, lock);
 			casEntry = &m_casLookup.find(casKey)->second;
 			lock.Leave();
 
@@ -233,7 +233,7 @@ namespace uba
 			bool wantsProxy = false;
 			if (allowProxy)
 			{
-				ScopedWriteLock proxyLock(m_proxyClientLock);
+				SCOPED_WRITE_LOCK(m_proxyClientLock, proxyLock);
 				if (m_proxyClient)
 				{
 					if (m_proxyClient->client.IsConnected())
@@ -257,7 +257,7 @@ namespace uba
 				{
 					if (proxy)
 					{
-						ScopedWriteLock proxyLock(m_proxyClientLock);
+						SCOPED_WRITE_LOCK(m_proxyClientLock, proxyLock);
 						if (!--proxy->refCount)
 						{
 							if (!proxy->client.IsConnected())
@@ -317,7 +317,7 @@ namespace uba
 						// TODO: Check if self, then use loopback.. this could be received before local proxy is started due to network ordering
 					}
 
-					ScopedWriteLock proxyLock2(m_proxyClientLock);
+					SCOPED_WRITE_LOCK(m_proxyClientLock, proxyLock2);
 					if (m_proxyClient)
 						continue;
 
@@ -658,7 +658,7 @@ namespace uba
 	bool StorageClient::HasCasFile(const CasKey& casKey, CasEntry** out)
 	{
 		CasKey localKey = AsCompressed(casKey, false);
-		ScopedReadLock lock(m_localStorageFilesLock);
+		SCOPED_READ_LOCK(m_localStorageFilesLock, lock);
 		auto findIt = m_localStorageFiles.find(localKey);
 		if (findIt != m_localStorageFiles.end())
 		{
@@ -710,7 +710,7 @@ namespace uba
 
 			if (keepMappingInMemory)
 			{
-				ScopedWriteLock lock(m_localStorageFilesLock);
+				SCOPED_WRITE_LOCK(m_localStorageFilesLock, lock);
 				auto insres = m_localStorageFiles.try_emplace(AsCompressed(casKey, false));
 				LocalFile& localFile = insres.first->second;
 				if (insres.second || !localFile.casEntry.mappingHandle.IsValid())
@@ -766,7 +766,7 @@ namespace uba
 
 	void StorageClient::Ping()
 	{
-		ScopedReadLock lock(m_proxyClientLock);
+		SCOPED_READ_LOCK(m_proxyClientLock, lock);
 		if (!m_proxyClient)
 			return;
 		u64 now = GetTime();
@@ -1075,7 +1075,7 @@ namespace uba
 				fullPath.Resize(dirLen).Append(e.name);
 				if (IsDirectory(e.attributes))
 				{
-					ScopedWriteLock lock(seenIdsLock);
+					SCOPED_WRITE_LOCK(seenIdsLock, lock);
 					if (!seenIds.insert(e.id).second)
 						return;
 					lock.Leave();
@@ -1099,7 +1099,7 @@ namespace uba
 							return;
 						}
 
-						ScopedWriteLock lookupLock(m_localStorageFilesLock);
+						SCOPED_WRITE_LOCK(m_localStorageFilesLock, lookupLock);
 						auto insres = m_localStorageFiles.try_emplace(AsCompressed(casKey, false));
 						if (!insres.second)
 							return;
diff --git a/Engine/Source/Programs/UnrealBuildAccelerator/Common/Private/UbaStorageServer.cpp b/Engine/Source/Programs/UnrealBuildAccelerator/Common/Private/UbaStorageServer.cpp
--- a/Engine/Source/Programs/UnrealBuildAccelerator/Common/Private/UbaStorageServer.cpp
+++ b/Engine/Source/Programs/UnrealBuildAccelerator/Common/Private/UbaStorageServer.cpp
@@ -28,7 +28,7 @@ namespace uba
 
 		m_server.RegisterOnClientConnected(ServiceId, [this](const Guid& clientUid, u32 clientId)
 			{
-				ScopedWriteLock lock(m_loadCasTableLock);
+				SCOPED_WRITE_LOCK(m_loadCasTableLock, lock);
 				if (!m_casTableLoaded)
 					LoadCasTable(true);
 			});
@@ -72,14 +72,14 @@ namespace uba
 		if (m_storeCompressed)
 			actualKey = AsCompressed(casKey, true);
 
-		ScopedWriteLock waitLock(m_waitEntriesLock);
+		SCOPED_WRITE_LOCK(m_waitEntriesLock, waitLock);
 		WaitEntry& waitEntry = m_waitEntries[actualKey];
 		++waitEntry.refCount;
 		waitLock.Leave();
 
 		auto g = MakeGuard([&]()
 			{
-				ScopedWriteLock waitLock2(m_waitEntriesLock);
+				SCOPED_WRITE_LOCK(m_waitEntriesLock, waitLock2);
 				if (!--waitEntry.refCount)
 					m_waitEntries.erase(actualKey);
 			});
@@ -120,7 +120,7 @@ namespace uba
 	void StorageServer::OnDisconnected(u32 clientId)
 	{
 		{
-			ScopedWriteLock lock(m_proxiesLock);
+			SCOPED_WRITE_LOCK(m_proxiesLock, lock);
 			for (auto it=m_proxies.begin(); it!=m_proxies.end(); ++it)
 			{
 				ProxyEntry& e = it->second;
@@ -132,7 +132,7 @@ namespace uba
 			}
 		}
 		{
-			ScopedWriteLock lock(m_activeStoresLock);
+			SCOPED_WRITE_LOCK(m_activeStoresLock, lock);
 			for (auto it=m_activeStores.begin(); it!=m_activeStores.end();)
 			{
 				ActiveStore& store = it->second;
@@ -143,7 +143,7 @@ namespace uba
 				}
 
 				{
-					ScopedWriteLock entryLock(store.casEntry->lock);
+					SCOPED_WRITE_LOCK(store.casEntry->lock, entryLock);
 					store.casEntry->verified = false;
 					store.casEntry->beingWritten = false;
 					if (m_traceStore)
@@ -155,7 +155,7 @@ namespace uba
 			}
 		}
 		{
-			ScopedWriteLock lock(m_activeFetchesLock);
+			SCOPED_WRITE_LOCK(m_activeFetchesLock, lock);
 			for (auto it=m_activeFetches.begin(); it!=m_activeFetches.end();)
 			{
 				ActiveFetch& fetch = it->second;
@@ -187,16 +187,16 @@ namespace uba
 		if (out == CasKeyZero)
 			return false;
 
-		ScopedWriteLock lookupLock(m_fileTableLookupLock);
+		SCOPED_WRITE_LOCK(m_fileTableLookupLock, lookupLock);
 		auto insres = m_fileTableLookup.try_emplace(fileNameKey);
 		FileEntry& fileEntry = insres.first->second;
 		lookupLock.Leave();
-		ScopedWriteLock entryLock(fileEntry.lock);
+		SCOPED_WRITE_LOCK(fileEntry.lock, entryLock);
 		fileEntry.verified = true;
 		fileEntry.casKey = out;
 		fileEntry.size = fileSize;
 
-		ScopedWriteLock externalFileLock(m_externalFileMappingsLock);
+		SCOPED_WRITE_LOCK(m_externalFileMappingsLock, externalFileLock);
 		m_externalFileMappings.try_emplace(fileNameKey, ExternalFileMapping{mappingHandle, mappingOffset, fileSize});
 		externalFileLock.Leave();
 
@@ -214,7 +214,7 @@ namespace uba
 			fromForKey.MakeLower();
 		StringKey fileNameKey = ToStringKey(fromForKey);
 
-		ScopedWriteLock lock(m_externalFileMappingsLock);
+		SCOPED_WRITE_LOCK(m_externalFileMappingsLock, lock);
 		auto findIt = m_externalFileMappings.find(fileNameKey);
 		if (findIt == m_externalFileMappings.end())
 		{
@@ -282,7 +282,7 @@ namespace uba
 				if (reader.ReadBool()) // is proxy
 					return m_logger.Error(TC("Proxy is sending connect message. This path is not implemented"));
 				u16 proxyPort = reader.ReadU16();
-				ScopedWriteLock lock(m_connectionInfoLock);
+				SCOPED_WRITE_LOCK(m_connectionInfoLock, lock);
 				Info& info = m_connectionInfo[connectionInfo.GetId()];
 				info.zone = reader.ReadString();
 				info.storageSize = reader.ReadU64();
@@ -297,7 +297,7 @@ namespace uba
 			{
 				if (reader.ReadBool()) // Wants proxy
 				{
-					ScopedReadLock lock(m_connectionInfoLock);
+					SCOPED_READ_LOCK(m_connectionInfoLock, lock);
 					auto findIt = m_connectionInfo.find(connectionInfo.GetId());
 					UBA_ASSERT(findIt != m_connectionInfo.end());
 					Info& info = findIt->second;
@@ -330,7 +330,7 @@ namespace uba
 							writer.WriteByte(1 << 2);
 
 							auto proxyKey = ToStringKeyNoCheck(proxyName.data, proxyName.count);
-							ScopedWriteLock proxiesLock(m_proxiesLock);
+							SCOPED_WRITE_LOCK(m_proxiesLock, proxiesLock);
 							ProxyEntry& proxy = m_proxies[proxyKey];
 							if (proxy.clientId == ~0u)
 							{
@@ -376,13 +376,13 @@ namespace uba
 					CasKey checkedCasKey;
 					{
 						StringKey fileNameKey = CaseInsensitiveFs ? ToStringKeyLower(hint) : ToStringKey(hint);
-						ScopedReadLock lookupLock(m_fileTableLookupLock);
+						SCOPED_READ_LOCK(m_fileTableLookupLock, lookupLock);
 						auto findIt = m_fileTableLookup.find(fileNameKey);
 						if (findIt != m_fileTableLookup.end())
 						{
 							FileEntry& fileEntry = findIt->second;
 							lookupLock.Leave();
-							ScopedReadLock entryLock(fileEntry.lock);
+							SCOPED_READ_LOCK(fileEntry.lock, entryLock);
 							if (fileEntry.verified)
 								checkedCasKey = fileEntry.casKey;
 						}
@@ -411,7 +411,7 @@ namespace uba
 						writer.WriteU16(0);
 						return true;
 					}
-					ScopedWriteLock lookupLock(m_casLookupLock);
+					SCOPED_WRITE_LOCK(m_casLookupLock, lookupLock);
 					auto findIt = m_casLookup.find(casKey);
 					UBA_ASSERT(findIt != m_casLookup.end());
 					casEntry = &findIt->second;
@@ -548,7 +548,7 @@ namespace uba
 
 				*fetchId = PopId();
 
-				ScopedWriteLock lock(m_activeFetchesLock);
+				SCOPED_WRITE_LOCK(m_activeFetchesLock, lock);
 				auto insres = m_activeFetches.try_emplace(*fetchId);
 				UBA_ASSERT(insres.second);
 				lock.Leave();
@@ -574,7 +574,7 @@ namespace uba
 				u16 fetchId = reader.ReadU16();
 				u32 fetchIndex = reader.ReadU32();
 
-				ScopedReadLock lock(m_activeFetchesLock);
+				SCOPED_READ_LOCK(m_activeFetchesLock, lock);
 				auto findIt = m_activeFetches.find(fetchId);
 				if (findIt == m_activeFetches.end())
 					return m_logger.Error(TC("Can't find active fetch %u, disconnected client? (index %u)"), fetchId, fetchIndex);
@@ -599,7 +599,7 @@ namespace uba
 				fetch.Release(*this, TC("FetchDone"));
 
 				u64 sendCasTime = fetch.sendCasTime;
-				ScopedWriteLock activeLock(m_activeFetchesLock);
+				SCOPED_WRITE_LOCK(m_activeFetchesLock, activeLock);
 				m_activeFetches.erase(fetchId);
 				activeLock.Leave();
 				PushId(fetchId);
@@ -621,12 +621,12 @@ namespace uba
 			{
 				CasKey casKey = reader.ReadCasKey();
 				UBA_ASSERT(IsCompressed(casKey));
-				ScopedWriteLock lookupLock(m_casLookupLock);
+				SCOPED_WRITE_LOCK(m_casLookupLock, lookupLock);
 				auto casInsres = m_casLookup.try_emplace(casKey);
 				CasEntry& casEntry = casInsres.first->second;
 				lookupLock.Leave();
 
-				ScopedWriteLock entryLock(casEntry.lock);
+				SCOPED_WRITE_LOCK(casEntry.lock, entryLock);
 				
 				if (!WaitForWritten(casEntry, entryLock, TC("UNKNOWN")))
 					return false;
@@ -678,12 +678,12 @@ namespace uba
 				StringBuffer<> hint;
 				reader.ReadString(hint);
 
-				ScopedWriteLock lookupLock(m_casLookupLock);
+				SCOPED_WRITE_LOCK(m_casLookupLock, lookupLock);
 				auto casInsres = m_casLookup.try_emplace(casKey);
 				CasEntry& casEntry = casInsres.first->second;
 				lookupLock.Leave();
 
-				ScopedWriteLock entryLock(casEntry.lock);
+				SCOPED_WRITE_LOCK(casEntry.lock, entryLock);
 				if (!casEntry.verified)
 				{
 					casEntry.key = casKey;
@@ -745,7 +745,7 @@ namespace uba
 				{
 					storeId = reader.ReadU16();
 					memOffset = reader.ReadU64();
-					ScopedReadLock activeLock(m_activeStoresLock);
+					SCOPED_READ_LOCK(m_activeStoresLock, activeLock);
 					auto storeIt = m_activeStores.find(storeId);
 					if (storeIt == m_activeStores.end())
 						return m_logger.Error(TC("Can't find active store %u, disconnected client?"), storeId);
@@ -767,7 +767,7 @@ namespace uba
 
 					CasEntry& casEntry = *activeStore.casEntry;
 					{
-						ScopedWriteLock entryLock(casEntry.lock);
+						SCOPED_WRITE_LOCK(casEntry.lock, entryLock);
 						casEntry.mappingHandle = activeStore.mappedView.handle;
 						casEntry.mappingOffset = activeStore.mappedView.offset;
 						casEntry.mappingSize = totalWritten;
@@ -786,7 +786,7 @@ namespace uba
 					stats.recvCasBytesComp += activeStore.fileSize;
 					stats.recvCasBytesRaw += activeStore.actualSize;
 
-					ScopedWriteLock waitLock(m_waitEntriesLock);
+					SCOPED_WRITE_LOCK(m_waitEntriesLock, waitLock);
 					auto waitFindIt = m_waitEntries.find(casEntry.key);
 					if (waitFindIt != m_waitEntries.end())
 					{
@@ -798,7 +798,7 @@ namespace uba
 
 					if (!firstStore)
 					{
-						ScopedWriteLock activeLock(m_activeStoresLock);
+						SCOPED_WRITE_LOCK(m_activeStoresLock, activeLock);
 						m_activeStores.erase(storeId);
 						activeLock.Leave();
 						PushId(storeId);
@@ -819,7 +819,7 @@ namespace uba
 					writer.WriteU16(storeId);
 					writer.WriteBool(m_traceStore);
 
-					ScopedWriteLock activeLock(m_activeStoresLock);
+					SCOPED_WRITE_LOCK(m_activeStoresLock, activeLock);
 					auto insres = m_activeStores.try_emplace(storeId);
 					UBA_ASSERT(insres.second);
 					activeLock.Leave();
@@ -850,7 +850,7 @@ namespace uba
 
 	u16 StorageServer::PopId()
 	{
-		ScopedWriteLock lock(m_availableIdsLock);
+		SCOPED_WRITE_LOCK(m_availableIdsLock, lock);
 		if (m_availableIds.empty())
 		{
 			UBA_ASSERT(m_availableIdsHigh < u16(~0) - 1);
@@ -863,7 +863,7 @@ namespace uba
 
 	void StorageServer::PushId(u16 id)
 	{
-		ScopedWriteLock lock(m_availableIdsLock);
+		SCOPED_WRITE_LOCK(m_availableIdsLock, lock);
 		m_availableIds.push_back(id);
 	}
 
diff --git a/Engine/Source/Programs/UnrealBuildAccelerator/Common/Private/UbaTrace.cpp b/Engine/Source/Programs/UnrealBuildAccelerator/Common/Private/UbaTrace.cpp
--- a/Engine/Source/Programs/UnrealBuildAccelerator/Common/Private/UbaTrace.cpp
+++ b/Engine/Source/Programs/UnrealBuildAccelerator/Common/Private/UbaTrace.cpp
@@ -130,7 +130,7 @@ namespace uba
 
 		u64 stringLen = TStrlen(string);
 
-		ScopedWriteLock lock(m_stringsLock);
+		SCOPED_WRITE_LOCK(m_stringsLock, lock);
 		auto insres = m_strings.try_emplace(ToStringKeyNoCheck(string, stringLen));
 		if (insres.second)
 		{
diff --git a/Engine/Source/Programs/UnrealBuildAccelerator/Common/Private/UbaWorkManager.cpp b/Engine/Source/Programs/UnrealBuildAccelerator/Common/Private/UbaWorkManager.cpp
--- a/Engine/Source/Programs/UnrealBuildAccelerator/Common/Private/UbaWorkManager.cpp
+++ b/Engine/Source/Programs/UnrealBuildAccelerator/Common/Private/UbaWorkManager.cpp
@@ -39,7 +39,7 @@ namespace uba
 				{
 					while (true)
 					{
-						ScopedWriteLock lock(manager.m_workLock);
+						SCOPED_WRITE_LOCK(manager.m_workLock, lock);
 						if (manager.m_work.empty())
 							break;
 						Work work = manager.m_work.front();
@@ -49,8 +49,8 @@ namespace uba
 						work.func();
 					}
 
-					ScopedWriteLock lock1(manager.m_availableWorkersLock);
-					ScopedReadLock lock2(manager.m_workLock);
+					SCOPED_WRITE_LOCK(manager.m_availableWorkersLock, lock1);
+					SCOPED_READ_LOCK(manager.m_workLock, lock2);
 					if (!manager.m_work.empty())
 						continue;
 
@@ -86,12 +86,12 @@ namespace uba
 
 	void WorkManagerImpl::AddWork(const Function<void()>& work, u32 count, const tchar* desc)
 	{
-		ScopedWriteLock lock(m_workLock);
+		SCOPED_WRITE_LOCK(m_workLock, lock);
 		for (u32 i = 0; i != count; ++i)
 			m_work.push_back({ work });
 		lock.Leave();
 
-		ScopedWriteLock lock2(m_availableWorkersLock);
+		SCOPED_WRITE_LOCK(m_availableWorkersLock, lock2);
 		while (count--)
 		{
 			Worker* worker = PopWorkerNoLock();
@@ -108,7 +108,7 @@ namespace uba
 
 	void WorkManagerImpl::PushWorker(Worker* worker)
 	{
-		ScopedWriteLock lock(m_availableWorkersLock);
+		SCOPED_WRITE_LOCK(m_availableWorkersLock, lock);
 		PushWorkerNoLock(worker);
 	}
 
@@ -126,7 +126,7 @@ namespace uba
 	{
 		while (count--)
 		{
-			ScopedWriteLock lock(m_workLock);
+			SCOPED_WRITE_LOCK(m_workLock, lock);
 			if (m_work.empty())
 				break;
 			Work work = m_work.front();
@@ -141,7 +141,7 @@ namespace uba
 	{
 		while (true)
 		{
-			ScopedReadLock lock(m_workLock);
+			SCOPED_READ_LOCK(m_workLock, lock);
 			bool workEmpty = m_work.empty();
 			lock.Leave();
 			if (workEmpty)
diff --git a/Engine/Source/Programs/UnrealBuildAccelerator/Common/Public/UbaBottleneck.h b/Engine/Source/Programs/UnrealBuildAccelerator/Common/Public/UbaBottleneck.h
--- a/Engine/Source/Programs/UnrealBuildAccelerator/Common/Public/UbaBottleneck.h
+++ b/Engine/Source/Programs/UnrealBuildAccelerator/Common/Public/UbaBottleneck.h
@@ -21,7 +21,7 @@ namespace uba
 	{
 		BottleneckScope(Bottleneck& b) : bottleneck(b)
 		{
-			ScopedWriteLock lock(bottleneck.lock);
+			SCOPED_WRITE_LOCK(bottleneck.lock, lock);
 			while (true)
 			{
 				if (bottleneck.activeCount < bottleneck.maxCount)
@@ -40,7 +40,7 @@ namespace uba
 
 		~BottleneckScope()
 		{
-			ScopedWriteLock lock(bottleneck.lock);
+			SCOPED_WRITE_LOCK(bottleneck.lock, lock);
 			if (bottleneck.activeCount == bottleneck.maxCount)
 				bottleneck.underMax.Set();
 			--bottleneck.activeCount;
diff --git a/Engine/Source/Programs/UnrealBuildAccelerator/Common/Public/UbaLogger.h b/Engine/Source/Programs/UnrealBuildAccelerator/Common/Public/UbaLogger.h
--- a/Engine/Source/Programs/UnrealBuildAccelerator/Common/Public/UbaLogger.h
+++ b/Engine/Source/Programs/UnrealBuildAccelerator/Common/Public/UbaLogger.h
@@ -94,4 +94,6 @@ namespace uba
 	void StopDebugLogger();
 	extern LoggerWithWriter g_debugLogger;
 	#endif
+
+	void PrintContentionSummary(class Logger& logger);
 }
diff --git a/Engine/Source/Programs/UnrealBuildAccelerator/Core/Private/UbaMemory.cpp b/Engine/Source/Programs/UnrealBuildAccelerator/Core/Private/UbaMemory.cpp
--- a/Engine/Source/Programs/UnrealBuildAccelerator/Core/Private/UbaMemory.cpp
+++ b/Engine/Source/Programs/UnrealBuildAccelerator/Core/Private/UbaMemory.cpp
@@ -60,7 +60,7 @@ namespace uba
 		if (!memory)
 			return aligned_alloc(alignment, bytes);
 
-		ScopedWriteLock l(lock);
+		SCOPED_WRITE_LOCK(lock, l);
 		return AllocateNoLock(bytes, alignment, hint);
 	}
 
diff --git a/Engine/Source/Programs/UnrealBuildAccelerator/Core/Private/UbaSynchronization.cpp b/Engine/Source/Programs/UnrealBuildAccelerator/Core/Private/UbaSynchronization.cpp
--- a/Engine/Source/Programs/UnrealBuildAccelerator/Core/Private/UbaSynchronization.cpp
+++ b/Engine/Source/Programs/UnrealBuildAccelerator/Core/Private/UbaSynchronization.cpp
@@ -122,4 +122,24 @@ namespace uba
 		pthread_rwlock_unlock((pthread_rwlock_t*)data);
 		#endif
 	}
+
+	#if UBA_TRACK_CONTENTION
+
+	List<ContentionTracker>& GetContentionTrackerList()
+	{
+		static List<ContentionTracker> trackers;
+		return trackers;
+	}
+
+
+	ContentionTracker& GetContentionTracker(const char* file, u64 line)
+	{
+		static ReaderWriterLock rwl;
+		ScopedWriteLock l(rwl);
+		ContentionTracker& ct = GetContentionTrackerList().emplace_back();
+		ct.file = file;
+		ct.line = line;
+		return ct;
+	}
+	#endif
 }
diff --git a/Engine/Source/Programs/UnrealBuildAccelerator/Core/Public/UbaDirectoryTable.h b/Engine/Source/Programs/UnrealBuildAccelerator/Core/Public/UbaDirectoryTable.h
--- a/Engine/Source/Programs/UnrealBuildAccelerator/Core/Public/UbaDirectoryTable.h
+++ b/Engine/Source/Programs/UnrealBuildAccelerator/Core/Public/UbaDirectoryTable.h
@@ -24,7 +24,7 @@ namespace uba
 
 		void ParseDirectoryTable(u32 size)
 		{
-			ScopedWriteLock lock(m_lookupLock);
+			SCOPED_WRITE_LOCK(m_lookupLock, lock);
 			ParseDirectoryTableNoLock(size);
 		}
 
@@ -52,7 +52,7 @@ namespace uba
 		{
 			if (dir.parseOffset == dir.tableOffset)
 				return;
-			ScopedWriteLock lock(dir.lock);
+			SCOPED_WRITE_LOCK(dir.lock, lock);
 			PopulateDirectoryRecursive(hasher, dir.tableOffset, dir.parseOffset, dir.files);
 			dir.parseOffset = dir.tableOffset;
 		}
@@ -127,7 +127,7 @@ namespace uba
 
 		Exists EntryExists(StringKey entryKey, const tchar* entryName, u64 entryNameLen, bool checkIfDir = false, u32* tableOffset = nullptr)
 		{
-			ScopedReadLock lock(m_lookupLock);
+			SCOPED_READ_LOCK(m_lookupLock, lock);
 			return EntryExistsNoLock(entryKey, entryName, entryNameLen, checkIfDir, tableOffset);
 		}
 
@@ -167,12 +167,12 @@ namespace uba
 						return Exists_No;
 					if (parentDir.parseOffset != parentDir.tableOffset)
 					{
-						ScopedWriteLock lock(parentDir.lock);
+						SCOPED_WRITE_LOCK(parentDir.lock, lock);
 						PopulateDirectoryRecursive(ancestorHasher, parentDir.tableOffset, parentDir.parseOffset, parentDir.files);
 						parentDir.parseOffset = parentDir.tableOffset;
 					}
 
-					ScopedReadLock lock(parentDir.lock);
+					SCOPED_READ_LOCK(parentDir.lock, lock);
 					auto entryIt = parentDir.files.find(entryKey);
 					if (entryIt == parentDir.files.end())
 						return Exists_No;
@@ -258,7 +258,7 @@ namespace uba
 			StringKeyHasher hasher;
 			hasher.Update(forHash.data, forHash.count);
 
-			ScopedReadLock lock(m_lookupLock);
+			SCOPED_READ_LOCK(m_lookupLock, lock);
 			while (true)
 			{
 				const tchar* slash = TStrchr(prevSlash + 1, PathSeparator);
@@ -273,7 +273,7 @@ namespace uba
 
 				if (directory)
 				{
-					ScopedReadLock lock2(directory->lock);
+					SCOPED_READ_LOCK(directory->lock, lock2);
 					auto fileIt = directory->files.find(fileNameKey);
 					if (fileIt != directory->files.end())
 					{
diff --git a/Engine/Source/Programs/UnrealBuildAccelerator/Core/Public/UbaMemory.h b/Engine/Source/Programs/UnrealBuildAccelerator/Core/Public/UbaMemory.h
--- a/Engine/Source/Programs/UnrealBuildAccelerator/Core/Public/UbaMemory.h
+++ b/Engine/Source/Programs/UnrealBuildAccelerator/Core/Public/UbaMemory.h
@@ -159,7 +159,7 @@ namespace uba
 
 		void* Allocate()
 		{
-			ScopedWriteLock lock(m_lock);
+			SCOPED_WRITE_LOCK(m_lock, lock);
 			if (m_nextFree)
 			{
 				void* ptr = (void*)m_nextFree;
@@ -175,7 +175,7 @@ namespace uba
 			#if UBA_DEBUG
 			memset(mem, 0xFE, sizeof(Type));
 			#endif
-			ScopedWriteLock lock(m_lock);
+			SCOPED_WRITE_LOCK(m_lock, lock);
 			*(u64*)mem = m_nextFree;
 			m_nextFree = u64(mem);
 		}
diff --git a/Engine/Source/Programs/UnrealBuildAccelerator/Core/Public/UbaSynchronization.h b/Engine/Source/Programs/UnrealBuildAccelerator/Core/Public/UbaSynchronization.h
--- a/Engine/Source/Programs/UnrealBuildAccelerator/Core/Public/UbaSynchronization.h
+++ b/Engine/Source/Programs/UnrealBuildAccelerator/Core/Public/UbaSynchronization.h
@@ -6,6 +6,11 @@
 #include <atomic>
 #include <utility>
 
+#define UBA_TRACK_CONTENTION 0
+
+#define STRING_JOIN(arg1, arg2) STRING_JOIN_INNER(arg1, arg2)
+#define STRING_JOIN_INNER(arg1, arg2) arg1 ## arg2
+
 namespace uba
 {
 	template<typename Type>
@@ -67,9 +72,6 @@ namespace uba
 		void EnterWrite();
 		void LeaveWrite();
 
-		template<class Functor> auto ScopedRead(const Functor& f);
-		template<class Functor> auto ScopedWrite(const Functor& f);
-
 	private:
 
 		#if PLATFORM_WINDOWS
@@ -84,6 +86,25 @@ namespace uba
 		ReaderWriterLock& operator=(const ReaderWriterLock&) = delete;
 	};
 
+	#if !UBA_TRACK_CONTENTION
+	#define SCOPED_READ_LOCK(readerWriterLock, name) ScopedReadLock name(readerWriterLock);
+	#define SCOPED_WRITE_LOCK(readerWriterLock, name) ScopedWriteLock name(readerWriterLock);
+	#else
+	u64 GetTime();
+	struct ContentionTracker { void Add(u64 t) { time += t; ++count; }; Atomic<u64> time; Atomic<u64> count; const char* file; u64 line; };
+	ContentionTracker& GetContentionTracker(const char* file, u64 line);
+
+	#define SCOPED_LOCK(readerWriterLock, lockType, name) \
+		u64 STRING_JOIN(contentionStart, __LINE__) = GetTime(); \
+		lockType name(readerWriterLock); \
+		static ContentionTracker& STRING_JOIN(tracker, __LINE__) = GetContentionTracker(__FILE__, __LINE__); \
+		STRING_JOIN(tracker, __LINE__).Add(GetTime() - STRING_JOIN(contentionStart, __LINE__));
+
+	#define SCOPED_READ_LOCK(readerWriterLock, name) SCOPED_LOCK(readerWriterLock, ScopedReadLock, name)
+	#define SCOPED_WRITE_LOCK(readerWriterLock, name) SCOPED_LOCK(readerWriterLock, ScopedWriteLock, name)
+	#endif
+
+
 	class ScopedReadLock
 	{
 	public:
@@ -107,10 +128,6 @@ namespace uba
 		bool m_active = true;
 	};
 
-	template<class Functor> auto ReaderWriterLock::ScopedRead(const Functor& f) { ScopedReadLock l(*this); return f(); }
-	template<class Functor> auto ReaderWriterLock::ScopedWrite(const Functor& f) { ScopedWriteLock l(*this); return f(); }
-
-
 	template<typename Lambda>
 	struct ScopeGuard
 	{
diff --git a/Engine/Source/Programs/UnrealBuildAccelerator/Detours/Private/UbaDetoursFileMappingTable.cpp b/Engine/Source/Programs/UnrealBuildAccelerator/Detours/Private/UbaDetoursFileMappingTable.cpp
--- a/Engine/Source/Programs/UnrealBuildAccelerator/Detours/Private/UbaDetoursFileMappingTable.cpp
+++ b/Engine/Source/Programs/UnrealBuildAccelerator/Detours/Private/UbaDetoursFileMappingTable.cpp
@@ -46,13 +46,13 @@ namespace uba
 
 	void MappedFileTable::Parse(u32 tableSize)
 	{
-		ScopedWriteLock lock(m_lookupLock);
+		SCOPED_WRITE_LOCK(m_lookupLock, lock);
 		ParseNoLock(tableSize);
 	}
 
 	void MappedFileTable::SetDeleted(const StringKey& key, const tchar* name, bool deleted)
 	{
-		ScopedWriteLock lock(m_lookupLock);
+		SCOPED_WRITE_LOCK(m_lookupLock, lock);
 		auto it = m_lookup.find(key);
 		if (it == m_lookup.end())
 			return;
@@ -64,7 +64,7 @@ namespace uba
 	void Rpc_CreateFileW(const tchar* fileName, const StringKey& fileNameKey, u8 access, tchar* outNewName, u64 newNameCapacity, u64& outSize, u32& outCloseId, bool lock)
 	{
 		TimerScope ts(g_stats.createFile);
-		ScopedWriteLock pcs(g_communicationLock);
+		SCOPED_WRITE_LOCK(g_communicationLock, pcs);
 		BinaryWriter writer;
 		writer.WriteByte(MessageType_CreateFileW);
 		writer.WriteString(fileName);
@@ -93,7 +93,7 @@ namespace uba
 		u32 tableOffset;
 		{
 			TimerScope ts(g_stats.listDirectory);
-			ScopedWriteLock pcs(g_communicationLock);
+			SCOPED_WRITE_LOCK(g_communicationLock, pcs);
 			BinaryWriter writer;
 			writer.WriteByte(MessageType_ListDirectory);
 			writer.WriteString(dirName, dirNameLen);
@@ -118,7 +118,7 @@ namespace uba
 		u32 directoryTableSize;
 		{
 			TimerScope ts(g_stats.closeFile);
-			ScopedWriteLock pcs(g_communicationLock);
+			SCOPED_WRITE_LOCK(g_communicationLock, pcs);
 			BinaryWriter writer;
 			writer.WriteByte(MessageType_CloseFile);
 			writer.WriteString(handleName);
@@ -153,7 +153,7 @@ namespace uba
 	void Rpc_UpdateTables()
 	{
 		TimerScope ts(g_stats.updateTables);
-		ScopedWriteLock pcs(g_communicationLock);
+		SCOPED_WRITE_LOCK(g_communicationLock, pcs);
 		BinaryWriter writer;
 		writer.WriteByte(MessageType_UpdateTables);
 		writer.Flush();
@@ -192,7 +192,7 @@ namespace uba
 		if (Rpc_UpdateDirectory(hash.key, entryName, dirNameLen) == ~u32(0))
 			return ~u32(0);
 
-		ScopedWriteLock lookLock(g_directoryTable.m_lookupLock);
+		SCOPED_WRITE_LOCK(g_directoryTable.m_lookupLock, lookLock);
 		auto dirFindIt = g_directoryTable.m_lookup.find(hash.key);
 		if (dirFindIt == g_directoryTable.m_lookup.end())
 			return ~u32(0);
@@ -203,7 +203,7 @@ namespace uba
 
 		g_directoryTable.PopulateDirectory(hash.open, dir);
 
-		ScopedReadLock lock(dir.lock);
+		SCOPED_READ_LOCK(dir.lock, lock);
 		auto findIt = dir.files.find(entryNameKey);
 		if (findIt == dir.files.end())
 			return ~u32(0);
@@ -227,7 +227,7 @@ namespace uba
 
 		{
 			TimerScope ts(g_stats.getFullFileName);
-			ScopedWriteLock pcs(g_communicationLock);
+			SCOPED_WRITE_LOCK(g_communicationLock, pcs);
 			BinaryWriter writer;
 			writer.WriteByte(MessageType_GetFullFileName);
 			writer.WriteString(path);
diff --git a/Engine/Source/Programs/UnrealBuildAccelerator/Detours/Private/UbaDetoursShared.cpp b/Engine/Source/Programs/UnrealBuildAccelerator/Detours/Private/UbaDetoursShared.cpp
--- a/Engine/Source/Programs/UnrealBuildAccelerator/Detours/Private/UbaDetoursShared.cpp
+++ b/Engine/Source/Programs/UnrealBuildAccelerator/Detours/Private/UbaDetoursShared.cpp
@@ -155,7 +155,7 @@ namespace uba
 	{
 		DEBUG_LOG(TC("LOG  %.*s"), u32(textCharLength), text); // TODO: Investigate, deadlocks on non-windows
 		TimerScope ts(g_stats.log);
-		ScopedWriteLock pcs(g_communicationLock);
+		SCOPED_WRITE_LOCK(g_communicationLock, pcs);
 		BinaryWriter writer;
 		writer.WriteByte(MessageType_Log);
 		writer.WriteBool(printInSession);
@@ -190,7 +190,7 @@ namespace uba
 		if (!g_echoOn)
 			return;
 
-		ScopedWriteLock lock(g_consoleStringCs);
+		SCOPED_WRITE_LOCK(g_consoleStringCs, lock);
 		const CharType* read = chars;
 		tchar* write = g_consoleString + g_consoleStringIndex;
 		int left = sizeof_array(g_consoleString) - g_consoleStringIndex - 1;
@@ -241,7 +241,7 @@ namespace uba
 		bool keepInMemory = KeepInMemory(fileName, fileNameForKey.count);
 		if (keepInMemory)
 		{
-			ScopedReadLock lock(g_mappedFileTable.m_lookupLock);
+			SCOPED_READ_LOCK(g_mappedFileTable.m_lookupLock, lock);
 			auto it = g_mappedFileTable.m_lookup.find(fileNameKey);
 			if (it == g_mappedFileTable.m_lookup.end() || it->second.deleted)
 			{
@@ -303,7 +303,7 @@ namespace uba
 				{
 					if (g_runningRemote) // This could be a written file not reported to server yet
 					{
-						ScopedReadLock lock(g_mappedFileTable.m_lookupLock);
+						SCOPED_READ_LOCK(g_mappedFileTable.m_lookupLock, lock);
 						auto findIt = g_mappedFileTable.m_lookup.find(fileNameKey);
 						if (findIt != g_mappedFileTable.m_lookup.end() && !findIt->second.deleted)
 						{
diff --git a/Engine/Source/Programs/UnrealBuildAccelerator/Detours/Private/UbaDetoursShared.h b/Engine/Source/Programs/UnrealBuildAccelerator/Detours/Private/UbaDetoursShared.h
--- a/Engine/Source/Programs/UnrealBuildAccelerator/Detours/Private/UbaDetoursShared.h
+++ b/Engine/Source/Programs/UnrealBuildAccelerator/Detours/Private/UbaDetoursShared.h
@@ -42,8 +42,6 @@ namespace uba
 		#define DEBUG_LOG_PIPE(Command, ...) ts.Leave();
 		//#define DEBUG_LOG(...)
 		#define DEBUG_LOG(...) { if (isLogging()) { GetLogTlsBuffer().Clear().Appendf(__VA_ARGS__).Append(TC("\n")); WriteDebugLog(); }}
-		#define STRING_JOIN(arg1, arg2) STRING_JOIN_INNER(arg1, arg2)
-		#define STRING_JOIN_INNER(arg1, arg2) arg1 ## arg2
 	#else
 		#define DEBUG_LOG(...)
 		#define DEBUG_LOG_DETOURED(Command, ...)
diff --git a/Engine/Source/Programs/UnrealBuildAccelerator/Detours/Private/Windows/UbaDetoursFunctionsKernelBase.inl b/Engine/Source/Programs/UnrealBuildAccelerator/Detours/Private/Windows/UbaDetoursFunctionsKernelBase.inl
--- a/Engine/Source/Programs/UnrealBuildAccelerator/Detours/Private/Windows/UbaDetoursFunctionsKernelBase.inl
+++ b/Engine/Source/Programs/UnrealBuildAccelerator/Detours/Private/Windows/UbaDetoursFunctionsKernelBase.inl
@@ -4,7 +4,7 @@ const wchar_t* ToString(BOOL b) { return b ? L"Success" : L"Error"; }
 
 DWORD Local_GetLongPathNameW(LPCWSTR lpszShortPath, LPWSTR lpszLongPath, DWORD cchBuffer)
 {
-	ScopedWriteLock lock(g_longPathNameCacheLock);
+	SCOPED_WRITE_LOCK(g_longPathNameCacheLock, lock);
 	auto findIt = g_longPathNameCache.find(lpszShortPath);
 	if (findIt != g_longPathNameCache.end())
 	{
@@ -176,7 +176,7 @@ BOOL Detoured_CreateDirectoryW(LPCWSTR lpPathName, LPSECURITY_ATTRIBUTES lpSecur
 		FixPath(pathName, lpPathName);
 		StringKey pathNameKey = ToStringKeyLower(pathName);
 		TimerScope ts(g_stats.createFile);
-		ScopedWriteLock pcs(g_communicationLock);
+		SCOPED_WRITE_LOCK(g_communicationLock, pcs);
 		BinaryWriter writer;
 		writer.WriteByte(MessageType_CreateDirectory);
 		writer.WriteStringKey(pathNameKey);
@@ -503,7 +503,7 @@ void WriteStdFile(LPCVOID lpBuffer, DWORD nNumberOfBytesToWrite, bool isError)
 	if (!g_echoOn)
 		return;
 
-	ScopedWriteLock lock(g_stdFileLock);
+	SCOPED_WRITE_LOCK(g_stdFileLock, lock);
 	u32 start = 0;
 	u32 i = 0;
 	auto bufferStr = (const char*)lpBuffer;
@@ -1083,7 +1083,7 @@ DWORD Detoured_GetModuleFileNameW(HMODULE hModule, LPWSTR lpFilename, DWORD nSiz
 
 	{
 		// Check if there are any stored paths from dynamically loaded dlls
-		ScopedReadLock lock(g_loadedModulesLock);
+		SCOPED_READ_LOCK(g_loadedModulesLock, lock);
 		auto findIt = g_loadedModules.find(hModule);
 		if (findIt != g_loadedModules.end())
 			return Shared_GetModuleFileNameW(hModule, findIt->second.c_str(), u32(findIt->second.size()), lpFilename, nSize);
@@ -1138,7 +1138,7 @@ BOOL Detoured_CopyFileExW(LPCWSTR lpExistingFileName, LPCWSTR lpNewFileName, LPP
 	u32 directoryTableSize;
 	{
 		TimerScope ts(g_stats.copyFile);
-		ScopedWriteLock pcs(g_communicationLock);
+		SCOPED_WRITE_LOCK(g_communicationLock, pcs);
 		BinaryWriter writer;
 		writer.WriteByte(MessageType_CopyFile);
 		writer.WriteStringKey(fromKey);
@@ -1208,7 +1208,7 @@ BOOL Detoured_CreateHardLinkW(LPCWSTR lpFileName, LPCWSTR lpExistingFileName, LP
 	u32 closeId;
 	{
 		TimerScope ts(g_stats.copyFile);
-		ScopedWriteLock pcs(g_communicationLock);
+		SCOPED_WRITE_LOCK(g_communicationLock, pcs);
 		BinaryWriter writer;
 		writer.WriteByte(MessageType_CopyFile);
 		writer.WriteStringKey(fromKey);
@@ -1268,7 +1268,7 @@ BOOL Detoured_DeleteFileW(LPCWSTR lpFileName)
 	{
 		u32 closeId = 0;
 		TimerScope ts(g_stats.deleteFile);
-		ScopedWriteLock pcs(g_communicationLock);
+		SCOPED_WRITE_LOCK(g_communicationLock, pcs);
 		BinaryWriter writer;
 		writer.WriteByte(MessageType_DeleteFileW);
 		writer.WriteString(fixedName);
@@ -1300,7 +1300,7 @@ bool Shared_MoveFile(LPCWSTR lpExistingFileName, LPCWSTR lpNewFileName, DWORD dw
 
 	if (KeepInMemory(source.data, source.count))
 	{
-		ScopedWriteLock lock(g_mappedFileTable.m_lookupLock);
+		SCOPED_WRITE_LOCK(g_mappedFileTable.m_lookupLock, lock);
 		auto it = g_mappedFileTable.m_lookup.find(sourceKey);
 		UBA_ASSERTF(it != g_mappedFileTable.m_lookup.end(), L"Can't find %ls", source.data);
 		FileInfo& sourceInfo = it->second;
@@ -1315,7 +1315,7 @@ bool Shared_MoveFile(LPCWSTR lpExistingFileName, LPCWSTR lpNewFileName, DWORD dw
 			UBA_ASSERT(!sourceInfo.memoryFile->isLocalOnly);
 			dest.MakeLower();
 			StringKey destKey = ToStringKey(dest);
-			ScopedWriteLock lock2(g_mappedFileTable.m_lookupLock);
+			SCOPED_WRITE_LOCK(g_mappedFileTable.m_lookupLock, lock2);
 			auto insres = g_mappedFileTable.m_lookup.try_emplace(destKey);
 			lock2.Leave();
 			FileInfo& destInfo = insres.first->second;
@@ -1355,7 +1355,7 @@ bool Shared_MoveFile(LPCWSTR lpExistingFileName, LPCWSTR lpNewFileName, DWORD dw
 	bool result;
 	{
 		TimerScope ts(g_stats.moveFile);
-		ScopedWriteLock pcs(g_communicationLock);
+		SCOPED_WRITE_LOCK(g_communicationLock, pcs);
 		BinaryWriter writer;
 		writer.WriteByte(MessageType_MoveFileW);
 		writer.WriteStringKey(sourceKey);
@@ -1501,7 +1501,7 @@ __forceinline HANDLE Shared_FindFirstFileExW(LPCWSTR lpFileName, FINDEX_INFO_LEV
 
 	DirHash hash(buf, bufChars);
 
-	ScopedWriteLock lookLock(g_directoryTable.m_lookupLock);
+	SCOPED_WRITE_LOCK(g_directoryTable.m_lookupLock, lookLock);
 	auto insres = g_directoryTable.m_lookup.try_emplace(hash.key, &g_memoryBlock);
 	DirectoryTable::Directory& dir = insres.first->second;
 	if (insres.second)
@@ -1556,7 +1556,7 @@ __forceinline HANDLE Shared_FindFirstFileExW(LPCWSTR lpFileName, FINDEX_INFO_LEV
 	else
 		listHandle->it = 0;
 
-	ScopedReadLock lock(dir.lock);
+	SCOPED_READ_LOCK(dir.lock, lock);
 	listHandle->fileTableOffsets.resize(dir.files.size());
 	u32 it = 0;
 	for (auto& pair : dir.files)
@@ -2106,7 +2106,7 @@ LPVOID Detoured_MapViewOfFileEx(HANDLE hFileMappingObject, DWORD dwDesiredAccess
 				//if (res = (u8*)True_MapViewOfFileEx(trueMappingObject, dwDesiredAccess, ToHigh(fi.trueFileMapOffset), ToLow(fi.trueFileMapOffset), dwNumberOfBytesToMap + offset, (u8*)lpBaseAddress - offset))
 				//{
 				//	// If we don't have any real MapViewOfFile on existing fileMapMem we can just drop the old one and use the new one.
-				//	ScopedWriteLock lock(g_mappedFileTable.m_memLookupLock);
+				//	SCOPED_WRITE_LOCK(g_mappedFileTable.m_memLookupLock, lock);
 				//	auto it = g_mappedFileTable.m_memLookup.find(fi.fileMapMem);
 				//	if (it == g_mappedFileTable.m_memLookup.end())
 				//	{
@@ -2157,7 +2157,7 @@ LPVOID Detoured_MapViewOfFileEx(HANDLE hFileMappingObject, DWORD dwDesiredAccess
 			if (fi.memoryFile && (dwDesiredAccess & FILE_MAP_WRITE)) // We assume changes will happen
 				fi.memoryFile->isReported = false;
 
-			ScopedWriteLock lock(g_mappedFileTable.m_memLookupLock);
+			SCOPED_WRITE_LOCK(g_mappedFileTable.m_memLookupLock, lock);
 			++g_mappedFileTable.m_memLookup[mem];
 			return mem;
 		}
@@ -2181,7 +2181,7 @@ BOOL Detoured_UnmapViewOfFileEx(PVOID lpBaseAddress, ULONG UnmapFlags)
 	DETOURED_CALL(UnmapViewOfFileEx);
 
 	{
-		ScopedWriteLock lock(g_mappedFileTable.m_memLookupLock);
+		SCOPED_WRITE_LOCK(g_mappedFileTable.m_memLookupLock, lock);
 		auto it = g_mappedFileTable.m_memLookup.find(lpBaseAddress);
 		if (it != g_mappedFileTable.m_memLookup.end())
 		{
@@ -2346,7 +2346,7 @@ HMODULE Recursive_LoadLibraryExW(LPCWSTR lpLibFileName, LPCWSTR originalName, DW
 	{
 		if (originalName[1] == ':')
 		{
-			ScopedWriteLock lock(g_loadedModulesLock);
+			SCOPED_WRITE_LOCK(g_loadedModulesLock, lock);
 			g_loadedModules[res] = originalName;
 		}
 		if (g_isRunningWine)
@@ -2492,7 +2492,7 @@ BOOL Detoured_SetConsoleMode(HANDLE hConsoleHandle, DWORD mode)
 
 	g_echoOn = (mode & ~503) != 0; // TODO: This might be wrong. Trying to figure out how echo off in batch files work in terms of win32 calls
 	{
-		ScopedWriteLock pcs(g_communicationLock);
+		SCOPED_WRITE_LOCK(g_communicationLock, pcs);
 		BinaryWriter writer;
 		writer.WriteByte(MessageType_EchoOn);
 		writer.WriteBool(g_echoOn);
@@ -2543,7 +2543,7 @@ BOOL Detoured_CreateProcessW(LPCWSTR lpApplicationName, LPWSTR lpCommandLine, LP
 	char dll[1024];
 	{
 		TimerScope ts(g_stats.createProcess);
-		ScopedWriteLock pcs(g_communicationLock);
+		SCOPED_WRITE_LOCK(g_communicationLock, pcs);
 		BinaryWriter writer;
 		writer.WriteByte(MessageType_CreateProcess);
 		writer.WriteString(lpApplicationName ? lpApplicationName : L"");
@@ -2610,7 +2610,7 @@ BOOL Detoured_CreateProcessW(LPCWSTR lpApplicationName, LPWSTR lpCommandLine, LP
 
 	{
 		TimerScope ts(g_stats.createProcess);
-		ScopedWriteLock pcs(g_communicationLock);
+		SCOPED_WRITE_LOCK(g_communicationLock, pcs);
 		BinaryWriter writer;
 		writer.WriteByte(MessageType_StartProcess);
 		writer.WriteU32(processId);
@@ -2936,7 +2936,7 @@ DWORD Detoured_GetModuleFileNameExA(HANDLE hProcess, HMODULE hModule, LPSTR lpFi
 
 	{
 		// Check if there are any stored paths from dynamically loaded dlls
-		ScopedReadLock lock(g_loadedModulesLock);
+		SCOPED_READ_LOCK(g_loadedModulesLock, lock);
 		auto findIt = g_loadedModules.find(hModule);
 		if (findIt != g_loadedModules.end())
 			return Shared_GetModuleFileNameA(hModule, findIt->second.c_str(), u32(findIt->second.size()), lpFilename, nSize);
diff --git a/Engine/Source/Programs/UnrealBuildAccelerator/Detours/Private/Windows/UbaDetoursFunctionsNtDll.inl b/Engine/Source/Programs/UnrealBuildAccelerator/Detours/Private/Windows/UbaDetoursFunctionsNtDll.inl
--- a/Engine/Source/Programs/UnrealBuildAccelerator/Detours/Private/Windows/UbaDetoursFunctionsNtDll.inl
+++ b/Engine/Source/Programs/UnrealBuildAccelerator/Detours/Private/Windows/UbaDetoursFunctionsNtDll.inl
@@ -329,7 +329,7 @@ NTSTATUS Detoured_NtSetInformationFile(HANDLE FileHandle, PIO_STATUS_BLOCK IoSta
 #if 0
 			FileInfo* newFileInfo = nullptr;
 			{
-				ScopedReadLock _(g_mappedFileTable.m_lookupLock);
+				SCOPED_READ_LOCK(g_mappedFileTable.m_lookupLock, _);
 				auto findIt = g_mappedFileTable.m_lookup.find(newFileNameKey);
 				if (findIt != g_mappedFileTable.m_lookup.end())
 					newFileInfo = &findIt->second;
@@ -362,7 +362,7 @@ NTSTATUS Detoured_NtSetInformationFile(HANDLE FileHandle, PIO_STATUS_BLOCK IoSta
 			// It then renames the old file to that file..
 			// ..and finally opens it again with attributes only and deleteonclose flag.. and closes it.
 			// TODO: This is wrong.. we just delete the file.. but what we should do is copy it over in memory to the memory file before setting it to delete
-			ScopedReadLock _(g_mappedFileTable.m_lookupLock);
+			SCOPED_READ_LOCK(g_mappedFileTable.m_lookupLock, _);
 			auto findIt = g_mappedFileTable.m_lookup.find(newFileNameKey);
 			if (findIt != g_mappedFileTable.m_lookup.end())
 			{
@@ -651,7 +651,7 @@ NTSTATUS NTAPI Shared_NtCreateFile(bool IsCreateFunc, PHANDLE hFileHandle, ACCES
 				{
 					// This could be a written file not reported to server yet
 					{
-						ScopedReadLock lock(g_mappedFileTable.m_lookupLock);
+						SCOPED_READ_LOCK(g_mappedFileTable.m_lookupLock, lock);
 						auto findIt = g_mappedFileTable.m_lookup.find(fileNameKey);
 						if (findIt != g_mappedFileTable.m_lookup.end())
 							allowEarlyOut = findIt->second.deleted;
@@ -724,7 +724,7 @@ NTSTATUS NTAPI Shared_NtCreateFile(bool IsCreateFunc, PHANDLE hFileHandle, ACCES
 		UBA_ASSERT(fileNameLower.data[fileNameLower.count - 1] != '\\');
 		DirHash hash(fileNameLower.data, pathLen);
 
-		ScopedWriteLock lookupLock(g_directoryTable.m_lookupLock);
+		SCOPED_WRITE_LOCK(g_directoryTable.m_lookupLock, lookupLock);
 		auto insres = g_directoryTable.m_lookup.try_emplace(hash.key, &g_memoryBlock);
 		DirectoryTable::Directory& dir = insres.first->second;
 		if (insres.second)
@@ -766,7 +766,7 @@ NTSTATUS NTAPI Shared_NtCreateFile(bool IsCreateFunc, PHANDLE hFileHandle, ACCES
 		auto listHandle = new ListDirectoryHandle{ hash.key, insres.first->second };
 		listHandle->it = 0;
 
-		ScopedReadLock lock(dir.lock);
+		SCOPED_READ_LOCK(dir.lock, lock);
 		listHandle->fileTableOffsets.resize(dir.files.size());
 		u32 it = 0;
 		for (auto& pair : dir.files)
@@ -799,7 +799,7 @@ NTSTATUS NTAPI Shared_NtCreateFile(bool IsCreateFunc, PHANDLE hFileHandle, ACCES
 	const wchar_t* lpFileName = fileName.data;
 	u32 closeId = 0;
 
-	ScopedWriteLock _(g_mappedFileTable.m_lookupLock);
+	SCOPED_WRITE_LOCK(g_mappedFileTable.m_lookupLock, _);
 	auto insres = g_mappedFileTable.m_lookup.try_emplace(fileNameKey);
 	FileInfo& info = insres.first->second;
 	u32 lastDesiredAccess = info.lastDesiredAccess;
@@ -951,7 +951,7 @@ NTSTATUS NTAPI Shared_NtCreateFile(bool IsCreateFunc, PHANDLE hFileHandle, ACCES
 				else
 				{
 					TimerScope ts(g_stats.openTempFile);
-					ScopedWriteLock pcs(g_communicationLock);
+					SCOPED_WRITE_LOCK(g_communicationLock, pcs);
 					BinaryWriter writer;
 					writer.WriteByte(MessageType_OpenTempFile);
 					writer.WriteStringKey(fileNameKey);
@@ -1247,7 +1247,7 @@ NTSTATUS NTAPI Detoured_NtClose(HANDLE handle)
 				FixPath(fixedNewName, fo->newName.c_str());
 				fixedNewName.MakeLower();
 				StringKey fileNameKey = ToStringKey(fixedNewName);
-				ScopedWriteLock _(g_mappedFileTable.m_lookupLock);
+				SCOPED_WRITE_LOCK(g_mappedFileTable.m_lookupLock, _);
 				auto insres = g_mappedFileTable.m_lookup.try_emplace(fileNameKey);
 				FileInfo& newInfo = insres.first->second;
 				newInfo = fi;
@@ -1268,7 +1268,7 @@ NTSTATUS NTAPI Detoured_NtClose(HANDLE handle)
 			FixPath(fixedName, path);
 
 			TimerScope ts(g_stats.createTempFile);
-			ScopedWriteLock pcs(g_communicationLock);
+			SCOPED_WRITE_LOCK(g_communicationLock, pcs);
 			BinaryWriter writer;
 			writer.WriteByte(MessageType_CreateTempFile);
 			writer.WriteStringKey(ToStringKeyLower(fixedName));
diff --git a/Engine/Source/Programs/UnrealBuildAccelerator/Detours/Private/Windows/UbaDetoursFunctionsWin.cpp b/Engine/Source/Programs/UnrealBuildAccelerator/Detours/Private/Windows/UbaDetoursFunctionsWin.cpp
--- a/Engine/Source/Programs/UnrealBuildAccelerator/Detours/Private/Windows/UbaDetoursFunctionsWin.cpp
+++ b/Engine/Source/Programs/UnrealBuildAccelerator/Detours/Private/Windows/UbaDetoursFunctionsWin.cpp
@@ -494,7 +494,7 @@ GrowingUnorderedMap<const wchar_t*, const wchar_t*, HashString, EqualString> g_l
 void Rpc_AllocFailed(const wchar_t* allocType, u32 error)
 {
 	TimerScope ts(g_stats.virtualAllocFailed);
-	ScopedWriteLock pcs(g_communicationLock);
+	SCOPED_WRITE_LOCK(g_communicationLock, pcs);
 	BinaryWriter writer;
 	writer.WriteByte(MessageType_VirtualAllocFailed);
 	writer.WriteString(allocType);
@@ -548,7 +548,7 @@ void SendExitMessage(DWORD exitCode, u64 startTime)
 		while (left)
 		{
 			u32 toWrite = Min(left, u32(30 * 1024));
-			ScopedWriteLock pcs(g_communicationLock);
+			SCOPED_WRITE_LOCK(g_communicationLock, pcs);
 			BinaryWriter writer;
 			writer.WriteByte(MessageType_InputDependencies);
 			if (pos == 0)
@@ -563,7 +563,7 @@ void SendExitMessage(DWORD exitCode, u64 startTime)
 
 	g_stats.usedMemory = u32(g_memoryBlock.writtenSize);
 
-	ScopedWriteLock pcs(g_communicationLock);
+	SCOPED_WRITE_LOCK(g_communicationLock, pcs);
 	BinaryWriter writer;
 	writer.WriteByte(MessageType_Exit);
 	writer.WriteU32(exitCode);
@@ -880,7 +880,7 @@ void Init(const DetoursPayload& payload, u64 startTime)
 
 	{
 		TimerScope ts(g_stats.init);
-		ScopedWriteLock pcs(g_communicationLock);
+		SCOPED_WRITE_LOCK(g_communicationLock, pcs);
 		BinaryWriter writer;
 		writer.WriteByte(MessageType_Init);
 		writer.Flush();
@@ -1011,7 +1011,7 @@ extern "C"
 	UBA_DETOURED_API u32 UbaSendCustomMessage(const void* send, u32 sendSize, void* recv, u32 recvCapacity)
 	{
 		//TimerScope ts(g_stats.init);
-		ScopedWriteLock pcs(g_communicationLock);
+		SCOPED_WRITE_LOCK(g_communicationLock, pcs);
 		BinaryWriter writer;
 		writer.WriteByte(MessageType_Custom);
 		writer.WriteU32(sendSize);
@@ -1026,7 +1026,7 @@ extern "C"
 
 	UBA_DETOURED_API bool UbaFlushWrittenFiles()
 	{
-		ScopedWriteLock pcs(g_communicationLock);
+		SCOPED_WRITE_LOCK(g_communicationLock, pcs);
 		BinaryWriter writer;
 		writer.WriteByte(MessageType_FlushWrittenFiles);
 		writer.Flush();
@@ -1037,7 +1037,7 @@ extern "C"
 	UBA_DETOURED_API bool UbaUpdateEnvironment(const wchar_t* reason, bool resetStats)
 	{
 		{
-			ScopedWriteLock pcs(g_communicationLock);
+			SCOPED_WRITE_LOCK(g_communicationLock, pcs);
 			BinaryWriter writer;
 			writer.WriteByte(MessageType_UpdateEnvironment);
 			writer.WriteString(reason ? reason : L"");
@@ -1065,7 +1065,7 @@ extern "C"
 		*outArguments = 0;
 		bool newProcess;
 		{
-			ScopedWriteLock pcs(g_communicationLock);
+			SCOPED_WRITE_LOCK(g_communicationLock, pcs);
 			BinaryWriter writer;
 			writer.WriteByte(MessageType_GetNextProcess);
 			writer.WriteU32(prevExitCode);
diff --git a/Engine/Source/Programs/UnrealBuildAccelerator/Host/Private/UbaExports.cpp b/Engine/Source/Programs/UnrealBuildAccelerator/Host/Private/UbaExports.cpp
--- a/Engine/Source/Programs/UnrealBuildAccelerator/Host/Private/UbaExports.cpp
+++ b/Engine/Source/Programs/UnrealBuildAccelerator/Host/Private/UbaExports.cpp
@@ -394,6 +394,7 @@ uba::StorageClient* CreateStorageClient(uba::NetworkClient& client, const uba::t
 		server->GetStorage().PrintSummary(logger);
 		server->GetServer().PrintSummary(logger);
 		uba::SystemStats::GetGlobal().Print(logger, true);
+		uba::PrintContentionSummary(logger);
 	}
 	void SessionServer_CancelAll(uba::SessionServer* server)
 	{
