[UBA]
* Added Uba NetworkClient::StopAll
* Fixed shutdown race condition in NetworkClient
* Changed client work count from 16 to no of logical cores

diff --git a/Engine/Source/Programs/UnrealBuildAccelerator/Common/Private/UbaNetworkClient.cpp b/Engine/Source/Programs/UnrealBuildAccelerator/Common/Private/UbaNetworkClient.cpp
--- a/Engine/Source/Programs/UnrealBuildAccelerator/Common/Private/UbaNetworkClient.cpp
+++ b/Engine/Source/Programs/UnrealBuildAccelerator/Common/Private/UbaNetworkClient.cpp
@@ -11,7 +11,7 @@
 namespace uba
 {
 	NetworkClient::NetworkClient(bool& outCtorSuccess, const NetworkClientCreateInfo& info, const tchar* name)
-	: WorkManagerImpl(16)
+	: WorkManagerImpl(GetLogicalProcessorCount())
 	,	m_logWriter(info.logWriter)
 	,	m_logger(info.logWriter, SetGetPrefix(name))
 	,	m_isConnected(true)
@@ -39,12 +39,7 @@ namespace uba
 
 	NetworkClient::~NetworkClient()
 	{
-		StopListen();
-
-		Disconnect();
-
-		SCOPED_WRITE_LOCK(m_connectionsLock, lock);
-		m_connections.clear(); // Before tcpBackend
+		StopAll();
 
 		delete m_tcpBackend;
 
@@ -52,6 +47,20 @@ namespace uba
 			Crypto::DestroyKey(m_cryptoKey);
 	}
 
+	void NetworkClient::StopAll()
+	{
+		StopListen();
+
+		Disconnect();
+
+		{
+			SCOPED_WRITE_LOCK(m_connectionsLock, lock);
+			m_connections.clear();
+		}
+
+		FlushWork();
+	}
+
 	bool NetworkClient::Connect(NetworkBackend& backend, const tchar* ip, u16 port, bool* timedOut)
 	{
 		return backend.Connect(m_logger, ip, [&](void* connection, const sockaddr& remoteSocketAddr, bool* timedOut)
@@ -219,8 +228,8 @@ namespace uba
 		backend.SetDisconnectCallback(backendConnection, connection, [](void* context, void* connection)
 			{
 				auto& c = *(Connection*)context;
-				c.disconnectedEvent.Set();
 				c.owner.OnDisconnected(c, true);
+				c.disconnectedEvent.Set();
 			});
 		backend.SetRecvCallbacks(backendConnection, connection, ReceiveHeaderSize, ReceiveResponseHeader, ReceiveResponseBody, TC("ReceiveMessageResponse"));
 	}
diff --git a/Engine/Source/Programs/UnrealBuildAccelerator/Common/Public/UbaNetworkClient.h b/Engine/Source/Programs/UnrealBuildAccelerator/Common/Public/UbaNetworkClient.h
--- a/Engine/Source/Programs/UnrealBuildAccelerator/Common/Public/UbaNetworkClient.h
+++ b/Engine/Source/Programs/UnrealBuildAccelerator/Common/Public/UbaNetworkClient.h
@@ -33,6 +33,8 @@ namespace uba
 		NetworkClient(bool& outCtorSuccess, const NetworkClientCreateInfo& info = {}, const tchar* name = TC("UbaClient"));
 		~NetworkClient();
 
+		void StopAll();
+
 		bool Connect(NetworkBackend& backend, const tchar* ip, u16 port = DefaultPort, bool* timedOut = nullptr);
 		void Disconnect();
 
