[UBA]
* Changed so StorageProxy is not reading entire files first before sending them out to clients. Instead it now lazily loads each segment based on requests and send it

diff --git a/Engine/Source/Programs/UnrealBuildAccelerator/Common/Private/UbaStorageProxy.cpp b/Engine/Source/Programs/UnrealBuildAccelerator/Common/Private/UbaStorageProxy.cpp
--- a/Engine/Source/Programs/UnrealBuildAccelerator/Common/Private/UbaStorageProxy.cpp
+++ b/Engine/Source/Programs/UnrealBuildAccelerator/Common/Private/UbaStorageProxy.cpp
@@ -17,10 +17,24 @@ namespace uba
 	,	m_storageServerUid(storageServerUid)
 	,	m_name(name)
 	,	m_hasActiveFetchesEvent(true)
-	,	m_fileMapping(m_logger)
 	{
 		m_hasActiveFetchesEvent.Set();
 
+		m_server.RegisterOnClientDisconnected(0, [this](const Guid& clientUid, u32 clientId)
+			{
+				SCOPED_WRITE_LOCK(m_activeFetchesLock, lock);
+				for (auto it=m_activeFetches.begin(); it!=m_activeFetches.end();)
+				{
+					if (it->second.clientUid != clientUid)
+					{
+						++it;
+						continue;
+					}
+					PushId(it->first);
+					it = m_activeFetches.erase(it);
+				}
+			});
+
 		m_server.RegisterService(StorageServiceId,
 			[this](const ConnectionInfo& connectionInfo, u8 messageType, BinaryReader& reader, BinaryWriter& writer)
 			{
@@ -31,8 +45,13 @@ namespace uba
 				return ToString(StorageMessageType(messageType));
 			}
 		);
+	}
 
-		m_fileMapping.AddTransient(TC("Files"));
+	StorageProxy::~StorageProxy()
+	{
+		m_server.StopAll();
+		for (auto& kv : m_files)
+			delete[] kv.second.memory;
 	}
 
 	bool StorageProxy::Disconnect(u32 timeoutMs)
@@ -90,28 +109,22 @@ namespace uba
 				StringBuffer<> hint;
 				reader.ReadString(hint);
 
-				ScopedWriteLock filesLock(m_filesLock);
+				SCOPED_WRITE_LOCK(m_filesLock, filesLock);
 				auto insres = m_files.try_emplace(casKey);
 				FileEntry& file = insres.first->second;
 				filesLock.Leave();
 
-				MappedView view;
-				auto viewGuard = MakeGuard([&]() { m_fileMapping.UnmapView(view, hint.data); });
-
-				ScopedWriteLock fileLock(file.lock);
+				SCOPED_WRITE_LOCK(file.lock, fileLock);
 
+				bool hasAllSegments = false;
 				while (true)
 				{
-					if (file.view.size != 0)
-					{
-						fileLock.Leave();
-						view = m_fileMapping.MapView(file.view.handle, file.view.offset, file.view.size, hint.data);
+					if (file.memory)
 						break;
-					}
 
+					constexpr bool useLocalStorage = true;
 					bool storeCompressed = true;
-
-					if (m_localStorage && IsCompressed(casKey) && m_inProcessClientId && connectionInfo.GetId() != m_inProcessClientId)
+					if (useLocalStorage && m_localStorage && IsCompressed(casKey) && m_inProcessClientId && connectionInfo.GetId() != m_inProcessClientId)
 					{
 						// We need to leave this lock here since the in-process storage client might be asking for this file too and then we can end up in a deadlock
 						fileLock.Leave();
@@ -120,7 +133,7 @@ namespace uba
 
 						// Enter lock again, and also check if another thread might have already handled this file while we looked if it existed in local storage
 						fileLock.Enter();
-						if (file.view.size != 0)
+						if (file.size != 0)
 							continue;
 
 						if (hasCas)
@@ -133,16 +146,18 @@ namespace uba
 								if (sourceFile.OpenMemoryRead())
 								{
 									u64 fileSize = sourceFile.GetSize();
-									view = m_fileMapping.AllocAndMapView(MappedView_Transient, fileSize, 1, hint.data);
-									if (!view.memory)
+									file.memory = new u8[fileSize];
+									if (!file.memory)
 										return false;
-									memcpy(view.memory, sourceFile.GetData(), fileSize);
+									file.size = fileSize;
+									memcpy(file.memory, sourceFile.GetData(), fileSize);
+									hasAllSegments = true;
 								}
 							}
 						}
 					}
 
-					if (!view.handle.IsValid())
+					if (!file.memory)
 					{
 						NetworkMessage msg(m_client, ServiceId, messageType, writer2);
 						writer2.WriteBool(false); // Wants proxy
@@ -151,61 +166,66 @@ namespace uba
 						writer2.WriteBytes(reader.GetPositionData(), reader.GetLeft());
 
 						if (!msg.Send(reader2))
+						{
+							file.error = true;
 							return m_logger.Error(TC("FetchBegin failed for cas file %s (%s). Requested by %s"), CasKeyString(casKey).str, hint.data, GuidToString(connectionInfo.GetUid()).str);
+						}
 
 						BinaryReader tempReader(reader2.GetPositionData(), 0, reader2.GetLeft());
-						u32 sizeOfFirstMessage = u32(reader2.GetLeft());
 						u16 fetchId = tempReader.ReadU16();
 						if (fetchId == 0)
 						{
+							file.error = true;
 							m_logger.Error(TC("FetchBegin failed for cas file %s (%s). Requested by %s"), CasKeyString(casKey).str, hint.data, GuidToString(connectionInfo.GetUid()).str);
 							writer.WriteU16(0);
 							return true;
 						}
 						u64 fileSize = tempReader.Read7BitEncoded();
-						view = m_fileMapping.AllocAndMapView(MappedView_Transient, fileSize, 1, hint.data);
+						file.memory = new u8[fileSize];
+						file.size = fileSize;
 
 						u8 flags = tempReader.ReadByte();
 						storeCompressed = (flags >> 0) & 1;
 						bool sendEnd = (flags >> 1) & 1;
 						u64 fetchedSize = tempReader.GetLeft();
 
-						u32 responseSize = u32(fetchedSize);
+						memcpy(file.memory, tempReader.GetPositionData(), fetchedSize);
 
-						u64 left = fileSize;
-						u8* readBuffer = view.memory;
+						if (sendEnd && fetchedSize == fileSize)
+							SendEnd(casKey);
 
-						memcpy(view.memory, tempReader.GetPositionData(), responseSize);
-						left -= responseSize;
-						readBuffer += responseSize;
+						file.received = fetchedSize;
+						file.fetchId = fetchId;
+						file.sendEnd = sendEnd;
+					}
 
-						if (!StorageClient::SendAllSegments(m_client, fetchId, readBuffer, left, sizeOfFirstMessage))
-							return false;
+					file.storeCompressed = storeCompressed;
+					file.casKey = casKey;
 
-						if (sendEnd)
-						{
-							StackBinaryWriter<128> writer3;
-							NetworkMessage msg2(m_client, ServiceId, StorageMessageType_FetchEnd, writer3);
-							writer3.WriteCasKey(casKey);
-							if (!msg2.Send())
-								return false;
-						}
+					if (file.received < file.size)
+					{
+						u64 segmentSize = m_client.GetMessageMaxSize() - 5; // This is server response size - header.. TODO: Should be taken from server
+						u64 segmentCount = file.size / segmentSize + 1;
+						u64 lookupByteCount = (segmentCount + 7) / 8;
+						u8 initialValue = hasAllSegments ? 255 : 0;
+						file.segmentsAvailable.resize(lookupByteCount, initialValue);
 					}
-					file.storeCompressed = storeCompressed;
-					file.view = view;
-					fileLock.Leave();
 					break;
 				}
 
+				if (file.error)
+					return false;
+				fileLock.Leave();
+
 				u16 fetchId = u16(~0);
 
-				u64 headerSize = sizeof(u16) + Get7BitEncodedCount(file.view.size) + sizeof(u8);
-				u64 fetchedSize = Min(view.size, m_client.GetMessageMaxSize() - m_client.GetMessageReceiveHeaderSize() - headerSize);
+				u64 headerSize = sizeof(u16) + Get7BitEncodedCount(file.size) + sizeof(u8);
+				u64 fetchedSize = Min(file.size, m_client.GetMessageMaxSize() - m_client.GetMessageReceiveHeaderSize() - headerSize);
 
-				if (fetchedSize < view.size)
+				if (fetchedSize < file.size)
 				{
 					fetchId = PopId();
-					ScopedWriteLock lock(m_activeFetchesLock);
+					SCOPED_WRITE_LOCK(m_activeFetchesLock, lock);
 					if (m_activeFetches.empty())
 						m_hasActiveFetchesEvent.Reset();
 					auto res = m_activeFetches.try_emplace(fetchId);
@@ -213,6 +233,7 @@ namespace uba
 					ActiveFetch& fetch = res.first->second;
 					lock.Leave();
 
+					fetch.clientUid = connectionInfo.GetUid();
 					fetch.fetchedSize = fetchedSize;
 					fetch.file = &file;
 				}
@@ -221,9 +242,9 @@ namespace uba
 				flags |= u8(file.storeCompressed) << 0;
 
 				writer.WriteU16(fetchId);
-				writer.Write7BitEncoded(view.size);
+				writer.Write7BitEncoded(file.size);
 				writer.WriteByte(flags);
-				writer.WriteBytes(view.memory, fetchedSize);
+				writer.WriteBytes(file.memory, fetchedSize);
 
 				return true;
 			}
@@ -232,7 +253,7 @@ namespace uba
 				u16 fetchId = reader.ReadU16();
 				u32 fetchIndex = reader.ReadU32();
 
-				ScopedReadLock activeLock(m_activeFetchesLock);
+				SCOPED_READ_LOCK(m_activeFetchesLock, activeLock);
 				auto findIt = m_activeFetches.find(fetchId);
 				UBA_ASSERT(findIt != m_activeFetches.end());
 				ActiveFetch& fetch = findIt->second;
@@ -240,24 +261,109 @@ namespace uba
 
 				FileEntry& file = *fetch.file;
 
-				u64 headerSize = sizeof(u16) + Get7BitEncodedCount(file.view.size) + sizeof(u8);
+				u64 headerSize = sizeof(u16) + Get7BitEncodedCount(file.size) + sizeof(u8);
 				u64 firstFetchSize = m_client.GetMessageMaxSize() - m_client.GetMessageReceiveHeaderSize() - headerSize;
+				u64 segmentSize = m_client.GetMessageMaxSize() - 5; // This is server response size - header.. TODO: Should be taken from server
+
+				u64 offset = firstFetchSize + segmentSize * (fetchIndex - 1);
+				if (offset + segmentSize > file.size)
+					segmentSize = file.size - offset;
+				
+				SCOPED_WRITE_LOCK(file.lock, fileLock);
+				if (file.error)
+					return false;
+
+				u64 byteIndex = fetchIndex / 8;
+				u8 bitMask = u8(1 << (fetchIndex - byteIndex*8));
+				
+				u8& availableByte = file.segmentsAvailable[byteIndex];
+				if (!(availableByte & bitMask))
+				{
+					SegmentInFlight* activeSegment = nullptr;
+					auto asg = MakeGuard([&]
+						{
+							if (!--activeSegment->refCount)
+							{
+								if (auto next = activeSegment->next)
+									next->prev = activeSegment->prev;
+								else
+									file.lastInFlight = activeSegment->prev;
+								if (auto prev = activeSegment->prev)
+									prev->next = activeSegment->next;
+								else
+									file.firstInFlight = activeSegment->next;
+								delete activeSegment;
+							}
+						});
+
+					for (auto it=file.firstInFlight; it; it=it->next)
+					{
+						if (it->segmentIndex != fetchIndex)
+							continue;
+						activeSegment = it;
+						break;
+					}
 
-				u64 offset = firstFetchSize + writer.GetCapacityLeft() * (fetchIndex - 1);
+					if (!activeSegment)
+					{
+						activeSegment = new SegmentInFlight;
+						activeSegment->segmentIndex = fetchIndex;
+						activeSegment->refCount = 1;
+						activeSegment->done.Create(true);
+
+						activeSegment->next = nullptr;
+						activeSegment->prev = file.lastInFlight;
+						if (auto last = file.lastInFlight)
+							last->next = activeSegment;
+						file.lastInFlight = activeSegment;
+						if (!file.firstInFlight)
+							file.firstInFlight = activeSegment;
 
-				u64 segmentSize = writer.GetCapacityLeft();
-				if (offset + segmentSize > file.view.size)
-					segmentSize = file.view.size - offset;
-				MappedView view = m_fileMapping.MapView(file.view.handle, file.view.offset + offset, segmentSize, TC(""));
-				auto viewGuard = MakeGuard([&]() { m_fileMapping.UnmapView(view, TC("")); });
+						fileLock.Leave();
 
-				writer.WriteBytes(view.memory, segmentSize);
+						NetworkMessage msg(m_client, ServiceId, StorageMessageType_FetchSegment, writer2);
+						writer2.WriteU16(file.fetchId);
+						writer2.WriteU32(fetchIndex);
+
+						if (!msg.Send(reader2))
+						{
+							file.error = true;
+							activeSegment->done.Set();
+							return m_logger.Error(TC("FetchSegment failed. Requested by %s"), GuidToString(connectionInfo.GetUid()).str);
+						}
+						
+						file.received += segmentSize;
+						if (file.sendEnd && file.size == file.received)
+							SendEnd(file.casKey);
+
+						memcpy(file.memory + offset, reader2.GetPositionData(), segmentSize);
+
+						activeSegment->done.Set();
+
+						fileLock.Enter();
+						availableByte |= bitMask;
+					}
+					else
+					{
+						++activeSegment->refCount;
+						fileLock.Leave();
+						activeSegment->done.IsSet();
+						fileLock.Enter();
+						if (file.error)
+							return false;
+					}
+				}
+				fileLock.Leave();
+
+				const u8* memory = file.memory + offset;
+
+				writer.WriteBytes(memory, segmentSize);
 
 				u64 fetchedSize = fetch.fetchedSize.fetch_add(segmentSize) + segmentSize;
-				if (fetchedSize != file.view.size)
+				if (fetchedSize != file.size)
 					return true;
 
-				ScopedWriteLock activeLock2(m_activeFetchesLock);
+				SCOPED_WRITE_LOCK(m_activeFetchesLock, activeLock2);
 				m_activeFetches.erase(findIt);
 				if (m_activeFetches.empty())
 					m_hasActiveFetchesEvent.Set();
@@ -284,7 +390,7 @@ namespace uba
 
 	u16 StorageProxy::PopId()
 	{
-		ScopedWriteLock lock(m_availableIdsLock);
+		SCOPED_WRITE_LOCK(m_availableIdsLock, lock);
 		if (m_availableIds.empty())
 			return m_availableIdsHigh++;
 		u16 storeId = m_availableIds.back();
@@ -294,7 +400,16 @@ namespace uba
 
 	void StorageProxy::PushId(u16 id)
 	{
-		ScopedWriteLock lock(m_availableIdsLock);
+		SCOPED_WRITE_LOCK(m_availableIdsLock, lock);
 		m_availableIds.push_back(id);
 	}
+
+	bool StorageProxy::SendEnd(const CasKey& key)
+	{
+		StackBinaryWriter<128> writer;
+		NetworkMessage msg(m_client, ServiceId, StorageMessageType_FetchEnd, writer);
+		writer.WriteCasKey(key);
+		return msg.Send();
+	}
+
 }
diff --git a/Engine/Source/Programs/UnrealBuildAccelerator/Common/Public/UbaStorageProxy.h b/Engine/Source/Programs/UnrealBuildAccelerator/Common/Public/UbaStorageProxy.h
--- a/Engine/Source/Programs/UnrealBuildAccelerator/Common/Public/UbaStorageProxy.h
+++ b/Engine/Source/Programs/UnrealBuildAccelerator/Common/Public/UbaStorageProxy.h
@@ -18,6 +18,7 @@ namespace uba
 	{
 	public:
 		StorageProxy(NetworkServer& server, NetworkClient& client, const Guid& storageServerUid, const tchar* name, StorageImpl* localStorage = nullptr);
+		~StorageProxy();
 
 		bool Disconnect(u32 timeoutMs); // Use a timeout to try to make more graceful disconnect
 		void PrintSummary();
@@ -26,6 +27,7 @@ namespace uba
 		bool HandleMessage(const ConnectionInfo& connectionInfo, u8 messageType, BinaryReader& reader, BinaryWriter& writer);
 		u16 PopId();
 		void PushId(u16 id);
+		bool SendEnd(const CasKey& key);
 
 		static constexpr u8 ServiceId = StorageServiceId;
 
@@ -41,13 +43,31 @@ namespace uba
 
 		u32 m_inProcessClientId = 0;
 
-		struct FileEntry { ReaderWriterLock lock; MappedView view; bool storeCompressed = false; };
+		struct SegmentInFlight { u32 refCount; u32 segmentIndex; Event done; SegmentInFlight* prev; SegmentInFlight* next; };
+
+		struct FileEntry
+		{
+			ReaderWriterLock lock;
+			u8* memory = nullptr;
+			u64 size = 0;
+			Atomic<u64> received;
+			CasKey casKey;
+			u16 fetchId = 0;
+			bool storeCompressed = false;
+			bool sendEnd = false;
+			bool error = false;
+			Vector<u8> segmentsAvailable;
+			SegmentInFlight* firstInFlight = nullptr;
+			SegmentInFlight* lastInFlight = nullptr;
+		};
+
 		ReaderWriterLock m_filesLock;
 		UnorderedMap<CasKey, FileEntry> m_files;
 
 		struct ActiveFetch
 		{
 			FileEntry* file = nullptr;
+			Guid clientUid;
 			Atomic<u64> fetchedSize;
 		};
 
@@ -55,8 +75,6 @@ namespace uba
 		ReaderWriterLock m_activeFetchesLock;
 		UnorderedMap<u16, ActiveFetch> m_activeFetches;
 
-		FileMappingBuffer m_fileMapping;
-
 		ReaderWriterLock m_availableIdsLock;
 		Vector<u16> m_availableIds;
 		u16 m_availableIdsHigh = 1;
