[UBA]
* Elevate thread priority for the threads doing all the networking to make sure proxy and host are prioritizing receiving/sending network data

diff --git a/Engine/Source/Programs/UnrealBuildAccelerator/Common/Private/UbaNetworkBackendTcp.cpp b/Engine/Source/Programs/UnrealBuildAccelerator/Common/Private/UbaNetworkBackendTcp.cpp
--- a/Engine/Source/Programs/UnrealBuildAccelerator/Common/Private/UbaNetworkBackendTcp.cpp
+++ b/Engine/Source/Programs/UnrealBuildAccelerator/Common/Private/UbaNetworkBackendTcp.cpp
@@ -400,6 +400,8 @@ namespace uba
 
 	void NetworkBackendTcp::ThreadRecv(Connection& connection)
 	{
+		ElevateCurrentThreadPriority();
+		
 		if (connection.ready.IsSet(60000)) // This should never time out!
 		{
 			Guid connectionUid;
diff --git a/Engine/Source/Programs/UnrealBuildAccelerator/Common/Private/UbaNetworkServer.cpp b/Engine/Source/Programs/UnrealBuildAccelerator/Common/Private/UbaNetworkServer.cpp
--- a/Engine/Source/Programs/UnrealBuildAccelerator/Common/Private/UbaNetworkServer.cpp
+++ b/Engine/Source/Programs/UnrealBuildAccelerator/Common/Private/UbaNetworkServer.cpp
@@ -501,6 +501,8 @@ namespace uba
 
 	void NetworkServer::Worker::ThreadWorker(NetworkServer& server)
 	{
+		ElevateCurrentThreadPriority();
+
 		t_worker = this;
 		while (m_context->workAvailable.IsSet(~0u) && m_loop)
 			Update(*m_context, true);
diff --git a/Engine/Source/Programs/UnrealBuildAccelerator/Core/Private/UbaPlatform.cpp b/Engine/Source/Programs/UnrealBuildAccelerator/Core/Private/UbaPlatform.cpp
--- a/Engine/Source/Programs/UnrealBuildAccelerator/Core/Private/UbaPlatform.cpp
+++ b/Engine/Source/Programs/UnrealBuildAccelerator/Core/Private/UbaPlatform.cpp
@@ -341,6 +341,13 @@ namespace uba
 		return 1u;
 	}
 
+	void ElevateCurrentThreadPriority()
+	{
+		#if PLATFORM_WINDOWS
+		SetThreadPriority(GetCurrentThread(), THREAD_PRIORITY_ABOVE_NORMAL);
+		#endif
+	}
+
 
 #if !PLATFORM_WINDOWS
 	void GetMappingHandleName(StringBufferBase& out, u64 uid)
diff --git a/Engine/Source/Programs/UnrealBuildAccelerator/Core/Public/UbaPlatform.h b/Engine/Source/Programs/UnrealBuildAccelerator/Core/Public/UbaPlatform.h
--- a/Engine/Source/Programs/UnrealBuildAccelerator/Core/Public/UbaPlatform.h
+++ b/Engine/Source/Programs/UnrealBuildAccelerator/Core/Public/UbaPlatform.h
@@ -98,6 +98,8 @@ namespace uba
 	u32 GetLogicalProcessorCount();
 	u32 GetProcessorGroupCount();
 
+	void ElevateCurrentThreadPriority();
+
 #if PLATFORM_WINDOWS
 	inline constexpr bool CaseInsensitiveFs = true;
 	inline constexpr tchar PathSeparator = '\\';
