[UBA]
* Changed so sending in a 0 workercount will use logical core count
* Cleaned up code around work tracking

diff --git a/Engine/Source/Programs/UnrealBuildAccelerator/Common/Private/UbaNetworkServer.cpp b/Engine/Source/Programs/UnrealBuildAccelerator/Common/Private/UbaNetworkServer.cpp
--- a/Engine/Source/Programs/UnrealBuildAccelerator/Common/Private/UbaNetworkServer.cpp
+++ b/Engine/Source/Programs/UnrealBuildAccelerator/Common/Private/UbaNetworkServer.cpp
@@ -442,12 +442,7 @@ namespace uba
 		u32 size;
 		WorkerRec& rec = server.m_workerFunctions[context.serviceId];
 
-		u32 workIndex = 0;
-		if (server.m_trackWork)
-		{
-			workIndex = server.m_workCounter++;
-			server.m_startWork(workIndex, rec.toString(context.messageType));
-		}
+		u32 workId = server.TrackWorkStart(rec.toString(context.messageType));
 
 		if (!rec.func)
 		{
@@ -466,8 +461,7 @@ namespace uba
 			size = u32(writer.GetPosition());
 		}
 
-		if (server.m_trackWork)
-			server.m_endWork(workIndex);
+		server.TrackWorkEnd(workId);
 
 		if (context.id)
 		{
@@ -526,17 +520,11 @@ namespace uba
 				server.m_additionalWork.pop_front();
 				lock.Leave();
 
-				u32 workIndex = 0;
-				if (server.m_trackWork)
-				{
-					workIndex = server.m_workCounter++;
-					server.m_startWork(workIndex, work.desc.c_str());
-				}
+				u32 workId = server.TrackWorkStart(work.desc.c_str());
 
 				work.func();
 
-				if (server.m_trackWork)
-					server.m_endWork(workIndex);
+				server.TrackWorkEnd(workId);
 			}
 
 			// Both locks needs to be taken to verify if additional work
@@ -569,7 +557,11 @@ namespace uba
 	{
 		outCtorSuccess = true;
 
-		u32 workerCount = Min(Max(info.workerCount, (u32)(1u)), (u32)(1024u));
+		u32 workerCount;
+		if (info.workerCount == 0)
+			workerCount = GetLogicalProcessorCount();
+		else
+			workerCount = Min(Max(info.workerCount, (u32)(1u)), (u32)(1024u));
 		m_maxWorkerCount = workerCount;
 
 		#if UBA_DEBUG
@@ -828,6 +820,21 @@ namespace uba
 		return m_maxWorkerCount;
 	}
 
+	u32 NetworkServer::TrackWorkStart(const tchar* desc)
+	{
+		if (!m_trackWork)
+			return 0;
+		u32 workId = m_workCounter++;
+		m_startWork(workId, desc);
+		return workId;
+	}
+
+	void NetworkServer::TrackWorkEnd(u32 id)
+	{
+		if (m_trackWork)
+			m_endWork(id);
+	}
+
 	u64 NetworkServer::GetTotalSentBytes()
 	{
 		return m_sendBytes;
@@ -903,17 +910,12 @@ namespace uba
 		m_additionalWork.pop_front();
 		lock.Leave();
 
-		u32 workIndex = 0;
-		if (m_trackWork)
-		{
-			workIndex = m_workCounter++;
-			m_startWork(workIndex, work.desc.c_str());
-		}
+		u32 workId = TrackWorkStart(work.desc.c_str());
 
 		work.func();
 
-		if (m_trackWork)
-			m_endWork(workIndex);
+		TrackWorkEnd(workId);
+
 		return true;
 	}
 
diff --git a/Engine/Source/Programs/UnrealBuildAccelerator/Common/Public/UbaNetworkServer.h b/Engine/Source/Programs/UnrealBuildAccelerator/Common/Public/UbaNetworkServer.h
--- a/Engine/Source/Programs/UnrealBuildAccelerator/Common/Public/UbaNetworkServer.h
+++ b/Engine/Source/Programs/UnrealBuildAccelerator/Common/Public/UbaNetworkServer.h
@@ -29,7 +29,7 @@ namespace uba
 	{
 		NetworkServerCreateInfo(LogWriter& w = g_consoleLogWriter) : logWriter(w) {}
 		LogWriter& logWriter;
-		u32 workerCount = 64;
+		u32 workerCount = 0; // Zero means it will use the number of logical cores as worker count
 		u32 sendSize = SendDefaultSize;
 		u32 receiveTimeoutSeconds = 0;
 	};
@@ -66,8 +66,10 @@ namespace uba
 		//void RegisterOnConnection(u8 id, const OnConnectionFunction& func);
 		//void UnregisterOnConnection(u8 id);
 
-		virtual void AddWork(const Function<void()>& work, u32 count, const tchar* desc) override;
-		virtual u32 GetWorkerCount() override;
+		virtual void AddWork(const Function<void()>& work, u32 count, const tchar* desc) override final;
+		virtual u32 GetWorkerCount() override final;
+		virtual u32 TrackWorkStart(const tchar* desc) override final;
+		virtual void TrackWorkEnd(u32 id) override final;
 
 		struct ClientStats
 		{
diff --git a/Engine/Source/Programs/UnrealBuildAccelerator/Common/Public/UbaWorkManager.h b/Engine/Source/Programs/UnrealBuildAccelerator/Common/Public/UbaWorkManager.h
--- a/Engine/Source/Programs/UnrealBuildAccelerator/Common/Public/UbaWorkManager.h
+++ b/Engine/Source/Programs/UnrealBuildAccelerator/Common/Public/UbaWorkManager.h
@@ -11,6 +11,9 @@ namespace uba
 	public:
 		virtual void AddWork(const Function<void()>& work, u32 count, const tchar* desc) = 0;
 		virtual u32 GetWorkerCount() = 0;
+
+		virtual u32 TrackWorkStart(const tchar* desc) { return 0; }
+		virtual void TrackWorkEnd(u32 id) {}
 	};
 
 
@@ -39,4 +42,12 @@ namespace uba
 		ReaderWriterLock m_availableWorkersLock;
 		Worker* m_firstAvailableWorker = nullptr;
 	};
+
+	struct TrackWorkScope
+	{
+		TrackWorkScope(WorkManager& wm, const tchar* desc) : workManager(wm), workIndex(wm.TrackWorkStart(desc)) {}
+		~TrackWorkScope() { workManager.TrackWorkEnd(workIndex); }
+		WorkManager& workManager;
+		u32 workIndex;
+	};
 }
