[UBA]
* Changed UbaAgent max connection count from 4 to 8 to improve throughput when being proxy
* Changed UbaAgent worker count for storage proxy from 16 to 192.. these workers are just receiving/sending data and need to be able to go wide
* Fixed so UbaCli use storage proxy when testing in process client

diff --git a/Engine/Source/Programs/UnrealBuildAccelerator/Agent/Private/UbaAgent.cpp b/Engine/Source/Programs/UnrealBuildAccelerator/Agent/Private/UbaAgent.cpp
--- a/Engine/Source/Programs/UnrealBuildAccelerator/Agent/Private/UbaAgent.cpp
+++ b/Engine/Source/Programs/UnrealBuildAccelerator/Agent/Private/UbaAgent.cpp
@@ -49,7 +49,7 @@ namespace uba
 	}();
 	u32				DefaultProcessorCount = []() { return GetLogicalProcessorCount(); }();
 	const tchar*	DefaultAgentName = []() { static tchar buf[256]; GetComputerNameW(buf, sizeof_array(buf)); return buf; }();
-	u32				DefaultMaxConnectionCount = 4;
+	u32				DefaultMaxConnectionCount = 8;
 
 	int PrintHelp(const tchar* message)
 	{
@@ -1028,7 +1028,7 @@ namespace uba
 					auto& proxy = *(Proxy*)userData;
 
 					NetworkServerCreateInfo nsci(proxy.logWriter);
-					nsci.workerCount = 16;
+					nsci.workerCount = 192;
 					nsci.receiveTimeoutSeconds = 60;
 
 					StringBuffer<256> prefix;
@@ -1132,8 +1132,7 @@ namespace uba
 						LoggerWithWriter(g_consoleLogWriter, TC("")).Info(TC("Proxy timed out waiting for zero active fetches"));
 				if (proxy.server)
 					proxy.server->StopAll();
-				client->StopListen();
-				client->Disconnect();
+				client->StopAll();
 				sessionClient->Stop();
 				storageClient->StopProxy();
 				loopLogging = false;
diff --git a/Engine/Source/Programs/UnrealBuildAccelerator/Cli/Private/UbaCli.cpp b/Engine/Source/Programs/UnrealBuildAccelerator/Cli/Private/UbaCli.cpp
--- a/Engine/Source/Programs/UnrealBuildAccelerator/Cli/Private/UbaCli.cpp
+++ b/Engine/Source/Programs/UnrealBuildAccelerator/Cli/Private/UbaCli.cpp
@@ -7,6 +7,7 @@
 #include "UbaSessionClient.h"
 #include "UbaSessionServer.h"
 #include "UbaStorageClient.h"
+#include "UbaStorageProxy.h"
 #include "UbaStorageServer.h"
 #include "UbaPlatform.h"
 #include "UbaVersion.h"
@@ -550,9 +551,55 @@ namespace uba
 			{
 				NetworkClient client(ctorSuccess, { logWriter });
 
+				Event wakeupSessionWait(false);
+				Atomic<u32> targetConnectionCount = 1;
+				u32 maxConnectionCount = 4;
+				struct Proxy
+				{
+					LogWriter& logWriter;
+					NetworkClient* client;
+					Event& wakeupSessionWait;
+					u32& maxConnectionCount;
+					Atomic<u32>& targetConnectionCount;
+					NetworkServer* server = nullptr;
+					StorageProxy* storage = nullptr;
+					StorageClient* storageClient = nullptr;
+					TString serverPrefix;
+				} proxy { g_consoleLogWriter, &client, wakeupSessionWait, maxConnectionCount, targetConnectionCount };
+				auto psg = MakeGuard([&]() { delete proxy.server; });
+				auto pg = MakeGuard([&]() { delete proxy.storage; });
+
+				static auto startProxy = [](void* userData, u16 proxyPort, const Guid& storageServerUid)
+					{
+						auto& proxy = *(Proxy*)userData;
+
+						NetworkServerCreateInfo nsci(proxy.logWriter);
+						nsci.workerCount = 192;
+						nsci.receiveTimeoutSeconds = 60;
+
+						StringBuffer<256> prefix;
+						prefix.Append(TC("UbaProxyServer (")).Append(GuidToString(proxy.client->GetUid()).str).Append(')');
+						proxy.serverPrefix = prefix.data;
+						bool ctorSuccess = true;
+						proxy.server = new NetworkServer(ctorSuccess, nsci, proxy.serverPrefix.c_str());
+						if (!ctorSuccess)
+						{
+							delete proxy.server;
+							return false;
+						}
+						proxy.storage = new StorageProxy(*proxy.server, *proxy.client, storageServerUid, TC("Wooohoo"), proxy.storageClient);
+						proxy.server->StartListen(proxy.client->GetTcpBackend(), proxyPort);
+						proxy.targetConnectionCount = proxy.maxConnectionCount;
+						proxy.wakeupSessionWait.Set();
+						return true;
+					};
+
 				StringBuffer<> clientRootDir;
 				clientRootDir.Append(g_rootDir).Append("Agent");
 				StorageClientCreateInfo storageClientInfo(client, clientRootDir.data);
+				storageClientInfo.zone = TC("FOO");
+				storageClientInfo.startProxyCallback = startProxy;
+				storageClientInfo.startProxyUserData = &proxy;
 				StorageClient storageClient(storageClientInfo);
 
 				SessionClientCreateInfo sessionClientInfo(storageClient, client, logWriter);
@@ -565,7 +612,7 @@ namespace uba
 					if (!client.Connect(networkBackend, TC("127.0.0.1"), port))
 						return logger.Error(TC("Failed to connect"));
 
-				auto cg = MakeGuard([&]() { sessionClient.Stop(); client.Disconnect(); });
+				auto cg = MakeGuard([&]() { sessionClient.Stop(); client.StopAll(); });
 
 				return func();
 			};
diff --git a/Engine/Source/Programs/UnrealBuildAccelerator/Test/Private/UbaTestSession.h b/Engine/Source/Programs/UnrealBuildAccelerator/Test/Private/UbaTestSession.h
--- a/Engine/Source/Programs/UnrealBuildAccelerator/Test/Private/UbaTestSession.h
+++ b/Engine/Source/Programs/UnrealBuildAccelerator/Test/Private/UbaTestSession.h
@@ -80,7 +80,7 @@ namespace uba
 		sessionClientInfo.rootDir = rootDir.data;
 		SessionClient sessionClient(sessionClientInfo);
 
-		auto cg = MakeGuard([&]() { sessionClient.Stop(); client.Disconnect(); });
+		auto cg = MakeGuard([&]() { sessionClient.Stop(); client.StopAll(); });
 
 		StringBuffer<> workingDir;
 		workingDir.Append(testRootDir).Append(TC("WorkingDir"));
