[UBA]
* Reduced contention in SessionClient when testing if processes are done and spawning new processes
* Changed so FileMappingBuffer is lazily closing mapped memory. For some reason these calls end up being very expensive and we don't have to close them straight away

diff --git a/Engine/Source/Programs/UnrealBuildAccelerator/Common/Private/UbaFileMapping.cpp b/Engine/Source/Programs/UnrealBuildAccelerator/Common/Private/UbaFileMapping.cpp
--- a/Engine/Source/Programs/UnrealBuildAccelerator/Common/Private/UbaFileMapping.cpp
+++ b/Engine/Source/Programs/UnrealBuildAccelerator/Common/Private/UbaFileMapping.cpp
@@ -289,8 +289,9 @@ namespace uba
 #endif
 	}
 
-	FileMappingBuffer::FileMappingBuffer(Logger& logger)
+	FileMappingBuffer::FileMappingBuffer(Logger& logger, WorkManager* workManager)
 	:	m_logger(logger)
+	,	m_workManager(workManager)
 	{
 		m_pageSize = 64*1024;
 	}
@@ -520,37 +521,45 @@ namespace uba
 		return res;
 	}
 
-	void FileMappingBuffer::UnmapView(MappedView view, const tchar* hint, u64 newSize)
+	void FileMappingBuffer::UnmapView(MappedView view, const tchar* hint_, u64 newSize)
 	{
 		if (!view.handle.IsValid())
 			return;
-		u8 storageIndex = 255;
-		File& file = GetFile(view.handle, storageIndex);
 
-		u64 alignedOffsetStart = AlignUp(view.offset - (m_pageSize - 1), m_pageSize);
-		u64 alignedOffsetEnd = AlignUp(view.offset + view.size, m_pageSize);
+		auto unmap = [=](const tchar* hint)
+			{
+				u8 storageIndex = 255;
+				File& file = GetFile(view.handle, storageIndex);
 
-		u8* memory = view.memory - (view.offset - alignedOffsetStart);
-		u64 mapSize = alignedOffsetEnd - alignedOffsetStart;
-		if (!UnmapViewOfFile(memory, mapSize, hint))
-		{
-			m_logger.Error(TC("%s - Failed to unmap view on address %llx (offset %llu) - %s (%s)"), file.name, u64(memory), view.offset, hint, LastErrorToText().data);
-		}
+				u64 alignedOffsetStart = AlignUp(view.offset - (m_pageSize - 1), m_pageSize);
+				u64 alignedOffsetEnd = AlignUp(view.offset + view.size, m_pageSize);
 
-		if (newSize != InvalidValue)
-		{
-			if (newSize != view.size)
-			{
-				UBA_ASSERT(!file.commitOnAlloc);
-				UBA_ASSERTF(newSize < view.size, TC("%s - Reserved too little memory. Reserved %llu, needed %llu for %s"), file.name, view.size, newSize, hint);
-				file.size -= view.size - newSize;
-			}
+				u8* memory = view.memory - (view.offset - alignedOffsetStart);
+				u64 mapSize = alignedOffsetEnd - alignedOffsetStart;
+				if (!UnmapViewOfFile(memory, mapSize, hint))
+				{
+					m_logger.Error(TC("%s - Failed to unmap view on address %llx (offset %llu) - %s (%s)"), file.name, u64(memory), view.offset, hint, LastErrorToText().data);
+				}
 
-			MappingStorage& storage = m_storage[storageIndex];
-			PushFile(storage, &file);
-		}
+				if (newSize != InvalidValue)
+				{
+					if (newSize != view.size)
+					{
+						UBA_ASSERT(!file.commitOnAlloc);
+						UBA_ASSERTF(newSize < view.size, TC("%s - Reserved too little memory. Reserved %llu, needed %llu for %s"), file.name, view.size, newSize, hint);
+						file.size -= view.size - newSize;
+					}
+
+					MappingStorage& storage = m_storage[storageIndex];
+					PushFile(storage, &file);
+				}
 
-		--file.activeMapCount;
+				--file.activeMapCount;
+			};
+		if (m_workManager && newSize == InvalidValue)
+			m_workManager->AddWork([=, h = TString(hint_)]() { unmap(h.c_str()); }, 1, TC("UnmapView"));
+		else
+			unmap(hint_);
 	}
 
 	void FileMappingBuffer::GetSizeAndCount(FileMappingType type, u64& outSize, u32& outCount)
diff --git a/Engine/Source/Programs/UnrealBuildAccelerator/Common/Private/UbaSession.cpp b/Engine/Source/Programs/UnrealBuildAccelerator/Common/Private/UbaSession.cpp
--- a/Engine/Source/Programs/UnrealBuildAccelerator/Common/Private/UbaSession.cpp
+++ b/Engine/Source/Programs/UnrealBuildAccelerator/Common/Private/UbaSession.cpp
@@ -918,11 +918,12 @@ namespace uba
 		#endif
 	}
 
-	Session::Session(const SessionCreateInfo& info, const tchar* logPrefix, bool runningRemote)
+	Session::Session(const SessionCreateInfo& info, const tchar* logPrefix, bool runningRemote, WorkManager* workManager)
 	:	m_storage(info.storage)
 	,	m_logger(info.logWriter, logPrefix)
+	,	m_workManager(workManager)
 	,	m_directoryTable(&m_directoryTableMemory)
-	,	m_fileMappingBuffer(m_logger)
+	,	m_fileMappingBuffer(m_logger, workManager)
 	,	m_processCommunicationAllocator(m_logger, TC("CommunicationAllocator"))
 	,	m_trace(info.logWriter)
 	{
@@ -1338,6 +1339,7 @@ namespace uba
 			m_trace.ProcessExited(id, exitCode, writer.GetData(), writer.GetPosition(), logLines);
 			SCOPED_WRITE_LOCK(m_processStatsLock, lock);
 			m_processStats.Add(process.m_processStats);
+			m_stats.Add(process.m_sessionStats);
 		}
 
 		SCOPED_WRITE_LOCK(m_processesLock, lock);
@@ -2085,6 +2087,7 @@ namespace uba
 		logger.Info(TC("  MappingTable        %7u %9s"), u32(m_fileMappingTableLookup.size()), BytesToText(GetFileMappingSize()).str);
 		logger.Info(TC("  MappingBuffer       %7u %9s"), mappingBufferCount, BytesToText(mappingBufferSize).str);
 		logger.Info(TC(""));
+		m_stats.Print(logger);
 	}
 
 	void Session::CreateProcessJobObject()
diff --git a/Engine/Source/Programs/UnrealBuildAccelerator/Common/Private/UbaSessionClient.cpp b/Engine/Source/Programs/UnrealBuildAccelerator/Common/Private/UbaSessionClient.cpp
--- a/Engine/Source/Programs/UnrealBuildAccelerator/Common/Private/UbaSessionClient.cpp
+++ b/Engine/Source/Programs/UnrealBuildAccelerator/Common/Private/UbaSessionClient.cpp
@@ -11,7 +11,7 @@
 namespace uba
 {
 	SessionClient::SessionClient(const SessionClientCreateInfo& info)
-	: Session(info, TC("UbaSessionClient"), true)
+	: Session(info, TC("UbaSessionClient"), true, &info.client)
 	,	m_client(info.client)
 	,	m_name(info.name.data)
 	,	m_terminationTime(~0ull)
@@ -435,7 +435,8 @@ namespace uba
 					casKey = AsCompressed(casKey, false);
 					entry.handled = true;
 					Storage::RetrieveResult result;
-					if (!m_storage.RetrieveCasFile(result, casKey, fileName.data, &m_fileMappingBuffer, memoryMapAlignment, !IsRarelyRead(msg.process, fileName)))
+					bool allowProxy = !IsRarelyRead(msg.process, fileName);
+					if (!m_storage.RetrieveCasFile(result, casKey, fileName.data, &m_fileMappingBuffer, memoryMapAlignment, allowProxy))
 						return m_logger.Error(TC("Error retrieving cas entry %s (%s)"), CasKeyString(casKey).str, fileName.data);
 					entry.success = true;
 					entry.size = result.size;
@@ -1288,8 +1289,8 @@ namespace uba
 			ProcessRec(ProcessImpl* impl) : handle(impl) {}
 			ProcessHandle handle;
 			ReaderWriterLock lock;
-			bool isKilled = false;
-			bool isDone = false;
+			Atomic<bool> isKilled;
+			Atomic<bool> isDone;
 			float weight = 1.0f;
 		};
 		List<ProcessRec> activeProcesses;
@@ -1320,13 +1321,13 @@ namespace uba
 			for (auto it=activeProcesses.begin();it!=activeProcesses.end();)
 			{
 				ProcessRec& r = *it;
-				SCOPED_WRITE_LOCK(r.lock, lock);
 				if (!r.isDone)
 				{
 					++it;
 					continue;
 				}
-				lock.Leave();
+				r.lock.EnterWrite();
+				r.lock.LeaveWrite();
 				it = activeProcesses.erase(it);
 			}
 
@@ -1366,9 +1367,9 @@ namespace uba
 				for (auto it = activeProcesses.rbegin(); it != activeProcesses.rend(); ++it)
 				{
 					ProcessRec& rec = *it;
-					SCOPED_WRITE_LOCK(rec.lock, lock);
 					if (rec.isKilled || rec.isDone)
 						continue;
+					SCOPED_WRITE_LOCK(rec.lock, lock);
 					rec.handle.Cancel(true);
 					rec.isKilled = true;
 					SendReturnProcess(rec.handle.GetId(), TC("Running out of memory"));
@@ -1656,7 +1657,6 @@ namespace uba
 	void SessionClient::PrintSessionStats(Logger& logger)
 	{
 		Session::PrintSessionStats(logger);
-		m_stats.Print(logger);
 	}
 
 	bool SessionClient::GetNextProcess(Process& process, bool& outNewProcess, NextProcessInfo& outNextProcess, u32 prevExitCode, BinaryReader& statsReader)
diff --git a/Engine/Source/Programs/UnrealBuildAccelerator/Common/Private/UbaSessionServer.cpp b/Engine/Source/Programs/UnrealBuildAccelerator/Common/Private/UbaSessionServer.cpp
--- a/Engine/Source/Programs/UnrealBuildAccelerator/Common/Private/UbaSessionServer.cpp
+++ b/Engine/Source/Programs/UnrealBuildAccelerator/Common/Private/UbaSessionServer.cpp
@@ -89,7 +89,7 @@ namespace uba
 
 
 	SessionServer::SessionServer(const SessionServerCreateInfo& info)
-	:	Session(info, TC("UbaSessionServer"), false)
+	:	Session(info, TC("UbaSessionServer"), false, &info.server)
 	,	m_server(info.server)
 	,	m_maxRemoteProcessCount(~0u)
 	{
diff --git a/Engine/Source/Programs/UnrealBuildAccelerator/Common/Private/UbaStorage.cpp b/Engine/Source/Programs/UnrealBuildAccelerator/Common/Private/UbaStorage.cpp
--- a/Engine/Source/Programs/UnrealBuildAccelerator/Common/Private/UbaStorage.cpp
+++ b/Engine/Source/Programs/UnrealBuildAccelerator/Common/Private/UbaStorage.cpp
@@ -258,7 +258,14 @@ namespace uba
 				if (!uncompressedData)
 					return m_logger.Error(TC("Failed to map view of file mapping for %s (%s)"), from, LastErrorToText().data);
 
-				auto udg = MakeGuard([&]() { UnmapViewOfFile(uncompressedData, fileSize, from); });
+				auto udg = MakeGuard([&]()
+					{
+						#if UBA_EXPERIMENTAL
+						m_workManager->AddWork([=, f = TString(from)]() { UnmapViewOfFile(uncompressedData, fileSize, f.c_str()); }, 1, TC("UnmapFile"));
+						#else
+						UnmapViewOfFile(uncompressedData, fileSize, from);
+						#endif
+					});
 
 				if (!WriteMemToCompressedFile(destinationFile, workCount, uncompressedData, fileSize, maxUncompressedBlock, totalWritten))
 					return false;
@@ -1046,7 +1053,7 @@ namespace uba
 	:	m_workManager(info.workManager)
 	,	m_logger(info.writer, logPrefix)
 	,	m_activeCopyOrLinkEvent(false)
-	,	m_casDataBuffer(m_logger)
+	,	m_casDataBuffer(m_logger, info.workManager)
 	{
 		m_casCapacityBytes = info.casCapacityBytes;
 		m_storeCompressed = info.storeCompressed;
@@ -2310,7 +2317,13 @@ namespace uba
 				m_logger.Error(TC("Failed to map view of file mapping for %s (%s)"), fileName, LastErrorToText().data);
 				return CasKeyZero;
 			}
-			auto udg = MakeGuard([&]() { UnmapViewOfFile(fileData, fileSize, fileName); });
+			auto udg = MakeGuard([&]() { 
+					#if UBA_EXPERIMENTAL
+					m_workManager->AddWork([=, fn = TString(fileName)]() { UnmapViewOfFile(fileData, fileSize, fn.c_str()); }, 1, TC("UnmapFile"));
+					#else
+					UnmapViewOfFile(fileData, fileSize, fileName);
+					#endif
+				});
 
 			struct WorkRec
 			{
diff --git a/Engine/Source/Programs/UnrealBuildAccelerator/Common/Public/UbaFileMapping.h b/Engine/Source/Programs/UnrealBuildAccelerator/Common/Public/UbaFileMapping.h
--- a/Engine/Source/Programs/UnrealBuildAccelerator/Common/Public/UbaFileMapping.h
+++ b/Engine/Source/Programs/UnrealBuildAccelerator/Common/Public/UbaFileMapping.h
@@ -9,6 +9,8 @@
 
 namespace uba
 {
+	class WorkManager;
+
 #if PLATFORM_WINDOWS
 	struct FileMappingHandle
 	{
@@ -71,7 +73,7 @@ namespace uba
 	class FileMappingBuffer
 	{
 	public:
-		FileMappingBuffer(Logger& logger);
+		FileMappingBuffer(Logger& logger, WorkManager* workManager);
 		~FileMappingBuffer();
 
 		bool AddTransient(const tchar* name);
@@ -103,6 +105,7 @@ namespace uba
 		File& GetFile(FileMappingHandle handle, u8& outStorageIndex);
 
 		Logger& m_logger;
+		WorkManager* m_workManager;
 		u64 m_pageSize = 0;
 
 		struct MappingStorage
@@ -123,6 +126,9 @@ namespace uba
 		void CloseMappingStorage(MappingStorage& storage);
 
 		MappingStorage m_storage[2];
+
+		FileMappingBuffer(const FileMappingBuffer&) = delete;
+		void operator=(const FileMappingBuffer&) = delete;
 	};
 
 
@@ -161,5 +167,8 @@ namespace uba
 		u64 m_mappingCount = 0;
 
 		Set<u64> m_availableBlocks;
+
+		FileMappingAllocator(const FileMappingAllocator&) = delete;
+		void operator=(const FileMappingAllocator&) = delete;
 	};
 }
diff --git a/Engine/Source/Programs/UnrealBuildAccelerator/Common/Public/UbaSession.h b/Engine/Source/Programs/UnrealBuildAccelerator/Common/Public/UbaSession.h
--- a/Engine/Source/Programs/UnrealBuildAccelerator/Common/Public/UbaSession.h
+++ b/Engine/Source/Programs/UnrealBuildAccelerator/Common/Public/UbaSession.h
@@ -20,6 +20,7 @@ namespace uba
 	class ProcessHandle;
 	class ProcessImpl;
 	class Storage;
+	class WorkManager;
 	struct ProcessStartInfo;
 	struct ProcessStats;
 	struct InitMessage;
@@ -68,7 +69,7 @@ namespace uba
 		Logger& GetLogger();			// Logger used for logging 
 		LogWriter& GetLogWriter();		// LogWriter used by logger
 
-		Session(const SessionCreateInfo& info, const tchar* logPrefix, bool runningRemote);
+		Session(const SessionCreateInfo& info, const tchar* logPrefix, bool runningRemote, WorkManager* workManager = nullptr);
 		virtual ~Session();
 
 	protected:
@@ -150,6 +151,8 @@ namespace uba
 		Storage& m_storage;
 		MutableLogger m_logger;
 
+		WorkManager* m_workManager;
+
 		StringBuffer<32> m_id;
 		StringBuffer<MaxPath> m_rootDir;
 		StringBuffer<MaxPath> m_sessionDir;
diff --git a/Engine/Source/Programs/UnrealBuildAccelerator/Core/Public/UbaPlatform.h b/Engine/Source/Programs/UnrealBuildAccelerator/Core/Public/UbaPlatform.h
--- a/Engine/Source/Programs/UnrealBuildAccelerator/Core/Public/UbaPlatform.h
+++ b/Engine/Source/Programs/UnrealBuildAccelerator/Core/Public/UbaPlatform.h
@@ -50,6 +50,8 @@
     #define ERROR_SUCCESS 0
 #endif
 
+#define UBA_EXPERIMENTAL 0
+
 namespace uba
 {
 	#ifndef sizeof_array
