[UBA]
* Added potential race condition fix that could happen during certain disconnect situations in storage proxy
* Added Deploy-Linux.bat that can be used to build all linux binaries. It also enables tsan when running "Deploy-Linux.bat debug"

diff --git a/Engine/Source/Programs/UnrealBuildAccelerator/Common/Private/UbaStorageProxy.cpp b/Engine/Source/Programs/UnrealBuildAccelerator/Common/Private/UbaStorageProxy.cpp
--- a/Engine/Source/Programs/UnrealBuildAccelerator/Common/Private/UbaStorageProxy.cpp
+++ b/Engine/Source/Programs/UnrealBuildAccelerator/Common/Private/UbaStorageProxy.cpp
@@ -130,29 +130,26 @@ namespace uba
 						fileLock.Leave();
 
 						bool hasCas = m_localStorage->EnsureCasFile(casKey, nullptr);
+						StringBuffer<> casFile;
+						hasCas = hasCas && m_localStorage->GetCasFileName(casFile, casKey);
 
 						// Enter lock again, and also check if another thread might have already handled this file while we looked if it existed in local storage
 						fileLock.Enter();
-						if (file.size != 0)
-							continue;
+						if (file.memory)
+							break;
 
 						if (hasCas)
 						{
-
-							StringBuffer<> casFile;
-							if (m_localStorage->GetCasFileName(casFile, casKey))
+							FileAccessor sourceFile(m_logger, casFile.data);
+							if (sourceFile.OpenMemoryRead())
 							{
-								FileAccessor sourceFile(m_logger, casFile.data);
-								if (sourceFile.OpenMemoryRead())
-								{
-									u64 fileSize = sourceFile.GetSize();
-									file.memory = new u8[fileSize];
-									if (!file.memory)
-										return false;
-									file.size = fileSize;
-									memcpy(file.memory, sourceFile.GetData(), fileSize);
-									hasAllSegments = true;
-								}
+								u64 fileSize = sourceFile.GetSize();
+								file.memory = new u8[fileSize];
+								if (!file.memory)
+									return false;
+								file.size = fileSize;
+								memcpy(file.memory, sourceFile.GetData(), fileSize);
+								hasAllSegments = true;
 							}
 						}
 					}
@@ -329,6 +326,7 @@ namespace uba
 						{
 							file.error = true;
 							activeSegment->done.Set();
+							fileLock.Enter();
 							return m_logger.Error(TC("FetchSegment failed. Requested by %s"), GuidToString(connectionInfo.GetUid()).str);
 						}
 						
@@ -347,8 +345,10 @@ namespace uba
 					{
 						++activeSegment->refCount;
 						fileLock.Leave();
-						activeSegment->done.IsSet();
+						bool success = activeSegment->done.IsSet(10*60*1000); // This should never happen.
 						fileLock.Enter();
+						if (!success)
+							return m_logger.Error(TC("Connection %s timed out after 10 minutes waiting for segment %u on cas entry %s to be available in storage proxy"), fetchIndex, CasKeyString(file.casKey).str, GuidToString(connectionInfo.GetUid()).str);
 						if (file.error)
 							return false;
 					}
diff --git a/Engine/Source/Programs/UnrealBuildAccelerator/Deploy-Linux.bat b/Engine/Source/Programs/UnrealBuildAccelerator/Deploy-Linux.bat
new file mode 100755
--- /dev/null
+++ b/Engine/Source/Programs/UnrealBuildAccelerator/Deploy-Linux.bat
@@ -0,0 +1,40 @@
+@echo off
+setlocal 
+
+pushd "%~dp0"
+
+set config=Shipping
+
+if "%~1"=="" goto noconfig
+set config=%1
+:noconfig
+
+set options=""
+if /I not "%config%"=="debug" goto nooptions
+set options="-enabletsan"
+:nooptions
+
+
+echo.
+echo === Building Targets ===
+pushd "../../../.."
+call Engine/Build/BatchFiles/RunUBT.bat -NoUba -NoUbaLocal -NoSNDBS -NoXGE ^
+	-Target="UbaAgent Linux %config% %options%" ^
+	-Target="UbaCli Linux %config% %options%" ^
+	-Target="UbaDetours Linux %config%" ^
+	-Target="UbaHost Linux %config%" ^
+	-Target="UbaTest Linux %config% %options%" ^
+	-Target="UbaTestApp Linux %config% %options%" ^
+
+if %ERRORLEVEL% NEQ 0 exit /b %ERRORLEVEL%
+popd
+
+echo.
+echo === Reconciling artifacts ===
+p4 reconcile ../../../Binaries/Linux/UnrealBuildAccelerator/...
+
+popd
+
+endlocal
+
+exit /b 0
