[Uba]
* Fixed bug in trace code where errors are more than 100 lines long

diff --git a/Engine/Source/Programs/UnrealBuildAccelerator/Common/Private/UbaTrace.cpp b/Engine/Source/Programs/UnrealBuildAccelerator/Common/Private/UbaTrace.cpp
--- a/Engine/Source/Programs/UnrealBuildAccelerator/Common/Private/UbaTrace.cpp
+++ b/Engine/Source/Programs/UnrealBuildAccelerator/Common/Private/UbaTrace.cpp
@@ -210,23 +210,17 @@ namespace uba
 		writer.WriteU32(processId);
 		writer.WriteU32(exitCode);
 		writer.WriteBytes(data, dataSize);
-		writer.Write7BitEncoded(u64(logLines.size()));
 		u32 lineCounter = 0;
 		for (auto& line : logLines)
 		{
 			if (lineCounter++ == 100) // We don't want to write the entire error in the trace stream to blow the entire buffer
-			{
-				if (!writer.EnsureMemory(100))
-					return;
-				writer.WriteByte(LogEntryType_Info);
-				writer.WriteString(TC("Error is cut-off. Look in normal log to see full error"));
 				break;
-			}
 			if (!writer.EnsureMemory(1 + (line.text.size()+1)*sizeof(tchar)))
-				return;
+				break;
 			writer.WriteByte(line.type);
 			writer.WriteString(line.text);
 		}
+		writer.WriteByte(255);
 	}
 
 	void Trace::ProcessReturned(u32 processId)
diff --git a/Engine/Source/Programs/UnrealBuildAccelerator/Common/Public/UbaTrace.h b/Engine/Source/Programs/UnrealBuildAccelerator/Common/Public/UbaTrace.h
--- a/Engine/Source/Programs/UnrealBuildAccelerator/Common/Public/UbaTrace.h
+++ b/Engine/Source/Programs/UnrealBuildAccelerator/Common/Public/UbaTrace.h
@@ -51,7 +51,7 @@ namespace uba
 		TraceType_StatusUpdate,
 	};
 
-	static constexpr u32 TraceVersion = 21;
+	static constexpr u32 TraceVersion = 22;
 	static constexpr u32 TraceReadCompatibilityVersion = 6;
 
 	class Trace
diff --git a/Engine/Source/Programs/UnrealBuildAccelerator/Visualizer/Private/UbaTraceReader.cpp b/Engine/Source/Programs/UnrealBuildAccelerator/Visualizer/Private/UbaTraceReader.cpp
--- a/Engine/Source/Programs/UnrealBuildAccelerator/Visualizer/Private/UbaTraceReader.cpp
+++ b/Engine/Source/Programs/UnrealBuildAccelerator/Visualizer/Private/UbaTraceReader.cpp
@@ -533,7 +533,17 @@ namespace uba
 				process.systemStats.Read(reader);
 			}
 
-			if (out.version >= 20)
+			if (out.version >= 22)
+			{
+				while (true)
+				{
+					auto type = (LogEntryType)reader.ReadByte();
+					if (type == 255)
+						break;
+					process.logLines.emplace_back(reader.ReadString(), type);
+				}
+			}
+			else if (out.version >= 20)
 			{
 				u64 logLineCount = reader.Read7BitEncoded();
 				process.logLines.reserve(logLineCount);
