[UBA]
* Changed NetworkMessage to use a callback for when message is done and moved event out of message. This is a change in preparation for UbaStorageProxy optimization

diff --git a/Engine/Source/Programs/UnrealBuildAccelerator/Common/Private/UbaNetworkClient.cpp b/Engine/Source/Programs/UnrealBuildAccelerator/Common/Private/UbaNetworkClient.cpp
--- a/Engine/Source/Programs/UnrealBuildAccelerator/Common/Private/UbaNetworkClient.cpp
+++ b/Engine/Source/Programs/UnrealBuildAccelerator/Common/Private/UbaNetworkClient.cpp
@@ -256,13 +256,13 @@ namespace uba
 		if (messageSize == ErrorSize)
 		{
 			msg->m_error = true;
-			msg->m_event.Set();
+			msg->Done();
 			return true;
 		}
 		else if (!messageSize)
 		{
 			connection.recvCount++;
-			msg->m_event.Set();
+			msg->Done();
 			return true;
 		}
 
@@ -282,7 +282,7 @@ namespace uba
 		auto& msg = *(NetworkMessage*)bodyContext;
 		if (recvError)
 			msg.m_error = true;
-		msg.m_event.Set();
+		msg.Done();
 		return true;
 	}
 
@@ -442,7 +442,7 @@ namespace uba
 			{
 				m->m_error = true;
 				//m->m_responseSize = 0; // There is a race here where a message could just have fully arrived when disconnected. Can't set this to zero because we might be passed the m_error check but before setting reader size
-				m->m_event.Set();
+				m->Done();
 			}
 			++messageId;
 		}
@@ -505,13 +505,7 @@ namespace uba
 		}
 
 		UBA_ASSERT(messageId < 65535);
-
-		auto rm = MakeGuard([&]()
-		{
-			if (response)
-				ReturnMessageId(messageId);
-		});
-
+		message.m_id = messageId;
 
 		u32 sendSize = u32(writer.GetPosition());
 		u8* data = writer.GetData();
@@ -534,9 +528,21 @@ namespace uba
 
 		connection.sendBytes += sendSize;
 
+		Event gotResponse;
+
+		if (response)
+		{
+			message.m_sendContext.flags = NetworkBackend::SendFlags_ExternalWait;
+			if (!async)
+			{
+				gotResponse.Create(true);
+				UBA_ASSERT(!message.m_doneFunc);
+				message.m_doneUserData = &gotResponse;
+				message.m_doneFunc = [](bool error, void* userData) { ((Event*)userData)->Set(); };
+			}
+		}
+
 		{
-			if (response)
-				message.m_sendContext.flags = NetworkBackend::SendFlags_ExternalWait;
 			TimerScope ts(connection.sendTimer);
 			if (!connection.backend->Send(m_logger, connection.backendConnection, data, sendSize, message.m_sendContext))
 			{
@@ -546,16 +552,12 @@ namespace uba
 		}
 
 		if (async)
-		{
-			message.m_id = messageId;
-			rm.Cancel();
 			return true;
-		}
 
 		if (response)
 		{
 			u32 timeoutMs = 10 * 60 * 1000;
-			if (!message.m_event.IsSet(timeoutMs))
+			if (!gotResponse.IsSet(timeoutMs))
 			{
 				m_logger.Error(TC("Timed out after 10 minutes waiting for message."));
 				message.m_error = true;
@@ -589,7 +591,6 @@ namespace uba
 	NetworkMessage::NetworkMessage(NetworkClient& client, u8 serviceId, u8 messageType, BinaryWriter& sendWriter)
 	:	m_client(client)
 	,	m_sendWriter(sendWriter)
-	,	m_event(true)
 	{
 		// Header (SendHeaderSize):
 		// 1 byte    - 2 bits for serviceid, 6 bits for messagetype
@@ -604,6 +605,7 @@ namespace uba
 
 	NetworkMessage::~NetworkMessage()
 	{
+		UBA_ASSERT(!m_id);
 	}
 
 	bool NetworkMessage::Send()
@@ -626,21 +628,16 @@ namespace uba
 		return res;
 	}
 
-	bool NetworkMessage::SendAsync(BinaryReader& response)
+	bool NetworkMessage::SendAsync(BinaryReader& response, DoneFunc* func, void* userData)
 	{
+		UBA_ASSERT(!m_doneFunc);
+		m_doneFunc = func;
+		m_doneUserData = userData;
 		return m_client.Send(*this, (u8*)response.GetPositionData(), u32(response.GetLeft()), true);
 	}
 
-	bool NetworkMessage::WaitForAsync(BinaryReader& response)
+	bool NetworkMessage::ProcessAsyncResults(BinaryReader& response)
 	{
-		u32 timeoutMs = 10 * 60 * 1000;
-		bool success = m_event.IsSet(timeoutMs);
-		m_client.ReturnMessageId(m_id);
-		if (!success)
-		{
-			m_client.m_logger.Error(TC("Timed out after 10 minutes waiting for message."));
-			return false;
-		}
 		if (m_error)
 			return false;
 
@@ -654,4 +651,14 @@ namespace uba
 		response.SetSize(response.GetPosition() + m_responseSize);
 		return true;
 	}
+
+	void NetworkMessage::Done()
+	{
+		if (m_id)
+		{
+			m_client.ReturnMessageId(m_id);
+			m_id = 0;
+		}
+		m_doneFunc(m_error, m_doneUserData);
+	}
 }
diff --git a/Engine/Source/Programs/UnrealBuildAccelerator/Common/Private/UbaStorageClient.cpp b/Engine/Source/Programs/UnrealBuildAccelerator/Common/Private/UbaStorageClient.cpp
--- a/Engine/Source/Programs/UnrealBuildAccelerator/Common/Private/UbaStorageClient.cpp
+++ b/Engine/Source/Programs/UnrealBuildAccelerator/Common/Private/UbaStorageClient.cpp
@@ -785,6 +785,7 @@ namespace uba
 			Entry(NetworkClient& client, u16 fetchId, u8* slot, u32 readIndex, u32 i, u32 messageMaxSize)
 				: message(client, ServiceId, StorageMessageType_FetchSegment, writer)
 				, reader(slot + i * messageMaxSize, 0, SendMaxSize)
+				, done(true)
 			{
 				writer.WriteU16(fetchId);
 				writer.WriteU32(readIndex + i + 1);
@@ -793,6 +794,7 @@ namespace uba
 			StackBinaryWriter<32> writer;
 			NetworkMessage message;
 			BinaryReader reader;
+			Event done;
 		};
 
 		u64 sendCountCapacity = capacity / messageMaxSize;
@@ -819,13 +821,18 @@ namespace uba
 			});
 
 		for (u32 i=0; i!=sendCount; ++i)
-			if (!entries[i].message.SendAsync(entries[i].reader))
+		{
+			auto& entry = entries[i];
+			if (!entry.message.SendAsync(entry.reader, [](bool error, void* userData) { ((Event*)userData)->Set(); }, &entry.done))
 				return false;
+		}
 
 		for (u32 i=0; i!=sendCount; ++i)
 		{
 			Entry& entry = entries[i];
-			if (!entries[i].message.WaitForAsync(entries[i].reader))
+			if (!entry.done.IsSet())
+				return false;
+			if (!entry.message.ProcessAsyncResults(entry.reader))
 				return false;
 			responseSize += u32(entry.reader.GetLeft());
 		}
@@ -841,6 +848,7 @@ namespace uba
 			Entry(NetworkClient& client, u16 fetchId, u8* readBuffer, u32 fetchIndex)
 				: message(client, ServiceId, StorageMessageType_FetchSegment, writer)
 				, reader(readBuffer, 0, SendMaxSize)
+				, done(true)
 			{
 				writer.WriteU16(fetchId);
 				writer.WriteU32(fetchIndex + 1);
@@ -849,6 +857,7 @@ namespace uba
 			StackBinaryWriter<16> writer;
 			NetworkMessage message;
 			BinaryReader reader;
+			Event done;
 		};
 
 		constexpr u32 EntryCount = 128;
@@ -863,7 +872,7 @@ namespace uba
 			if (fetchIndex >= EntryCount)
 			{
 				auto& entry = entries[entryIndex];
-				if (!entry.message.WaitForAsync(entry.reader))
+				if (!entry.done.IsSet())
 				{
 					success = false;
 					break;
@@ -873,7 +882,7 @@ namespace uba
 			}
 
 			auto& entry = *new (entries + entryIndex) Entry(client, fetchId, readBuffer, fetchIndex);
-			if (!entry.message.SendAsync(entry.reader))
+			if (!entry.message.SendAsync(entry.reader, [](bool error, void* userData) { ((Event*)userData)->Set(); }, &entry.done))
 			{
 				entry.~Entry();
 				success = false;
@@ -898,7 +907,9 @@ namespace uba
 		for (u32 i = 0; i != waitCount; ++i)
 		{
 			auto& entry = entries[waitIndex];
-			if (!entry.message.WaitForAsync(entry.reader))
+			if (!entry.done.IsSet())
+				success = false;
+			if (!entry.message.ProcessAsyncResults(entry.reader))
 				success = false;
 			entry.~Entry();
 			waitIndex = (waitIndex + 1) % EntryCount;
diff --git a/Engine/Source/Programs/UnrealBuildAccelerator/Common/Public/UbaNetworkMessage.h b/Engine/Source/Programs/UnrealBuildAccelerator/Common/Public/UbaNetworkMessage.h
--- a/Engine/Source/Programs/UnrealBuildAccelerator/Common/Public/UbaNetworkMessage.h
+++ b/Engine/Source/Programs/UnrealBuildAccelerator/Common/Public/UbaNetworkMessage.h
@@ -16,13 +16,17 @@ namespace uba
 		bool Send(BinaryReader& response, Timer& outTimer);
 		bool Send();
 
-		bool SendAsync(BinaryReader& response);
-		bool WaitForAsync(BinaryReader& response);
+		using DoneFunc = void(bool error, void* userData);
+		bool SendAsync(BinaryReader& response, DoneFunc* func, void* userData); // Try to minimize work in GetResponseFunc since it is running on receiving thread
+		bool ProcessAsyncResults(BinaryReader& response); // Note, this must be called after GetResponseFunc has been called
 
 	private:
+		void Done();
+
 		NetworkClient& m_client;
 		BinaryWriter& m_sendWriter;
-		Event m_event;
+		DoneFunc* m_doneFunc = nullptr;
+		void* m_doneUserData = nullptr;
 		void* m_response = nullptr;
 		NetworkBackend::SendContext m_sendContext;
 		u32 m_responseSize = 0;
