[Uba]
* Fixed so visualizer does not crash with version 21 traces that contains log messages that have over 100 lines

diff --git a/Engine/Source/Programs/UnrealBuildAccelerator/Visualizer/Private/UbaTraceReader.cpp b/Engine/Source/Programs/UnrealBuildAccelerator/Visualizer/Private/UbaTraceReader.cpp
--- a/Engine/Source/Programs/UnrealBuildAccelerator/Visualizer/Private/UbaTraceReader.cpp
+++ b/Engine/Source/Programs/UnrealBuildAccelerator/Visualizer/Private/UbaTraceReader.cpp
@@ -546,6 +546,8 @@ namespace uba
 			else if (out.version >= 20)
 			{
 				u64 logLineCount = reader.Read7BitEncoded();
+				if (logLineCount >= 101)
+					logLineCount = 101;
 				process.logLines.reserve(logLineCount);
 				while (logLineCount--)
 				{
