[UBA]
* Added PrefetchVirtualMemory function
* Fixed bug where code writing callstacks checks end of module to figure out if address belongs to module
* Changed so FileAccessor::CreateMemoryWrite set always allow read because you can't create file mappings that are write-only
* Added option to scheduler to force all runs to run native
* Fixed so bottleneck class has enter/leave function
* Added ScopedReadLock::Enter
* Moved LogWriter to its own file
* Visualizer - Fixed so when hovering over timeline time under pointer is showing
* Visualizer - Fixed so hovering over work records shows name, start and duration
* Visualizer - Added context menu item that can show/hide text of process bars
* Improved assert message
* Changed so lock is not on stack but is leaked to prevent shutdown warnings in debug
* Added missing includes

diff --git a/Engine/Source/Programs/UnrealBuildAccelerator/Common/Private/UbaFile.cpp b/Engine/Source/Programs/UnrealBuildAccelerator/Common/Private/UbaFile.cpp
--- a/Engine/Source/Programs/UnrealBuildAccelerator/Common/Private/UbaFile.cpp
+++ b/Engine/Source/Programs/UnrealBuildAccelerator/Common/Private/UbaFile.cpp
@@ -457,7 +457,7 @@ namespace uba
 
 		if (close(asFileDescriptor(h)) == 0)
 			return true;
-		UBA_ASSERTF(false, TC("CloseFile error handling not implemented"));
+		UBA_ASSERTF(false, TC("CloseFile error handling not implemented while failing to close %s (%s)"), fileName, strerror(errno));
 		return false;
 #endif
 	}
diff --git a/Engine/Source/Programs/UnrealBuildAccelerator/Common/Private/UbaFileAccessor.cpp b/Engine/Source/Programs/UnrealBuildAccelerator/Common/Private/UbaFileAccessor.cpp
--- a/Engine/Source/Programs/UnrealBuildAccelerator/Common/Private/UbaFileAccessor.cpp
+++ b/Engine/Source/Programs/UnrealBuildAccelerator/Common/Private/UbaFileAccessor.cpp
@@ -2,6 +2,7 @@
 
 #include "UbaFileAccessor.h"
 #include "UbaFile.h"
+#include "UbaLogger.h"
 #include "UbaStats.h"
 
 #if PLATFORM_LINUX
@@ -103,6 +104,8 @@ namespace uba
 
 	bool FileAccessor::CreateMemoryWrite(bool allowRead, u32 flagsAndAttributes, u64 size, const tchar* tempPath)
 	{
+		allowRead = true; // It is not possible to have write only access to file mappings it seems
+
 		m_size = size;
 
 		UBA_ASSERT(flagsAndAttributes != 0);
diff --git a/Engine/Source/Programs/UnrealBuildAccelerator/Common/Private/UbaFileMapping.cpp b/Engine/Source/Programs/UnrealBuildAccelerator/Common/Private/UbaFileMapping.cpp
--- a/Engine/Source/Programs/UnrealBuildAccelerator/Common/Private/UbaFileMapping.cpp
+++ b/Engine/Source/Programs/UnrealBuildAccelerator/Common/Private/UbaFileMapping.cpp
@@ -15,7 +15,7 @@
 namespace uba
 {
 #if PLATFORM_WINDOWS
-	ReaderWriterLock g_createFileHandleLock;
+	ReaderWriterLock& g_createFileHandleLock = *new ReaderWriterLock(); // Allocated and leaked just to prevent shutdown asserts in debug
 
 	//Atomic<u64> g_fileMappingCount;
 
@@ -392,6 +392,8 @@ namespace uba
 				m_logger.Error(TC("%s - Failed to allocate memory for %s (%s)"), f.name, hint, LastErrorToText().data);
 				return res;
 			}
+
+			PrefetchVirtualMemory(data + commitStart, commitSize);
 		}
 
 		f.size = newOffset;
diff --git a/Engine/Source/Programs/UnrealBuildAccelerator/Common/Private/UbaProcessStartInfoHolder.h b/Engine/Source/Programs/UnrealBuildAccelerator/Common/Private/UbaProcessStartInfoHolder.h
--- a/Engine/Source/Programs/UnrealBuildAccelerator/Common/Private/UbaProcessStartInfoHolder.h
+++ b/Engine/Source/Programs/UnrealBuildAccelerator/Common/Private/UbaProcessStartInfoHolder.h
@@ -4,8 +4,10 @@
 
 #define Local_GetLongPathNameW uba::GetLongPathNameW
 
+#include "UbaBinaryReaderWriter.h"
 #include "UbaFile.h"
 #include "UbaPathUtils.h"
+#include "UbaProcessStartInfo.h"
 
 namespace uba
 {
diff --git a/Engine/Source/Programs/UnrealBuildAccelerator/Common/Private/UbaScheduler.cpp b/Engine/Source/Programs/UnrealBuildAccelerator/Common/Private/UbaScheduler.cpp
--- a/Engine/Source/Programs/UnrealBuildAccelerator/Common/Private/UbaScheduler.cpp
+++ b/Engine/Source/Programs/UnrealBuildAccelerator/Common/Private/UbaScheduler.cpp
@@ -60,6 +60,7 @@ namespace uba
 	,	m_updateThreadLoop(false)
 	,	m_enableProcessReuse(info.enableProcessReuse)
 	,	m_forceRemote(info.forceRemote)
+	,	m_forceNative(info.forceNative)
 	{
 		m_session.RegisterGetNextProcess([this](Process& process, NextProcessInfo& outNextProcess, u32 prevExitCode)
 			{
@@ -285,7 +286,7 @@ namespace uba
 
 			auto& processEntry = m_processEntries[indexToRun];
 			auto info = processEntry.info;
-			bool canDetour = processEntry.canDetour;
+			bool canDetour = processEntry.canDetour && !m_forceNative;
 			bool wasSkipped = processEntry.status == ProcessStatus_Skipped;
 			lock.Leave();
 
@@ -499,6 +500,17 @@ namespace uba
 				si.arguments = arg.c_str();
 				si.workingDir = dir.c_str();
 				si.description = desc.c_str();
+
+				#if UBA_DEBUG
+				StringBuffer<> logFile;
+				if (true)
+				{
+					GetNameFromArguments(logFile, si.arguments, true);
+					logFile.Append(TC(".log"));
+					si.logFile = logFile.data;
+				};
+				#endif
+
 				EnqueueProcessInfo info { si };
 				info.dependencies = deps.data();
 				info.dependencyCount = u32(deps.size());
diff --git a/Engine/Source/Programs/UnrealBuildAccelerator/Common/Public/UbaBottleneck.h b/Engine/Source/Programs/UnrealBuildAccelerator/Common/Public/UbaBottleneck.h
--- a/Engine/Source/Programs/UnrealBuildAccelerator/Common/Public/UbaBottleneck.h
+++ b/Engine/Source/Programs/UnrealBuildAccelerator/Common/Public/UbaBottleneck.h
@@ -10,42 +10,47 @@ namespace uba
 
 	struct Bottleneck
 	{
-		Bottleneck(u32 mc) : underMax(true), activeCount(0), maxCount(mc) {}
-		ReaderWriterLock lock;
-		Event underMax;
-		u32 activeCount;
-		u32 maxCount;
-	};
+		Bottleneck(u32 mc) : m_underMax(true), m_activeCount(0), m_maxCount(mc) {}
 
-	struct BottleneckScope
-	{
-		BottleneckScope(Bottleneck& b) : bottleneck(b)
+		void Enter()
 		{
-			SCOPED_WRITE_LOCK(bottleneck.lock, lock);
+			SCOPED_WRITE_LOCK(m_lock, lock);
 			while (true)
 			{
-				if (bottleneck.activeCount < bottleneck.maxCount)
+				if (m_activeCount < m_maxCount)
 				{
-					++bottleneck.activeCount;
-					if (bottleneck.activeCount == bottleneck.maxCount)
-						bottleneck.underMax.Reset();
+					++m_activeCount;
+					if (m_activeCount == m_maxCount)
+						m_underMax.Reset();
 					break;
 				}
 
 				lock.Leave();
-				bottleneck.underMax.IsSet();
+				m_underMax.IsSet();
 				lock.Enter();
 			}
 		}
 
-		~BottleneckScope()
+		void Leave()
 		{
-			SCOPED_WRITE_LOCK(bottleneck.lock, lock);
-			if (bottleneck.activeCount == bottleneck.maxCount)
-				bottleneck.underMax.Set();
-			--bottleneck.activeCount;
+			SCOPED_WRITE_LOCK(m_lock, lock);
+			if (m_activeCount == m_maxCount)
+				m_underMax.Set();
+			--m_activeCount;
 		}
 
+	private:
+		ReaderWriterLock m_lock;
+		Event m_underMax;
+		u32 m_activeCount;
+		u32 m_maxCount;
+	};
+
+	struct BottleneckScope
+	{
+		BottleneckScope(Bottleneck& b) : bottleneck(b) { b.Enter(); }
+		~BottleneckScope() { bottleneck.Leave(); }
+
 		Bottleneck& bottleneck;
 	};
 
diff --git a/Engine/Source/Programs/UnrealBuildAccelerator/Common/Public/UbaDirectoryIterator.h b/Engine/Source/Programs/UnrealBuildAccelerator/Common/Public/UbaDirectoryIterator.h
--- a/Engine/Source/Programs/UnrealBuildAccelerator/Common/Public/UbaDirectoryIterator.h
+++ b/Engine/Source/Programs/UnrealBuildAccelerator/Common/Public/UbaDirectoryIterator.h
@@ -3,6 +3,7 @@
 #pragma once
 
 #include "UbaFile.h"
+#include "UbaLogger.h"
 #include "UbaStringBuffer.h"
 
 namespace uba
diff --git a/Engine/Source/Programs/UnrealBuildAccelerator/Common/Public/UbaFile.h b/Engine/Source/Programs/UnrealBuildAccelerator/Common/Public/UbaFile.h
--- a/Engine/Source/Programs/UnrealBuildAccelerator/Common/Public/UbaFile.h
+++ b/Engine/Source/Programs/UnrealBuildAccelerator/Common/Public/UbaFile.h
@@ -2,12 +2,13 @@
 
 #pragma once
 
-#include "UbaLogger.h"
 #include "UbaMemory.h"
 #include "UbaPlatform.h"
 
 namespace uba
 {
+	class Logger;
+
 	#if !PLATFORM_WINDOWS
 	inline constexpr u32 ERROR_FILE_NOT_FOUND = ENOENT;
 	inline constexpr u32 ERROR_PATH_NOT_FOUND = ENOENT;
diff --git a/Engine/Source/Programs/UnrealBuildAccelerator/Common/Public/UbaFileMapping.h b/Engine/Source/Programs/UnrealBuildAccelerator/Common/Public/UbaFileMapping.h
--- a/Engine/Source/Programs/UnrealBuildAccelerator/Common/Public/UbaFileMapping.h
+++ b/Engine/Source/Programs/UnrealBuildAccelerator/Common/Public/UbaFileMapping.h
@@ -5,10 +5,10 @@
 #include "UbaEvent.h"
 #include "UbaHash.h"
 #include "UbaSynchronization.h"
-#include "UbaLogger.h"
 
 namespace uba
 {
+	class Logger;
 	class WorkManager;
 
 #if PLATFORM_WINDOWS
diff --git a/Engine/Source/Programs/UnrealBuildAccelerator/Common/Public/UbaLogWriter.h b/Engine/Source/Programs/UnrealBuildAccelerator/Common/Public/UbaLogWriter.h
new file mode 100644
--- /dev/null
+++ b/Engine/Source/Programs/UnrealBuildAccelerator/Common/Public/UbaLogWriter.h
@@ -0,0 +1,29 @@
+// Copyright Epic Games, Inc. All Rights Reserved.
+
+#pragma once
+
+#include "UbaBase.h"
+
+namespace uba
+{
+	enum LogEntryType : u8
+	{
+		LogEntryType_Error = 0,
+		LogEntryType_Warning = 1,
+		LogEntryType_Info = 2,
+		LogEntryType_Detail = 3,
+		LogEntryType_Debug = 4,
+	};
+
+	class LogWriter
+	{
+	public:
+		virtual ~LogWriter() = default;
+		virtual void BeginScope() = 0;
+		virtual void EndScope() = 0;
+		virtual void Log(LogEntryType type, const tchar* str, u32 strLen, const tchar* prefix = nullptr, u32 prefixLen = 0) = 0;
+	};
+
+	extern LogWriter& g_consoleLogWriter;
+	extern LogWriter& g_nullLogWriter;
+}
diff --git a/Engine/Source/Programs/UnrealBuildAccelerator/Common/Public/UbaLogger.h b/Engine/Source/Programs/UnrealBuildAccelerator/Common/Public/UbaLogger.h
--- a/Engine/Source/Programs/UnrealBuildAccelerator/Common/Public/UbaLogger.h
+++ b/Engine/Source/Programs/UnrealBuildAccelerator/Common/Public/UbaLogger.h
@@ -2,38 +2,19 @@
 
 #pragma once
 
-#include "UbaBase.h"
-#include <stdarg.h>
-
+#include "UbaLogWriter.h"
+#include "UbaSynchronization.h"
 
 #define UBA_DEBUG_LOGGER 0
 
 namespace uba
 {
-	enum LogEntryType : u8
-	{
-		LogEntryType_Error = 0,
-		LogEntryType_Warning = 1,
-		LogEntryType_Info = 2,
-		LogEntryType_Detail = 3,
-		LogEntryType_Debug = 4,
-	};
-
 	struct LogEntry
 	{
 		LogEntryType type;
 		const tchar* string;
 	};
 
-	class LogWriter
-	{
-	public:
-		virtual ~LogWriter() = default;
-		virtual void BeginScope() = 0;
-		virtual void EndScope() = 0;
-		virtual void Log(LogEntryType type, const tchar* str, u32 strLen, const tchar* prefix = nullptr, u32 prefixLen = 0) = 0;
-	};
-
 	class Logger
 	{
 	public:
@@ -65,6 +46,13 @@ namespace uba
 		u32 m_prefixLen;
 	};
 
+	struct MutableLogger : public LoggerWithWriter
+	{
+		MutableLogger(LogWriter& writer, const tchar* prefix) : LoggerWithWriter(writer, prefix) {}
+		virtual void Log(LogEntryType type, const tchar* str, u32 strLen) override { if (!isMuted) LoggerWithWriter::Log(type, str, strLen); }
+		Atomic<bool> isMuted;
+	};
+
 	class FilteredLogWriter : public LogWriter
 	{
 	public:
@@ -77,10 +65,6 @@ namespace uba
 		LogEntryType m_level;
 	};
 
-	extern LogWriter& g_consoleLogWriter;
-	extern LogWriter& g_nullLogWriter;
-
-
 	struct BytesToText
 	{
 		BytesToText(u64 bytes);
diff --git a/Engine/Source/Programs/UnrealBuildAccelerator/Common/Public/UbaProcessHandle.h b/Engine/Source/Programs/UnrealBuildAccelerator/Common/Public/UbaProcessHandle.h
--- a/Engine/Source/Programs/UnrealBuildAccelerator/Common/Public/UbaProcessHandle.h
+++ b/Engine/Source/Programs/UnrealBuildAccelerator/Common/Public/UbaProcessHandle.h
@@ -2,7 +2,7 @@
 
 #pragma once
 
-#include "UbaLogger.h"
+#include "UbaLogWriter.h"
 #include "UbaMemory.h"
 
 namespace uba
diff --git a/Engine/Source/Programs/UnrealBuildAccelerator/Common/Public/UbaProcessStartInfo.h b/Engine/Source/Programs/UnrealBuildAccelerator/Common/Public/UbaProcessStartInfo.h
--- a/Engine/Source/Programs/UnrealBuildAccelerator/Common/Public/UbaProcessStartInfo.h
+++ b/Engine/Source/Programs/UnrealBuildAccelerator/Common/Public/UbaProcessStartInfo.h
@@ -2,7 +2,7 @@
 
 #pragma once
 
-#include "UbaLogger.h"
+#include "UbaLogWriter.h"
 
 namespace uba
 {
diff --git a/Engine/Source/Programs/UnrealBuildAccelerator/Common/Public/UbaScheduler.h b/Engine/Source/Programs/UnrealBuildAccelerator/Common/Public/UbaScheduler.h
--- a/Engine/Source/Programs/UnrealBuildAccelerator/Common/Public/UbaScheduler.h
+++ b/Engine/Source/Programs/UnrealBuildAccelerator/Common/Public/UbaScheduler.h
@@ -21,6 +21,7 @@ namespace uba
 		u32 maxLocalProcessors = ~0u; // Max local processors to use. ~0u means it will use all processors
 		bool enableProcessReuse = false; // If this is true, the system will allow processes to be reused when they're asking for it.
 		bool forceRemote = false; // Force all processes that can run remotely to run remotely.
+		bool forceNative = false; // Force all processes to run native (not detoured)
 	};
 
 	struct EnqueueProcessInfo
@@ -110,6 +111,7 @@ namespace uba
 		Atomic<bool> m_loop;
 		bool m_enableProcessReuse;
 		bool m_forceRemote;
+		bool m_forceNative;
 
 		float m_activeLocalProcessWeight = 0.0f;
 
diff --git a/Engine/Source/Programs/UnrealBuildAccelerator/Common/Public/UbaSessionCreateInfo.h b/Engine/Source/Programs/UnrealBuildAccelerator/Common/Public/UbaSessionCreateInfo.h
--- a/Engine/Source/Programs/UnrealBuildAccelerator/Common/Public/UbaSessionCreateInfo.h
+++ b/Engine/Source/Programs/UnrealBuildAccelerator/Common/Public/UbaSessionCreateInfo.h
@@ -2,7 +2,7 @@
 
 #pragma once
 
-#include "UbaLogger.h"
+#include "UbaLogWriter.h"
 
 namespace uba
 {
diff --git a/Engine/Source/Programs/UnrealBuildAccelerator/Common/Public/UbaStorage.h b/Engine/Source/Programs/UnrealBuildAccelerator/Common/Public/UbaStorage.h
--- a/Engine/Source/Programs/UnrealBuildAccelerator/Common/Public/UbaStorage.h
+++ b/Engine/Source/Programs/UnrealBuildAccelerator/Common/Public/UbaStorage.h
@@ -5,9 +5,10 @@
 #define Local_GetLongPathNameW uba::GetLongPathNameW
 
 #include "UbaFile.h"
-#include "UbaPathUtils.h"
-#include "UbaMemory.h"
 #include "UbaFileMapping.h"
+#include "UbaLogger.h"
+#include "UbaMemory.h"
+#include "UbaPathUtils.h"
 #include "UbaStats.h"
 #include <oodle2.h>
 
diff --git a/Engine/Source/Programs/UnrealBuildAccelerator/Common/Public/UbaStorageProxy.h b/Engine/Source/Programs/UnrealBuildAccelerator/Common/Public/UbaStorageProxy.h
--- a/Engine/Source/Programs/UnrealBuildAccelerator/Common/Public/UbaStorageProxy.h
+++ b/Engine/Source/Programs/UnrealBuildAccelerator/Common/Public/UbaStorageProxy.h
@@ -3,6 +3,7 @@
 #pragma once
 
 #include "UbaFileMapping.h"
+#include "UbaLogger.h"
 #include "UbaNetwork.h"
 
 namespace uba
diff --git a/Engine/Source/Programs/UnrealBuildAccelerator/Common/Public/UbaTrace.h b/Engine/Source/Programs/UnrealBuildAccelerator/Common/Public/UbaTrace.h
--- a/Engine/Source/Programs/UnrealBuildAccelerator/Common/Public/UbaTrace.h
+++ b/Engine/Source/Programs/UnrealBuildAccelerator/Common/Public/UbaTrace.h
@@ -2,8 +2,9 @@
 
 #pragma once
 
-#include "UbaProcessStats.h"
 #include "UbaFileMapping.h"
+#include "UbaLogger.h"
+#include "UbaProcessStats.h"
 
 namespace uba
 {
diff --git a/Engine/Source/Programs/UnrealBuildAccelerator/Core/Private/UbaPlatform.cpp b/Engine/Source/Programs/UnrealBuildAccelerator/Core/Private/UbaPlatform.cpp
--- a/Engine/Source/Programs/UnrealBuildAccelerator/Core/Private/UbaPlatform.cpp
+++ b/Engine/Source/Programs/UnrealBuildAccelerator/Core/Private/UbaPlatform.cpp
@@ -166,16 +166,17 @@ namespace uba
 				for (u32 i = 0; i < count; i++)
 				{
 					tchar str[1024];
+					auto addr = u64(callers[i]);
 
-					auto findIt = moduleEndAddresses.lower_bound(u64(callers[i]));
-					if (findIt != moduleEndAddresses.end())
+					auto findIt = moduleEndAddresses.lower_bound(addr);
+					if (findIt != moduleEndAddresses.end() && addr >= findIt->second.start)
 					{
 						if (GetModuleFileNameW(findIt->second.handle, str, sizeof_array(str)))
 						{
 							const tchar* moduleName = str;
 							if (const tchar* lastSeparator = TStrrchr(str, PathSeparator))
 								moduleName = lastSeparator + 1;
-							out.Appendf(L"\n   %s: +0x%llx", moduleName, u64(callers[i]) - findIt->second.start);
+							out.Appendf(L"\n   %s: +0x%llx", moduleName, addr - findIt->second.start);
 						}
 					}
 					else
@@ -344,7 +345,17 @@ namespace uba
 	void ElevateCurrentThreadPriority()
 	{
 		#if PLATFORM_WINDOWS
-		SetThreadPriority(GetCurrentThread(), THREAD_PRIORITY_ABOVE_NORMAL);
+		SetThreadPriority(GetCurrentThread(), THREAD_PRIORITY_HIGHEST);
+		#endif
+	}
+
+	void PrefetchVirtualMemory(const void* mem, u64 size)
+	{
+		#if PLATFORM_WINDOWS
+		WIN32_MEMORY_RANGE_ENTRY entry;
+		entry.VirtualAddress = const_cast<void*>(mem);
+		entry.NumberOfBytes = size;
+		PrefetchVirtualMemory(GetCurrentProcess(), 1, &entry, 0);
 		#endif
 	}
 
diff --git a/Engine/Source/Programs/UnrealBuildAccelerator/Core/Public/UbaPlatform.h b/Engine/Source/Programs/UnrealBuildAccelerator/Core/Public/UbaPlatform.h
--- a/Engine/Source/Programs/UnrealBuildAccelerator/Core/Public/UbaPlatform.h
+++ b/Engine/Source/Programs/UnrealBuildAccelerator/Core/Public/UbaPlatform.h
@@ -101,6 +101,7 @@ namespace uba
 	u32 GetProcessorGroupCount();
 
 	void ElevateCurrentThreadPriority();
+	void PrefetchVirtualMemory(const void* mem, u64 size);
 
 #if PLATFORM_WINDOWS
 	inline constexpr bool CaseInsensitiveFs = true;
diff --git a/Engine/Source/Programs/UnrealBuildAccelerator/Core/Public/UbaSynchronization.h b/Engine/Source/Programs/UnrealBuildAccelerator/Core/Public/UbaSynchronization.h
--- a/Engine/Source/Programs/UnrealBuildAccelerator/Core/Public/UbaSynchronization.h
+++ b/Engine/Source/Programs/UnrealBuildAccelerator/Core/Public/UbaSynchronization.h
@@ -110,6 +110,7 @@ namespace uba
 	public:
 		ScopedReadLock(ReaderWriterLock& lock) : m_lock(lock) { lock.EnterRead(); }
 		~ScopedReadLock() { Leave(); }
+		inline void Enter() { if (m_active) return; m_active = true; m_lock.EnterRead(); }
 		inline void Leave() { if (!m_active) return; m_active = false; m_lock.LeaveRead(); }
 
 		ReaderWriterLock& m_lock;
diff --git a/Engine/Source/Programs/UnrealBuildAccelerator/Visualizer/Private/UbaVisualizer.cpp b/Engine/Source/Programs/UnrealBuildAccelerator/Visualizer/Private/UbaVisualizer.cpp
--- a/Engine/Source/Programs/UnrealBuildAccelerator/Visualizer/Private/UbaVisualizer.cpp
+++ b/Engine/Source/Programs/UnrealBuildAccelerator/Visualizer/Private/UbaVisualizer.cpp
@@ -18,6 +18,7 @@ enum
 	Popup_Pause,
 	Popup_Play,
 	Popup_JumpToEnd,
+	Popup_ShowText,
 	Popup_SaveAs,
 	Popup_Quit,
 };
@@ -492,7 +493,7 @@ namespace uba
 		progressRect.left += ProgressRectLeft;
 		progressRect.bottom -= 30;
 
-		bool shouldDrawText = m_zoomValue > 0.22f;
+		bool shouldDrawText = m_showText && m_zoomValue > 0.27f;
 
 		SetBkMode(hdc, TRANSPARENT);
 		SetTextColor(hdc, m_textColor);
@@ -544,6 +545,7 @@ namespace uba
 			posY += 4;
 		}
 
+		TraceView::WorkRecord selectedWork;
 
 		struct SessionRec { TraceView::Session* session; u32 index; };
 		SessionRec sortedSessions[256];
@@ -564,8 +566,6 @@ namespace uba
 			if (!isFirst)
 				posY += 3;
 
-			bool isRemote = processLocation.sessionIndex != 0;
-
 			if (posY + stepY >= progressRect.top && posY <= progressRect.bottom)
 			{
 				SelectObject(hdc, m_separatorPen);
@@ -661,14 +661,10 @@ namespace uba
 						{
 							if (m_visibleComponents[ComponentType_SendRecv] && updateSend != 0 && updateRecv != 0)
 							{
-								auto sendPen = m_sendPen;
-								auto recvPen = m_recvPen;
-								if (!isRemote)
-									std::swap(sendPen, recvPen);
-								SelectObject(hdc, sendPen);
+								SelectObject(hdc, m_sendPen);
 								MoveToEx(hdc, prevX, prevSendY, NULL);
 								LineTo(hdc, x, sendY);
-								SelectObject(hdc, recvPen);
+								SelectObject(hdc, m_recvPen);
 								MoveToEx(hdc, prevX, prevRecvY, NULL);
 								LineTo(hdc, x, recvY);
 							}
@@ -707,6 +703,7 @@ namespace uba
 						if (selected)
 							SetBkMode(hdc, TRANSPARENT);
 					};
+				bool isRemote = processLocation.sessionIndex != 0;
 				PaintDetailedStats(posY, progressRect, session, isRemote, playTime, drawText);
 			}
 
@@ -884,7 +881,7 @@ namespace uba
 
 			if (m_visibleComponents[ComponentType_Workers] && isFirst)
 			{
-				int trackIndex = 0;
+				u32 trackIndex = 0;
 				for (auto& workTrack : m_traceView.workTracks)
 				{
 					if (posY + SessionStepY >= progressRect.top && posY <= progressRect.bottom)
@@ -915,29 +912,29 @@ namespace uba
 							ExtTextOutW(hdc, 5, posY + offsetY, ETO_CLIPPED, &rect, buf.data, buf.count, NULL);
 						}
 
-						u32 recordIndex = 0;
+						u32 workIndex = 0;
 						int posX = int(m_scrollPosX) + progressRect.left;
-						for (auto& record : workTrack.records)
+						for (auto& work : workTrack.records)
 						{
-							if (record.start == record.stop)
+							if (work.start == work.stop)
 							{
-								++recordIndex;
+								++workIndex;
 								continue;
 							}
 
-							float startTime = TimeToS(record.start);
+							float startTime = TimeToS(work.start);
 
 							int left = int(posX + startTime * scaleX);
 
 							if (left >= progressRect.right)
 							{
-								++recordIndex;
+								++workIndex;
 								continue;
 							}
 
 							HBRUSH brush;
 
-							u64 stop = record.stop;
+							u64 stop = work.stop;
 
 							bool done = stop != ~u64(0);
 							if (done)
@@ -953,7 +950,7 @@ namespace uba
 							float stopTime = TimeToS(stop);
 							if ((stopTime - startTime) * scaleX < 0.05f)
 							{
-								++recordIndex;
+								++workIndex;
 								continue;
 							}
 
@@ -966,12 +963,16 @@ namespace uba
 
 							if (rect.right <= progressRect.left)
 							{
-								++recordIndex;
+								++workIndex;
 								continue;
 							}
 
 							rect.right = Max(int(rect.right), left + 1);
 
+							bool selected = m_workSelected && m_workTrack == trackIndex && m_workIndex == workIndex;
+							if (selected)
+								selectedWork = work;
+
 							if (lastSelectedBrush != brush)
 							{
 								SelectObject(hdc, brush);
@@ -989,7 +990,7 @@ namespace uba
 							int processWidth = rect.right - rect.left;
 							if (shouldDrawText && processWidth > 3)
 							{
-								if (!record.bitmap)
+								if (!work.bitmap)
 								{
 									if (m_lastBitmapOffset == BitmapCacheHeight)
 									{
@@ -1004,20 +1005,20 @@ namespace uba
 									RECT rect2{ 0,m_lastBitmapOffset,256, m_lastBitmapOffset + FontHeight };
 
 									FillRect(textDC, &rect2, m_workBrush);
-									ExtTextOutW(textDC, rect2.left, rect2.top, ETO_CLIPPED, &rect2, record.description, int(wcslen(record.description)), NULL);
-									record.bitmap = m_lastBitmap;
-									record.bitmapOffset = m_lastBitmapOffset;
+									ExtTextOutW(textDC, rect2.left, rect2.top, ETO_CLIPPED, &rect2, work.description, int(wcslen(work.description)), NULL);
+									work.bitmap = m_lastBitmap;
+									work.bitmapOffset = m_lastBitmapOffset;
 									m_lastBitmapOffset += FontHeight;
 								}
 
-								if (lastSelectedBitmap != record.bitmap)
+								if (lastSelectedBitmap != work.bitmap)
 								{
-									SelectObject(textDC, record.bitmap);
-									lastSelectedBitmap = record.bitmap;
+									SelectObject(textDC, work.bitmap);
+									lastSelectedBitmap = work.bitmap;
 								}
 
 								int width = Min(processWidth, 256);
-								int bitmapOffsetY = record.bitmapOffset;
+								int bitmapOffsetY = work.bitmapOffset;
 								int bltOffsetY = offsetY;
 								if (bltOffsetY < 0)
 								{
@@ -1042,7 +1043,7 @@ namespace uba
 									BitBlt(hdc, rect.left, rect.top + bltOffsetY, width, height, textDC, bitmapOffsetX, bitmapOffsetY, SRCCOPY);
 								}
 							}
-							++recordIndex;
+							++workIndex;
 						}
 					}
 					++trackIndex;
@@ -1229,6 +1230,34 @@ namespace uba
 				}
 			}
 		}
+		else if (m_workSelected && selectedWork.description)
+		{
+			int width = 290;
+			int height = 3*m_popupFontHeight;
+			POINT p;
+			GetClientCursorPos(p);
+			RECT r;
+			r.left = p.x;
+			r.top = p.y;
+			r.right = r.left + width;
+			r.bottom = r.top + height;
+
+			u64 duration;
+			if (selectedWork.stop != ~u64(0))
+				duration = selectedWork.stop - selectedWork.start;
+			else
+				duration = playTime - selectedWork.start;
+
+			r.bottom += 10;
+			FillRect(hdc, &r, m_tooltipBackgroundBrush);
+
+			r.top += 5;
+			SelectObject(hdc, m_popupFont);
+			DrawTextLogger logger(hdc, r, m_popupFontHeight);
+			logger.Info(L"  %ls", selectedWork.description);
+			logger.Info(L"  Start:     %ls", TimeToText(selectedWork.start, true).str);
+			logger.Info(L"  Duration:  %ls", TimeToText(duration, true).str);
+		}
 		else if (m_sessionSelectedIndex != ~0u)
 		{
 			int width = 290;
@@ -1354,10 +1383,53 @@ namespace uba
 		{
 			int posX = int(m_scrollPosX) + progressRect.left;
 			int left = int(posX + m_timelineSelected * scaleX);
+			int timelineTop = Min(posY, int(progressRect.bottom)) + 2;
 
 			// TODO: Draw line up
 			MoveToEx(hdc, left, 2, NULL);
-			LineTo(hdc, left + 1, progressRect.bottom);
+			LineTo(hdc, left, timelineTop);
+
+			if (m_timelineSelected >= 0)
+			{
+				StringBuffer<> b;
+				u32 milliseconds = u32(m_timelineSelected * 1000.0f);
+				u32 seconds = milliseconds / 1000;
+				milliseconds -= seconds * 1000;
+				u32 minutes = seconds / 60;
+				seconds -= minutes * 60;
+				u32 hours = minutes / 60;
+				minutes -= hours * 60;
+				if (hours)
+				{
+					b.AppendValue(hours).Append('h');
+					if (hours < 10)
+						b.Append('0');
+				}
+				if (minutes || hours)
+				{
+					b.AppendValue(minutes).Append('m');
+					if (seconds < 10)
+						b.Append('0');
+				}
+				b.AppendValue(seconds).Append('.');
+				if (milliseconds < 100)
+					b.Append('0');
+				if (milliseconds < 10)
+					b.Append('0');
+				b.AppendValue(milliseconds);
+
+				RECT r;
+				r.left = left + 4;
+				r.top = timelineTop - 20;
+				r.right = left + b.count*8;
+				r.bottom = r.top + 15;
+
+				FillRect(hdc, &r, m_tooltipBackgroundBrush);
+				SelectObject(hdc, m_popupFont);
+				DrawTextLogger logger(hdc, r, m_popupFontHeight);
+
+				logger.Info(L"%s", b.data);
+			}
 		}
 		else if (m_fetchedFilesSelected != ~0u)
 		{
@@ -1411,6 +1483,8 @@ namespace uba
 		}
 	}
 
+	u64 ConvertTime(const TraceView& view, u64 time);
+
 	void Visualizer::PaintProcessRect(TraceView::Process& process, HDC hdc, RECT rect, const RECT& progressRect, bool selected, bool writingBitmap)
 	{
 		auto clampRect = [&](RECT& r) { r.left = Min(Max(r.left, progressRect.left), progressRect.right); r.right = Max(Min(r.right, progressRect.right), progressRect.left); };
@@ -1438,7 +1512,7 @@ namespace uba
 		RECT main = rect;
 		int width = rect.right - rect.left;
 
-		double recvPart = (double(process.processStats.createFile.time) / duration);
+		double recvPart = (double(ConvertTime(m_traceView, process.processStats.createFile.time)) / duration);
 		if (int headSize = int(recvPart * width))
 		{
 			UBA_ASSERT(headSize > 0);
@@ -1450,7 +1524,7 @@ namespace uba
 				FillRect(hdc, &r2, m_processBrushes[selected].recv);
 		}
 
-		double sendPart = (double(process.processStats.sendFiles.time) / duration);
+		double sendPart = (double(ConvertTime(m_traceView, process.processStats.sendFiles.time)) / duration);
 		if (int tailSize = int(sendPart * width))
 		{
 			UBA_ASSERT(tailSize > 0);
@@ -1688,8 +1762,8 @@ namespace uba
 						sendPerS = u64((u.send - session.prevSend) / duration);
 						recvPerS = u64((u.recv - session.prevRecv) / duration);
 					}
-					drawText(L"Recv: %ls (%s/s)", BytesToText(u.send), BytesToText(sendPerS));
-					drawText(L"Send: %ls (%s/s)", BytesToText(u.recv), BytesToText(recvPerS));
+					drawText(L"Recv: %ls (%s/s)", BytesToText(u.recv), BytesToText(recvPerS));
+					drawText(L"Send: %ls (%s/s)", BytesToText(u.send), BytesToText(sendPerS));
 				}
 			}
 		}
@@ -1764,8 +1838,6 @@ namespace uba
 			if (!isFirst)
 				posY += 3;
 
-			bool isRemote = sessionIndex != 0;
-
 			if (pos.y >= posY && pos.y < posY + SessionStepY)
 			{
 				if (pos.x < 500)
@@ -1803,14 +1875,8 @@ namespace uba
 						if (pos.x + hitOffset >= prevX && pos.x + hitOffset <= x)
 						{
 							double duration = TimeToS(update.time - prevTime);
-
-							auto recvBytesPerSecond = u64((update.send - prevSend) / duration);
-							auto sendBytesPerSecond = u64((update.recv - prevRecv) / duration);
-							if (isRemote)
-								std::swap(recvBytesPerSecond, sendBytesPerSecond);
-
-							outResult.stats.recvBytesPerSecond = recvBytesPerSecond;
-							outResult.stats.sendBytesPerSecond = sendBytesPerSecond;
+							outResult.stats.recvBytesPerSecond = u64((update.recv - prevRecv) / duration);
+							outResult.stats.sendBytesPerSecond = u64((update.send - prevSend) / duration);
 							outResult.stats.ping = update.ping;
 							outResult.stats.memAvail = update.memAvail;
 							outResult.stats.cpuLoad = update.cpuLoad;
@@ -1907,11 +1973,63 @@ namespace uba
 			}
 
 			if (m_visibleComponents[ComponentType_Workers] && isFirst)
-				posY += u32(m_traceView.workTracks.size()) * stepY;
+			{
+				int trackIndex = 0;
+				for (auto& workTrack : m_traceView.workTracks)
+				{
+					if (pos.y < progressRect.bottom && posY + stepY >= progressRect.top && posY <= progressRect.bottom && pos.y >= posY-1 && pos.y < posY-1 + stepY)
+					{
+						u32 workIndex = 0;
+						int posX = int(m_scrollPosX) + progressRect.left;
+						for (auto& work : workTrack.records)
+						{
+							int left = int(posX + TimeToS(work.start) * scaleX);
+
+							if (left >= progressRect.right)
+							{
+								++workIndex;
+								continue;
+							}
+							if (left < progressRect.left)
+								left = progressRect.left;
+
+							u64 stopTime = work.stop;
+							bool done = stopTime != ~u64(0);
+							if (!done)
+								stopTime = playTime;
+
+							RECT rect;
+							rect.left = left;
+							rect.right = int(posX + TimeToS(stopTime) * scaleX);
+
+							if (rect.right <= progressRect.left)
+							{
+								++workIndex;
+								continue;
+							}
+							rect.right = Max(int(rect.right), left + 1);
+							rect.top = posY;
+							rect.bottom = posY + int(float(18) * m_zoomValue);
+
+							if (pos.x >= rect.left && pos.x <= rect.right)
+							{
+								outResult.workTrack = trackIndex;
+								outResult.workIndex = workIndex;
+								outResult.workSelected = true;
+								return;
+							}
+							++workIndex;
+						}
+					}
+					++trackIndex;
+					posY += stepY;
+				}
+			}
 
 			if (m_visibleComponents[ComponentType_Timeline])
 			{
-				if (pos.y >= progressRect.bottom)
+				int timelineTop = Min(posY, int(progressRect.bottom));
+				if (pos.y >= timelineTop && pos.y < timelineTop + 40)
 				{
 					float timeScale = (m_horizontalScaleValue * m_zoomValue)*50.0f;
 					float startOffset = -(m_scrollPosX / timeScale);
@@ -1975,7 +2093,7 @@ namespace uba
 
 	void Visualizer::UnselectAndRedraw()
 	{
-		if (m_processSelected || m_sessionSelectedIndex != ~0u || m_statsSelected || m_timelineSelected || m_fetchedFilesSelected != ~0u)
+		if (m_processSelected || m_sessionSelectedIndex != ~0u || m_statsSelected || m_timelineSelected || m_fetchedFilesSelected != ~0u || m_workSelected)
 		{
 			m_processSelected = false;
 			m_sessionSelectedIndex = ~0u;
@@ -1983,6 +2101,7 @@ namespace uba
 			m_buttonSelected = ~0u;
 			m_timelineSelected = 0;
 			m_fetchedFilesSelected = ~0u;
+			m_workSelected = false;
 			RedrawWindow(m_hwnd, NULL, NULL, RDW_INVALIDATE);
 		}
 	}
@@ -2018,7 +2137,8 @@ namespace uba
 			res.sessionSelectedIndex == m_sessionSelectedIndex &&
 			res.statsSelected == m_statsSelected && memcmp(&res.stats, &m_stats, sizeof(Stats)) == 0 &&
 			res.buttonSelected == m_buttonSelected && res.timelineSelected == m_timelineSelected &&
-			res.fetchedFilesSelected == m_fetchedFilesSelected)
+			res.fetchedFilesSelected == m_fetchedFilesSelected &&
+			res.workSelected == m_workSelected && res.workTrack == m_workTrack && res.workIndex == m_workIndex)
 			return false;
 		m_processSelected = res.processSelected;
 		m_processSelectedLocation = res.processLocation;
@@ -2028,6 +2148,9 @@ namespace uba
 		m_buttonSelected = res.buttonSelected;
 		m_timelineSelected = res.timelineSelected;
 		m_fetchedFilesSelected = res.fetchedFilesSelected;
+		m_workSelected = res.workSelected;
+		m_workTrack = res.workTrack;
+		m_workIndex = res.workIndex;
 		return true;
 	}
 
@@ -2295,6 +2418,7 @@ namespace uba
 			m_buttonSelected = ~0u;
 			m_timelineSelected = 0;
 			m_fetchedFilesSelected = ~0u;
+			m_workSelected = false;
 			m_autoScroll = false;
 			POINTS p = MAKEPOINTS(lParam);
 			m_mouseAnchor = {p.x, p.y};
@@ -2363,7 +2487,9 @@ namespace uba
 			HMENU hMenu = CreatePopupMenu();
 			ClientToScreen(hWnd, &point);
 
-			//AppendMenuW(hMenu, MF_STRING, 1, L"&New");
+			AppendMenuW(hMenu, MF_STRING, Popup_ShowText, m_showText ? L"&Hide process text" : L"&Show process text");
+			AppendMenuW(hMenu, MF_SEPARATOR, 0, NULL);
+
 			if (m_sessionSelectedIndex != ~0u)
 			{
 				AppendMenuW(hMenu, MF_STRING, Popup_CopySessionInfo, L"&Copy Session Info");
@@ -2392,7 +2518,6 @@ namespace uba
 				}
 			}
 
-
 			AppendMenuW(hMenu, MF_STRING, Popup_SaveAs, L"&Save Trace");
 			AppendMenuW(hMenu, MF_SEPARATOR, 0, NULL);
 			AppendMenuW(hMenu, MF_STRING, Popup_Quit, L"&Quit");
@@ -2421,6 +2546,10 @@ namespace uba
 					m_trace.SaveAs(ofn.lpstrFile);
 				break;
 			}
+			case Popup_ShowText:
+				m_showText = !m_showText;
+				RedrawWindow(hWnd, NULL, NULL, RDW_INVALIDATE | RDW_UPDATENOW);
+				break;
 			case Popup_Replay:
 				m_replay = 1;
 				PostMessage(m_hwnd, WM_NEWTRACE, 0, 0);
diff --git a/Engine/Source/Programs/UnrealBuildAccelerator/Visualizer/Private/UbaVisualizer.h b/Engine/Source/Programs/UnrealBuildAccelerator/Visualizer/Private/UbaVisualizer.h
--- a/Engine/Source/Programs/UnrealBuildAccelerator/Visualizer/Private/UbaVisualizer.h
+++ b/Engine/Source/Programs/UnrealBuildAccelerator/Visualizer/Private/UbaVisualizer.h
@@ -52,7 +52,9 @@ namespace uba
 			u32 buttonSelected = ~0u;
 			float timelineSelected = 0;
 			u32 fetchedFilesSelected = ~0u;
-
+			bool workSelected = false;
+			u32 workTrack = ~0u;
+			u32 workIndex = ~0u;
 		};
 		void HitTest(HitTestResult& outResult, const POINT& pos);
 
@@ -109,6 +111,7 @@ namespace uba
 		int m_popupFontHeight = 0;
 		bool m_useDarkMode = true;
 		bool m_isThemeSet = false;
+		bool m_showText = true;
 
 		Logger& m_logger;
 		NetworkClient* m_client = nullptr;
@@ -143,6 +146,10 @@ namespace uba
 		u32 m_buttonSelected = ~0u;
 		float m_timelineSelected = 0;
 		u32 m_fetchedFilesSelected = ~0u;
+
+		bool m_workSelected = false;
+		u32 m_workTrack = ~0u;
+		u32 m_workIndex = ~0u;
 		
 		bool m_mouseOverWindow = false;
 		bool m_showPopup = false;
