[UBA]
* Extracted work tracker out from work manager so it can be used for multiple work managers when running with proxies

diff --git a/Engine/Source/Programs/UnrealBuildAccelerator/Common/Private/UbaNetworkServer.cpp b/Engine/Source/Programs/UnrealBuildAccelerator/Common/Private/UbaNetworkServer.cpp
--- a/Engine/Source/Programs/UnrealBuildAccelerator/Common/Private/UbaNetworkServer.cpp
+++ b/Engine/Source/Programs/UnrealBuildAccelerator/Common/Private/UbaNetworkServer.cpp
@@ -154,7 +154,7 @@ namespace uba
 			}
 
 			if (!m_disconnectCallbackCalled.IsSet(60000)) // This should never time out!
-				m_server.m_logger.Warning(TC("This should never happen!! Unknown conseqences"));
+				m_server.m_logger.Warning(TC("This should never happen!! Unknown consequences"));
 			return true;
 		}
 
@@ -516,7 +516,7 @@ namespace uba
 				SCOPED_WRITE_LOCK(server.m_additionalWorkLock, lock);
 				if (server.m_additionalWork.empty())
 					break;
-				work = server.m_additionalWork.front();
+				work = std::move(server.m_additionalWork.front());
 				server.m_additionalWork.pop_front();
 				lock.Leave();
 
@@ -592,7 +592,7 @@ namespace uba
 
 	bool NetworkServer::StartListen(NetworkBackend& backend, u16 port, const tchar* ip, const u8* cryptoKey128)
 	{
-		UBA_ASSERT(!m_listenBackend);
+		//UBA_ASSERT(!m_listenBackend);
 		m_listenBackend = &backend;
 
 		if (cryptoKey128)
@@ -798,7 +798,7 @@ namespace uba
 		for (u32 i = 0; i != count; ++i)
 		{
 			m_additionalWork.push_back({ work });
-			if (m_trackWork)
+			if (m_workTracker)
 				m_additionalWork.back().desc = desc;
 		}
 		lock.Leave();
@@ -820,21 +820,6 @@ namespace uba
 		return m_maxWorkerCount;
 	}
 
-	u32 NetworkServer::TrackWorkStart(const tchar* desc)
-	{
-		if (!m_trackWork)
-			return 0;
-		u32 workId = m_workCounter++;
-		m_startWork(workId, desc);
-		return workId;
-	}
-
-	void NetworkServer::TrackWorkEnd(u32 id)
-	{
-		if (m_trackWork)
-			m_endWork(id);
-	}
-
 	u64 NetworkServer::GetTotalSentBytes()
 	{
 		return m_sendBytes;
@@ -845,6 +830,16 @@ namespace uba
 		return m_recvBytes;
 	}
 
+	u32 NetworkServer::GetConnectionCount()
+	{
+		SCOPED_READ_LOCK(m_connectionsLock, lock);
+		u32 count = 0;
+		for (auto& con : m_connections)
+			if (!con.m_disconnected)
+				++count;
+		return count;
+	}
+
 	void NetworkServer::GetClientStats(ClientStats& out, u32 clientId)
 	{
 		SCOPED_READ_LOCK(m_clientsLock, lock);
@@ -906,32 +901,17 @@ namespace uba
 			worker->m_context = oldContext;
 			return true;
 		}
-		work = m_additionalWork.front();
+		work = std::move(m_additionalWork.front());
 		m_additionalWork.pop_front();
 		lock.Leave();
 
-		u32 workId = TrackWorkStart(work.desc.c_str());
+		TrackWorkScope tws(*this, work.desc.c_str());
 
 		work.func();
 
-		TrackWorkEnd(workId);
-
 		return true;
 	}
 
-	void NetworkServer::SetWorkListener(const WorkBeginFunction& start, const WorkEndFunction& end)
-	{
-		m_startWork = start;
-		m_endWork = end;
-		m_trackWork = true;
-	}
-
-	void NetworkServer::ResetWorkListener()
-	{
-		m_trackWork = false;
-		m_startWork = {};
-		m_endWork = {};
-	}
 
 	NetworkServer::Worker* NetworkServer::PopWorker()
 	{
diff --git a/Engine/Source/Programs/UnrealBuildAccelerator/Common/Private/UbaSessionServer.cpp b/Engine/Source/Programs/UnrealBuildAccelerator/Common/Private/UbaSessionServer.cpp
--- a/Engine/Source/Programs/UnrealBuildAccelerator/Common/Private/UbaSessionServer.cpp
+++ b/Engine/Source/Programs/UnrealBuildAccelerator/Common/Private/UbaSessionServer.cpp
@@ -126,17 +126,7 @@ namespace uba
 		m_storage.SetTrace(&m_trace, m_detailedTrace);
 
 		if (m_detailedTrace)
-		{
-			m_server.SetWorkListener(
-				[this](u32 workIndex, const tchar* desc)
-				{
-					m_trace.BeginWork(workIndex, desc);
-				},
-				[this](u32 workIndex)
-				{
-					m_trace.EndWork(workIndex);
-				});
-		}
+			m_server.SetWorkTracker(&m_trace);
 
 		m_memoryThreadEvent.Create(true);
 		if (info.checkMemory)
@@ -188,7 +178,7 @@ namespace uba
 
 		StopTraceThread();
 
-		m_server.ResetWorkListener();
+		m_server.SetWorkTracker(nullptr);
 		m_server.UnregisterOnClientDisconnected(ServiceId);
 		m_server.UnregisterService(ServiceId);
 
@@ -773,10 +763,13 @@ namespace uba
 				UBA_ASSERT(attributes);
 				CasKey casKey = reader.ReadCasKey();
 				Storage::RetrieveResult res;
-				bool success = m_storage.RetrieveCasFile(res, casKey, TC(""));
+				bool success = m_storage.RetrieveCasFile(res, casKey, destination.data);
 				casKey = res.casKey;
 				if (!success)
-					m_logger.Error(TC("Failed to retrieve cas for %s from client (Needed to write %s)"), CasKeyString(casKey).str, destination);
+				{
+					auto logType = connectionInfo.ShouldDisconnect() ? LogEntryType_Info : LogEntryType_Warning;
+					m_logger.Logf(logType, TC("Failed to retrieve cas for %s from client (Needed to write %s)"), CasKeyString(casKey).str, destination.data);
+				}
 				if (success)
 				{
 					if (destination.StartsWith(TC("<log>")))
@@ -1604,8 +1597,8 @@ namespace uba
 		lock.Leave();
 
 		float cpuLoad = UpdateCpuLoad();
-		u64 serverSend = m_server.GetTotalRecvBytes();
-		u64 serverRecv = m_server.GetTotalSentBytes();
+		u64 serverSend = m_server.GetTotalSentBytes();
+		u64 serverRecv = m_server.GetTotalRecvBytes();
 		u64 memAvail = m_memAvail;
 		u64 memTotal = m_memTotal;
 		m_trace.SessionUpdate(0, 0, serverSend, serverRecv, 0, memAvail, memTotal, cpuLoad);
diff --git a/Engine/Source/Programs/UnrealBuildAccelerator/Common/Private/UbaTrace.cpp b/Engine/Source/Programs/UnrealBuildAccelerator/Common/Private/UbaTrace.cpp
--- a/Engine/Source/Programs/UnrealBuildAccelerator/Common/Private/UbaTrace.cpp
+++ b/Engine/Source/Programs/UnrealBuildAccelerator/Common/Private/UbaTrace.cpp
@@ -86,7 +86,7 @@ namespace uba
 		}
 
 
-		if (m_channel.Init())
+		if (namedTrace && m_channel.Init())
 			m_channel.Write(namedTrace);
 
 		return true;
@@ -123,6 +123,18 @@ namespace uba
 		return true;
 	}
 
+	u32 Trace::TrackWorkStart(const tchar* desc)
+	{
+		u32 workId = m_workCounter++;
+		BeginWork(workId, desc);
+		return workId;
+	}
+
+	void Trace::TrackWorkEnd(u32 id)
+	{
+		EndWork(id);
+	}
+
 	u32 Trace::AddString(const tchar* string)
 	{
 		if (!m_memoryBegin)
diff --git a/Engine/Source/Programs/UnrealBuildAccelerator/Common/Private/UbaWorkManager.cpp b/Engine/Source/Programs/UnrealBuildAccelerator/Common/Private/UbaWorkManager.cpp
--- a/Engine/Source/Programs/UnrealBuildAccelerator/Common/Private/UbaWorkManager.cpp
+++ b/Engine/Source/Programs/UnrealBuildAccelerator/Common/Private/UbaWorkManager.cpp
@@ -7,6 +7,20 @@
 
 namespace uba
 {
+	u32 WorkManager::TrackWorkStart(const tchar* desc)
+	{
+		if (auto t = m_workTracker.load())
+			return t->TrackWorkStart(desc);
+		return 0;
+	}
+
+	void WorkManager::TrackWorkEnd(u32 id)
+	{
+		if (auto t = m_workTracker.load())
+			return t->TrackWorkEnd(id);
+	}
+
+
 	struct WorkManagerImpl::Worker
 	{
 		Worker(WorkManagerImpl& manager) : m_workAvailable(false)
@@ -46,6 +60,7 @@ namespace uba
 						manager.m_work.pop_front();
 						lock.Leave();
 
+						TrackWorkScope tws(manager, work.desc.c_str());
 						work.func();
 					}
 
@@ -88,7 +103,11 @@ namespace uba
 	{
 		SCOPED_WRITE_LOCK(m_workLock, lock);
 		for (u32 i = 0; i != count; ++i)
+		{
 			m_work.push_back({ work });
+			if (m_workTracker.load())
+				m_work.back().desc = desc;
+		}
 		lock.Leave();
 
 		SCOPED_WRITE_LOCK(m_availableWorkersLock, lock2);
@@ -133,6 +152,7 @@ namespace uba
 			m_work.pop_front();
 			lock.Leave();
 
+			TrackWorkScope tws(*this, work.desc.c_str());
 			work.func();
 		}
 	}
diff --git a/Engine/Source/Programs/UnrealBuildAccelerator/Common/Public/UbaNetworkServer.h b/Engine/Source/Programs/UnrealBuildAccelerator/Common/Public/UbaNetworkServer.h
--- a/Engine/Source/Programs/UnrealBuildAccelerator/Common/Public/UbaNetworkServer.h
+++ b/Engine/Source/Programs/UnrealBuildAccelerator/Common/Public/UbaNetworkServer.h
@@ -68,8 +68,6 @@ namespace uba
 
 		virtual void AddWork(const Function<void()>& work, u32 count, const tchar* desc) override final;
 		virtual u32 GetWorkerCount() override final;
-		virtual u32 TrackWorkStart(const tchar* desc) override final;
-		virtual void TrackWorkEnd(u32 id) override final;
 
 		struct ClientStats
 		{
@@ -80,14 +78,11 @@ namespace uba
 
 		u64 GetTotalSentBytes();
 		u64 GetTotalRecvBytes();
+		u32 GetConnectionCount();
 		void GetClientStats(ClientStats& out, u32 clientId);
 
 		bool DoAdditionalWork();
 
-		using WorkBeginFunction = Function<void(u32, const tchar*)>;
-		using WorkEndFunction = Function<void(u32)>;
-		void SetWorkListener(const WorkBeginFunction& start, const WorkEndFunction& end);
-		void ResetWorkListener();
 
 		class Worker;
 	private:
@@ -162,11 +157,6 @@ namespace uba
 		Atomic<u64> m_recvBytes;
 		Atomic<u32> m_recvCount;
 
-		bool m_trackWork = false;
-		WorkBeginFunction m_startWork;
-		WorkEndFunction m_endWork;
-		Atomic<u32> m_workCounter;
-
 		NetworkServer(const NetworkServer&) = delete;
 		void operator=(const NetworkServer&) = delete;
 	};
diff --git a/Engine/Source/Programs/UnrealBuildAccelerator/Common/Public/UbaTrace.h b/Engine/Source/Programs/UnrealBuildAccelerator/Common/Public/UbaTrace.h
--- a/Engine/Source/Programs/UnrealBuildAccelerator/Common/Public/UbaTrace.h
+++ b/Engine/Source/Programs/UnrealBuildAccelerator/Common/Public/UbaTrace.h
@@ -5,6 +5,7 @@
 #include "UbaFileMapping.h"
 #include "UbaLogger.h"
 #include "UbaProcessStats.h"
+#include "UbaWorkManager.h"
 
 namespace uba
 {
@@ -55,7 +56,7 @@ namespace uba
 	static constexpr u32 TraceVersion = 22;
 	static constexpr u32 TraceReadCompatibilityVersion = 6;
 
-	class Trace
+	class Trace : public WorkTracker
 	{
 	public:
 		Trace(LogWriter& logWriter);
@@ -84,6 +85,9 @@ namespace uba
 
 		bool StopWrite(const tchar* writeFileName);
 
+		virtual u32 TrackWorkStart(const tchar* desc) final override;
+		virtual void TrackWorkEnd(u32 id) final override;
+
 	private:
 		struct WriterScope;
 		u32 AddString(const tchar* string);
@@ -101,6 +105,8 @@ namespace uba
 		ReaderWriterLock m_stringsLock;
 		UnorderedMap<StringKey, u32> m_strings;
 
+		Atomic<u32> m_workCounter;
+
 		friend class SessionServer;
 	};
 }
\ No newline at end of file
diff --git a/Engine/Source/Programs/UnrealBuildAccelerator/Common/Public/UbaWorkManager.h b/Engine/Source/Programs/UnrealBuildAccelerator/Common/Public/UbaWorkManager.h
--- a/Engine/Source/Programs/UnrealBuildAccelerator/Common/Public/UbaWorkManager.h
+++ b/Engine/Source/Programs/UnrealBuildAccelerator/Common/Public/UbaWorkManager.h
@@ -6,14 +6,27 @@
 
 namespace uba
 {
+	class WorkTracker
+	{
+	public:
+		virtual u32 TrackWorkStart(const tchar* desc) { return 0; }
+		virtual void TrackWorkEnd(u32 id) {}
+	};
+
 	class WorkManager
 	{
 	public:
 		virtual void AddWork(const Function<void()>& work, u32 count, const tchar* desc) = 0;
 		virtual u32 GetWorkerCount() = 0;
 
-		virtual u32 TrackWorkStart(const tchar* desc) { return 0; }
-		virtual void TrackWorkEnd(u32 id) {}
+		u32 TrackWorkStart(const tchar* desc);
+		void TrackWorkEnd(u32 id);
+
+		void SetWorkTracker(WorkTracker* workTracker) { m_workTracker = workTracker; }
+		WorkTracker* GetWorkTracker() { return m_workTracker; }
+
+	protected:
+		Atomic<WorkTracker*> m_workTracker;
 	};
 
 
@@ -34,10 +47,11 @@ namespace uba
 		Worker* PopWorkerNoLock();
 
 		Vector<Worker*> m_workers;
-		struct Work { Function<void()> func; };
+		struct Work { Function<void()> func; TString desc; };
 		ReaderWriterLock m_workLock;
 		List<Work> m_work;
 		Atomic<u32> m_activeWorkerCount;
+		Atomic<u32> m_workCounter;
 
 		ReaderWriterLock m_availableWorkersLock;
 		Worker* m_firstAvailableWorker = nullptr;
