[UBA]
* Fixed race condition in UbaNetworkClient related to message ids and returning message ids after use
* Implemented network backend that is using memory for communication (will be used between client and proxy server when running inproc)
* Added connection uid provided by network backend to be able to improve error messages
* Added NetworkBackend GetTotalSendAndRecv that can be used to fetch all traffic on backend

diff --git a/Engine/Source/Programs/UnrealBuildAccelerator/Common/Private/UbaNetworkBackendMemory.cpp b/Engine/Source/Programs/UnrealBuildAccelerator/Common/Private/UbaNetworkBackendMemory.cpp
new file mode 100644
--- /dev/null
+++ b/Engine/Source/Programs/UnrealBuildAccelerator/Common/Private/UbaNetworkBackendMemory.cpp
@@ -0,0 +1,154 @@
+// Copyright Epic Games, Inc. All Rights Reserved.
+
+#include "UbaNetworkBackendMemory.h"
+#include "UbaEvent.h"
+#include "UbaHash.h"
+#include "UbaPlatform.h"
+#include "UbaStringBuffer.h"
+#include "UbaTimer.h"
+
+namespace uba
+{
+	struct NetworkBackendMemory::Connection
+	{
+		struct Peer
+		{
+			ReaderWriterLock lock;
+			RecvHeaderCallback* headerCallback = nullptr;
+			RecvBodyCallback* bodyCallback = nullptr;
+			void* context = nullptr;
+			u32 headerSize = 0;
+			const tchar* recvHint = nullptr;
+			
+			DisconnectCallback* disconnectCallback = nullptr;
+			void* disconnectContext = nullptr;
+		};
+
+		Peer peer[2];
+		Guid uid;
+		bool connected = true;
+	};
+
+	NetworkBackendMemory::NetworkBackendMemory(LogWriter& writer, const tchar* prefix)
+	{
+	}
+
+	NetworkBackendMemory::~NetworkBackendMemory()
+	{
+	}
+
+	void NetworkBackendMemory::Shutdown(void* connection)
+	{
+		u64 from = ((uintptr_t)connection)-1;
+		u64 to = from == 0 ? 1 : 0;
+
+		auto& rc = m_connection->peer[to];
+		SCOPED_WRITE_LOCK(rc.lock, l);
+		m_connection->connected = false;
+		if (auto cb = rc.disconnectCallback)
+			cb(rc.disconnectContext, m_connection->uid, connection);
+	}
+
+	bool NetworkBackendMemory::Send(Logger& logger, void* connection, const void* data, u32 dataSize, SendContext& sendContext)
+	{
+		UBA_ASSERT(m_connection);
+		sendContext.isUsed = true;
+		sendContext.isFinished = true;
+
+		u64 from = ((uintptr_t)connection)-1;
+		u64 to = from == 0 ? 1 : 0;
+
+		auto& peer = m_connection->peer[to];
+
+		SCOPED_READ_LOCK(peer.lock, l);
+
+		if (!m_connection->connected)
+			return false;
+
+		void* bodyContext = nullptr;
+		u8* bodyData = nullptr;
+		u32 bodySize = 0;
+
+		if (!peer.headerCallback(peer.context, m_connection->uid, (u8*)data, bodyContext, bodyData, bodySize))
+			return false;
+
+		if (!bodySize)
+			return true;
+
+		memcpy(bodyData, ((const u8*)data) + peer.headerSize, bodySize);
+		auto bc = peer.bodyCallback;
+		if (!bc)
+			return logger.Error(TC("Connection body callback not set"));
+
+		if (!bc(peer.context, false, (u8*)data, bodyContext, bodyData, bodySize))
+			return false;
+		return true;
+	}
+
+	void NetworkBackendMemory::SetDataSentCallback(void* connection, void* context, DataSentCallback* callback)
+	{
+	}
+
+	void NetworkBackendMemory::SetRecvCallbacks(void* connection, void* context, u32 headerSize, RecvHeaderCallback* h, RecvBodyCallback* b, const tchar* recvHint)
+	{
+		UBA_ASSERT(m_connection);
+		auto& peer = m_connection->peer[((uintptr_t)connection)-1];
+		peer.headerCallback = h;
+		peer.bodyCallback = b;
+		peer.context = context;
+		peer.headerSize = headerSize;
+		peer.recvHint = recvHint;
+	}
+
+	void NetworkBackendMemory::SetRecvTimeout(void* connection, u32 timeoutMs)
+	{
+	}
+
+	void NetworkBackendMemory::SetDisconnectCallback(void* connection, void* context, DisconnectCallback* callback)
+	{
+		UBA_ASSERT(m_connection);
+		auto& peer = m_connection->peer[((uintptr_t)connection)-1];
+		peer.disconnectCallback = callback;
+		peer.disconnectContext = context;
+	}
+
+	bool NetworkBackendMemory::StartListen(Logger& logger, u16 port, const tchar* ip, const ListenConnectedFunc& connectedFunc)
+	{
+		m_connectedFunc = connectedFunc;
+		return true;
+	}
+
+	void NetworkBackendMemory::StopListen()
+	{
+	}
+
+	bool NetworkBackendMemory::Connect(Logger& logger, const tchar* ip, const ConnectedFunc& connectedFunc, u16 port, bool* timedOut)
+	{
+		u64 start = GetTime();
+		SCOPED_READ_LOCK(m_connectedFuncLock, l);
+		while (!m_connectedFunc)
+		{
+			if (TimeToMs(GetTime() - start) > 1000*2)
+				return false;
+			l.Leave();
+			Sleep(10);
+			l.Enter();
+		}
+		UBA_ASSERT(!m_connection);
+		m_connection = new Connection();
+		if (!m_connectedFunc((void*)1, {}))
+			return false;
+		return connectedFunc((void*)2, {}, timedOut);
+	}
+
+	bool NetworkBackendMemory::Connect(Logger& logger, const sockaddr& remoteSocketAddr, const ConnectedFunc& connectedFunc, bool* timedOut, const tchar* nameHint)
+	{
+		return false;
+	}
+
+	void NetworkBackendMemory::GetTotalSendAndRecv(u64& outSend, u64& outRecv)
+	{
+		outSend = 0;
+		outRecv = 0;
+	}
+}
\ No newline at end of file
diff --git a/Engine/Source/Programs/UnrealBuildAccelerator/Common/Private/UbaNetworkBackendTcp.cpp b/Engine/Source/Programs/UnrealBuildAccelerator/Common/Private/UbaNetworkBackendTcp.cpp
--- a/Engine/Source/Programs/UnrealBuildAccelerator/Common/Private/UbaNetworkBackendTcp.cpp
+++ b/Engine/Source/Programs/UnrealBuildAccelerator/Common/Private/UbaNetworkBackendTcp.cpp
@@ -48,17 +48,21 @@ namespace uba
 	{
 		StringBuffer<128> ip;
 		u16 port;
-		Thread thread;
+		ListenConnectedFunc connectedFunc;
 		Event listening;
 		Atomic<SOCKET> socket = INVALID_SOCKET;
+		Thread thread;
 	};
 
 	struct NetworkBackendTcp::Connection
 	{
-		Connection(SOCKET s) : socket(s), ready(true) {}
+		Connection(Logger& l, SOCKET s) : logger(l), socket(s), ready(true) { CreateGuid(uid); }
 
+		Logger& logger;
 		Atomic<SOCKET> socket;
+
 		Event ready;
+		Guid uid;
 		u32 headerSize = 0;
 		u32 recvTimeoutMs = 0;
 
@@ -133,7 +137,7 @@ namespace uba
 			shutdown(s, SD_BOTH);
 			lock2.Leave();
 			conn.recvThread.Wait();
-			CloseSocket(m_logger, s);
+			CloseSocket(conn.logger, s);
 		}
 		m_connections.clear();
 
@@ -168,6 +172,9 @@ namespace uba
 		#endif
 
 		sendContext.isFinished = true;
+
+		m_totalSend += dataSize;
+
 		if (auto c = conn.dataSentCallback)
 			c(conn.dataSentContext, dataSize);
 		return res;
@@ -215,12 +222,17 @@ namespace uba
 		if (!EnsureInitialized(logger))
 			return false;
 
+		SCOPED_WRITE_LOCK(m_listenEntriesLock, lock);
+
+		auto prevListenEntryCount = int(m_listenEntries.size());
+
 		auto AddAddr = [&](const tchar* addr)
 			{
 				m_listenEntries.emplace_back();
 				auto& entry = m_listenEntries.back();
 				entry.ip.Append(addr);
 				entry.port = port;
+				entry.connectedFunc = connectedFunc;
 			};
 
 		if (ip && *ip)
@@ -243,10 +255,11 @@ namespace uba
 			return false;
 		}
 
-		m_connectedFunc = connectedFunc;
-
+		auto skipCount = prevListenEntryCount;
 		for (auto& e : m_listenEntries)
 		{
+			if (skipCount-- > 0)
+				continue;
 			e.listening.Create(true);
 			e.thread.Start([this, &logger, &e]
 				{
@@ -256,8 +269,11 @@ namespace uba
 		}
 
 		bool success = true;
+		skipCount = prevListenEntryCount;
 		for (auto& e : m_listenEntries)
 		{
+			if (skipCount-- > 0)
+				continue;
 			if (!e.listening.IsSet(4000))
 				success = false;
 			if (e.socket == INVALID_SOCKET)
@@ -269,6 +285,7 @@ namespace uba
 
 	void NetworkBackendTcp::StopListen()
 	{
+		SCOPED_WRITE_LOCK(m_listenEntriesLock, lock);
 		for (auto& e : m_listenEntries)
 			e.socket = INVALID_SOCKET;
 		for (auto& e : m_listenEntries)
@@ -304,7 +321,7 @@ namespace uba
 		if (listenSocket == INVALID_SOCKET)
 			return logger.Error(TC("socket failed (%s)"), LastErrorToText(WSAGetLastError()).data);
 
-		auto listenSocketCleanup = MakeGuard([&]() { CloseSocket(m_logger, listenSocket); });
+		auto listenSocketCleanup = MakeGuard([&]() { CloseSocket(logger, listenSocket); });
 
 		u32 reuseAddr = 1;
 		if (::setsockopt(listenSocket, SOL_SOCKET, SO_REUSEADDR, (const char*)&reuseAddr, sizeof reuseAddr) == SOCKET_ERROR)
@@ -379,12 +396,12 @@ namespace uba
 			}
 
 			SCOPED_WRITE_LOCK(m_connectionsLock, lock);
-			auto it = m_connections.emplace(m_connections.end(), clientSocket);
+			auto it = m_connections.emplace(m_connections.end(), logger, clientSocket);
 			auto& conn = *it;
 			conn.recvThread.Start([this, connPtr = &conn] { ThreadRecv(*connPtr); return 0; });
 			lock.Leave();
 
-			if (!m_connectedFunc(&conn, remoteSockAddr))
+			if (!entry.connectedFunc(&conn, remoteSockAddr))
 			{
 				shutdown(clientSocket, SD_BOTH);
 				conn.ready.Set();
@@ -402,11 +419,10 @@ namespace uba
 	{
 		ElevateCurrentThreadPriority();
 		
+		auto& logger = connection.logger;
+
 		if (connection.ready.IsSet(60000)) // This should never time out!
 		{
-			Guid connectionUid;
-			CreateGuid(connectionUid);
-
 			bool isFirst = true;
 			while (connection.socket != INVALID_SOCKET)
 			{
@@ -415,27 +431,32 @@ namespace uba
 				u32 bodySize = 0;
 
 				u8 headerData[MaxHeaderSize];
-				if (!RecvSocket(m_logger, connection.socket, headerData, connection.headerSize, connection.recvTimeoutMs, connectionUid, connection.recvHint, TC(""), isFirst))
+				if (!RecvSocket(logger, connection.socket, headerData, connection.headerSize, connection.recvTimeoutMs, connection.uid, connection.recvHint, TC(""), isFirst))
 					break;
 				isFirst = false;
 
+				m_totalRecv += connection.headerSize;
+
 				auto hc = connection.headerCallback;
 				if (!hc)
 				{
-					m_logger.Error(TC("Header callback not set"));
+					logger.Error(TC("Tcp connection header callback not set"));
 					break;
 				}
 
-				if (!hc(connection.recvContext, headerData, bodyContext, bodyData, bodySize))
+				if (!hc(connection.recvContext, connection.uid, headerData, bodyContext, bodyData, bodySize))
 					break;
 				if (!bodySize)
 					continue;
 
-				bool success = RecvSocket(m_logger, connection.socket, bodyData, bodySize, connection.recvTimeoutMs, connectionUid, connection.recvHint, TC("Body"), false);
+				bool success = RecvSocket(logger, connection.socket, bodyData, bodySize, connection.recvTimeoutMs, connection.uid, connection.recvHint, TC("Body"), false);
+
+				m_totalRecv += bodySize;
+
 				auto bc = connection.bodyCallback;
 				if (!bc)
 				{
-					m_logger.Error(TC("Body callback not set"));
+					logger.Error(TC("Tcp connection body callback not set"));
 					break;
 				}
 
@@ -447,7 +468,7 @@ namespace uba
 		}
 		else
 		{
-			m_logger.Warning(TC("Timed out waiting for recv thread to be ready"));
+			logger.Warning(TC("Tcp connection timed out waiting for recv thread to be ready"));
 		}
 
 		ScopedCriticalSection lock2(connection.shutdownLock);
@@ -459,13 +480,13 @@ namespace uba
 			auto context = connection.disconnectContext;
 			connection.disconnectCallback = nullptr;
 			connection.disconnectContext = nullptr;
-			cb(context, &connection);
+			cb(context, connection.uid, &connection);
 		}
 
 		if (s == INVALID_SOCKET)
 			return;
 		shutdown(s, SD_BOTH);
-		CloseSocket(m_logger, s);
+		CloseSocket(logger, s);
 	}
 
 	bool NetworkBackendTcp::Connect(Logger& logger, const tchar* ip, const ConnectedFunc& connectedFunc, u16 port, bool* timedOut)
@@ -619,7 +640,7 @@ namespace uba
 		socketClose.Cancel();
 
 		SCOPED_WRITE_LOCK(m_connectionsLock, lock);
-		auto it = m_connections.emplace(m_connections.end(), socketFd);
+		auto it = m_connections.emplace(m_connections.end(), logger, socketFd);
 		auto& conn = *it;
 		conn.recvThread.Start([this, connPtr = &conn] { ThreadRecv(*connPtr); return 0; });
 		lock.Leave();
@@ -636,13 +657,19 @@ namespace uba
 
 		//char* ip = inet_ntoa(((sockaddr_in*)const_cast<sockaddr*>(&remoteSocketAddr))->sin_addr);
 		if (nameHint)
-			logger.Detail(TC("Connected to %s:%u"), nameHint, ((sockaddr_in&)remoteSocketAddr).sin_port);
+			logger.Detail(TC("Connected to %s:%u (%s)"), nameHint, ((sockaddr_in&)remoteSocketAddr).sin_port, GuidToString(conn.uid).str);
 		else
-			logger.Detail(TC("Connected using sockaddr"));
+			logger.Detail(TC("Connected using sockaddr (%s)"), GuidToString(conn.uid).str);
 
 		return true;
 	}
 
+	void NetworkBackendTcp::GetTotalSendAndRecv(u64& outSend, u64& outRecv)
+	{
+		outSend = m_totalSend;
+		outRecv = m_totalRecv;
+	}
+
 	bool SetBlocking(Logger& logger, SOCKET socket, bool blocking)
 	{
 #if PLATFORM_WINDOWS
diff --git a/Engine/Source/Programs/UnrealBuildAccelerator/Common/Private/UbaNetworkClient.cpp b/Engine/Source/Programs/UnrealBuildAccelerator/Common/Private/UbaNetworkClient.cpp
--- a/Engine/Source/Programs/UnrealBuildAccelerator/Common/Private/UbaNetworkClient.cpp
+++ b/Engine/Source/Programs/UnrealBuildAccelerator/Common/Private/UbaNetworkClient.cpp
@@ -96,7 +96,7 @@ namespace uba
 
 		// The only way out of this function is to get a call to one of the below callbacks since exitScopeEvent must be set.
 
-		backend.SetDisconnectCallback(backendConnection, &rc, [](void* context, void* connection)
+		backend.SetDisconnectCallback(backendConnection, &rc, [](void* context, const Guid& connectionUid, void* connection)
 			{
 				auto& rc = *(RecvContext*)context;
 				rc.error = 4;
@@ -104,7 +104,7 @@ namespace uba
 				rc.exitScopeEvent.Set();
 			});
 
-		backend.SetRecvCallbacks(backendConnection, &rc, 1 + sizeof(Guid), [](void* context, u8* headerData, void*& outBodyContext, u8*& outBodyData, u32& outBodySize)
+		backend.SetRecvCallbacks(backendConnection, &rc, 1 + sizeof(Guid), [](void* context, const Guid& connectionUid, u8* headerData, void*& outBodyContext, u8*& outBodyData, u32& outBodySize)
 			{
 				auto& rc = *(RecvContext*)context;
 				rc.error = *headerData;
@@ -126,7 +126,7 @@ namespace uba
 				else
 				{
 					// Zero out callbacks since we will be leaving this scope when exiting current callback and we don't want the disconnect callback
-					rc.backend.SetRecvCallbacks(rc.backendConnection, nullptr, 0, [](void*, u8*, void*&, u8*&, u32&) { return false; }, nullptr, TC("Null"));
+					rc.backend.SetRecvCallbacks(rc.backendConnection, nullptr, 0, [](void*, const Guid&, u8*, void*&, u8*&, u32&) { return false; }, nullptr, TC("Null"));
 					rc.backend.SetDisconnectCallback(rc.backendConnection, nullptr, nullptr);
 				}
 
@@ -225,7 +225,7 @@ namespace uba
 		}
 		lock.Leave();
 
-		backend.SetDisconnectCallback(backendConnection, connection, [](void* context, void* connection)
+		backend.SetDisconnectCallback(backendConnection, connection, [](void* context, const Guid& connectionUid, void* connection)
 			{
 				auto& c = *(Connection*)context;
 				c.owner.OnDisconnected(c, true);
@@ -234,7 +234,7 @@ namespace uba
 		backend.SetRecvCallbacks(backendConnection, connection, ReceiveHeaderSize, ReceiveResponseHeader, ReceiveResponseBody, TC("ReceiveMessageResponse"));
 	}
 
-	bool NetworkClient::ReceiveResponseHeader(void* context, u8* headerData, void*& outBodyContext, u8*& outBodyData, u32& outBodySize)
+	bool NetworkClient::ReceiveResponseHeader(void* context, const Guid& connectionUid, u8* headerData, void*& outBodyContext, u8*& outBodyData, u32& outBodySize)
 	{
 		auto& connection = *(Connection*)context;
 		auto& client = connection.owner;
@@ -261,7 +261,7 @@ namespace uba
 		}
 		else if (!messageSize)
 		{
-			connection.recvCount++;
+			++client.m_recvCount;
 			msg->Done();
 			return true;
 		}
@@ -274,6 +274,10 @@ namespace uba
 		outBodyData = (u8*)msg->m_response;
 		outBodySize = messageSize;
 
+		++client.m_recvCount;
+		client.m_recvBytes += ReceiveHeaderSize + messageSize;
+
+
 		return true;
 	}
 
@@ -343,28 +347,14 @@ namespace uba
 
 	void NetworkClient::PrintSummary(Logger& logger)
 	{
-		Timer sendTimer;
-		u64 sendBytes = 0;
-		u64 recvBytes = 0;
-		u32 recvCount = 0;
-
 		SCOPED_READ_LOCK(m_connectionsLock, lock);
 		u32 connectionsCount = u32(m_connections.size());
-		for (auto& c : m_connections)
-		{
-			sendTimer.count += c.sendTimer.count;
-			sendTimer.time += c.sendTimer.time;
-			sendBytes += c.sendBytes;
-			recvBytes += c.recvBytes;
-			recvCount += c.recvCount;
-		}
 		lock.Leave();
 
-
 		logger.Info(TC("  ----- Uba client stats summary ------"));
-		logger.Info(TC("  SendTotal          %8u %9s"), sendTimer.count.load(), TimeToText(sendTimer.time).str);
-		logger.Info(TC("     Bytes                    %9s"), BytesToText(sendBytes).str);
-		logger.Info(TC("  RecvTotal          %8u %9s"), recvCount, BytesToText(recvBytes).str);
+		logger.Info(TC("  SendTotal          %8u %9s"), m_sendTimer.count.load(), TimeToText(m_sendTimer.time).str);
+		logger.Info(TC("     Bytes                    %9s"), BytesToText(m_sendBytes).str);
+		logger.Info(TC("  RecvTotal          %8u %9s"), m_recvCount.load(), BytesToText(m_recvBytes).str);
 		if (m_cryptoKey)
 		{
 			logger.Info(TC("  EncryptTotal       %8u %9s"), m_encryptTimer.count.load(), TimeToText(m_encryptTimer.time).str);
@@ -388,6 +378,7 @@ namespace uba
 
 	void NetworkClient::RegisterOnDisconnected(const OnDisconnectedFunction& function)
 	{
+		SCOPED_WRITE_LOCK(m_onDisconnectedFunctionsLock, lock);
 		m_onDisconnectedFunctions.push_back(function);
 	}
 
@@ -418,6 +409,14 @@ namespace uba
 		return m_sendSize;
 	}
 
+	NetworkBackend* NetworkClient::GetFirstConnectionBackend()
+	{
+		SCOPED_READ_LOCK(m_connectionsLock, connectionLock);
+		if (m_connections.empty())
+			return nullptr;
+		return m_connections.front().backend;
+	}
+
 	void NetworkClient::OnDisconnected(Connection& connection, bool calledFromReceive)
 	{
 		if (connection.connected.exchange(0) == 1)
@@ -429,6 +428,7 @@ namespace uba
 			if (m_connectionCount.fetch_sub(1) == 1)
 			{
 				m_isConnected.Reset();
+				SCOPED_READ_LOCK(m_onDisconnectedFunctionsLock, lock);
 				for (auto& f : m_onDisconnectedFunctions)
 					f();
 			}
@@ -441,8 +441,7 @@ namespace uba
 			if (m && m->m_connection == &connection)
 			{
 				m->m_error = true;
-				//m->m_responseSize = 0; // There is a race here where a message could just have fully arrived when disconnected. Can't set this to zero because we might be passed the m_error check but before setting reader size
-				m->Done();
+				m->Done(false);
 			}
 			++messageId;
 		}
@@ -469,6 +468,7 @@ namespace uba
 		BinaryWriter& writer = message.m_sendWriter;
 
 		u16 messageId = 0;
+		Event gotResponse(true);
 
 		if (response)
 		{
@@ -500,12 +500,20 @@ namespace uba
 
 				UBA_ASSERT(!m_activeMessages[messageId]);
 				m_activeMessages[messageId] = &message;
+
+				message.m_id = messageId;
+				message.m_sendContext.flags = NetworkBackend::SendFlags_ExternalWait;
+				if (!async)
+				{
+					UBA_ASSERT(!message.m_doneFunc);
+					message.m_doneUserData = &gotResponse;
+					message.m_doneFunc = [](bool error, void* userData) { ((Event*)userData)->Set(); };
+				}
 				break;
 			}
 		}
 
 		UBA_ASSERT(messageId < 65535);
-		message.m_id = messageId;
 
 		u32 sendSize = u32(writer.GetPosition());
 		u8* data = writer.GetData();
@@ -526,24 +534,10 @@ namespace uba
 		}
 
 
-		connection.sendBytes += sendSize;
-
-		Event gotResponse;
-
-		if (response)
-		{
-			message.m_sendContext.flags = NetworkBackend::SendFlags_ExternalWait;
-			if (!async)
-			{
-				gotResponse.Create(true);
-				UBA_ASSERT(!message.m_doneFunc);
-				message.m_doneUserData = &gotResponse;
-				message.m_doneFunc = [](bool error, void* userData) { ((Event*)userData)->Set(); };
-			}
-		}
+		m_sendBytes += sendSize;
 
 		{
-			TimerScope ts(connection.sendTimer);
+			TimerScope ts(m_sendTimer);
 			if (!connection.backend->Send(m_logger, connection.backendConnection, data, sendSize, message.m_sendContext))
 			{
 				OnDisconnected(connection, false);
@@ -559,7 +553,7 @@ namespace uba
 			u32 timeoutMs = 10 * 60 * 1000;
 			if (!gotResponse.IsSet(timeoutMs))
 			{
-				m_logger.Error(TC("Timed out after 10 minutes waiting for message."));
+				m_logger.Error(TC("Timed out after 10 minutes waiting for message response from server."));
 				message.m_error = true;
 			}
 			else if (m_cryptoKey && !message.m_error && message.m_responseSize)
@@ -652,13 +646,30 @@ namespace uba
 		return true;
 	}
 
-	void NetworkMessage::Done()
+	void NetworkMessage::Done(bool shouldLock)
 	{
-		if (m_id)
+		bool hasId = false;
+		auto returnId = [&]()
+			{
+				if (m_id)
+				{
+					m_client.m_availableMessageIds.push_back(m_id);
+					m_client.m_activeMessages[m_id] = nullptr;
+					m_id = 0;
+					hasId = true;
+				}
+			};
+
+		if (shouldLock)
+		{
+			SCOPED_WRITE_LOCK(m_client.m_activeMessagesLock, lock);
+			returnId();
+		}
+		else
 		{
-			m_client.ReturnMessageId(m_id);
-			m_id = 0;
+			returnId();
 		}
-		m_doneFunc(m_error, m_doneUserData);
+		if (hasId)
+			m_doneFunc(m_error, m_doneUserData);
 	}
 }
diff --git a/Engine/Source/Programs/UnrealBuildAccelerator/Common/Private/UbaNetworkServer.cpp b/Engine/Source/Programs/UnrealBuildAccelerator/Common/Private/UbaNetworkServer.cpp
--- a/Engine/Source/Programs/UnrealBuildAccelerator/Common/Private/UbaNetworkServer.cpp
+++ b/Engine/Source/Programs/UnrealBuildAccelerator/Common/Private/UbaNetworkServer.cpp
@@ -89,17 +89,18 @@ namespace uba
 	class NetworkServer::Connection
 	{
 	public:
-		Connection(NetworkServer& server, NetworkBackend& backend, void* backendConnection, const sockaddr& remoteSockAddr, CryptoKey cryptoKey)
+		Connection(NetworkServer& server, NetworkBackend& backend, void* backendConnection, const sockaddr& remoteSockAddr, CryptoKey cryptoKey, u32 id)
 		:	m_server(server)
 		,	m_backend(backend)
 		,	m_remoteSockAddr(remoteSockAddr)
 		,	m_cryptoKey(cryptoKey)
 		,	m_disconnectCallbackCalled(true)
+		,	m_id(id)
 		,	m_backendConnection(backendConnection)
 		{
 			m_activeWorkerCount = 1;
 
-			m_backend.SetDisconnectCallback(m_backendConnection, this, [](void* context, void* connection)
+			m_backend.SetDisconnectCallback(m_backendConnection, this, [](void* context, const Guid& connectionUid, void* connection)
 				{
 					auto& conn = *(Connection*)context;
 					conn.Disconnect();
@@ -167,7 +168,7 @@ namespace uba
 			return m_backend.Send(m_server.m_logger, m_backendConnection, data, 1 + sizeof(Guid), context);
 		}
 
-		static bool ReceiveHandshakeHeader(void* context, u8* headerData, void*& outBodyContext, u8*& outBodyData, u32& outBodySize)
+		static bool ReceiveHandshakeHeader(void* context, const Guid& connectionUid, u8* headerData, void*& outBodyContext, u8*& outBodyData, u32& outBodySize)
 		{
 			u8* handshakeData = new u8[sizeof(EncryptionHandshakeString)];
 			outBodyData = handshakeData;
@@ -192,7 +193,7 @@ namespace uba
 			return true;
 		}
 
-		static bool ReceiveVersion(void* context, u8* headerData, void*& outBodyContext, u8*& outBodyData, u32& outBodySize)
+		static bool ReceiveVersion(void* context, const Guid& connectionUid, u8* headerData, void*& outBodyContext, u8*& outBodyData, u32& outBodySize)
 		{
 			auto& conn = *(Connection*)context;
 			u32 clientVersion = *(u32*)headerData;
@@ -207,7 +208,7 @@ namespace uba
 			return true;
 		}
 
-		static bool ReceiveClientUid(void* context, u8* headerData, void*& outBodyContext, u8*& outBodyData, u32& outBodySize)
+		static bool ReceiveClientUid(void* context, const Guid& connectionUid, u8* headerData, void*& outBodyContext, u8*& outBodyData, u32& outBodySize)
 		{
 			auto& conn = *(Connection*)context;
 			auto& server = conn.m_server;
@@ -243,9 +244,6 @@ namespace uba
 
 			conn.m_client = &client;
 
-			Guid connectionUid;
-			CreateGuid(connectionUid);
-
 			if (client.connectionCount.fetch_add(1) == 0)
 			{
 				if (server.m_onConnectionFunction)
@@ -259,7 +257,7 @@ namespace uba
 			return true;
 		}
 
-		static bool ReceiveMessageHeader(void* context, u8* headerData, void*& outBodyContext, u8*& outBodyData, u32& outBodySize)
+		static bool ReceiveMessageHeader(void* context, const Guid& connectionUid, u8* headerData, void*& outBodyContext, u8*& outBodyData, u32& outBodySize)
 		{
 			auto& conn = *(Connection*)context;
 
@@ -360,6 +358,7 @@ namespace uba
 		Atomic<int> m_activeWorkerCount;
 		Atomic<int> m_disconnectCalled;
 		Atomic<bool> m_disconnected;
+		u32 m_id = 0;
 		bool m_shouldDisconnect = false;
 		void* m_backendConnection = nullptr;
 
@@ -1054,7 +1053,7 @@ namespace uba
 
 		RemoveDisconnectedConnections();
 
-		m_connections.emplace_back(*this, backend, backendConnection, remoteSocketAddr, cryptoKey);
+		m_connections.emplace_back(*this, backend, backendConnection, remoteSocketAddr, cryptoKey, m_connectionIdCounter++);
 		m_maxActiveConnections = Max(m_maxActiveConnections, u32(m_connections.size()));
 		return true;
 	}
diff --git a/Engine/Source/Programs/UnrealBuildAccelerator/Common/Public/UbaNetworkBackend.h b/Engine/Source/Programs/UnrealBuildAccelerator/Common/Public/UbaNetworkBackend.h
--- a/Engine/Source/Programs/UnrealBuildAccelerator/Common/Public/UbaNetworkBackend.h
+++ b/Engine/Source/Programs/UnrealBuildAccelerator/Common/Public/UbaNetworkBackend.h
@@ -42,13 +42,13 @@ namespace uba
 
 		// Recv data from connection. This is callback based so header callback is called first after headerSize bytes have been read
 		// If header callback sets a value in outBodySize, body callback will be called once body size bytes are received.
-		using RecvHeaderCallback = bool(void* context, u8* headerData, void*& outBodyContext, u8*& outBodyData, u32& outBodySize);
+		using RecvHeaderCallback = bool(void* context, const Guid& connectionUid, u8* headerData, void*& outBodyContext, u8*& outBodyData, u32& outBodySize);
 		using RecvBodyCallback = bool(void* context, bool recvError, u8* headerData, void* bodyContext, u8* bodyData, u32 bodySize);
 		virtual void SetRecvCallbacks(void* connection, void* context, u32 headerSize, RecvHeaderCallback* h, RecvBodyCallback* b, const tchar* recvHint) = 0;
 		virtual void SetRecvTimeout(void* connection, u32 timeoutMs) = 0;
 
 		// Disconnect callback. This is called as soon as connection is interrupted from send, recv or shutdown.
-		using DisconnectCallback = void(void* context, void* connection);
+		using DisconnectCallback = void(void* context, const Guid& connectionUid, void* connection);
 		virtual void SetDisconnectCallback(void* connection, void* context, DisconnectCallback* callback) = 0;
 
 		// Start listen on port/ip.
@@ -60,6 +60,8 @@ namespace uba
 		using ConnectedFunc = Function<bool(void* connection, const sockaddr& remoteSocketAddr, bool* timedOut)>;
 		virtual bool Connect(Logger& logger, const tchar* ip, const ConnectedFunc& connectedFunc, u16 port = DefaultPort, bool* timedOut = nullptr) = 0;
 		virtual bool Connect(Logger& logger, const sockaddr& remoteSocketAddr, const ConnectedFunc& connectedFunc, bool* timedOut = nullptr, const tchar* nameHint = nullptr) = 0;
+
+		virtual void GetTotalSendAndRecv(u64& outSend, u64& outRecv) = 0;
 	};
 
 	struct NetworkBackend::SendContext
diff --git a/Engine/Source/Programs/UnrealBuildAccelerator/Common/Public/UbaNetworkBackendMemory.h b/Engine/Source/Programs/UnrealBuildAccelerator/Common/Public/UbaNetworkBackendMemory.h
new file mode 100644
--- /dev/null
+++ b/Engine/Source/Programs/UnrealBuildAccelerator/Common/Public/UbaNetworkBackendMemory.h
@@ -0,0 +1,37 @@
+// Copyright Epic Games, Inc. All Rights Reserved.
+
+#pragma once
+
+#include "UbaLogger.h"
+#include "UbaNetworkBackend.h"
+#include "UbaThread.h"
+
+namespace uba
+{
+	class NetworkBackendMemory : public NetworkBackend
+	{
+	public:
+		NetworkBackendMemory(LogWriter& writer, const tchar* prefix = TC("NetworkBackendMemory"));
+		virtual ~NetworkBackendMemory();
+		virtual void Shutdown(void* connection) override;
+		virtual bool Send(Logger& logger, void* connection, const void* data, u32 dataSize, SendContext& sendContext) override;
+		virtual void SetDataSentCallback(void* connection, void* context, DataSentCallback* callback) override;
+		virtual void SetRecvCallbacks(void* connection, void* context, u32 headerSize, RecvHeaderCallback* h, RecvBodyCallback* b, const tchar* recvHint) override;
+		virtual void SetRecvTimeout(void* connection, u32 timeoutMs) override;
+		virtual void SetDisconnectCallback(void* connection, void* context, DisconnectCallback* callback) override;
+
+		virtual bool StartListen(Logger& logger, u16 port, const tchar* ip, const ListenConnectedFunc& connectedFunc) override;
+		virtual void StopListen() override;
+		virtual bool Connect(Logger& logger, const tchar* ip, const ConnectedFunc& connectedFunc, u16 port = DefaultPort, bool* timedOut = nullptr) override;
+		virtual bool Connect(Logger& logger, const sockaddr& remoteSocketAddr, const ConnectedFunc& connectedFunc, bool* timedOut = nullptr, const tchar* nameHint = nullptr) override;
+
+		virtual void GetTotalSendAndRecv(u64& outSend, u64& outRecv) override;
+
+	private:
+		struct Connection;
+		Connection* m_connection = nullptr;
+
+		ReaderWriterLock m_connectedFuncLock;
+		ListenConnectedFunc m_connectedFunc;
+	};
+};
diff --git a/Engine/Source/Programs/UnrealBuildAccelerator/Common/Public/UbaNetworkBackendTcp.h b/Engine/Source/Programs/UnrealBuildAccelerator/Common/Public/UbaNetworkBackendTcp.h
--- a/Engine/Source/Programs/UnrealBuildAccelerator/Common/Public/UbaNetworkBackendTcp.h
+++ b/Engine/Source/Programs/UnrealBuildAccelerator/Common/Public/UbaNetworkBackendTcp.h
@@ -25,7 +25,7 @@ namespace uba
 		virtual bool Connect(Logger& logger, const tchar* ip, const ConnectedFunc& connectedFunc, u16 port = DefaultPort, bool* timedOut = nullptr) override;
 		virtual bool Connect(Logger& logger, const sockaddr& remoteSocketAddr, const ConnectedFunc& connectedFunc, bool* timedOut = nullptr, const tchar* nameHint = nullptr) override;
 
-		bool Empty() { return m_listenEntries.empty(); }
+		virtual void GetTotalSendAndRecv(u64& outSend, u64& outRecv) override;
 
 	private:
 		bool EnsureInitialized(Logger& logger);
@@ -36,12 +36,14 @@ namespace uba
 		void ThreadRecv(Connection& connection);
 
 		LoggerWithWriter m_logger;
+		ReaderWriterLock m_listenEntriesLock;
 		List<ListenEntry> m_listenEntries;
-		ListenConnectedFunc m_connectedFunc;
 
 		ReaderWriterLock m_connectionsLock;
 		List<Connection> m_connections;
 
+		Atomic<u64> m_totalSend;
+		Atomic<u64> m_totalRecv;
 
 		#if PLATFORM_WINDOWS
 		bool m_wsaInitDone = false;
diff --git a/Engine/Source/Programs/UnrealBuildAccelerator/Common/Public/UbaNetworkClient.h b/Engine/Source/Programs/UnrealBuildAccelerator/Common/Public/UbaNetworkClient.h
--- a/Engine/Source/Programs/UnrealBuildAccelerator/Common/Public/UbaNetworkClient.h
+++ b/Engine/Source/Programs/UnrealBuildAccelerator/Common/Public/UbaNetworkClient.h
@@ -65,6 +65,10 @@ namespace uba
 		LogWriter& GetLogWriter() { return m_logWriter; }
 		NetworkBackendTcp& GetTcpBackend() { return *m_tcpBackend; }
 		u32 GetConnectionCount() { return m_connectionCount; }
+		u64 GetTotalSentBytes() { return m_sendBytes; }
+		u64 GetTotalRecvBytes() { return m_recvBytes; }
+
+		NetworkBackend* GetFirstConnectionBackend();
 
 	private:
 		struct Connection
@@ -74,16 +78,12 @@ namespace uba
 			void* backendConnection = nullptr;
 			Atomic<u32> connected;
 			Event disconnectedEvent;
-			Timer sendTimer;
-			Atomic<u64> sendBytes;
-			u64 recvBytes = 0;
-			u32 recvCount = 0;
 			NetworkBackend* backend = nullptr;
 		};
 
 		bool AddConnection(NetworkBackend& backend, void* backendConnection, bool* timedOut);
 		void ConnectedCallback(NetworkBackend& backend, void* backendConnection);
-		static bool ReceiveResponseHeader(void* context, u8* headerData, void*& outBodyContext, u8*& outBodyData, u32& outBodySize);
+		static bool ReceiveResponseHeader(void* context, const Guid& connectionUid, u8* headerData, void*& outBodyContext, u8*& outBodyData, u32& outBodySize);
 		static bool ReceiveResponseBody(void* context, bool recvError, u8* headerData, void* bodyContext, u8* bodyData, u32 bodySize);
 		void OnDisconnected(Connection& connection, bool calledFromReceive);
 		bool Send(NetworkMessage& message, void* response, u32 responseCapacity, bool async);
@@ -96,6 +96,10 @@ namespace uba
 		LoggerWithWriter m_logger;
 		u32 m_sendSize;
 		u32 m_receiveTimeoutSeconds;
+		Atomic<u64> m_sendBytes;
+		Atomic<u64> m_recvBytes;
+		Atomic<u32> m_recvCount;
+		Timer m_sendTimer;
 
 		ReaderWriterLock m_serverUidLock;
 		Guid m_serverUid;
@@ -105,6 +109,7 @@ namespace uba
 		Atomic<u32> m_connectionCount;
 		ReaderWriterLock m_onConnectedFunctionsLock;
 		Vector<OnConnectedFunction> m_onConnectedFunctions;
+		ReaderWriterLock m_onDisconnectedFunctionsLock;
 		Vector<OnDisconnectedFunction> m_onDisconnectedFunctions;
 		OnVersionMismatchFunction m_versionMismatchFunction;
 
diff --git a/Engine/Source/Programs/UnrealBuildAccelerator/Common/Public/UbaNetworkMessage.h b/Engine/Source/Programs/UnrealBuildAccelerator/Common/Public/UbaNetworkMessage.h
--- a/Engine/Source/Programs/UnrealBuildAccelerator/Common/Public/UbaNetworkMessage.h
+++ b/Engine/Source/Programs/UnrealBuildAccelerator/Common/Public/UbaNetworkMessage.h
@@ -21,7 +21,7 @@ namespace uba
 		bool ProcessAsyncResults(BinaryReader& response); // Note, this must be called after GetResponseFunc has been called
 
 	private:
-		void Done();
+		void Done(bool shouldLock = true);
 
 		NetworkClient& m_client;
 		BinaryWriter& m_sendWriter;
diff --git a/Engine/Source/Programs/UnrealBuildAccelerator/Common/Public/UbaNetworkServer.h b/Engine/Source/Programs/UnrealBuildAccelerator/Common/Public/UbaNetworkServer.h
--- a/Engine/Source/Programs/UnrealBuildAccelerator/Common/Public/UbaNetworkServer.h
+++ b/Engine/Source/Programs/UnrealBuildAccelerator/Common/Public/UbaNetworkServer.h
@@ -134,6 +134,7 @@ namespace uba
 
 		ReaderWriterLock m_connectionsLock;
 		List<Connection> m_connections;
+		u32 m_connectionIdCounter = 1;
 		u32 m_maxActiveConnections = 0;
 
 		struct Client
diff --git a/Engine/Source/Programs/UnrealBuildAccelerator/Test/Private/UbaTestAll.h b/Engine/Source/Programs/UnrealBuildAccelerator/Test/Private/UbaTestAll.h
--- a/Engine/Source/Programs/UnrealBuildAccelerator/Test/Private/UbaTestAll.h
+++ b/Engine/Source/Programs/UnrealBuildAccelerator/Test/Private/UbaTestAll.h
@@ -33,6 +33,7 @@ namespace uba
 		UBA_TEST(TestSockets) \
 		UBA_TEST(TestClientServer) \
 		UBA_TEST(TestClientServer2) \
+		UBA_TEST(TestClientServerMem) \
 		UBA_TEST(TestStorage) \
 		UBA_TEST(TestDetouredTestApp) \
 		UBA_TEST(TestDetouredClang) \
diff --git a/Engine/Source/Programs/UnrealBuildAccelerator/Test/Private/UbaTestNetwork.h b/Engine/Source/Programs/UnrealBuildAccelerator/Test/Private/UbaTestNetwork.h
--- a/Engine/Source/Programs/UnrealBuildAccelerator/Test/Private/UbaTestNetwork.h
+++ b/Engine/Source/Programs/UnrealBuildAccelerator/Test/Private/UbaTestNetwork.h
@@ -8,6 +8,7 @@
 #include "UbaNetworkServer.h"
 #include "UbaNetworkMessage.h"
 #include "UbaNetworkBackendTcp.h"
+#include "UbaNetworkBackendMemory.h"
 
 namespace uba
 {
@@ -23,7 +24,7 @@ namespace uba
 					{
 						logger.Info(TC("Listen got connection"));
 						tcp.SetRecvCallbacks(connection, &result, 1,
-							[](void* context, u8* headerData, void*& outBodyContext, u8*& outBodyData, u32& outBodySize)
+							[](void* context, const Guid& connectionUid, u8* headerData, void*& outBodyContext, u8*& outBodyData, u32& outBodySize)
 							{
 								*(u8*)context = *headerData;
 								//wprintf(TC("Listen got data from peer: %u\n"), *headerData);
@@ -39,7 +40,7 @@ namespace uba
 		logger.Info(TC("Starting to Connect"));
 		tcp.Connect(logger, TC("127.0.0.1"), [&](void* connection, const sockaddr& remoteSocketAddr, bool* timedOut)
 			{
-				tcp.SetRecvCallbacks(connection, &tcp, 1, [](void* context, u8* headerData, void*& outBodyContext, u8*& outBodyData, u32& outBodySize) { return true; }, nullptr, TC(""));
+				tcp.SetRecvCallbacks(connection, &tcp, 1, [](void* context, const Guid& connectionUid, u8* headerData, void*& outBodyContext, u8*& outBodyData, u32& outBodySize) { return true; }, nullptr, TC(""));
 				u8 b = 42;
 				NetworkBackend::SendContext sc;
 				tcp.Send(logger, connection, &b, 1, sc);
@@ -58,8 +59,7 @@ namespace uba
 	bool TestClientServer(LoggerWithWriter& logger, const StringBufferBase& testRootDir)
 	{
 		auto& logWriter = logger.m_writer;
-		NetworkBackendTcp serverTcp(logWriter, TC("ServerTcp"));
-		NetworkBackendTcp clientTcp(logWriter, TC("ClientTcp"));
+		NetworkBackendTcp networkBackend(logWriter);
 
 		bool ctorSuccess = true;
 		NetworkServer server(ctorSuccess, logWriter);
@@ -73,10 +73,10 @@ namespace uba
 				return true;
 			});
 
-		if (!server.StartListen(serverTcp, 1234))
+		if (!server.StartListen(networkBackend, 1234))
 			return logger.Error(TC("Failed to listen"));
 		Sleep(100);
-		if (!client.Connect(clientTcp, TC("127.0.0.1"), 1234))
+		if (!client.Connect(networkBackend, TC("127.0.0.1"), 1234))
 			return logger.Error(TC("Failed to connect"));
 
 
@@ -95,8 +95,7 @@ namespace uba
 	bool TestClientServer2(LoggerWithWriter& logger, const StringBufferBase& testRootDir)
 	{
 		auto& logWriter = logger.m_writer;
-		NetworkBackendTcp serverTcp(logWriter, TC("ServerTcp"));
-		NetworkBackendTcp clientTcp(logWriter, TC("ClientTcp"));
+		NetworkBackendTcp networkBackend(logWriter);
 
 		bool ctorSuccess = true;
 		NetworkServer server(ctorSuccess, logWriter);
@@ -110,12 +109,12 @@ namespace uba
 				return true;
 			});
 
-		if (!client.StartListen(clientTcp, 1239))
+		if (!client.StartListen(networkBackend, 1239))
 			return logger.Error(TC("Client failed to listen"));
 		Sleep(100);
-		if (!server.AddClient(serverTcp, TC("127.0.0.1"), 1239))
+		if (!server.AddClient(networkBackend, TC("127.0.0.1"), 1239))
 			return logger.Error(TC("Server failed to connect"));
-		if (!server.AddClient(serverTcp, TC("127.0.0.1"), 1239))
+		if (!server.AddClient(networkBackend, TC("127.0.0.1"), 1239))
 			return logger.Error(TC("Server failed to connect second"));
 
 		u64 time = GetTime();
@@ -126,6 +125,42 @@ namespace uba
 			Sleep(100);
 		}
 
+		StackBinaryWriter<128> writer;
+		NetworkMessage msg(client, 1, SessionMessageType_Ping, writer);
+		writer.WriteU32(32);
+		StackBinaryReader<32> reader;
+		if (!msg.Send(reader))
+			return logger.Error(TC("Failed to get message"));
+
+		u8 value = reader.ReadByte();
+		logger.Info(TC("Got value %u"), value);
+		return true;
+	}
+
+	bool TestClientServerMem(LoggerWithWriter& logger, const StringBufferBase& testRootDir)
+	{
+		auto& logWriter = logger.m_writer;
+		NetworkBackendMemory networkBackend(logWriter);
+
+		bool ctorSuccess = true;
+		NetworkServer server(ctorSuccess, logWriter);
+		NetworkClient client(ctorSuccess, logWriter);
+
+		server.RegisterService(1, [&](const ConnectionInfo& connectionInfo, u8 messageType, BinaryReader& reader, BinaryWriter& writer)
+			{
+				logger.Info(TC("Got ping!"));
+				UBA_ASSERT(messageType == SessionMessageType_Ping);
+				writer.WriteByte(42);
+				return true;
+			});
+
+		if (!server.StartListen(networkBackend, 1234))
+			return logger.Error(TC("Failed to listen"));
+		Sleep(100);
+		if (!client.Connect(networkBackend, TC("127.0.0.1"), 1234))
+			return logger.Error(TC("Failed to connect"));
+
+
 		StackBinaryWriter<128> writer;
 		NetworkMessage msg(client, 1, SessionMessageType_Ping, writer);
 		writer.WriteU32(32);
