[UBA]
* Added support for sending .uba trace files from helpers to server

diff --git a/Engine/Source/Programs/UnrealBuildAccelerator/Common/Private/UbaSession.cpp b/Engine/Source/Programs/UnrealBuildAccelerator/Common/Private/UbaSession.cpp
--- a/Engine/Source/Programs/UnrealBuildAccelerator/Common/Private/UbaSession.cpp
+++ b/Engine/Source/Programs/UnrealBuildAccelerator/Common/Private/UbaSession.cpp
@@ -538,6 +538,11 @@ namespace uba
 
 	bool Session::CreateMemoryMapFromView(MemoryMap& out, StringKey fileNameKey, const tchar* fileName, const CasKey& casKey, u64 alignment)
 	{
+		//StringBuffer<> workName;
+		//u32 len = TStrlen(fileName);
+		//workName.Append(TC("MM:")).Append(len > 30 ? fileName + (len - 30) : fileName);
+		//TrackWorkScope tws(*m_workManager, workName.data);
+
 		SCOPED_WRITE_LOCK(m_fileMappingTableLookupLock, lookupLock);
 		auto insres = m_fileMappingTableLookup.try_emplace(fileNameKey);
 		FileMappingEntry& entry = insres.first->second;
@@ -1039,23 +1044,7 @@ namespace uba
 			traceName.Append(m_id);
 
 		if (!traceName.IsEmpty())
-		{
-			#if PLATFORM_WINDOWS
-			LoggerWithWriter(m_logger.m_writer).Info(TC("---- Starting trace: %s ----"), traceName.data);
-			m_trace.StartWrite(traceName.data, m_detailedTrace ? 512*1024*1024 : 128*1024*1024);
-			#else
-			LoggerWithWriter(m_logger.m_writer).Info(TC("---- Starting trace ----")); // non-windows named shared memory not implemented (only needed for UbaVisualizer which you can't run on linux either way)
-			m_trace.StartWrite(nullptr, m_detailedTrace ? 512 * 1024 * 1024 : 128 * 1024 * 1024);
-			#endif
-			tchar buf[256];
-			if (!GetComputerNameW(buf, sizeof_array(buf)))
-				TStrcpy_s(buf, 256, TC("LOCAL"));
-			StringBuffer<> systemInfo;
-			GetSystemInfo(systemInfo);
-			m_trace.SessionAdded(0, {}, buf, systemInfo.data);
-			m_traceThreadEvent.Create(true);
-			m_traceThread.Start([this]() { ThreadTraceLoop(); return 0; });
-		}
+			StartTrace(IsWindows ? traceName.data : nullptr); // non-windows named shared memory not implemented (only needed for UbaVisualizer which you can't run on linux either way)
 
 		#if PLATFORM_WINDOWS
 		if (info.launchVisualizer)
@@ -1101,8 +1090,7 @@ namespace uba
 
 	Session::~Session()
 	{
-		StopTraceThread();
-		m_trace.StopWrite(m_traceOutputFile.data);
+		StopTrace(m_traceOutputFile.data);
 
 		CancelAllProcessesAndWait();
 		FlushDeadProcesses();
@@ -1412,6 +1400,31 @@ namespace uba
 		stats.Print(m_logger);
 	}
 
+	void Session::StartTrace(const tchar* traceName)
+	{
+		if (traceName)
+			LoggerWithWriter(m_logger.m_writer).Info(TC("---- Starting trace: %s ----"), traceName);
+		else
+			LoggerWithWriter(m_logger.m_writer).Info(TC("---- Starting trace ----"));
+
+		m_trace.StartWrite(traceName, m_detailedTrace ? 512*1024*1024 : 128*1024*1024);
+
+		tchar buf[256];
+		if (!GetComputerNameW(buf, sizeof_array(buf)))
+			TStrcpy_s(buf, 256, TC("LOCAL"));
+		StringBuffer<> systemInfo;
+		GetSystemInfo(systemInfo);
+		m_trace.SessionAdded(0, {}, buf, systemInfo.data);
+		m_traceThreadEvent.Create(true);
+		m_traceThread.Start([this]() { ThreadTraceLoop(); return 0; });
+	}
+
+	bool Session::StopTrace(const tchar* writeFile)
+	{
+		StopTraceThread();
+		return m_trace.StopWrite(writeFile);
+	}
+
 	void Session::StopTraceThread()
 	{
 		m_traceThreadEvent.Set();
@@ -2086,8 +2099,8 @@ namespace uba
 		logger.Info(TC("  DirectoryTable      %7u %9s"), u32(m_directoryTable.m_lookup.size()), BytesToText(GetDirectoryTableSize()).str);
 		logger.Info(TC("  MappingTable        %7u %9s"), u32(m_fileMappingTableLookup.size()), BytesToText(GetFileMappingSize()).str);
 		logger.Info(TC("  MappingBuffer       %7u %9s"), mappingBufferCount, BytesToText(mappingBufferSize).str);
-		logger.Info(TC(""));
 		m_stats.Print(logger);
+		logger.Info(TC(""));
 	}
 
 	void Session::CreateProcessJobObject()
diff --git a/Engine/Source/Programs/UnrealBuildAccelerator/Common/Private/UbaSessionClient.cpp b/Engine/Source/Programs/UnrealBuildAccelerator/Common/Private/UbaSessionClient.cpp
--- a/Engine/Source/Programs/UnrealBuildAccelerator/Common/Private/UbaSessionClient.cpp
+++ b/Engine/Source/Programs/UnrealBuildAccelerator/Common/Private/UbaSessionClient.cpp
@@ -85,6 +85,11 @@ namespace uba
 		m_terminationReason = reason;
 	}
 
+	void SessionClient::SetMaxProcessCount(u32 count)
+	{
+		m_maxProcessCount = count;
+	}
+
 	u64 SessionClient::GetBestPing()
 	{
 		return m_bestPing;
@@ -92,6 +97,11 @@ namespace uba
 
 	bool SessionClient::RetrieveCasFile(CasKey& outNewKey, u64& outSize, const CasKey& casKey, const tchar* hint, bool storeUncompressed, bool allowProxy)
 	{
+		//StringBuffer<> workName;
+		//u32 len = TStrlen(hint);
+		//workName.Append(TC("RC:")).Append(len > 30 ? hint + (len - 30) : hint);
+		//TrackWorkScope tws(m_client, workName.data);
+
 		TimerScope s(m_stats.storageRetrieve);
 		CasKey tempKey = casKey;
 
@@ -488,16 +498,15 @@ namespace uba
 			if (!pair.second.mappingHandle.IsValid())
 				m_logger.Warning(TC("%s is not using file mapping"), pair.first.c_str());
 #endif
-			if (!SendFile(process, pair.second, pair.first.c_str()))
+			bool keepMappingInMemory = IsWindows && !IsRarelyReadAfterWritten(process, pair.first.c_str(), pair.first.size());
+			if (!SendFile(pair.second, pair.first.c_str(), process.GetId(), keepMappingInMemory))
 				return false;
 		}
 		return true;
 	}
 
-	bool SessionClient::SendFile(ProcessImpl& process, WrittenFile& source, const tchar* destination)
+	bool SessionClient::SendFile(WrittenFile& source, const tchar* destination, u32 processId, bool keepMappingInMemory)
 	{
-		bool keepMappingInMemory = IsWindows && !IsRarelyReadAfterWritten(process, destination, TStrlen(destination));
-
 		CasKey casKey;
 		{
 			TimerScope ts(m_stats.storageSend);
@@ -514,7 +523,7 @@ namespace uba
 		{
 			StackBinaryWriter<1024> writer;
 			NetworkMessage msg(m_client, ServiceId, SessionMessageType_SendFileToServer, writer);
-			writer.WriteU32(process.GetId());
+			writer.WriteU32(processId);
 			writer.WriteString(destination);
 			writer.WriteStringKey(source.key);
 			writer.WriteU32(source.attributes);
@@ -1040,8 +1049,20 @@ namespace uba
 		m_uiLanguage = reader.ReadU32();
 		m_detailedTrace = reader.ReadBool();
 		m_shouldSendLogToServer = reader.ReadBool();
+		m_shouldSendTraceToServer = reader.ReadBool();
+
 		if (m_shouldSendLogToServer)
 			m_logToFile = true;
+
+		if (m_shouldSendTraceToServer)
+		{
+			//if (m_detailedTrace)
+			{
+				m_client.SetWorkTracker(&m_trace);
+			}
+			StartTrace(nullptr);
+		}
+
 		BuildEnvironmentVariables(reader);
 
 		m_loop = true;
@@ -1263,7 +1284,7 @@ namespace uba
 			logFile = lastSlash + 1;
 		dest.Append(TC("<log>")).Append(logFile);
 		f.key = ToStringKeyLower(dest);
-		SendFile(pi, f, dest.data);
+		SendFile(f, dest.data, pi.GetId(), false);
 	}
 
 	void SessionClient::GetLogFileName(StringBufferBase& out, const tchar* logFile, const tchar* arguments)
@@ -1309,7 +1330,6 @@ namespace uba
 		u64 memRequiredToSpawn = u64(double(memTotal) * double(100 - m_memWaitLoadPercent) / 100.0);
 		u64 memRequiredFree = u64(double(memTotal) * double(100 - m_memKillLoadPercent) / 100.0);
 
-		float maxWeight = float(m_maxProcessCount);
 		float activeWeight = 0;
 		ReaderWriterLock activeWeightLock;
 
@@ -1357,6 +1377,7 @@ namespace uba
 
 		while (m_loop)
 		{
+			float maxWeight = float(m_maxProcessCount);
 			u32 waitTimeoutMs = 3000;
 
 			FlushDeadProcesses();
@@ -1624,6 +1645,37 @@ namespace uba
 		SendSummary();
 
 		m_client.FlushWork();
+
+
+		if (m_shouldSendTraceToServer)
+		{
+			m_client.SetWorkTracker(nullptr);
+
+			StopTraceThread();
+
+			StackBinaryWriter<SendMaxSize> writer;
+			WriteSummary(writer, [&](Logger& logger)
+				{
+					PrintSummary(logger);
+					m_storage.PrintSummary(logger);
+					m_client.PrintSummary(logger);
+					SystemStats::GetGlobal().Print(logger, true);
+				});
+			m_trace.SessionSummary(0, writer.GetData(), writer.GetPosition());
+
+			StringBuffer<> ubaFile(m_sessionLogDir);
+			ubaFile.Append(TC("Trace.uba"));
+			if (StopTrace(ubaFile.data))
+			{
+				WrittenFile f;
+				f.name = ubaFile.data;
+				f.attributes = DefaultAttributes();
+				StringBuffer<> dest(TC("<uba>"));
+				f.key = ToStringKeyLower(dest);
+
+				SendFile(f, dest.data, 0, false);
+			}
+		}
 	}
 
 	u32 SessionClient::CountLogLines(ProcessImpl& process)
@@ -1787,4 +1839,25 @@ namespace uba
 		}
 		return SendUpdateDirectoryTable(reader);
 	}
+
+	void SessionClient::TraceSessionUpdate()
+	{
+		float cpuLoad = 0.0f;//UpdateCpuLoad();
+		u64 send;
+		u64 recv;
+		if (auto backend = m_client.GetFirstConnectionBackend())
+		{
+			backend->GetTotalSendAndRecv(send, recv);
+		}
+		else
+		{
+			recv = m_client.GetTotalRecvBytes();
+			send = m_client.GetTotalSentBytes();
+		}
+		u64 memAvail = 0;//m_memAvail;
+		u64 memTotal = 0;//m_memTotal;
+
+		// send and recv are swapped on purpose because that is how visualizer is visualizing 
+		m_trace.SessionUpdate(0, 0, send, recv, 0, memAvail, memTotal, cpuLoad);
+	}
 }
diff --git a/Engine/Source/Programs/UnrealBuildAccelerator/Common/Private/UbaSessionServer.cpp b/Engine/Source/Programs/UnrealBuildAccelerator/Common/Private/UbaSessionServer.cpp
--- a/Engine/Source/Programs/UnrealBuildAccelerator/Common/Private/UbaSessionServer.cpp
+++ b/Engine/Source/Programs/UnrealBuildAccelerator/Common/Private/UbaSessionServer.cpp
@@ -93,7 +93,7 @@ namespace uba
 	,	m_server(info.server)
 	,	m_maxRemoteProcessCount(~0u)
 	{
-		m_server.RegisterOnClientDisconnected(ServiceId, [this](const Guid& clientUid, u32 clientId) { OnDisconnected(clientId); });
+		m_server.RegisterOnClientDisconnected(ServiceId, [this](const Guid& clientUid, u32 clientId) { OnDisconnected(clientUid, clientId); });
 
 		m_server.RegisterService(ServiceId,
 			[this](const ConnectionInfo& connectionInfo, u8 messageType, BinaryReader& reader, BinaryWriter& writer)
@@ -119,6 +119,7 @@ namespace uba
 		m_nameToHashTableEnabled = info.nameToHashTableEnabled;
 		m_memKillLoadPercent = info.memKillLoadPercent;
 		m_remoteLogEnabled = info.remoteLogEnabled;
+		m_remoteTraceEnabled = info.remoteTraceEnabled;
 
 		if (m_resetCas)
 			m_storage.Reset();
@@ -436,7 +437,7 @@ namespace uba
 		return m_server;
 	}
 
-	void SessionServer::OnDisconnected(u32 clientId)
+	void SessionServer::OnDisconnected(const Guid& clientUid, u32 clientId)
 	{
 		u32 returnCount = 0;
 		ScopedCriticalSection queueLock(m_remoteProcessAndSessionLock);
@@ -493,7 +494,7 @@ namespace uba
 			if (sessionName.IsEmpty())
 				sessionName.Append(TC("<can't find session>"));
 
-			m_logger.Info(TC("Client session %s disconnected. Returned %u process(s) to queue"), sessionName.data, returnCount);
+			m_logger.Info(TC("Client session %s (%s) disconnected. Returned %u process(s) to queue"), sessionName.data, GuidToString(clientUid).str, returnCount);
 		}
 
 		if (m_connectionCount)
@@ -610,6 +611,7 @@ namespace uba
 				writer.WriteU32(m_uiLanguage);
 				writer.WriteBool(m_detailedTrace);
 				writer.WriteBool(m_remoteLogEnabled);
+				writer.WriteBool(m_remoteTraceEnabled);
 				WriteRemoteEnvironmentVariables(writer);
 
 				m_trace.SessionAdded(sessionId, connectionInfo.GetId(), name.data, info.data); // Must be inside lock for TraceSessionUpdate() to not include
@@ -782,6 +784,16 @@ namespace uba
 						return true;
 					}
 
+					if (destination.StartsWith(TC("<uba>")))
+					{
+						StringBuffer<> ubaPath;
+						ubaPath.Append(m_sessionLogDir).AppendValue(connectionInfo.GetId()).Append(TC(".uba"));
+						m_storage.CopyOrLink(casKey, ubaPath.data, attributes);
+						m_storage.DropCasFile(casKey, false, ubaPath.data);
+						writer.WriteBool(true);
+						return true;
+					}
+
 					if (ShouldWriteToDisk(destination.data, destination.count))
 					{
 						success = m_storage.CopyOrLink(casKey, destination.data, attributes);
diff --git a/Engine/Source/Programs/UnrealBuildAccelerator/Common/Public/UbaNetwork.h b/Engine/Source/Programs/UnrealBuildAccelerator/Common/Public/UbaNetwork.h
--- a/Engine/Source/Programs/UnrealBuildAccelerator/Common/Public/UbaNetwork.h
+++ b/Engine/Source/Programs/UnrealBuildAccelerator/Common/Public/UbaNetwork.h
@@ -18,7 +18,7 @@ namespace uba
 	
 	static constexpr u32 SystemNetworkVersion = 1339;
 	static constexpr u32 StorageNetworkVersion = 3;
-	static constexpr u32 SessionNetworkVersion = 27;
+	static constexpr u32 SessionNetworkVersion = 28;
 
 
 	// Messages used over network between client/server (system, storage and session)
diff --git a/Engine/Source/Programs/UnrealBuildAccelerator/Common/Public/UbaSession.h b/Engine/Source/Programs/UnrealBuildAccelerator/Common/Public/UbaSession.h
--- a/Engine/Source/Programs/UnrealBuildAccelerator/Common/Public/UbaSession.h
+++ b/Engine/Source/Programs/UnrealBuildAccelerator/Common/Public/UbaSession.h
@@ -79,6 +79,8 @@ namespace uba
 		void ProcessExited(ProcessImpl& process, u64 executionTime);
 		void FlushDeadProcesses();
 		void PrintProcessStats(ProcessStats& stats, const tchar* logName);
+		void StartTrace(const tchar* traceName);
+		bool StopTrace(const tchar* writeFile);
 		void StopTraceThread();
 		u32 GetDirectoryTableSize();
 		u32 GetFileMappingSize();
@@ -140,14 +142,6 @@ namespace uba
 		void ThreadTraceLoop();
 		virtual void TraceSessionUpdate();
 
-		struct MutableLogger : public LoggerWithWriter
-		{
-			MutableLogger(LogWriter& writer, const tchar* prefix) : LoggerWithWriter(writer, prefix) {}
-			virtual void Log(LogEntryType type, const tchar* str, u32 strLen) override { if (!isMuted) LoggerWithWriter::Log(type, str, strLen); }
-			Atomic<bool> isMuted;
-		};
-
-
 		Storage& m_storage;
 		MutableLogger m_logger;
 
diff --git a/Engine/Source/Programs/UnrealBuildAccelerator/Common/Public/UbaSessionClient.h b/Engine/Source/Programs/UnrealBuildAccelerator/Common/Public/UbaSessionClient.h
--- a/Engine/Source/Programs/UnrealBuildAccelerator/Common/Public/UbaSessionClient.h
+++ b/Engine/Source/Programs/UnrealBuildAccelerator/Common/Public/UbaSessionClient.h
@@ -38,6 +38,7 @@ namespace uba
 		void Stop();
 		bool Wait(u32 milliseconds = 0xFFFFFFFF, Event* wakeupEvent = nullptr);
 		void SetIsTerminating(const tchar* reason = TC("Terminating"), u64 delayMs = 0); // Session stores pointer directly. Can't be temporary
+		void SetMaxProcessCount(u32 count);
 
 		u64 GetBestPing();
 
@@ -61,6 +62,7 @@ namespace uba
 		virtual bool CustomMessage(Process& process, BinaryReader& reader, BinaryWriter& writer) override;
 		virtual bool FlushWrittenFiles(ProcessImpl& process) override;
 		virtual bool UpdateEnvironment(ProcessImpl& process, const tchar* reason, bool resetStats) override;
+		virtual void TraceSessionUpdate() override;
 
 		struct InternalProcessStartInfo;
 
@@ -69,7 +71,7 @@ namespace uba
 		bool EnsureBinaryFile(StringBufferBase& out, StringBufferBase& outVirtual, u32 processId, const StringBufferBase& fileName, const StringKey& fileNameKey, const tchar* applicationDir);
 		bool WriteBinFile(StringBufferBase& out, const tchar* binaryName, const CasKey& casKey, const KeyToString& applicationDir, u32 fileAttributes);
 		bool SendFiles(ProcessImpl& process, Timer& sendFiles);
-		bool SendFile(ProcessImpl& process, WrittenFile& source, const tchar* destination);
+		bool SendFile(WrittenFile& source, const tchar* destination, u32 processId, bool keepMappingInMemory);
 		bool SendUpdateDirectoryTable(StackBinaryReader<SendMaxSize>& reader); // Note, reader is sent in to save stack space.
 		bool UpdateDirectoryTableFromServer(StackBinaryReader<SendMaxSize>& reader);
 		bool SendUpdateNameToHashTable(StackBinaryReader<SendMaxSize>& reader);
@@ -93,7 +95,6 @@ namespace uba
 		StringBuffer<128> m_name;
 		StringBuffer<MaxPath> m_processWorkingDir;
 		u32 m_sessionId = 0;
-		u32 m_maxProcessCount = 1;
 		u32 m_uiLanguage = 0;
 		u32 m_defaultPriorityClass;
 		u32 m_outputStatsThresholdMs = 0;
@@ -108,10 +109,12 @@ namespace uba
 		bool m_dedicated = false;
 		bool m_useStorage = true;
 		bool m_shouldSendLogToServer = false;
+		bool m_shouldSendTraceToServer = false;
 		bool m_remoteExecutionEnabled = true;
 		
 		Atomic<const tchar*> m_terminationReason;
 		Atomic<u64> m_terminationTime;
+		Atomic<u32> m_maxProcessCount;
 
 		ReaderWriterLock m_handledApplicationEnvironmentsLock;
 		UnorderedSet<TString> m_handledApplicationEnvironments;
diff --git a/Engine/Source/Programs/UnrealBuildAccelerator/Common/Public/UbaSessionServer.h b/Engine/Source/Programs/UnrealBuildAccelerator/Common/Public/UbaSessionServer.h
--- a/Engine/Source/Programs/UnrealBuildAccelerator/Common/Public/UbaSessionServer.h
+++ b/Engine/Source/Programs/UnrealBuildAccelerator/Common/Public/UbaSessionServer.h
@@ -65,7 +65,7 @@ namespace uba
 	protected:
 		struct ClientSession;
 
-		void OnDisconnected(u32 clientId);
+		void OnDisconnected(const Guid& clientUid, u32 clientId);
 		bool HandleMessage(const ConnectionInfo& connectionInfo, u8 messageType, BinaryReader& reader, BinaryWriter& writer);
 
 		bool StoreCasFile(CasKey& out, const StringKey& fileNameKey, const tchar* fileName);
@@ -171,6 +171,7 @@ namespace uba
 		bool m_allowWaitOnMem = false;
 		bool m_allowKillOnMem = false;
 		bool m_remoteLogEnabled = false;
+		bool m_remoteTraceEnabled = false;
 
 		SessionServer(const SessionServer&) = delete;
 		void operator=(const SessionServer&) = delete;
diff --git a/Engine/Source/Programs/UnrealBuildAccelerator/Common/Public/UbaSessionServerCreateInfo.h b/Engine/Source/Programs/UnrealBuildAccelerator/Common/Public/UbaSessionServerCreateInfo.h
--- a/Engine/Source/Programs/UnrealBuildAccelerator/Common/Public/UbaSessionServerCreateInfo.h
+++ b/Engine/Source/Programs/UnrealBuildAccelerator/Common/Public/UbaSessionServerCreateInfo.h
@@ -22,5 +22,6 @@ namespace uba
 		bool allowWaitOnMem = false;
 		bool allowKillOnMem = false;
 		bool remoteLogEnabled = false; // If Uba is built in debug, then the logs will be sent back to server
+		bool remoteTraceEnabled = false; // If this is true, the agents will run trace and send the .uba file back to server
 	};
 }
diff --git a/Engine/Source/Programs/UnrealBuildAccelerator/Host/Private/UbaExports.cpp b/Engine/Source/Programs/UnrealBuildAccelerator/Host/Private/UbaExports.cpp
--- a/Engine/Source/Programs/UnrealBuildAccelerator/Host/Private/UbaExports.cpp
+++ b/Engine/Source/Programs/UnrealBuildAccelerator/Host/Private/UbaExports.cpp
@@ -331,6 +331,7 @@ uba::StorageClient* CreateStorageClient(uba::NetworkClient& client, const uba::t
 		info->detailedTrace = detailedTrace;
 		info->allowWaitOnMem = allowWaitOnMem;
 		info->allowKillOnMem = allowKillOnMem;
+		//info->remoteTraceEnabled = true;
 		//info->remoteLogEnabled = true;
 		return info;
 	}
