[UBA]
* Added support for handling server side messages and send response later
* Changed so storage proxy errors are muted when disconnected
* Improved error messages

diff --git a/Engine/Source/Programs/UnrealBuildAccelerator/Common/Private/UbaNetworkServer.cpp b/Engine/Source/Programs/UnrealBuildAccelerator/Common/Private/UbaNetworkServer.cpp
--- a/Engine/Source/Programs/UnrealBuildAccelerator/Common/Private/UbaNetworkServer.cpp
+++ b/Engine/Source/Programs/UnrealBuildAccelerator/Common/Private/UbaNetworkServer.cpp
@@ -443,13 +443,18 @@ namespace uba
 
 		u32 workId = server.TrackWorkStart(rec.toString(context.messageType));
 
+		MessageInfo mi;
+		mi.type = context.messageType;
+		mi.connectionId = connection.m_id;
+		mi.messageId = context.id;
+
 		if (!rec.func)
 		{
 			server.m_logger.Error(TC("WORKER FUNCTION NOT FOUND. id: %u, serviceid: %u type: %s, client: %s"), context.id, context.serviceId, rec.toString(context.messageType), GuidToString(connection.m_client->uid).str);
 			connection.SetShouldDisconnect();
 			size = ErrorSize;
 		}
-		else if (!rec.func({&connection}, context.messageType, reader, writer))
+		else if (!rec.func({&connection}, mi, reader, writer))
 		{
 			if (connection.SetShouldDisconnect())
 				server.m_logger.Error(TC("WORKER FUNCTION FAILED. id: %u, serviceid: %u type: %s, client: %s"), context.id, context.serviceId, rec.toString(context.messageType), GuidToString(connection.m_client->uid).str);
@@ -462,7 +467,7 @@ namespace uba
 
 		server.TrackWorkEnd(workId);
 
-		if (context.id)
+		if (mi.messageId)
 		{
 			UBA_ASSERT(size < (1 << 24));
 				
@@ -575,9 +580,9 @@ namespace uba
 		m_receiveTimeoutMs = info.receiveTimeoutSeconds * 1000;
 
 		m_workerFunctions[SystemServiceId].toString = GetMessageTypeToName;
-		m_workerFunctions[SystemServiceId].func = [this](const ConnectionInfo& connectionInfo, u8 messageType, BinaryReader& reader, BinaryWriter& writer)
+		m_workerFunctions[SystemServiceId].func = [this](const ConnectionInfo& connectionInfo, MessageInfo& messageInfo, BinaryReader& reader, BinaryWriter& writer)
 			{
-				return HandleSystemMessage(connectionInfo, messageType, reader, writer);
+				return HandleSystemMessage(connectionInfo, messageInfo.type, reader, writer);
 			};
 
 		if (!CreateGuid(m_uid))
@@ -911,6 +916,53 @@ namespace uba
 		return true;
 	}
 
+	bool NetworkServer::SendResponse(const MessageInfo& info, const u8* body, u32 bodySize)
+	{
+		UBA_ASSERT(info.connectionId);
+		UBA_ASSERT(info.messageId);
+
+		SCOPED_READ_LOCK(m_connectionsLock, lock);
+		Connection* found = nullptr;
+		for (auto& it : m_connections)
+		{
+			if (it.m_id != info.connectionId)
+				continue;
+			found = &it;
+			break;
+		}
+		if (!found)
+			return false;
+		Connection& connection = *found;
+		SCOPED_WRITE_LOCK(connection.m_shutdownLock, connectionLock);
+		lock.Leave();
+
+		u8 buffer[SendMaxSize];
+
+		constexpr u32 HeaderSize = 5; // 2 byte id, 3 bytes size
+
+		BinaryWriter writer(buffer, 0, sizeof_array(buffer));
+		u8* idAndSizePtr = writer.AllocWrite(HeaderSize);
+
+		idAndSizePtr[0] = info.messageId >> 8;
+		*(u32*)(idAndSizePtr + 1) = bodySize | u32(info.messageId << 24);
+
+		writer.WriteBytes(body, bodySize);
+
+		if (connection.m_cryptoKey && bodySize)
+		{
+			TimerScope ts(connection.m_encryptTimer);
+			u8* bodyData = writer.GetData() + HeaderSize;
+			if (!Crypto::Encrypt(m_logger, connection.m_cryptoKey, bodyData, bodySize))
+			{
+				connection.SetShouldDisconnect();
+				return false;
+			}
+		}
+
+		connection.Send(writer.GetData(), u32(writer.GetPosition()));
+		return true;
+	}
+
 
 	NetworkServer::Worker* NetworkServer::PopWorker()
 	{
diff --git a/Engine/Source/Programs/UnrealBuildAccelerator/Common/Private/UbaSessionServer.cpp b/Engine/Source/Programs/UnrealBuildAccelerator/Common/Private/UbaSessionServer.cpp
--- a/Engine/Source/Programs/UnrealBuildAccelerator/Common/Private/UbaSessionServer.cpp
+++ b/Engine/Source/Programs/UnrealBuildAccelerator/Common/Private/UbaSessionServer.cpp
@@ -96,9 +96,9 @@ namespace uba
 		m_server.RegisterOnClientDisconnected(ServiceId, [this](const Guid& clientUid, u32 clientId) { OnDisconnected(clientUid, clientId); });
 
 		m_server.RegisterService(ServiceId,
-			[this](const ConnectionInfo& connectionInfo, u8 messageType, BinaryReader& reader, BinaryWriter& writer)
+			[this](const ConnectionInfo& connectionInfo, MessageInfo& messageInfo, BinaryReader& reader, BinaryWriter& writer)
 			{
-				return HandleMessage(connectionInfo, messageType, reader, writer);
+				return HandleMessage(connectionInfo, messageInfo.type, reader, writer);
 			},
 			[](u8 type)
 			{
diff --git a/Engine/Source/Programs/UnrealBuildAccelerator/Common/Private/UbaStorage.cpp b/Engine/Source/Programs/UnrealBuildAccelerator/Common/Private/UbaStorage.cpp
--- a/Engine/Source/Programs/UnrealBuildAccelerator/Common/Private/UbaStorage.cpp
+++ b/Engine/Source/Programs/UnrealBuildAccelerator/Common/Private/UbaStorage.cpp
@@ -1275,7 +1275,7 @@ namespace uba
 		if (logStats)
 		{
 			u64 duration = GetTime() - startTime;
-			m_logger.Detail(TC("Database loaded from %s in %s"), fileName.data, TimeToText(duration).str);
+			m_logger.Detail(TC("Database loaded from %s in %s (contained %llu entries)"), fileName.data, TimeToText(duration).str, m_casLookup.size());
 		}
 
 		return true;
diff --git a/Engine/Source/Programs/UnrealBuildAccelerator/Common/Private/UbaStorageProxy.cpp b/Engine/Source/Programs/UnrealBuildAccelerator/Common/Private/UbaStorageProxy.cpp
--- a/Engine/Source/Programs/UnrealBuildAccelerator/Common/Private/UbaStorageProxy.cpp
+++ b/Engine/Source/Programs/UnrealBuildAccelerator/Common/Private/UbaStorageProxy.cpp
@@ -36,15 +36,17 @@ namespace uba
 			});
 
 		m_server.RegisterService(StorageServiceId,
-			[this](const ConnectionInfo& connectionInfo, u8 messageType, BinaryReader& reader, BinaryWriter& writer)
+			[this](const ConnectionInfo& connectionInfo, MessageInfo& messageInfo, BinaryReader& reader, BinaryWriter& writer)
 			{
-				return HandleMessage(connectionInfo, messageType, reader, writer);
+				return HandleMessage(connectionInfo, messageInfo, reader, writer);
 			},
 			[](u8 messageType)
 			{
 				return ToString(StorageMessageType(messageType));
 			}
 		);
+
+		m_client.RegisterOnDisconnected([this]() { m_logger.isMuted = true; });
 	}
 
 	StorageProxy::~StorageProxy()
@@ -73,14 +75,14 @@ namespace uba
 		logger.Info(TC(""));
 	}
 
-	bool StorageProxy::HandleMessage(const ConnectionInfo& connectionInfo, u8 messageType, BinaryReader& reader, BinaryWriter& writer)
+	bool StorageProxy::HandleMessage(const ConnectionInfo& connectionInfo, MessageInfo& messageInfo, BinaryReader& reader, BinaryWriter& writer)
 	{
 		LoggerWithWriter logger(g_consoleLogWriter, TC(""));
 
 		StackBinaryWriter<1024> writer2;
 		StackBinaryReader<SendMaxSize> reader2;
 
-		switch (messageType)
+		switch (messageInfo.type)
 		{
 		case StorageMessageType_Connect:
 			{
@@ -156,7 +158,7 @@ namespace uba
 
 					if (!file.memory)
 					{
-						NetworkMessage msg(m_client, ServiceId, messageType, writer2);
+						NetworkMessage msg(m_client, ServiceId, messageInfo.type, writer2);
 						writer2.WriteBool(false); // Wants proxy
 						writer2.WriteCasKey(casKey);
 						writer2.WriteString(hint);
@@ -378,7 +380,7 @@ namespace uba
 			}
 		default:
 			{
-				NetworkMessage msg(m_client, ServiceId, messageType, writer2);
+				NetworkMessage msg(m_client, ServiceId, messageInfo.type, writer2);
 				writer2.WriteBytes(reader.GetPositionData(), reader.GetLeft());
 				if (!msg.Send(reader2))
 					return false;
diff --git a/Engine/Source/Programs/UnrealBuildAccelerator/Common/Private/UbaStorageServer.cpp b/Engine/Source/Programs/UnrealBuildAccelerator/Common/Private/UbaStorageServer.cpp
--- a/Engine/Source/Programs/UnrealBuildAccelerator/Common/Private/UbaStorageServer.cpp
+++ b/Engine/Source/Programs/UnrealBuildAccelerator/Common/Private/UbaStorageServer.cpp
@@ -16,9 +16,9 @@ namespace uba
 			UBA_ASSERT(false);
 
 		m_server.RegisterService(ServiceId,
-			[this](const ConnectionInfo& connectionInfo, u8 messageType, BinaryReader& reader, BinaryWriter& writer)
+			[this](const ConnectionInfo& connectionInfo, MessageInfo& messageInfo, BinaryReader& reader, BinaryWriter& writer)
 			{
-				return HandleMessage(connectionInfo, messageType, reader, writer);
+				return HandleMessage(connectionInfo, messageInfo.type, reader, writer);
 			},
 			[](u8 messageType)
 			{
@@ -95,7 +95,10 @@ namespace uba
 
 			u64 waited = GetTime() - startTime;
 			if (TimeToMs(waited) > 4 * 60 * 1000) // 4 minutes timeout
-				return m_logger.Error(TC("Timed out waiting %s for cas %s to be transferred from remote to storage (%s)"), TimeToText(waited).str, CasKeyString(casKey).str, hint);
+			{
+				m_logger.Info(TC("Timed out waiting %s for cas %s to be transferred from remote to storage (%s)"), TimeToText(waited).str, CasKeyString(casKey).str, hint);
+				return false;
+			}
 		}
 		return waitEntry.Success;
 	}
diff --git a/Engine/Source/Programs/UnrealBuildAccelerator/Common/Public/UbaNetworkServer.h b/Engine/Source/Programs/UnrealBuildAccelerator/Common/Public/UbaNetworkServer.h
--- a/Engine/Source/Programs/UnrealBuildAccelerator/Common/Public/UbaNetworkServer.h
+++ b/Engine/Source/Programs/UnrealBuildAccelerator/Common/Public/UbaNetworkServer.h
@@ -34,6 +34,13 @@ namespace uba
 		u32 receiveTimeoutSeconds = 0;
 	};
 
+	struct MessageInfo
+	{
+		u32 connectionId;
+		u16 messageId;
+		u8 type;
+	};
+
 	class NetworkServer : public WorkManager
 	{
 	public:
@@ -50,7 +57,7 @@ namespace uba
 		void PrintSummary(Logger& logger);
 
 		using TypeToNameFunction = const tchar*(u8 type);
-		using WorkerFunction = Function<bool(const ConnectionInfo& connectionInfo, u8 messageType, BinaryReader& reader, BinaryWriter& writer)>;
+		using WorkerFunction = Function<bool(const ConnectionInfo& connectionInfo, MessageInfo& messageInfo, BinaryReader& reader, BinaryWriter& writer)>;
 		void RegisterService(u8 serviceId, const WorkerFunction& function, TypeToNameFunction* typeToNameFunc = nullptr);
 		void UnregisterService(u8 serviceId);
 
@@ -83,6 +90,7 @@ namespace uba
 
 		bool DoAdditionalWork();
 
+		bool SendResponse(const MessageInfo& info, const u8* body, u32 bodySize);
 
 		class Worker;
 	private:
diff --git a/Engine/Source/Programs/UnrealBuildAccelerator/Common/Public/UbaStorage.h b/Engine/Source/Programs/UnrealBuildAccelerator/Common/Public/UbaStorage.h
--- a/Engine/Source/Programs/UnrealBuildAccelerator/Common/Public/UbaStorage.h
+++ b/Engine/Source/Programs/UnrealBuildAccelerator/Common/Public/UbaStorage.h
@@ -152,7 +152,7 @@ namespace uba
 
 
 		WorkManager* m_workManager;
-		LoggerWithWriter m_logger;
+		MutableLogger m_logger;
 
 		StringBuffer<> m_rootDir;
 		StringBuffer<> m_tempPath;
diff --git a/Engine/Source/Programs/UnrealBuildAccelerator/Common/Public/UbaStorageProxy.h b/Engine/Source/Programs/UnrealBuildAccelerator/Common/Public/UbaStorageProxy.h
--- a/Engine/Source/Programs/UnrealBuildAccelerator/Common/Public/UbaStorageProxy.h
+++ b/Engine/Source/Programs/UnrealBuildAccelerator/Common/Public/UbaStorageProxy.h
@@ -14,6 +14,7 @@ namespace uba
 	struct BinaryReader;
 	struct BinaryWriter;
 	struct ConnectionInfo;
+	struct MessageInfo;
 
 	class StorageProxy
 	{
@@ -25,9 +26,9 @@ namespace uba
 		void PrintSummary();
 
 	protected:
-		bool HandleMessage(const ConnectionInfo& connectionInfo, u8 messageType, BinaryReader& reader, BinaryWriter& writer);
 		u16 PopId();
 		void PushId(u16 id);
+		bool HandleMessage(const ConnectionInfo& connectionInfo, MessageInfo& messageInfo, BinaryReader& reader, BinaryWriter& writer);
 		bool SendEnd(const CasKey& key);
 
 		static constexpr u8 ServiceId = StorageServiceId;
@@ -36,13 +37,13 @@ namespace uba
 		NetworkClient& m_client;
 		StorageImpl* m_localStorage;
 
-		LoggerWithWriter m_logger;
+		MutableLogger m_logger;
 
 		Guid m_storageServerUid;
 
 		TString m_name;
 
-		u32 m_inProcessClientId = 0;
+		Atomic<u32> m_inProcessClientId;
 
 		struct SegmentInFlight { u32 refCount; u32 segmentIndex; Event done; SegmentInFlight* prev; SegmentInFlight* next; };
 
