[UBA]
* Added __try/__catch to UbaAgent main and disabled the other seh exception handling to see if we get more luck with this on wine/linux
* Moved code from SessionClient ctor to start function to get around ordering issue where session can get connection callback in ctor
* Compile fix for test code caused by previous changelist

diff --git a/Engine/Source/Programs/UnrealBuildAccelerator/Agent/Private/UbaAgent.cpp b/Engine/Source/Programs/UnrealBuildAccelerator/Agent/Private/UbaAgent.cpp
--- a/Engine/Source/Programs/UnrealBuildAccelerator/Agent/Private/UbaAgent.cpp
+++ b/Engine/Source/Programs/UnrealBuildAccelerator/Agent/Private/UbaAgent.cpp
@@ -1,13 +1,14 @@
 // Copyright Epic Games, Inc. All Rights Reserved.
 
+#include "UbaAWS.h"
+#include "UbaDirectoryIterator.h"
+#include "UbaNetworkBackendMemory.h"
 #include "UbaNetworkBackendQuic.h"
 #include "UbaNetworkBackendTcp.h"
 #include "UbaNetworkServer.h"
 #include "UbaSessionClient.h"
 #include "UbaStorageClient.h"
 #include "UbaStorageProxy.h"
-#include "UbaDirectoryIterator.h"
-#include "UbaAWS.h"
 #include "UbaSentry.h"
 #include "UbaVersion.h"
 
@@ -28,8 +29,10 @@
 
 #if PLATFORM_WINDOWS
 #define UBA_AUTO_UPDATE 1
+#define UBA_USE_EXCEPTION_HANDLER 0
 #else
 #define UBA_AUTO_UPDATE 0
+#define UBA_USE_EXCEPTION_HANDLER 0
 #endif
 //#include <dbghelp.h>
 //#pragma comment (lib, "Dbghelp.lib")
@@ -124,6 +127,14 @@ namespace uba
 	}
 
 	#if PLATFORM_WINDOWS
+	int ReportSEH(LPEXCEPTION_POINTERS exceptionInfo)
+	{
+		StringBuffer<4096> assertInfo;
+		assertInfo.Appendf(TC("SEH EXCEPTION %u (0x%llx)"), exceptionInfo->ExceptionRecord->ExceptionCode, exceptionInfo->ExceptionRecord->ExceptionAddress);
+		WriteAssertInfo(assertInfo, nullptr, nullptr, 0, nullptr, 1);
+		LoggerWithWriter(g_consoleLogWriter).Log(LogEntryType_Error, assertInfo.data, assertInfo.count);
+		return EXCEPTION_CONTINUE_SEARCH;
+	}
 	BOOL ConsoleHandler(DWORD signal)
 	{
 		if (signal == CTRL_C_EVENT)
@@ -140,8 +151,12 @@ namespace uba
 
 	StringBuffer<> g_rootDir(DefaultRootDir);
 
-	//LONG WINAPI UbaUnhandledExceptionFilter(EXCEPTION_POINTERS* ExceptionInfo)
-	//{
+#if UBA_USE_EXCEPTION_HANDLER
+	LONG WINAPI UbaUnhandledExceptionFilter(EXCEPTION_POINTERS* ExceptionInfo)
+	{
+		StringBuffer<4096> assertInfo;
+		WriteAssertInfo(assertInfo, TC(""), nullptr, 0, nullptr, 1);
+		LoggerWithWriter(g_consoleLogWriter).Log(LogEntryType_Error, assertInfo.data, assertInfo.count);
 	//	time_t rawtime;
 	//	time(&rawtime);
 	//	tm ti;
@@ -157,8 +172,9 @@ namespace uba
 	//	mei.ClientPointers = TRUE;
 	//	mei.ExceptionPointers = ExceptionInfo;
 	//	MiniDumpWriteDump(GetCurrentProcess(), GetCurrentProcessId(), hFile, MiniDumpNormal, &mei, NULL, NULL);
-	//	return EXCEPTION_EXECUTE_HANDLER;
-	//}
+		return EXCEPTION_EXECUTE_HANDLER;
+	}
+#endif
 
 #if UBA_AUTO_UPDATE
 	const tchar* g_ubaAgentBinaries[] = { UBA_AGENT_EXECUTABLE, UBA_DETOURS_LIBRARY };
@@ -371,7 +387,9 @@ namespace uba
 
 	int WrappedMain(int argc, tchar*argv[])
 	{
-		//SetUnhandledExceptionFilter(UbaUnhandledExceptionFilter);
+		#if UBA_USE_EXCEPTION_HANDLER
+		SetUnhandledExceptionFilter(UbaUnhandledExceptionFilter);
+		#endif
 
 		u32 maxProcessCount = DefaultProcessorCount;
 		float mulProcessValue = 1.0f;
@@ -1132,9 +1150,11 @@ namespace uba
 						LoggerWithWriter(g_consoleLogWriter, TC("")).Info(TC("Proxy timed out waiting for zero active fetches"));
 				if (proxy.server)
 					proxy.server->StopAll();
-				client->StopAll();
-				sessionClient->Stop();
 				storageClient->StopProxy();
+				client->StopListen();
+				sessionClient->Stop();
+				sessionClient->SendSummary([&](Logger& logger) { if (proxy.server) proxy.server->PrintSummary(logger); });
+				client->StopAll();
 				loopLogging = false;
 				logLinesAvailable.Set();
 				loggingThread.Wait();
@@ -1165,6 +1185,12 @@ namespace uba
 
 			u32 connectionCount = 1;
 
+			//#if PLATFORM_WINDOWS
+			//SetPriorityClass(GetCurrentProcess(), ABOVE_NORMAL_PRIORITY_CLASS);
+			//#endif
+
+			sessionClient->Start();
+
 			while (true)
 			{
 				if (useListen)
@@ -1272,7 +1298,15 @@ namespace uba
 #if PLATFORM_WINDOWS
 int wmain(int argc, wchar_t* argv[])
 {
-	return uba::WrappedMain(argc, argv);
+	using namespace uba;
+	__try
+	{
+		return WrappedMain(argc, argv);
+	}
+	__except(ReportSEH(GetExceptionInformation()))
+	{
+		return -1;
+	}
 }
 #else
 int main(int argc, char* argv[])
diff --git a/Engine/Source/Programs/UnrealBuildAccelerator/Cli/Private/UbaCli.cpp b/Engine/Source/Programs/UnrealBuildAccelerator/Cli/Private/UbaCli.cpp
--- a/Engine/Source/Programs/UnrealBuildAccelerator/Cli/Private/UbaCli.cpp
+++ b/Engine/Source/Programs/UnrealBuildAccelerator/Cli/Private/UbaCli.cpp
@@ -1,7 +1,8 @@
 // Copyright Epic Games, Inc. All Rights Reserved.
 
-#include "UbaNetworkBackendTcp.h"
 #include "UbaFileAccessor.h"
+#include "UbaNetworkBackendTcp.h"
+#include "UbaPlatform.h"
 #include "UbaProtocol.h"
 #include "UbaScheduler.h"
 #include "UbaSessionClient.h"
@@ -9,7 +10,6 @@
 #include "UbaStorageClient.h"
 #include "UbaStorageProxy.h"
 #include "UbaStorageServer.h"
-#include "UbaPlatform.h"
 #include "UbaVersion.h"
 
 #include "UbaAWS.h"
@@ -70,6 +70,7 @@ namespace uba
 		logger.Info(TC("   -nocustomalloc          Disable custom allocator for processes. If you see odd crashes this can be tested"));
 		logger.Info(TC("   -nostdout               Disable stdout from process."));
 		logger.Info(TC("   -storeraw               Disable compression of storage. This will use more storage and might improve performance"));
+		logger.Info(TC("   -maxcpu=<number>        Max number of processes that can be started. Defaults to \"%u\" on this machine"), DefaultProcessorCount);
 		if (IsWindows)
 			logger.Info(TC("   -visualizer             Spawn a visualizer that visualizes progress"));
 		logger.Info(TC(""));
@@ -131,9 +132,10 @@ namespace uba
 		//SetUnhandledExceptionFilter(UbaUnhandledExceptionFilter);
 
 		u32 storageCapacityGb = DefaultCapacityGb;
-		StringBuffer<> workDir;
-		StringBuffer<> listenIp;
+		StringBuffer<256> workDir;
+		StringBuffer<128> listenIp;
 		u16 port = DefaultPort;
+		u32 maxProcessCount = DefaultProcessorCount;
 		bool launchVisualizer = false;
 		bool storeCompressed = true;
 		bool disableCustomAllocator = false;
@@ -252,6 +254,11 @@ namespace uba
 			{
 				disableCustomAllocator = true;
 			}
+			else if (name.Equals(TC("-maxcpu")))
+			{
+				if (!value.Parse(maxProcessCount))
+					return PrintHelp(TC("Invalid value for -maxcpu"));
+			}
 			else if (name.Equals(TC("-nostdout")))
 			{
 				enableStdOut = false;
@@ -309,9 +316,12 @@ namespace uba
 		if (deleteCas)
 		{
 			StorageImpl(StorageCreateInfo(g_rootDir.data, logWriter)).DeleteAllCas();
-			StringBuffer<> clientRootDir;
-			clientRootDir.Append(g_rootDir).Append("Agent");
-			StorageImpl(StorageCreateInfo(clientRootDir.data, logWriter)).DeleteAllCas();
+			for (u32 i=0; i!=4; ++i)
+			{
+				StringBuffer<> clientRootDir;
+				clientRootDir.Append(g_rootDir).Append("Agent").AppendValue(i);
+				StorageImpl(StorageCreateInfo(clientRootDir.data, logWriter)).DeleteAllCas();
+			}
 		}
 
 		if (checkCas)
@@ -489,6 +499,7 @@ namespace uba
 		#if UBA_DEBUG_LOG_ENABLED
 		info.remoteLogEnabled = true;
 		#endif
+		//info.remoteTraceEnabled = true;
 
 		info.deleteSessionsOlderThanSeconds = 1;
 		auto session = new SessionServer(info);
@@ -607,6 +618,7 @@ namespace uba
 				sessionClientInfo.rootDir = clientRootDir.data;
 				sessionClientInfo.deleteSessionsOlderThanSeconds = 1;
 				SessionClient sessionClient(sessionClientInfo);
+				sessionClient.Start();
 
 				//for (u32 i=0; i!=4; ++i)
 					if (!client.Connect(networkBackend, TC("127.0.0.1"), port))
@@ -626,6 +638,8 @@ namespace uba
 		{
 			SchedulerCreateInfo info(*session);
 			info.forceRemote = isRemote;
+			info.forceNative = commandType == CommandType_Native;
+			info.maxLocalProcessors = maxProcessCount;
 			Scheduler scheduler(info);
 
 			if (!scheduler.EnqueueFromFile(yamlFile))
@@ -639,11 +653,14 @@ namespace uba
 			Event finished(true);
 			scheduler.SetProcessFinishedCallback([&](const ProcessHandle& ph)
 				{
-					if (ph.GetExitCode() != 0)
+					const tchar* desc = ph.GetStartInfo().description;
+					if (ph.GetExitCode() != 0 && ph.GetExitCode() != ProcessCancelExitCode)
+					{
+						logger.Error(TC("%s - Error exit code: %u"), desc, ph.GetExitCode());
 						success = false;
+					}
 					u32 c = ++counter;
 					logger.BeginScope();
-					const tchar* desc = ph.GetStartInfo().description;
 					StringBuffer<128> extra;
 					if (ph.IsRemote())
 						extra.Append(TC(" [RemoteExecutor: ")).Append(ph.GetExecutingHost()).Append(']');
diff --git a/Engine/Source/Programs/UnrealBuildAccelerator/Common/Private/UbaSessionClient.cpp b/Engine/Source/Programs/UnrealBuildAccelerator/Common/Private/UbaSessionClient.cpp
--- a/Engine/Source/Programs/UnrealBuildAccelerator/Common/Private/UbaSessionClient.cpp
+++ b/Engine/Source/Programs/UnrealBuildAccelerator/Common/Private/UbaSessionClient.cpp
@@ -49,7 +49,15 @@ namespace uba
 		}
 
 		m_nameToHashTableMem.Init(NameToHashMemSize);
+	}
 
+	SessionClient::~SessionClient()
+	{
+		Stop();
+	}
+
+	bool SessionClient::Start()
+	{
 		m_client.RegisterOnDisconnected([this]()
 			{
 				m_loop = false;
@@ -59,11 +67,7 @@ namespace uba
 			{
 				Connect();
 			});
-	}
-
-	SessionClient::~SessionClient()
-	{
-		Stop();
+		return true;
 	}
 
 	void SessionClient::Stop()
@@ -1253,7 +1257,7 @@ namespace uba
 		m_storage.Ping();
 	}
 
-	void SessionClient::SendSummary()
+	void SessionClient::SendSummary(const Function<void(Logger&)>& extraInfo)
 	{
 		StackBinaryWriter<SendMaxSize> writer;
 		NetworkMessage msg(m_client, ServiceId, SessionMessageType_Summary, writer);
@@ -1266,6 +1270,8 @@ namespace uba
 				m_storage.PrintSummary(logger);
 				m_client.PrintSummary(logger);
 				SystemStats::GetGlobal().Print(logger, true);
+				if (extraInfo)
+					extraInfo(logger);
 			});
 
 		msg.Send();
@@ -1642,7 +1648,6 @@ namespace uba
 			m_waitToSendEvent.IsSet(100);
 		};
 
-		SendSummary();
 
 		m_client.FlushWork();
 
diff --git a/Engine/Source/Programs/UnrealBuildAccelerator/Common/Public/UbaSessionClient.h b/Engine/Source/Programs/UnrealBuildAccelerator/Common/Public/UbaSessionClient.h
--- a/Engine/Source/Programs/UnrealBuildAccelerator/Common/Public/UbaSessionClient.h
+++ b/Engine/Source/Programs/UnrealBuildAccelerator/Common/Public/UbaSessionClient.h
@@ -35,8 +35,10 @@ namespace uba
 		SessionClient(const SessionClientCreateInfo& info);
 		~SessionClient();
 
+		bool Start();
 		void Stop();
 		bool Wait(u32 milliseconds = 0xFFFFFFFF, Event* wakeupEvent = nullptr);
+		void SendSummary(const Function<void(Logger&)>& extraInfo);
 		void SetIsTerminating(const tchar* reason = TC("Terminating"), u64 delayMs = 0); // Session stores pointer directly. Can't be temporary
 		void SetMaxProcessCount(u32 count);
 
@@ -81,7 +83,6 @@ namespace uba
 		bool SendProcessAvailable(Vector<InternalProcessStartInfo>& out, float availableWeight);
 		void SendReturnProcess(u32 processId, const tchar* reason);
 		void SendPing(u64 memAvail, u64 memTotal);
-		void SendSummary();
 		void SendLogFileToServer(ProcessImpl& pi);
 		void GetLogFileName(StringBufferBase& out, const tchar* logFile, const tchar* arguments);
 		u32 CountLogLines(ProcessImpl& process);
diff --git a/Engine/Source/Programs/UnrealBuildAccelerator/Test/Private/UbaTestNetwork.h b/Engine/Source/Programs/UnrealBuildAccelerator/Test/Private/UbaTestNetwork.h
--- a/Engine/Source/Programs/UnrealBuildAccelerator/Test/Private/UbaTestNetwork.h
+++ b/Engine/Source/Programs/UnrealBuildAccelerator/Test/Private/UbaTestNetwork.h
@@ -23,6 +23,7 @@ namespace uba
 				tcp.StartListen(logger, port, TC("127.0.0.1"), [&](void* connection, const sockaddr& remoteSocketAddr)
 					{
 						logger.Info(TC("Listen got connection"));
+						tcp.SetDisconnectCallback(connection, nullptr, [](void*, const Guid&, void*) {});
 						tcp.SetRecvCallbacks(connection, &result, 1,
 							[](void* context, const Guid& connectionUid, u8* headerData, void*& outBodyContext, u8*& outBodyData, u32& outBodySize)
 							{
@@ -40,6 +41,7 @@ namespace uba
 		logger.Info(TC("Starting to Connect"));
 		tcp.Connect(logger, TC("127.0.0.1"), [&](void* connection, const sockaddr& remoteSocketAddr, bool* timedOut)
 			{
+				tcp.SetDisconnectCallback(connection, nullptr, [](void*, const Guid&, void*) {});
 				tcp.SetRecvCallbacks(connection, &tcp, 1, [](void* context, const Guid& connectionUid, u8* headerData, void*& outBodyContext, u8*& outBodyData, u32& outBodySize) { return true; }, nullptr, TC(""));
 				u8 b = 42;
 				NetworkBackend::SendContext sc;
@@ -65,10 +67,10 @@ namespace uba
 		NetworkServer server(ctorSuccess, logWriter);
 		NetworkClient client(ctorSuccess, logWriter);
 
-		server.RegisterService(1, [&](const ConnectionInfo& connectionInfo, u8 messageType, BinaryReader& reader, BinaryWriter& writer)
+		server.RegisterService(1, [&](const ConnectionInfo& connectionInfo, MessageInfo& messageInfo, BinaryReader& reader, BinaryWriter& writer)
 			{
 				logger.Info(TC("Got ping!"));
-				UBA_ASSERT(messageType == SessionMessageType_Ping);
+				UBA_ASSERT(messageInfo.type == SessionMessageType_Ping);
 				writer.WriteByte(42);
 				return true;
 			});
@@ -101,10 +103,10 @@ namespace uba
 		NetworkServer server(ctorSuccess, logWriter);
 		NetworkClient client(ctorSuccess, logWriter);
 
-		server.RegisterService(1, [&](const ConnectionInfo& connectionInfo, u8 messageType, BinaryReader& reader, BinaryWriter& writer)
+		server.RegisterService(1, [&](const ConnectionInfo& connectionInfo, MessageInfo& messageInfo, BinaryReader& reader, BinaryWriter& writer)
 			{
 				logger.Info(TC("Got ping!"));
-				UBA_ASSERT(messageType == SessionMessageType_Ping);
+				UBA_ASSERT(messageInfo.type == SessionMessageType_Ping);
 				writer.WriteByte(42);
 				return true;
 			});
@@ -146,10 +148,10 @@ namespace uba
 		NetworkServer server(ctorSuccess, logWriter);
 		NetworkClient client(ctorSuccess, logWriter);
 
-		server.RegisterService(1, [&](const ConnectionInfo& connectionInfo, u8 messageType, BinaryReader& reader, BinaryWriter& writer)
+		server.RegisterService(1, [&](const ConnectionInfo& connectionInfo, MessageInfo& messageInfo, BinaryReader& reader, BinaryWriter& writer)
 			{
 				logger.Info(TC("Got ping!"));
-				UBA_ASSERT(messageType == SessionMessageType_Ping);
+				UBA_ASSERT(messageInfo.type == SessionMessageType_Ping);
 				writer.WriteByte(42);
 				return true;
 			});
diff --git a/Engine/Source/Programs/UnrealBuildAccelerator/Test/Private/UbaTestSession.h b/Engine/Source/Programs/UnrealBuildAccelerator/Test/Private/UbaTestSession.h
--- a/Engine/Source/Programs/UnrealBuildAccelerator/Test/Private/UbaTestSession.h
+++ b/Engine/Source/Programs/UnrealBuildAccelerator/Test/Private/UbaTestSession.h
@@ -91,6 +91,8 @@ namespace uba
 		if (!DeleteAllFiles(logger, workingDir.data, false))
 			return false;
 
+		sessionClient.Start();
+
 		u16 port = 1356;
 		if (!server.StartListen(tcpBackend, port))
 			return logger.Error(TC("Failed to listen"));
