[UBA]
* Changed so in-process storage proxy is using memory network backend for communication with client

diff --git a/Engine/Source/Programs/UnrealBuildAccelerator/Agent/Private/UbaAgent.cpp b/Engine/Source/Programs/UnrealBuildAccelerator/Agent/Private/UbaAgent.cpp
--- a/Engine/Source/Programs/UnrealBuildAccelerator/Agent/Private/UbaAgent.cpp
+++ b/Engine/Source/Programs/UnrealBuildAccelerator/Agent/Private/UbaAgent.cpp
@@ -953,6 +953,7 @@ namespace uba
 			//u32 receiveTimeoutSeconds = 60; // We are sending pings at 5 seconds cadence, so timeout if no data appears on recv socket within 60 seconds
 			u32 receiveTimeoutSeconds = 0;
 
+			NetworkBackendMemory networkBackendMem(logWriter);
 			NetworkBackend* networkBackend;
 			#if UBA_USE_QUIC
 			if (useQuic)
@@ -1029,6 +1030,8 @@ namespace uba
 			struct Proxy
 			{
 				LogWriter& logWriter;
+				NetworkBackend& networkBackend;
+				NetworkBackend& networkBackendMem;
 				NetworkClient* client;
 				Event& wakeupSessionWait;
 				u32& maxConnectionCount;
@@ -1037,7 +1040,7 @@ namespace uba
 				StorageProxy* storage = nullptr;
 				StorageClient* storageClient = nullptr;
 				TString serverPrefix;
-			} proxy { g_consoleLogWriter, client, wakeupSessionWait, maxConnectionCount, targetConnectionCount };
+			} proxy { g_consoleLogWriter, *networkBackend, networkBackendMem, client, wakeupSessionWait, maxConnectionCount, targetConnectionCount };
 			auto psg = MakeGuard([&]() { delete proxy.server; });
 			auto pg = MakeGuard([&]() { delete proxy.storage; });
 
@@ -1059,19 +1062,39 @@ namespace uba
 						delete proxy.server;
 						return false;
 					}
+
+					proxy.client->RegisterOnDisconnected([&]() { proxy.server->StopAll(); });
 					proxy.storage = new StorageProxy(*proxy.server, *proxy.client, storageServerUid, TC("Wooohoo"), proxy.storageClient);
-					proxy.server->StartListen(proxy.client->GetTcpBackend(), proxyPort);
+
+					proxy.server->RegisterOnClientConnected(0, [&](const Guid& clientUid, u32 clientId) { proxy.wakeupSessionWait.Set(); });
+					proxy.server->SetWorkTracker(proxy.client->GetWorkTracker());
+					proxy.server->StartListen(proxy.networkBackendMem, proxyPort);
+					proxy.server->StartListen(proxy.networkBackend, proxyPort);
 					proxy.targetConnectionCount = proxy.maxConnectionCount;
+
 					proxy.wakeupSessionWait.Set();
 					return true;
 				};
 
+			struct NetworkBackends
+			{
+				NetworkBackend& tcp;
+				NetworkBackend& mem;
+			} backends { *networkBackend, networkBackendMem };
+
+			static auto getProxyBackend = [](void* userData, const tchar* host) -> NetworkBackend&
+				{
+					auto& backends = *(NetworkBackends*)userData;
+					return Equals(host, TC("inprocess")) ? backends.mem : backends.tcp;
+				};
 
 			StorageClientCreateInfo storageInfo(*client, g_rootDir.data);
 			storageInfo.casCapacityBytes = storageCapacity;
 			storageInfo.storeCompressed = storeCompressed;
 			storageInfo.sendCompressed = sendCompressed;
 			storageInfo.workManager = client;
+			storageInfo.getProxyBackendCallback = getProxyBackend;
+			storageInfo.getProxyBackendUserData = &backends;
 			storageInfo.startProxyCallback = startProxy;
 			storageInfo.startProxyUserData = &proxy;
 			storageInfo.zone = zone.data;
@@ -1214,6 +1237,10 @@ namespace uba
 				if (sessionClient->Wait(5*1000, &wakeupSessionWait))
 					break;
 
+				// If we are the proxy server and have external connections we lower max process count. Note that it will always have one connection which is itself
+				if (proxy.server && proxy.server->GetConnectionCount() > 1)
+					sessionClient->SetMaxProcessCount(maxProcessCount - 2);
+
 				// This is an estimation based on tcp limitations (ack and sliding windows).
 				// For every 15ms latency on "best ping") we increase targetConnectionCount up to maxConnectionCount
 				if (!storageClient->IsUsingProxy())
diff --git a/Engine/Source/Programs/UnrealBuildAccelerator/Cli/Private/UbaCli.cpp b/Engine/Source/Programs/UnrealBuildAccelerator/Cli/Private/UbaCli.cpp
--- a/Engine/Source/Programs/UnrealBuildAccelerator/Cli/Private/UbaCli.cpp
+++ b/Engine/Source/Programs/UnrealBuildAccelerator/Cli/Private/UbaCli.cpp
@@ -1,6 +1,7 @@
 // Copyright Epic Games, Inc. All Rights Reserved.
 
 #include "UbaFileAccessor.h"
+#include "UbaNetworkBackendMemory.h"
 #include "UbaNetworkBackendTcp.h"
 #include "UbaPlatform.h"
 #include "UbaProtocol.h"
@@ -470,6 +471,7 @@ namespace uba
 		#endif
 
 		NetworkBackendTcp networkBackend(logWriter);
+		NetworkBackendMemory networkBackendMem(logWriter);
 		NetworkServerCreateInfo nsci(logWriter);
 		//nsci.workerCount = 4;
 		bool ctorSuccess = true;
@@ -560,71 +562,103 @@ namespace uba
 
 		auto RunWithClient = [&](const Function<bool()>& func)
 			{
-				NetworkClient client(ctorSuccess, { logWriter });
-
-				Event wakeupSessionWait(false);
-				Atomic<u32> targetConnectionCount = 1;
-				u32 maxConnectionCount = 4;
-				struct Proxy
+				struct Client
 				{
-					LogWriter& logWriter;
-					NetworkClient* client;
-					Event& wakeupSessionWait;
-					u32& maxConnectionCount;
-					Atomic<u32>& targetConnectionCount;
-					NetworkServer* server = nullptr;
-					StorageProxy* storage = nullptr;
-					StorageClient* storageClient = nullptr;
-					TString serverPrefix;
-				} proxy { g_consoleLogWriter, &client, wakeupSessionWait, maxConnectionCount, targetConnectionCount };
-				auto psg = MakeGuard([&]() { delete proxy.server; });
-				auto pg = MakeGuard([&]() { delete proxy.storage; });
-
-				static auto startProxy = [](void* userData, u16 proxyPort, const Guid& storageServerUid)
+					bool Init(LogWriter& logWriter, NetworkBackend& nb, NetworkBackend& nbMem, u16 port, u32 maxProcessorCount, u32 index)
 					{
-						auto& proxy = *(Proxy*)userData;
-
-						NetworkServerCreateInfo nsci(proxy.logWriter);
-						nsci.workerCount = 192;
-						nsci.receiveTimeoutSeconds = 60;
-
-						StringBuffer<256> prefix;
-						prefix.Append(TC("UbaProxyServer (")).Append(GuidToString(proxy.client->GetUid()).str).Append(')');
-						proxy.serverPrefix = prefix.data;
+						networkBackend = &nb;
+						networkBackendMem = &nbMem;
 						bool ctorSuccess = true;
-						proxy.server = new NetworkServer(ctorSuccess, nsci, proxy.serverPrefix.c_str());
+						networkClient = new NetworkClient(ctorSuccess, { logWriter });
 						if (!ctorSuccess)
-						{
-							delete proxy.server;
 							return false;
-						}
-						proxy.storage = new StorageProxy(*proxy.server, *proxy.client, storageServerUid, TC("Wooohoo"), proxy.storageClient);
-						proxy.server->StartListen(proxy.client->GetTcpBackend(), proxyPort);
-						proxy.targetConnectionCount = proxy.maxConnectionCount;
-						proxy.wakeupSessionWait.Set();
+
+
+						static auto startProxy = [](void* userData, u16 proxyPort, const Guid& storageServerUid)
+							{
+								auto& client = *(Client*)userData;
+
+								NetworkServerCreateInfo nsci(client.networkClient->GetLogWriter());
+								nsci.workerCount = 192;
+								nsci.receiveTimeoutSeconds = 60;
+
+								StringBuffer<256> prefix;
+								prefix.Append(TC("UbaProxyServer (")).Append(GuidToString(client.networkClient->GetUid()).str).Append(')');
+								client.serverPrefix = prefix.data;
+								bool ctorSuccess = true;
+								client.proxyNetworkServer = new NetworkServer(ctorSuccess, nsci, client.serverPrefix.c_str());
+								if (!ctorSuccess)
+								{
+									delete client.proxyNetworkServer;
+									return false;
+								}
+								client.proxyStorage = new StorageProxy(*client.proxyNetworkServer, *client.networkClient, storageServerUid, TC("Wooohoo"), client.storageClient);
+								client.proxyNetworkServer->StartListen(*client.networkBackendMem, proxyPort);
+								client.proxyNetworkServer->StartListen(*client.networkBackend, proxyPort);
+								//proxy.targetConnectionCount = proxy.maxConnectionCount;
+								//proxy.wakeupSessionWait.Set();
+								return true;
+							};
+
+						static auto getProxyBackend = [](void* userData, const tchar* host) -> NetworkBackend&
+							{
+								auto& client = *(Client*)userData;
+								return Equals(host, TC("inprocess")) ? *client.networkBackendMem : *client.networkBackend;
+							};
+
+						StringBuffer<> clientRootDir;
+						clientRootDir.Append(g_rootDir).Append("Agent").AppendValue(index);
+						StorageClientCreateInfo storageClientInfo(*networkClient, clientRootDir.data);
+						storageClientInfo.zone = TC("FOO");
+						storageClientInfo.getProxyBackendCallback = getProxyBackend;
+						storageClientInfo.getProxyBackendUserData = this;
+						storageClientInfo.startProxyCallback = startProxy;
+						storageClientInfo.startProxyUserData = this;
+						storageClient = new StorageClient(storageClientInfo);
+
+						SessionClientCreateInfo sessionClientInfo(*storageClient, *networkClient, logWriter);
+						sessionClientInfo.maxProcessCount = maxProcessorCount;
+						sessionClientInfo.rootDir = clientRootDir.data;
+						sessionClientInfo.deleteSessionsOlderThanSeconds = 1;
+
+						sessionClient = new SessionClient(sessionClientInfo);
+						sessionClient->Start();
+						if (!networkClient->Connect(*networkBackend, TC("127.0.0.1"), port))
+							return false;//logger.Error(TC("Failed to connect"));
 						return true;
-					};
+					}
 
-				StringBuffer<> clientRootDir;
-				clientRootDir.Append(g_rootDir).Append("Agent");
-				StorageClientCreateInfo storageClientInfo(client, clientRootDir.data);
-				storageClientInfo.zone = TC("FOO");
-				storageClientInfo.startProxyCallback = startProxy;
-				storageClientInfo.startProxyUserData = &proxy;
-				StorageClient storageClient(storageClientInfo);
-
-				SessionClientCreateInfo sessionClientInfo(storageClient, client, logWriter);
-				sessionClientInfo.maxProcessCount = DefaultProcessorCount;
-				sessionClientInfo.rootDir = clientRootDir.data;
-				sessionClientInfo.deleteSessionsOlderThanSeconds = 1;
-				SessionClient sessionClient(sessionClientInfo);
-				sessionClient.Start();
-
-				//for (u32 i=0; i!=4; ++i)
-					if (!client.Connect(networkBackend, TC("127.0.0.1"), port))
-						return logger.Error(TC("Failed to connect"));
-
-				auto cg = MakeGuard([&]() { sessionClient.Stop(); client.StopAll(); });
+					~Client()
+					{
+						if (proxyNetworkServer)
+							proxyNetworkServer->StopAll();
+						sessionClient->Stop();
+						networkClient->StopAll();
+						delete proxyStorage;
+						delete proxyNetworkServer;
+						delete sessionClient;
+						delete storageClient;
+						delete networkClient;
+					}
+
+					NetworkClient* networkClient = nullptr;
+					StorageClient* storageClient = nullptr;
+					SessionClient* sessionClient = nullptr;
+
+					NetworkBackend* networkBackend = nullptr;
+					NetworkBackend* networkBackendMem = nullptr;
+					NetworkServer* proxyNetworkServer = nullptr;
+					StorageProxy* proxyStorage = nullptr;
+					TString serverPrefix;
+				};
+
+
+				Vector<Client> clients;
+				clients.resize(4);
+				u32 clientIndex = 0;
+				for (auto& c : clients)
+					if (!c.Init(logWriter, networkBackend, networkBackendMem, port, maxProcessCount/4, clientIndex++))
+						return false;
 
 				return func();
 			};
diff --git a/Engine/Source/Programs/UnrealBuildAccelerator/Common/Private/UbaNetworkClient.cpp b/Engine/Source/Programs/UnrealBuildAccelerator/Common/Private/UbaNetworkClient.cpp
--- a/Engine/Source/Programs/UnrealBuildAccelerator/Common/Private/UbaNetworkClient.cpp
+++ b/Engine/Source/Programs/UnrealBuildAccelerator/Common/Private/UbaNetworkClient.cpp
@@ -16,7 +16,6 @@ namespace uba
 	,	m_logger(info.logWriter, SetGetPrefix(name))
 	,	m_isConnected(true)
 	,	m_isOrWasConnected(true)
-	,	m_tcpBackend(new NetworkBackendTcp(info.logWriter))
 	{
 		outCtorSuccess = true;
 
@@ -41,8 +40,6 @@ namespace uba
 	{
 		StopAll();
 
-		delete m_tcpBackend;
-
 		if (m_cryptoKey)
 			Crypto::DestroyKey(m_cryptoKey);
 	}
diff --git a/Engine/Source/Programs/UnrealBuildAccelerator/Common/Private/UbaStorageClient.cpp b/Engine/Source/Programs/UnrealBuildAccelerator/Common/Private/UbaStorageClient.cpp
--- a/Engine/Source/Programs/UnrealBuildAccelerator/Common/Private/UbaStorageClient.cpp
+++ b/Engine/Source/Programs/UnrealBuildAccelerator/Common/Private/UbaStorageClient.cpp
@@ -14,6 +14,8 @@ namespace uba
 	,	m_client(info.client)
 	,	m_sendCompressed(info.sendCompressed)
 	,	m_zone(info.zone)
+	,	m_getProxyBackendCallback(info.getProxyBackendCallback)
+	,	m_getProxyBackendUserData(info.getProxyBackendUserData)
 	,	m_startProxyCallback(info.startProxyCallback)
 	,	m_startProxyUserData(info.startProxyUserData)
 	{
@@ -42,6 +44,8 @@ namespace uba
 
 				m_storageServerUid = reader.ReadGuid();
 			});
+
+		m_client.RegisterOnDisconnected([this]() { m_logger.isMuted = true; });
 	}
 
 	struct StorageClient::ProxyClient
@@ -300,14 +304,17 @@ namespace uba
 					StringBuffer<> proxyHost;
 					u16 proxyPort;
 					bool isInProcessClient = false;
-					if (reader.ReadBool())
+
+					if (reader.ReadBool()) // This will be true for only one message.. no need to guard it
 					{
 						proxyPort = reader.ReadU16();
 						if (!m_startProxyCallback(m_startProxyUserData, proxyPort, m_storageServerUid))
 						{
 							// TODO: Tell server we failed
+							m_logger.Warning(TC("Failed to create proxy server. This should never happen!"));
+							continue;
 						}
-						proxyHost.Append(TC("127.0.0.1"));
+						proxyHost.Append(TC("inprocess"));
 						isInProcessClient = true;
 					}
 					else
@@ -328,7 +335,6 @@ namespace uba
 						continue;
 					m_lastTestedProxyIp = proxyHost.data;
 
-					NetworkBackendTcp& proxyBackend = m_client.GetTcpBackend();
 					NetworkClientCreateInfo ncci(m_logger.m_writer);
 					bool ctorSuccess = true;
 					proxy = new ProxyClient(ctorSuccess, ncci);
@@ -343,6 +349,8 @@ namespace uba
 					if (!ctorSuccess)
 						continue;
 
+					NetworkBackend& proxyBackend = m_getProxyBackendCallback(m_getProxyBackendUserData, proxyHost.data);
+
 					u64 startTime = GetTime();
 					if (!proxy->client.Connect(proxyBackend, proxyHost.data, proxyPort))
 					{
@@ -350,6 +358,9 @@ namespace uba
 						continue;
 					}
 					
+					//for (u32 i=0;i!=4;++i)
+					//	proxy->client.Connect(proxyBackend, proxyHost.data, proxyPort);
+
 					u64 connectTime = GetTime() - startTime;
 					if (connectTime > MsToTime(2000))
 						m_logger.Info(TC("Took %s to connect to proxy %s:%u"), TimeToText(connectTime).str, proxyHost.data, proxyPort);
@@ -375,6 +386,7 @@ namespace uba
 					destroyProxy.Cancel();
 
 					++proxy->refCount;
+					proxy->client.SetWorkTracker(m_client.GetWorkTracker());
 					m_proxyClient = proxy;
 					continue;
 				}
@@ -425,7 +437,7 @@ namespace uba
 			// This is here just to prevent server from getting a million messages at the same time.
 			// In theory we could have 10 clients with 48 processes each where each one of the processes asks for a large file (64 messages in flight)
 			// So worst case in that scenario would be 10*48*64 = 30000 messages.
-			bool oneAtTheTime = left > client->GetMessageMaxSize() * 2;
+			bool oneAtTheTime = false;//left > client->GetMessageMaxSize() * 2;
 			if (oneAtTheTime)
 				m_retrieveOneBatchAtTheTimeLock.EnterWrite();
 			auto oatg = MakeGuard([&]() { if (oneAtTheTime) m_retrieveOneBatchAtTheTimeLock.LeaveWrite(); });
@@ -776,6 +788,13 @@ namespace uba
 		m_proxyClient->client.SendKeepAlive();
 	}
 
+	void StorageClient::PrintSummary(Logger& logger)
+	{
+		StorageImpl::PrintSummary(logger);
+		if (m_proxyClient)
+			m_proxyClient->client.PrintSummary(logger);
+	}
+
 	bool StorageClient::SendBatchMessages(NetworkClient& client, u16 fetchId, u8* slot, u64 capacity, u64 left, u32 messageMaxSize, u32& readIndex, u32& responseSize)
 	{
 		responseSize = 0;
diff --git a/Engine/Source/Programs/UnrealBuildAccelerator/Common/Private/UbaStorageServer.cpp b/Engine/Source/Programs/UnrealBuildAccelerator/Common/Private/UbaStorageServer.cpp
--- a/Engine/Source/Programs/UnrealBuildAccelerator/Common/Private/UbaStorageServer.cpp
+++ b/Engine/Source/Programs/UnrealBuildAccelerator/Common/Private/UbaStorageServer.cpp
@@ -351,8 +351,12 @@ namespace uba
 							}
 							else
 							{
+								const tchar* proxyHost = proxy.host.c_str();
+								if (connectionInfo.GetId() == proxy.clientId)
+									proxyHost = TC("inprocess");
+
 								writer.WriteBool(false);
-								writer.WriteString(proxy.host);
+								writer.WriteString(proxyHost);
 								writer.WriteU16(proxy.port);
 								if (m_trace)
 									m_trace->ProxyUsed(connectionInfo.GetId(), proxyName.data);
diff --git a/Engine/Source/Programs/UnrealBuildAccelerator/Common/Public/UbaNetworkClient.h b/Engine/Source/Programs/UnrealBuildAccelerator/Common/Public/UbaNetworkClient.h
--- a/Engine/Source/Programs/UnrealBuildAccelerator/Common/Public/UbaNetworkClient.h
+++ b/Engine/Source/Programs/UnrealBuildAccelerator/Common/Public/UbaNetworkClient.h
@@ -63,7 +63,6 @@ namespace uba
 		u64 GetMessageReceiveHeaderSize();
 		const Guid& GetUid() { return m_uid; }
 		LogWriter& GetLogWriter() { return m_logWriter; }
-		NetworkBackendTcp& GetTcpBackend() { return *m_tcpBackend; }
 		u32 GetConnectionCount() { return m_connectionCount; }
 		u64 GetTotalSentBytes() { return m_sendBytes; }
 		u64 GetTotalRecvBytes() { return m_recvBytes; }
@@ -113,7 +112,6 @@ namespace uba
 		Vector<OnDisconnectedFunction> m_onDisconnectedFunctions;
 		OnVersionMismatchFunction m_versionMismatchFunction;
 
-		NetworkBackendTcp* m_tcpBackend; // Used by proxies
 		NetworkBackend* m_listenBackend = nullptr;
 
 		ReaderWriterLock m_connectionsLock;
diff --git a/Engine/Source/Programs/UnrealBuildAccelerator/Common/Public/UbaStorageClient.h b/Engine/Source/Programs/UnrealBuildAccelerator/Common/Public/UbaStorageClient.h
--- a/Engine/Source/Programs/UnrealBuildAccelerator/Common/Public/UbaStorageClient.h
+++ b/Engine/Source/Programs/UnrealBuildAccelerator/Common/Public/UbaStorageClient.h
@@ -10,6 +10,7 @@ namespace uba
 	class NetworkClient;
 
 	using StartProxyCallback = bool(void* userData, u16 port, const Guid& storageServerUid);
+	using GetProxyBackendCallback = NetworkBackend&(void* userData, const tchar* proxyHost);
 
 	struct StorageClientCreateInfo : StorageCreateInfo
 	{
@@ -19,6 +20,8 @@ namespace uba
 		u16 proxyPort = DefaultStorageProxyPort;
 		bool sendCompressed = true;
 		
+		GetProxyBackendCallback* getProxyBackendCallback = nullptr;
+		void* getProxyBackendUserData = nullptr;
 		StartProxyCallback* startProxyCallback = nullptr;
 		void* startProxyUserData = nullptr;
 	};
@@ -48,6 +51,7 @@ namespace uba
 		virtual bool StoreCasFile(CasKey& out, StringKey fileNameKey, const tchar* fileName, FileMappingHandle mappingFile, u64 mappingOffset, u64 fileSize, const tchar* hint, bool deferCreation = false, bool keepMappingInMemory = false) override;
 		virtual bool HasCasFile(const CasKey& casKey, CasEntry** out = nullptr) override;
 		virtual void Ping() override;
+		virtual void PrintSummary(Logger& logger) override;
 
 		static bool SendBatchMessages(NetworkClient& client, u16 fetchId, u8* slot, u64 capacity, u64 left, u32 messageMaxSize, u32& readIndex, u32& responseSize);
 		static bool SendAllSegments(NetworkClient& client, u16 fetchId, u8* readBuffer, u64 left, u32 messageMaxSize);
@@ -78,6 +82,8 @@ namespace uba
 		ProxyClient* m_proxyClient = nullptr;
 		u64 m_proxyClientKeepAliveTime = 0;
 
+		GetProxyBackendCallback* m_getProxyBackendCallback = nullptr;
+		void* m_getProxyBackendUserData = nullptr;
 		StartProxyCallback* m_startProxyCallback = nullptr;
 		void* m_startProxyUserData = nullptr;
 	};
