[UBA]
* Added null guard for WriteAssertInfo text.. should have been part of earlier changelist

diff --git a/Engine/Source/Programs/UnrealBuildAccelerator/Core/Private/UbaPlatform.cpp b/Engine/Source/Programs/UnrealBuildAccelerator/Core/Private/UbaPlatform.cpp
--- a/Engine/Source/Programs/UnrealBuildAccelerator/Core/Private/UbaPlatform.cpp
+++ b/Engine/Source/Programs/UnrealBuildAccelerator/Core/Private/UbaPlatform.cpp
@@ -131,7 +131,8 @@ namespace uba
 	void WriteAssertInfo(StringBufferBase& out, const tchar* text, const char* file, u32 line, const char* expr, u32 skipCallstack = 0)
 	{
 #if PLATFORM_WINDOWS
-		out.Append(L"ASSERT: ").Append(text).Append(L"\r\n");
+		if (text)
+			out.Append(L"ASSERT: ").Append(text).Append(L"\r\n");
 		if (file)
 			out.Appendf(L"%hs:%u (%hs)", file, line, expr);
 
