[UBA]
* Fixed bug where additional work could be added without worker picking it up causing a deadlock

diff --git a/Engine/Source/Programs/UnrealBuildAccelerator/Common/Private/UbaNetworkServer.cpp b/Engine/Source/Programs/UnrealBuildAccelerator/Common/Private/UbaNetworkServer.cpp
--- a/Engine/Source/Programs/UnrealBuildAccelerator/Common/Private/UbaNetworkServer.cpp
+++ b/Engine/Source/Programs/UnrealBuildAccelerator/Common/Private/UbaNetworkServer.cpp
@@ -808,7 +808,7 @@ namespace uba
 		lock.Leave();
 
 		SCOPED_WRITE_LOCK(m_availableWorkersLock, lock2);
-		while (count-- && m_createdWorkerCount < m_maxWorkerCount)
+		while (count--)
 		{
 			Worker* worker = PopWorkerNoLock();
 			if (!worker)
