[UBA]
* Fixed missing assignment of process weight in assignment operator.. this caused remote machines to think all processes had the processor weight of 1 which is wrong for processes like cl.exe which cost 1.5 processors.

diff --git a/Engine/Source/Programs/UnrealBuildAccelerator/Common/Private/UbaProcessStartInfoHolder.h b/Engine/Source/Programs/UnrealBuildAccelerator/Common/Private/UbaProcessStartInfoHolder.h
--- a/Engine/Source/Programs/UnrealBuildAccelerator/Common/Private/UbaProcessStartInfoHolder.h
+++ b/Engine/Source/Programs/UnrealBuildAccelerator/Common/Private/UbaProcessStartInfoHolder.h
@@ -106,6 +106,8 @@ namespace uba
 			startInfo.description = description.c_str();
 			logFile = o.logFile;
 			startInfo.logFile = logFile.c_str();
+
+			weight = o.weight;
 		}
 
 		ProcessStartInfo startInfo;
