[UBA]
* Added session notification message to be able to see why client is about to go away.. hopefully this will help us identify if there are crashes happening on remotes that we don't track
* Added ApplicationRules::AllowStorageProxy and changed so .obj files are never going through the storage proxy on fetch (they are only read by one process so no point feeding them through the proxy
* Fixed the file write-through code and gave the #if 0 a name (#if UBA_USE_WRITE_THROUGH).. still disabled since it didn't give any benefits enabling it

diff --git a/Engine/Source/Programs/UnrealBuildAccelerator/Common/Private/UbaFileAccessor.cpp b/Engine/Source/Programs/UnrealBuildAccelerator/Common/Private/UbaFileAccessor.cpp
--- a/Engine/Source/Programs/UnrealBuildAccelerator/Common/Private/UbaFileAccessor.cpp
+++ b/Engine/Source/Programs/UnrealBuildAccelerator/Common/Private/UbaFileAccessor.cpp
@@ -11,6 +11,8 @@
 #include <copyfile.h>
 #endif
 
+#define UBA_USE_WRITE_THROUGH 0
+
 namespace uba
 {
 	#if PLATFORM_WINDOWS
@@ -61,6 +63,10 @@ namespace uba
 		}
 		#endif
 
+		#if UBA_USE_WRITE_THROUGH
+		flagsAndAttributes |= FILE_FLAG_NO_BUFFERING | FILE_FLAG_WRITE_THROUGH;
+		#endif
+
 		u32 createDisp = CREATE_ALWAYS;
 		u32 dwDesiredAccess = GENERIC_WRITE | DELETE;
 		if (allowRead)
@@ -135,10 +141,11 @@ namespace uba
 		if (!m_isWrite)
 			return false;
 
-#if 0
+		#if UBA_USE_WRITE_THROUGH
+		bool useWriteThrough = true;
 		u8 writeThroughBuffer[4 * 1024];
 		bool setFileSize = false;
-#endif
+		#endif
 
 #if PLATFORM_WINDOWS
 		if ((u64)m_fileHandle & OverlappedIoFlag)
@@ -189,7 +196,7 @@ namespace uba
 				u64 toWrite = Min(writeLeft, BlockSize);
 				u64 toActuallyWrite = toWrite;
 
-#if 0
+				#if UBA_USE_WRITE_THROUGH
 				if (useWriteThrough)
 				{
 					if (toWrite < BlockSize)
@@ -205,7 +212,7 @@ namespace uba
 							toWrite = toActuallyWrite;
 					}
 				}
-#endif
+				#endif
 
 				ol[index] = {};
 				ol[index].hEvent = ev[index].GetHandle();
@@ -226,9 +233,9 @@ namespace uba
 			if (!eg.Execute())
 				return false;
 
-			#if 0
+			#if UBA_USE_WRITE_THROUGH
 			if (setFileSize)
-				SetEndOfFile(logger, fileName, m_fileHandle, bufferLen);
+				SetEndOfFile(m_logger, m_fileName, m_fileHandle, dataLen);
 			#endif
 
 			return true;
@@ -242,8 +249,7 @@ namespace uba
 			u32 toWrite = u32(Min(writeLeft, 256llu * 1024 * 1024));
 			u32 toActuallyWrite = toWrite;
 
-#if 0
-			u8 writeThroughBuffer[4 * 1024];
+			#if UBA_USE_WRITE_THROUGH
 			if (useWriteThrough)
 			{
 				toActuallyWrite = (toWrite / 4096) * 4096;
@@ -254,7 +260,7 @@ namespace uba
 					setFileSize = true;
 				}
 			}
-#endif
+			#endif
 
 #if PLATFORM_WINDOWS
 			DWORD written;
@@ -282,10 +288,10 @@ namespace uba
 			pos += written;
 		}
 
-#if 0
+		#if UBA_USE_WRITE_THROUGH
 		if (setFileSize)
-			SetEndOfFile(logger, fileName, m_fileHandle, bufferLen);
-#endif
+			SetEndOfFile(m_logger, m_fileName, m_fileHandle, dataLen);
+		#endif
 
 		return true;
 	}
diff --git a/Engine/Source/Programs/UnrealBuildAccelerator/Common/Private/UbaSessionClient.cpp b/Engine/Source/Programs/UnrealBuildAccelerator/Common/Private/UbaSessionClient.cpp
--- a/Engine/Source/Programs/UnrealBuildAccelerator/Common/Private/UbaSessionClient.cpp
+++ b/Engine/Source/Programs/UnrealBuildAccelerator/Common/Private/UbaSessionClient.cpp
@@ -1,6 +1,7 @@
 // Copyright Epic Games, Inc. All Rights Reserved.
 
 #include "UbaSessionClient.h"
+#include "UbaApplicationRules.h"
 #include "UbaNetworkClient.h"
 #include "UbaNetworkMessage.h"
 #include "UbaProcess.h"
@@ -87,6 +88,12 @@ namespace uba
 	{
 		m_terminationTime = GetTime() + MsToTime(delayMs);
 		m_terminationReason = reason;
+
+		StackBinaryWriter<1024> writer;
+		NetworkMessage msg(m_client, ServiceId, SessionMessageType_Notification, writer);
+		writer.WriteU32(m_sessionId);
+		writer.WriteString(reason);
+		msg.Send();
 	}
 
 	void SessionClient::SetMaxProcessCount(u32 count)
@@ -449,7 +456,7 @@ namespace uba
 					casKey = AsCompressed(casKey, false);
 					entry.handled = true;
 					Storage::RetrieveResult result;
-					bool allowProxy = !IsRarelyRead(msg.process, fileName);
+					bool allowProxy = GetApplicationRules()[msg.process.m_rulesIndex].rules->AllowStorageProxy(fileName);
 					if (!m_storage.RetrieveCasFile(result, casKey, fileName.data, &m_fileMappingBuffer, memoryMapAlignment, allowProxy))
 						return m_logger.Error(TC("Error retrieving cas entry %s (%s)"), CasKeyString(casKey).str, fileName.data);
 					entry.success = true;
diff --git a/Engine/Source/Programs/UnrealBuildAccelerator/Common/Private/UbaSessionServer.cpp b/Engine/Source/Programs/UnrealBuildAccelerator/Common/Private/UbaSessionServer.cpp
--- a/Engine/Source/Programs/UnrealBuildAccelerator/Common/Private/UbaSessionServer.cpp
+++ b/Engine/Source/Programs/UnrealBuildAccelerator/Common/Private/UbaSessionServer.cpp
@@ -1153,6 +1153,14 @@ namespace uba
 
 				return true;
 			}
+			case SessionMessageType_Notification:
+			{
+				u32 sessionId = reader.ReadU32();
+				StringBuffer<1024> str;
+				reader.ReadString(str);
+				m_trace.SessionNotification(sessionId, str.data);
+				return true;
+			}
 			case SessionMessageType_GetNextProcess:
 			{
 				u32 processId = reader.ReadU32();
diff --git a/Engine/Source/Programs/UnrealBuildAccelerator/Common/Private/UbaTrace.cpp b/Engine/Source/Programs/UnrealBuildAccelerator/Common/Private/UbaTrace.cpp
--- a/Engine/Source/Programs/UnrealBuildAccelerator/Common/Private/UbaTrace.cpp
+++ b/Engine/Source/Programs/UnrealBuildAccelerator/Common/Private/UbaTrace.cpp
@@ -187,6 +187,13 @@ namespace uba
 		writer.WriteU32(*(u32*)&cpuLoad);
 	}
 
+	void Trace::SessionNotification(u32 sessionId, const tchar* text)
+	{
+		BEGIN_TRACE_ENTRY(TraceType_SessionNotification);
+		writer.WriteU32(sessionId);
+		writer.WriteString(text);
+	}
+
 	void Trace::SessionSummary(u32 sessionId, const u8* data, u64 dataSize)
 	{
 		BEGIN_TRACE_ENTRY(TraceType_SessionSummary);
diff --git a/Engine/Source/Programs/UnrealBuildAccelerator/Common/Public/UbaNetwork.h b/Engine/Source/Programs/UnrealBuildAccelerator/Common/Public/UbaNetwork.h
--- a/Engine/Source/Programs/UnrealBuildAccelerator/Common/Public/UbaNetwork.h
+++ b/Engine/Source/Programs/UnrealBuildAccelerator/Common/Public/UbaNetwork.h
@@ -18,7 +18,7 @@ namespace uba
 	
 	static constexpr u32 SystemNetworkVersion = 1339;
 	static constexpr u32 StorageNetworkVersion = 3;
-	static constexpr u32 SessionNetworkVersion = 28;
+	static constexpr u32 SessionNetworkVersion = 29;
 
 
 	// Messages used over network between client/server (system, storage and session)
@@ -77,6 +77,7 @@ namespace uba
 		UBA_SESSION_MESSAGE(VirtualAllocFailed) \
 		UBA_SESSION_MESSAGE(GetTraceInformation) \
 		UBA_SESSION_MESSAGE(Ping) \
+		UBA_SESSION_MESSAGE(Notification) \
 		UBA_SESSION_MESSAGE(GetNextProcess) \
 		UBA_SESSION_MESSAGE(Custom) \
 		UBA_SESSION_MESSAGE(UpdateEnvironment) \
diff --git a/Engine/Source/Programs/UnrealBuildAccelerator/Common/Public/UbaTrace.h b/Engine/Source/Programs/UnrealBuildAccelerator/Common/Public/UbaTrace.h
--- a/Engine/Source/Programs/UnrealBuildAccelerator/Common/Public/UbaTrace.h
+++ b/Engine/Source/Programs/UnrealBuildAccelerator/Common/Public/UbaTrace.h
@@ -51,9 +51,13 @@ namespace uba
 		TraceType_FileFetchLight,
 		TraceType_FileStoreLight,
 		TraceType_StatusUpdate,
+		TraceType_SessionNotification,
 	};
 
-	static constexpr u32 TraceVersion = 22;
+	using Color = u32;
+	inline Color toColor(u8 r, u8 g, u8 b) { return (r << 16) + (g << 8) + b; }
+
+	static constexpr u32 TraceVersion = 23;
 	static constexpr u32 TraceReadCompatibilityVersion = 6;
 
 	class Trace : public WorkTracker
@@ -67,6 +71,7 @@ namespace uba
 		bool StartWrite(const tchar* namedTrace, u64 traceMemCapacity = 64*1024*1024);
 		void SessionAdded(u32 sessionId, u32 clientId, const tchar* name, const tchar* info);
 		void SessionUpdate(u32 sessionId, u32 connectionCount, u64 send, u64 recv, u64 lastPing, u64 memAvail, u64 memTotal, float cpuLoad);
+		void SessionNotification(u32 sessionId, const tchar* text);
 		void SessionSummary(u32 sessionId, const u8* data, u64 dataSize);
 		void SessionDisconnect(u32 sessionId);
 		void ProcessAdded(u32 sessionId, u32 processId, const tchar* description);
diff --git a/Engine/Source/Programs/UnrealBuildAccelerator/Core/Private/UbaApplicationRules.cpp b/Engine/Source/Programs/UnrealBuildAccelerator/Core/Private/UbaApplicationRules.cpp
--- a/Engine/Source/Programs/UnrealBuildAccelerator/Core/Private/UbaApplicationRules.cpp
+++ b/Engine/Source/Programs/UnrealBuildAccelerator/Core/Private/UbaApplicationRules.cpp
@@ -85,6 +85,7 @@ namespace uba
 
 	class ApplicationRulesLinkExe : public ApplicationRulesVC
 	{
+		using Super = ApplicationRulesVC;
 	public:
 		virtual bool IsOutputFile(const tchar* file, u64 fileLen) override
 		{
@@ -120,6 +121,13 @@ namespace uba
 				|| file.EndsWith(TC(".pdb"));
 		}
 
+		virtual bool AllowStorageProxy(const StringBufferBase& file) override
+		{
+			if (file.EndsWith(TC(".obj")))
+				return false;
+			return Super::AllowStorageProxy(file);
+		}
+
 		virtual bool IsRarelyReadAfterWritten(const tchar* fileName, u64 fileNameLen) const override
 		{
 			return EndsWith(fileName, fileNameLen, TC(".pdb"))
diff --git a/Engine/Source/Programs/UnrealBuildAccelerator/Core/Public/UbaApplicationRules.h b/Engine/Source/Programs/UnrealBuildAccelerator/Core/Public/UbaApplicationRules.h
--- a/Engine/Source/Programs/UnrealBuildAccelerator/Core/Public/UbaApplicationRules.h
+++ b/Engine/Source/Programs/UnrealBuildAccelerator/Core/Public/UbaApplicationRules.h
@@ -85,6 +85,11 @@ namespace uba
 			return false;
 		}
 
+		virtual bool AllowStorageProxy(const StringBufferBase& file)
+		{
+			return !IsRarelyRead(file);
+		}
+
 		// Enable vectored exception handler. This can't be enabled for all processes because some of them actually allow
 		// access violations etc and then catch them.
 		virtual bool EnableVectoredExceptionHandler()
diff --git a/Engine/Source/Programs/UnrealBuildAccelerator/Visualizer/Private/UbaTraceReader.cpp b/Engine/Source/Programs/UnrealBuildAccelerator/Visualizer/Private/UbaTraceReader.cpp
--- a/Engine/Source/Programs/UnrealBuildAccelerator/Visualizer/Private/UbaTraceReader.cpp
+++ b/Engine/Source/Programs/UnrealBuildAccelerator/Visualizer/Private/UbaTraceReader.cpp
@@ -433,6 +433,13 @@ namespace uba
 			--out.activeSessionCount;
 			break;
 		}
+		case TraceType_SessionNotification:
+		{
+			u32 sessionIndex = reader.ReadU32();
+			TraceView::Session& session = GetSession(out, sessionIndex);
+			session.notification = reader.ReadString();
+			break;
+		}
 		case TraceType_SessionSummary:
 		{
 			u32 sessionIndex = reader.ReadU32();
diff --git a/Engine/Source/Programs/UnrealBuildAccelerator/Visualizer/Private/UbaTraceReader.h b/Engine/Source/Programs/UnrealBuildAccelerator/Visualizer/Private/UbaTraceReader.h
--- a/Engine/Source/Programs/UnrealBuildAccelerator/Visualizer/Private/UbaTraceReader.h
+++ b/Engine/Source/Programs/UnrealBuildAccelerator/Visualizer/Private/UbaTraceReader.h
@@ -89,6 +89,7 @@ namespace uba
 			Vector<TString> summary;
 			Vector<FileTransfer> fetchedFiles;
 			Vector<FileTransfer> storedFiles;
+			TString notification;
 			u64 fetchedFilesBytes = 0;
 			u64 storedFilesBytes = 0;
 			u32 maxVisibleFiles = 0;
diff --git a/Engine/Source/Programs/UnrealBuildAccelerator/Visualizer/Private/UbaVisualizer.cpp b/Engine/Source/Programs/UnrealBuildAccelerator/Visualizer/Private/UbaVisualizer.cpp
--- a/Engine/Source/Programs/UnrealBuildAccelerator/Visualizer/Private/UbaVisualizer.cpp
+++ b/Engine/Source/Programs/UnrealBuildAccelerator/Visualizer/Private/UbaVisualizer.cpp
@@ -586,6 +586,8 @@ namespace uba
 						text.Appendf(L" Mem: %ls/%ls", BytesToText(session.memTotal - memAvail).str, BytesToText(session.memTotal).str);
 					if (ping)
 						text.Appendf(L" Ping: %ls", TimeToText(ping, false, m_traceView.frequency).str);
+					if (!session.notification.empty())
+						text.Append(L" - ").Append(session.notification);
 				}
 				else
 					text.Append(L" - Disconnected");
