[UBA]
* Fixed bug in binarywriter/reader for wide characters that has a higher numeric value (like korean characters)

diff --git a/Engine/Source/Programs/UnrealBuildAccelerator/Core/Public/UbaBinaryReaderWriter.h b/Engine/Source/Programs/UnrealBuildAccelerator/Core/Public/UbaBinaryReaderWriter.h
--- a/Engine/Source/Programs/UnrealBuildAccelerator/Core/Public/UbaBinaryReaderWriter.h
+++ b/Engine/Source/Programs/UnrealBuildAccelerator/Core/Public/UbaBinaryReaderWriter.h
@@ -153,11 +153,7 @@ namespace uba
 		UBA_ASSERT(str);
 
 		#if PLATFORM_WINDOWS
-		u64 actualBytes = GetWrittenBytes(str, strLen);
-
-		Write7BitEncoded(actualBytes);
-
-		UBA_ASSERT_WRITE(actualBytes);
+		Write7BitEncoded(strLen);
 
 		for (const tchar* i = str, *e = str + strLen; i != e; ++i)
 		{
@@ -314,33 +310,29 @@ namespace uba
 		#if PLATFORM_WINDOWS
 		tchar* it = str;
 		u64 left = charLen;
-		while (left)
+		while (left--)
 		{
 			u8 a = *m_pos++;
 			if (a <= 127)
 			{
-				--left;
 				*it++ = a;
 				continue;
 			}
 			u8 b = *m_pos++;
 			if (a >= 192 && a <= 223)
 			{
-				left -= 2;
 				*it++ = (a-192)*64 + (b-128);
 				continue;
 			}
 			u8 c = *m_pos++;
 			if (a >= 224 && a <= 239)
 			{
-				left -= 3;
 				*it++ = (a-224)*4096 + (b-128)*64 + (c-128);
 				continue;
 			}
 			if (a >= 240 && a <= 253)
 			{
 				UBA_ASSERT(false); // Wide chars cannot exceed 16 bits
-				left -= 4;
 				*it++ = tchar(~0);
 				continue;
 			}
diff --git a/Engine/Source/Programs/UnrealBuildAccelerator/Test/Private/UbaTestBasics.h b/Engine/Source/Programs/UnrealBuildAccelerator/Test/Private/UbaTestBasics.h
--- a/Engine/Source/Programs/UnrealBuildAccelerator/Test/Private/UbaTestBasics.h
+++ b/Engine/Source/Programs/UnrealBuildAccelerator/Test/Private/UbaTestBasics.h
@@ -227,7 +227,7 @@ namespace uba
 	{
 		Vector<TString> arguments;
 		ParseArguments(arguments, TC("foo bar"));
-		UBA_TEST_CHECK(arguments.size() == 2, "ParseAgsuments 1 failed (%llu)", arguments.size());
+		UBA_TEST_CHECK(arguments.size() == 2, "ParseArguments 1 failed (%llu)", arguments.size());
 
 		Vector<TString> arguments2;
 		ParseArguments(arguments2, TC("\"foo\" bar"));
@@ -260,4 +260,35 @@ namespace uba
 		UBA_TEST_CHECK(Contains(arguments7[1].data(), TC("\"")), "ParseArguments 7 failed");
 		return true;
 	}
+
+	bool TestBinaryWriter(Logger& logger, const StringBufferBase& rootDir)
+	{
+		auto testString = [&](const tchar* str)
+		{
+			u8 mem[1024];
+			BinaryWriter writer(mem);
+			writer.WriteString(str);
+			BinaryReader reader(mem);
+			TString s = reader.ReadString();
+			if (s.size() != TStrlen(str))
+				return logger.Error(TC("Serialized string '%s' has wrong strlen"), str);
+			if (s != str)
+				return logger.Error(TC("Serialized string '%s' is different from source"), str);
+			return true;
+		};
+
+		if (!testString(TC("Foo")))
+			return false;
+
+		#if PLATFORM_WINDOWS
+		tchar str1[] = { 54620, 44544, 0 };
+		if (!testString(str1))
+			return false;
+		tchar str2[] = { 'f', 54620, 'o', 44544, 0 };
+		if (!testString(str2))
+			return false;
+		#endif
+
+		return true;
+	}
 }
