[UBA]
* Added custom application rule for ShaderCompileWorker to prevent it from using storage proxy for .uba.in files since they are always unique anyway

diff --git a/Engine/Source/Programs/UnrealBuildAccelerator/Core/Private/UbaApplicationRules.cpp b/Engine/Source/Programs/UnrealBuildAccelerator/Core/Private/UbaApplicationRules.cpp
--- a/Engine/Source/Programs/UnrealBuildAccelerator/Core/Private/UbaApplicationRules.cpp
+++ b/Engine/Source/Programs/UnrealBuildAccelerator/Core/Private/UbaApplicationRules.cpp
@@ -365,8 +365,19 @@ namespace uba
 #endif // #if PLATFORM_WINDOWS
 	};
 
+	class ApplicationRulesShaderCompileWorker : public ApplicationRules
+	{
+		virtual bool IsRarelyRead(const StringBufferBase& file) override
+		{
+			return file.Contains(TC(".uba."));
+		}
+	};
+
 	const RulesRec* GetApplicationRules()
 	{
+		// TODO: Add support for data driven rules.
+		// Note, they need to be possible to serialize from server to client and then from client to each detoured process
+
 		static RulesRec rules[]
 		{
 			{ TC(""),							new ApplicationRules() },		// Must be index 0
@@ -390,6 +401,7 @@ namespace uba
 			{ TC("llvm-objcopy.exe"),			new ApplicationRulesLlvmObjCopyExe() },
 			{ TC("UnrealBuildTool.dll"),		new ApplicationRulesUBTDll() },
 			{ TC("PVS-Studio.exe"),				new ApplicationRulesPVSStudio() },
+			{ TC("ShaderCompileWorker.exe"),	new ApplicationRulesShaderCompileWorker() },
 			//{ L"MSBuild.dll"),				new ApplicationRules() },
 			//{ L"BreakpadSymbolEncoder.exe"),	new ApplicationRulesClang() },
 			//{ L"cmd.exe"),		new ApplicationRules() },
