[UBA]
* Huge stability push... all (known) paths have been tested heavily on linux with tsan and every single found race condition report has been fixed. Lots of locks have been added/moved/changed and some instances of things have been leaked on purpose to prevent tsan reports during shutdown
* More efficient storage proxy implementation which immediately forward segments to clients once they are available in proxy
* Added UbaAgent -command=x which can be used to send commands to host. Supported commands are "status"which prints out status of all remote sessions. "abort/abortproxy/abortnonproxy" that can be used to abort remote sessions and "disableremote" to have host stop accepting more helpers
* Fixed scheduler::stop bug if remotes were still requesting processes
* Added support for process reuse on linux/macos
* Added support for Coordinator interface and dynamically load coordinator dll in UbaCli
* Restructured code a little bit to be able to queue up all writes in parallel
* Added Show create/write colors to visualizer (defaults to on)
* Fixed so write file times are visualized in visualizer
* Improved socket path for visualizer
* Improved a lot of error messages
* Fixed double close of memory handle in StorageServer
* Changed some ScopedWriteLock to SCOPED_WRITE_LOCK (same for read locks)
* Fixed some missing cleanup of trace view when starting a new trace view in visualizer

diff --git a/Engine/Source/Programs/UnrealBuildAccelerator/Agent/Private/UbaAgent.cpp b/Engine/Source/Programs/UnrealBuildAccelerator/Agent/Private/UbaAgent.cpp
--- a/Engine/Source/Programs/UnrealBuildAccelerator/Agent/Private/UbaAgent.cpp
+++ b/Engine/Source/Programs/UnrealBuildAccelerator/Agent/Private/UbaAgent.cpp
@@ -5,6 +5,7 @@
 #include "UbaNetworkBackendMemory.h"
 #include "UbaNetworkBackendQuic.h"
 #include "UbaNetworkBackendTcp.h"
+#include "UbaNetworkMessage.h"
 #include "UbaNetworkServer.h"
 #include "UbaSessionClient.h"
 #include "UbaStorageClient.h"
@@ -52,7 +53,7 @@ namespace uba
 	}();
 	u32				DefaultProcessorCount = []() { return GetLogicalProcessorCount(); }();
 	const tchar*	DefaultAgentName = []() { static tchar buf[256]; GetComputerNameW(buf, sizeof_array(buf)); return buf; }();
-	u32				DefaultMaxConnectionCount = 8;
+	u32				DefaultMaxConnectionCount = 4;
 
 	int PrintHelp(const tchar* message)
 	{
@@ -122,8 +123,10 @@ namespace uba
 			g_storageClient->SaveCasTable(true);
 			LoggerWithWriter(g_consoleLogWriter).Info(TC("CAS table saved..."));
 		}
-		if (g_client)
-			g_client->Disconnect();
+
+		abort();
+		//if (g_client)
+		//	g_client->Disconnect();
 	}
 
 	#if PLATFORM_WINDOWS
@@ -145,7 +148,6 @@ namespace uba
 	void ConsoleHandler(int sig)
 	{
 		CtrlBreakPressed();
-		exit(-1);
 	}
 	#endif
 
@@ -400,6 +402,7 @@ namespace uba
 		StringBuffer<256> named;
 		StringBuffer<512> relaunchPath;
 		StringBuffer<256> eventFile;
+		TString command;
 		u16 port = DefaultPort;
 		u16 proxyPort = DefaultStorageProxyPort;
 		StringBuffer<> agentName(DefaultAgentName);
@@ -657,6 +660,14 @@ namespace uba
 					return PrintHelp(TC("-zone needs a value"));
 				zone.Append(value);
 			}
+			else if (name.Equals(TC("-command")))
+			{
+				if (value.IsEmpty())
+					return PrintHelp(TC("-command needs a value"));
+				command = value.data;
+				poll = false;
+				quiet = true;
+			}
 			else if (name.Equals(TC("-?")))
 			{
 				return PrintHelp(TC(""));
@@ -765,6 +776,7 @@ namespace uba
 
 		u64 storageCapacity = u64(storageCapacityGb)*1000*1000*1000;
 
+		if (command.empty())
 		{
 			// Create a uba storage quickly just to fix non-graceful shutdowns
 			StorageCreateInfo info(g_rootDir.data, logWriter);
@@ -843,7 +855,6 @@ namespace uba
 		// if we didn't want a single version, or all xcodes, then use active xcode (useful for user running their own agents)
 		else
 		{
-			StringBuffer<512> command;
 			StringBuffer<512> xcodeSelectOutput;
 			FILE* xcodeSelect = popen("/usr/bin/xcode-select -p", "r");
 			if (xcodeSelect == nullptr || fgets(xcodeSelectOutput.data, xcodeSelectOutput.capacity, xcodeSelect) == nullptr || pclose(xcodeSelect) != 0)
@@ -971,7 +982,7 @@ namespace uba
 			bool ctorSuccess = true;
 			NetworkClient* client = new NetworkClient(ctorSuccess, ncci);
 			g_client = client;
-			auto csg = MakeGuard([&]() { g_client = nullptr; delete client; });
+			auto csg = MakeGuard([&]() { g_client = nullptr; client->Disconnect(); delete client; });
 			if (!ctorSuccess)
 				return -1;
 
@@ -1024,6 +1035,32 @@ namespace uba
 			if (exit)
 				return 0;
 
+			if (!command.empty())
+			{
+				StackBinaryWriter<128> writer;
+				NetworkMessage msg(*client, SessionServiceId, SessionMessageType_Command, writer);
+				writer.WriteString(command);
+				StackBinaryReader<8*1024> reader;
+				if (!msg.Send(reader))
+				{
+					logger.Error(TC("Failed to send command to host"));
+					return -1;
+				}
+				LoggerWithWriter commandLogger(g_consoleLogWriter, TC(""));
+				commandLogger.Info(TC("----------------------------------"));
+				while (true)
+				{
+					auto logType = (LogEntryType)reader.ReadByte();
+					if (logType == 255)
+						break;
+					TString result = reader.ReadString();
+					commandLogger.Log(logType, result.c_str(), u32(result.size()));
+				}
+				commandLogger.Info(TC("----------------------------------"));
+				return 0;
+			}
+
+
 			Event wakeupSessionWait(false);
 			Atomic<u32> targetConnectionCount = 1;
 
@@ -1036,7 +1073,7 @@ namespace uba
 				Event& wakeupSessionWait;
 				u32& maxConnectionCount;
 				Atomic<u32>& targetConnectionCount;
-				NetworkServer* server = nullptr;
+				Atomic<NetworkServer*> server;
 				StorageProxy* storage = nullptr;
 				StorageClient* storageClient = nullptr;
 				TString serverPrefix;
@@ -1056,26 +1093,28 @@ namespace uba
 					prefix.Append(TC("UbaProxyServer (")).Append(GuidToString(proxy.client->GetUid()).str).Append(')');
 					proxy.serverPrefix = prefix.data;
 					bool ctorSuccess = true;
-					proxy.server = new NetworkServer(ctorSuccess, nsci, proxy.serverPrefix.c_str());
+					auto proxyServer = new NetworkServer(ctorSuccess, nsci, proxy.serverPrefix.c_str());
 					if (!ctorSuccess)
 					{
-						delete proxy.server;
+						delete proxyServer;
 						return false;
 					}
 
-					proxy.client->RegisterOnDisconnected([&]() { proxy.server->StopAll(); });
-					proxy.storage = new StorageProxy(*proxy.server, *proxy.client, storageServerUid, TC("Wooohoo"), proxy.storageClient);
+					proxy.storage = new StorageProxy(*proxyServer, *proxy.client, storageServerUid, TC("Wooohoo"), proxy.storageClient);
 
-					proxy.server->RegisterOnClientConnected(0, [&](const Guid& clientUid, u32 clientId) { proxy.wakeupSessionWait.Set(); });
-					proxy.server->SetWorkTracker(proxy.client->GetWorkTracker());
-					proxy.server->StartListen(proxy.networkBackendMem, proxyPort);
-					proxy.server->StartListen(proxy.networkBackend, proxyPort);
+					proxyServer->RegisterOnClientConnected(0, [p = &proxy](const Guid& clientUid, u32 clientId) { p->wakeupSessionWait.Set(); });
+					proxyServer->SetWorkTracker(proxy.client->GetWorkTracker());
+					proxyServer->StartListen(proxy.networkBackendMem, proxyPort);
+					proxyServer->StartListen(proxy.networkBackend, proxyPort);
 					proxy.targetConnectionCount = proxy.maxConnectionCount;
 
+					proxy.server = proxyServer;
 					proxy.wakeupSessionWait.Set();
 					return true;
 				};
 
+			client->RegisterOnDisconnected([&]() { networkBackend->StopListen(); if (auto proxyServer = proxy.server.load()) proxyServer->DisconnectClients(); });
+
 			struct NetworkBackends
 			{
 				NetworkBackend& tcp;
@@ -1168,16 +1207,14 @@ namespace uba
 
 			auto disconnectAndStopLoggingThread = MakeGuard([&]()
 			{
-				if (proxy.storage)
-					if (!proxy.storage->Disconnect(30 * 1000))
-						LoggerWithWriter(g_consoleLogWriter, TC("")).Info(TC("Proxy timed out waiting for zero active fetches"));
-				if (proxy.server)
-					proxy.server->StopAll();
+				networkBackend->StopListen();
 				storageClient->StopProxy();
-				client->StopListen();
+				auto proxyServer = proxy.server.load();
+				if (proxyServer)
+					proxyServer->DisconnectClients();
 				sessionClient->Stop();
-				sessionClient->SendSummary([&](Logger& logger) { if (proxy.server) proxy.server->PrintSummary(logger); });
-				client->StopAll();
+				sessionClient->SendSummary([&](Logger& logger) { if (proxyServer) proxyServer->PrintSummary(logger); });
+				client->Disconnect();
 				loopLogging = false;
 				logLinesAvailable.Set();
 				loggingThread.Wait();
@@ -1212,6 +1249,7 @@ namespace uba
 			//SetPriorityClass(GetCurrentProcess(), ABOVE_NORMAL_PRIORITY_CLASS);
 			//#endif
 
+			storageClient->Start();
 			sessionClient->Start();
 
 			while (true)
@@ -1238,8 +1276,9 @@ namespace uba
 					break;
 
 				// If we are the proxy server and have external connections we lower max process count. Note that it will always have one connection which is itself
-				if (proxy.server && proxy.server->GetConnectionCount() > 1)
-					sessionClient->SetMaxProcessCount(maxProcessCount - 2);
+				if (auto proxyServer = proxy.server.load())
+					if (proxyServer->GetConnectionCount() > 1)
+						sessionClient->SetMaxProcessCount(maxProcessCount - 2);
 
 				// This is an estimation based on tcp limitations (ack and sliding windows).
 				// For every 15ms latency on "best ping") we increase targetConnectionCount up to maxConnectionCount
diff --git a/Engine/Source/Programs/UnrealBuildAccelerator/Cli/Private/UbaCli.cpp b/Engine/Source/Programs/UnrealBuildAccelerator/Cli/Private/UbaCli.cpp
--- a/Engine/Source/Programs/UnrealBuildAccelerator/Cli/Private/UbaCli.cpp
+++ b/Engine/Source/Programs/UnrealBuildAccelerator/Cli/Private/UbaCli.cpp
@@ -1,7 +1,8 @@
 // Copyright Epic Games, Inc. All Rights Reserved.
 
+#include "UbaClient.h"
+#include "UbaCoordinatorWrapper.h"
 #include "UbaFileAccessor.h"
-#include "UbaNetworkBackendMemory.h"
 #include "UbaNetworkBackendTcp.h"
 #include "UbaPlatform.h"
 #include "UbaProtocol.h"
@@ -72,8 +73,14 @@ namespace uba
 		logger.Info(TC("   -nostdout               Disable stdout from process."));
 		logger.Info(TC("   -storeraw               Disable compression of storage. This will use more storage and might improve performance"));
 		logger.Info(TC("   -maxcpu=<number>        Max number of processes that can be started. Defaults to \"%u\" on this machine"), DefaultProcessorCount);
-		if (IsWindows)
-			logger.Info(TC("   -visualizer             Spawn a visualizer that visualizes progress"));
+		logger.Info(TC("   -visualizer             Spawn a visualizer that visualizes progress"));
+		logger.Info(TC("   -coordinator=<name>     Load a UbaCoordinator<name>.dll to instantiate a coordinator to get helpers"));
+		logger.Info(TC(""));
+		logger.Info(TC("  CoordinatorOptions (if coordinator set):"));
+		logger.Info(TC("   -uri=<address>          Uri to coordinator"));
+		logger.Info(TC("   -pool=<name>            Name of helper pool inside coordinator"));
+		logger.Info(TC("   -oidc=<name>            Name of oidc"));
+		logger.Info(TC("   -maxcores=<number>      Max number of cores that will be asked for from coordinator"));
 		logger.Info(TC(""));
 		return -1;
 	}
@@ -87,6 +94,7 @@ namespace uba
 			g_storageServer->SaveCasTable(true);
 			LoggerWithWriter(g_consoleLogWriter).Info(TC("CAS table saved..."));
 		}
+		abort();
 	}
 
 	#if PLATFORM_WINDOWS
@@ -100,7 +108,6 @@ namespace uba
 	void ConsoleHandler(int sig)
 	{
 		CtrlBreakPressed();
-		exit(-1);
 	}
 	#endif
 	
@@ -135,6 +142,11 @@ namespace uba
 		u32 storageCapacityGb = DefaultCapacityGb;
 		StringBuffer<256> workDir;
 		StringBuffer<128> listenIp;
+		TString coordinatorName;
+		TString coordinatorUri = TC("https://horde.devtools.epicgames.com:443"); // TC("https://horde.devtools-dev.epicgames.com:443");
+		TString coordinatorPool = TC("UbaLinux-us-west-2"); // TC("BoxLinux");
+		TString coordinatorOidc = TC("EpicGames-Okta");
+		u32 coordinatorMaxCoreCount = 400;
 		u16 port = DefaultPort;
 		u32 maxProcessCount = DefaultProcessorCount;
 		bool launchVisualizer = false;
@@ -215,6 +227,37 @@ namespace uba
 			{
 				launchVisualizer = true;
 			}
+			else if (name.Equals(TC("-coordinator")))
+			{
+				if (value.IsEmpty())
+					return PrintHelp(TC("-coordinator needs a value"));
+				coordinatorName = value.data;
+			}
+			else if (name.Equals(TC("-uri")))
+			{
+				if (value.IsEmpty())
+					return PrintHelp(TC("-uri needs a value"));
+				coordinatorUri = value.data;
+			}
+			else if (name.Equals(TC("-pool")))
+			{
+				if (value.IsEmpty())
+					return PrintHelp(TC("-pool needs a value"));
+				coordinatorPool = value.data;
+			}
+			else if (name.Equals(TC("-oidc")))
+			{
+				if (value.IsEmpty())
+					return PrintHelp(TC("-oidc needs a value"));
+				coordinatorOidc = value.data;
+			}
+			else if (name.Equals(TC("-maxcores")))
+			{
+				if (value.IsEmpty())
+					return PrintHelp(TC("-maxcores needs a value"));
+				if (!value.Parse(coordinatorMaxCoreCount))
+					return PrintHelp(TC("Invalid value for -maxcores"));
+			}
 			else if (name.Equals(TC("-workdir")))
 			{
 				if (value.IsEmpty())
@@ -351,7 +394,8 @@ namespace uba
 			rootDir2.Append("_CHECKCAS2");
 			StorageClientCreateInfo scci(client, rootDir2.data);
 			StorageClient storageClient(scci);
-			auto g = MakeGuard([&]() { server.StopAll(); });
+			storageClient.Start();
+			auto g = MakeGuard([&]() { server.DisconnectClients(); });
 			if (!server.StartListen(networkBackend, 1347, TC("127.0.0.1")))
 				return false;
 			if (!client.Connect(networkBackend, TC("127.0.0.1"), 1347))
@@ -471,7 +515,6 @@ namespace uba
 		#endif
 
 		NetworkBackendTcp networkBackend(logWriter);
-		NetworkBackendMemory networkBackendMem(logWriter);
 		NetworkServerCreateInfo nsci(logWriter);
 		//nsci.workerCount = 4;
 		bool ctorSuccess = true;
@@ -514,7 +557,9 @@ namespace uba
 			if (!server->StartListen(networkBackend, port, listenIp.data))
 				return -1;
 		}
-		auto stopServer = MakeGuard([&]() { server->StopAll(); });
+		auto stopServer = MakeGuard([&]() { server->DisconnectClients(); });
+
+		auto stopListen = MakeGuard([&]() { networkBackend.StopListen(); });
 
 		auto RunLocal = [&](const TString& app, const TString& arg, bool enableDetour, bool trackInputs = false)
 		{
@@ -562,104 +607,16 @@ namespace uba
 
 		auto RunWithClient = [&](const Function<bool()>& func)
 			{
-				struct Client
-				{
-					bool Init(LogWriter& logWriter, NetworkBackend& nb, NetworkBackend& nbMem, u16 port, u32 maxProcessorCount, u32 index)
-					{
-						networkBackend = &nb;
-						networkBackendMem = &nbMem;
-						bool ctorSuccess = true;
-						networkClient = new NetworkClient(ctorSuccess, { logWriter });
-						if (!ctorSuccess)
-							return false;
-
-
-						static auto startProxy = [](void* userData, u16 proxyPort, const Guid& storageServerUid)
-							{
-								auto& client = *(Client*)userData;
-
-								NetworkServerCreateInfo nsci(client.networkClient->GetLogWriter());
-								nsci.workerCount = 192;
-								nsci.receiveTimeoutSeconds = 60;
-
-								StringBuffer<256> prefix;
-								prefix.Append(TC("UbaProxyServer (")).Append(GuidToString(client.networkClient->GetUid()).str).Append(')');
-								client.serverPrefix = prefix.data;
-								bool ctorSuccess = true;
-								client.proxyNetworkServer = new NetworkServer(ctorSuccess, nsci, client.serverPrefix.c_str());
-								if (!ctorSuccess)
-								{
-									delete client.proxyNetworkServer;
-									return false;
-								}
-								client.proxyStorage = new StorageProxy(*client.proxyNetworkServer, *client.networkClient, storageServerUid, TC("Wooohoo"), client.storageClient);
-								client.proxyNetworkServer->StartListen(*client.networkBackendMem, proxyPort);
-								client.proxyNetworkServer->StartListen(*client.networkBackend, proxyPort);
-								//proxy.targetConnectionCount = proxy.maxConnectionCount;
-								//proxy.wakeupSessionWait.Set();
-								return true;
-							};
-
-						static auto getProxyBackend = [](void* userData, const tchar* host) -> NetworkBackend&
-							{
-								auto& client = *(Client*)userData;
-								return Equals(host, TC("inprocess")) ? *client.networkBackendMem : *client.networkBackend;
-							};
-
-						StringBuffer<> clientRootDir;
-						clientRootDir.Append(g_rootDir).Append("Agent").AppendValue(index);
-						StorageClientCreateInfo storageClientInfo(*networkClient, clientRootDir.data);
-						storageClientInfo.zone = TC("FOO");
-						storageClientInfo.getProxyBackendCallback = getProxyBackend;
-						storageClientInfo.getProxyBackendUserData = this;
-						storageClientInfo.startProxyCallback = startProxy;
-						storageClientInfo.startProxyUserData = this;
-						storageClient = new StorageClient(storageClientInfo);
-
-						SessionClientCreateInfo sessionClientInfo(*storageClient, *networkClient, logWriter);
-						sessionClientInfo.maxProcessCount = maxProcessorCount;
-						sessionClientInfo.rootDir = clientRootDir.data;
-						sessionClientInfo.deleteSessionsOlderThanSeconds = 1;
-
-						sessionClient = new SessionClient(sessionClientInfo);
-						sessionClient->Start();
-						if (!networkClient->Connect(*networkBackend, TC("127.0.0.1"), port))
-							return false;//logger.Error(TC("Failed to connect"));
-						return true;
-					}
-
-					~Client()
-					{
-						if (proxyNetworkServer)
-							proxyNetworkServer->StopAll();
-						sessionClient->Stop();
-						networkClient->StopAll();
-						delete proxyStorage;
-						delete proxyNetworkServer;
-						delete sessionClient;
-						delete storageClient;
-						delete networkClient;
-					}
-
-					NetworkClient* networkClient = nullptr;
-					StorageClient* storageClient = nullptr;
-					SessionClient* sessionClient = nullptr;
-
-					NetworkBackend* networkBackend = nullptr;
-					NetworkBackend* networkBackendMem = nullptr;
-					NetworkServer* proxyNetworkServer = nullptr;
-					StorageProxy* proxyStorage = nullptr;
-					TString serverPrefix;
-				};
-
-
 				Vector<Client> clients;
+				auto slg = MakeGuard([&]() { networkBackend.StopListen(); });
 				clients.resize(4);
 				u32 clientIndex = 0;
 				for (auto& c : clients)
-					if (!c.Init(logWriter, networkBackend, networkBackendMem, port, maxProcessCount/4, clientIndex++))
+				{
+					ClientInitInfo cii { logWriter, networkBackend, g_rootDir.data, TC("127.0.0.1"), port, TC("DummyZone"), maxProcessCount/4, clientIndex++};
+					if (!c.Init(cii))
 						return false;
-
+				}
 				return func();
 			};
 
@@ -668,6 +625,8 @@ namespace uba
 			return RunWithClient([&]() { return RunRemote(app, arg); });
 		};
 
+		CoordinatorWrapper coordinator;
+
 		auto RunScheduler = [&](const tchar* yamlFile)
 		{
 			SchedulerCreateInfo info(*session);
@@ -684,7 +643,8 @@ namespace uba
 
 			bool success = true;
 			Atomic<u32> counter;
-			Event finished(true);
+			static Event finished(true);
+
 			scheduler.SetProcessFinishedCallback([&](const ProcessHandle& ph)
 				{
 					const tchar* desc = ph.GetStartInfo().description;
@@ -730,6 +690,31 @@ namespace uba
 				return RunQueue();
 		};
 
+
+		if (!coordinatorName.empty())
+		{
+			StringBuffer<512> coordinatorWorkDir(g_rootDir);
+			coordinatorWorkDir.EnsureEndsWithSlash().Append(coordinatorName);
+			StringBuffer<512> binariesDir;
+			if (!GetDirectoryOfCurrentModule(logger, binariesDir))
+				return false;
+
+			CoordinatorCreateInfo cinfo;
+			cinfo.workDir = coordinatorWorkDir.data;
+			cinfo.binariesDir = binariesDir.data;
+
+			// TODO: This is very horde specific.. maybe all these parameters should be a string or something
+			cinfo.uri = coordinatorUri.c_str();
+			cinfo.pool = coordinatorPool.c_str();
+			cinfo.oidc = coordinatorOidc.c_str();
+			cinfo.maxCoreCount = 400;
+			cinfo.logging = true;
+			if (!coordinator.Create(logger, coordinatorName.c_str(), cinfo, networkBackend, *server))
+				return false;
+		}
+		auto cg = MakeGuard([&]() { coordinator.Destroy(); });
+
+
 		for (u32 i=0; i!=loopCount; ++i)
 		{
 			bool success = false;
diff --git a/Engine/Source/Programs/UnrealBuildAccelerator/Cli/Private/UbaClient.h b/Engine/Source/Programs/UnrealBuildAccelerator/Cli/Private/UbaClient.h
new file mode 100644
--- /dev/null
+++ b/Engine/Source/Programs/UnrealBuildAccelerator/Cli/Private/UbaClient.h
@@ -0,0 +1,113 @@
+// Copyright Epic Games, Inc. All Rights Reserved.
+
+#include "UbaNetworkBackendMemory.h"
+#include "UbaNetworkServer.h"
+#include "UbaSessionClient.h"
+#include "UbaStorageClient.h"
+#include "UbaStorageProxy.h"
+
+namespace uba
+{
+	struct ClientInitInfo
+	{
+		LogWriter& logWriter;
+		NetworkBackend& networkBackend;
+		const tchar* rootDir = nullptr;
+		const tchar* host = nullptr;
+		u16 port = 0;
+		const tchar* zone = nullptr;
+		u32 maxProcessorCount = 0;
+		u32 index = 0;
+	};
+
+	class Client
+	{
+	public:
+		bool Init(const ClientInitInfo& info)
+		{
+			networkBackend = &info.networkBackend;
+			networkBackendMem = new NetworkBackendMemory(info.logWriter);
+			bool ctorSuccess = true;
+			networkClient = new NetworkClient(ctorSuccess, { info.logWriter });
+			if (!ctorSuccess)
+				return false;
+
+			StringBuffer<> clientRootDir;
+			clientRootDir.Append(info.rootDir).Append("Agent").AppendValue(info.index);
+			StorageClientCreateInfo storageClientInfo(*networkClient, clientRootDir.data);
+			storageClientInfo.zone = info.zone;
+			storageClientInfo.getProxyBackendCallback = [](void* ud, const tchar* h) -> NetworkBackend& { return ((Client*)ud)->GetProxyBackend(h); };
+			storageClientInfo.getProxyBackendUserData = this;
+			storageClientInfo.startProxyCallback = [](void* ud, u16 p, const Guid& ssu) { return ((Client*)ud)->StartProxy(p, ssu); };
+			storageClientInfo.startProxyUserData = this;
+			storageClient = new StorageClient(storageClientInfo);
+
+			SessionClientCreateInfo sessionClientInfo(*storageClient, *networkClient, info.logWriter);
+			sessionClientInfo.maxProcessCount = info.maxProcessorCount;
+			sessionClientInfo.rootDir = clientRootDir.data;
+			sessionClientInfo.deleteSessionsOlderThanSeconds = 1;
+
+			sessionClient = new SessionClient(sessionClientInfo);
+
+			storageClient->Start();
+			sessionClient->Start();
+			
+			return networkClient->Connect(*networkBackend, info.host, info.port);
+		}
+
+		bool StartProxy(u16 proxyPort, const Guid& storageServerUid)
+		{
+			NetworkServerCreateInfo nsci(networkClient->GetLogWriter());
+			nsci.workerCount = 192;
+			nsci.receiveTimeoutSeconds = 60;
+
+			StringBuffer<256> prefix;
+			prefix.Append(TC("UbaProxyServer (")).Append(GuidToString(networkClient->GetUid()).str).Append(')');
+			serverPrefix = prefix.data;
+			bool ctorSuccess = true;
+			proxyNetworkServer = new NetworkServer(ctorSuccess, nsci, serverPrefix.c_str());
+			if (!ctorSuccess)
+			{
+				delete proxyNetworkServer;
+				return false;
+			}
+			proxyStorage = new StorageProxy(*proxyNetworkServer, *networkClient, storageServerUid, TC("Wooohoo"), storageClient);
+			proxyNetworkServer->StartListen(*networkBackendMem, proxyPort);
+			proxyNetworkServer->StartListen(*networkBackend, proxyPort);
+			return true;
+		}
+
+		NetworkBackend& GetProxyBackend(const tchar* host)
+		{
+			return Equals(host, TC("inprocess")) ? *networkBackendMem : *networkBackend;
+		}
+
+		~Client()
+		{
+			if (proxyNetworkServer)
+				proxyNetworkServer->DisconnectClients();
+			if (storageClient)
+				storageClient->StopProxy();
+			if (sessionClient)
+				sessionClient->Stop();
+			if (networkClient)
+				networkClient->Disconnect();
+			delete proxyStorage;
+			delete proxyNetworkServer;
+			delete sessionClient;
+			delete storageClient;
+			delete networkClient;
+			delete networkBackendMem;
+		}
+
+		NetworkBackendMemory* networkBackendMem = nullptr;
+		NetworkClient* networkClient = nullptr;
+		StorageClient* storageClient = nullptr;
+		SessionClient* sessionClient = nullptr;
+
+		NetworkBackend* networkBackend = nullptr;
+		NetworkServer* proxyNetworkServer = nullptr;
+		StorageProxy* proxyStorage = nullptr;
+		TString serverPrefix;
+	};
+}
diff --git a/Engine/Source/Programs/UnrealBuildAccelerator/Common/Private/UbaLogger.cpp b/Engine/Source/Programs/UnrealBuildAccelerator/Common/Private/UbaLogger.cpp
--- a/Engine/Source/Programs/UnrealBuildAccelerator/Common/Private/UbaLogger.cpp
+++ b/Engine/Source/Programs/UnrealBuildAccelerator/Common/Private/UbaLogger.cpp
@@ -22,7 +22,7 @@ namespace uba
 
 	ANALYSIS_NORETURN void UbaAssert(const tchar* text, const char* file, u32 line, const char* expr, u32 terminateCode)
 	{
-		static ReaderWriterLock assertLock;
+		static ReaderWriterLock& assertLock = *new ReaderWriterLock(); // Leak to prevent asan annoyances during shutdown when asserts happen
 		SCOPED_WRITE_LOCK(assertLock, lock);
 
 		StringBuffer<4096> b;
@@ -180,8 +180,8 @@ namespace uba
 		HANDLE m_stdout = 0;
 		u32 m_defaultAttributes = 0;
 #endif
-	} g_consoleLogWriterImpl;
-	LogWriter& g_consoleLogWriter = g_consoleLogWriterImpl;
+	};
+	LogWriter& g_consoleLogWriter = *new ConsoleLogWriter(); // Leak to prevent asan annoyances during shutdown when asserts happen
 	thread_local u32 t_consoleLogScopeCount = 0;
 
 	class NullLogWriter : public LogWriter
diff --git a/Engine/Source/Programs/UnrealBuildAccelerator/Common/Private/UbaNetworkBackendMemory.cpp b/Engine/Source/Programs/UnrealBuildAccelerator/Common/Private/UbaNetworkBackendMemory.cpp
--- a/Engine/Source/Programs/UnrealBuildAccelerator/Common/Private/UbaNetworkBackendMemory.cpp
+++ b/Engine/Source/Programs/UnrealBuildAccelerator/Common/Private/UbaNetworkBackendMemory.cpp
@@ -45,6 +45,8 @@ namespace uba
 		auto& rc = m_connection->peer[to];
 		SCOPED_WRITE_LOCK(rc.lock, l);
 		m_connection->connected = false;
+		l.Leave();
+
 		if (auto cb = rc.disconnectCallback)
 			cb(rc.disconnectContext, m_connection->uid, connection);
 	}
@@ -114,6 +116,7 @@ namespace uba
 
 	bool NetworkBackendMemory::StartListen(Logger& logger, u16 port, const tchar* ip, const ListenConnectedFunc& connectedFunc)
 	{
+		SCOPED_WRITE_LOCK(m_connectedFuncLock, l);
 		m_connectedFunc = connectedFunc;
 		return true;
 	}
@@ -138,6 +141,7 @@ namespace uba
 		m_connection = new Connection();
 		if (!m_connectedFunc((void*)1, {}))
 			return false;
+		l.Leave();
 		return connectedFunc((void*)2, {}, timedOut);
 	}
 
diff --git a/Engine/Source/Programs/UnrealBuildAccelerator/Common/Private/UbaNetworkBackendTcp.cpp b/Engine/Source/Programs/UnrealBuildAccelerator/Common/Private/UbaNetworkBackendTcp.cpp
--- a/Engine/Source/Programs/UnrealBuildAccelerator/Common/Private/UbaNetworkBackendTcp.cpp
+++ b/Engine/Source/Programs/UnrealBuildAccelerator/Common/Private/UbaNetworkBackendTcp.cpp
@@ -77,10 +77,7 @@ namespace uba
 		void* disconnectContext = nullptr;
 		DisconnectCallback* disconnectCallback = nullptr;
 
-		#if UBA_LOCK_AROUND_SEND
 		ReaderWriterLock sendLock;
-		#endif
-
 		CriticalSection shutdownLock;
 
 		Thread recvThread;
@@ -103,6 +100,14 @@ namespace uba
 			if (int res = WSAStartup(MAKEWORD(2, 2), &wsaData))
 				return logger.Error(TC("WSAStartup failed (%d)"), res);
 		m_wsaInitDone = true;
+		#else
+		static bool initOnce = []()
+			{
+				struct sigaction sa = { { SIG_IGN } };
+				sigaction(SIGPIPE, &sa, NULL); // Needed for broken pipe that can happen if helpers crash
+				return true;
+			}();
+
 		#endif
 		return true;
 	}
@@ -163,8 +168,11 @@ namespace uba
 
 		#if UBA_LOCK_AROUND_SEND
 		SCOPED_WRITE_LOCK(conn.sendLock, lock);
+		#else
+		SCOPED_READ_LOCK(conn.sendLock, lock);
 		#endif
-
+		if (conn.socket == INVALID_SOCKET)
+			return false;
 		bool res = SendSocket(logger, conn.socket, data, dataSize);
 
 		#if UBA_LOCK_AROUND_SEND
@@ -473,8 +481,11 @@ namespace uba
 
 		ScopedCriticalSection lock2(connection.shutdownLock);
 		SOCKET s = connection.socket;
-		connection.socket = INVALID_SOCKET;
 
+		{
+			SCOPED_WRITE_LOCK(connection.sendLock, lock);
+			connection.socket = INVALID_SOCKET;
+		}
 		if (auto cb = connection.disconnectCallback)
 		{
 			auto context = connection.disconnectContext;
diff --git a/Engine/Source/Programs/UnrealBuildAccelerator/Common/Private/UbaNetworkClient.cpp b/Engine/Source/Programs/UnrealBuildAccelerator/Common/Private/UbaNetworkClient.cpp
--- a/Engine/Source/Programs/UnrealBuildAccelerator/Common/Private/UbaNetworkClient.cpp
+++ b/Engine/Source/Programs/UnrealBuildAccelerator/Common/Private/UbaNetworkClient.cpp
@@ -38,26 +38,12 @@ namespace uba
 
 	NetworkClient::~NetworkClient()
 	{
-		StopAll();
+		UBA_ASSERTF(m_connections.empty(), TC("Client still has connections (%llu). %s"), m_connections.size(), m_isDisconnecting ? TC("") : TC("Disconnect has not been called"));
 
 		if (m_cryptoKey)
 			Crypto::DestroyKey(m_cryptoKey);
 	}
 
-	void NetworkClient::StopAll()
-	{
-		StopListen();
-
-		Disconnect();
-
-		{
-			SCOPED_WRITE_LOCK(m_connectionsLock, lock);
-			m_connections.clear();
-		}
-
-		FlushWork();
-	}
-
 	bool NetworkClient::Connect(NetworkBackend& backend, const tchar* ip, u16 port, bool* timedOut)
 	{
 		return backend.Connect(m_logger, ip, [&](void* connection, const sockaddr& remoteSocketAddr, bool* timedOut)
@@ -72,6 +58,7 @@ namespace uba
 		{
 			RecvContext(NetworkClient& c, NetworkBackend& b, void* bc) : client(c), backend(b), backendConnection(bc), recvEvent(true), exitScopeEvent(true)
 			{
+				error = 255;
 			}
 
 			~RecvContext()
@@ -86,7 +73,7 @@ namespace uba
 			void* backendConnection;
 			Event recvEvent;
 			Event exitScopeEvent;
-			u8 error = 255;
+			Atomic<u8> error;
 		};
 
 		RecvContext rc(*this, backend, backendConnection);
@@ -96,7 +83,8 @@ namespace uba
 		backend.SetDisconnectCallback(backendConnection, &rc, [](void* context, const Guid& connectionUid, void* connection)
 			{
 				auto& rc = *(RecvContext*)context;
-				rc.error = 4;
+				if (rc.error == 0)
+					rc.error = 4;
 				rc.recvEvent.Set();
 				rc.exitScopeEvent.Set();
 			});
@@ -117,15 +105,11 @@ namespace uba
 				}
 
 				if (!rc.error)
-				{
-					rc.client.ConnectedCallback(rc.backend, rc.backendConnection);
-				}
-				else
-				{
-					// Zero out callbacks since we will be leaving this scope when exiting current callback and we don't want the disconnect callback
-					rc.backend.SetRecvCallbacks(rc.backendConnection, nullptr, 0, [](void*, const Guid&, u8*, void*&, u8*&, u32&) { return false; }, nullptr, TC("Null"));
-					rc.backend.SetDisconnectCallback(rc.backendConnection, nullptr, nullptr);
-				}
+					if (!rc.client.ConnectedCallback(rc.backend, rc.backendConnection))
+						rc.error = 4;
+
+				if (rc.error != 0)
+					return false;
 
 				rc.recvEvent.Set();
 				rc.exitScopeEvent.Set();
@@ -208,9 +192,18 @@ namespace uba
 	constexpr u32 SendHeaderSize = 6;
 	constexpr u32 ReceiveHeaderSize = 5;
 
-	void NetworkClient::ConnectedCallback(NetworkBackend& backend, void* backendConnection)
+	void NetworkClient::DisconnectCallback(void* context, const Guid& connectionUid, void* connection)
+	{
+		auto& c = *(Connection*)context;
+		c.owner.OnDisconnected(c);
+		c.disconnectedEvent.Set();
+	}
+
+	bool NetworkClient::ConnectedCallback(NetworkBackend& backend, void* backendConnection)
 	{
 		SCOPED_WRITE_LOCK(m_connectionsLock, lock);
+		if (m_isDisconnecting)
+			return false;
 		m_connections.emplace_back(*this);
 		Connection* connection = &m_connections.back();
 		connection->backendConnection = backendConnection;
@@ -222,13 +215,9 @@ namespace uba
 		}
 		lock.Leave();
 
-		backend.SetDisconnectCallback(backendConnection, connection, [](void* context, const Guid& connectionUid, void* connection)
-			{
-				auto& c = *(Connection*)context;
-				c.owner.OnDisconnected(c, true);
-				c.disconnectedEvent.Set();
-			});
+		backend.SetDisconnectCallback(backendConnection, connection, DisconnectCallback);
 		backend.SetRecvCallbacks(backendConnection, connection, ReceiveHeaderSize, ReceiveResponseHeader, ReceiveResponseBody, TC("ReceiveMessageResponse"));
+		return true;
 	}
 
 	bool NetworkClient::ReceiveResponseHeader(void* context, const Guid& connectionUid, u8* headerData, void*& outBodyContext, u8*& outBodyData, u32& outBodySize)
@@ -289,19 +278,26 @@ namespace uba
 
 	void NetworkClient::Disconnect()
 	{
-		SCOPED_READ_LOCK(m_connectionsLock, lock);
-		for (auto& c : m_connections)
+		m_isDisconnecting = true;
+		{
+			SCOPED_READ_LOCK(m_connectionsLock, lock);
+			for (auto& c : m_connections)
+			{
+				OnDisconnected(c);
+				c.disconnectedEvent.IsSet(~0u);
+			}
+		}
+
 		{
-			OnDisconnected(c, false);
-			c.disconnectedEvent.IsSet(~0u);
+			SCOPED_WRITE_LOCK(m_connectionsLock, lock2);
+			m_connections.clear();
 		}
+
+		FlushWork();
 	}
 
 	bool NetworkClient::StartListen(NetworkBackend& backend, u16 port)
 	{
-		UBA_ASSERT(!m_listenBackend); // Only one allowed atm
-		m_listenBackend = &backend;
-
 		backend.StartListen(m_logger, port, nullptr, [&](void* connection, const sockaddr& remoteSockAddr)
 			{
 				return AddConnection(backend, connection, nullptr);
@@ -310,13 +306,6 @@ namespace uba
 		return true;
 	}
 
-	void NetworkClient::StopListen()
-	{
-		if (m_listenBackend)
-			m_listenBackend->StopListen();
-		m_listenBackend = nullptr;
-	}
-
 	bool NetworkClient::SetConnectionCount(u32 count)
 	{
 		StackBinaryWriter<64> writer;
@@ -414,7 +403,7 @@ namespace uba
 		return m_connections.front().backend;
 	}
 
-	void NetworkClient::OnDisconnected(Connection& connection, bool calledFromReceive)
+	void NetworkClient::OnDisconnected(Connection& connection)
 	{
 		if (connection.connected.exchange(0) == 1)
 		{
@@ -525,7 +514,7 @@ namespace uba
 			TimerScope ts(m_encryptTimer);
 			if (!Crypto::Encrypt(m_logger, m_cryptoKey, data + SendHeaderSize, bodySize))
 			{
-				OnDisconnected(connection, false);
+				OnDisconnected(connection);
 				return false;
 			}
 		}
@@ -537,7 +526,7 @@ namespace uba
 			TimerScope ts(m_sendTimer);
 			if (!connection.backend->Send(m_logger, connection.backendConnection, data, sendSize, message.m_sendContext))
 			{
-				OnDisconnected(connection, false);
+				OnDisconnected(connection);
 				return false;
 			}
 		}
diff --git a/Engine/Source/Programs/UnrealBuildAccelerator/Common/Private/UbaNetworkServer.cpp b/Engine/Source/Programs/UnrealBuildAccelerator/Common/Private/UbaNetworkServer.cpp
--- a/Engine/Source/Programs/UnrealBuildAccelerator/Common/Private/UbaNetworkServer.cpp
+++ b/Engine/Source/Programs/UnrealBuildAccelerator/Common/Private/UbaNetworkServer.cpp
@@ -233,7 +233,9 @@ namespace uba
 
 			if (!conn.SendInitialResponse(0))
 				return false;
-			
+
+			SCOPED_WRITE_LOCK(conn.m_shutdownLock, shutdownLock);
+
 			SCOPED_WRITE_LOCK(server.m_clientsLock, clientsLock);
 			u32 clientId = u32(server.m_clients.size() + 1);
 			for (auto& kv : server.m_clients)
@@ -341,6 +343,7 @@ namespace uba
 			m_backend.Shutdown(m_backendConnection);
 			if (m_client && m_client->connectionCount.fetch_sub(1) == 1)
 			{
+				SCOPED_READ_LOCK(m_server.m_onDisconnectFunctionsLock, l);
 				for (auto& entry : m_server.m_onDisconnectFunctions)
 					entry.function(m_client->uid, m_client->id);
 				m_server.m_logger.Detail(TC("Client %s disconnected"), GuidToString(m_client->uid).str);
@@ -591,14 +594,12 @@ namespace uba
 
 	NetworkServer::~NetworkServer()
 	{
-		StopAll();
+		UBA_ASSERT(m_connections.empty());
+		FlushWorkers();
 	}
 
 	bool NetworkServer::StartListen(NetworkBackend& backend, u16 port, const tchar* ip, const u8* cryptoKey128)
 	{
-		//UBA_ASSERT(!m_listenBackend);
-		m_listenBackend = &backend;
-
 		if (cryptoKey128)
 		{
 			m_listenCrypto = Crypto::CreateKey(m_logger, cryptoKey128);
@@ -618,23 +619,14 @@ namespace uba
 				return AddConnection(backend, connection, remoteSockAddr, cryptoKey);
 			});
 	}
-
-	void NetworkServer::StopListen()
-	{
-		if (m_listenBackend)
-			m_listenBackend->StopListen();
-		m_listenBackend = nullptr;
-	}
-
+	
 	void NetworkServer::DisallowNewClients()
 	{
 		m_allowNewClients = false;
 	}
 
-	void NetworkServer::StopAll()
+	void NetworkServer::DisconnectClients()
 	{
-		StopListen();
-
 		{
 			SCOPED_WRITE_LOCK(m_availableWorkersLock, lock);
 			m_workersEnabled = false;
@@ -659,33 +651,24 @@ namespace uba
 
 			// If stopping connections fail we need to abort because we will most likely run into a deadlock when deleting the workers.
 			if (!success)
+			{
+				m_logger.Info(TC("Failed to stop connection(s) in a graceful way. Will abort process"));
 				abort(); // TODO: Does this produce core dump on windows?
+			}
 		}
 
+		FlushWorkers();
 
-		SCOPED_WRITE_LOCK(m_availableWorkersLock, lock);
-		while (auto worker = m_firstActiveWorker)
-		{
-			lock.Leave();
-			worker->Stop(*this);
-			lock.Enter();
-		}
-
-		auto worker = m_firstAvailableWorker;
-		while (worker)
-		{
-			auto temp = worker;
-			worker = worker->m_nextWorker;
-			delete temp;
-		}
-		m_firstAvailableWorker = nullptr;
-
+		SCOPED_WRITE_LOCK(m_connectionsLock, lock);
 		m_connections.clear();
 	}
 
 	bool NetworkServer::AddClient(NetworkBackend& backend, const tchar* ip, u16 port, const u8* cryptoKey128)
 	{
 		SCOPED_WRITE_LOCK(m_addConnectionsLock, lock);
+		if (!m_workersEnabled)
+			return false;
+
 		for (auto it = m_addConnections.begin(); it != m_addConnections.end();)
 		{
 			if (it->Wait(0))
@@ -759,7 +742,6 @@ namespace uba
 	void NetworkServer::UnregisterService(u8 serviceId)
 	{
 		SCOPED_WRITE_LOCK(m_connectionsLock, lock);
-		UBA_ASSERTF(!m_listenBackend, TC("Server still listens to new connections while service is unregistered"));
 		UBA_ASSERTF(m_connections.empty(), TC("Unregistering service while still having live connections"));
 		WorkerRec& rec = m_workerFunctions[serviceId];
 		rec.func = {};
@@ -775,18 +757,19 @@ namespace uba
 	void NetworkServer::UnregisterOnClientConnected(u8 id)
 	{
 		SCOPED_WRITE_LOCK(m_connectionsLock, lock);
-		UBA_ASSERT(!m_listenBackend);
 		UBA_ASSERT(m_connections.empty());
 		m_onConnectionFunction = {};
 	}
 
 	void NetworkServer::RegisterOnClientDisconnected(u8 id, const OnDisconnectFunction& func)
 	{
+		SCOPED_WRITE_LOCK(m_onDisconnectFunctionsLock, l);
 		m_onDisconnectFunctions.emplace_back(OnDisconnectEntry{id, func});
 	}
 
 	void NetworkServer::UnregisterOnClientDisconnected(u8 id)
 	{
+		SCOPED_WRITE_LOCK(m_onDisconnectFunctionsLock, l);
 		for (auto it = m_onDisconnectFunctions.begin(); it != m_onDisconnectFunctions.end(); ++it)
 		{
 			if (it->id != id)
@@ -808,6 +791,8 @@ namespace uba
 		lock.Leave();
 
 		SCOPED_WRITE_LOCK(m_availableWorkersLock, lock2);
+		if (!m_workersEnabled)
+			return;
 		while (count--)
 		{
 			Worker* worker = PopWorkerNoLock();
@@ -933,8 +918,6 @@ namespace uba
 		if (!found)
 			return false;
 		Connection& connection = *found;
-		SCOPED_WRITE_LOCK(connection.m_shutdownLock, connectionLock);
-		lock.Leave();
 
 		u8 buffer[SendMaxSize];
 
@@ -943,21 +926,30 @@ namespace uba
 		BinaryWriter writer(buffer, 0, sizeof_array(buffer));
 		u8* idAndSizePtr = writer.AllocWrite(HeaderSize);
 
-		idAndSizePtr[0] = info.messageId >> 8;
-		*(u32*)(idAndSizePtr + 1) = bodySize | u32(info.messageId << 24);
-
-		writer.WriteBytes(body, bodySize);
-
-		if (connection.m_cryptoKey && bodySize)
+		if (body)
 		{
-			TimerScope ts(connection.m_encryptTimer);
-			u8* bodyData = writer.GetData() + HeaderSize;
-			if (!Crypto::Encrypt(m_logger, connection.m_cryptoKey, bodyData, bodySize))
+			writer.WriteBytes(body, bodySize);
+
+			if (connection.m_cryptoKey && bodySize)
 			{
-				connection.SetShouldDisconnect();
-				return false;
+				TimerScope ts(connection.m_encryptTimer);
+				u8* bodyData = writer.GetData() + HeaderSize;
+				if (!Crypto::Encrypt(m_logger, connection.m_cryptoKey, bodyData, bodySize))
+				{
+					connection.SetShouldDisconnect();
+					return false;
+				}
 			}
 		}
+		else
+		{
+			constexpr u32 ErrorSize = 0xffffff;
+			bodySize = ErrorSize;
+			connection.SetShouldDisconnect();
+		}
+
+		idAndSizePtr[0] = info.messageId >> 8;
+		*(u32*)(idAndSizePtr + 1) = bodySize | u32(info.messageId << 24);
 
 		connection.Send(writer.GetData(), u32(writer.GetPosition()));
 		return true;
@@ -1033,6 +1025,26 @@ namespace uba
 		m_workerAvailable.Set();
 	}
 
+	void NetworkServer::FlushWorkers()
+	{
+		SCOPED_WRITE_LOCK(m_availableWorkersLock, lock);
+		while (auto worker = m_firstActiveWorker)
+		{
+			lock.Leave();
+			worker->Stop(*this);
+			lock.Enter();
+		}
+
+		auto worker = m_firstAvailableWorker;
+		while (worker)
+		{
+			auto temp = worker;
+			worker = worker->m_nextWorker;
+			delete temp;
+		}
+		m_firstAvailableWorker = nullptr;
+	}
+
 	void NetworkServer::RemoveDisconnectedConnections()
 	{
 		for (auto it=m_connections.begin(); it!=m_connections.end();)
@@ -1105,6 +1117,14 @@ namespace uba
 
 		RemoveDisconnectedConnections();
 
+		if (!m_workersEnabled)
+		{
+			// Just to prevent errors in log
+			backend.SetDisconnectCallback(backendConnection, nullptr, [](void*, const Guid&, void*) {});
+			backend.SetRecvCallbacks(backendConnection, nullptr, 0, [](void*, const Guid&, u8*, void*&, u8*&, u32&) { return false; }, nullptr, TC("Disconnecting"));
+			return false;
+		}
+
 		m_connections.emplace_back(*this, backend, backendConnection, remoteSocketAddr, cryptoKey, m_connectionIdCounter++);
 		m_maxActiveConnections = Max(m_maxActiveConnections, u32(m_connections.size()));
 		return true;
diff --git a/Engine/Source/Programs/UnrealBuildAccelerator/Common/Private/UbaProcess.cpp b/Engine/Source/Programs/UnrealBuildAccelerator/Common/Private/UbaProcess.cpp
--- a/Engine/Source/Programs/UnrealBuildAccelerator/Common/Private/UbaProcess.cpp
+++ b/Engine/Source/Programs/UnrealBuildAccelerator/Common/Private/UbaProcess.cpp
@@ -904,8 +904,10 @@ namespace uba
 					BinaryReader statsReader(statsWriter.GetData(), 0, statsWriter.GetPosition());
 
 					bool newProcess = false;
+					m_exitCode = prevExitCode;
 					m_messageSuccess = m_session.GetNextProcess(*this, newProcess, nextProcess, prevExitCode, statsReader) && m_messageSuccess;
 					writer.WriteBool(newProcess);
+					m_exitCode = ~0u;
 					if (!newProcess)
 						return true;
 
@@ -1080,6 +1082,7 @@ namespace uba
 
 	bool ProcessImpl::WriteFilesToDisk()
 	{
+		Vector<WrittenFile*> files;
 		TimerScope ts(m_processStats.writeFiles);
 		SCOPED_WRITE_LOCK(m_writtenFilesLock, lock);
 		for (auto& kv : m_writtenFiles)
@@ -1088,10 +1091,10 @@ namespace uba
 				continue;
 			kv.second.owner = nullptr;
 			if (kv.second.mappingHandle.IsValid())
-				if (!m_session.WriteFileToDisk(*this, kv.second))
-					return false;
+				files.push_back(&kv.second);
 		}
-		return true;
+
+		return m_session.WriteFilesToDisk(*this, files.data(), u32(files.size()));
 	}
 
 	const tchar* ProcessImpl::InternalGetChildLogFile(StringBufferBase& temp)
diff --git a/Engine/Source/Programs/UnrealBuildAccelerator/Common/Private/UbaScheduler.cpp b/Engine/Source/Programs/UnrealBuildAccelerator/Common/Private/UbaScheduler.cpp
--- a/Engine/Source/Programs/UnrealBuildAccelerator/Common/Private/UbaScheduler.cpp
+++ b/Engine/Source/Programs/UnrealBuildAccelerator/Common/Private/UbaScheduler.cpp
@@ -92,11 +92,12 @@ namespace uba
 		SCOPED_WRITE_LOCK(m_processEntriesLock, lock);
 		for (auto& entry : m_processEntries)
 		{
-			UBA_ASSERT(entry.status !=ProcessStatus_Running);
+			UBA_ASSERTF(entry.status !=ProcessStatus_Running, TC("Found processes in running state when stopping scheduler."));
 			delete[] entry.dependencies;
 			delete entry.info;
 		}
 		m_processEntries.clear();
+		m_processEntriesStart = 0;
 	}
 
 	void Scheduler::SetMaxLocalProcessors(u32 maxLocalProcessors)
@@ -247,8 +248,8 @@ namespace uba
 			bool canRun = true;
 			for (u32 j=0, je=entry.dependencyCount; j!=je; ++j)
 			{
-				UBA_ASSERT(j < m_processEntries.size());
 				auto depIndex = entry.dependencies[j];
+				UBA_ASSERTF(depIndex < m_processEntries.size(), TC("Found dependency on index %u but there are only %u processes registered"), depIndex, u32(m_processEntries.size()));
 				auto depStatus = processEntries[depIndex].status;
 				if (depStatus == ProcessStatus_Failed || depStatus == ProcessStatus_Skipped)
 				{
diff --git a/Engine/Source/Programs/UnrealBuildAccelerator/Common/Private/UbaSession.cpp b/Engine/Source/Programs/UnrealBuildAccelerator/Common/Private/UbaSession.cpp
--- a/Engine/Source/Programs/UnrealBuildAccelerator/Common/Private/UbaSession.cpp
+++ b/Engine/Source/Programs/UnrealBuildAccelerator/Common/Private/UbaSession.cpp
@@ -1037,6 +1037,27 @@ namespace uba
 		if (info.traceOutputFile)
 			m_traceOutputFile.Append(info.traceOutputFile);
 
+		m_fileMappingBuffer.AddTransient(TC("FileMappings"));
+
+		u64 reserveSize = CommunicationMemSize * 512;
+		m_processCommunicationAllocator.Init(CommunicationMemSize, reserveSize);
+
+		CreateProcessJobObject();
+
+		// Environment variables that should stay local when building remote (not replicated)
+		#if PLATFORM_WINDOWS
+		m_localEnvironmentVariables.insert(TC("SystemRoot"));
+		m_localEnvironmentVariables.insert(TC("SystemDrive"));
+		m_localEnvironmentVariables.insert(TC("NUMBER_OF_PROCESSORS"));
+		m_localEnvironmentVariables.insert(TC("PROCESSOR_ARCHITECTURE"));
+		m_localEnvironmentVariables.insert(TC("PROCESSOR_IDENTIFIER"));
+		m_localEnvironmentVariables.insert(TC("PROCESSOR_LEVEL"));
+		m_localEnvironmentVariables.insert(TC("PROCESSOR_REVISION"));
+		#endif
+	}
+
+	bool Session::Create(const SessionCreateInfo& info)
+	{
 		StringBuffer<> traceName;
 		if (info.traceName && *info.traceName)
 			traceName.Append(info.traceName);
@@ -1069,23 +1090,7 @@ namespace uba
 		}
 		#endif
 
-		m_fileMappingBuffer.AddTransient(TC("FileMappings"));
-
-		u64 reserveSize = CommunicationMemSize * 512;
-		m_processCommunicationAllocator.Init(CommunicationMemSize, reserveSize);
-
-		CreateProcessJobObject();
-
-		// Environment variables that should stay local when building remote (not replicated)
-		#if PLATFORM_WINDOWS
-		m_localEnvironmentVariables.insert(TC("SystemRoot"));
-		m_localEnvironmentVariables.insert(TC("SystemDrive"));
-		m_localEnvironmentVariables.insert(TC("NUMBER_OF_PROCESSORS"));
-		m_localEnvironmentVariables.insert(TC("PROCESSOR_ARCHITECTURE"));
-		m_localEnvironmentVariables.insert(TC("PROCESSOR_IDENTIFIER"));
-		m_localEnvironmentVariables.insert(TC("PROCESSOR_LEVEL"));
-		m_localEnvironmentVariables.insert(TC("PROCESSOR_REVISION"));
-		#endif
+		return true;
 	}
 
 	Session::~Session()
@@ -1476,6 +1481,10 @@ namespace uba
 #else
 		// TODO: This. Does non-windows have dlls that needs to be downloaded here?
 		out.push_back({ TString(applicationName), TString(application), S_IRUSR | S_IWUSR | S_IXUSR });
+		
+		// This code is needed if application is compiled with tsan
+		//strcpy(applicationDirEnd, "libclang_rt.tsan.so");
+		//out.push_back({ "libclang_rt.tsan.so", applicationDir, S_IRUSR | S_IWUSR });
 #endif
 		return true;
 	}
@@ -1939,93 +1948,133 @@ namespace uba
 		return true;
 	}
 
-	bool Session::WriteFileToDisk(ProcessImpl& process, WrittenFile& file)
+	bool Session::WriteFilesToDisk(ProcessImpl& process, WrittenFile** files, u32 fileCount)
 	{
-		if (ShouldWriteToDisk(file.name.c_str(), file.name.size()))
+
+		auto writeFile = [&](WrittenFile& file)
 		{
-			u64 fileSize = file.mappingWritten;
-			u8* mem = MapViewOfFile(file.mappingHandle, FILE_MAP_READ, 0, fileSize);
-			if (!mem)
-				return m_logger.Error(TC("Failed to map view of filehandle for read %s (%s)"), file.name.c_str(), LastErrorToText().data);
+			if (ShouldWriteToDisk(file.name.c_str(), file.name.size()))
+			{
+				// This is to kill I/O when writing lots of pdb/dlls in parallel
+				#if PLATFORM_WINDOWS
+				constexpr u32 bottleneckMax = 8;
+				bool shouldBottleneck = false;//useOverlap;
+				static Bottleneck bottleneck(bottleneckMax);
+				auto bng = MakeGuard([&]() { if (shouldBottleneck) bottleneck.Leave(); });
+				#endif
 
-			auto memClose = MakeGuard([&](){ UnmapViewOfFile(mem, fileSize, file.name.c_str()); });
+				u64 fileSize = file.mappingWritten;
+				u8* mem = MapViewOfFile(file.mappingHandle, FILE_MAP_READ, 0, fileSize);
+				if (!mem)
+					return m_logger.Error(TC("Failed to map view of filehandle for read %s (%s)"), file.name.c_str(), LastErrorToText().data);
 
-			// Seems like best combo (for windows at least) is to use writes with overlap and max 16 at the same time.
-			// On one machine we get twice as fast without overlap if no bottleneck. On another machine (ntfs compression on) we get twice as slow without overlap
-			// Both machines behaves well with overlap AND bottleneck. Both machine are 128 logical core thread rippers.
-			constexpr bool useFileMapForWrite = false;
-			constexpr u32 bottleneckMax = 16;
-			bool useOverlap = fileSize > 8 * 1024 * 1024;
+				PrefetchVirtualMemory(mem, fileSize);
 
+				auto memClose = MakeGuard([&](){ UnmapViewOfFile(mem, fileSize, file.name.c_str()); });
 
-			u32 attributes = DefaultAttributes();
-			if (useOverlap)
-				attributes |= FILE_FLAG_OVERLAPPED;
+				// Seems like best combo (for windows at least) is to use writes with overlap and max 16 at the same time.
+				// On one machine we get twice as fast without overlap if no bottleneck. On another machine (ntfs compression on) we get twice as slow without overlap
+				// Both machines behaves well with overlap AND bottleneck. Both machine are 128 logical core thread rippers.
+				constexpr bool useFileMapForWrite = false;
+				bool useOverlap = false;//fileSize > 8 * 1024 * 1024;
 
-			FileAccessor destinationFile(m_logger, file.name.c_str());
 
-			if (useFileMapForWrite)
-			{
-				if (!destinationFile.CreateMemoryWrite(false, attributes, fileSize, m_tempPath.data))
+				u32 attributes = DefaultAttributes();
+				if (useOverlap)
+					attributes |= FILE_FLAG_OVERLAPPED;
+
+				FileAccessor destinationFile(m_logger, file.name.c_str());
+
+				if (useFileMapForWrite)
+				{
+					if (!destinationFile.CreateMemoryWrite(false, attributes, fileSize, m_tempPath.data))
+						return false;
+					memcpy(destinationFile.GetData(), mem, fileSize);
+				}
+				else
+				{
+					if (!destinationFile.CreateWrite(false, attributes, fileSize, m_tempPath.data))
+						return false;
+
+					#if PLATFORM_WINDOWS
+					//shouldBottleneck = fileSize > 64 * 1024 * 1024;
+					//if (shouldBottleneck)
+					//	bottleneck.Enter();
+					#endif
+
+					if (!destinationFile.Write(mem, fileSize))
+						return false;
+				}
+				if (u64 time = file.lastWriteTime)
+					if (!SetFileLastWriteTime(destinationFile.GetHandle(), time))
+						return m_logger.Error(TC("Failed to set file time on filehandle for %s"), file.name.c_str());
+
+				if (!destinationFile.Close())
 					return false;
-				memcpy(destinationFile.GetData(), mem, fileSize);
 			}
 			else
 			{
-				if (!destinationFile.CreateWrite(false, attributes, fileSize, m_tempPath.data))
-					return false;
-
-				// This is to kill I/O when writing lots of pdb/dlls in parallel
-				#if PLATFORM_WINDOWS
-				static Bottleneck bottleneck(bottleneckMax);
-				BottleneckScope scope(bottleneck);
-				#endif
+				// Delete existing file to make sure it is not picked up (since it is out of date)
+				uba::DeleteFileW(file.name.c_str());
+			}
 
-				if (!destinationFile.Write(mem, fileSize))
-					return false;
+			if (IsRarelyReadAfterWritten(process, file.name.c_str(), file.name.size()) || file.mappingWritten > m_keepOutputFileMemoryMapsThreshold)
+			{
+				m_workManager->AddWork([mh = file.mappingHandle]()
+					{
+						CloseFileMapping(mh);
+					}, 1, TC("CFM"));
+				//CloseFileMapping(file.mappingHandle);
 			}
-			if (u64 time = file.lastWriteTime)
-				if (!SetFileLastWriteTime(destinationFile.GetHandle(), time))
-					return m_logger.Error(TC("Failed to set file time on filehandle for %s"), file.name.c_str());
+			else
+			{
+				StringBuffer<> name;
+				Storage::GetMappingString(name, file.mappingHandle, 0);
+				SCOPED_WRITE_LOCK(m_fileMappingTableLookupLock, lookupLock);
+				auto insres = m_fileMappingTableLookup.try_emplace(file.key);
+				FileMappingEntry& entry = insres.first->second;
+				lookupLock.Leave();
+				SCOPED_WRITE_LOCK(entry.lock, entryCs);
+				entry.handled = true;
+				entry.mapping = file.mappingHandle;
+				entry.mappingOffset = 0;
+				entry.size = file.mappingWritten;
+				entry.isDir = false;
+				entry.success = true;
+
+				SCOPED_WRITE_LOCK(m_fileMappingTableMemLock, lock);
+				BinaryWriter writer(m_fileMappingTableMem, m_fileMappingTableSize);
+				writer.WriteStringKey(file.key);
+				writer.WriteString(name);
+				writer.Write7BitEncoded(file.mappingWritten);
+				u32 newSize = (u32)writer.GetPosition();
+				m_fileMappingTableSize = (u32)newSize;
+			}
+			file.mappingHandle = {};
+			return true;
+		};
 
-			if (!destinationFile.Close())
+		for (u32 i=0; i!=fileCount; ++i)
+			if (!writeFile(*files[i]))
 				return false;
-		}
-		else
-		{
-			// Delete existing file to make sure it is not picked up (since it is out of date)
-			uba::DeleteFileW(file.name.c_str());
-		}
+		/*
+		Event events[32];
+		UBA_ASSERT(fileCount < 32);
 
-		if (IsRarelyReadAfterWritten(process, file.name.c_str(), file.name.size()) || file.mappingWritten > m_keepOutputFileMemoryMapsThreshold)
-		{
-			CloseFileMapping(file.mappingHandle);
-		}
-		else
+
+		for (u32 i=0; i!=fileCount; ++i)
 		{
-			StringBuffer<> name;
-			Storage::GetMappingString(name, file.mappingHandle, 0);
-			SCOPED_WRITE_LOCK(m_fileMappingTableLookupLock, lookupLock);
-			auto insres = m_fileMappingTableLookup.try_emplace(file.key);
-			FileMappingEntry& entry = insres.first->second;
-			lookupLock.Leave();
-			SCOPED_WRITE_LOCK(entry.lock, entryCs);
-			entry.handled = true;
-			entry.mapping = file.mappingHandle;
-			entry.mappingOffset = 0;
-			entry.size = file.mappingWritten;
-			entry.isDir = false;
-			entry.success = true;
+			events[i].Create(true);
 
-			SCOPED_WRITE_LOCK(m_fileMappingTableMemLock, lock);
-			BinaryWriter writer(m_fileMappingTableMem, m_fileMappingTableSize);
-			writer.WriteStringKey(file.key);
-			writer.WriteString(name);
-			writer.Write7BitEncoded(file.mappingWritten);
-			u32 newSize = (u32)writer.GetPosition();
-			m_fileMappingTableSize = (u32)newSize;
+			m_workManager->AddWork([&, ii = i]()
+				{
+					writeFile(*files[ii]);
+					events[ii].Set();
+				}, 1, TC(""));
 		}
-		file.mappingHandle = {};
+		for (u32 i=0; i!=fileCount; ++i)
+			events[i].IsSet();
+		*/
 		return true;
 	}
 
diff --git a/Engine/Source/Programs/UnrealBuildAccelerator/Common/Private/UbaSessionClient.cpp b/Engine/Source/Programs/UnrealBuildAccelerator/Common/Private/UbaSessionClient.cpp
--- a/Engine/Source/Programs/UnrealBuildAccelerator/Common/Private/UbaSessionClient.cpp
+++ b/Engine/Source/Programs/UnrealBuildAccelerator/Common/Private/UbaSessionClient.cpp
@@ -50,6 +50,8 @@ namespace uba
 		}
 
 		m_nameToHashTableMem.Init(NameToHashMemSize);
+
+		Create(info);
 	}
 
 	SessionClient::~SessionClient()
@@ -416,7 +418,8 @@ namespace uba
 			if (m_useStorage || memoryMapAlignment == 0)
 			{
 				bool storeUncompressed = memoryMapAlignment == 0;
-				if (!RetrieveCasFile(newCasKey, fileSize, casKey, fileName.data, storeUncompressed, !IsRarelyRead(msg.process, fileName)))
+				bool allowProxy = GetApplicationRules()[msg.process.m_rulesIndex].rules->AllowStorageProxy(fileName);
+				if (!RetrieveCasFile(newCasKey, fileSize, casKey, fileName.data, storeUncompressed, allowProxy))
 					return m_logger.Error(TC("Error retrieving cas entry %s (%s)"), CasKeyString(casKey).str, fileName.data);
 
 				#if !UBA_USE_SPARSEFILE
@@ -801,7 +804,7 @@ namespace uba
 		return true;
 	}
 
-	bool SessionClient::WriteFileToDisk(ProcessImpl& process, WrittenFile& file)
+	bool SessionClient::WriteFilesToDisk(ProcessImpl& process, WrittenFile** files, u32 fileCount)
 	{
 		// Do nothing, we will send the data to the host when process is finished
 		return true;
@@ -1262,6 +1265,9 @@ namespace uba
 			m_bestPing = m_lastPing;
 
 		m_storage.Ping();
+
+		if (reader.ReadBool()) // abort
+			abort();
 	}
 
 	void SessionClient::SendSummary(const Function<void(Logger&)>& extraInfo)
diff --git a/Engine/Source/Programs/UnrealBuildAccelerator/Common/Private/UbaSessionServer.cpp b/Engine/Source/Programs/UnrealBuildAccelerator/Common/Private/UbaSessionServer.cpp
--- a/Engine/Source/Programs/UnrealBuildAccelerator/Common/Private/UbaSessionServer.cpp
+++ b/Engine/Source/Programs/UnrealBuildAccelerator/Common/Private/UbaSessionServer.cpp
@@ -71,7 +71,7 @@ namespace uba
 		SessionServer* m_server;
 		ReaderWriterLock m_exitedLock;
 		u32 m_processId;
-		u32 m_exitCode = ~u32(0);
+		u32 m_exitCode = ~0u;
 		u64 m_processorTime = 0;
 		u64 m_wallTime = 0;
 		Event m_done;
@@ -170,6 +170,8 @@ namespace uba
 		#else
 		m_detoursLibrary = detoursFile.data;
 		#endif
+
+		Create(info);
 	}
 
 	SessionServer::~SessionServer()
@@ -297,11 +299,13 @@ namespace uba
 			if (!m_remoteExecutionEnabled)
 			{
 				m_logger.Info(TC("Process queued for remote but remote execution was disabled, returning process to queue"));
+				ProcessHandle h(remoteProcess);
 				m_remoteProcessReturnedEvent(*remoteProcess);
 			}
 			else if (!m_connectionCount)
 			{
 				m_logger.Info(TC("Process queued for remote but there are no active connections, returning process to queue"));
+				ProcessHandle h(remoteProcess);
 				m_remoteProcessReturnedEvent(*remoteProcess);
 			}
 		}
@@ -553,7 +557,8 @@ namespace uba
 							m_storage.StoreCasFile(m_agentBinaryKey, dir.data, CasKeyZero, true);
 						}
 						dir.Resize(dirCount).Append(PathSeparator).Append(UBA_DETOURS_LIBRARY);
-						m_storage.StoreCasFile(m_detoursBinaryKey, dir.data, CasKeyZero, true);
+						if (!m_storage.StoreCasFile(m_detoursBinaryKey, dir.data, CasKeyZero, true))
+							return m_logger.Error(TC("Failed to create cas for %s"), dir.data);
 						UBA_ASSERT(m_detoursBinaryKey != CasKeyZero);
 					}
 				}
@@ -778,8 +783,10 @@ namespace uba
 					{
 						StringBuffer<> logPath;
 						logPath.Append(m_sessionLogDir).Append(destination.data + 5);
-						m_storage.CopyOrLink(casKey, logPath.data, attributes);
-						m_storage.DropCasFile(casKey, false, logPath.data);
+						if (!m_storage.CopyOrLink(casKey, logPath.data, attributes))
+							m_logger.Error(TC("Failed to copy cas from %s to %s"), CasKeyString(casKey).str, logPath.data);
+						else if (!m_storage.DropCasFile(casKey, false, logPath.data))
+							m_logger.Error(TC("Failed to drop cas %s"), CasKeyString(casKey).str);
 						writer.WriteBool(true);
 						return true;
 					}
@@ -895,7 +902,7 @@ namespace uba
 				u32 requestedSize = reader.ReadU32();
 
 				SCOPED_READ_LOCK(m_nameToHashLookupLock, lock);
-				if (requestedSize == ~u32(0))
+				if (requestedSize == ~0u)
 				{
 					requestedSize = u32(m_nameToHashTableMem.writtenSize);
 					writer.WriteU32(requestedSize);
@@ -1142,14 +1149,17 @@ namespace uba
 				u64 memTotal = reader.ReadU64();
 				u32 cpuLoadValue = reader.ReadU32();
 
+				u64 pingTime = GetTime();
 				u32 sessionIndex = sessionId - 1;
 				ScopedCriticalSection lock(m_remoteProcessAndSessionLock);
 				UBA_ASSERT(sessionIndex < m_clientSessions.size());
 				auto& session = *m_clientSessions[sessionIndex];
+				session.pingTime = pingTime;
 				session.lastPing = lastPing;
 				session.memAvail = memAvail;
 				session.memTotal = memTotal;
 				session.cpuLoad = *(float*)&cpuLoadValue;
+				writer.WriteBool(session.abort);
 
 				return true;
 			}
@@ -1176,7 +1186,12 @@ namespace uba
 				SCOPED_WRITE_LOCK(remoteProcess.m_exitedLock, exitedLock);
 				NextProcessInfo nextProcess;
 				bool newProcess;
-				if (!GetNextProcess(remoteProcess, newProcess, nextProcess, prevExitCode, reader))
+				remoteProcess.m_exitCode = prevExitCode;
+				remoteProcess.m_done.Set();
+				bool success = GetNextProcess(remoteProcess, newProcess, nextProcess, prevExitCode, reader);
+				remoteProcess.m_exitCode = ~0u;
+				remoteProcess.m_done.Reset();
+				if (!success)
 					return false;
 
 				writer.WriteBool(newProcess);
@@ -1218,6 +1233,69 @@ namespace uba
 				m_trace.SessionSummary(sessionId, reader.GetPositionData(), reader.GetLeft());
 				return true;
 			}
+			case SessionMessageType_Command:
+			{
+				StringBuffer<128> command;
+				reader.ReadString(command);
+
+				auto WriteString = [&](const tchar* str, LogEntryType type = LogEntryType_Info) { writer.WriteByte(type); writer.WriteString(str); };
+
+				if (command.Equals(TC("status")))
+				{
+					u32 totalUsed = 0;
+					u32 totalSlots = 0;
+					ScopedCriticalSection queueLock(m_remoteProcessAndSessionLock);
+					u64 time = GetTime();
+					for (auto& s : m_clientSessions)
+					{
+						if (!s->enabled)
+							continue;
+						WriteString(StringBuffer<>().Appendf(TC("Session %u (%s)"), s->id, s->name.c_str()).data);
+						WriteString(StringBuffer<>().Appendf(TC("   Process slots used %u/%u"), s->usedSlotCount, s->processSlotCount).data);
+						if (s->pingTime)
+							WriteString(StringBuffer<>().Appendf(TC("   Last ping %s ago"), TimeToText(time - s->pingTime).str).data);
+						totalUsed += s->usedSlotCount;
+						totalSlots += s->processSlotCount;
+					}
+					WriteString(StringBuffer<>().Appendf(TC("Total remote slots used %u/%u"), totalUsed, totalSlots).data);
+				}
+				else if (command.StartsWith(TC("abort")))
+				{
+					bool abortWithProxy = command.Equals(TC("abortproxy"));
+					bool abortUseProxy = command.Equals(TC("abortnonproxy"));
+					if (!abortWithProxy && !abortUseProxy)
+					{
+						abortWithProxy = true;
+						abortUseProxy = true;
+					}
+					ScopedCriticalSection queueLock(m_remoteProcessAndSessionLock);
+					u32 abortCount = 0;
+					for (auto& s : m_clientSessions)
+					{
+						if (!s->enabled || s->abort)
+							continue;
+						bool hasProxy = m_storage.HasProxy(s->id);
+						if (abortWithProxy && hasProxy)
+							s->abort = true;
+						else if (abortUseProxy && !hasProxy)
+							s->abort = true;
+						if (s->abort)
+							++abortCount;
+					}
+					WriteString(StringBuffer<>().Appendf(TC("Aborting: %u remote sessions"), abortCount).data);
+				}
+				else if (command.Equals(TC("disableremote")))
+				{
+					DisableRemoteExecution();
+					WriteString(StringBuffer<>().Appendf(TC("Remote execution is disabled")).data);
+				}
+				else
+				{
+					WriteString(StringBuffer<>().Appendf(TC("Unknown command: %s"), command.data).data, LogEntryType_Error);
+				}
+				writer.WriteByte(255);
+				return true;
+			}
 		}
 
 		UBA_ASSERT(false);
diff --git a/Engine/Source/Programs/UnrealBuildAccelerator/Common/Private/UbaStorageClient.cpp b/Engine/Source/Programs/UnrealBuildAccelerator/Common/Private/UbaStorageClient.cpp
--- a/Engine/Source/Programs/UnrealBuildAccelerator/Common/Private/UbaStorageClient.cpp
+++ b/Engine/Source/Programs/UnrealBuildAccelerator/Common/Private/UbaStorageClient.cpp
@@ -18,15 +18,28 @@ namespace uba
 	,	m_getProxyBackendUserData(info.getProxyBackendUserData)
 	,	m_startProxyCallback(info.startProxyCallback)
 	,	m_startProxyUserData(info.startProxyUserData)
+	,	m_proxyPort(info.proxyPort)
 	{
-		m_client.RegisterOnConnected([this, proxyPort = info.proxyPort]()
+	}
+
+	struct StorageClient::ProxyClient
+	{
+		ProxyClient(bool& outCtorSuccess, const NetworkClientCreateInfo& info) : client(outCtorSuccess, info, TC("UbaProxyClient")) {}
+		~ProxyClient() { client.Disconnect(); }
+		NetworkClient client;
+		u32 refCount = 0;
+	};
+
+	bool StorageClient::Start()
+	{
+		m_client.RegisterOnConnected([this]()
 			{
 				StackBinaryWriter<1024> writer;
 				NetworkMessage msg(m_client, ServiceId, StorageMessageType_Connect, writer);
 				writer.WriteString(TC("Client"));
 				writer.WriteU32(StorageNetworkVersion);
 				writer.WriteBool(false); // Is Proxy
-				writer.WriteU16(proxyPort);
+				writer.WriteU16(m_proxyPort);
 				writer.WriteString(m_zone);
 				writer.WriteU64(m_casTotalBytes);
 
@@ -46,15 +59,9 @@ namespace uba
 			});
 
 		m_client.RegisterOnDisconnected([this]() { m_logger.isMuted = true; });
+		return true;
 	}
 
-	struct StorageClient::ProxyClient
-	{
-		ProxyClient(bool& outCtorSuccess, const NetworkClientCreateInfo& info) : client(outCtorSuccess, info, TC("UbaProxyClient")) {}
-		NetworkClient client;
-		u32 refCount = 0;
-	};
-
 	StorageClient::~StorageClient()
 	{
 		delete m_proxyClient;
@@ -455,7 +462,7 @@ namespace uba
 						if (fetchId == u16(~0))
 							return m_logger.Error(TC("Cas content error. Server believes %s was only one segment but client sees more. Size: %llu Left to read: %llu ResponseSize: %u. (%s)"), hint, fileSize, left, responseSize, casFile.data);
 						readBuffer = slot;
-						if (!SendBatchMessages(*client, fetchId, readBuffer, BufferSlotSize, left, sizeOfFirstMessage, readIndex, responseSize))
+						if (!SendBatchMessages(m_logger, *client, fetchId, readBuffer, BufferSlotSize, left, sizeOfFirstMessage, readIndex, responseSize))
 						{
 							if (proxy)
 							{
@@ -515,7 +522,7 @@ namespace uba
 							{
 								if (fetchId == u16(~0))
 									return m_logger.Error(TC("Cas content error (2). Server believes %s was only one segment but client sees more. UncompressedSize: %llu LeftUncompressed: %llu Size: %llu Left to read: %llu ResponseSize: %u. (%s)"), hint, actualSize, leftUncompressed, fileSize, left, responseSize, casFile.data);
-								if (!SendBatchMessages(*client, fetchId, readPosition, maxReadSize - u32(readPosition - readBuffer), leftCompressed, sizeOfFirstMessage, readIndex, responseSize))
+								if (!SendBatchMessages(m_logger, *client, fetchId, readPosition, maxReadSize - u32(readPosition - readBuffer), leftCompressed, sizeOfFirstMessage, readIndex, responseSize))
 								{
 									if (proxy)
 									{
@@ -795,7 +802,7 @@ namespace uba
 			m_proxyClient->client.PrintSummary(logger);
 	}
 
-	bool StorageClient::SendBatchMessages(NetworkClient& client, u16 fetchId, u8* slot, u64 capacity, u64 left, u32 messageMaxSize, u32& readIndex, u32& responseSize)
+	bool StorageClient::SendBatchMessages(Logger& logger, NetworkClient& client, u16 fetchId, u8* slot, u64 capacity, u64 left, u32 messageMaxSize, u32& readIndex, u32& responseSize)
 	{
 		responseSize = 0;
 
@@ -833,31 +840,34 @@ namespace uba
 		for (u32 i=0; i!=sendCount; ++i)
 			new (entries + i) Entry(client, fetchId, slot, readIndex, i, messageMaxSize);
 
-		auto entryGuard = MakeGuard([&]()
-			{
-				for (u32 i=0; i!=sendCount; ++i)
-					entries[i].~Entry();
-			});
-
+		bool success = true;
+		u32 inFlight = u32(sendCount);
 		for (u32 i=0; i!=sendCount; ++i)
 		{
 			auto& entry = entries[i];
-			if (!entry.message.SendAsync(entry.reader, [](bool error, void* userData) { ((Event*)userData)->Set(); }, &entry.done))
-				return false;
+			if (entry.message.SendAsync(entry.reader, [](bool error, void* userData) { ((Event*)userData)->Set(); }, &entry.done))
+				continue;
+			inFlight = i;
+			success = false;
+			break;
 		}
 
-		for (u32 i=0; i!=sendCount; ++i)
+		for (u32 i=0; i!=inFlight; ++i)
 		{
 			Entry& entry = entries[i];
-			if (!entry.done.IsSet())
-				return false;
+			if (!entry.done.IsSet(5*60*1000))
+				logger.Error(TC("SendBatchMessages timed out getting async message response"));
 			if (!entry.message.ProcessAsyncResults(entry.reader))
-				return false;
-			responseSize += u32(entry.reader.GetLeft());
+				success = false;
+			else
+				responseSize += u32(entry.reader.GetLeft());
 		}
 
+		for (u32 i=0; i!=sendCount; ++i)
+			entries[i].~Entry();
+
 		readIndex += u32(sendCount);
-		return true;
+		return success;
 	}
 	
 	bool StorageClient::SendAllSegments(NetworkClient& client, u16 fetchId, u8* readBuffer, u64 left, u32 messageMaxSize)
diff --git a/Engine/Source/Programs/UnrealBuildAccelerator/Common/Private/UbaStorageProxy.cpp b/Engine/Source/Programs/UnrealBuildAccelerator/Common/Private/UbaStorageProxy.cpp
--- a/Engine/Source/Programs/UnrealBuildAccelerator/Common/Private/UbaStorageProxy.cpp
+++ b/Engine/Source/Programs/UnrealBuildAccelerator/Common/Private/UbaStorageProxy.cpp
@@ -9,6 +9,32 @@
 
 namespace uba
 {
+	struct StorageProxy::MessageInFlight
+	{
+		MessageInFlight(StorageProxy& p, FileEntry& f, u8* readBuffer, u32 fi)
+		:	proxy(p)
+		,	file(f)
+		,	message(p.m_client, ServiceId, StorageMessageType_FetchSegment, writer)
+		,	reader(readBuffer, 0, SendMaxSize)
+		,	fetchIndex(fi)
+		{
+			writer.WriteU16(file.fetchId);
+			writer.WriteU32(fetchIndex + 1);
+		}
+
+		StorageProxy& proxy;
+		FileEntry& file;
+		StackBinaryWriter<16> writer;
+		NetworkMessage message;
+		BinaryReader reader;
+		struct DeferredResponse { u32 clientId; u16 fetchId; MessageInfo info; };
+		List<DeferredResponse> deferredResponses;
+		u32 fetchIndex;
+		bool done = false;
+		bool error = false;
+	};
+
+
 	StorageProxy::StorageProxy(NetworkServer& server, NetworkClient& client, const Guid& storageServerUid, const tchar* name, StorageImpl* localStorage)
 	:	m_server(server)
 	,	m_client(client)
@@ -16,16 +42,13 @@ namespace uba
 	,	m_logger(client.GetLogWriter(), TC("StorageProxy"))
 	,	m_storageServerUid(storageServerUid)
 	,	m_name(name)
-	,	m_hasActiveFetchesEvent(true)
 	{
-		m_hasActiveFetchesEvent.Set();
-
 		m_server.RegisterOnClientDisconnected(0, [this](const Guid& clientUid, u32 clientId)
 			{
 				SCOPED_WRITE_LOCK(m_activeFetchesLock, lock);
 				for (auto it=m_activeFetches.begin(); it!=m_activeFetches.end();)
 				{
-					if (it->second.clientUid != clientUid)
+					if (it->second.clientId != clientId)
 					{
 						++it;
 						continue;
@@ -51,21 +74,11 @@ namespace uba
 
 	StorageProxy::~StorageProxy()
 	{
-		m_server.StopAll();
+		m_server.UnregisterService(StorageServiceId);
 		for (auto& kv : m_files)
 			delete[] kv.second.memory;
 	}
 
-	bool StorageProxy::Disconnect(u32 timeoutMs)
-	{
-		// TODO: Send disconnect to server to make it transition proxy to another proxy
-		//StackBinaryWriter<128> writer;
-		//NetworkMessage msg(m_client, ServiceId, StorageMessageType_Disconnect, writer);
-		//if (!msg.Send())
-		//	return false;
-		return m_hasActiveFetchesEvent.IsSet(timeoutMs);
-	}
-
 	void StorageProxy::PrintSummary()
 	{
 		LoggerWithWriter logger(m_logger.m_writer);
@@ -75,6 +88,20 @@ namespace uba
 		logger.Info(TC(""));
 	}
 
+	u16 StorageProxy::PopId()
+	{
+		if (m_availableIds.empty())
+			return m_availableIdsHigh++;
+		u16 storeId = m_availableIds.back();
+		m_availableIds.pop_back();
+		return storeId;
+	}
+
+	void StorageProxy::PushId(u16 id)
+	{
+		m_availableIds.push_back(id);
+	}
+
 	bool StorageProxy::HandleMessage(const ConnectionInfo& connectionInfo, MessageInfo& messageInfo, BinaryReader& reader, BinaryWriter& writer)
 	{
 		LoggerWithWriter logger(g_consoleLogWriter, TC(""));
@@ -118,13 +145,12 @@ namespace uba
 
 				SCOPED_WRITE_LOCK(file.lock, fileLock);
 
-				bool hasAllSegments = false;
 				while (true)
 				{
 					if (file.memory)
 						break;
 
-					constexpr bool useLocalStorage = true;
+					bool useLocalStorage = false; // Seems like it might be deadlocks in this code.. need to revisit
 					bool storeCompressed = true;
 					if (useLocalStorage && m_localStorage && IsCompressed(casKey) && m_inProcessClientId && connectionInfo.GetId() != m_inProcessClientId)
 					{
@@ -150,25 +176,30 @@ namespace uba
 								if (!file.memory)
 									return false;
 								file.size = fileSize;
+								file.received = fileSize;
 								memcpy(file.memory, sourceFile.GetData(), fileSize);
-								hasAllSegments = true;
+								file.available = true;
 							}
 						}
 					}
 
 					if (!file.memory)
 					{
-						NetworkMessage msg(m_client, ServiceId, messageInfo.type, writer2);
+						file.trackId = m_client.TrackWorkStart(StringBuffer<512>().AppendFileName(hint.data).data);
+
+						NetworkMessage msg(m_client, ServiceId, StorageMessageType_FetchBegin, writer2);
 						writer2.WriteBool(false); // Wants proxy
 						writer2.WriteCasKey(casKey);
 						writer2.WriteString(hint);
 						writer2.WriteBytes(reader.GetPositionData(), reader.GetLeft());
 
+						SCOPED_READ_LOCK(m_largeFileLock, largeFileLock);
 						if (!msg.Send(reader2))
 						{
 							file.error = true;
 							return m_logger.Error(TC("FetchBegin failed for cas file %s (%s). Requested by %s"), CasKeyString(casKey).str, hint.data, GuidToString(connectionInfo.GetUid()).str);
 						}
+						largeFileLock.Leave();
 
 						BinaryReader tempReader(reader2.GetPositionData(), 0, reader2.GetLeft());
 						u16 fetchId = tempReader.ReadU16();
@@ -204,10 +235,43 @@ namespace uba
 					if (file.received < file.size)
 					{
 						u64 segmentSize = m_client.GetMessageMaxSize() - 5; // This is server response size - header.. TODO: Should be taken from server
-						u64 segmentCount = file.size / segmentSize + 1;
-						u64 lookupByteCount = (segmentCount + 7) / 8;
-						u8 initialValue = hasAllSegments ? 255 : 0;
-						file.segmentsAvailable.resize(lookupByteCount, initialValue);
+						u32 segmentCount = u32(file.size / segmentSize);
+						file.messagesInFlight.resize(segmentCount);
+						for (u32 i=0; i!=segmentCount; ++i)
+						{
+							u64 offset = file.received + segmentSize * i;
+							u8* memory = file.memory + offset;
+							auto mif = new MessageInFlight(*this, file, memory, i);
+							file.messagesInFlight[i] = mif;
+						}
+
+						// Move the additional messages to a job to be able to return this one quickly.
+						m_server.AddWork([f = &file, segmentCount, this]()
+							{
+								SCOPED_WRITE_LOCK(m_largeFileLock, lock);
+								//TrackWorkScope tws(m_client, TC("SEGMENTS"));
+								auto& file = *f;
+								for (u32 i=0; i!=segmentCount; ++i)
+								{
+									auto mif = file.messagesInFlight[i];
+
+									bool res = mif->message.SendAsync(mif->reader, [](bool error, void* userData)
+										{
+											auto mif = (MessageInFlight*)userData;
+											mif->error = error;
+											mif->proxy.m_server.AddWork([mif]() { mif->proxy.HandleReceivedData(*mif); }, 1, TC(""));
+										}, mif);
+									if (!res)
+									{
+										// TODO: Don't leak mif
+										mif->error = true;
+									}
+								}
+							}, 1, TC(""));
+					}
+					else
+					{
+						m_client.TrackWorkEnd(file.trackId);
 					}
 					break;
 				}
@@ -223,16 +287,14 @@ namespace uba
 
 				if (fetchedSize < file.size)
 				{
-					fetchId = PopId();
 					SCOPED_WRITE_LOCK(m_activeFetchesLock, lock);
-					if (m_activeFetches.empty())
-						m_hasActiveFetchesEvent.Reset();
+					fetchId = PopId();
 					auto res = m_activeFetches.try_emplace(fetchId);
 					UBA_ASSERT(res.second);
 					ActiveFetch& fetch = res.first->second;
+					fetch.clientId = connectionInfo.GetId();
 					lock.Leave();
 
-					fetch.clientUid = connectionInfo.GetUid();
 					fetch.fetchedSize = fetchedSize;
 					fetch.file = &file;
 				}
@@ -250,129 +312,40 @@ namespace uba
 		case StorageMessageType_FetchSegment:
 			{
 				u16 fetchId = reader.ReadU16();
-				u32 fetchIndex = reader.ReadU32();
+				u32 fetchIndex = reader.ReadU32() - 1;
 
 				SCOPED_READ_LOCK(m_activeFetchesLock, activeLock);
 				auto findIt = m_activeFetches.find(fetchId);
 				UBA_ASSERT(findIt != m_activeFetches.end());
 				ActiveFetch& fetch = findIt->second;
+				u32 clientId = fetch.clientId;
 				activeLock.Leave();
 
 				FileEntry& file = *fetch.file;
-
-				u64 headerSize = sizeof(u16) + Get7BitEncodedCount(file.size) + sizeof(u8);
-				u64 firstFetchSize = m_client.GetMessageMaxSize() - m_client.GetMessageReceiveHeaderSize() - headerSize;
-				u64 segmentSize = m_client.GetMessageMaxSize() - 5; // This is server response size - header.. TODO: Should be taken from server
-
-				u64 offset = firstFetchSize + segmentSize * (fetchIndex - 1);
-				if (offset + segmentSize > file.size)
-					segmentSize = file.size - offset;
-				
 				SCOPED_WRITE_LOCK(file.lock, fileLock);
 				if (file.error)
 					return false;
 
-				u64 byteIndex = fetchIndex / 8;
-				u8 bitMask = u8(1 << (fetchIndex - byteIndex*8));
-				
-				u8& availableByte = file.segmentsAvailable[byteIndex];
-				if (!(availableByte & bitMask))
+				if (!file.available)
 				{
-					SegmentInFlight* activeSegment = nullptr;
-					auto asg = MakeGuard([&]
-						{
-							if (!--activeSegment->refCount)
-							{
-								if (auto next = activeSegment->next)
-									next->prev = activeSegment->prev;
-								else
-									file.lastInFlight = activeSegment->prev;
-								if (auto prev = activeSegment->prev)
-									prev->next = activeSegment->next;
-								else
-									file.firstInFlight = activeSegment->next;
-								delete activeSegment;
-							}
-						});
-
-					for (auto it=file.firstInFlight; it; it=it->next)
+					if (auto mif = file.messagesInFlight[fetchIndex])
 					{
-						if (it->segmentIndex != fetchIndex)
-							continue;
-						activeSegment = it;
-						break;
-					}
-
-					if (!activeSegment)
-					{
-						activeSegment = new SegmentInFlight;
-						activeSegment->segmentIndex = fetchIndex;
-						activeSegment->refCount = 1;
-						activeSegment->done.Create(true);
-
-						activeSegment->next = nullptr;
-						activeSegment->prev = file.lastInFlight;
-						if (auto last = file.lastInFlight)
-							last->next = activeSegment;
-						file.lastInFlight = activeSegment;
-						if (!file.firstInFlight)
-							file.firstInFlight = activeSegment;
-
-						fileLock.Leave();
-
-						NetworkMessage msg(m_client, ServiceId, StorageMessageType_FetchSegment, writer2);
-						writer2.WriteU16(file.fetchId);
-						writer2.WriteU32(fetchIndex);
-
-						if (!msg.Send(reader2))
-						{
-							file.error = true;
-							activeSegment->done.Set();
-							fileLock.Enter();
-							return m_logger.Error(TC("FetchSegment failed. Requested by %s"), GuidToString(connectionInfo.GetUid()).str);
-						}
-						
-						file.received += segmentSize;
-						if (file.sendEnd && file.size == file.received)
-							SendEnd(file.casKey);
-
-						memcpy(file.memory + offset, reader2.GetPositionData(), segmentSize);
-
-						activeSegment->done.Set();
-
-						fileLock.Enter();
-						availableByte |= bitMask;
-					}
-					else
-					{
-						++activeSegment->refCount;
-						fileLock.Leave();
-						bool success = activeSegment->done.IsSet(10*60*1000); // This should never happen.
-						fileLock.Enter();
-						if (!success)
-							return m_logger.Error(TC("Connection %s timed out after 10 minutes waiting for segment %u on cas entry %s to be available in storage proxy"), fetchIndex, CasKeyString(file.casKey).str, GuidToString(connectionInfo.GetUid()).str);
-						if (file.error)
-							return false;
+						UBA_ASSERT(clientId == connectionInfo.GetId());
+						mif->deferredResponses.push_back({clientId, fetchId, messageInfo});
+						messageInfo = {};
+						return true;
 					}
 				}
 				fileLock.Leave();
 
-				const u8* memory = file.memory + offset;
-
-				writer.WriteBytes(memory, segmentSize);
-
-				u64 fetchedSize = fetch.fetchedSize.fetch_add(segmentSize) + segmentSize;
-				if (fetchedSize != file.size)
-					return true;
-
-				SCOPED_WRITE_LOCK(m_activeFetchesLock, activeLock2);
-				m_activeFetches.erase(findIt);
-				if (m_activeFetches.empty())
-					m_hasActiveFetchesEvent.Set();
-				activeLock2.Leave();
-
-				PushId(fetchId);
-				return true;
+				u64 headerSize = sizeof(u16) + Get7BitEncodedCount(file.size) + sizeof(u8);
+				u64 firstFetchSize = m_client.GetMessageMaxSize() - m_client.GetMessageReceiveHeaderSize() - headerSize;
+				u64 segmentSize = m_client.GetMessageMaxSize() - 5; // This is server response size - header.. TODO: Should be taken from server
+				u64 offset = firstFetchSize + segmentSize * (fetchIndex);
+				if (offset + segmentSize > file.size)
+					segmentSize = file.size - offset;
+				writer.WriteBytes(file.memory + offset, segmentSize);
+				return UpdateFetch(fetch.clientId, fetchId, segmentSize);
 			}
 		case StorageMessageType_FetchEnd:
 			{
@@ -390,20 +363,73 @@ namespace uba
 		}
 	}
 
-	u16 StorageProxy::PopId()
+	void StorageProxy::HandleReceivedData(MessageInFlight& mif)
 	{
-		SCOPED_WRITE_LOCK(m_availableIdsLock, lock);
-		if (m_availableIds.empty())
-			return m_availableIdsHigh++;
-		u16 storeId = m_availableIds.back();
-		m_availableIds.pop_back();
-		return storeId;
+		auto& file = mif.file;
+		if (mif.error)
+		{
+			SCOPED_WRITE_LOCK(file.lock, fileLock);
+			file.error = true;
+		}
+		else
+		{
+			mif.message.ProcessAsyncResults(mif.reader);
+		}
+
+		SCOPED_WRITE_LOCK(file.lock, fileLock);
+
+		UBA_ASSERT(file.messagesInFlight[mif.fetchIndex] == &mif);
+		file.messagesInFlight[mif.fetchIndex] = nullptr;
+		file.received += mif.reader.GetLeft();
+		bool finished = file.received == file.size;
+		if (finished)
+			file.available = true;
+		fileLock.Leave();
+
+		if (finished)
+		{
+			m_client.TrackWorkEnd(file.trackId);
+			SendEnd(file.casKey);
+		}
+
+
+		for (auto& r : mif.deferredResponses)
+		{
+			if (UpdateFetch(r.clientId, r.fetchId, mif.reader.GetLeft()) && !mif.error)
+				m_server.SendResponse(r.info, mif.reader.GetPositionData(), u32(mif.reader.GetLeft()));
+			else
+				m_server.SendResponse(r.info, nullptr, 0);
+		}
+
+		delete &mif;
 	}
 
-	void StorageProxy::PushId(u16 id)
+	bool StorageProxy::UpdateFetch(u32 clientId, u16 fetchId, u64 segmentSize)
 	{
-		SCOPED_WRITE_LOCK(m_availableIdsLock, lock);
-		m_availableIds.push_back(id);
+		SCOPED_WRITE_LOCK(m_activeFetchesLock, activeLock);
+		auto findIt = m_activeFetches.find(fetchId);
+		if (findIt == m_activeFetches.end())
+		{
+			// This can happen if we have async downloading and client is disconnected
+			//m_logger.Info(TC("Failed to find active fetch with id %u"), fetchId);
+			return false;
+		}
+
+		ActiveFetch& fetch = findIt->second;
+		if (fetch.clientId != clientId)
+		{
+			// This can happen if we have async downloading and client is disconnected and new client have reused fetch id
+			//m_logger.Info(TC("Active fetch %i has a different client id."), fetchId);
+			return false;
+		}
+
+		fetch.fetchedSize += segmentSize;
+		if (fetch.fetchedSize != fetch.file->size)
+			return true;
+
+		m_activeFetches.erase(findIt);
+		PushId(fetchId);
+		return true;
 	}
 
 	bool StorageProxy::SendEnd(const CasKey& key)
diff --git a/Engine/Source/Programs/UnrealBuildAccelerator/Common/Private/UbaStorageServer.cpp b/Engine/Source/Programs/UnrealBuildAccelerator/Common/Private/UbaStorageServer.cpp
--- a/Engine/Source/Programs/UnrealBuildAccelerator/Common/Private/UbaStorageServer.cpp
+++ b/Engine/Source/Programs/UnrealBuildAccelerator/Common/Private/UbaStorageServer.cpp
@@ -173,7 +173,6 @@ namespace uba
 				if (m_traceFetch)
 					m_trace->FileEndFetch(clientId, AsCompressed(fetch.casKey, m_storeCompressed));
 
-				CloseFile(nullptr, fetch.readFileHandle);
 				it = m_activeFetches.erase(it);
 			}
 		}
@@ -248,6 +247,15 @@ namespace uba
 		m_traceStore = detailed;
 	}
 
+	bool StorageServer::HasProxy(u32 clientId)
+	{
+		SCOPED_READ_LOCK(m_proxiesLock, l);
+		for (auto& kv : m_proxies)
+			if (kv.second.clientId == clientId)
+				return true;
+		return false;
+	}
+
 	bool StorageServer::WaitForWritten(CasEntry& casEntry, ScopedWriteLock& entryLock, const tchar* hint)
 	{
 		int waitCount = 0;
@@ -558,12 +566,12 @@ namespace uba
 				SCOPED_WRITE_LOCK(m_activeFetchesLock, lock);
 				auto insres = m_activeFetches.try_emplace(*fetchId);
 				UBA_ASSERT(insres.second);
+				ActiveFetch& fetch = insres.first->second;
+				fetch.clientId = connectionInfo.GetId();
 				lock.Leave();
 
 				mappedView.size = fileSize;
 
-				ActiveFetch& fetch = insres.first->second;
-				fetch.clientId = connectionInfo.GetId();
 				fetch.readFileHandle = readFileHandle;
 				fetch.mappedView = mappedView;
 				fetch.ownsMapping = ownsMapping;
@@ -586,6 +594,7 @@ namespace uba
 				if (findIt == m_activeFetches.end())
 					return m_logger.Error(TC("Can't find active fetch %u, disconnected client? (index %u)"), fetchId, fetchIndex);
 				ActiveFetch& fetch = findIt->second;
+				UBA_ASSERT(fetch.clientId == connectionInfo.GetId());
 				lock.Leave();
 
 				UBA_ASSERT(fetchIndex);
@@ -729,7 +738,6 @@ namespace uba
 
 				*(u64*)mappedView.memory = fileSize;
 
-				firstStore->clientId = connectionInfo.GetId();
 				firstStore->casEntry  = &casEntry;
 				firstStore->fileSize = fileSize;
 				firstStore->actualSize = actualSize;
@@ -757,6 +765,7 @@ namespace uba
 					if (storeIt == m_activeStores.end())
 						return m_logger.Error(TC("Can't find active store %u, disconnected client?"), storeId);
 					activeStoreTemp = &storeIt->second;
+					UBA_ASSERT(activeStoreTemp->clientId == connectionInfo.GetId());
 				}
 				ActiveStore& activeStore = *activeStoreTemp;
 
@@ -829,10 +838,10 @@ namespace uba
 					SCOPED_WRITE_LOCK(m_activeStoresLock, activeLock);
 					auto insres = m_activeStores.try_emplace(storeId);
 					UBA_ASSERT(insres.second);
+					ActiveStore& s = insres.first->second;
+					s.clientId = connectionInfo.GetId();
 					activeLock.Leave();
 
-					ActiveStore& s = insres.first->second;
-					s.clientId = firstStore->clientId;
 					s.fileSize = firstStore->fileSize;
 					s.mappedView = firstStore->mappedView;
 					s.casEntry = firstStore->casEntry;
diff --git a/Engine/Source/Programs/UnrealBuildAccelerator/Common/Public/UbaCoordinator.h b/Engine/Source/Programs/UnrealBuildAccelerator/Common/Public/UbaCoordinator.h
new file mode 100644
--- /dev/null
+++ b/Engine/Source/Programs/UnrealBuildAccelerator/Common/Public/UbaCoordinator.h
@@ -0,0 +1,32 @@
+// Copyright Epic Games, Inc. All Rights Reserved.
+
+#include "UbaBase.h"
+
+namespace uba
+{
+	struct CoordinatorCreateInfo
+	{
+		const tchar* workDir = nullptr;
+		const tchar* binariesDir = nullptr;
+
+		// TODO: This is very horde specific.. maybe all these parameters should be a string or something
+		const tchar* uri = nullptr;
+		const tchar* pool = nullptr;
+		const tchar* oidc = nullptr;
+		u32 maxCoreCount = 0;
+
+		bool logging = false;
+	};
+
+	class Coordinator
+	{
+	public:
+		using AddClientCallback = bool(void* userData, const tchar* ip, u16 port);
+		virtual void SetAddClientCallback(AddClientCallback* callback, void* userData) = 0;
+		virtual void SetTargetCoreCount(u32 count) = 0;
+	};
+
+	// This is how the function signatures creating/destroying the coordinator module needs to exported
+	using UbaCreateCoordinatorFunc = Coordinator*(const CoordinatorCreateInfo& info);
+	using UbaDestroyCoordinatorFunc = void(Coordinator* coordinator);
+}
\ No newline at end of file
diff --git a/Engine/Source/Programs/UnrealBuildAccelerator/Common/Public/UbaCoordinatorWrapper.h b/Engine/Source/Programs/UnrealBuildAccelerator/Common/Public/UbaCoordinatorWrapper.h
new file mode 100644
--- /dev/null
+++ b/Engine/Source/Programs/UnrealBuildAccelerator/Common/Public/UbaCoordinatorWrapper.h
@@ -0,0 +1,100 @@
+// Copyright Epic Games, Inc. All Rights Reserved.
+
+#pragma once
+
+#include "UbaCoordinator.h"
+#include "UbaFile.h"
+#include "UbaNetworkServer.h"
+#include "UbaStringBuffer.h"
+
+#if !PLATFORM_WINDOWS
+#include <dlfcn.h>
+#define GetProcAddress dlsym
+#define LoadLibrary(name) dlopen(name, RTLD_LAZY);
+#define LoadLibraryError dlerror()
+#define HMODULE void*
+#else
+#define LoadLibraryError LastErrorToText().data
+#endif
+
+namespace uba
+{
+	class CoordinatorWrapper
+	{
+	public:
+		bool Create(Logger& logger, const tchar* coordinatorType, const CoordinatorCreateInfo& info, NetworkBackend& networkBackend, NetworkServer& networkServer)
+		{
+			UbaCreateCoordinatorFunc* createCoordinator = nullptr;
+
+			if (!*coordinatorType)
+				return false;
+
+			StringBuffer<128> coordinatorBin;
+			#if PLATFORM_WINDOWS
+			coordinatorBin.Append(TC("UbaCoordinator")).Append(coordinatorType).Append(TC(".dll"));
+			#else
+			coordinatorBin.Append(TC("libUbaCoordinator")).Append(coordinatorType).Append(TC(".so"));
+			#endif
+
+			HMODULE coordinatorModule = LoadLibrary(coordinatorBin.data);
+			if (!coordinatorModule)
+				return logger.Error(TC("Failed to load coordinator binary %s (%s)"), coordinatorBin.data, LoadLibraryError);
+
+			if (!coordinatorModule)
+				return logger.Error(TC("Failed to load coordinator binary %s (%s)"), coordinatorBin.data, LastErrorToText().data);
+			createCoordinator = (UbaCreateCoordinatorFunc*)(void*)GetProcAddress(coordinatorModule,"UbaCreateCoordinator");
+			if (!createCoordinator)
+				return logger.Error(TC("Failed to find UbaCreateCoordinator function inside %s (%s)"), coordinatorBin.data, LastErrorToText().data);
+			m_destroyCoordinator = (UbaDestroyCoordinatorFunc*)(void*)GetProcAddress(coordinatorModule, "UbaDestroyCoordinator");
+			if (!m_destroyCoordinator)
+				return logger.Error(TC("Failed to find UbaDestroyCoordinator function inside %s (%s)"), coordinatorBin.data, LastErrorToText().data);
+
+			StringBuffer<512> binariesDir;
+			if (!GetDirectoryOfCurrentModule(logger, binariesDir))
+				return false;
+
+			m_coordinator = createCoordinator(info);
+			if (!m_coordinator)
+				return false;
+
+			m_loopCoordinator.Create(true);
+			m_networkBackend = &networkBackend;
+			m_networkServer = &networkServer;
+
+			m_coordinatorThread.Start([this, nb = &networkBackend, ns = &networkServer, tc = info.maxCoreCount]()
+			{
+				m_coordinator->SetAddClientCallback([](void* userData, const tchar* ip, u16 port)
+					{
+						auto& cw = *(CoordinatorWrapper*)userData;
+						return cw.m_networkServer->AddClient(*cw.m_networkBackend, ip, port);
+					}, this);
+
+				do
+				{
+					m_coordinator->SetTargetCoreCount(tc);
+				}
+				while (!m_loopCoordinator.IsSet(3000));
+
+				return 0;
+			});
+			return true;
+		}
+
+		void Destroy()
+		{
+			if (!m_coordinator)
+				return;
+			m_loopCoordinator.Set();
+			m_coordinatorThread.Wait();
+			m_destroyCoordinator(m_coordinator);
+			m_coordinator = nullptr;
+		}
+
+		Coordinator* m_coordinator = nullptr;
+		NetworkBackend* m_networkBackend = nullptr;
+		NetworkServer* m_networkServer = nullptr;
+		UbaDestroyCoordinatorFunc* m_destroyCoordinator = nullptr;
+		Event m_loopCoordinator;
+		Thread m_coordinatorThread;
+	};
+}
\ No newline at end of file
diff --git a/Engine/Source/Programs/UnrealBuildAccelerator/Common/Public/UbaNetwork.h b/Engine/Source/Programs/UnrealBuildAccelerator/Common/Public/UbaNetwork.h
--- a/Engine/Source/Programs/UnrealBuildAccelerator/Common/Public/UbaNetwork.h
+++ b/Engine/Source/Programs/UnrealBuildAccelerator/Common/Public/UbaNetwork.h
@@ -18,7 +18,7 @@ namespace uba
 	
 	static constexpr u32 SystemNetworkVersion = 1339;
 	static constexpr u32 StorageNetworkVersion = 3;
-	static constexpr u32 SessionNetworkVersion = 29;
+	static constexpr u32 SessionNetworkVersion = 30;
 
 
 	// Messages used over network between client/server (system, storage and session)
@@ -82,6 +82,7 @@ namespace uba
 		UBA_SESSION_MESSAGE(Custom) \
 		UBA_SESSION_MESSAGE(UpdateEnvironment) \
 		UBA_SESSION_MESSAGE(Summary) \
+		UBA_SESSION_MESSAGE(Command) \
 
 	enum SessionMessageType : u8 
 	{
diff --git a/Engine/Source/Programs/UnrealBuildAccelerator/Common/Public/UbaNetworkClient.h b/Engine/Source/Programs/UnrealBuildAccelerator/Common/Public/UbaNetworkClient.h
--- a/Engine/Source/Programs/UnrealBuildAccelerator/Common/Public/UbaNetworkClient.h
+++ b/Engine/Source/Programs/UnrealBuildAccelerator/Common/Public/UbaNetworkClient.h
@@ -33,13 +33,10 @@ namespace uba
 		NetworkClient(bool& outCtorSuccess, const NetworkClientCreateInfo& info = {}, const tchar* name = TC("UbaClient"));
 		~NetworkClient();
 
-		void StopAll();
-
 		bool Connect(NetworkBackend& backend, const tchar* ip, u16 port = DefaultPort, bool* timedOut = nullptr);
 		void Disconnect();
 
 		bool StartListen(NetworkBackend& backend, u16 port = DefaultPort);
-		void StopListen();
 		bool SetConnectionCount(u32 count);
 		bool SendKeepAlive();
 
@@ -81,10 +78,11 @@ namespace uba
 		};
 
 		bool AddConnection(NetworkBackend& backend, void* backendConnection, bool* timedOut);
-		void ConnectedCallback(NetworkBackend& backend, void* backendConnection);
+		bool ConnectedCallback(NetworkBackend& backend, void* backendConnection);
+		static void DisconnectCallback(void* context, const Guid& connectionUid, void* connection);
 		static bool ReceiveResponseHeader(void* context, const Guid& connectionUid, u8* headerData, void*& outBodyContext, u8*& outBodyData, u32& outBodySize);
 		static bool ReceiveResponseBody(void* context, bool recvError, u8* headerData, void* bodyContext, u8* bodyData, u32 bodySize);
-		void OnDisconnected(Connection& connection, bool calledFromReceive);
+		void OnDisconnected(Connection& connection);
 		bool Send(NetworkMessage& message, void* response, u32 responseCapacity, bool async);
 		void ReturnMessageId(u16 id);
 		const tchar* SetGetPrefix(const tchar* originalPrefix);
@@ -98,6 +96,7 @@ namespace uba
 		Atomic<u64> m_sendBytes;
 		Atomic<u64> m_recvBytes;
 		Atomic<u32> m_recvCount;
+		Atomic<bool> m_isDisconnecting;
 		Timer m_sendTimer;
 
 		ReaderWriterLock m_serverUidLock;
@@ -112,8 +111,6 @@ namespace uba
 		Vector<OnDisconnectedFunction> m_onDisconnectedFunctions;
 		OnVersionMismatchFunction m_versionMismatchFunction;
 
-		NetworkBackend* m_listenBackend = nullptr;
-
 		ReaderWriterLock m_connectionsLock;
 		List<Connection> m_connections;
 		ReaderWriterLock m_connectionsItLock;
diff --git a/Engine/Source/Programs/UnrealBuildAccelerator/Common/Public/UbaNetworkServer.h b/Engine/Source/Programs/UnrealBuildAccelerator/Common/Public/UbaNetworkServer.h
--- a/Engine/Source/Programs/UnrealBuildAccelerator/Common/Public/UbaNetworkServer.h
+++ b/Engine/Source/Programs/UnrealBuildAccelerator/Common/Public/UbaNetworkServer.h
@@ -48,9 +48,8 @@ namespace uba
 		virtual ~NetworkServer();
 
 		bool StartListen(NetworkBackend& backend, u16 port = DefaultPort, const tchar* ip = nullptr, const u8* cryptoKey128 = nullptr); // Start listen for new connections/clients
-		void StopListen();			// Stop listen for new connections
 		void DisallowNewClients();	// Disallow new clients to connect but old clients can still create more connections
-		void StopAll();				// Stops all listen and disconnect all active connections
+		void DisconnectClients();				// Stops all listen and disconnect all active connections
 
 		bool AddClient(NetworkBackend& backend, const tchar* ip, u16 port = DefaultPort, const u8* cryptoKey128 = nullptr); // Adds a client that server will create one or more connections to (note this will return before we know if it was a success or not)
 
@@ -102,6 +101,7 @@ namespace uba
 		Worker* PopWorkerNoLock();
 		void PushWorker(Worker* worker);
 		void PushWorkerNoLock(Worker* worker);
+		void FlushWorkers();
 
 		bool HandleSystemMessage(const ConnectionInfo& connectionInfo, u8 messageType, BinaryReader& reader, BinaryWriter& writer);
 		bool AddConnection(NetworkBackend& backend, void* backendConnection, const sockaddr& remoteSocketAddr, CryptoKey cryptoKey);
@@ -110,7 +110,6 @@ namespace uba
 
 		LoggerWithWriter m_logger;
 
-		NetworkBackend* m_listenBackend = nullptr;
 		CryptoKey m_listenCrypto = InvalidCryptoKey;
 		Guid m_uid;
 		bool m_allowNewClients = true;
@@ -120,6 +119,7 @@ namespace uba
 		OnConnectionFunction m_onConnectionFunction;
 		
 		struct OnDisconnectEntry { u8 id; OnDisconnectFunction function; };
+		ReaderWriterLock m_onDisconnectFunctionsLock;
 		List<OnDisconnectEntry> m_onDisconnectFunctions;
 
 		u32 m_maxWorkerCount = 0;
diff --git a/Engine/Source/Programs/UnrealBuildAccelerator/Common/Public/UbaSession.h b/Engine/Source/Programs/UnrealBuildAccelerator/Common/Public/UbaSession.h
--- a/Engine/Source/Programs/UnrealBuildAccelerator/Common/Public/UbaSession.h
+++ b/Engine/Source/Programs/UnrealBuildAccelerator/Common/Public/UbaSession.h
@@ -69,10 +69,12 @@ namespace uba
 		Logger& GetLogger();			// Logger used for logging 
 		LogWriter& GetLogWriter();		// LogWriter used by logger
 
-		Session(const SessionCreateInfo& info, const tchar* logPrefix, bool runningRemote, WorkManager* workManager = nullptr);
 		virtual ~Session();
 
 	protected:
+		Session(const SessionCreateInfo& info, const tchar* logPrefix, bool runningRemote, WorkManager* workManager = nullptr);
+		bool Create(const SessionCreateInfo& info);
+
 		void ValidateStartInfo(const ProcessStartInfo& startInfo);
 		ProcessHandle InternalRunProcess(const ProcessStartInfo& startInfo, bool async, ProcessImpl* parent, bool enableDetour);
 		void ProcessAdded(Process& process, u32 sessionId);
@@ -119,7 +121,7 @@ namespace uba
 		virtual bool CreateDirectory(CreateDirectoryResponse& out, const CreateDirectoryMessage& msg);
 		virtual bool GetFullFileName(GetFullFileNameResponse& out, const GetFullFileNameMessage& msg);
 		virtual bool GetListDirectoryInfo(ListDirectoryResponse& out, tchar* dirName, const StringKey& dirKey);
-		virtual bool WriteFileToDisk(ProcessImpl& process, WrittenFile& file);
+		virtual bool WriteFilesToDisk(ProcessImpl& process, WrittenFile** files, u32 fileCount);
 		virtual bool AllocFailed(Process& process, const tchar* allocType, u32 error);
 		virtual bool GetNextProcess(Process& process, bool& outNewProcess, NextProcessInfo& outNextProcess, u32 prevExitCode, BinaryReader& statsReader);
 		virtual bool CustomMessage(Process& process, BinaryReader& reader, BinaryWriter& writer);
diff --git a/Engine/Source/Programs/UnrealBuildAccelerator/Common/Public/UbaSessionClient.h b/Engine/Source/Programs/UnrealBuildAccelerator/Common/Public/UbaSessionClient.h
--- a/Engine/Source/Programs/UnrealBuildAccelerator/Common/Public/UbaSessionClient.h
+++ b/Engine/Source/Programs/UnrealBuildAccelerator/Common/Public/UbaSessionClient.h
@@ -57,7 +57,7 @@ namespace uba
 		virtual bool CreateDirectory(CreateDirectoryResponse& out, const CreateDirectoryMessage& msg) override;
 		virtual bool GetFullFileName(GetFullFileNameResponse& out, const GetFullFileNameMessage& msg) override;
 		virtual bool GetListDirectoryInfo(ListDirectoryResponse& out, tchar* dirName, const StringKey& dirKey) override;
-		virtual bool WriteFileToDisk(ProcessImpl& process, WrittenFile& file) override;
+		virtual bool WriteFilesToDisk(ProcessImpl& process, WrittenFile** files, u32 fileCount) override;
 		virtual bool AllocFailed(Process& process, const tchar* allocType, u32 error) override;
 		virtual void PrintSessionStats(Logger& logger) override;
 		virtual bool GetNextProcess(Process& process, bool& outNewProcess, NextProcessInfo& outNextProcess, u32 prevExitCode, BinaryReader& statsReader) override;
diff --git a/Engine/Source/Programs/UnrealBuildAccelerator/Common/Public/UbaSessionServer.h b/Engine/Source/Programs/UnrealBuildAccelerator/Common/Public/UbaSessionServer.h
--- a/Engine/Source/Programs/UnrealBuildAccelerator/Common/Public/UbaSessionServer.h
+++ b/Engine/Source/Programs/UnrealBuildAccelerator/Common/Public/UbaSessionServer.h
@@ -130,9 +130,11 @@ namespace uba
 			u64 lastPing = 0;
 			u64 memAvail = 0;
 			u64 memTotal = 0;
+			u64 pingTime = 0;
 			float cpuLoad = 0;
 			bool enabled = true;
 			bool dedicated = false;
+			bool abort = false;
 		};
 		Vector<ClientSession*> m_clientSessions;
 
diff --git a/Engine/Source/Programs/UnrealBuildAccelerator/Common/Public/UbaStorage.h b/Engine/Source/Programs/UnrealBuildAccelerator/Common/Public/UbaStorage.h
--- a/Engine/Source/Programs/UnrealBuildAccelerator/Common/Public/UbaStorage.h
+++ b/Engine/Source/Programs/UnrealBuildAccelerator/Common/Public/UbaStorage.h
@@ -32,6 +32,7 @@ namespace uba
 		virtual bool SaveCasTable(bool deleteIsRunningfile, bool deleteDropped = true) = 0;
 		virtual u64 GetStorageCapacity() = 0;
 		virtual bool GetZone(StringBufferBase& out) = 0;
+		virtual bool HasProxy(u32 clientId) { return false; }
 
 		virtual bool DecompressFileToMemory(const tchar* fileName, FileHandle fileHandle, u8* dest, u64 decompressedSize) = 0;
 		virtual bool DecompressMemoryToMemory(u8* compressedData, u8* writeData, u64 decompressedSize, const tchar* readHint) = 0;
diff --git a/Engine/Source/Programs/UnrealBuildAccelerator/Common/Public/UbaStorageClient.h b/Engine/Source/Programs/UnrealBuildAccelerator/Common/Public/UbaStorageClient.h
--- a/Engine/Source/Programs/UnrealBuildAccelerator/Common/Public/UbaStorageClient.h
+++ b/Engine/Source/Programs/UnrealBuildAccelerator/Common/Public/UbaStorageClient.h
@@ -33,6 +33,8 @@ namespace uba
 		StorageClient(const StorageClientCreateInfo& info);
 		~StorageClient();
 
+		bool Start();
+
 		bool IsUsingProxy();
 		void StopProxy();
 
@@ -53,7 +55,7 @@ namespace uba
 		virtual void Ping() override;
 		virtual void PrintSummary(Logger& logger) override;
 
-		static bool SendBatchMessages(NetworkClient& client, u16 fetchId, u8* slot, u64 capacity, u64 left, u32 messageMaxSize, u32& readIndex, u32& responseSize);
+		static bool SendBatchMessages(Logger& logger, NetworkClient& client, u16 fetchId, u8* slot, u64 capacity, u64 left, u32 messageMaxSize, u32& readIndex, u32& responseSize);
 		static bool SendAllSegments(NetworkClient& client, u16 fetchId, u8* readBuffer, u64 left, u32 messageMaxSize);
 
 	private:
@@ -86,6 +88,7 @@ namespace uba
 		void* m_getProxyBackendUserData = nullptr;
 		StartProxyCallback* m_startProxyCallback = nullptr;
 		void* m_startProxyUserData = nullptr;
+		u16 m_proxyPort = 0;
 	};
 
 }
diff --git a/Engine/Source/Programs/UnrealBuildAccelerator/Common/Public/UbaStorageProxy.h b/Engine/Source/Programs/UnrealBuildAccelerator/Common/Public/UbaStorageProxy.h
--- a/Engine/Source/Programs/UnrealBuildAccelerator/Common/Public/UbaStorageProxy.h
+++ b/Engine/Source/Programs/UnrealBuildAccelerator/Common/Public/UbaStorageProxy.h
@@ -21,14 +21,17 @@ namespace uba
 	public:
 		StorageProxy(NetworkServer& server, NetworkClient& client, const Guid& storageServerUid, const tchar* name, StorageImpl* localStorage = nullptr);
 		~StorageProxy();
-
-		bool Disconnect(u32 timeoutMs); // Use a timeout to try to make more graceful disconnect
 		void PrintSummary();
 
 	protected:
 		u16 PopId();
 		void PushId(u16 id);
 		bool HandleMessage(const ConnectionInfo& connectionInfo, MessageInfo& messageInfo, BinaryReader& reader, BinaryWriter& writer);
+
+		struct MessageInFlight;
+		void HandleReceivedData(MessageInFlight& mif);
+
+		bool UpdateFetch(u32 clientId, u16 fetchId, u64 segmentSize);
 		bool SendEnd(const CasKey& key);
 
 		static constexpr u8 ServiceId = StorageServiceId;
@@ -45,8 +48,6 @@ namespace uba
 
 		Atomic<u32> m_inProcessClientId;
 
-		struct SegmentInFlight { u32 refCount; u32 segmentIndex; Event done; SegmentInFlight* prev; SegmentInFlight* next; };
-
 		struct FileEntry
 		{
 			ReaderWriterLock lock;
@@ -54,13 +55,13 @@ namespace uba
 			u64 size = 0;
 			Atomic<u64> received;
 			CasKey casKey;
+			u32 trackId = 0;
 			u16 fetchId = 0;
 			bool storeCompressed = false;
 			bool sendEnd = false;
 			bool error = false;
-			Vector<u8> segmentsAvailable;
-			SegmentInFlight* firstInFlight = nullptr;
-			SegmentInFlight* lastInFlight = nullptr;
+			bool available = false;
+			Vector<MessageInFlight*> messagesInFlight;
 		};
 
 		ReaderWriterLock m_filesLock;
@@ -69,15 +70,16 @@ namespace uba
 		struct ActiveFetch
 		{
 			FileEntry* file = nullptr;
-			Guid clientUid;
-			Atomic<u64> fetchedSize;
+			u64 fetchedSize = 0;
+			u32 clientId = ~0u;
+			u32 connectionId = 0;
 		};
 
-		Event m_hasActiveFetchesEvent;
 		ReaderWriterLock m_activeFetchesLock;
 		UnorderedMap<u16, ActiveFetch> m_activeFetches;
 
-		ReaderWriterLock m_availableIdsLock;
+		ReaderWriterLock m_largeFileLock;
+
 		Vector<u16> m_availableIds;
 		u16 m_availableIdsHigh = 1;
 	};
diff --git a/Engine/Source/Programs/UnrealBuildAccelerator/Common/Public/UbaStorageServer.h b/Engine/Source/Programs/UnrealBuildAccelerator/Common/Public/UbaStorageServer.h
--- a/Engine/Source/Programs/UnrealBuildAccelerator/Common/Public/UbaStorageServer.h
+++ b/Engine/Source/Programs/UnrealBuildAccelerator/Common/Public/UbaStorageServer.h
@@ -28,6 +28,7 @@ namespace uba
 		virtual bool WriteCompressed(WriteResult& out, const tchar* from, const tchar* toFile) override;
 		virtual bool IsDisallowedPath(const tchar* fileName) override;
 		virtual void SetTrace(Trace* trace, bool detailed) override;
+		virtual bool HasProxy(u32 clientId) override;
 		void OnDisconnected(u32 clientId);
 		bool HandleMessage(const ConnectionInfo& connectionInfo, u8 messageType, BinaryReader& reader, BinaryWriter& writer);
 		bool WaitForWritten(CasEntry& casEntry, ScopedWriteLock& entryLock, const tchar* hint);
diff --git a/Engine/Source/Programs/UnrealBuildAccelerator/Core/Public/UbaDirectoryTable.h b/Engine/Source/Programs/UnrealBuildAccelerator/Core/Public/UbaDirectoryTable.h
--- a/Engine/Source/Programs/UnrealBuildAccelerator/Core/Public/UbaDirectoryTable.h
+++ b/Engine/Source/Programs/UnrealBuildAccelerator/Core/Public/UbaDirectoryTable.h
@@ -62,8 +62,9 @@ namespace uba
 			BinaryReader reader(m_memory, tableOffset);
 			u32 prevTableOffset = u32(reader.Read7BitEncoded());
 
+			u32 buffer[48*1024];
 			u32 count = 0;
-			u32 readerOffsets[64*1024];
+			u32* readerOffsets = buffer;
 			readerOffsets[count++] = u32(reader.GetPosition());
 			bool firstIsRoot = true;
 			while (true)
@@ -76,7 +77,11 @@ namespace uba
 				reader.SetPosition(prevTableOffset);
 				prevTableOffset = u32(reader.Read7BitEncoded());
 				readerOffsets[count++] = u32(reader.GetPosition());
-				UBA_ASSERT(count < sizeof_array(readerOffsets));
+				if (count == sizeof_array(buffer))
+				{
+					readerOffsets = new u32[1024*1024]; // This sucks, but somethings the directory is huuuge. Ideally these files should be spread out over multiple directories
+					memcpy(readerOffsets, buffer, sizeof(buffer));
+				}
 			}
 
 			for (u32 i=count; i>0; --i)
@@ -93,6 +98,9 @@ namespace uba
 
 				PopulateDirectoryWithFiles(reader, hasher, files);
 			}
+
+			if (readerOffsets != buffer)
+				delete[] readerOffsets;
 		}
 
 		void PopulateDirectoryWithFiles(BinaryReader& reader, const StringKeyHasher& hasher, EntryLookup& files)
diff --git a/Engine/Source/Programs/UnrealBuildAccelerator/Core/Public/UbaStringBuffer.h b/Engine/Source/Programs/UnrealBuildAccelerator/Core/Public/UbaStringBuffer.h
--- a/Engine/Source/Programs/UnrealBuildAccelerator/Core/Public/UbaStringBuffer.h
+++ b/Engine/Source/Programs/UnrealBuildAccelerator/Core/Public/UbaStringBuffer.h
@@ -77,8 +77,7 @@ namespace uba
 		StringBuffer() : StringBufferBase(Capacity) { *buf = 0; }
 		explicit StringBuffer(const TString& str) : StringBufferBase(Capacity) { *buf = 0; Append(str); }
 		explicit StringBuffer(const tchar* str) : StringBufferBase(Capacity) { *buf = 0; if (str) Append(str); }
-		template<u32 Capacity2>
-		explicit StringBuffer(const StringBuffer<Capacity2>& str) : StringBufferBase(Capacity) { *buf = 0; Append(str); }
+		StringBuffer(const StringBufferBase& str) : StringBufferBase(Capacity) { *buf = 0; Append(str); }
 	private:
 		tchar buf[Capacity];
 	};
diff --git a/Engine/Source/Programs/UnrealBuildAccelerator/Detours/Private/PosixOS/UbaDetoursFunctionsPosix.cpp b/Engine/Source/Programs/UnrealBuildAccelerator/Detours/Private/PosixOS/UbaDetoursFunctionsPosix.cpp
--- a/Engine/Source/Programs/UnrealBuildAccelerator/Detours/Private/PosixOS/UbaDetoursFunctionsPosix.cpp
+++ b/Engine/Source/Programs/UnrealBuildAccelerator/Detours/Private/PosixOS/UbaDetoursFunctionsPosix.cpp
@@ -314,7 +314,7 @@ int Shared_open(const char* funcName, const char* file, int flags, int mode, con
 			{
 				// This could be a written file not reported to server yet
 				{
-					ScopedReadLock lock(g_mappedFileTable.m_lookupLock);
+					SCOPED_READ_LOCK(g_mappedFileTable.m_lookupLock, lock);
 					auto findIt = g_mappedFileTable.m_lookup.find(fileNameKey);
 					if (findIt != g_mappedFileTable.m_lookup.end())
 						allowEarlyOut = findIt->second.deleted;
@@ -342,7 +342,7 @@ int Shared_open(const char* funcName, const char* file, int flags, int mode, con
 	}
 
 
-	ScopedWriteLock _(g_mappedFileTable.m_lookupLock);
+	SCOPED_WRITE_LOCK(g_mappedFileTable.m_lookupLock, _);
 	auto insres = g_mappedFileTable.m_lookup.try_emplace(fileNameKey);
 	FileInfo& info = insres.first->second;
 	FileInfo* fileInfo = &info;
@@ -400,7 +400,7 @@ int Shared_open(const char* funcName, const char* file, int flags, int mode, con
 		if (realFileName[1] == 'd') // This is a directory and we will need to fake it locally.. 
 		{
 			int fd = TRUE_WRAPPER(open)("/dev/null", O_RDONLY);
-			ScopedWriteLock lock(g_fileHandlesLock);
+			SCOPED_WRITE_LOCK(g_fileHandlesLock, lock);
 			auto insres2 = g_fileHandles.insert({ fd, DetouredHandle() });
 			UBA_ASSERTF(insres2.second, "File handle for directory already added");
 			DetouredHandle& h = insres2.first->second;
@@ -443,7 +443,7 @@ int Shared_open(const char* funcName, const char* file, int flags, int mode, con
 	if (fd == -1)
 		return fd;
 
-	ScopedWriteLock lock(g_fileHandlesLock);
+	SCOPED_WRITE_LOCK(g_fileHandlesLock, lock);
 	auto insres2 = g_fileHandles.insert({ fd, DetouredHandle() });
 	UBA_ASSERTF(insres2.second, "File handle already added");
 	DetouredHandle& h = insres2.first->second;
@@ -502,7 +502,7 @@ void Shared_close(int fd, const TrueClose& trueClose)
 		return;
 	}
 
-	ScopedWriteLock lock(g_fileHandlesLock);
+	SCOPED_WRITE_LOCK(g_fileHandlesLock, lock);
 	auto findIt = g_fileHandles.find(fd);
 	if (findIt == g_fileHandles.end())
 	{
@@ -541,7 +541,7 @@ int Shared_fstat(const char* funcName, int fd, struct stat* attr, const True_fst
 	//if (!g_isDetouring || t_disallowDetour)
 	//	return trueFstat(fd, attr);
 
-	ScopedReadLock lock(g_fileHandlesLock);
+	SCOPED_READ_LOCK(g_fileHandlesLock, lock);
 	auto findIt = g_fileHandles.find(fd);
 	if (findIt == g_fileHandles.end())
 	{
@@ -646,7 +646,7 @@ int Shared_stat(const char* funcName, const char* file, struct stat* attr, const
 	{
 		struct stat attr2;
 		int res2 = trueStat(file, &attr2);
-		UBA_ASSERTF(res == res2, "stat: return value differs for %s (%i vs %i) [fixed: %s]", file, res, res2, fixedFile.data);
+		UBA_ASSERTF(res == res2, "stat: return value differs for %s (cached %i vs actual %i) [fixed: %s]", file, res, res2, fixedFile.data);
 		if (res != -1)
 		{
 			bool isDir = S_ISDIR(attr->st_mode);
@@ -796,7 +796,7 @@ UBA_EXPORT ssize_t UBA_WRAPPER(readlink)(const char* pathname, char* buf, size_t
 		u32 fd;
 		if (!fdStr.Parse(fd))
 			UBA_ASSERTF(false, "Failed to parse /proc/self/fd");
-		ScopedReadLock lock(g_fileHandlesLock);
+		SCOPED_READ_LOCK(g_fileHandlesLock, lock);
 		auto findIt = g_fileHandles.find(fd);
 		if (findIt != g_fileHandles.end())
 		{
@@ -896,7 +896,7 @@ UBA_EXPORT int UBA_WRAPPER(dup2)(int oldfd, int newfd)
 
 	if (res != -1)
 	{
-		ScopedWriteLock lock(g_fileHandlesLock);
+		SCOPED_WRITE_LOCK(g_fileHandlesLock, lock);
 		auto findIt = g_fileHandles.find(oldfd);
 		if (findIt != g_fileHandles.end())
 		{
@@ -1213,7 +1213,7 @@ UBA_EXPORT int UBA_WRAPPER(rename)(const char* oldpath, const char* newpath)
 
 	// TODO: This might be really slow but it seems you can rename files on linux while they are open and they won't be properly renamed until closed
 	{
-		ScopedReadLock lock(g_fileHandlesLock);
+		SCOPED_READ_LOCK(g_fileHandlesLock, lock);
 		for (auto& kv : g_fileHandles)
 		{
 			FileObject& fo = *kv.second.fileObject;
@@ -1252,7 +1252,7 @@ UBA_EXPORT int UBA_WRAPPER(rename)(const char* oldpath, const char* newpath)
 	bool result;
 	{
 		TimerScope ts(g_stats.moveFile);
-		ScopedWriteLock pcs(g_communicationLock);
+		SCOPED_WRITE_LOCK(g_communicationLock, pcs);
 		BinaryWriter writer;
 		writer.WriteByte(MessageType_MoveFileW);
 		writer.WriteStringKey(oldKey);
@@ -1292,7 +1292,7 @@ UBA_EXPORT int UBA_WRAPPER(chmod)(const char* pathname, mode_t mode)
 	u32 errorCode;
 	{
 		TimerScope ts(g_stats.chmod);
-		ScopedWriteLock pcs(g_communicationLock);
+		SCOPED_WRITE_LOCK(g_communicationLock, pcs);
 		BinaryWriter writer;
 		writer.WriteByte(MessageType_Chmod);
 		writer.WriteStringKey(key);
@@ -1383,7 +1383,7 @@ UBA_EXPORT int UBA_WRAPPER(remove)(const char* pathname)
 	{
 		u32 closeId = 0;
 		TimerScope ts(g_stats.deleteFile);
-		ScopedWriteLock pcs(g_communicationLock);
+		SCOPED_WRITE_LOCK(g_communicationLock, pcs);
 		BinaryWriter writer;
 		writer.WriteByte(MessageType_DeleteFileW);
 		writer.WriteString(fixedName);
@@ -1435,7 +1435,7 @@ UBA_EXPORT int UBA_WRAPPER(posix_spawn)(pid_t* pid, const char* path, const posi
 
 	{
 		TimerScope ts(g_stats.createProcess);
-		ScopedWriteLock pcs(g_communicationLock);
+		SCOPED_WRITE_LOCK(g_communicationLock, pcs);
 
 		const char* pwd = "";
 		//while (char* env = envp[i++])
@@ -1511,7 +1511,7 @@ UBA_EXPORT int UBA_WRAPPER(posix_spawn)(pid_t* pid, const char* path, const posi
 
 	{
 		TimerScope ts(g_stats.createProcess);
-		ScopedWriteLock pcs(g_communicationLock);
+		SCOPED_WRITE_LOCK(g_communicationLock, pcs);
 		BinaryWriter writer;
 		writer.WriteByte(MessageType_StartProcess);
 		writer.WriteU32(processId);
@@ -1614,7 +1614,7 @@ int Internal_execve(const char* pathname, char* const _Nullable argv[], char* co
 
 	{
 		TimerScope ts(g_stats.createProcess);
-		ScopedWriteLock pcs(g_communicationLock);
+		SCOPED_WRITE_LOCK(g_communicationLock, pcs);
 		BinaryWriter writer;
 		writer.WriteByte(MessageType_ExitChildProcess);
 		writer.WriteU32(pid);
@@ -1872,7 +1872,7 @@ namespace uba
 
 		{
 			TimerScope ts(g_stats.init);
-			ScopedWriteLock pcs(g_communicationLock);
+			SCOPED_WRITE_LOCK(g_communicationLock, pcs);
 			BinaryWriter writer;
 			writer.WriteByte(MessageType_Init);
 			writer.Flush();
@@ -1937,7 +1937,7 @@ namespace uba
 		g_isDetouring = false;
 
 		{
-			ScopedWriteLock lock(g_fileHandlesLock);
+			SCOPED_WRITE_LOCK(g_fileHandlesLock, lock);
 			for (auto& kv : g_fileHandles)
 			{
 				TRUE_WRAPPER(close)(kv.first);
@@ -2005,3 +2005,52 @@ namespace uba
 		TRUE_WRAPPER(_exit)(int(terminateCode));
 	}
 }
+
+extern "C"
+{
+	UBA_EXPORT bool UbaRequestNextProcess(u32 prevExitCode, char* outArguments, u32 outArgumentsCapacity)
+	{
+		#if UBA_DEBUG_LOG_ENABLED
+		//FlushDebugLog();
+		#endif
+
+		*outArguments = 0;
+		bool newProcess;
+		{
+			SCOPED_WRITE_LOCK(g_communicationLock, pcs);
+			BinaryWriter writer;
+			writer.WriteByte(MessageType_GetNextProcess);
+			writer.WriteU32(prevExitCode);
+			g_stats.Write(writer);
+
+
+			writer.Flush();
+			BinaryReader reader;
+			newProcess = reader.ReadBool();
+			if (newProcess)
+			{
+				reader.ReadString(outArguments, outArgumentsCapacity);
+				reader.SkipString(); // workingDir
+				reader.SkipString(); // description
+				reader.ReadString(g_logName.Clear());
+			}
+		}
+
+		if (newProcess)
+		{
+			g_stats = {};
+
+			//#if UBA_DEBUG_LOG_ENABLED
+			//SuppressCreateFileDetourScope scope;
+			//HANDLE debugFile = (HANDLE)g_debugFile;
+			//g_debugFile = InvalidFileHandle;
+			//CloseHandle(debugFile);
+			//debugFile = CreateFileW(g_logName.data, GENERIC_WRITE, FILE_SHARE_READ, NULL, CREATE_ALWAYS, FILE_ATTRIBUTE_NORMAL, NULL);
+			//g_debugFile = (FileHandle)(u64)debugFile;
+			//#endif
+		}
+
+		Rpc_UpdateTables();
+		return newProcess;
+	}
+}
\ No newline at end of file
diff --git a/Engine/Source/Programs/UnrealBuildAccelerator/Detours/Private/UbaDetoursShared.cpp b/Engine/Source/Programs/UnrealBuildAccelerator/Detours/Private/UbaDetoursShared.cpp
--- a/Engine/Source/Programs/UnrealBuildAccelerator/Detours/Private/UbaDetoursShared.cpp
+++ b/Engine/Source/Programs/UnrealBuildAccelerator/Detours/Private/UbaDetoursShared.cpp
@@ -37,7 +37,8 @@ namespace uba
 		g_systemRootMem.Create();
 		g_systemTempMem.Create();
 
-		g_memoryBlockMem.Create(160 * 1024 * 1024);
+		u64 reserveSizeMb = IsWindows ? 160 : 1024; // The sync primitives on linux/macos is much bigger
+		g_memoryBlockMem.Create(reserveSizeMb * 1024 * 1024);
 		g_directoryTableMem.Create(&g_memoryBlock);
 		g_mappedFileTableMem.Create(g_memoryBlock);
 		g_consoleStringCsMem.Create();
diff --git a/Engine/Source/Programs/UnrealBuildAccelerator/Detours/Private/Windows/UbaDetoursFunctionsKernelBase.inl b/Engine/Source/Programs/UnrealBuildAccelerator/Detours/Private/Windows/UbaDetoursFunctionsKernelBase.inl
--- a/Engine/Source/Programs/UnrealBuildAccelerator/Detours/Private/Windows/UbaDetoursFunctionsKernelBase.inl
+++ b/Engine/Source/Programs/UnrealBuildAccelerator/Detours/Private/Windows/UbaDetoursFunctionsKernelBase.inl
@@ -1320,7 +1320,7 @@ bool Shared_MoveFile(LPCWSTR lpExistingFileName, LPCWSTR lpNewFileName, DWORD dw
 			lock2.Leave();
 			FileInfo& destInfo = insres.first->second;
 			UBA_ASSERT(!insres.second); // This is here just to get a chance to investigate this scenario.. might work
-			UBA_ASSERT(!destInfo.trueFileMapHandle && !destInfo.memoryFile);
+			UBA_ASSERTF(!destInfo.trueFileMapHandle && (!destInfo.memoryFile || g_rules->IsThrowAway(dest.data, dest.count)), TC("Moving file %s to %s that is an output file that is not a memory file is not supported"), source.data, lpNewFileName);
 			destInfo.memoryFile = sourceInfo.memoryFile;
 			sourceInfo.memoryFile = nullptr;
 			DEBUG_LOG_DETOURED(L"MoveFileExW", L"(memfile->memfile) %ls to %ls -> Success", lpExistingFileName, lpNewFileName);
diff --git a/Engine/Source/Programs/UnrealBuildAccelerator/Detours/Private/Windows/UbaDetoursFunctionsNtDll.inl b/Engine/Source/Programs/UnrealBuildAccelerator/Detours/Private/Windows/UbaDetoursFunctionsNtDll.inl
--- a/Engine/Source/Programs/UnrealBuildAccelerator/Detours/Private/Windows/UbaDetoursFunctionsNtDll.inl
+++ b/Engine/Source/Programs/UnrealBuildAccelerator/Detours/Private/Windows/UbaDetoursFunctionsNtDll.inl
@@ -679,7 +679,10 @@ NTSTATUS NTAPI Shared_NtCreateFile(bool IsCreateFunc, PHANDLE hFileHandle, ACCES
 					DirectoryTable::EntryInformation entryInfo;
 					g_directoryTable.GetEntryInformation(entryInfo, dirTableOffset);
 					if (entryInfo.attributes == 0)
+					{
+						DEBUG_LOG_DETOURED(funcName, L"DELETED %llu, (%ls) -> Success", uintptr_t(*hFileHandle), fileName.data);
 						return STATUS_OBJECT_NAME_NOT_FOUND;
+					}
 				}
 
 				bool isWriteAttributes = (DesiredAccess & FILE_WRITE_ATTRIBUTES) != 0;
@@ -983,7 +986,7 @@ NTSTATUS NTAPI Shared_NtCreateFile(bool IsCreateFunc, PHANDLE hFileHandle, ACCES
 			{
 				if (g_rules->IsThrowAway(fileName.data, fileName.count))
 				{
-					isDeleteOnClose = true;
+					//isDeleteOnClose = true;
 				}
 				else if (CreateDisposition == FILE_OPEN)
 				{
@@ -993,7 +996,7 @@ NTSTATUS NTAPI Shared_NtCreateFile(bool IsCreateFunc, PHANDLE hFileHandle, ACCES
 				}
 
 				bool isLocal = !IsOutputFile(fileName.data, fileName.count, dwDesiredAccess, isDeleteOnClose);
-				UBA_ASSERT(CreateDisposition != FILE_OPEN || Contains(fileName.data, L"vctip_"));
+				//UBA_ASSERTF(CreateDisposition != FILE_OPEN || Contains(fileName.data, L"vctip_"), TC("Unsupported disposition %u for file %s"), CreateDisposition, fileName.data);
 				info.memoryFile = new MemoryFile(isLocal, FileTypeMaxSize(fileName, isSystemOrTempFile));
 			}
 
@@ -1223,7 +1226,8 @@ NTSTATUS NTAPI Detoured_NtClose(HANDLE handle)
 		mappingHandle = fi.memoryFile->mappingHandle;
 		mappingWritten = fi.memoryFile->writtenSize;
 
-		if (IsOutputFile(fi.originalName, wcslen(fi.originalName), fo->desiredAccess, fo->deleteOnClose))
+		u32 orginalNameLen = TStrlen(fi.originalName);
+		if (IsOutputFile(fi.originalName, orginalNameLen, fo->desiredAccess, fo->deleteOnClose) && !g_rules->IsThrowAway(fi.originalName, orginalNameLen))
 		{
 			// Need to report this file to host so it can be tracked in directory table
 			if (!fi.memoryFile->isReported)
diff --git a/Engine/Source/Programs/UnrealBuildAccelerator/Host/Private/UbaExports.cpp b/Engine/Source/Programs/UnrealBuildAccelerator/Host/Private/UbaExports.cpp
--- a/Engine/Source/Programs/UnrealBuildAccelerator/Host/Private/UbaExports.cpp
+++ b/Engine/Source/Programs/UnrealBuildAccelerator/Host/Private/UbaExports.cpp
@@ -143,7 +143,10 @@ extern "C"
 
 	void Server_Stop(uba::NetworkServer* server)
 	{
-		server->StopAll();
+		auto s = (uba::NetworkServerWithBackend*)server;
+		auto networkBackend = s->networkBackend;
+		networkBackend->StopListen();
+		server->DisconnectClients();
 	}
 
 
@@ -432,7 +435,11 @@ uba::StorageClient* CreateStorageClient(uba::NetworkClient& client, const uba::t
 	void DestroySessionServer(uba::SessionServer* server)
 	{
 		if (server)
-			server->GetServer().StopAll();
+		{
+			auto& s = (uba::NetworkServerWithBackend&)server->GetServer();
+			s.networkBackend->StopListen();
+			s.DisconnectClients();
+		}
 		delete server;
 	}
 
diff --git a/Engine/Source/Programs/UnrealBuildAccelerator/StorageProxy/Private/UbaStorageProxyApp.cpp b/Engine/Source/Programs/UnrealBuildAccelerator/StorageProxy/Private/UbaStorageProxyApp.cpp
--- a/Engine/Source/Programs/UnrealBuildAccelerator/StorageProxy/Private/UbaStorageProxyApp.cpp
+++ b/Engine/Source/Programs/UnrealBuildAccelerator/StorageProxy/Private/UbaStorageProxyApp.cpp
@@ -184,8 +184,8 @@ int wmain(int argc, wchar_t *argv[])
 			});
 
 		disconnected.IsSet();
-
-		server.StopAll();
+		networkBackend.StopListen();
+		server.DisconnectClients();
 		server.PrintSummary(logger);
 		client->PrintSummary(logger);
 
diff --git a/Engine/Source/Programs/UnrealBuildAccelerator/Test/Private/Linux/UbaTestLinux.cpp b/Engine/Source/Programs/UnrealBuildAccelerator/Test/Private/Linux/UbaTestLinux.cpp
--- a/Engine/Source/Programs/UnrealBuildAccelerator/Test/Private/Linux/UbaTestLinux.cpp
+++ b/Engine/Source/Programs/UnrealBuildAccelerator/Test/Private/Linux/UbaTestLinux.cpp
@@ -9,14 +9,14 @@ void segfault_sigaction(int signal, siginfo_t* si, void* arg)
 
 int main()
 {
+	#if !__has_feature(thread_sanitizer)
 	struct sigaction sa;
-
 	memset(&sa, 0, sizeof(struct sigaction));
 	sigemptyset(&sa.sa_mask);
 	sa.sa_sigaction = segfault_sigaction;
 	sa.sa_flags = SA_SIGINFO;
-
 	sigaction(SIGSEGV, &sa, NULL);
+	#endif
 
 	using namespace uba;
 	if (!RunAllTests())
diff --git a/Engine/Source/Programs/UnrealBuildAccelerator/Test/Private/UbaTestAll.h b/Engine/Source/Programs/UnrealBuildAccelerator/Test/Private/UbaTestAll.h
--- a/Engine/Source/Programs/UnrealBuildAccelerator/Test/Private/UbaTestAll.h
+++ b/Engine/Source/Programs/UnrealBuildAccelerator/Test/Private/UbaTestAll.h
@@ -5,8 +5,9 @@
 #include "UbaTestBasics.h"
 #include "UbaTestNetwork.h"
 #include "UbaTestScheduler.h"
-#include "UbaTestStorage.h"
 #include "UbaTestSession.h"
+#include "UbaTestStorage.h"
+#include "UbaTestStress.h"
 
 namespace uba
 {
@@ -29,7 +30,9 @@ namespace uba
 		UBA_TEST(TestEvents) \
 		UBA_TEST(TestPaths) \
 		UBA_TEST(TestFiles) \
+		UBA_TEST(TestMemoryBlock) \
 		UBA_TEST(TestParseArguments) \
+		UBA_TEST(TestBinaryWriter) \
 		UBA_TEST(TestSockets) \
 		UBA_TEST(TestClientServer) \
 		UBA_TEST(TestClientServer2) \
@@ -71,7 +74,7 @@ namespace uba
 
 		logger.Info(TC("Running tests (Test rootdir: %s)"), testRootDir.data);
 
-		//UBA_TEST(TestLogLines)
+		//UBA_TEST(TestStress) // This can not be submitted.. it depends on CoordinatorHorde and credentials
 		UBA_TESTS
 
 		logger.Info(TC("Tests finished successfully!"));
diff --git a/Engine/Source/Programs/UnrealBuildAccelerator/Test/Private/UbaTestNetwork.h b/Engine/Source/Programs/UnrealBuildAccelerator/Test/Private/UbaTestNetwork.h
--- a/Engine/Source/Programs/UnrealBuildAccelerator/Test/Private/UbaTestNetwork.h
+++ b/Engine/Source/Programs/UnrealBuildAccelerator/Test/Private/UbaTestNetwork.h
@@ -75,9 +75,11 @@ namespace uba
 				return true;
 			});
 
+		auto ds = MakeGuard([&]() { server.DisconnectClients(); });
 		if (!server.StartListen(networkBackend, 1234))
 			return logger.Error(TC("Failed to listen"));
 		Sleep(100);
+		auto dg = MakeGuard([&]() { client.Disconnect(); });
 		if (!client.Connect(networkBackend, TC("127.0.0.1"), 1234))
 			return logger.Error(TC("Failed to connect"));
 
@@ -114,6 +116,8 @@ namespace uba
 		if (!client.StartListen(networkBackend, 1239))
 			return logger.Error(TC("Client failed to listen"));
 		Sleep(100);
+		auto ds = MakeGuard([&]() { server.DisconnectClients(); });
+		auto dg = MakeGuard([&]() { client.Disconnect(); });
 		if (!server.AddClient(networkBackend, TC("127.0.0.1"), 1239))
 			return logger.Error(TC("Server failed to connect"));
 		if (!server.AddClient(networkBackend, TC("127.0.0.1"), 1239))
@@ -156,9 +160,11 @@ namespace uba
 				return true;
 			});
 
+		auto ds = MakeGuard([&]() { server.DisconnectClients(); });
 		if (!server.StartListen(networkBackend, 1234))
 			return logger.Error(TC("Failed to listen"));
 		Sleep(100);
+		auto dg = MakeGuard([&]() { client.Disconnect(); });
 		if (!client.Connect(networkBackend, TC("127.0.0.1"), 1234))
 			return logger.Error(TC("Failed to connect"));
 
diff --git a/Engine/Source/Programs/UnrealBuildAccelerator/Test/Private/UbaTestSession.h b/Engine/Source/Programs/UnrealBuildAccelerator/Test/Private/UbaTestSession.h
--- a/Engine/Source/Programs/UnrealBuildAccelerator/Test/Private/UbaTestSession.h
+++ b/Engine/Source/Programs/UnrealBuildAccelerator/Test/Private/UbaTestSession.h
@@ -27,6 +27,7 @@ namespace uba
 
 		bool ctorSuccess = true;
 		NetworkServer server(ctorSuccess, { logWriter });
+
 		SessionServerCreateInfo sessionServerInfo(storage, server, logWriter);
 		sessionServerInfo.checkMemory = false;
 		sessionServerInfo.rootDir = rootDir.data;
@@ -67,7 +68,7 @@ namespace uba
 		sessionServerInfo.rootDir = rootDir.data;
 		SessionServer sessionServer(sessionServerInfo);
 
-		auto sg = MakeGuard([&]() { server.StopAll(); });
+		auto sg = MakeGuard([&]() { server.DisconnectClients(); });
 
 		rootDir.Append(TC("Client"));
 		if (!DeleteAllFiles(logger, rootDir.data))
@@ -80,7 +81,7 @@ namespace uba
 		sessionClientInfo.rootDir = rootDir.data;
 		SessionClient sessionClient(sessionClientInfo);
 
-		auto cg = MakeGuard([&]() { sessionClient.Stop(); client.StopAll(); });
+		auto cg = MakeGuard([&]() { sessionClient.Stop(); client.Disconnect(); });
 
 		StringBuffer<> workingDir;
 		workingDir.Append(testRootDir).Append(TC("WorkingDir"));
@@ -91,6 +92,7 @@ namespace uba
 		if (!DeleteAllFiles(logger, workingDir.data, false))
 			return false;
 
+		storageClient.Start();
 		sessionClient.Start();
 
 		u16 port = 1356;
diff --git a/Engine/Source/Programs/UnrealBuildAccelerator/Test/Private/UbaTestStorage.h b/Engine/Source/Programs/UnrealBuildAccelerator/Test/Private/UbaTestStorage.h
--- a/Engine/Source/Programs/UnrealBuildAccelerator/Test/Private/UbaTestStorage.h
+++ b/Engine/Source/Programs/UnrealBuildAccelerator/Test/Private/UbaTestStorage.h
@@ -106,7 +106,7 @@ namespace uba
 		storageServerInfo.casCapacityBytes = 1024ull * 1024 * 1024;
 		StorageServer storageServer(storageServerInfo);
 
-		auto g = MakeGuard([&]() { server.StopAll(); });
+		auto g = MakeGuard([&]() { server.DisconnectClients(); });
 
 		rootDir.Append(TC("Client"));
 		if (!DeleteAllFiles(logger, rootDir.data))
@@ -114,6 +114,7 @@ namespace uba
 
 		StorageClientCreateInfo storageClientInfo(client, rootDir.data);
 		StorageClient storageClient(storageClientInfo);
+		storageClient.Start();
 
 		if (!storageClient.LoadCasTable(true))
 			return false;
diff --git a/Engine/Source/Programs/UnrealBuildAccelerator/Test/Private/UbaTestStress.h b/Engine/Source/Programs/UnrealBuildAccelerator/Test/Private/UbaTestStress.h
new file mode 100644
--- /dev/null
+++ b/Engine/Source/Programs/UnrealBuildAccelerator/Test/Private/UbaTestStress.h
@@ -0,0 +1,189 @@
+// Copyright Epic Games, Inc. All Rights Reserved.
+
+#pragma once
+
+#include "UbaCoordinatorWrapper.h"
+#include "UbaScheduler.h"
+#include "UbaSessionServer.h"
+#include "UbaStorageServer.h"
+
+namespace uba
+{
+	bool TestStress(LoggerWithWriter& logger, const StringBufferBase& testRootDir)
+	{
+		LogWriter& logWriter = logger.m_writer;
+		NetworkBackendTcp networkBackend(logWriter);
+
+		bool ctorSuccess = true;
+		NetworkServer networkServer(ctorSuccess, { logWriter });
+
+		StringBuffer<> rootDir;
+		rootDir.Append(testRootDir).Append(TC("Uba"));
+		if (!DeleteAllFiles(logger, rootDir.data))
+			return false;
+
+		StorageServerCreateInfo storageServerInfo(networkServer, rootDir.data, logWriter);
+		storageServerInfo.casCapacityBytes = 1024ull * 1024 * 1024;
+		StorageServer storageServer(storageServerInfo);
+
+		SessionServerCreateInfo sessionServerInfo(storageServer, networkServer, logWriter);
+		sessionServerInfo.checkMemory = false;
+		sessionServerInfo.rootDir = rootDir.data;
+		sessionServerInfo.traceEnabled = true;
+		//sessionServerInfo.remoteLogEnabled = true;
+		SessionServer sessionServer(sessionServerInfo);
+
+		auto sg = MakeGuard([&]() { networkServer.DisconnectClients(); });
+
+		StringBuffer<> workingDir;
+		workingDir.Append(testRootDir).Append(TC("WorkingDir"));
+		if (!DeleteAllFiles(logger, workingDir.data))
+			return false;
+		if (!storageServer.CreateDirectory(workingDir.data))
+			return false;
+		if (!DeleteAllFiles(logger, workingDir.data, false))
+			return false;
+		workingDir.EnsureEndsWithSlash();
+
+		SchedulerCreateInfo info(sessionServer);
+		//info.forceRemote = true;
+		info.enableProcessReuse = true;
+		info.maxLocalProcessors = 1;
+		Scheduler scheduler(info);
+		scheduler.Start();
+
+		if (!networkServer.StartListen(networkBackend))
+			return logger.Error(TC("Failed to listen"));
+
+		CoordinatorWrapper coordinator;
+
+		bool useCoordinator = true;
+		if (useCoordinator)
+		{
+			return logger.Error(TC("uri/pool/oidc must be populated for this test to work"));
+
+			StringBuffer<512> coordinatorWorkDir(testRootDir);
+			coordinatorWorkDir.EnsureEndsWithSlash().Append(TC("Horde"));
+			StringBuffer<512> binariesDir;
+			if (!GetDirectoryOfCurrentModule(logger, binariesDir))
+				return false;
+
+			CoordinatorCreateInfo cinfo;
+			cinfo.workDir = coordinatorWorkDir.data;
+			cinfo.binariesDir = binariesDir.data;
+
+			cinfo.uri = TC("https://horde-uri");
+			cinfo.pool = TC("PoolToUse");
+			cinfo.oidc = TC("Oidc");
+
+
+			cinfo.maxCoreCount = 400;
+			cinfo.logging = true;
+			if (!coordinator.Create(logger, TC("Horde"), cinfo, networkBackend, networkServer))
+				return false;
+		}
+		auto cg = MakeGuard([&]() { coordinator.Destroy(); });
+
+		StringBuffer<> testApp;
+		GetTestAppPath(logger, testApp);
+
+		u32 counter = 0;
+
+		Atomic<u32> activeCount;
+		Atomic<bool> error;
+		while (true)
+		{
+			StringBuffer<512> fileName;
+			fileName.Append(workingDir).Append(TC("File_")).AppendValue(counter++).Append(TC(".in"));
+			{
+				FileAccessor fa(logger, fileName.data);
+				fa.CreateWrite();
+				Guid g;
+				CreateGuid(g);
+				fa.Write(&g, sizeof(g));
+				fa.Close();
+			}
+			sessionServer.RegisterNewFile(fileName.data);
+			StringBuffer<512> args;
+			args.Append(TC("-file=")).Append(fileName);
+			ProcessStartInfo pi;
+			pi.application = testApp.data;
+			pi.arguments = args.data;
+			pi.workingDir = workingDir.data;
+			pi.description = fileName.Last(PathSeparator) + 6;
+
+			struct Rec
+			{
+				Logger& logger;
+				StringBuffer<256> inFile;
+				Atomic<u32>& activeCount;
+				Atomic<bool>& error;
+			};
+			pi.userData = new Rec { logger, fileName, activeCount, error };
+			pi.exitedFunc = [](void* userData, const ProcessHandle& ph)
+				{
+					auto rec = (Rec*)userData;
+					for (auto& line : ph.GetLogLines())
+						rec->logger.Log(line.type, line.text.c_str(), u32(line.text.size()));
+
+					if (u32 exitCode = ph.GetExitCode())
+					{
+						rec->logger.Error(TC("Process %s failed with exit code %u"), ph.GetStartInfo().description, exitCode);
+						rec->error = true;
+					}
+					else
+					{
+						auto& logger = rec->logger;
+						if (!DeleteFileW(rec->inFile.data))
+						{
+							logger.Error(TC("Failed to delete file %s"), rec->inFile.data);
+							rec->error = true;
+						}
+						StringBuffer<256> outFile;
+						outFile.Append(rec->inFile.data, rec->inFile.count - 2).Append(TC("out"));
+						if (!FileExists(logger, outFile.data))
+						{
+							logger.Error(TC("Couldn't find output file %s"), outFile.data);
+							rec->error = true;
+						}
+						if (!DeleteFileW(outFile.data))
+						{
+							logger.Error(TC("Failed to delete file %s"), outFile.data);
+							rec->error = true;
+						}
+					}
+					--rec->activeCount;
+					delete rec;
+				};
+			EnqueueProcessInfo epi(pi);
+			scheduler.EnqueueProcess(epi);
+			++activeCount;
+
+			while (activeCount > 800)
+				Sleep(100);
+			if (error)
+				break;
+
+			if (counter != 20000)
+				continue;
+
+			while (true)
+			{
+				u32 queued;
+				u32 activeLocal;
+				u32 activeRemote;
+				u32 finished;
+				scheduler.GetStats(queued, activeLocal, activeRemote, finished);
+				if (finished == 20000)
+					break;
+			}
+			break;
+		}
+
+		networkBackend.StopListen();
+		networkServer.DisconnectClients();
+		scheduler.Stop();
+
+		return !error;
+	}
+}
diff --git a/Engine/Source/Programs/UnrealBuildAccelerator/TestApp/Private/PosixOS/UbaTestAppPosix.cpp b/Engine/Source/Programs/UnrealBuildAccelerator/TestApp/Private/PosixOS/UbaTestAppPosix.cpp
--- a/Engine/Source/Programs/UnrealBuildAccelerator/TestApp/Private/PosixOS/UbaTestAppPosix.cpp
+++ b/Engine/Source/Programs/UnrealBuildAccelerator/TestApp/Private/PosixOS/UbaTestAppPosix.cpp
@@ -61,7 +61,7 @@ bool IsLibraryLoaded(const char* libraryToMatch)
 // UbaTestAppPosix cannot be run standalone. It is extremely
 // dependent on the UbaTest runner. Please see details in:
 // UbaTestSession.h
-int main()
+int main(int argc, char* argv[])
 {
 	bool runningRemote = false;
 	char* tmp = getenv("UBA_REMOTE");
@@ -75,84 +75,133 @@ int main()
 	if (!foundDetoursLib)
 		return LogError("libUbaDetours not loaded. This app is designed to only start from inside UnrealBuildAccelerator.");
 
-	char cwd[1024];
-	if (!getcwd(cwd, sizeof(cwd)))
-		return LogError("getcwd failed");
-
-	struct stat attrR;
-	if (stat("FileR.h", &attrR) == -1)
-		return LogError("stat for FileR.h failed");
-	if (S_ISREG(attrR.st_mode) == 0)
-		return LogError("stat for FileR.h did not return normal file");
-
-	int fdr = open("FileR.h", O_RDONLY);
-	if (fdr == -1)
-		return LogError("open FileR.h failed");
+	if (argc == 1)
+	{
+		char cwd[1024];
+		if (!getcwd(cwd, sizeof(cwd)))
+			return LogError("getcwd failed");
+
+		struct stat attrR;
+		if (stat("FileR.h", &attrR) == -1)
+			return LogError("stat for FileR.h failed");
+		if (S_ISREG(attrR.st_mode) == 0)
+			return LogError("stat for FileR.h did not return normal file");
+
+		int fdr = open("FileR.h", O_RDONLY);
+		if (fdr == -1)
+			return LogError("open FileR.h failed");
 	
-	char buf[4];
-	if (read(fdr, buf, 4) != 4)
-		return LogError("Failed to read FileR.h");
-
-	if (strcmp(buf, "Foo") != 0)
-		return LogError("FileR.h content was wrong");
-
-	if (close(fdr) == -1)
-		return LogError("close FileR.h failed");
-
-	int fdw = open("FileW", O_WRONLY | O_CREAT | O_TRUNC, S_IRUSR | S_IWUSR);
-	if (fdw == -1)
-		return LogError("open FileW failed");
-
-	if (write(fdw, "hello", 6) == -1)
-		return LogError("write FileW failed");
-
-	struct stat attrW1;
-	if (fstat(fdw, &attrW1) == -1)
-		return LogError("fstat FileW failed");
-
-	if (close(fdw) == -1)
-		return LogError("closed FileW failed");
-
-	struct stat attrW2;
-	if (stat("FileW", &attrW2) == -1)
-		return LogError("stat for FileW failed");
-	if (S_ISREG(attrW2.st_mode) == 0)
-		return LogError("stat for FileW did not return normal file");
-
-	if (rename("FileW", "FileW2") == -1)
-		return LogError("rename for FileW to FileW2 failed");
-
-
-	struct stat attrD1;
-	if (stat("Dir1", &attrD1) == -1)
-		return LogError("stat for FileW failed");
-	if (S_ISREG(attrD1.st_mode) != 0)
-		return LogError("stat for FileW did not return directory");
-
-	struct stat attr3;
-	if (stat("/usr", &attr3) == -1)
-		return LogError("stat for /usr failed");
-
-	FILE* f = fopen("FileWF", "w+");
-	if (f == nullptr)
-		return LogError("fopen FileWF failed");
-	if (fwrite("Hello", 1, 6, f) != 6)
-		return LogError("fwrite FileWF failed");
-	if (fclose(f) != 0)
-		return LogError("fclose FileWF failed");
-	struct stat attrWF;
-	if (stat("FileWF", &attrWF) == -1)
-		return LogError("stat for FileW failed");
-
-#if !PLATFORM_LINUX // Farm linux machines do not have clang installed... need to revisit this
-	char fullPath[PATH_MAX];
-	if (realpath("/usr/bin/clang", fullPath) == nullptr)
-		return LogError("realpath for 'clang' failed");
-#endif
-
-	struct stat attrRoot;
-	if (stat("/", &attrRoot) != 0)
-		return LogError("stat for '/' failed");
+		char buf[4];
+		if (read(fdr, buf, 4) != 4)
+			return LogError("Failed to read FileR.h");
+
+		if (strcmp(buf, "Foo") != 0)
+			return LogError("FileR.h content was wrong");
+
+		if (close(fdr) == -1)
+			return LogError("close FileR.h failed");
+
+		int fdw = open("FileW", O_WRONLY | O_CREAT | O_TRUNC, S_IRUSR | S_IWUSR);
+		if (fdw == -1)
+			return LogError("open FileW failed");
+
+		if (write(fdw, "hello", 6) == -1)
+			return LogError("write FileW failed");
+
+		struct stat attrW1;
+		if (fstat(fdw, &attrW1) == -1)
+			return LogError("fstat FileW failed");
+
+		if (close(fdw) == -1)
+			return LogError("closed FileW failed");
+
+		struct stat attrW2;
+		if (stat("FileW", &attrW2) == -1)
+			return LogError("stat for FileW failed");
+		if (S_ISREG(attrW2.st_mode) == 0)
+			return LogError("stat for FileW did not return normal file");
+
+		if (rename("FileW", "FileW2") == -1)
+			return LogError("rename for FileW to FileW2 failed");
+
+
+		struct stat attrD1;
+		if (stat("Dir1", &attrD1) == -1)
+			return LogError("stat for FileW failed");
+		if (S_ISREG(attrD1.st_mode) != 0)
+			return LogError("stat for FileW did not return directory");
+
+		struct stat attr3;
+		if (stat("/usr", &attr3) == -1)
+			return LogError("stat for /usr failed");
+
+		FILE* f = fopen("FileWF", "w+");
+		if (f == nullptr)
+			return LogError("fopen FileWF failed");
+		if (fwrite("Hello", 1, 6, f) != 6)
+			return LogError("fwrite FileWF failed");
+		if (fclose(f) != 0)
+			return LogError("fclose FileWF failed");
+		struct stat attrWF;
+		if (stat("FileWF", &attrWF) == -1)
+			return LogError("stat for FileW failed");
+
+	#if !PLATFORM_LINUX // Farm linux machines do not have clang installed... need to revisit this
+		char fullPath[PATH_MAX];
+		if (realpath("/usr/bin/clang", fullPath) == nullptr)
+			return LogError("realpath for 'clang' failed");
+	#endif
+
+		struct stat attrRoot;
+		if (stat("/", &attrRoot) != 0)
+			return LogError("stat for '/' failed");
+		return 0;
+	}
+	else if (strncmp(argv[1], "-file=", 6) == 0)
+	{
+		void* detoursHandle = dlopen(UBA_DETOURS_LIBRARY, RTLD_LAZY);
+		if (!detoursHandle)
+			return -3;
+		using UbaRequestNextProcessFunc = bool(unsigned int prevExitCode, char* outArguments, unsigned int outArgumentsCapacity);
+		static UbaRequestNextProcessFunc* requestNextProcess = (UbaRequestNextProcessFunc*)(void*)dlsym(detoursHandle, "UbaRequestNextProcess");
+		if (!requestNextProcess)
+			return -8;
+
+		char arguments[1024];
+		const char* file = argv[1] + 6;
+		while (true)
+		{
+			int rh = open(file, O_RDONLY);
+			if (rh == -1)
+				return LogError("Failed to open file %s", file);
+			if (close(rh) == -1)
+				return LogError("Failed to close file %s", file);
+
+			srand(getpid());
+			int milliseconds = rand() % 2000;
+			struct timespec ts;
+			ts.tv_sec = milliseconds / 1000;
+			ts.tv_nsec = (milliseconds % 1000) * 1000000;
+			nanosleep(&ts, NULL);
+
+			char outFile[1024];
+			strcpy(outFile, file);
+			outFile[strlen(file)-3] = 0;
+			strcat(outFile, ".out");
+
+			int wh = open(outFile, O_WRONLY | O_CREAT | O_TRUNC, S_IRUSR | S_IWUSR);
+			if (wh == -1)
+				return LogError("Failed to create file %s", file);
+			if (close(wh) == -1)
+				return LogError("Failed to close created file %s", file);
+
+			// Request new process
+			if (!requestNextProcess(0, arguments, 1024))
+				break; // No process available, exit loop
+			file = arguments + 6;
+		}
 
-	return 0;
+		return 0;
+	}
+	return -2;
 }
diff --git a/Engine/Source/Programs/UnrealBuildAccelerator/TestApp/Private/Windows/UbaTestAppWin.cpp b/Engine/Source/Programs/UnrealBuildAccelerator/TestApp/Private/Windows/UbaTestAppWin.cpp
--- a/Engine/Source/Programs/UnrealBuildAccelerator/TestApp/Private/Windows/UbaTestAppWin.cpp
+++ b/Engine/Source/Programs/UnrealBuildAccelerator/TestApp/Private/Windows/UbaTestAppWin.cpp
@@ -3,6 +3,7 @@
 #include <Windows.h>
 #include <wchar.h> 
 #include <stdarg.h>
+#include <time.h>
 
 int LogError(const wchar_t* format, ...)
 {
@@ -27,6 +28,9 @@ int wmain(int argc, wchar_t* argv[])
 		return LogError(L"Couldn't find UbaRunningRemote function in UbaDetours.dll");
 	bool runningRemote = (*runningRemoteFunc)();
 
+	using UbaRequestNextProcessFunc = bool(unsigned int prevExitCode, wchar_t* outArguments, unsigned int outArgumentsCapacity);
+	static UbaRequestNextProcessFunc* requestNextProcess = (UbaRequestNextProcessFunc*)(void*)GetProcAddress(detoursHandle, "UbaRequestNextProcess");
+
 	if (argc == 1)
 	{
 		HMODULE modules[] = { 0, detoursHandle, GetModuleHandleW(L"UbaTestApp.exe") };
@@ -133,13 +137,50 @@ int wmain(int argc, wchar_t* argv[])
 	}
 	else if (wcscmp(argv[1], L"-reuse") == 0)
 	{
-		using UbaRequestNextProcessFunc = bool(unsigned int prevExitCode, wchar_t* outArguments, unsigned int outArgumentsCapacity);
-		static UbaRequestNextProcessFunc* requestNextProcess = (UbaRequestNextProcessFunc*)(void*)GetProcAddress(detoursHandle, "UbaRequestNextProcess");
-
 		wchar_t arguments[1024];
 		if (requestNextProcess(0, arguments, sizeof(arguments)))
 			return LogError(L"Didn't expect another process");
 	}
+	else if (wcsncmp(argv[1], L"-file=", 6) == 0)
+	{
+		wchar_t arguments[1024];
+		const wchar_t* file = argv[1] + 6;
+		while (true)
+		{
+			HANDLE rh = CreateFileW(file, GENERIC_READ, FILE_SHARE_READ, NULL, OPEN_EXISTING, 0, NULL);
+			if (rh == INVALID_HANDLE_VALUE)
+				return LogError(L"Failed to open file %s", file);
+			char data[17] = {};
+			DWORD bytesRead;
+			if (!ReadFile(rh, data, 16, &bytesRead, NULL) || bytesRead != 16)
+				return LogError(L"Failed to read 16 bytes from file %s", file);
+			CloseHandle(rh);
+
+			srand(GetProcessId(GetCurrentProcess()));
+			Sleep(rand() % 2000);
+			wchar_t outFile[1024];
+			wcscpy_s(outFile, 1024, file);
+			outFile[wcslen(file)-3] = 0;
+			wcscat_s(outFile, 1024, L".out");
+
+			HANDLE wh = CreateFileW(outFile, GENERIC_WRITE, FILE_SHARE_WRITE, NULL, CREATE_ALWAYS, 0, NULL);
+			if (wh == INVALID_HANDLE_VALUE)
+				return LogError(L"Failed to create file File");
+			data[16] = 1;
+			DWORD bytesWritten;
+			if (!WriteFile(wh, data, 17, &bytesWritten, NULL) || bytesWritten != 17)
+				return LogError(L"Failed to read 16 bytes from file %s", file);
+
+			CloseHandle(wh);
+
+			// Request new process
+			if (!requestNextProcess(0, arguments, 1024))
+				break; // No process available, exit loop
+			file = arguments + 6;
+		}
+
+		return 0;
+	}
 	else
 	{
 		using u32 = unsigned int;
diff --git a/Engine/Source/Programs/UnrealBuildAccelerator/Visualizer/Private/UbaTraceReader.cpp b/Engine/Source/Programs/UnrealBuildAccelerator/Visualizer/Private/UbaTraceReader.cpp
--- a/Engine/Source/Programs/UnrealBuildAccelerator/Visualizer/Private/UbaTraceReader.cpp
+++ b/Engine/Source/Programs/UnrealBuildAccelerator/Visualizer/Private/UbaTraceReader.cpp
@@ -308,6 +308,11 @@ namespace uba
 						break;
 					oldSession.isReset = true;
 					oldSession.disconnectTime = ~u64(0);
+					oldSession.proxyName.clear();
+					oldSession.proxyCreated = false;
+					oldSession.notification.clear();
+					//oldSession.fetchedFiles.clear();
+					//oldSession.storedFiles.clear();
 					virtualSessionIndex = i;
 					break;
 				}
diff --git a/Engine/Source/Programs/UnrealBuildAccelerator/Visualizer/Private/UbaVisualizer.cpp b/Engine/Source/Programs/UnrealBuildAccelerator/Visualizer/Private/UbaVisualizer.cpp
--- a/Engine/Source/Programs/UnrealBuildAccelerator/Visualizer/Private/UbaVisualizer.cpp
+++ b/Engine/Source/Programs/UnrealBuildAccelerator/Visualizer/Private/UbaVisualizer.cpp
@@ -19,6 +19,7 @@ enum
 	Popup_Play,
 	Popup_JumpToEnd,
 	Popup_ShowText,
+	Popup_ShowCreateWriteColors,
 	Popup_SaveAs,
 	Popup_Quit,
 };
@@ -178,8 +179,10 @@ namespace uba
 			while (m_hwnd && m_client->IsConnected())
 				Sleep(1000);
 
+			m_client->Disconnect();
 			delete m_client;
 			m_client = nullptr;
+			Sleep(2000); // To prevent it from reconnecting to the same thing again and get thrown out (since it will post a WM_NEWTRACE and clean everything
 		}
 		return true;
 	}
@@ -477,6 +480,30 @@ namespace uba
 	constexpr int SessionStepY = RawBoxHeight + 2;
 	constexpr int FontHeight = 13;
 
+	struct SessionRec
+	{
+		TraceView::Session* session;
+		u32 index;
+	};
+	void Populate(SessionRec* recs, TraceView& traceView)
+	{
+		u32 count = u32(traceView.sessions.size());
+		for (u32 i = 0, e = count; i != e; ++i)
+			recs[i] = { &traceView.sessions[i], i };
+		if (count <= 1)
+			return;
+		std::sort(recs + 1, recs + traceView.sessions.size(), [](SessionRec& a, SessionRec& b)
+			{
+				auto& as = *a.session;
+				auto& bs = *b.session;
+				if ((as.processActiveCount != 0) != (bs.processActiveCount != 0))
+					return as.processActiveCount > bs.processActiveCount;
+				if (as.processActiveCount && as.proxyCreated != bs.proxyCreated)
+					return int(as.proxyCreated) > int(bs.proxyCreated);
+				return a.index < b.index;
+			});
+	}
+
 	void Visualizer::PaintAll(HDC hdc, const RECT& clientRect)
 	{
 		u64 currentTime = m_paused ? m_pauseStart : GetTime();
@@ -547,11 +574,8 @@ namespace uba
 
 		TraceView::WorkRecord selectedWork;
 
-		struct SessionRec { TraceView::Session* session; u32 index; };
-		SessionRec sortedSessions[256];
-		for (u32 i = 0, e = u32(m_traceView.sessions.size()); i != e; ++i)
-			sortedSessions[i] = { &m_traceView.sessions[i], i };
-		std::sort(sortedSessions, sortedSessions + m_traceView.sessions.size(), [](SessionRec& a, SessionRec& b) { return a.session->disconnectTime > b.session->disconnectTime; });
+		SessionRec sortedSessions[1024];
+		Populate(sortedSessions, m_traceView);
 
 		TraceView::ProcessLocation processLocation { 0, 0, 0 };
 		for (u64 i = 0, e = m_traceView.sessions.size(); i != e; ++i)
@@ -590,7 +614,11 @@ namespace uba
 						text.Append(L" - ").Append(session.notification);
 				}
 				else
+				{
 					text.Append(L" - Disconnected");
+					if (!session.notification.empty())
+						text.Append(L" (").Append(session.notification).Append(')');
+				}
 
 				bool selected = m_sessionSelectedIndex == processLocation.sessionIndex;
 
@@ -1500,7 +1528,9 @@ namespace uba
 		else if (process.exitCode != 0)
 			brush = m_processBrushes[selected].error;
 
-		if (!done || process.exitCode != 0 || (TimeToMs(process.processStats.sendFiles.time) < 300 && TimeToMs(process.processStats.createFile.time) < 300))
+		u64 writeFilesTime = Max(process.processStats.writeFiles.time, process.processStats.sendFiles.time);
+
+		if (!done || process.exitCode != 0 || !m_showCreateWriteColors || (TimeToMs(writeFilesTime) < 300 && TimeToMs(process.processStats.createFile.time) < 300))
 		{
 			if (writingBitmap)
 				rect.right = 256;
@@ -1526,7 +1556,7 @@ namespace uba
 				FillRect(hdc, &r2, m_processBrushes[selected].recv);
 		}
 
-		double sendPart = (double(ConvertTime(m_traceView, process.processStats.sendFiles.time)) / duration);
+		double sendPart = (double(ConvertTime(m_traceView, writeFilesTime)) / duration);
 		if (int tailSize = int(sendPart * width))
 		{
 			UBA_ASSERT(tailSize > 0);
@@ -1821,11 +1851,8 @@ namespace uba
 
 		TraceView::ProcessLocation& outLocation = outResult.processLocation;
 
-		struct SessionRec { TraceView::Session* session; u32 index; };
-		SessionRec sortedSessions[256];
-		for (u32 i = 0, e = u32(m_traceView.sessions.size()); i != e; ++i)
-			sortedSessions[i] = { &m_traceView.sessions[i], i };
-		std::sort(sortedSessions, sortedSessions + m_traceView.sessions.size(), [](SessionRec& a, SessionRec& b) { return a.session->disconnectTime > b.session->disconnectTime; });
+		SessionRec sortedSessions[1024];
+		Populate(sortedSessions, m_traceView);
 
 		for (u64 i = 0, e = m_traceView.sessions.size(); i != e; ++i)
 		{
@@ -2189,6 +2216,12 @@ namespace uba
 			StringBuffer<> title;
 			GetTitlePrefix(title);
 
+			auto g = MakeGuard([&]()
+				{
+					RedrawWindow(hWnd, NULL, NULL, RDW_INVALIDATE|RDW_UPDATENOW);
+					UpdateScrollbars(true);
+				});
+
 			if (m_client)
 			{
 				if (!m_trace.StartReadClient(m_traceView, *m_client))
@@ -2215,8 +2248,6 @@ namespace uba
 			}
 
 			SetWindowTextW(m_hwnd, title.data);
-			SendMessage(m_hwnd, WM_TIMER, 0, 0);
-			UpdateScrollbars(true);
 			SetTimer(m_hwnd, 0, 200, NULL);
 			return 0;
 		}
@@ -2333,11 +2364,12 @@ namespace uba
 			//	m_autoScroll = true;
 
 		
-			for (auto& session : m_traceView.sessions)
-				for (auto& processor : session.processors)
-					for (auto& process : processor.processes)
-						if (TimeToMs(process.processStats.sendFiles.time) >= 300 || TimeToMs(process.processStats.createFile.time) >= 300)
-							process.bitmapDirty = true;
+			if (m_showCreateWriteColors)
+				for (auto& session : m_traceView.sessions)
+					for (auto& processor : session.processors)
+						for (auto& process : processor.processes)
+							if (TimeToMs(Max(process.processStats.writeFiles.time, process.processStats.sendFiles.time)) >= 300 || TimeToMs(process.processStats.createFile.time) >= 300)
+								process.bitmapDirty = true;
 
 			UpdateScrollbars(true);
 			RedrawWindow(hWnd, NULL, NULL, RDW_INVALIDATE);
@@ -2489,7 +2521,8 @@ namespace uba
 			HMENU hMenu = CreatePopupMenu();
 			ClientToScreen(hWnd, &point);
 
-			AppendMenuW(hMenu, MF_STRING, Popup_ShowText, m_showText ? L"&Hide process text" : L"&Show process text");
+			AppendMenuW(hMenu, MF_STRING, Popup_ShowText, m_showText ? L"&Hide Process text" : L"&Show Process text");
+			AppendMenuW(hMenu, MF_STRING, Popup_ShowCreateWriteColors, m_showCreateWriteColors ? L"&Hide Create/Write colors" : L"&Show Create/Write colors");
 			AppendMenuW(hMenu, MF_SEPARATOR, 0, NULL);
 
 			if (m_sessionSelectedIndex != ~0u)
@@ -2552,6 +2585,15 @@ namespace uba
 				m_showText = !m_showText;
 				RedrawWindow(hWnd, NULL, NULL, RDW_INVALIDATE | RDW_UPDATENOW);
 				break;
+			case Popup_ShowCreateWriteColors:
+				m_showCreateWriteColors = !m_showCreateWriteColors;
+				for (auto& session : m_traceView.sessions)
+					for (auto& processor : session.processors)
+						for (auto& process : processor.processes)
+								process.bitmapDirty = true;
+				RedrawWindow(hWnd, NULL, NULL, RDW_INVALIDATE | RDW_UPDATENOW);
+				break;
+
 			case Popup_Replay:
 				m_replay = 1;
 				PostMessage(m_hwnd, WM_NEWTRACE, 0, 0);
diff --git a/Engine/Source/Programs/UnrealBuildAccelerator/Visualizer/Private/UbaVisualizer.h b/Engine/Source/Programs/UnrealBuildAccelerator/Visualizer/Private/UbaVisualizer.h
--- a/Engine/Source/Programs/UnrealBuildAccelerator/Visualizer/Private/UbaVisualizer.h
+++ b/Engine/Source/Programs/UnrealBuildAccelerator/Visualizer/Private/UbaVisualizer.h
@@ -112,6 +112,7 @@ namespace uba
 		bool m_useDarkMode = true;
 		bool m_isThemeSet = false;
 		bool m_showText = true;
+		bool m_showCreateWriteColors = true;
 
 		Logger& m_logger;
 		NetworkClient* m_client = nullptr;
