[UBA]
* Removed hard coded epic strings that should not be in code

diff --git a/Engine/Source/Programs/UnrealBuildAccelerator/Cli/Private/UbaCli.cpp b/Engine/Source/Programs/UnrealBuildAccelerator/Cli/Private/UbaCli.cpp
--- a/Engine/Source/Programs/UnrealBuildAccelerator/Cli/Private/UbaCli.cpp
+++ b/Engine/Source/Programs/UnrealBuildAccelerator/Cli/Private/UbaCli.cpp
@@ -143,9 +143,9 @@ namespace uba
 		StringBuffer<256> workDir;
 		StringBuffer<128> listenIp;
 		TString coordinatorName;
-		TString coordinatorUri = TC("https://horde.devtools.epicgames.com:443"); // TC("https://horde.devtools-dev.epicgames.com:443");
-		TString coordinatorPool = TC("UbaLinux-us-west-2"); // TC("BoxLinux");
-		TString coordinatorOidc = TC("EpicGames-Okta");
+		TString coordinatorUri;
+		TString coordinatorPool;
+		TString coordinatorOidc;
 		u32 coordinatorMaxCoreCount = 400;
 		u16 port = DefaultPort;
 		u32 maxProcessCount = DefaultProcessorCount;
