[UBA]
* Lowered stack usage for some unit tests to see if that solves macos failures

diff --git a/Engine/Source/Programs/UnrealBuildAccelerator/Test/Private/UbaTestSession.h b/Engine/Source/Programs/UnrealBuildAccelerator/Test/Private/UbaTestSession.h
--- a/Engine/Source/Programs/UnrealBuildAccelerator/Test/Private/UbaTestSession.h
+++ b/Engine/Source/Programs/UnrealBuildAccelerator/Test/Private/UbaTestSession.h
@@ -16,7 +16,7 @@ namespace uba
 	{
 		LogWriter& logWriter = logger.m_writer;
 
-		StringBuffer<> rootDir;
+		StringBuffer<MaxPath> rootDir;
 		rootDir.Append(testRootDir).Append(TC("Uba"));
 		if (!DeleteAllFiles(logger, rootDir.data))
 			return false;
@@ -33,7 +33,7 @@ namespace uba
 		sessionServerInfo.rootDir = rootDir.data;
 		SessionServer session(sessionServerInfo);
 
-		StringBuffer<> workingDir;
+		StringBuffer<MaxPath> workingDir;
 		workingDir.Append(testRootDir).Append(TC("WorkingDir"));
 		if (!DeleteAllFiles(logger, workingDir.data))
 			return false;
@@ -55,18 +55,21 @@ namespace uba
 		NetworkServer server(ctorSuccess, { logWriter });
 		NetworkClient client(ctorSuccess, { logWriter });
 
-		StringBuffer<> rootDir;
+		StringBuffer<MaxPath> rootDir;
 		rootDir.Append(testRootDir).Append(TC("Uba"));
 		if (!DeleteAllFiles(logger, rootDir.data))
 			return false;
 
 		StorageServerCreateInfo storageServerInfo(server, rootDir.data, logWriter);
 		storageServerInfo.casCapacityBytes = 1024ull * 1024 * 1024;
-		StorageServer storageServer(storageServerInfo);
+		auto& storageServer = *new StorageServer(storageServerInfo);
+		auto ssg = MakeGuard([&]() { delete &storageServer; });
+
 		SessionServerCreateInfo sessionServerInfo(storageServer, server, logWriter);
 		sessionServerInfo.checkMemory = false;
 		sessionServerInfo.rootDir = rootDir.data;
-		SessionServer sessionServer(sessionServerInfo);
+		auto& sessionServer = *new SessionServer(sessionServerInfo);
+		auto ssg2 = MakeGuard([&]() { delete &sessionServer; });
 
 		auto sg = MakeGuard([&]() { server.DisconnectClients(); });
 
@@ -75,15 +78,17 @@ namespace uba
 			return false;
 
 		StorageClientCreateInfo storageClientInfo(client, rootDir.data);
-		StorageClient storageClient(storageClientInfo);
+		auto& storageClient = *new StorageClient(storageClientInfo);
+		auto scg = MakeGuard([&]() { delete &storageClient; });
 
 		SessionClientCreateInfo sessionClientInfo(storageClient, client, logWriter);
 		sessionClientInfo.rootDir = rootDir.data;
-		SessionClient sessionClient(sessionClientInfo);
+		auto& sessionClient = *new SessionClient(sessionClientInfo);
+		auto scg2 = MakeGuard([&]() { delete &sessionClient; });
 
 		auto cg = MakeGuard([&]() { sessionClient.Stop(); client.Disconnect(); });
 
-		StringBuffer<> workingDir;
+		StringBuffer<MaxPath> workingDir;
 		workingDir.Append(testRootDir).Append(TC("WorkingDir"));
 		if (!DeleteAllFiles(logger, workingDir.data))
 			return false;
@@ -114,23 +119,23 @@ namespace uba
 
 	bool RunTestApp(LoggerWithWriter& logger, SessionServer& session, const tchar* workingDir, const RunProcessFunction& runProcess)
 	{
-		StringBuffer<> testApp;
+		StringBuffer<MaxPath> testApp;
 		GetTestAppPath(logger, testApp);
 
-		StringBuffer<> fileR;
 		{
+			StringBuffer<MaxPath> fileR;
 			fileR.Append(workingDir).Append(TC("FileR.h"));
 			FileAccessor fr(logger, fileR.data);
 			fr.CreateWrite();
 			fr.Write("Foo", 4);
 			fr.Close();
 		}
-
-		StringBuffer<> dir1;
-		dir1.Append(workingDir).Append(TC("Dir1"));
-		if (!CreateDirectoryW(dir1.data))
-			return logger.Error(TC("Failed to create dir %s"), dir1.data);
-
+		{
+			StringBuffer<MaxPath> dir1;
+			dir1.Append(workingDir).Append(TC("Dir1"));
+			if (!CreateDirectoryW(dir1.data))
+				return logger.Error(TC("Failed to create dir %s"), dir1.data);
+		}
 
 		ProcessStartInfo processInfo;
 		processInfo.application = testApp.data;
@@ -143,16 +148,18 @@ namespace uba
 		if (exitCode != 0)
 			return logger.Error(TC("UbaTestApp returned exit code %u"), exitCode);
 
-		StringBuffer<> fileW2;
-		fileW2.Append(workingDir).Append(TC("FileW2"));
-		if (!FileExists(logger, fileW2.data))
-			return logger.Error(TC("Can't find file %s"), fileW2.data);
-
-		StringBuffer<> fileWF;
-		fileWF.Append(workingDir).Append(TC("FileWF"));
-		if (!FileExists(logger, fileWF.data))
-			return logger.Error(TC("Can't find file %s"), fileWF.data);
-
+		{
+			StringBuffer<MaxPath> fileW2;
+			fileW2.Append(workingDir).Append(TC("FileW2"));
+			if (!FileExists(logger, fileW2.data))
+				return logger.Error(TC("Can't find file %s"), fileW2.data);
+		}
+		{
+			StringBuffer<MaxPath> fileWF;
+			fileWF.Append(workingDir).Append(TC("FileWF"));
+			if (!FileExists(logger, fileWF.data))
+				return logger.Error(TC("Can't find file %s"), fileWF.data);
+		}
 		return true;
 	}
 
@@ -178,7 +185,7 @@ namespace uba
 
 	bool RunClang(LoggerWithWriter& logger, SessionServer& session, const tchar* workingDir, const RunProcessFunction& runProcess)
 	{
-		StringBuffer<> sourceFile;
+		StringBuffer<MaxPath> sourceFile;
 		sourceFile.Append(workingDir).Append(TC("Code.cpp"));
 		FileAccessor codeFile(logger, sourceFile.data);
 		if (!codeFile.CreateWrite())
@@ -190,7 +197,7 @@ namespace uba
 			return false;
 
 #if PLATFORM_MAC
-		StringBuffer<512> xcodePath;
+		StringBuffer<MaxPath> xcodePath;
 		ExecuteCommand(logger, "/usr/bin/xcrun -f clang++", xcodePath);
 		const tchar* clangPath = xcodePath.data;
 #else
@@ -203,8 +210,8 @@ namespace uba
 		ProcessStartInfo processInfo;
 		processInfo.application = clangPath;
 #if PLATFORM_MAC
-		StringBuffer<512> xcodeSDKPath;
-		StringBuffer<> args;
+		StringBuffer<MaxPath> xcodeSDKPath;
+		StringBuffer<MaxPath> args;
 		ExecuteCommand(logger, "xcrun --show-sdk-path", xcodeSDKPath);
 		args.Append("-isysroot ");
 		args.Append(xcodeSDKPath.data);
