[UBA]
* Fixed so ref count of remote process is kept up in order to prevent it from being deleted at wrong point in time

diff --git a/Engine/Source/Programs/UnrealBuildAccelerator/Common/Private/UbaSessionServer.cpp b/Engine/Source/Programs/UnrealBuildAccelerator/Common/Private/UbaSessionServer.cpp
--- a/Engine/Source/Programs/UnrealBuildAccelerator/Common/Private/UbaSessionServer.cpp
+++ b/Engine/Source/Programs/UnrealBuildAccelerator/Common/Private/UbaSessionServer.cpp
@@ -289,7 +289,9 @@ namespace uba
 			}
 			remoteProcess->m_knownInputsCount = keysIndex;
 		}
-		
+
+		ProcessHandle h(remoteProcess); // Keep ref count up even if process is removed by callbacks etc.
+
 		ScopedCriticalSection lock(m_remoteProcessAndSessionLock);
 		m_queuedRemoteProcesses.push_back(remoteProcess);
 
@@ -299,17 +301,15 @@ namespace uba
 			if (!m_remoteExecutionEnabled)
 			{
 				m_logger.Info(TC("Process queued for remote but remote execution was disabled, returning process to queue"));
-				ProcessHandle h(remoteProcess);
 				m_remoteProcessReturnedEvent(*remoteProcess);
 			}
 			else if (!m_connectionCount)
 			{
 				m_logger.Info(TC("Process queued for remote but there are no active connections, returning process to queue"));
-				ProcessHandle h(remoteProcess);
 				m_remoteProcessReturnedEvent(*remoteProcess);
 			}
 		}
-		return remoteProcess;
+		return h;
 	}
 
 	void SessionServer::DisableRemoteExecution()
