[UBA]
* Added support for running individual unit tests from command line

diff --git a/Engine/Source/Programs/UnrealBuildAccelerator/Test/Private/Linux/UbaTestLinux.cpp b/Engine/Source/Programs/UnrealBuildAccelerator/Test/Private/Linux/UbaTestLinux.cpp
--- a/Engine/Source/Programs/UnrealBuildAccelerator/Test/Private/Linux/UbaTestLinux.cpp
+++ b/Engine/Source/Programs/UnrealBuildAccelerator/Test/Private/Linux/UbaTestLinux.cpp
@@ -7,7 +7,7 @@ void segfault_sigaction(int signal, siginfo_t* si, void* arg)
 	uba::UbaAssert("Segmentation fault", "", 0, "", -1);
 }
 
-int main()
+int main(int argc, char* argv[])
 {
 	#if !__has_feature(thread_sanitizer)
 	struct sigaction sa;
@@ -19,7 +19,7 @@ int main()
 	#endif
 
 	using namespace uba;
-	if (!RunAllTests())
+	if (!RunTests(argc, argv))
 		return -1;
 	return 0;
 }
diff --git a/Engine/Source/Programs/UnrealBuildAccelerator/Test/Private/Mac/UbaTestMac.cpp b/Engine/Source/Programs/UnrealBuildAccelerator/Test/Private/Mac/UbaTestMac.cpp
--- a/Engine/Source/Programs/UnrealBuildAccelerator/Test/Private/Mac/UbaTestMac.cpp
+++ b/Engine/Source/Programs/UnrealBuildAccelerator/Test/Private/Mac/UbaTestMac.cpp
@@ -6,7 +6,7 @@ int main(int argc, char* argv[])
 {
 	using namespace uba;
 
-	if (!RunAllTests())
+	if (!RunTests(argc, argv))
 		return -1;
 	return 0;
 }
diff --git a/Engine/Source/Programs/UnrealBuildAccelerator/Test/Private/UbaTestAll.h b/Engine/Source/Programs/UnrealBuildAccelerator/Test/Private/UbaTestAll.h
--- a/Engine/Source/Programs/UnrealBuildAccelerator/Test/Private/UbaTestAll.h
+++ b/Engine/Source/Programs/UnrealBuildAccelerator/Test/Private/UbaTestAll.h
@@ -46,13 +46,15 @@ namespace uba
 
 
 	#define UBA_TEST(x) \
-		logger.Info(TC("Running %s..."), TC(#x)); \
-		if (!x(testLogger, testRootDir)) \
-			return logger.Error(TC("  %s failed"), TC(#x)); \
-		logger.Info(TC("  %s success!"),  TC(#x));
-
-
-	bool RunAllTests()
+		if (!filter || Contains(TC(#x), filter)) \
+		{ \
+			logger.Info(TC("Running %s..."), TC(#x)); \
+			if (!x(testLogger, testRootDir)) \
+				return logger.Error(TC("  %s failed"), TC(#x)); \
+			logger.Info(TC("  %s success!"),  TC(#x)); \
+		}
+
+	bool RunTests(int argc, tchar* argv[])
 	{
 		LoggerWithWriter logger(g_consoleLogWriter, TC(""));
 
@@ -74,6 +76,10 @@ namespace uba
 
 		logger.Info(TC("Running tests (Test rootdir: %s)"), testRootDir.data);
 
+		const tchar* filter = nullptr;
+		if (argc > 1)
+			filter = argv[1];
+
 		//UBA_TEST(TestStress) // This can not be submitted.. it depends on CoordinatorHorde and credentials
 		UBA_TESTS
 
diff --git a/Engine/Source/Programs/UnrealBuildAccelerator/Test/Private/UbaTestNetwork.h b/Engine/Source/Programs/UnrealBuildAccelerator/Test/Private/UbaTestNetwork.h
--- a/Engine/Source/Programs/UnrealBuildAccelerator/Test/Private/UbaTestNetwork.h
+++ b/Engine/Source/Programs/UnrealBuildAccelerator/Test/Private/UbaTestNetwork.h
@@ -15,7 +15,7 @@ namespace uba
 	bool TestSockets(LoggerWithWriter& logger, const StringBufferBase& testRootDir)
 	{
 		static constexpr u16 port = 1346;
-		u8 result = 0;
+		Atomic<u8> result;
 		NetworkBackendTcp tcp(logger.m_writer);
 		Thread t([&]()
 			{
@@ -27,7 +27,7 @@ namespace uba
 						tcp.SetRecvCallbacks(connection, &result, 1,
 							[](void* context, const Guid& connectionUid, u8* headerData, void*& outBodyContext, u8*& outBodyData, u32& outBodySize)
 							{
-								*(u8*)context = *headerData;
+								*(Atomic<u8>*)context = *headerData;
 								//wprintf(TC("Listen got data from peer: %u\n"), *headerData);
 								return true;
 							},
diff --git a/Engine/Source/Programs/UnrealBuildAccelerator/Test/Private/Windows/UbaTestWin.cpp b/Engine/Source/Programs/UnrealBuildAccelerator/Test/Private/Windows/UbaTestWin.cpp
--- a/Engine/Source/Programs/UnrealBuildAccelerator/Test/Private/Windows/UbaTestWin.cpp
+++ b/Engine/Source/Programs/UnrealBuildAccelerator/Test/Private/Windows/UbaTestWin.cpp
@@ -6,7 +6,7 @@ int wmain(int argc, wchar_t* argv[])
 {
 	using namespace uba;
 
-	if (!RunAllTests())
+	if (!RunTests(argc, argv))
 		return -1;
 	return 0;
 }
