[UBA]
* Changed so macos is using ftruncate instead of lseek+write.  This fixes failed unit test on some macs
* Removed locks in UbaNetworkBackendMemory to fix potential deadlock
* Added StringBuffer::GetFileName since it is a common usecase when doing debug printing
* Added missing close of file descriptor in error situation

diff --git a/Engine/Source/Programs/UnrealBuildAccelerator/Common/Private/UbaFile.cpp b/Engine/Source/Programs/UnrealBuildAccelerator/Common/Private/UbaFile.cpp
--- a/Engine/Source/Programs/UnrealBuildAccelerator/Common/Private/UbaFile.cpp
+++ b/Engine/Source/Programs/UnrealBuildAccelerator/Common/Private/UbaFile.cpp
@@ -425,6 +425,7 @@ namespace uba
 			}
 			if (S_ISREG(attr.st_mode))
 				return (FileHandle)fd;
+			close(fd);
 			SetLastError(ERROR_ACCESS_DENIED);
 			return InvalidFileHandle;
 		}
diff --git a/Engine/Source/Programs/UnrealBuildAccelerator/Common/Private/UbaFileMapping.cpp b/Engine/Source/Programs/UnrealBuildAccelerator/Common/Private/UbaFileMapping.cpp
--- a/Engine/Source/Programs/UnrealBuildAccelerator/Common/Private/UbaFileMapping.cpp
+++ b/Engine/Source/Programs/UnrealBuildAccelerator/Common/Private/UbaFileMapping.cpp
@@ -175,7 +175,14 @@ namespace uba
 		int fd = asFileDescriptor(file);
 		if (maxSize)
 		{
-			if (lseek(fd, maxSize - 1, SEEK_SET) == -1)
+#if PLATFORM_MAC // For some reason some macs fails to use lseek+write
+			if (ftruncate(fd, maxSize) == -1)
+			{
+				UBA_ASSERTF(false, "ftruncate to %llu failed for %s: %s\n", maxSize, hint, strerror(errno));
+				return h;
+			}
+#else
+			if (lseek(fd, maxSize - 1, SEEK_SET) != maxSize - 1)
 			{
 				UBA_ASSERTF(false, "lseek to %llu failed for %s: %s\n", maxSize - 1, hint, strerror(errno));
 				return h;
@@ -187,9 +194,10 @@ namespace uba
 				res = write(fd, "", 1);
 			if (res != 1)
 			{
-				//UBA_ASSERTF(false, "write one byte at %llu on fd %i (%s) failed (res: %i): %s\n", maxSize - 1, fd, hint, res, strerror(errno));
-				//return h;
+				UBA_ASSERTF(false, "write one byte at %llu on fd %i (%s) failed (res: %i): %s\n", maxSize - 1, fd, hint, res, strerror(errno));
+				return h;
 			}
+#endif
 		}
 
 		h.shmFd = fd;
diff --git a/Engine/Source/Programs/UnrealBuildAccelerator/Common/Private/UbaNetworkBackendMemory.cpp b/Engine/Source/Programs/UnrealBuildAccelerator/Common/Private/UbaNetworkBackendMemory.cpp
--- a/Engine/Source/Programs/UnrealBuildAccelerator/Common/Private/UbaNetworkBackendMemory.cpp
+++ b/Engine/Source/Programs/UnrealBuildAccelerator/Common/Private/UbaNetworkBackendMemory.cpp
@@ -13,7 +13,6 @@ namespace uba
 	{
 		struct Peer
 		{
-			ReaderWriterLock lock;
 			RecvHeaderCallback* headerCallback = nullptr;
 			RecvBodyCallback* bodyCallback = nullptr;
 			void* context = nullptr;
@@ -26,7 +25,7 @@ namespace uba
 
 		Peer peer[2];
 		Guid uid;
-		bool connected = true;
+		Atomic<bool> connected;
 	};
 
 	NetworkBackendMemory::NetworkBackendMemory(LogWriter& writer, const tchar* prefix)
@@ -35,17 +34,16 @@ namespace uba
 
 	NetworkBackendMemory::~NetworkBackendMemory()
 	{
+		delete m_connection;
 	}
 
 	void NetworkBackendMemory::Shutdown(void* connection)
 	{
 		u64 from = ((uintptr_t)connection)-1;
 		u64 to = from == 0 ? 1 : 0;
-
 		auto& rc = m_connection->peer[to];
-		SCOPED_WRITE_LOCK(rc.lock, l);
+
 		m_connection->connected = false;
-		l.Leave();
 
 		if (auto cb = rc.disconnectCallback)
 			cb(rc.disconnectContext, m_connection->uid, connection);
@@ -62,8 +60,6 @@ namespace uba
 
 		auto& peer = m_connection->peer[to];
 
-		SCOPED_READ_LOCK(peer.lock, l);
-
 		if (!m_connection->connected)
 			return false;
 
@@ -139,6 +135,7 @@ namespace uba
 		}
 		UBA_ASSERT(!m_connection);
 		m_connection = new Connection();
+		m_connection->connected = true;
 		if (!m_connectedFunc((void*)1, {}))
 			return false;
 		l.Leave();
diff --git a/Engine/Source/Programs/UnrealBuildAccelerator/Core/Private/UbaStringBuffer.cpp b/Engine/Source/Programs/UnrealBuildAccelerator/Core/Private/UbaStringBuffer.cpp
--- a/Engine/Source/Programs/UnrealBuildAccelerator/Core/Private/UbaStringBuffer.cpp
+++ b/Engine/Source/Programs/UnrealBuildAccelerator/Core/Private/UbaStringBuffer.cpp
@@ -271,6 +271,13 @@ namespace uba
 		return TStrrchr(data + offset, c);
 	}
 
+	const tchar* StringBufferBase::GetFileName() const
+	{
+		if (const tchar* lps = TStrrchr(data, PathSeparator))
+			return lps + 1;
+		return data;
+	}
+
 	StringBufferBase& StringBufferBase::EnsureEndsWithSlash()
 	{
 		UBA_ASSERT(count);
diff --git a/Engine/Source/Programs/UnrealBuildAccelerator/Core/Public/UbaStringBuffer.h b/Engine/Source/Programs/UnrealBuildAccelerator/Core/Public/UbaStringBuffer.h
--- a/Engine/Source/Programs/UnrealBuildAccelerator/Core/Public/UbaStringBuffer.h
+++ b/Engine/Source/Programs/UnrealBuildAccelerator/Core/Public/UbaStringBuffer.h
@@ -50,6 +50,7 @@ namespace uba
 		bool Equals(const tchar* str, bool ignoreCase = true) const { return uba::Equals(data, str, ignoreCase); }
 		const tchar* First(tchar c, u64 offset = 0) const;
 		const tchar* Last(tchar c, u64 offset = 0) const;
+		const tchar* GetFileName() const;
 		inline StringBufferBase& Replace(tchar from, tchar to) { uba::Replace(data, from, to);  return *this; }
 
 		StringBufferBase& EnsureEndsWithSlash();
