[Uba]
* Added missing check for if AppendDir actually contains a dir

diff --git a/Engine/Source/Programs/UnrealBuildAccelerator/Core/Private/UbaStringBuffer.cpp b/Engine/Source/Programs/UnrealBuildAccelerator/Core/Private/UbaStringBuffer.cpp
--- a/Engine/Source/Programs/UnrealBuildAccelerator/Core/Private/UbaStringBuffer.cpp
+++ b/Engine/Source/Programs/UnrealBuildAccelerator/Core/Private/UbaStringBuffer.cpp
@@ -155,8 +155,9 @@ namespace uba
 
 	StringBufferBase& StringBufferBase::AppendDir(const StringBufferBase& str)
 	{
-		const tchar* last = str.Last(PathSeparator);
-		return Append(str.data, u64(last - str.data));
+		if (const tchar* last = str.Last(PathSeparator))
+			return Append(str.data, u64(last - str.data));
+		return *this;
 	}
 
 	StringBufferBase& StringBufferBase::AppendDir(const tchar* dir)
