[Uba]
* Fixed stack overflow issue on mac silicon machines

diff --git a/Engine/Source/Programs/UnrealBuildAccelerator/Common/Private/UbaSessionClient.cpp b/Engine/Source/Programs/UnrealBuildAccelerator/Common/Private/UbaSessionClient.cpp
--- a/Engine/Source/Programs/UnrealBuildAccelerator/Common/Private/UbaSessionClient.cpp
+++ b/Engine/Source/Programs/UnrealBuildAccelerator/Common/Private/UbaSessionClient.cpp
@@ -234,9 +234,60 @@ namespace uba
 		return EnsureApplicationEnvironment(outRealApplication, 0, startInfo.application);
 	}
 
+	struct SessionClient::ModuleInfo
+	{
+		ModuleInfo(const tchar* n, const CasKey& c, u32 a) : name(n), casKey(c), attributes(a), done(true) {}
+		TString name;
+		CasKey casKey;
+		u32 attributes;
+		Event done;
+	};
+
+	bool SessionClient::ReadModules(List<ModuleInfo>& outModules, u32 processId, const tchar* application)
+	{
+		StackBinaryReader<SendMaxSize> reader;
+		{
+			StackBinaryWriter<1024> writer;
+			NetworkMessage msg(m_client, ServiceId, SessionMessageType_GetApplication, writer);
+			writer.WriteU32(processId);
+			writer.WriteString(application);
+			if (!msg.Send(reader, m_stats.getApplicationMsg))
+				return false;
+		}
+
+		u32 serverSystemPathLen = reader.ReadU32();
+		u32 moduleCount = reader.ReadU32();
+		if (moduleCount == 0)
+			return m_logger.Error(TC("Application %s not found"), application);
+
+		while (moduleCount--)
+		{
+			StringBuffer<> moduleFile;
+			reader.ReadString(moduleFile);
+			u32 fileAttributes = reader.ReadU32();
+			bool isSystem = reader.ReadBool();
+
+			CasKey casKey = reader.ReadCasKey();
+			if (casKey == CasKeyZero)
+				return m_logger.Error(TC("Bad CasKey for %s (%s)"), moduleFile.data, CasKeyString(casKey).str);
+
+			if (isSystem)
+			{
+				StringBuffer<> localSystemModule;
+				localSystemModule.Append(m_systemPath).Append(moduleFile.data + serverSystemPathLen);
+				if (FileExists(m_logger, localSystemModule.data))
+					continue;
+				moduleFile.Clear().Append(localSystemModule);
+			}
+			outModules.emplace_back(moduleFile.data, casKey, fileAttributes);
+		}
+
+		return true;
+	}
+
 	bool SessionClient::EnsureApplicationEnvironment(StringBufferBase& out, u32 processId, const tchar* application)
 	{
-		StringBuffer<> applicationDir;
+		StringBuffer<MaxPath> applicationDir;
 		applicationDir.AppendDir(application);
 		KeyToString keyStr(ToStringKeyLower(applicationDir));
 
@@ -247,51 +298,16 @@ namespace uba
 		{
 			auto failGuard = MakeGuard([&]() { m_handledApplicationEnvironments.erase(insres.first); });
 
-			StackBinaryReader<SendMaxSize> reader;
-			{
-				StackBinaryWriter<1024> writer;
-				NetworkMessage msg(m_client, ServiceId, SessionMessageType_GetApplication, writer);
-				writer.WriteU32(processId);
-				writer.WriteString(application);
-				if (!msg.Send(reader, m_stats.getApplicationMsg))
-					return false;
-			}
-
-			u32 serverSystemPathLen = reader.ReadU32();
-			u32 moduleCount = reader.ReadU32();
-			if (moduleCount == 0)
-				return m_logger.Error(TC("Application %s not found"), application);
-
-			struct ModuleInfo { ModuleInfo(const tchar* n, const CasKey& c, u32 a) : name(n), casKey(c), attributes(a), done(true) {} TString name; CasKey casKey; u32 attributes; Event done; };
 			List<ModuleInfo> modules;
 
+			if (!ReadModules(modules, processId, application))
+				return false;
+
 			Atomic<bool> success = true;
 			Atomic<u32> handledCount;
-			u32 workCount = 0;
-			while (moduleCount--)
+			for (auto& m : modules)
 			{
-				StringBuffer<> moduleFile;
-				reader.ReadString(moduleFile);
-				u32 fileAttributes = reader.ReadU32();
-				bool isSystem = reader.ReadBool();
-
-				CasKey casKey = reader.ReadCasKey();
-				if (casKey == CasKeyZero)
-					return m_logger.Error(TC("Bad CasKey for %s (%s)"), moduleFile.data, CasKeyString(casKey).str);
-
-				if (isSystem)
-				{
-					StringBuffer<> localSystemModule;
-					localSystemModule.Append(m_systemPath).Append(moduleFile.data + serverSystemPathLen);
-					if (FileExists(m_logger, localSystemModule.data))
-						continue;
-					moduleFile.Clear().Append(localSystemModule);
-				}
-
-				++workCount;
-
-				auto& m = modules.emplace_back(moduleFile.data, casKey, fileAttributes);
-				m_client.AddWork([&]()
+				m_client.AddWork([&handledCount, &m, this, &success, &keyStr]()
 					{
 						++handledCount;
 						auto g = MakeGuard([&]() { m.done.Set(); });
@@ -307,13 +323,13 @@ namespace uba
 						}
 						if (const tchar* lastSeparator = TStrrchr(moduleName, PathSeparator))
 							moduleName = lastSeparator + 1;
-						StringBuffer<> temp;
+						StringBuffer<MaxPath> temp;
 						if (!WriteBinFile(temp, moduleName, newCasKey, keyStr, m.attributes))
 							success = false;
 					}, 1, TC("EnsureApp"));
 			}
 
-			while (handledCount < workCount)
+			while (handledCount < modules.size())
 				m_client.DoWork();
 
 			// Wait for all to be done
diff --git a/Engine/Source/Programs/UnrealBuildAccelerator/Common/Public/UbaSessionClient.h b/Engine/Source/Programs/UnrealBuildAccelerator/Common/Public/UbaSessionClient.h
--- a/Engine/Source/Programs/UnrealBuildAccelerator/Common/Public/UbaSessionClient.h
+++ b/Engine/Source/Programs/UnrealBuildAccelerator/Common/Public/UbaSessionClient.h
@@ -67,8 +67,10 @@ namespace uba
 		virtual void TraceSessionUpdate() override;
 
 		struct InternalProcessStartInfo;
+		struct ModuleInfo;
 
 		bool GetCasKeyForFile(CasKey& out, u32 processId, const StringBufferBase& fileName, const StringKey& fileNameKey);
+		bool ReadModules(List<ModuleInfo>& outModules, u32 processId, const tchar* application);
 		bool EnsureApplicationEnvironment(StringBufferBase& out, u32 processId, const tchar* application);
 		bool EnsureBinaryFile(StringBufferBase& out, StringBufferBase& outVirtual, u32 processId, const StringBufferBase& fileName, const StringKey& fileNameKey, const tchar* applicationDir);
 		bool WriteBinFile(StringBufferBase& out, const tchar* binaryName, const CasKey& casKey, const KeyToString& applicationDir, u32 fileAttributes);
