[Uba]
* Fixed bug where NtQueryDirectoryFile did not handle buffer overflow properly

diff --git a/Engine/Source/Programs/UnrealBuildAccelerator/Detours/Private/Windows/UbaDetoursFunctionsNtDll.inl b/Engine/Source/Programs/UnrealBuildAccelerator/Detours/Private/Windows/UbaDetoursFunctionsNtDll.inl
--- a/Engine/Source/Programs/UnrealBuildAccelerator/Detours/Private/Windows/UbaDetoursFunctionsNtDll.inl
+++ b/Engine/Source/Programs/UnrealBuildAccelerator/Detours/Private/Windows/UbaDetoursFunctionsNtDll.inl
@@ -160,7 +160,7 @@ NTSTATUS NTAPI Detoured_NtQueryDirectoryFile(HANDLE FileHandle, HANDLE Event, PI
 
 		u8* prevInformation = nullptr;
 		u8* it = (u8*)FileInformation;
-		u32 left = Length;
+		u8* bufferEnd = it + Length;
 
 		while (true)
 		{
@@ -178,7 +178,7 @@ NTSTATUS NTAPI Detoured_NtQueryDirectoryFile(HANDLE FileHandle, HANDLE Event, PI
 			if (FileName && wcsncmp(FileName->Buffer, fileName, FileName->Length / 2) != 0)
 				continue;
 
-			u32 fileNameBytes = u32(wcslen(fileName) * 2) + 2;
+			u32 fileNameBytes = u32(wcslen(fileName) * 2);
 
 			wchar_t* fileNamePos = nullptr;
 			u32 structSize = 0;
@@ -198,10 +198,10 @@ NTSTATUS NTAPI Detoured_NtQueryDirectoryFile(HANDLE FileHandle, HANDLE Event, PI
 				return STATUS_OBJECT_NAME_NOT_FOUND;
 			}
 
-			if (left < structSize + fileNameBytes)
+			u8* writeEnd = (u8*)fileNamePos + fileNameBytes;
+			if (writeEnd > bufferEnd)
 			{
-				if (!prevInformation)
-					res = STATUS_BUFFER_OVERFLOW;
+				res = STATUS_BUFFER_OVERFLOW;
 				break;
 			}
 
@@ -210,7 +210,7 @@ NTSTATUS NTAPI Detoured_NtQueryDirectoryFile(HANDLE FileHandle, HANDLE Event, PI
 
 			memcpy(fileNamePos, fileName, fileNameBytes);
 
-			info.FileNameLength = fileNameBytes - 2;
+			info.FileNameLength = fileNameBytes;
 			info.FileAttributes = entryInfo.attributes;
 			info.LastWriteTime.QuadPart = entryInfo.lastWrite;
 			info.EndOfFile.QuadPart = entryInfo.size;
