[UBA]
* Proper fix for setting size of file mapping.. bug was that it tried to set file mapping size on files that were readonly and error was confusing

diff --git a/Engine/Source/Programs/UnrealBuildAccelerator/Common/Private/UbaFileAccessor.cpp b/Engine/Source/Programs/UnrealBuildAccelerator/Common/Private/UbaFileAccessor.cpp
--- a/Engine/Source/Programs/UnrealBuildAccelerator/Common/Private/UbaFileAccessor.cpp
+++ b/Engine/Source/Programs/UnrealBuildAccelerator/Common/Private/UbaFileAccessor.cpp
@@ -110,7 +110,9 @@ namespace uba
 
 	bool FileAccessor::CreateMemoryWrite(bool allowRead, u32 flagsAndAttributes, u64 size, const tchar* tempPath)
 	{
+		#if PLATFORM_WINDOWS
 		allowRead = true; // It is not possible to have write only access to file mappings it seems
+		#endif
 
 		m_size = size;
 
diff --git a/Engine/Source/Programs/UnrealBuildAccelerator/Common/Private/UbaFileMapping.cpp b/Engine/Source/Programs/UnrealBuildAccelerator/Common/Private/UbaFileMapping.cpp
--- a/Engine/Source/Programs/UnrealBuildAccelerator/Common/Private/UbaFileMapping.cpp
+++ b/Engine/Source/Programs/UnrealBuildAccelerator/Common/Private/UbaFileMapping.cpp
@@ -165,20 +165,20 @@ namespace uba
 #endif
 	}
 
-	FileMappingHandle CreateFileMappingW(FileHandle file, u32 flProtect, u64 maxSize, const tchar* hint)
+	FileMappingHandle CreateFileMappingW(FileHandle file, u32 protect, u64 maxSize, const tchar* hint)
 	{
 		ExtendedTimerScope ts(SystemStats::GetCurrent().createFileMapping);
 #if PLATFORM_WINDOWS
-		return { InternalCreateFileMappingW(asHANDLE(file), flProtect, (DWORD)ToHigh(maxSize), ToLow(maxSize), NULL) };
+		return { InternalCreateFileMappingW(asHANDLE(file), protect, (DWORD)ToHigh(maxSize), ToLow(maxSize), NULL) };
 #else
 		FileMappingHandle h;
 		int fd = asFileDescriptor(file);
-		if (maxSize)
+		if (maxSize && (protect & (~PAGE_READONLY)) != 0)
 		{
-#if PLATFORM_MAC // For some reason some macs fails to use lseek+write
+#if 0
 			if (ftruncate(fd, maxSize) == -1)
 			{
-				UBA_ASSERTF(false, "ftruncate to %llu failed for %s: %s\n", maxSize, hint, strerror(errno));
+				UBA_ASSERTF(false, "ftruncate to %llu on fd %i failed for %s: %s\n", maxSize, fd, hint, strerror(errno));
 				return h;
 			}
 #else
@@ -190,8 +190,6 @@ namespace uba
 
 			errno = 0;
 			int res = write(fd, "", 1);
-			if (res == 0)
-				res = write(fd, "", 1);
 			if (res != 1)
 			{
 				UBA_ASSERTF(false, "write one byte at %llu on fd %i (%s) failed (res: %i): %s\n", maxSize - 1, fd, hint, res, strerror(errno));
