[UBA]
* Re-added use for ftruncate on mac because it still didn't work after the fix for linux

diff --git a/Engine/Source/Programs/UnrealBuildAccelerator/Common/Private/UbaFileMapping.cpp b/Engine/Source/Programs/UnrealBuildAccelerator/Common/Private/UbaFileMapping.cpp
--- a/Engine/Source/Programs/UnrealBuildAccelerator/Common/Private/UbaFileMapping.cpp
+++ b/Engine/Source/Programs/UnrealBuildAccelerator/Common/Private/UbaFileMapping.cpp
@@ -175,7 +175,7 @@ namespace uba
 		int fd = asFileDescriptor(file);
 		if (maxSize && (protect & (~PAGE_READONLY)) != 0)
 		{
-#if 0
+#if PLATFORM_MAC // For some reason lseek+write does not work on apple silicon platform
 			if (ftruncate(fd, maxSize) == -1)
 			{
 				UBA_ASSERTF(false, "ftruncate to %llu on fd %i failed for %s: %s\n", maxSize, fd, hint, strerror(errno));
