[UBA]
* Added description to error message when process exits before getting exit message
* Changed NtQueryDirectoryFile usage in UbaDirectoryIterator to use ReturnSingleEntry=false to reduce number of kernel calls. It now uses a 64k buffer to receive as many entries as possible at the time
* Fixed bug in detoured NtQueryDirectoryFile after learning how it works by implementing ReturnSingleEntry=false

diff --git a/Engine/Source/Programs/UnrealBuildAccelerator/Common/Private/UbaProcess.cpp b/Engine/Source/Programs/UnrealBuildAccelerator/Common/Private/UbaProcess.cpp
--- a/Engine/Source/Programs/UnrealBuildAccelerator/Common/Private/UbaProcess.cpp
+++ b/Engine/Source/Programs/UnrealBuildAccelerator/Common/Private/UbaProcess.cpp
@@ -292,7 +292,7 @@ namespace uba
 		if (!m_gotExitMessage)
 		{
 			StringBuffer<> err;
-			err.Appendf(TC("ERROR: Process %u not active but did not get exit message. Received %u messages"), m_nativeProcessId, m_messageCount);
+			err.Appendf(TC("ERROR: Process %u (%s) not active but did not get exit message. Received %u messages"), m_nativeProcessId, m_description.c_str(), m_messageCount);
 			LogLine(false, err.data, LogEntryType_Error);
 			m_nativeProcessExitCode = UBA_EXIT_CODE(666);
 		}
diff --git a/Engine/Source/Programs/UnrealBuildAccelerator/Common/Public/UbaDirectoryIterator.h b/Engine/Source/Programs/UnrealBuildAccelerator/Common/Public/UbaDirectoryIterator.h
--- a/Engine/Source/Programs/UnrealBuildAccelerator/Common/Public/UbaDirectoryIterator.h
+++ b/Engine/Source/Programs/UnrealBuildAccelerator/Common/Public/UbaDirectoryIterator.h
@@ -121,30 +121,51 @@ namespace uba
 			infoFunc(di);
 		}
 
-		u64 buff[(sizeof(FILE_ID_BOTH_DIR_INFORMATION) + 512) / 8];
-		auto& FileInformation = (FILE_ID_BOTH_DIR_INFORMATION&)buff;
+		u8 buff[64*1024];
 		while (true)
 		{
-			res = NtQueryDirectoryFile(handle, 0, NULL, NULL, &IoStatusBlock, &FileInformation, sizeof(buff), (FILE_INFORMATION_CLASS)FileIdBothDirectoryInformation, TRUE, NULL, FALSE);
+			res = NtQueryDirectoryFile(handle, 0, NULL, NULL, &IoStatusBlock, buff, sizeof(buff) - 2, (FILE_INFORMATION_CLASS)FileIdBothDirectoryInformation, FALSE, NULL, FALSE);
 			if (res != STATUS_SUCCESS)
 			{
 				if (res != STATUS_NO_MORE_FILES)
-					logger.Warning(L"NtQueryDirectoryFile returned error %i", res);
+					logger.Error(L"NtQueryDirectoryFile returned error 0x%x", res);
 				break;
 			}
-			FileInformation.FileName[FileInformation.FileNameLength / 2] = 0;
-			if (wcscmp(FileInformation.FileName, L".") == 0 || wcscmp(FileInformation.FileName, L"..") == 0)
-				continue;
 
-			DirectoryEntry entry;
-			entry.name = FileInformation.FileName;
-			entry.nameLen = u32(FileInformation.FileNameLength/sizeof(tchar));
-			entry.lastWritten = FileInformation.LastWriteTime.QuadPart;
-			entry.attributes = FileInformation.FileAttributes;
-			entry.volumeSerial = volumeSerial;
-			entry.id = FileInformation.FileId.QuadPart;
-			entry.size = FileInformation.EndOfFile.QuadPart;
-			iteratorFunc(entry);
+			u8* it = buff;
+			while (true)
+			{
+				auto& FileInformation = *(FILE_ID_BOTH_DIR_INFORMATION*)it;
+
+				auto fileName = FileInformation.FileName;
+				u64 fileNameCharCount = FileInformation.FileNameLength/sizeof(tchar);
+				wchar_t& fileNameEnd = fileName[fileNameCharCount];
+				wchar_t old = fileNameEnd;
+				fileNameEnd = 0;
+				if (fileNameCharCount == 1 && fileName[0] == '.' || (fileNameCharCount == 2 && fileName[0] == '.' && fileName[1] == '.'))
+				{
+					if (!FileInformation.NextEntryOffset)
+						break;
+					fileNameEnd = old;
+					it += FileInformation.NextEntryOffset;
+					continue;
+				}
+
+				DirectoryEntry entry;
+				entry.name = fileName;
+				entry.nameLen = u32(fileNameCharCount);
+				entry.lastWritten = FileInformation.LastWriteTime.QuadPart;
+				entry.attributes = FileInformation.FileAttributes;
+				entry.volumeSerial = volumeSerial;
+				entry.id = FileInformation.FileId.QuadPart;
+				entry.size = FileInformation.EndOfFile.QuadPart;
+				iteratorFunc(entry);
+
+				if (!FileInformation.NextEntryOffset)
+					break;
+				fileNameEnd = old;
+				it += FileInformation.NextEntryOffset;
+			}
 		}
 		return true;
 	}
diff --git a/Engine/Source/Programs/UnrealBuildAccelerator/Detours/Private/Windows/UbaDetoursFunctionsNtDll.inl b/Engine/Source/Programs/UnrealBuildAccelerator/Detours/Private/Windows/UbaDetoursFunctionsNtDll.inl
--- a/Engine/Source/Programs/UnrealBuildAccelerator/Detours/Private/Windows/UbaDetoursFunctionsNtDll.inl
+++ b/Engine/Source/Programs/UnrealBuildAccelerator/Detours/Private/Windows/UbaDetoursFunctionsNtDll.inl
@@ -201,7 +201,8 @@ NTSTATUS NTAPI Detoured_NtQueryDirectoryFile(HANDLE FileHandle, HANDLE Event, PI
 			u8* writeEnd = (u8*)fileNamePos + fileNameBytes;
 			if (writeEnd > bufferEnd)
 			{
-				res = STATUS_BUFFER_OVERFLOW;
+				if (!prevInformation)
+					res = STATUS_BUFFER_OVERFLOW;
 				break;
 			}
 
