[UBA]
* Some minor fixes in uba visualizer

diff --git a/Engine/Source/Programs/UnrealBuildAccelerator/Visualizer/Private/UbaVisualizer.cpp b/Engine/Source/Programs/UnrealBuildAccelerator/Visualizer/Private/UbaVisualizer.cpp
--- a/Engine/Source/Programs/UnrealBuildAccelerator/Visualizer/Private/UbaVisualizer.cpp
+++ b/Engine/Source/Programs/UnrealBuildAccelerator/Visualizer/Private/UbaVisualizer.cpp
@@ -1432,7 +1432,7 @@ namespace uba
 				if (hours)
 				{
 					b.AppendValue(hours).Append('h');
-					if (hours < 10)
+					if (minutes < 10)
 						b.Append('0');
 				}
 				if (minutes || hours)
@@ -1530,7 +1530,7 @@ namespace uba
 
 		u64 writeFilesTime = Max(process.processStats.writeFiles.time, process.processStats.sendFiles.time);
 
-		if (!done || process.exitCode != 0 || !m_showCreateWriteColors || (TimeToMs(writeFilesTime) < 300 && TimeToMs(process.processStats.createFile.time) < 300))
+		if (!done || process.exitCode != 0 || !m_showCreateWriteColors || (TimeToMs(writeFilesTime, m_traceView.frequency) < 300 && TimeToMs(process.processStats.createFile.time, m_traceView.frequency) < 300))
 		{
 			if (writingBitmap)
 				rect.right = 256;
@@ -2368,7 +2368,7 @@ namespace uba
 				for (auto& session : m_traceView.sessions)
 					for (auto& processor : session.processors)
 						for (auto& process : processor.processes)
-							if (TimeToMs(Max(process.processStats.writeFiles.time, process.processStats.sendFiles.time)) >= 300 || TimeToMs(process.processStats.createFile.time) >= 300)
+							if (TimeToMs(Max(process.processStats.writeFiles.time, process.processStats.sendFiles.time), m_traceView.frequency) >= 300 || TimeToMs(process.processStats.createFile.time, m_traceView.frequency) >= 300)
 								process.bitmapDirty = true;
 
 			UpdateScrollbars(true);
