[Uba]
* Added lots of more logging around processes disappearing without exit message
* Added descriptions to lots of asserts since macos refuse to give us good line numbers
* Added more error handling

diff --git a/Engine/Source/Programs/UnrealBuildAccelerator/Common/Private/UbaFileMapping.cpp b/Engine/Source/Programs/UnrealBuildAccelerator/Common/Private/UbaFileMapping.cpp
--- a/Engine/Source/Programs/UnrealBuildAccelerator/Common/Private/UbaFileMapping.cpp
+++ b/Engine/Source/Programs/UnrealBuildAccelerator/Common/Private/UbaFileMapping.cpp
@@ -606,10 +606,7 @@ namespace uba
 	{
 		m_mappingHandle = uba::CreateMemoryMappingW(m_logger, PAGE_READWRITE|SEC_RESERVE, capacity);
 		if (!m_mappingHandle.IsValid())
-		{
-			m_logger.Error(TC("%s - Failed to create memory map (%s)"), m_name, LastErrorToText().data);
-			return false;
-		}
+			return m_logger.Error(TC("%s - Failed to create memory map (%s)"), m_name, LastErrorToText().data);
 
 		m_blockSize = blockSize;
 		m_pageSize = 64*1024;
diff --git a/Engine/Source/Programs/UnrealBuildAccelerator/Common/Private/UbaProcess.cpp b/Engine/Source/Programs/UnrealBuildAccelerator/Common/Private/UbaProcess.cpp
--- a/Engine/Source/Programs/UnrealBuildAccelerator/Common/Private/UbaProcess.cpp
+++ b/Engine/Source/Programs/UnrealBuildAccelerator/Common/Private/UbaProcess.cpp
@@ -266,10 +266,10 @@ namespace uba
 			return false;
 		#endif
 
+		siginfo_t signalInfo;
 		while (true)
 		{
-			siginfo_t signalInfo;
-			signalInfo.si_pid = 0;	// if remains 0, treat as child was not waitable (i.e. was running)
+			memset(&signalInfo, 0, sizeof(signalInfo));
 			int res = waitid(P_PID, (unsigned int)m_nativeProcessId, &signalInfo, WEXITED | WNOHANG | WNOWAIT);
 			if (res)
 			{
@@ -284,7 +284,10 @@ namespace uba
 			else
 			{
 				if (signalInfo.si_pid != (pid_t)m_nativeProcessId)
+				{
+					++m_waitidCount;
 					return true;
+				}
 				break;
 			}
 		}
@@ -292,7 +295,7 @@ namespace uba
 		if (!m_gotExitMessage)
 		{
 			StringBuffer<> err;
-			err.Appendf(TC("ERROR: Process %u (%s) not active but did not get exit message. Received %u messages"), m_nativeProcessId, m_description.c_str(), m_messageCount);
+			err.Appendf(TC("ERROR: Process %u (%s) not active but did not get exit message. Received %u messages. Called waitid %u times. Signal code: %i. Exit value or signal: %i. Execution time: %s."), m_nativeProcessId, m_description.c_str(), m_messageCount, m_waitidCount, signalInfo.si_code, signalInfo.si_status, TimeToText(GetTime() - m_startTime).str);
 			LogLine(false, err.data, LogEntryType_Error);
 			m_nativeProcessExitCode = UBA_EXIT_CODE(666);
 		}
@@ -370,13 +373,12 @@ namespace uba
 		SessionStatsScope sessionStatsScope(m_sessionStats);
 
 		u8* comMemory = m_comMemory.memory;
+		u64 comMemorySize = CommunicationMemSize;
 		#if !PLATFORM_WINDOWS
 		comMemory += sizeof(Event) * 3;
+		comMemorySize -= sizeof(Event) * 3;
 		#endif
 
-		u8* writeMemory = comMemory;
-		u8* readMemory = comMemory;// +CommunicationMemSize / 2;
-
 		u32 retryCount = 0; // Do not allow retry
 
 		u32 exitCode = ~0u;
@@ -394,8 +396,8 @@ namespace uba
 				while (loop && WaitForRead(outReader, errReader))
 				{
 					u64 startTime = GetTime();
-					BinaryReader reader(readMemory);
-					BinaryWriter writer(writeMemory);
+					BinaryReader reader(comMemory, 0, comMemorySize);
+					BinaryWriter writer(comMemory, 0, comMemorySize);
 					loop = HandleMessage(reader, writer);
 					SetWritten();
 					m_processStats.hostTotalTime += GetTime() - startTime;
@@ -701,6 +703,7 @@ namespace uba
 					ProcessStartInfo info;
 					info.application = application.data;
 					info.arguments = commandLine;
+					info.description = application.GetFileName();
 					info.workingDir = currentDir.data;
 					info.logFile = InternalGetChildLogFile(temp);
 					info.priorityClass = m_startInfo.priorityClass;
@@ -741,7 +744,7 @@ namespace uba
 			case MessageType_StartProcess:
 				{
 					u32 processId = reader.ReadU32();
-					UBA_ASSERT(processId > 0);
+					UBA_ASSERT(processId > 0 && processId <= m_childProcesses.size());
 					auto& process = *(ProcessImpl*)m_childProcesses[processId - 1].m_process;
 					bool result = reader.ReadBool();
 					u32 lastError = reader.ReadU32();
@@ -1617,7 +1620,7 @@ namespace uba
 				if (TimeToMs(GetTime() - startTime) > 120 * 1000) // 
 				{
 					startTime = GetTime();
-					logger.Error(TC("Waiting for parent process in createprocess has now taken more than 120 seconds."));
+					logger.Error(TC("Waiting for parent process (%s) has now taken more than 120 seconds. (%s)"), m_parentProcess->m_startInfo.description, m_startInfo.description);
 				}
 			}
 
@@ -1649,7 +1652,7 @@ namespace uba
 				if (TimeToMs(GetTime() - startTime) > 120 * 1000) // 
 				{
 					startTime = GetTime();
-					logger.Error(TC("Waiting for parent process in exitprocess has now taken more than 120 seconds."));
+					logger.Error(TC("Waiting for parent process (%s) while exiting has now taken more than 120 seconds."), m_parentProcess->m_startInfo.description);
 				}
 			}
 		}
diff --git a/Engine/Source/Programs/UnrealBuildAccelerator/Common/Private/UbaSession.cpp b/Engine/Source/Programs/UnrealBuildAccelerator/Common/Private/UbaSession.cpp
--- a/Engine/Source/Programs/UnrealBuildAccelerator/Common/Private/UbaSession.cpp
+++ b/Engine/Source/Programs/UnrealBuildAccelerator/Common/Private/UbaSession.cpp
@@ -1009,7 +1009,12 @@ namespace uba
 		m_storage.CreateDirectory(m_sessionLogDir.data);
 		m_sessionLogDir.EnsureEndsWithSlash();
 
+		if (info.traceOutputFile)
+			m_traceOutputFile.Append(info.traceOutputFile);
+	}
 
+	bool Session::Create(const SessionCreateInfo& info)
+	{
 		#if UBA_DEBUG_LOGGER
 		StartDebugLogger(m_logger, StringBuffer<512>().Append(m_sessionDir).Append(TC("SessionDebug.log")).data);
 		#endif
@@ -1034,15 +1039,14 @@ namespace uba
 		m_directoryTable.m_lookup.reserve(30000);
 		m_fileMappingTableLookup.reserve(70000);
 
-		if (info.traceOutputFile)
-			m_traceOutputFile.Append(info.traceOutputFile);
-
 		m_fileMappingBuffer.AddTransient(TC("FileMappings"));
 
 		u64 reserveSize = CommunicationMemSize * 512;
-		m_processCommunicationAllocator.Init(CommunicationMemSize, reserveSize);
+		if (!m_processCommunicationAllocator.Init(CommunicationMemSize, reserveSize))
+			return false;
 
-		CreateProcessJobObject();
+		if (!CreateProcessJobObject())
+			return false;
 
 		// Environment variables that should stay local when building remote (not replicated)
 		#if PLATFORM_WINDOWS
@@ -1054,10 +1058,7 @@ namespace uba
 		m_localEnvironmentVariables.insert(TC("PROCESSOR_LEVEL"));
 		m_localEnvironmentVariables.insert(TC("PROCESSOR_REVISION"));
 		#endif
-	}
 
-	bool Session::Create(const SessionCreateInfo& info)
-	{
 		StringBuffer<> traceName;
 		if (info.traceName && *info.traceName)
 			traceName.Append(info.traceName);
@@ -2152,19 +2153,17 @@ namespace uba
 		logger.Info(TC(""));
 	}
 
-	void Session::CreateProcessJobObject()
+	bool Session::CreateProcessJobObject()
 	{
 		#if PLATFORM_WINDOWS
 		m_processJobObject = CreateJobObject(nullptr, nullptr);
 		if (!m_processJobObject)
-		{
-			m_logger.Error(TC("Failed to create process job object"));
-			return;
-		}
+			return m_logger.Error(TC("Failed to create process job object"));
 		JOBOBJECT_EXTENDED_LIMIT_INFORMATION info = { };
 		info.BasicLimitInformation.LimitFlags = JOB_OBJECT_LIMIT_BREAKAWAY_OK | JOB_OBJECT_LIMIT_KILL_ON_JOB_CLOSE;
 		SetInformationJobObject(m_processJobObject, JobObjectExtendedLimitInformation, &info, sizeof(info));
 		#endif
+		return true;
 	}
 
 
diff --git a/Engine/Source/Programs/UnrealBuildAccelerator/Common/Private/UbaSessionServer.cpp b/Engine/Source/Programs/UnrealBuildAccelerator/Common/Private/UbaSessionServer.cpp
--- a/Engine/Source/Programs/UnrealBuildAccelerator/Common/Private/UbaSessionServer.cpp
+++ b/Engine/Source/Programs/UnrealBuildAccelerator/Common/Private/UbaSessionServer.cpp
@@ -171,7 +171,11 @@ namespace uba
 		m_detoursLibrary = detoursFile.data;
 		#endif
 
-		Create(info);
+		if (!Create(info))
+		{
+			UBA_ASSERT(false);
+			return;
+		}
 	}
 
 	SessionServer::~SessionServer()
diff --git a/Engine/Source/Programs/UnrealBuildAccelerator/Common/Public/UbaProcess.h b/Engine/Source/Programs/UnrealBuildAccelerator/Common/Public/UbaProcess.h
--- a/Engine/Source/Programs/UnrealBuildAccelerator/Common/Public/UbaProcess.h
+++ b/Engine/Source/Programs/UnrealBuildAccelerator/Common/Public/UbaProcess.h
@@ -104,7 +104,7 @@ namespace uba
 
 		int m_stdOutPipe = -1;
 		int m_stdErrPipe = -1;
-
+		u32 m_waitidCount = 0;
 	#endif
 
 		ProcHandle m_nativeProcessHandle = InvalidProcHandle;
diff --git a/Engine/Source/Programs/UnrealBuildAccelerator/Common/Public/UbaSession.h b/Engine/Source/Programs/UnrealBuildAccelerator/Common/Public/UbaSession.h
--- a/Engine/Source/Programs/UnrealBuildAccelerator/Common/Public/UbaSession.h
+++ b/Engine/Source/Programs/UnrealBuildAccelerator/Common/Public/UbaSession.h
@@ -135,7 +135,7 @@ namespace uba
 		bool WriteDirectoryEntriesInternal(DirectoryTable::Directory& dir, const StringKey& dirKey, const tchar* dirPath, bool isRefresh, u32& outTableOffset);
 		void WriteDirectoryEntriesRecursive(const StringKey& dirKey, tchar* dirPath, u32& outTableOffset);
 		bool CopyImports(Vector<BinaryModule>& out, const tchar* library, tchar* applicationDir, tchar* applicationDirEnd, UnorderedSet<TString>& handledImports);
-		void CreateProcessJobObject();
+		bool CreateProcessJobObject();
 		void GetSystemInfo(StringBufferBase& out);
 		bool GetMemoryInfo(u64& outAvailable, u64& outTotal);
 		void WriteSummary(BinaryWriter& writer, const Function<void(Logger& logger)>& summaryFunc);
diff --git a/Engine/Source/Programs/UnrealBuildAccelerator/Core/Private/UbaEvent.cpp b/Engine/Source/Programs/UnrealBuildAccelerator/Core/Private/UbaEvent.cpp
--- a/Engine/Source/Programs/UnrealBuildAccelerator/Core/Private/UbaEvent.cpp
+++ b/Engine/Source/Programs/UnrealBuildAccelerator/Core/Private/UbaEvent.cpp
@@ -2,6 +2,7 @@
 
 #include "UbaEvent.h"
 #include "UbaPlatform.h"
+#include "UbaStringBuffer.h"
 #include "UbaSynchronization.h"
 
 #if !PLATFORM_WINDOWS
@@ -15,11 +16,6 @@ namespace uba
 	struct EventImpl
 	{
 		EventImpl() = default;
-
-		EventImpl(bool manualReset) { Create(manualReset); }
-
-		EventImpl(const EventImpl& o) { UBA_ASSERT(false); }
-
 		~EventImpl() { Destroy(); }
 
 		bool Create(bool manualReset, bool shared = false)
@@ -28,13 +24,13 @@ namespace uba
 			static_assert(Atomic<TriggerType>::is_always_lock_free, "atomic<T> not lock free, can't work in shared mem");
 			static_assert(Atomic<bool>::is_always_lock_free, "atomic<T> not lock free, can't work in shared mem");
 
-			UBA_ASSERT(!m_initialized);
+			UBA_ASSERTF(!m_initialized, "Can't create already created Event");
 			m_manualReset = manualReset;
 
 			pthread_mutexattr_t attrmutex;
 			if (pthread_mutexattr_init(&attrmutex) != 0)
 			{
-				UBA_ASSERT(false);
+				UBA_ASSERTF(false, "pthread_mutexattr_init failed");
 				return false;
 			}
 
@@ -42,13 +38,13 @@ namespace uba
 			{
 				if (pthread_mutexattr_setpshared(&attrmutex, PTHREAD_PROCESS_SHARED) != 0)
 				{
-					UBA_ASSERT(false);
+					UBA_ASSERTF(false, "pthread_mutexattr_setpshared failed");
 					return false;
 				}
 				#if PLATFORM_LINUX
 				if (pthread_mutexattr_setrobust(&attrmutex, PTHREAD_MUTEX_ROBUST) != 0)
 				{
-					UBA_ASSERT(false);
+					UBA_ASSERTF(false, "pthread_mutexattr_setrobust failed");
 					return false;
 				}
 				#endif
@@ -56,7 +52,7 @@ namespace uba
 
 			if (pthread_mutex_init(&m_mutex, &attrmutex) != 0)
 			{
-				UBA_ASSERT(false);
+				UBA_ASSERTF(false, "pthread_mutex_init failed");
 				return false;
 			}
 			pthread_mutexattr_destroy(&attrmutex);
@@ -64,7 +60,7 @@ namespace uba
 			pthread_condattr_t attrcond;
 			if (pthread_condattr_init(&attrcond) != 0)
 			{
-				UBA_ASSERT(false);
+				UBA_ASSERTF(false, "pthread_condattr_init failed");
 				return false;
 			}
 
@@ -72,14 +68,14 @@ namespace uba
 			{
 				if (pthread_condattr_setpshared(&attrcond, PTHREAD_PROCESS_SHARED) != 0)
 				{
-					UBA_ASSERT(false);
+					UBA_ASSERTF(false, "pthread_condattr_setpshared failed");
 					return false;
 				}
 			}
 
 			if (pthread_cond_init(&m_condition, &attrcond) != 0)
 			{
-				UBA_ASSERT(false);
+				UBA_ASSERTF(false, "pthread_cond_init failed");
 				pthread_mutex_destroy(&m_mutex);
 				return false;
 			}
@@ -122,13 +118,13 @@ namespace uba
 			{
 				m_triggered = TriggerType_All;
 				if (pthread_cond_broadcast(&m_condition) != 0)
-					UBA_ASSERT(false);
+					UBA_ASSERTF(false, "pthread_cond_broadcast failed");
 			}
 			else
 			{
 				m_triggered = TriggerType_One;
 				if (pthread_cond_signal(&m_condition) != 0)  // may release multiple threads anyhow!
-					UBA_ASSERT(false);
+					UBA_ASSERTF(false, "pthread_cond_signal failed");
 			}
 
 			UnlockEventMutex();
@@ -179,7 +175,7 @@ namespace uba
 					if (timeoutMs == ~0u) // infinite wait?
 					{
 						if (pthread_cond_wait(&m_condition, &m_mutex) != 0)  // unlocks Mutex while blocking...
-							UBA_ASSERT(false);
+							UBA_ASSERTF(false, "pthread_cond_wait failed");
 					}
 					else  // timed wait.
 					{
@@ -188,7 +184,7 @@ namespace uba
 						TimeOut.tv_sec = (startTime.tv_sec) + (ms / 1000);
 						TimeOut.tv_nsec = (ms % 1000) * 1000000;  // remainder of milliseconds converted to nanoseconds.
 						int rc = pthread_cond_timedwait(&m_condition, &m_mutex, &TimeOut);    // unlocks Mutex while blocking...
-						UBA_ASSERT((rc == 0) || (rc == ETIMEDOUT)); (void)rc;
+						UBA_ASSERTF((rc == 0) || (rc == ETIMEDOUT), "pthread_cond_timedwait failed"); (void)rc;
 
 						// Update timeoutMs and startTime in case we have to go again...
 						struct timeval now, difference;
@@ -199,7 +195,7 @@ namespace uba
 						startTime = now;
 					}
 					--m_waitingThreads;
-					UBA_ASSERT(m_waitingThreads >= 0);
+					UBA_ASSERTF(m_waitingThreads >= 0, "m_waitingThreads less than 0");
 				}
 
 			} while (timeoutMs != 0);
@@ -208,8 +204,6 @@ namespace uba
 			return false;
 		}
 
-		void* GetHandle() { UBA_ASSERT(false); return 0; }
-
 	private:
 		enum TriggerType { TriggerType_None, TriggerType_One, TriggerType_All };
 
@@ -222,14 +216,14 @@ namespace uba
 
 		inline void LockEventMutex()
 		{
-			if (pthread_mutex_lock(&m_mutex) != 0)
-				UBA_ASSERT(false);
+			int res = pthread_mutex_lock(&m_mutex);(void)res;
+			UBA_ASSERTF(res == 0, "pthread_mutex_lock failed (error code %i)", res);
 		}
 
 		inline void UnlockEventMutex()
 		{
-			if (pthread_mutex_unlock(&m_mutex) != 0)
-				UBA_ASSERT(false);
+			int res = pthread_mutex_unlock(&m_mutex);
+			UBA_ASSERTF(res == 0, "pthread_mutex_unlock failed (error code %i)", res);
 		}
 
 		static inline void SubtractTimevals(const struct timeval* FromThis, struct timeval* SubThis, struct timeval* Difference)
@@ -333,7 +327,7 @@ namespace uba
 		#if PLATFORM_WINDOWS
 		return m_ev;
 		#else
-		UBA_ASSERT(false);
+		UBA_ASSERTF(false, "Event::GetHandle not available");
 		return 0;
 		#endif
 	}
diff --git a/Engine/Source/Programs/UnrealBuildAccelerator/Core/Public/UbaBinaryReaderWriter.h b/Engine/Source/Programs/UnrealBuildAccelerator/Core/Public/UbaBinaryReaderWriter.h
--- a/Engine/Source/Programs/UnrealBuildAccelerator/Core/Public/UbaBinaryReaderWriter.h
+++ b/Engine/Source/Programs/UnrealBuildAccelerator/Core/Public/UbaBinaryReaderWriter.h
@@ -155,6 +155,8 @@ namespace uba
 		#if PLATFORM_WINDOWS
 		Write7BitEncoded(strLen);
 
+		UBA_ASSERT_WRITE(GetWrittenBytes(str, strLen));
+
 		for (const tchar* i = str, *e = str + strLen; i != e; ++i)
 		{
 			int c = *i;
diff --git a/Engine/Source/Programs/UnrealBuildAccelerator/Detours/Private/PosixOS/UbaDetoursMainPosix.cpp b/Engine/Source/Programs/UnrealBuildAccelerator/Detours/Private/PosixOS/UbaDetoursMainPosix.cpp
--- a/Engine/Source/Programs/UnrealBuildAccelerator/Detours/Private/PosixOS/UbaDetoursMainPosix.cpp
+++ b/Engine/Source/Programs/UnrealBuildAccelerator/Detours/Private/PosixOS/UbaDetoursMainPosix.cpp
@@ -18,8 +18,7 @@ namespace uba
 	Event* g_cancelEvent;
 	Event* g_readEvent;
 	Event* g_writeEvent;
-	u8* g_readMessageMappingMem;
-	u8* g_writeMessageMappingMem;
+	u8* g_messageMappingMem;
 
 	void PreInit(const char* logFile);
 	void Init();
@@ -94,8 +93,7 @@ static void __attribute__((constructor(102))) PreInitCtor()
 		printf("UbaDetours: Failed to open shared mem: %s\n", comIdName.data);
 		return;
 	}
-	u64 bytesToMap = 64 * 1024;
-	u8* rptr = (u8*)mmap(NULL, bytesToMap, PROT_READ | PROT_WRITE, MAP_SHARED, g_comFd, s64(comIdOffset));
+	u8* rptr = (u8*)mmap(NULL, CommunicationMemSize, PROT_READ | PROT_WRITE, MAP_SHARED, g_comFd, s64(comIdOffset));
 	if (rptr == MAP_FAILED)
 	{
 		printf("UbaDetours: Failed to mmap fd: %u\n", g_comFd);
@@ -105,11 +103,7 @@ static void __attribute__((constructor(102))) PreInitCtor()
 	g_cancelEvent = ((Event*)rptr);
 	g_readEvent = ((Event*)rptr) + 1;
 	g_writeEvent = ((Event*)rptr) + 2;
-
-	u8* comMem = rptr + sizeof(Event) * 3;
-
-	g_readMessageMappingMem = comMem;
-	g_writeMessageMappingMem = comMem;// +bytesToMap / 2;
+	g_messageMappingMem = rptr + sizeof(Event) * 3;
 
 	#if 0 // TODO: For some reason this does not work
 	struct sigaction sa;
@@ -140,12 +134,12 @@ void CloseCom()
 {
 	if (g_comFd == -1)
 		return;
-	u64 bytesToMap = 64 * 1024;
 	void* mem = g_cancelEvent;
 	g_cancelEvent = nullptr;
 	g_readEvent = nullptr;
 	g_writeEvent = nullptr;
-	munmap(mem, bytesToMap);
+	g_messageMappingMem = nullptr;
+	munmap(mem, CommunicationMemSize);
 	close(g_comFd);
 	g_comFd = -1;
 }
@@ -163,9 +157,9 @@ namespace uba
 {
 	BinaryWriter::BinaryWriter()
 	{
-		m_begin = g_writeMessageMappingMem;
+		m_begin = g_messageMappingMem;
 		m_pos = m_begin;
-		m_end = m_begin + CommunicationMemSize;// / 2;
+		m_end = m_begin + CommunicationMemSize - sizeof(Event) * 3;
 	}
 
 	void BinaryWriter::Flush(bool waitOnResponse)
@@ -192,8 +186,8 @@ namespace uba
 
 	BinaryReader::BinaryReader()
 	{
-		m_begin = g_readMessageMappingMem;
+		m_begin = g_messageMappingMem;
 		m_pos = m_begin;
-		m_end = m_begin + CommunicationMemSize;// / 2;
+		m_end = m_begin + CommunicationMemSize - sizeof(Event) * 3;
 	}
 }
