[UBA]
* Added _RDATA section to callback in detoured ImageGetDigestStream.. Needed to new dxil.dll to match hash on linux/wine

diff --git a/Engine/Source/Programs/UnrealBuildAccelerator/Detours/Private/Windows/UbaDetoursFunctionsImagehlp.inl b/Engine/Source/Programs/UnrealBuildAccelerator/Detours/Private/Windows/UbaDetoursFunctionsImagehlp.inl
--- a/Engine/Source/Programs/UnrealBuildAccelerator/Detours/Private/Windows/UbaDetoursFunctionsImagehlp.inl
+++ b/Engine/Source/Programs/UnrealBuildAccelerator/Detours/Private/Windows/UbaDetoursFunctionsImagehlp.inl
@@ -82,6 +82,7 @@ BOOL Detoured_ImageGetDigestStream(HANDLE FileHandle, DWORD DigestLevel, DIGEST_
 	ReportSection(".rdata");
 	ReportSection(".data");
 	ReportSection(".pdata");
+	ReportSection("_RDATA");
 	ReportSection(".debug");
 	ReportSection(".didat");
 	ReportSection(".rsrc");
