[UBA]
* Enabled file writing using memory maps (we've tested on two machines that behaves very different with file writing and mmap is the only path that makes write fast on both machines)
* Changed from lock to bottleneck for creating file mappings and allow four at the time.

diff --git a/Engine/Source/Programs/UnrealBuildAccelerator/Common/Private/UbaFileMapping.cpp b/Engine/Source/Programs/UnrealBuildAccelerator/Common/Private/UbaFileMapping.cpp
--- a/Engine/Source/Programs/UnrealBuildAccelerator/Common/Private/UbaFileMapping.cpp
+++ b/Engine/Source/Programs/UnrealBuildAccelerator/Common/Private/UbaFileMapping.cpp
@@ -1,6 +1,7 @@
 // Copyright Epic Games, Inc. All Rights Reserved.
 
 #include "UbaFileMapping.h"
+#include "UbaBottleneck.h"
 #include "UbaDirectoryIterator.h"
 #include "UbaFile.h"
 #include "UbaPlatform.h"
@@ -15,7 +16,7 @@
 namespace uba
 {
 #if PLATFORM_WINDOWS
-	ReaderWriterLock& g_createFileHandleLock = *new ReaderWriterLock(); // Allocated and leaked just to prevent shutdown asserts in debug
+	Bottleneck& g_createFileHandleBottleneck = *new Bottleneck(8); // Allocated and leaked just to prevent shutdown asserts in debug
 
 	//Atomic<u64> g_fileMappingCount;
 
@@ -29,7 +30,7 @@ namespace uba
 			return ::CreateFileMappingW(hFile, NULL, flProtect, dwMaximumSizeHigh, dwMaximumSizeLow, lpName);
 
 		// Experiment to try to prevent lock happening on AWS servers when lots of helpers are sending back obj files.
-		SCOPED_WRITE_LOCK(g_createFileHandleLock, lock);
+		BottleneckScope bs(g_createFileHandleBottleneck);
 		return ::CreateFileMappingW(hFile, NULL, flProtect, dwMaximumSizeHigh, dwMaximumSizeLow, lpName);
 	}
 #else
diff --git a/Engine/Source/Programs/UnrealBuildAccelerator/Common/Private/UbaSession.cpp b/Engine/Source/Programs/UnrealBuildAccelerator/Common/Private/UbaSession.cpp
--- a/Engine/Source/Programs/UnrealBuildAccelerator/Common/Private/UbaSession.cpp
+++ b/Engine/Source/Programs/UnrealBuildAccelerator/Common/Private/UbaSession.cpp
@@ -1958,25 +1958,28 @@ namespace uba
 			{
 				// This is to kill I/O when writing lots of pdb/dlls in parallel
 				#if PLATFORM_WINDOWS
-				constexpr u32 bottleneckMax = 8;
+				constexpr u32 bottleneckMax = 32;
 				bool shouldBottleneck = false;//useOverlap;
 				static Bottleneck bottleneck(bottleneckMax);
 				auto bng = MakeGuard([&]() { if (shouldBottleneck) bottleneck.Leave(); });
 				#endif
 
+				shouldBottleneck = true;
+				bottleneck.Enter();
+
 				u64 fileSize = file.mappingWritten;
 				u8* mem = MapViewOfFile(file.mappingHandle, FILE_MAP_READ, 0, fileSize);
 				if (!mem)
 					return m_logger.Error(TC("Failed to map view of filehandle for read %s (%s)"), file.name.c_str(), LastErrorToText().data);
 
-				PrefetchVirtualMemory(mem, fileSize);
+				//PrefetchVirtualMemory(mem, fileSize);
 
 				auto memClose = MakeGuard([&](){ UnmapViewOfFile(mem, fileSize, file.name.c_str()); });
 
 				// Seems like best combo (for windows at least) is to use writes with overlap and max 16 at the same time.
 				// On one machine we get twice as fast without overlap if no bottleneck. On another machine (ntfs compression on) we get twice as slow without overlap
 				// Both machines behaves well with overlap AND bottleneck. Both machine are 128 logical core thread rippers.
-				constexpr bool useFileMapForWrite = false;
+				constexpr bool useFileMapForWrite = true;
 				bool useOverlap = false;//fileSize > 8 * 1024 * 1024;
 
 
diff --git a/Engine/Source/Programs/UnrealBuildAccelerator/Common/Public/UbaProcessStartInfo.h b/Engine/Source/Programs/UnrealBuildAccelerator/Common/Public/UbaProcessStartInfo.h
--- a/Engine/Source/Programs/UnrealBuildAccelerator/Common/Public/UbaProcessStartInfo.h
+++ b/Engine/Source/Programs/UnrealBuildAccelerator/Common/Public/UbaProcessStartInfo.h
@@ -19,7 +19,7 @@ namespace uba
 		const tchar* workingDir = TC("");		// Working directory. Use full path
 		const tchar* description = TC("");		// Description. Used for on-screen logging and log file names if session.logToFile is set but logFile is ""
 		const tchar* logFile = TC("");			// Log file. If set, will always log. If not full path the session log dir will be prepended.
-		u32 priorityClass = 0; // Priority of process
+		u32 priorityClass = 0x00000020;			// Priority of process. Defaults to NORMAL_PRIORITY_CLASS
 		u64 outputStatsThresholdMs = ~u64(0);	// Threshold in milliseconds where process summary should be printed to log
 		bool trackInputs = false;				// Track all files read. Can read result in ProcessHandle.GetTrackedInputs()
 		bool useCustomAllocator = true;			// Disable detouring of allocator inside processes. If Session.disableCustomAllocator is false this will be overridden
