[UBA]
* Fixed so NetworkMessage can be created without initialized
* Removed unused code
* Fixed StorageClient::SendBatchMessages so it deletes correct number of messages if it disconnects in the middle of transfer

diff --git a/Engine/Source/Programs/UnrealBuildAccelerator/Common/Private/UbaNetworkClient.cpp b/Engine/Source/Programs/UnrealBuildAccelerator/Common/Private/UbaNetworkClient.cpp
--- a/Engine/Source/Programs/UnrealBuildAccelerator/Common/Private/UbaNetworkClient.cpp
+++ b/Engine/Source/Programs/UnrealBuildAccelerator/Common/Private/UbaNetworkClient.cpp
@@ -451,7 +451,7 @@ namespace uba
 		message.m_responseCapacity = responseCapacity;
 		message.m_connection = &connection;
 
-		BinaryWriter& writer = message.m_sendWriter;
+		BinaryWriter& writer = *message.m_sendWriter;
 
 		u16 messageId = 0;
 		Event gotResponse(true);
@@ -569,9 +569,20 @@ namespace uba
 	}
 
 	NetworkMessage::NetworkMessage(NetworkClient& client, u8 serviceId, u8 messageType, BinaryWriter& sendWriter)
-	:	m_client(client)
-	,	m_sendWriter(sendWriter)
 	{
+		Init(client, serviceId, messageType, sendWriter);
+	}
+
+	NetworkMessage::~NetworkMessage()
+	{
+		UBA_ASSERT(!m_id);
+	}
+
+	void NetworkMessage::Init(NetworkClient& client, u8 serviceId, u8 messageType, BinaryWriter& sendWriter)
+	{
+		m_client = &client;
+		m_sendWriter = &sendWriter;
+
 		// Header (SendHeaderSize):
 		// 1 byte    - 2 bits for serviceid, 6 bits for messagetype
 		// 2 byte    - message id
@@ -583,19 +594,14 @@ namespace uba
 		data[0] = u8(serviceId << 6) | messageType;
 	}
 
-	NetworkMessage::~NetworkMessage()
-	{
-		UBA_ASSERT(!m_id);
-	}
-
 	bool NetworkMessage::Send()
 	{
-		return m_client.Send(*this, nullptr, 0, false);
+		return m_client->Send(*this, nullptr, 0, false);
 	}
 
 	bool NetworkMessage::Send(BinaryReader& response)
 	{
-		if (!m_client.Send(*this, (u8*)response.GetPositionData(), u32(response.GetLeft()), false))
+		if (!m_client->Send(*this, (u8*)response.GetPositionData(), u32(response.GetLeft()), false))
 			return false;
 		response.SetSize(response.GetPosition() + m_responseSize);
 		return true;
@@ -613,7 +619,7 @@ namespace uba
 		UBA_ASSERT(!m_doneFunc);
 		m_doneFunc = func;
 		m_doneUserData = userData;
-		return m_client.Send(*this, (u8*)response.GetPositionData(), u32(response.GetLeft()), true);
+		return m_client->Send(*this, (u8*)response.GetPositionData(), u32(response.GetLeft()), true);
 	}
 
 	bool NetworkMessage::ProcessAsyncResults(BinaryReader& response)
@@ -621,11 +627,11 @@ namespace uba
 		if (m_error)
 			return false;
 
-		if (m_client.m_cryptoKey)
+		if (m_client->m_cryptoKey)
 		{
 			UBA_ASSERT(!response.GetPosition());
-			TimerScope ts(m_client.m_decryptTimer);
-			if (!Crypto::Decrypt(m_client.m_logger, m_client.m_cryptoKey, (u8*)m_response, m_responseSize))
+			TimerScope ts(m_client->m_decryptTimer);
+			if (!Crypto::Decrypt(m_client->m_logger, m_client->m_cryptoKey, (u8*)m_response, m_responseSize))
 				return false;
 		}
 		response.SetSize(response.GetPosition() + m_responseSize);
@@ -639,8 +645,8 @@ namespace uba
 			{
 				if (m_id)
 				{
-					m_client.m_availableMessageIds.push_back(m_id);
-					m_client.m_activeMessages[m_id] = nullptr;
+					m_client->m_availableMessageIds.push_back(m_id);
+					m_client->m_activeMessages[m_id] = nullptr;
 					m_id = 0;
 					hasId = true;
 				}
@@ -648,7 +654,7 @@ namespace uba
 
 		if (shouldLock)
 		{
-			SCOPED_WRITE_LOCK(m_client.m_activeMessagesLock, lock);
+			SCOPED_WRITE_LOCK(m_client->m_activeMessagesLock, lock);
 			returnId();
 		}
 		else
diff --git a/Engine/Source/Programs/UnrealBuildAccelerator/Common/Private/UbaStorageClient.cpp b/Engine/Source/Programs/UnrealBuildAccelerator/Common/Private/UbaStorageClient.cpp
--- a/Engine/Source/Programs/UnrealBuildAccelerator/Common/Private/UbaStorageClient.cpp
+++ b/Engine/Source/Programs/UnrealBuildAccelerator/Common/Private/UbaStorageClient.cpp
@@ -808,16 +808,7 @@ namespace uba
 
 		struct Entry
 		{
-			Entry(NetworkClient& client, u16 fetchId, u8* slot, u32 readIndex, u32 i, u32 messageMaxSize)
-				: message(client, ServiceId, StorageMessageType_FetchSegment, writer)
-				, reader(slot + i * messageMaxSize, 0, SendMaxSize)
-				, done(true)
-			{
-				writer.WriteU16(fetchId);
-				writer.WriteU32(readIndex + i + 1);
-			}
-
-			StackBinaryWriter<32> writer;
+			Entry(u8* slot, u32 i, u32 messageMaxSize) : reader(slot + i * messageMaxSize, 0, SendMaxSize), done(true) {}
 			NetworkMessage message;
 			BinaryReader reader;
 			Event done;
@@ -837,22 +828,24 @@ namespace uba
 
 		Entry* entries = (Entry*)entriesMem;
 
-		for (u32 i=0; i!=sendCount; ++i)
-			new (entries + i) Entry(client, fetchId, slot, readIndex, i, messageMaxSize);
-
 		bool success = true;
-		u32 inFlight = u32(sendCount);
+		u32 inFlightCount = u32(sendCount);
 		for (u32 i=0; i!=sendCount; ++i)
 		{
-			auto& entry = entries[i];
+			auto& entry = *new (entries + i) Entry(slot, i, messageMaxSize);
+			StackBinaryWriter<32> writer;
+			entry.message.Init(client, ServiceId, StorageMessageType_FetchSegment, writer);
+			writer.WriteU16(fetchId);
+			writer.WriteU32(readIndex + i + 1);
 			if (entry.message.SendAsync(entry.reader, [](bool error, void* userData) { ((Event*)userData)->Set(); }, &entry.done))
 				continue;
-			inFlight = i;
+			entry.~Entry();
+			inFlightCount = i;
 			success = false;
 			break;
 		}
 
-		for (u32 i=0; i!=inFlight; ++i)
+		for (u32 i=0; i!=inFlightCount; ++i)
 		{
 			Entry& entry = entries[i];
 			if (!entry.done.IsSet(5*60*1000))
@@ -863,90 +856,13 @@ namespace uba
 				responseSize += u32(entry.reader.GetLeft());
 		}
 
-		for (u32 i=0; i!=sendCount; ++i)
+		for (u32 i=0; i!=inFlightCount; ++i)
 			entries[i].~Entry();
 
 		readIndex += u32(sendCount);
 		return success;
 	}
 	
-	bool StorageClient::SendAllSegments(NetworkClient& client, u16 fetchId, u8* readBuffer, u64 left, u32 messageMaxSize)
-	{
-		struct Entry
-		{
-			Entry(NetworkClient& client, u16 fetchId, u8* readBuffer, u32 fetchIndex)
-				: message(client, ServiceId, StorageMessageType_FetchSegment, writer)
-				, reader(readBuffer, 0, SendMaxSize)
-				, done(true)
-			{
-				writer.WriteU16(fetchId);
-				writer.WriteU32(fetchIndex + 1);
-			}
-
-			StackBinaryWriter<16> writer;
-			NetworkMessage message;
-			BinaryReader reader;
-			Event done;
-		};
-
-		constexpr u32 EntryCount = 128;
-		u64 entriesMem[sizeof(Entry) * EntryCount / 8];
-		Entry* entries = (Entry*)entriesMem;
-
-		u32 entryIndex = 0;
-		u32 fetchIndex = 0;
-		bool success = true;
-		while (left)
-		{
-			if (fetchIndex >= EntryCount)
-			{
-				auto& entry = entries[entryIndex];
-				if (!entry.done.IsSet())
-				{
-					success = false;
-					break;
-				}
-				entry.~Entry();
-				UBA_ASSERT(entry.reader.GetLeft() == messageMaxSize);
-			}
-
-			auto& entry = *new (entries + entryIndex) Entry(client, fetchId, readBuffer, fetchIndex);
-			if (!entry.message.SendAsync(entry.reader, [](bool error, void* userData) { ((Event*)userData)->Set(); }, &entry.done))
-			{
-				entry.~Entry();
-				success = false;
-				break;
-			}
-			++fetchIndex;
-			entryIndex = (entryIndex + 1) % EntryCount;
-
-			if (left < messageMaxSize)
-				break;
-			left -= messageMaxSize;
-			readBuffer += messageMaxSize;
-		}
-
-		u32 waitIndex = 0;
-		u32 waitCount = entryIndex;
-		if (fetchIndex >= EntryCount)
-		{
-			waitIndex = entryIndex;
-			waitCount = EntryCount;
-		}
-		for (u32 i = 0; i != waitCount; ++i)
-		{
-			auto& entry = entries[waitIndex];
-			if (!entry.done.IsSet())
-				success = false;
-			if (!entry.message.ProcessAsyncResults(entry.reader))
-				success = false;
-			entry.~Entry();
-			waitIndex = (waitIndex + 1) % EntryCount;
-		}
-		return success;
-	}
-
-
 	bool StorageClient::SendFile(const CasKey& casKey, const tchar* fileName, u8* sourceMem, u64 sourceSize, const tchar* hint)
 	{
 		UBA_ASSERT(casKey != CasKeyZero);
diff --git a/Engine/Source/Programs/UnrealBuildAccelerator/Common/Public/UbaNetworkMessage.h b/Engine/Source/Programs/UnrealBuildAccelerator/Common/Public/UbaNetworkMessage.h
--- a/Engine/Source/Programs/UnrealBuildAccelerator/Common/Public/UbaNetworkMessage.h
+++ b/Engine/Source/Programs/UnrealBuildAccelerator/Common/Public/UbaNetworkMessage.h
@@ -10,8 +10,11 @@ namespace uba
 	struct NetworkMessage
 	{
 		NetworkMessage(NetworkClient& client, u8 serviceId, u8 messageType, BinaryWriter& sendWriter);
+		NetworkMessage() {}
 		~NetworkMessage();
 
+		void Init(NetworkClient& client, u8 serviceId, u8 messageType, BinaryWriter& sendWriter);
+
 		bool Send(BinaryReader& response);
 		bool Send(BinaryReader& response, Timer& outTimer);
 		bool Send();
@@ -23,8 +26,8 @@ namespace uba
 	private:
 		void Done(bool shouldLock = true);
 
-		NetworkClient& m_client;
-		BinaryWriter& m_sendWriter;
+		NetworkClient* m_client = nullptr;
+		BinaryWriter* m_sendWriter = nullptr;
 		DoneFunc* m_doneFunc = nullptr;
 		void* m_doneUserData = nullptr;
 		void* m_response = nullptr;
diff --git a/Engine/Source/Programs/UnrealBuildAccelerator/Common/Public/UbaStorageClient.h b/Engine/Source/Programs/UnrealBuildAccelerator/Common/Public/UbaStorageClient.h
--- a/Engine/Source/Programs/UnrealBuildAccelerator/Common/Public/UbaStorageClient.h
+++ b/Engine/Source/Programs/UnrealBuildAccelerator/Common/Public/UbaStorageClient.h
@@ -56,7 +56,6 @@ namespace uba
 		virtual void PrintSummary(Logger& logger) override;
 
 		static bool SendBatchMessages(Logger& logger, NetworkClient& client, u16 fetchId, u8* slot, u64 capacity, u64 left, u32 messageMaxSize, u32& readIndex, u32& responseSize);
-		static bool SendAllSegments(NetworkClient& client, u16 fetchId, u8* readBuffer, u64 left, u32 messageMaxSize);
 
 	private:
 		bool SendFile(const CasKey& casKey, const tchar* fileName, u8* sourceMem, u64 sourceSize, const tchar* hint);
