[Uba]
* Fixed linux compile error

diff --git a/Engine/Source/Programs/UnrealBuildAccelerator/Common/Private/UbaSession.cpp b/Engine/Source/Programs/UnrealBuildAccelerator/Common/Private/UbaSession.cpp
--- a/Engine/Source/Programs/UnrealBuildAccelerator/Common/Private/UbaSession.cpp
+++ b/Engine/Source/Programs/UnrealBuildAccelerator/Common/Private/UbaSession.cpp
@@ -1964,8 +1964,10 @@ namespace uba
 				auto bng = MakeGuard([&]() { if (shouldBottleneck) bottleneck.Leave(); });
 				#endif
 
+				#if PLATFORM_WINDOWS
 				shouldBottleneck = true;
 				bottleneck.Enter();
+				#endif
 
 				u64 fileSize = file.mappingWritten;
 				u8* mem = MapViewOfFile(file.mappingHandle, FILE_MAP_READ, 0, fileSize);
