[UBA]
* Fixed exit race condition for posix processes... session can end up in a situation where it sees the process as exited without getting an exit message.. we have assumed this was a crash but it is actually just a matter of message wait event timing out just before the exit message is written, and then the detoured process manages to exit before the isProcessAlive-checking happens..

Solution is simply to do one more check for messages in the shared memory if an exit has happened without an exit message

diff --git a/Engine/Source/Programs/UnrealBuildAccelerator/Common/Private/UbaProcess.cpp b/Engine/Source/Programs/UnrealBuildAccelerator/Common/Private/UbaProcess.cpp
--- a/Engine/Source/Programs/UnrealBuildAccelerator/Common/Private/UbaProcess.cpp
+++ b/Engine/Source/Programs/UnrealBuildAccelerator/Common/Private/UbaProcess.cpp
@@ -267,7 +267,7 @@ namespace uba
 		#endif
 
 		siginfo_t signalInfo;
-		while (true)
+		while (m_nativeProcessId != 0)
 		{
 			memset(&signalInfo, 0, sizeof(signalInfo));
 			int res = waitid(P_PID, (unsigned int)m_nativeProcessId, &signalInfo, WEXITED | WNOHANG | WNOWAIT);
@@ -284,18 +284,57 @@ namespace uba
 			else
 			{
 				if (signalInfo.si_pid != (pid_t)m_nativeProcessId)
-				{
-					++m_waitidCount;
 					return true;
+
+				const char* codeType = nullptr;
+				const char* extraString = "";
+				switch (signalInfo.si_code)
+				{
+				case CLD_KILLED:
+					codeType = "killed";
+					break;
+				case CLD_DUMPED:
+					codeType = "killed";
+					extraString = " (dumped core)";
+					break;
+				case CLD_STOPPED:
+					codeType = "stopped";
+					break;
+				case CLD_TRAPPED:
+					codeType = "trapped";
+					break;
+				case CLD_CONTINUED:
+					codeType = "continued";
+					break;
 				}
-				break;
+
+				m_nativeProcessId = 0;
+				m_nativeProcessExitCode = signalInfo.si_status;
+				
+				if (!codeType) // Is null if graceful exit (CLD_EXITED)
+					break;
+					
+				StringBuffer<> err;
+				err.Appendf(TC("Process %u (%s) %s by signal %i%s. Received %u messages. Execution time: %s."), m_nativeProcessId, m_description.c_str(), codeType, signalInfo.si_status, m_messageCount, TimeToText(GetTime() - m_startTime).str);
+				LogLine(false, err.data, LogEntryType_Error);
+				m_nativeProcessExitCode = UBA_EXIT_CODE(666); // We do exit code 666 to trigger non-uba retry on the outside
+				return false;
 			}
 		}
 
+		// There is a small race condition between this process polling and exit message.
+		// Detoured process can't wait for exit message response and then close the shared memory because it might end up closing another process memory
+		// .. so solution is to do one more poll from here to make sure we pick up the message before leaving.
 		if (!m_gotExitMessage)
 		{
+			if (m_doOneExtraCheckForExitMessage)
+			{
+				m_doOneExtraCheckForExitMessage = false;
+				return true;
+			}
+
 			StringBuffer<> err;
-			err.Appendf(TC("ERROR: Process %u (%s) not active but did not get exit message. Received %u messages. Called waitid %u times. Signal code: %i. Exit value or signal: %i. Execution time: %s."), m_nativeProcessId, m_description.c_str(), m_messageCount, m_waitidCount, signalInfo.si_code, signalInfo.si_status, TimeToText(GetTime() - m_startTime).str);
+			err.Appendf(TC("ERROR: Process %u (%s) not active but did not get exit message. Received %u messages. Signal code: %i. Exit value or signal: %i. Execution time: %s."), m_nativeProcessId, m_description.c_str(), m_messageCount, signalInfo.si_code, signalInfo.si_status, TimeToText(GetTime() - m_startTime).str);
 			LogLine(false, err.data, LogEntryType_Error);
 			m_nativeProcessExitCode = UBA_EXIT_CODE(666);
 		}
@@ -1723,7 +1762,8 @@ namespace uba
 
 		if (cancel)
 		{
-			kill((pid_t)m_nativeProcessId, -1);
+			if (m_nativeProcessId)
+				kill((pid_t)m_nativeProcessId, -1);
 			return m_nativeProcessExitCode;
 		}
 
@@ -1733,7 +1773,7 @@ namespace uba
 
 		// Process should have been waited on here because of IsActive
 		int status = 0;
-		while (true)
+		while (m_nativeProcessId)
 		{
 			int res = waitpid((pid_t)m_nativeProcessId, &status, 0);
 			if (res == -1)
diff --git a/Engine/Source/Programs/UnrealBuildAccelerator/Common/Public/UbaProcess.h b/Engine/Source/Programs/UnrealBuildAccelerator/Common/Public/UbaProcess.h
--- a/Engine/Source/Programs/UnrealBuildAccelerator/Common/Public/UbaProcess.h
+++ b/Engine/Source/Programs/UnrealBuildAccelerator/Common/Public/UbaProcess.h
@@ -104,7 +104,7 @@ namespace uba
 
 		int m_stdOutPipe = -1;
 		int m_stdErrPipe = -1;
-		u32 m_waitidCount = 0;
+		bool m_doOneExtraCheckForExitMessage = true;
 	#endif
 
 		ProcHandle m_nativeProcessHandle = InvalidProcHandle;
