[UBA]
* Added retry code for opening file for write. Will retry 5 times with one second in between. We have seen fails to open files for write randomly over the last five months (only a few times).. have no idea what could cause this.. have never seen it locally myself
* Fixed non unity compile error

diff --git a/Engine/Source/Programs/UnrealBuildAccelerator/Common/Private/UbaFileAccessor.cpp b/Engine/Source/Programs/UnrealBuildAccelerator/Common/Private/UbaFileAccessor.cpp
--- a/Engine/Source/Programs/UnrealBuildAccelerator/Common/Private/UbaFileAccessor.cpp
+++ b/Engine/Source/Programs/UnrealBuildAccelerator/Common/Private/UbaFileAccessor.cpp
@@ -72,15 +72,32 @@ namespace uba
 		if (allowRead)
 			dwDesiredAccess |= GENERIC_READ;
 		u32 dwShareMode = 0;// FILE_SHARE_READ | FILE_SHARE_WRITE;
-		m_fileHandle = uba::CreateFileW(realFileName, dwDesiredAccess, dwShareMode, createDisp, flagsAndAttributes);
-		if (m_fileHandle == InvalidFileHandle)
+		u32 retryCount = 0;
+		StringBuffer<256> additionalInfo;
+		while (true)
 		{
+			m_fileHandle = uba::CreateFileW(realFileName, dwDesiredAccess, dwShareMode, createDisp, flagsAndAttributes);
+			if (m_fileHandle != InvalidFileHandle)
+			{
+				if (retryCount)
+					m_logger.Warning(TC("Had to retry %u times to open file %s for write (because is was being used%s)"), retryCount, realFileName, additionalInfo.data);
+				break;
+			}
 			u32 lastError = GetLastError();
-			StringBuffer<256> additionalInfo;
-#if PLATFORM_WINDOWS
+			#if PLATFORM_WINDOWS
 			if (lastError == ERROR_SHARING_VIOLATION)
-				GetProcessHoldingFile(additionalInfo, m_fileName);
-#endif
+			{
+				if (retryCount == 0)
+					GetProcessHoldingFile(additionalInfo, m_fileName);
+
+				if (retryCount < 5)
+				{
+					Sleep(1000);
+					++retryCount;
+					continue;
+				}
+			}
+			#endif
 			return m_logger.Error(TC("ERROR opening file %s for write (%s%s)"), realFileName, LastErrorToText(lastError).data, additionalInfo.data);
 		}
 
diff --git a/Engine/Source/Programs/UnrealBuildAccelerator/Common/Public/UbaLogger.h b/Engine/Source/Programs/UnrealBuildAccelerator/Common/Public/UbaLogger.h
--- a/Engine/Source/Programs/UnrealBuildAccelerator/Common/Public/UbaLogger.h
+++ b/Engine/Source/Programs/UnrealBuildAccelerator/Common/Public/UbaLogger.h
@@ -4,6 +4,7 @@
 
 #include "UbaLogWriter.h"
 #include "UbaSynchronization.h"
+#include <stdarg.h>
 
 #define UBA_DEBUG_LOGGER 0
 
