UBA: Bumping size limits for commandlines as some projects need >64k when building, due to all of the args/flags/paths/includes/etc.

diff --git a/Engine/Source/Programs/UnrealBuildAccelerator/Common/Private/UbaProcess.cpp b/Engine/Source/Programs/UnrealBuildAccelerator/Common/Private/UbaProcess.cpp
--- a/Engine/Source/Programs/UnrealBuildAccelerator/Common/Private/UbaProcess.cpp
+++ b/Engine/Source/Programs/UnrealBuildAccelerator/Common/Private/UbaProcess.cpp
@@ -704,7 +704,7 @@ namespace uba
 
 					StringBuffer<> application;
 					reader.ReadString(application);
-					StringBuffer<64*1024> fullCommandLine;
+					StringBuffer<64*1024*2> fullCommandLine;
 					reader.ReadString(fullCommandLine);
 					StringBuffer<> currentDir;
 					reader.ReadString(currentDir);
diff --git a/Engine/Source/Programs/UnrealBuildAccelerator/Core/Public/UbaProtocol.h b/Engine/Source/Programs/UnrealBuildAccelerator/Core/Public/UbaProtocol.h
--- a/Engine/Source/Programs/UnrealBuildAccelerator/Core/Public/UbaProtocol.h
+++ b/Engine/Source/Programs/UnrealBuildAccelerator/Core/Public/UbaProtocol.h
@@ -37,7 +37,7 @@ namespace uba
 
 	inline constexpr u32 ProcessMessageVersion = 1339;
 
-	inline constexpr u32 CommunicationMemSize = 64*1024;
+	inline constexpr u32 CommunicationMemSize = 64*1024*2;
 
 	inline constexpr u32 FileMappingTableMemSize = 16 * 1024 * 1024;
 	inline constexpr u32 DirTableMemSize = 40 * 1024 * 1024;
