[UBA]
* Fixed badly formatted printf
* Added so failing to write output files return fail
* Moved a couple of message handlers to separate functions to solve stack complains from submit tool

diff --git a/Engine/Source/Programs/UnrealBuildAccelerator/Common/Private/UbaProcess.cpp b/Engine/Source/Programs/UnrealBuildAccelerator/Common/Private/UbaProcess.cpp
--- a/Engine/Source/Programs/UnrealBuildAccelerator/Common/Private/UbaProcess.cpp
+++ b/Engine/Source/Programs/UnrealBuildAccelerator/Common/Private/UbaProcess.cpp
@@ -315,7 +315,7 @@ namespace uba
 					break;
 					
 				StringBuffer<> err;
-				err.Appendf(TC("Process %u (%s) %s by signal %i%s. Received %u messages. Execution time: %s."), m_nativeProcessId, m_description.c_str(), codeType, signalInfo.si_status, m_messageCount, TimeToText(GetTime() - m_startTime).str);
+				err.Appendf(TC("Process %u (%s) %s by signal %i. Received %u messages. Execution time: %s."), m_nativeProcessId, m_description.c_str(), codeType, signalInfo.si_status, m_messageCount, TimeToText(GetTime() - m_startTime).str);
 				LogLine(false, err.data, LogEntryType_Error);
 				m_nativeProcessExitCode = UBA_EXIT_CODE(666); // We do exit code 666 to trigger non-uba retry on the outside
 				return false;
@@ -546,462 +546,465 @@ namespace uba
 		MessageType messageType = (MessageType)reader.ReadByte();
 		switch (messageType)
 		{
-			case MessageType_Init:
-				{
-					InitMessage msg;
-					InitResponse response;
-					m_messageSuccess = m_session.GetInitResponse(response, msg) && m_messageSuccess;
-					writer.WriteBool(m_echoOn);
-					writer.WriteString(m_startInfo.application);
-					writer.WriteString(m_startInfo.workingDir);
-					writer.WriteU64(response.directoryTableHandle);
-					writer.WriteU32(response.directoryTableSize);
-					writer.WriteU32(response.directoryTableCount);
-					writer.WriteU64(response.mappedFileTableHandle);
-					writer.WriteU32(response.mappedFileTableSize);
-					writer.WriteU32(response.mappedFileTableCount);
-					return true;
-				}
+			#define UBA_PROCESS_MESSAGE(type) case MessageType_##type: return Handle##type(reader, writer);
+			UBA_PROCESS_MESSAGES
+			#undef UBA_PROCESS_MESSAGE
+		}
+		return m_session.m_logger.Error(TC("Unknown message type %u"), messageType);
+	}
 
-			case MessageType_CreateFileW:
-				{
-					CreateFileMessage msg { *this };
-					reader.ReadString(msg.fileName);
-					msg.fileNameKey = reader.ReadStringKey();
-					msg.access = (FileAccess)reader.ReadByte();
-
-					CreateFileResponse response;
-					m_messageSuccess = m_session.CreateFile(response, msg) && m_messageSuccess;
-					writer.WriteString(response.fileName);
-					writer.WriteU64(response.size);
-					writer.WriteU32(response.closeId);
-					writer.WriteU32(response.mappedFileTableSize);
-					writer.WriteU32(response.directoryTableSize);
-					return true;
-				}
+	void ProcessImpl::LogLine(bool printInSession, TString&& line, LogEntryType logType)
+	{
+		if (IsCancelled())// || m_startInfo.stdoutHandle != INVALID_HANDLE_VALUE)
+			return;
+		if (printInSession)
+			m_session.m_logger.Log(LogEntryType_Warning, line.c_str(), u32(line.size()));
+		if (m_startInfo.logLineFunc)
+			m_startInfo.logLineFunc(m_startInfo.logLineUserData, line.c_str(), u32(line.size()), logType);
+		SCOPED_WRITE_LOCK(m_logLinesLock, l);
+		m_logLines.push_back({ std::move(line), logType });
+	}
 
-			case MessageType_GetFullFileName:
-				{
-					GetFullFileNameMessage msg { *this };
-					reader.ReadString(msg.fileName);
-					msg.fileNameKey = reader.ReadStringKey();
-					GetFullFileNameResponse response;
-					m_messageSuccess = m_session.GetFullFileName(response, msg) && m_messageSuccess;
-					writer.WriteString(response.fileName);
-					writer.WriteString(response.virtualFileName);
-					writer.WriteU32(response.mappedFileTableSize);
-					return true;
-				}
+	bool ProcessImpl::HandleInit(BinaryReader& reader, BinaryWriter& writer)
+	{
+		InitMessage msg;
+		InitResponse response;
+		m_messageSuccess = m_session.GetInitResponse(response, msg) && m_messageSuccess;
+		writer.WriteBool(m_echoOn);
+		writer.WriteString(m_startInfo.application);
+		writer.WriteString(m_startInfo.workingDir);
+		writer.WriteU64(response.directoryTableHandle);
+		writer.WriteU32(response.directoryTableSize);
+		writer.WriteU32(response.directoryTableCount);
+		writer.WriteU64(response.mappedFileTableHandle);
+		writer.WriteU32(response.mappedFileTableSize);
+		writer.WriteU32(response.mappedFileTableCount);
+		return true;
+	}
+	
+	bool ProcessImpl::HandleCreateFile(BinaryReader& reader, BinaryWriter& writer)
+	{
+		CreateFileMessage msg { *this };
+		reader.ReadString(msg.fileName);
+		msg.fileNameKey = reader.ReadStringKey();
+		msg.access = (FileAccess)reader.ReadByte();
+
+		CreateFileResponse response;
+		m_messageSuccess = m_session.CreateFile(response, msg) && m_messageSuccess;
+		writer.WriteString(response.fileName);
+		writer.WriteU64(response.size);
+		writer.WriteU32(response.closeId);
+		writer.WriteU32(response.mappedFileTableSize);
+		writer.WriteU32(response.directoryTableSize);
+		return true;
+	}
 
-			case MessageType_CloseFile:
-				{
-					CloseFileMessage msg { *this };
-					reader.ReadString(msg.fileName);
-					msg.closeId = reader.ReadU32();
-					msg.attributes = DefaultAttributes(); // reader.ReadU32(); TODO
-					msg.deleteOnClose = reader.ReadBool();
-					msg.success = reader.ReadBool();
-					msg.mappingHandle = reader.ReadU64();
-					msg.mappingWritten = reader.ReadU64();
-					msg.newNameKey = reader.ReadStringKey();
-					if (msg.newNameKey != StringKeyZero)
-						reader.ReadString(msg.newName);
-					CloseFileResponse response;
-					m_messageSuccess = m_session.CloseFile(response, msg) && m_messageSuccess;
-					writer.WriteU32(response.directoryTableSize);
-					return true;
-				}
+	bool ProcessImpl::HandleGetFullFileName(BinaryReader& reader, BinaryWriter& writer)
+	{
+		GetFullFileNameMessage msg { *this };
+		reader.ReadString(msg.fileName);
+		msg.fileNameKey = reader.ReadStringKey();
+		GetFullFileNameResponse response;
+		m_messageSuccess = m_session.GetFullFileName(response, msg) && m_messageSuccess;
+		writer.WriteString(response.fileName);
+		writer.WriteString(response.virtualFileName);
+		writer.WriteU32(response.mappedFileTableSize);
+		return true;
+	}
 
-			case MessageType_DeleteFileW:
-				{
-					DeleteFileMessage msg { *this };
-					reader.ReadString(msg.fileName);
-					msg.fileNameKey = reader.ReadStringKey();
-					msg.closeId = reader.ReadU32();
-					DeleteFileResponse response;
-					m_messageSuccess = m_session.DeleteFile(response, msg) && m_messageSuccess;
-					writer.WriteBool(response.result);
-					writer.WriteU32(response.errorCode);
-					writer.WriteU32(response.directoryTableSize);
-					return true;
-				}
+	bool ProcessImpl::HandleCloseFile(BinaryReader& reader, BinaryWriter& writer)
+	{
+		CloseFileMessage msg { *this };
+		reader.ReadString(msg.fileName);
+		msg.closeId = reader.ReadU32();
+		msg.attributes = DefaultAttributes(); // reader.ReadU32(); TODO
+		msg.deleteOnClose = reader.ReadBool();
+		msg.success = reader.ReadBool();
+		msg.mappingHandle = reader.ReadU64();
+		msg.mappingWritten = reader.ReadU64();
+		msg.newNameKey = reader.ReadStringKey();
+		if (msg.newNameKey != StringKeyZero)
+			reader.ReadString(msg.newName);
+		CloseFileResponse response;
+		m_messageSuccess = m_session.CloseFile(response, msg) && m_messageSuccess;
+		writer.WriteU32(response.directoryTableSize);
+		return true;
+	}
 
-			case MessageType_CopyFile:
-				{
-					CopyFileMessage msg{ *this };
-					msg.fromKey = reader.ReadStringKey();
-					reader.ReadString(msg.fromName);
-					msg.toKey = reader.ReadStringKey();
-					reader.ReadString(msg.toName);
-					CopyFileResponse response;
-					m_messageSuccess = m_session.CopyFile(response, msg) && m_messageSuccess;
-					writer.WriteString(response.fromName);
-					writer.WriteString(response.toName);
-					writer.WriteU32(response.closeId);
-					writer.WriteU32(response.errorCode);
-					writer.WriteU32(response.directoryTableSize);
-					return true;
-				}
+	bool ProcessImpl::HandleDeleteFile(BinaryReader& reader, BinaryWriter& writer)
+	{
+		DeleteFileMessage msg { *this };
+		reader.ReadString(msg.fileName);
+		msg.fileNameKey = reader.ReadStringKey();
+		msg.closeId = reader.ReadU32();
+		DeleteFileResponse response;
+		m_messageSuccess = m_session.DeleteFile(response, msg) && m_messageSuccess;
+		writer.WriteBool(response.result);
+		writer.WriteU32(response.errorCode);
+		writer.WriteU32(response.directoryTableSize);
+		return true;
+	}
 
-			case MessageType_MoveFileW:
-				{
-					MoveFileMessage msg { *this };
-					msg.fromKey = reader.ReadStringKey();
-					reader.ReadString(msg.fromName);
-					msg.toKey = reader.ReadStringKey();
-					reader.ReadString(msg.toName);
-					msg.flags = reader.ReadU32();
-					MoveFileResponse response;
-					m_messageSuccess = m_session.MoveFile(response, msg) && m_messageSuccess;
-					writer.WriteBool(response.result);
-					writer.WriteU32(response.errorCode);
-					writer.WriteU32(response.directoryTableSize);
-					return true;
-				}
+	bool ProcessImpl::HandleCopyFile(BinaryReader& reader, BinaryWriter& writer)
+	{
+		CopyFileMessage msg{ *this };
+		msg.fromKey = reader.ReadStringKey();
+		reader.ReadString(msg.fromName);
+		msg.toKey = reader.ReadStringKey();
+		reader.ReadString(msg.toName);
+		CopyFileResponse response;
+		m_messageSuccess = m_session.CopyFile(response, msg) && m_messageSuccess;
+		writer.WriteString(response.fromName);
+		writer.WriteString(response.toName);
+		writer.WriteU32(response.closeId);
+		writer.WriteU32(response.errorCode);
+		writer.WriteU32(response.directoryTableSize);
+		return true;
+	}
 
-			case MessageType_Chmod:
-				{
-					ChmodMessage msg { *this };
-					msg.fileNameKey = reader.ReadStringKey();
-					reader.ReadString(msg.fileName);
-					msg.fileMode = reader.ReadU32();
-					ChmodResponse response;
-					m_messageSuccess = m_session.Chmod(response, msg) && m_messageSuccess;
-					writer.WriteU32(response.errorCode);
-					return true;
-				}
+	bool ProcessImpl::HandleMoveFile(BinaryReader& reader, BinaryWriter& writer)
+	{
+		MoveFileMessage msg { *this };
+		msg.fromKey = reader.ReadStringKey();
+		reader.ReadString(msg.fromName);
+		msg.toKey = reader.ReadStringKey();
+		reader.ReadString(msg.toName);
+		msg.flags = reader.ReadU32();
+		MoveFileResponse response;
+		m_messageSuccess = m_session.MoveFile(response, msg) && m_messageSuccess;
+		writer.WriteBool(response.result);
+		writer.WriteU32(response.errorCode);
+		writer.WriteU32(response.directoryTableSize);
+		return true;
+	}
 
-			case MessageType_CreateDirectory:
-				{
-					CreateDirectoryMessage msg;
-					msg.nameKey = reader.ReadStringKey();
-					reader.ReadString(msg.name);
-					CreateDirectoryResponse response;
-					m_messageSuccess = m_session.CreateDirectory(response, msg) && m_messageSuccess;
-					writer.WriteBool(response.result);
-					writer.WriteU32(response.errorCode);
-					return true;
-				}
+	bool ProcessImpl::HandleChmod(BinaryReader& reader, BinaryWriter& writer)
+	{
+		ChmodMessage msg { *this };
+		msg.fileNameKey = reader.ReadStringKey();
+		reader.ReadString(msg.fileName);
+		msg.fileMode = reader.ReadU32();
+		ChmodResponse response;
+		m_messageSuccess = m_session.Chmod(response, msg) && m_messageSuccess;
+		writer.WriteU32(response.errorCode);
+		return true;
+	}
+	bool ProcessImpl::HandleCreateDirectory(BinaryReader& reader, BinaryWriter& writer)
+	{
+		CreateDirectoryMessage msg;
+		msg.nameKey = reader.ReadStringKey();
+		reader.ReadString(msg.name);
+		CreateDirectoryResponse response;
+		m_messageSuccess = m_session.CreateDirectory(response, msg) && m_messageSuccess;
+		writer.WriteBool(response.result);
+		writer.WriteU32(response.errorCode);
+		return true;
+	}
 
-			case MessageType_ListDirectory:
-				{
-					ListDirectoryMessage msg;
-					reader.ReadString(msg.directoryName);
-					msg.directoryNameKey = reader.ReadStringKey();
-					ListDirectoryResponse response;
-					m_messageSuccess = m_session.GetListDirectoryInfo(response, msg.directoryName.data, msg.directoryNameKey) && m_messageSuccess;
-					writer.WriteU32(response.tableSize);
-					writer.WriteU32(response.tableOffset);
-					return true;
-				}
+	bool ProcessImpl::HandleListDirectory(BinaryReader& reader, BinaryWriter& writer)
+	{
+		ListDirectoryMessage msg;
+		reader.ReadString(msg.directoryName);
+		msg.directoryNameKey = reader.ReadStringKey();
+		ListDirectoryResponse response;
+		m_messageSuccess = m_session.GetListDirectoryInfo(response, msg.directoryName.data, msg.directoryNameKey) && m_messageSuccess;
+		writer.WriteU32(response.tableSize);
+		writer.WriteU32(response.tableOffset);
+		return true;
+	}
 
-			case MessageType_CreateProcess:
-				{
-					#if PLATFORM_LINUX
-					// This process will become the parent of a process if it becomes orphaned
-					static bool subreaper = []() { prctl(PR_SET_CHILD_SUBREAPER, 1, 0, 0, 0); return true; }();
-					#endif
-
-					StringBuffer<> application;
-					reader.ReadString(application);
-					StringBuffer<64*1024*2> fullCommandLine;
-					reader.ReadString(fullCommandLine);
-					StringBuffer<> currentDir;
-					reader.ReadString(currentDir);
-					if (currentDir.IsEmpty())
-						currentDir.Append(m_startInfo.workingDir);
-
-					const tchar* commandLine = nullptr;
-
-					// Remove application from command line
-					if (fullCommandLine[0] == '"')
-					{
-						const tchar* second = fullCommandLine.First('"', 1);
-						if (!second)
-							second = fullCommandLine.data + fullCommandLine.count;
-						//UBA_ASSERTF(second, TC("Missing second '\"' in command line: %s"), fullCommandLine.data); // "Unsupported cmd line format"
-						commandLine = second + 1;
-						if (application.IsEmpty())
-							application.Append(fullCommandLine.data + 1, u64(second - fullCommandLine.data - 1));
-					}
-					else
-					{
-						const tchar* secondParamStart = fullCommandLine.First(' ', 1);
-						if (!secondParamStart)
-							commandLine = TC("");
-						else
-							commandLine = secondParamStart + 1;
-						if (application.IsEmpty())
-							application.Append(fullCommandLine.data, u64(secondParamStart - fullCommandLine.data));
-					}
+	bool ProcessImpl::HandleUpdateTables(BinaryReader& reader, BinaryWriter& writer)
+	{
+		u32 dirSize = m_session.GetDirectoryTableSize();
+		writer.WriteU32(dirSize);
+		return true;
+	}
+
+	bool ProcessImpl::HandleCreateProcess(BinaryReader& reader, BinaryWriter& writer)
+	{
+		#if PLATFORM_LINUX
+		// This process will become the parent of a process if it becomes orphaned
+		static bool subreaper = []() { prctl(PR_SET_CHILD_SUBREAPER, 1, 0, 0, 0); return true; }();
+		#endif
 
-					while (*commandLine == ' ')
-						++commandLine;
-
-					StringBuffer<> temp;
-					ProcessStartInfo info;
-					info.application = application.data;
-					info.arguments = commandLine;
-					info.description = application.GetFileName();
-					info.workingDir = currentDir.data;
-					info.logFile = InternalGetChildLogFile(temp);
-					info.priorityClass = m_startInfo.priorityClass;
-					info.outputStatsThresholdMs = m_startInfo.outputStatsThresholdMs;
-					info.logLineUserData = this;
-					info.logLineFunc = [](void* userData, const tchar* line, u32 length, LogEntryType type) { ((ProcessImpl*)userData)->LogLine(false, TString(line, length), type); };
-
-					ProcessHandle h = m_session.InternalRunProcess(info, true, this, true);
-					m_childProcesses.push_back(h);
-					u32 childProcessId = u32(m_childProcesses.size());
-
-					auto& process = *(ProcessImpl*)h.m_process;
-					process.m_echoOn = m_echoOn;
-
-					const char* detoursLib = m_session.m_detoursLibrary.c_str();
-					u32 detoursLibLen = u32(m_session.m_detoursLibrary.size());
-
-					writer.WriteU32(childProcessId);
-					writer.WriteU32(process.m_rulesIndex);
-					writer.WriteU32(detoursLibLen);
-					writer.WriteBytes(detoursLib, detoursLibLen);
-
-					writer.WriteString(m_realWorkingDir);
-					#if PLATFORM_WINDOWS
-					TString realCommandLine = TC("\"") + process.m_realApplication + TC("\" ") + commandLine;
-					writer.WriteString(realCommandLine);
-					#else
-					writer.WriteString(process.m_realApplication);
-					writer.WriteString(commandLine);
-					writer.WriteU64(process.m_comMemory.handle.uid);
-					writer.WriteU32(process.m_comMemory.offset);
-					writer.WriteString(info.logFile);
-					#endif
+		StringBuffer<> application;
+		reader.ReadString(application);
+		StringBuffer<64*1024*2> fullCommandLine;
+		reader.ReadString(fullCommandLine);
+		StringBuffer<> currentDir;
+		reader.ReadString(currentDir);
+		if (currentDir.IsEmpty())
+			currentDir.Append(m_startInfo.workingDir);
 
-					return true;
-				}
+		const tchar* commandLine = nullptr;
 
-			case MessageType_StartProcess:
-				{
-					u32 processId = reader.ReadU32();
-					UBA_ASSERT(processId > 0 && processId <= m_childProcesses.size());
-					auto& process = *(ProcessImpl*)m_childProcesses[processId - 1].m_process;
-					bool result = reader.ReadBool();
-					u32 lastError = reader.ReadU32();
-					if (!result)
-					{
-						m_session.m_logger.Logf(LogEntryType_Info, TC("Detoured process failed to start child process - %s. %s (Working dir: %s)"), LastErrorToText(lastError).data, process.m_realApplication.c_str(), process.m_realWorkingDir);
-						process.m_waitForParent.Set();
-						return true;
-					}
+		// Remove application from command line
+		if (fullCommandLine[0] == '"')
+		{
+			const tchar* second = fullCommandLine.First('"', 1);
+			if (!second)
+				second = fullCommandLine.data + fullCommandLine.count;
+			//UBA_ASSERTF(second, TC("Missing second '\"' in command line: %s"), fullCommandLine.data); // "Unsupported cmd line format"
+			commandLine = second + 1;
+			if (application.IsEmpty())
+				application.Append(fullCommandLine.data + 1, u64(second - fullCommandLine.data - 1));
+		}
+		else
+		{
+			const tchar* secondParamStart = fullCommandLine.First(' ', 1);
+			if (!secondParamStart)
+				commandLine = TC("");
+			else
+				commandLine = secondParamStart + 1;
+			if (application.IsEmpty())
+				application.Append(fullCommandLine.data, u64(secondParamStart - fullCommandLine.data));
+		}
+
+		while (*commandLine == ' ')
+			++commandLine;
+
+		StringBuffer<> temp;
+		ProcessStartInfo info;
+		info.application = application.data;
+		info.arguments = commandLine;
+		info.description = application.GetFileName();
+		info.workingDir = currentDir.data;
+		info.logFile = InternalGetChildLogFile(temp);
+		info.priorityClass = m_startInfo.priorityClass;
+		info.outputStatsThresholdMs = m_startInfo.outputStatsThresholdMs;
+		info.logLineUserData = this;
+		info.logLineFunc = [](void* userData, const tchar* line, u32 length, LogEntryType type) { ((ProcessImpl*)userData)->LogLine(false, TString(line, length), type); };
+
+		ProcessHandle h = m_session.InternalRunProcess(info, true, this, true);
+		m_childProcesses.push_back(h);
+		u32 childProcessId = u32(m_childProcesses.size());
+
+		auto& process = *(ProcessImpl*)h.m_process;
+		process.m_echoOn = m_echoOn;
+
+		const char* detoursLib = m_session.m_detoursLibrary.c_str();
+		u32 detoursLibLen = u32(m_session.m_detoursLibrary.size());
+
+		writer.WriteU32(childProcessId);
+		writer.WriteU32(process.m_rulesIndex);
+		writer.WriteU32(detoursLibLen);
+		writer.WriteBytes(detoursLib, detoursLibLen);
+
+		writer.WriteString(m_realWorkingDir);
+		#if PLATFORM_WINDOWS
+		TString realCommandLine = TC("\"") + process.m_realApplication + TC("\" ") + commandLine;
+		writer.WriteString(realCommandLine);
+		#else
+		writer.WriteString(process.m_realApplication);
+		writer.WriteString(commandLine);
+		writer.WriteU64(process.m_comMemory.handle.uid);
+		writer.WriteU32(process.m_comMemory.offset);
+		writer.WriteString(info.logFile);
+		#endif
+
+		return true;
+	}
+
+	bool ProcessImpl::HandleStartProcess(BinaryReader& reader, BinaryWriter& writer)
+	{
+		u32 processId = reader.ReadU32();
+		UBA_ASSERT(processId > 0 && processId <= m_childProcesses.size());
+		auto& process = *(ProcessImpl*)m_childProcesses[processId - 1].m_process;
+		bool result = reader.ReadBool();
+		u32 lastError = reader.ReadU32();
+		if (!result)
+		{
+			m_session.m_logger.Logf(LogEntryType_Info, TC("Detoured process failed to start child process - %s. %s (Working dir: %s)"), LastErrorToText(lastError).data, process.m_realApplication.c_str(), process.m_realWorkingDir);
+			process.m_waitForParent.Set();
+			return true;
+		}
 
 #if PLATFORM_WINDOWS
-					HANDLE nativeProcessHandle = (HANDLE)reader.ReadU64();
-					u32 nativeProcessId = reader.ReadU32();
-					HANDLE nativeThreadHandle = (HANDLE)reader.ReadU64();
+		HANDLE nativeProcessHandle = (HANDLE)reader.ReadU64();
+		u32 nativeProcessId = reader.ReadU32();
+		HANDLE nativeThreadHandle = (HANDLE)reader.ReadU64();
 
-					if (nativeProcessHandle)
-					{
-						DuplicateHandle((HANDLE)m_nativeProcessHandle, nativeProcessHandle, GetCurrentProcess(), (HANDLE*)&process.m_nativeProcessHandle, 0, false, DUPLICATE_SAME_ACCESS);
-						UBA_ASSERT(process.m_nativeProcessHandle && process.m_nativeProcessHandle != InvalidProcHandle);
-						DuplicateHandle((HANDLE)m_nativeProcessHandle, nativeThreadHandle, GetCurrentProcess(), &process.m_nativeThreadHandle, 0, false, DUPLICATE_SAME_ACCESS);
-						UBA_ASSERT(process.m_nativeThreadHandle && process.m_nativeThreadHandle != INVALID_HANDLE_VALUE);
-						process.m_nativeProcessId = nativeProcessId;
-					}
+		if (nativeProcessHandle)
+		{
+			DuplicateHandle((HANDLE)m_nativeProcessHandle, nativeProcessHandle, GetCurrentProcess(), (HANDLE*)&process.m_nativeProcessHandle, 0, false, DUPLICATE_SAME_ACCESS);
+			UBA_ASSERT(process.m_nativeProcessHandle && process.m_nativeProcessHandle != InvalidProcHandle);
+			DuplicateHandle((HANDLE)m_nativeProcessHandle, nativeThreadHandle, GetCurrentProcess(), &process.m_nativeThreadHandle, 0, false, DUPLICATE_SAME_ACCESS);
+			UBA_ASSERT(process.m_nativeThreadHandle && process.m_nativeThreadHandle != INVALID_HANDLE_VALUE);
+			process.m_nativeProcessId = nativeProcessId;
+		}
 #else
-					u64 nativeProcessHandle = reader.ReadU64();
-					u32 nativeProcessId = reader.ReadU32();
-					u64 nativeThreadHandle = reader.ReadU64();
-					process.m_nativeProcessHandle = (ProcHandle)nativeProcessHandle;
-					process.m_nativeProcessId = nativeProcessId;
+		u64 nativeProcessHandle = reader.ReadU64();
+		u32 nativeProcessId = reader.ReadU32();
+		u64 nativeThreadHandle = reader.ReadU64();
+		process.m_nativeProcessHandle = (ProcHandle)nativeProcessHandle;
+		process.m_nativeProcessId = nativeProcessId;
 #endif
-					process.m_waitForParent.Set();
-					return true;
-				}
+		process.m_waitForParent.Set();
+		return true;
+	}
+	bool ProcessImpl::HandleExitChildProcess(BinaryReader& reader, BinaryWriter& writer)
+	{
+		u32 nativeProcessId = reader.ReadU32();
+		for (auto& child : m_childProcesses)
+		{
+			auto& process = *(ProcessImpl*)child.m_process;
+			if (process.m_nativeProcessId != nativeProcessId)
+				continue;
+			process.m_parentReportedExit = true;
+			return true;
+		}
+		UBA_ASSERT(false);
+		return true;
+	}
 
-			case MessageType_ExitChildProcess:
-				{
-					u32 nativeProcessId = reader.ReadU32();
-					for (auto& child : m_childProcesses)
-					{
-						auto& process = *(ProcessImpl*)child.m_process;
-						if (process.m_nativeProcessId != nativeProcessId)
-							continue;
-						process.m_parentReportedExit = true;
-						return true;
-					}
-					UBA_ASSERT(false);
-					return true;
-				}
+	bool ProcessImpl::HandleCreateTempFile(BinaryReader& reader, BinaryWriter& writer)
+	{
+		CreateTempFile(reader, m_nativeProcessHandle, m_virtualApplication.c_str());
+		return true;
+	}
 
-			case MessageType_UpdateTables:
-				{
-					u32 dirSize = m_session.GetDirectoryTableSize();
-					writer.WriteU32(dirSize);
-					return true;
-				}
+	bool ProcessImpl::HandleOpenTempFile(BinaryReader& reader, BinaryWriter& writer)
+	{
+		OpenTempFile(reader, writer, m_virtualApplication.c_str());
+		return true;
+	}
 
-			case MessageType_CreateTempFile:
-				{
-					CreateTempFile(reader, m_nativeProcessHandle, m_virtualApplication.c_str());
-					return true;
-				}
+	bool ProcessImpl::HandleVirtualAllocFailed(BinaryReader& reader, BinaryWriter& writer)
+	{
+		StringBuffer<> allocType;
+		reader.ReadString(allocType);
+		u32 error = reader.ReadU32();
+		m_session.AllocFailed(*this, allocType.data, error);
+		return true;
+	}
 
-			case MessageType_OpenTempFile:
-				{
-					OpenTempFile(reader, writer, m_virtualApplication.c_str());
-					return true;
-				}
+	bool ProcessImpl::HandleLog(BinaryReader& reader, BinaryWriter& writer)
+	{
+		bool printInSession = reader.ReadBool();
+		bool isError = reader.ReadBool();
+		TString line = reader.ReadString();
+		LogLine(printInSession, std::move(line), isError ? LogEntryType_Error : LogEntryType_Info);
+		return true;
+	}
 
-			case MessageType_VirtualAllocFailed:
-				{
-					StringBuffer<> allocType;
-					reader.ReadString(allocType);
-					u32 error = reader.ReadU32();
-					m_session.AllocFailed(*this, allocType.data, error);
-					return true;
-				}
+	bool ProcessImpl::HandleEchoOn(BinaryReader& reader, BinaryWriter& writer)
+	{
+		m_echoOn = reader.ReadBool();
+		return true;
+	}
 
-			case MessageType_Log:
-				{
-					bool printInSession = reader.ReadBool();
-					bool isError = reader.ReadBool();
-					TString line = reader.ReadString();
-					LogLine(printInSession, std::move(line), isError ? LogEntryType_Error : LogEntryType_Info);
-					return true;
-				}
-				
-			case MessageType_EchoOn:
-				{
-					m_echoOn = reader.ReadBool();
-					return true;
-				}
+	bool ProcessImpl::HandleInputDependencies(BinaryReader& reader, BinaryWriter& writer)
+	{
+		UBA_ASSERT(m_startInfo.trackInputs);
+		if (m_trackedInputs.empty())
+		{
+			u32 trackedInputsSize = reader.ReadU32();
+			m_trackedInputs.reserve(trackedInputsSize);
+		}
+		u32 toRead = reader.ReadU32();
+		u8* pos = m_trackedInputs.data() + m_trackedInputs.size();
+		m_trackedInputs.resize(m_trackedInputs.size() + toRead);
+		reader.ReadBytes(pos, toRead);
+		return true;
+	}
 
-			case MessageType_InputDependencies:
-				{
-					UBA_ASSERT(m_startInfo.trackInputs);
-					if (m_trackedInputs.empty())
-					{
-						u32 trackedInputsSize = reader.ReadU32();
-						m_trackedInputs.reserve(trackedInputsSize);
-					}
-					u32 toRead = reader.ReadU32();
-					u8* pos = m_trackedInputs.data() + m_trackedInputs.size();
-					m_trackedInputs.resize(m_trackedInputs.size() + toRead);
-					reader.ReadBytes(pos, toRead);
-					return true;
-				}
+	bool ProcessImpl::HandleExit(BinaryReader& reader, BinaryWriter& writer)
+	{
+		m_gotExitMessage = true;
+		m_nativeProcessExitCode = reader.ReadU32();
 
-			case MessageType_Exit:
-				{
-					m_gotExitMessage = true;
-					m_nativeProcessExitCode = reader.ReadU32();
+		StringBuffer<> logName;
+		reader.ReadString(logName);
 
-					StringBuffer<> logName;
-					reader.ReadString(logName);
+		ProcessStats stats;
+		stats.Read(reader, ~0u);
 
-					ProcessStats stats;
-					stats.Read(reader, ~0u);
+		m_processStats.Add(stats);
 
-					m_processStats.Add(stats);
+		if (!IsCancelled())
+			if (m_startInfo.writeOutputFilesOnFail || GetApplicationRules()[m_rulesIndex].rules->IsExitCodeSuccess(m_nativeProcessExitCode))
+				m_messageSuccess = WriteFilesToDisk() && m_messageSuccess;
 
-					if (!IsCancelled())
-						if (m_startInfo.writeOutputFilesOnFail || GetApplicationRules()[m_rulesIndex].rules->IsExitCodeSuccess(m_nativeProcessExitCode))
-							WriteFilesToDisk();
+		if (m_parentProcess)
+		{
+			m_parentProcess->m_processStats.Add(m_processStats);
+			m_parentProcess->m_sessionStats.Add(m_sessionStats);
+			m_parentProcess->m_storageStats.Add(m_storageStats);
+			m_parentProcess->m_systemStats.Add(m_systemStats);
+		}
 
-					if (m_parentProcess)
-					{
-						m_parentProcess->m_processStats.Add(m_processStats);
-						m_parentProcess->m_sessionStats.Add(m_sessionStats);
-						m_parentProcess->m_storageStats.Add(m_storageStats);
-						m_parentProcess->m_systemStats.Add(m_systemStats);
-					}
+		if (m_startInfo.outputStatsThresholdMs && TimeToMs(m_processStats.GetTotalTime()) > m_startInfo.outputStatsThresholdMs)
+		{
+			m_session.PrintProcessStats(m_processStats, logName.data);
+			m_processStats.Print(m_session.m_logger, GetFrequency());
+		}
 
-					if (m_startInfo.outputStatsThresholdMs && TimeToMs(m_processStats.GetTotalTime()) > m_startInfo.outputStatsThresholdMs)
-					{
-						m_session.PrintProcessStats(m_processStats, logName.data);
-						m_processStats.Print(m_session.m_logger, GetFrequency());
-					}
+		return false;
+	}
 
-					return false;
-				}
+	bool ProcessImpl::HandleFlushWrittenFiles(BinaryReader& reader, BinaryWriter& writer)
+	{
+		WriteFilesToDisk();
+		bool result = m_session.FlushWrittenFiles(*this);
+		writer.WriteBool(result);
+		return true;
+	}
 
-			case MessageType_GetNextProcess:
-				{
-					u32 prevExitCode = reader.ReadU32();
-					NextProcessInfo nextProcess;
+	bool ProcessImpl::HandleUpdateEnvironment(BinaryReader& reader, BinaryWriter& writer)
+	{
+		StringBuffer<> reason;
+		reader.ReadString(reason);
+		bool resetStats = reader.ReadBool();
+		bool result = m_session.UpdateEnvironment(*this, reason.data, resetStats);
+		writer.WriteBool(result);
+		return true;
+	}
+
+	bool ProcessImpl::HandleGetNextProcess(BinaryReader& reader, BinaryWriter& writer)
+	{
+		u32 prevExitCode = reader.ReadU32();
+		NextProcessInfo nextProcess;
 
-					StackBinaryWriter<16 * 1024> statsWriter;
+		StackBinaryWriter<16 * 1024> statsWriter;
 					
-					ProcessStats processStats;
-					processStats.Read(reader, TraceVersion);
-					processStats.startupTime = m_processStats.startupTime;
-					processStats.wallTime = GetTime() - m_startTime;
-					processStats.cpuTime = 0;
-					processStats.hostTotalTime = m_processStats.hostTotalTime;
+		ProcessStats processStats;
+		processStats.Read(reader, TraceVersion);
+		processStats.startupTime = m_processStats.startupTime;
+		processStats.wallTime = GetTime() - m_startTime;
+		processStats.cpuTime = 0;
+		processStats.hostTotalTime = m_processStats.hostTotalTime;
 					
-					processStats.Write(statsWriter);
-					m_sessionStats.Write(statsWriter);
-					m_storageStats.Write(statsWriter);
-					m_systemStats.Write(statsWriter);
-					BinaryReader statsReader(statsWriter.GetData(), 0, statsWriter.GetPosition());
-
-					bool newProcess = false;
-					m_exitCode = prevExitCode;
-					m_messageSuccess = m_session.GetNextProcess(*this, newProcess, nextProcess, prevExitCode, statsReader) && m_messageSuccess;
-					writer.WriteBool(newProcess);
-					m_exitCode = ~0u;
-					if (!newProcess)
-						return true;
-
-					m_startTime = GetTime();
-					m_processStats = {};
-					m_sessionStats = {};
-					m_storageStats = {};
-					m_systemStats = {};
-
-					writer.WriteString(nextProcess.arguments);
-					writer.WriteString(nextProcess.workingDir);
-					writer.WriteString(nextProcess.description);
-					writer.WriteString(nextProcess.logFile);
-					return true;
-				}
-
-			case MessageType_Custom:
-				{
-					m_session.CustomMessage(*this, reader, writer);
-					return true;
-				}
-			case MessageType_FlushWrittenFiles:
-				{
-					WriteFilesToDisk();
-					bool result = m_session.FlushWrittenFiles(*this);
-					writer.WriteBool(result);
-					return true;
-				}
+		processStats.Write(statsWriter);
+		m_sessionStats.Write(statsWriter);
+		m_storageStats.Write(statsWriter);
+		m_systemStats.Write(statsWriter);
+		BinaryReader statsReader(statsWriter.GetData(), 0, statsWriter.GetPosition());
+
+		bool newProcess = false;
+		m_exitCode = prevExitCode;
+		m_messageSuccess = m_session.GetNextProcess(*this, newProcess, nextProcess, prevExitCode, statsReader) && m_messageSuccess;
+		writer.WriteBool(newProcess);
+		m_exitCode = ~0u;
+		if (!newProcess)
+			return true;
 
-			case MessageType_UpdateEnvironment:
-				{
-					StringBuffer<> reason;
-					reader.ReadString(reason);
-					bool resetStats = reader.ReadBool();
-					bool result = m_session.UpdateEnvironment(*this, reason.data, resetStats);
-					writer.WriteBool(result);
-					return true;
-				}
-		}
-		return m_session.m_logger.Error(TC("Unknown message type %u"), messageType);
+		m_startTime = GetTime();
+		m_processStats = {};
+		m_sessionStats = {};
+		m_storageStats = {};
+		m_systemStats = {};
+
+		writer.WriteString(nextProcess.arguments);
+		writer.WriteString(nextProcess.workingDir);
+		writer.WriteString(nextProcess.description);
+		writer.WriteString(nextProcess.logFile);
+		return true;
 	}
 
-	void ProcessImpl::LogLine(bool printInSession, TString&& line, LogEntryType logType)
+	bool ProcessImpl::HandleCustom(BinaryReader& reader, BinaryWriter& writer)
 	{
-		if (IsCancelled())// || m_startInfo.stdoutHandle != INVALID_HANDLE_VALUE)
-			return;
-		if (printInSession)
-			m_session.m_logger.Log(LogEntryType_Warning, line.c_str(), u32(line.size()));
-		if (m_startInfo.logLineFunc)
-			m_startInfo.logLineFunc(m_startInfo.logLineUserData, line.c_str(), u32(line.size()), logType);
-		SCOPED_WRITE_LOCK(m_logLinesLock, l);
-		m_logLines.push_back({ std::move(line), logType });
+		m_session.CustomMessage(*this, reader, writer);
+		return true;
 	}
 
 	bool ProcessImpl::CreateTempFile(BinaryReader& reader, ProcHandle nativeProcessHandle, const tchar* application)
diff --git a/Engine/Source/Programs/UnrealBuildAccelerator/Common/Public/UbaProcess.h b/Engine/Source/Programs/UnrealBuildAccelerator/Common/Public/UbaProcess.h
--- a/Engine/Source/Programs/UnrealBuildAccelerator/Common/Public/UbaProcess.h
+++ b/Engine/Source/Programs/UnrealBuildAccelerator/Common/Public/UbaProcess.h
@@ -67,6 +67,32 @@ namespace uba
 
 		void ThreadRun(bool runningRemote, void* environment);
 		bool HandleMessage(BinaryReader& reader, BinaryWriter& writer);
+		bool HandleInit(BinaryReader& reader, BinaryWriter& writer);
+		bool HandleCreateFile(BinaryReader& reader, BinaryWriter& writer);
+		bool HandleGetFullFileName(BinaryReader& reader, BinaryWriter& writer);
+		bool HandleCloseFile(BinaryReader& reader, BinaryWriter& writer);
+		bool HandleDeleteFile(BinaryReader& reader, BinaryWriter& writer);
+		bool HandleCopyFile(BinaryReader& reader, BinaryWriter& writer);
+		bool HandleMoveFile(BinaryReader& reader, BinaryWriter& writer);
+		bool HandleChmod(BinaryReader& reader, BinaryWriter& writer);
+		bool HandleCreateDirectory(BinaryReader& reader, BinaryWriter& writer);
+		bool HandleListDirectory(BinaryReader& reader, BinaryWriter& writer);
+		bool HandleUpdateTables(BinaryReader& reader, BinaryWriter& writer);
+		bool HandleCreateProcess(BinaryReader& reader, BinaryWriter& writer);
+		bool HandleStartProcess(BinaryReader& reader, BinaryWriter& writer);
+		bool HandleExitChildProcess(BinaryReader& reader, BinaryWriter& writer);
+		bool HandleCreateTempFile(BinaryReader& reader, BinaryWriter& writer);
+		bool HandleOpenTempFile(BinaryReader& reader, BinaryWriter& writer);
+		bool HandleVirtualAllocFailed(BinaryReader& reader, BinaryWriter& writer);
+		bool HandleLog(BinaryReader& reader, BinaryWriter& writer);
+		bool HandleEchoOn(BinaryReader& reader, BinaryWriter& writer);
+		bool HandleInputDependencies(BinaryReader& reader, BinaryWriter& writer);
+		bool HandleExit(BinaryReader& reader, BinaryWriter& writer);
+		bool HandleFlushWrittenFiles(BinaryReader& reader, BinaryWriter& writer);
+		bool HandleUpdateEnvironment(BinaryReader& reader, BinaryWriter& writer);
+		bool HandleGetNextProcess(BinaryReader& reader, BinaryWriter& writer);
+		bool HandleCustom(BinaryReader& reader, BinaryWriter& writer);
+
 		void LogLine(bool printInSession, TString&& line, LogEntryType logType);
 		bool CreateTempFile(BinaryReader& reader, ProcHandle nativeProcessHandle, const tchar* application);
 		bool OpenTempFile(BinaryReader& reader, BinaryWriter& writer, const tchar* application);
diff --git a/Engine/Source/Programs/UnrealBuildAccelerator/Core/Public/UbaProtocol.h b/Engine/Source/Programs/UnrealBuildAccelerator/Core/Public/UbaProtocol.h
--- a/Engine/Source/Programs/UnrealBuildAccelerator/Core/Public/UbaProtocol.h
+++ b/Engine/Source/Programs/UnrealBuildAccelerator/Core/Public/UbaProtocol.h
@@ -6,33 +6,39 @@
 
 namespace uba
 {
+	#define UBA_PROCESS_MESSAGES \
+		UBA_PROCESS_MESSAGE(Init) \
+		UBA_PROCESS_MESSAGE(CreateFile) \
+		UBA_PROCESS_MESSAGE(GetFullFileName) \
+		UBA_PROCESS_MESSAGE(CloseFile) \
+		UBA_PROCESS_MESSAGE(DeleteFile) \
+		UBA_PROCESS_MESSAGE(CopyFile) \
+		UBA_PROCESS_MESSAGE(MoveFile) \
+		UBA_PROCESS_MESSAGE(Chmod) \
+		UBA_PROCESS_MESSAGE(CreateDirectory) \
+		UBA_PROCESS_MESSAGE(ListDirectory) \
+		UBA_PROCESS_MESSAGE(UpdateTables) \
+		UBA_PROCESS_MESSAGE(CreateProcess) \
+		UBA_PROCESS_MESSAGE(StartProcess) \
+		UBA_PROCESS_MESSAGE(ExitChildProcess) \
+		UBA_PROCESS_MESSAGE(CreateTempFile) \
+		UBA_PROCESS_MESSAGE(OpenTempFile) \
+		UBA_PROCESS_MESSAGE(VirtualAllocFailed) \
+		UBA_PROCESS_MESSAGE(Log) \
+		UBA_PROCESS_MESSAGE(EchoOn) \
+		UBA_PROCESS_MESSAGE(InputDependencies) \
+		UBA_PROCESS_MESSAGE(Exit) \
+		UBA_PROCESS_MESSAGE(FlushWrittenFiles) \
+		UBA_PROCESS_MESSAGE(UpdateEnvironment) \
+		UBA_PROCESS_MESSAGE(GetNextProcess) \
+		UBA_PROCESS_MESSAGE(Custom) \
+
 	enum MessageType : u8
 	{
-		MessageType_Init = 1,
-		MessageType_CreateFileW,
-		MessageType_GetFullFileName,
-		MessageType_CloseFile,
-		MessageType_DeleteFileW,
-		MessageType_CopyFile,
-		MessageType_MoveFileW,
-		MessageType_Chmod,
-		MessageType_CreateDirectory,
-		MessageType_ListDirectory,
-		MessageType_UpdateTables,
-		MessageType_CreateProcess,
-		MessageType_StartProcess,
-		MessageType_ExitChildProcess, // Only used for Linux since we can't wait for grand children to exit
-		MessageType_CreateTempFile,
-		MessageType_OpenTempFile,
-		MessageType_VirtualAllocFailed,
-		MessageType_Log,
-		MessageType_EchoOn,
-		MessageType_InputDependencies,
-		MessageType_Exit,
-		MessageType_FlushWrittenFiles,
-		MessageType_UpdateEnvironment,
-		MessageType_GetNextProcess,
-		MessageType_Custom,
+		Unused = 0,
+		#define UBA_PROCESS_MESSAGE(type) MessageType_##type,
+		UBA_PROCESS_MESSAGES
+		#undef UBA_PROCESS_MESSAGE
 	};
 
 	inline constexpr u32 ProcessMessageVersion = 1339;
diff --git a/Engine/Source/Programs/UnrealBuildAccelerator/Detours/Private/PosixOS/UbaDetoursFunctionsPosix.cpp b/Engine/Source/Programs/UnrealBuildAccelerator/Detours/Private/PosixOS/UbaDetoursFunctionsPosix.cpp
--- a/Engine/Source/Programs/UnrealBuildAccelerator/Detours/Private/PosixOS/UbaDetoursFunctionsPosix.cpp
+++ b/Engine/Source/Programs/UnrealBuildAccelerator/Detours/Private/PosixOS/UbaDetoursFunctionsPosix.cpp
@@ -1254,7 +1254,7 @@ UBA_EXPORT int UBA_WRAPPER(rename)(const char* oldpath, const char* newpath)
 		TimerScope ts(g_stats.moveFile);
 		SCOPED_WRITE_LOCK(g_communicationLock, pcs);
 		BinaryWriter writer;
-		writer.WriteByte(MessageType_MoveFileW);
+		writer.WriteByte(MessageType_MoveFile);
 		writer.WriteStringKey(oldKey);
 		writer.WriteString(fixedOldPath);
 		writer.WriteStringKey(newKey);
@@ -1364,13 +1364,13 @@ UBA_EXPORT int UBA_WRAPPER(remove)(const char* pathname)
 
 	//if (!CanDetour(fixedName.data))
 	//{
-	//	DEBUG_LOG_TRUE(L"DeleteFileW", L"(%ls)", original);
+	//	DEBUG_LOG_TRUE(L"DeleteFile", L"(%ls)", original);
 	//	return TRUE_WRAPPER(remove)(pathname);
 	//}
 
 	//if (KeepInMemory(fixedName.data, fixedName.count))
 	//{
-	//	DEBUG_LOG_DETOURED(L"DeleteFileW", L"(INMEMORY) (%ls) -> Success", lpFileName);
+	//	DEBUG_LOG_DETOURED(L"DeleteFile", L"(INMEMORY) (%ls) -> Success", lpFileName);
 	//	SetLastError(ERROR_SUCCESS);
 	//	return TRUE;
 	//}
@@ -1385,7 +1385,7 @@ UBA_EXPORT int UBA_WRAPPER(remove)(const char* pathname)
 		TimerScope ts(g_stats.deleteFile);
 		SCOPED_WRITE_LOCK(g_communicationLock, pcs);
 		BinaryWriter writer;
-		writer.WriteByte(MessageType_DeleteFileW);
+		writer.WriteByte(MessageType_DeleteFile);
 		writer.WriteString(fixedName);
 		writer.WriteStringKey(fileNameKey);
 		writer.WriteU32(closeId);
@@ -1397,7 +1397,7 @@ UBA_EXPORT int UBA_WRAPPER(remove)(const char* pathname)
 		pcs.Leave();
 		DEBUG_LOG_PIPE(L"DeleteFile", L"%ls", lpFileName);
 	}
-	//DEBUG_LOG_DETOURED(L"DeleteFileW", L"(%ls) -> %ls", pathname, ToString(result));
+	//DEBUG_LOG_DETOURED(L"DeleteFile", L"(%ls) -> %ls", pathname, ToString(result));
 
 	g_directoryTable.ParseDirectoryTable(directoryTableSize);
 	g_mappedFileTable.SetDeleted(fileNameKey, fixedName.data, true);
diff --git a/Engine/Source/Programs/UnrealBuildAccelerator/Detours/Private/UbaDetoursFileMappingTable.cpp b/Engine/Source/Programs/UnrealBuildAccelerator/Detours/Private/UbaDetoursFileMappingTable.cpp
--- a/Engine/Source/Programs/UnrealBuildAccelerator/Detours/Private/UbaDetoursFileMappingTable.cpp
+++ b/Engine/Source/Programs/UnrealBuildAccelerator/Detours/Private/UbaDetoursFileMappingTable.cpp
@@ -66,7 +66,7 @@ namespace uba
 		TimerScope ts(g_stats.createFile);
 		SCOPED_WRITE_LOCK(g_communicationLock, pcs);
 		BinaryWriter writer;
-		writer.WriteByte(MessageType_CreateFileW);
+		writer.WriteByte(MessageType_CreateFile);
 		writer.WriteString(fileName);
 		writer.WriteStringKey(fileNameKey);
 		writer.WriteByte(access);
diff --git a/Engine/Source/Programs/UnrealBuildAccelerator/Detours/Private/Windows/UbaDetoursFunctionsKernelBase.inl b/Engine/Source/Programs/UnrealBuildAccelerator/Detours/Private/Windows/UbaDetoursFunctionsKernelBase.inl
--- a/Engine/Source/Programs/UnrealBuildAccelerator/Detours/Private/Windows/UbaDetoursFunctionsKernelBase.inl
+++ b/Engine/Source/Programs/UnrealBuildAccelerator/Detours/Private/Windows/UbaDetoursFunctionsKernelBase.inl
@@ -1270,7 +1270,7 @@ BOOL Detoured_DeleteFileW(LPCWSTR lpFileName)
 		TimerScope ts(g_stats.deleteFile);
 		SCOPED_WRITE_LOCK(g_communicationLock, pcs);
 		BinaryWriter writer;
-		writer.WriteByte(MessageType_DeleteFileW);
+		writer.WriteByte(MessageType_DeleteFile);
 		writer.WriteString(fixedName);
 		writer.WriteStringKey(fileNameKey);
 		writer.WriteU32(closeId);
@@ -1357,7 +1357,7 @@ bool Shared_MoveFile(LPCWSTR lpExistingFileName, LPCWSTR lpNewFileName, DWORD dw
 		TimerScope ts(g_stats.moveFile);
 		SCOPED_WRITE_LOCK(g_communicationLock, pcs);
 		BinaryWriter writer;
-		writer.WriteByte(MessageType_MoveFileW);
+		writer.WriteByte(MessageType_MoveFile);
 		writer.WriteStringKey(sourceKey);
 		writer.WriteString(source);
 		writer.WriteStringKey(destKey);
