[UBA]
* Changed so only non-windows platforms need 128kb message buffer... right now only macos creates those super long command lines

diff --git a/Engine/Source/Programs/UnrealBuildAccelerator/Core/Public/UbaProtocol.h b/Engine/Source/Programs/UnrealBuildAccelerator/Core/Public/UbaProtocol.h
--- a/Engine/Source/Programs/UnrealBuildAccelerator/Core/Public/UbaProtocol.h
+++ b/Engine/Source/Programs/UnrealBuildAccelerator/Core/Public/UbaProtocol.h
@@ -43,7 +43,7 @@ namespace uba
 
 	inline constexpr u32 ProcessMessageVersion = 1339;
 
-	inline constexpr u32 CommunicationMemSize = 64*1024*2;
+	inline constexpr u32 CommunicationMemSize = IsWindows ? 64*1024 : 64*1024*2; // Macos expands some commandlines to be crazy long
 
 	inline constexpr u32 FileMappingTableMemSize = 16 * 1024 * 1024;
 	inline constexpr u32 DirTableMemSize = 40 * 1024 * 1024;
