UBA: Don't increment the file list handle in NtQueryDirectoryFile if the next entry would overflow the buffer

#rnx

diff --git a/Engine/Source/Programs/UnrealBuildAccelerator/Detours/Private/Windows/UbaDetoursFunctionsNtDll.inl b/Engine/Source/Programs/UnrealBuildAccelerator/Detours/Private/Windows/UbaDetoursFunctionsNtDll.inl
--- a/Engine/Source/Programs/UnrealBuildAccelerator/Detours/Private/Windows/UbaDetoursFunctionsNtDll.inl
+++ b/Engine/Source/Programs/UnrealBuildAccelerator/Detours/Private/Windows/UbaDetoursFunctionsNtDll.inl
@@ -167,7 +167,7 @@ NTSTATUS NTAPI Detoured_NtQueryDirectoryFile(HANDLE FileHandle, HANDLE Event, PI
 			if (listHandle.it == listHandle.fileTableOffsets.size())
 				break;
 
-			u32 fileOffset = listHandle.fileTableOffsets[listHandle.it++];
+			u32 fileOffset = listHandle.fileTableOffsets[listHandle.it];
 
 			DirectoryTable::EntryInformation entryInfo;
 			wchar_t fileName[512];
@@ -226,6 +226,7 @@ NTSTATUS NTAPI Detoured_NtQueryDirectoryFile(HANDLE FileHandle, HANDLE Event, PI
 
 			prevInformation = it;
 			it = (u8*)fileNamePos + info.FileNameLength + 2;
+			listHandle.it++;
 
 			DEBUG_LOG_DETOURED(L"NtQueryDirectoryFile", L"%llu %ls", u64(FileHandle), fileNamePos);
 
