[UBA]
* Fixed bug in Detoured_NtQueryDirectoryFile when folder had deleted files and you could end up in infinite loops

diff --git a/Engine/Source/Programs/UnrealBuildAccelerator/Detours/Private/Windows/UbaDetoursFunctionsNtDll.inl b/Engine/Source/Programs/UnrealBuildAccelerator/Detours/Private/Windows/UbaDetoursFunctionsNtDll.inl
--- a/Engine/Source/Programs/UnrealBuildAccelerator/Detours/Private/Windows/UbaDetoursFunctionsNtDll.inl
+++ b/Engine/Source/Programs/UnrealBuildAccelerator/Detours/Private/Windows/UbaDetoursFunctionsNtDll.inl
@@ -167,7 +167,7 @@ NTSTATUS NTAPI Detoured_NtQueryDirectoryFile(HANDLE FileHandle, HANDLE Event, PI
 			if (listHandle.it == listHandle.fileTableOffsets.size())
 				break;
 
-			u32 fileOffset = listHandle.fileTableOffsets[listHandle.it];
+			u32 fileOffset = listHandle.fileTableOffsets[listHandle.it++];
 
 			DirectoryTable::EntryInformation entryInfo;
 			wchar_t fileName[512];
@@ -201,6 +201,7 @@ NTSTATUS NTAPI Detoured_NtQueryDirectoryFile(HANDLE FileHandle, HANDLE Event, PI
 			u8* writeEnd = (u8*)fileNamePos + fileNameBytes;
 			if (writeEnd > bufferEnd)
 			{
+				--listHandle.it;
 				if (!prevInformation)
 					res = STATUS_BUFFER_OVERFLOW;
 				break;
@@ -226,7 +227,6 @@ NTSTATUS NTAPI Detoured_NtQueryDirectoryFile(HANDLE FileHandle, HANDLE Event, PI
 
 			prevInformation = it;
 			it = (u8*)fileNamePos + info.FileNameLength + 2;
-			listHandle.it++;
 
 			DEBUG_LOG_DETOURED(L"NtQueryDirectoryFile", L"%llu %ls", u64(FileHandle), fileNamePos);
 
