[UBA]
* Added so server listens on wifi adapters

diff --git a/Engine/Source/Programs/UnrealBuildAccelerator/Common/Private/UbaNetworkBackendTcp.cpp b/Engine/Source/Programs/UnrealBuildAccelerator/Common/Private/UbaNetworkBackendTcp.cpp
--- a/Engine/Source/Programs/UnrealBuildAccelerator/Common/Private/UbaNetworkBackendTcp.cpp
+++ b/Engine/Source/Programs/UnrealBuildAccelerator/Common/Private/UbaNetworkBackendTcp.cpp
@@ -12,7 +12,8 @@
 #endif
 
 #if PLATFORM_WINDOWS
-#include "iphlpapi.h"
+#include <iphlpapi.h>
+#include <ipifcons.h>
 #pragma comment (lib, "Netapi32.lib")
 #pragma comment (lib, "Ws2_32.lib")
 #pragma comment(lib, "IPHLPAPI.lib") // For GetAdaptersInfo
@@ -810,7 +811,7 @@ namespace uba
 		}
 		for (IP_ADAPTER_INFO* it = info; it; it = it->Next)
 		{
-			if (it->Type != MIB_IF_TYPE_ETHERNET)
+			if (it->Type != MIB_IF_TYPE_ETHERNET && it->Type != IF_TYPE_IEEE80211)
 				continue;
 			for (IP_ADDR_STRING* s = &it->IpAddressList; s; s = s->Next)
 			{
