[UBA]
* Fixed bug where pipes between processes did not work properly on wine
* Added unit tests for pipes (win only)
* Improved error message when failing to open file
* Fixed so .obj is seen as output files from ispc.exe
* Fixed up application path when creating child processes

diff --git a/Engine/Source/Programs/UnrealBuildAccelerator/Common/Private/UbaFile.cpp b/Engine/Source/Programs/UnrealBuildAccelerator/Common/Private/UbaFile.cpp
--- a/Engine/Source/Programs/UnrealBuildAccelerator/Common/Private/UbaFile.cpp
+++ b/Engine/Source/Programs/UnrealBuildAccelerator/Common/Private/UbaFile.cpp
@@ -96,7 +96,7 @@ namespace uba
 			DWORD wasRead = 0;
 			if (!::ReadFile(asHANDLE(fileHandle), buffer, toRead, &wasRead, NULL))
 				if (GetLastError() != ERROR_IO_PENDING)
-					return logger.Error(TC("ERROR reading file %s (error: %s)"), fileName, LastErrorToText().data);
+					return logger.Error(TC("ERROR reading %llu bytes from file %s (error: %s)"), toRead, fileName, LastErrorToText().data);
 #else
 			ssize_t wasRead = read(asFileDescriptor(fileHandle), buffer, toRead);
 			if (wasRead == -1)
diff --git a/Engine/Source/Programs/UnrealBuildAccelerator/Common/Private/UbaProcess.cpp b/Engine/Source/Programs/UnrealBuildAccelerator/Common/Private/UbaProcess.cpp
--- a/Engine/Source/Programs/UnrealBuildAccelerator/Common/Private/UbaProcess.cpp
+++ b/Engine/Source/Programs/UnrealBuildAccelerator/Common/Private/UbaProcess.cpp
@@ -727,8 +727,7 @@ namespace uba
 		static bool subreaper = []() { prctl(PR_SET_CHILD_SUBREAPER, 1, 0, 0, 0); return true; }();
 		#endif
 
-		StringBuffer<> application;
-		reader.ReadString(application);
+		TString applicationStr = reader.ReadString();
 		StringBuffer<64*1024*2> fullCommandLine;
 		reader.ReadString(fullCommandLine);
 		StringBuffer<> currentDir;
@@ -746,8 +745,8 @@ namespace uba
 				second = fullCommandLine.data + fullCommandLine.count;
 			//UBA_ASSERTF(second, TC("Missing second '\"' in command line: %s"), fullCommandLine.data); // "Unsupported cmd line format"
 			commandLine = second + 1;
-			if (application.IsEmpty())
-				application.Append(fullCommandLine.data + 1, u64(second - fullCommandLine.data - 1));
+			if (applicationStr.empty())
+				applicationStr.assign(fullCommandLine.data + 1, u64(second - fullCommandLine.data - 1));
 		}
 		else
 		{
@@ -756,10 +755,13 @@ namespace uba
 				commandLine = TC("");
 			else
 				commandLine = secondParamStart + 1;
-			if (application.IsEmpty())
-				application.Append(fullCommandLine.data, u64(secondParamStart - fullCommandLine.data));
+			if (applicationStr.empty())
+				applicationStr.assign(fullCommandLine.data, u64(secondParamStart - fullCommandLine.data));
 		}
 
+		StringBuffer<512> application;
+		FixPath(applicationStr.c_str(), nullptr, 0, application);
+
 		while (*commandLine == ' ')
 			++commandLine;
 
diff --git a/Engine/Source/Programs/UnrealBuildAccelerator/Core/Private/UbaApplicationRules.cpp b/Engine/Source/Programs/UnrealBuildAccelerator/Core/Private/UbaApplicationRules.cpp
--- a/Engine/Source/Programs/UnrealBuildAccelerator/Core/Private/UbaApplicationRules.cpp
+++ b/Engine/Source/Programs/UnrealBuildAccelerator/Core/Private/UbaApplicationRules.cpp
@@ -320,7 +320,8 @@ namespace uba
 		{
 			return EndsWith(file, fileLen, TC(".dummy.h"))
 				|| EndsWith(file, fileLen, TC(".ispc.bc"))
-				|| EndsWith(file, fileLen, TC(".ispc.txt"));
+				|| EndsWith(file, fileLen, TC(".ispc.txt"))
+				|| EndsWith(file, fileLen, TC(".obj"));
 		}
 	};
 
diff --git a/Engine/Source/Programs/UnrealBuildAccelerator/Detours/Private/Windows/UbaDetoursFunctionsKernelBase.inl b/Engine/Source/Programs/UnrealBuildAccelerator/Detours/Private/Windows/UbaDetoursFunctionsKernelBase.inl
--- a/Engine/Source/Programs/UnrealBuildAccelerator/Detours/Private/Windows/UbaDetoursFunctionsKernelBase.inl
+++ b/Engine/Source/Programs/UnrealBuildAccelerator/Detours/Private/Windows/UbaDetoursFunctionsKernelBase.inl
@@ -324,7 +324,7 @@ BOOL Detoured_WriteConsoleA(HANDLE hConsoleOutput, const VOID* lpBuffer, DWORD n
 BOOL Detoured_WriteConsoleW(HANDLE hConsoleOutput, const VOID* lpBuffer, DWORD nNumberOfCharsToWrite, LPDWORD lpNumberOfCharsWritten, LPVOID lpReserved)
 {
 	DETOURED_CALL(WriteConsoleW);
-	//DEBUG_LOG_DETOURED(L"WriteConsoleW"", L""); // Too much spam
+	//DEBUG_LOG_DETOURED(L"WriteConsoleW", L""); // Too much spam
 	Shared_WriteConsole((const wchar_t*)lpBuffer, nNumberOfCharsToWrite, false);
 	if (lpNumberOfCharsWritten)
 		*lpNumberOfCharsWritten = nNumberOfCharsToWrite;
@@ -565,8 +565,11 @@ BOOL Detoured_WriteFile(HANDLE hFile, LPCVOID lpBuffer, DWORD nNumberOfBytesToWr
 	else if (hFile == g_stdHandle[1] || hFile == g_stdHandle[0])
 	{
 		WriteStdFile(lpBuffer, nNumberOfBytesToWrite, hFile == g_stdHandle[0]);
+		*lpNumberOfBytesWritten = nNumberOfBytesToWrite;
 		SetLastError(ERROR_SUCCESS);
 		return true;
+		//if (GetFileType(trueHandle) != FILE_TYPE_CHAR )
+		//return True_WriteFile(trueHandle, lpBuffer, nNumberOfBytesToWrite, lpNumberOfBytesWritten, lpOverlapped);
 	}
 
 	BOOL res = True_WriteFile(trueHandle, lpBuffer, nNumberOfBytesToWrite, lpNumberOfBytesWritten, lpOverlapped);
@@ -2537,6 +2540,10 @@ BOOL Detoured_CreateProcessW(LPCWSTR lpApplicationName, LPWSTR lpCommandLine, LP
 		}
 	}
 
+	StringBuffer<> application;
+	if (lpApplicationName)
+		FixPath(application, lpApplicationName);
+
 	TString commandLine;
 	TString currentDir;
 	u32 processId = 0;
@@ -2546,7 +2553,7 @@ BOOL Detoured_CreateProcessW(LPCWSTR lpApplicationName, LPWSTR lpCommandLine, LP
 		SCOPED_WRITE_LOCK(g_communicationLock, pcs);
 		BinaryWriter writer;
 		writer.WriteByte(MessageType_CreateProcess);
-		writer.WriteString(lpApplicationName ? lpApplicationName : L"");
+		writer.WriteString(application.data);
 		writer.WriteString(lpCommandLine ? lpCommandLine : L"");
 		writer.WriteString(lpCurrentDirectory ? lpCurrentDirectory : g_virtualWorkingDir.data);
 		writer.Flush();
@@ -2562,7 +2569,7 @@ BOOL Detoured_CreateProcessW(LPCWSTR lpApplicationName, LPWSTR lpCommandLine, LP
 
 		currentDir = reader.ReadString();
 		commandLine = reader.ReadString();
-		DEBUG_LOG_PIPE(L"CreateProcess", L"%ls %ls", lpApplicationName, lpCommandLine ? lpCommandLine : L"");
+		DEBUG_LOG_PIPE(L"CreateProcess", L"%ls %ls", application.data, lpCommandLine ? lpCommandLine : L"");
 	}
 
 	LPCSTR dlls[] = { dll };
@@ -2634,7 +2641,7 @@ BOOL Detoured_CreateProcessW(LPCWSTR lpApplicationName, LPWSTR lpCommandLine, LP
 	detouredHandle->trueHandle = trueHandle;
 	lpProcessInformation->hProcess = makeDetouredHandle(detouredHandle);
 
-	DEBUG_LOG_DETOURED(L"CreateProcessW", L"%llu", u64(lpProcessInformation->hProcess));
+	DEBUG_LOG_DETOURED(L"CreateProcessW", L"%llu (0x%llx)", lpProcessInformation->hProcess, trueHandle);
 	return TRUE;
 }
 
@@ -2715,14 +2722,15 @@ DWORD Detoured_WaitForSingleObjectEx(HANDLE hHandle, DWORD dwMilliseconds, BOOL
 {
 	DETOURED_CALL(WaitForSingleObjectEx);
 	bool isProcess = false;
+	HANDLE trueHandle = hHandle;
 	if (isDetouredHandle(hHandle))
 	{
 		DetouredHandle& dh = asDetouredHandle(hHandle);
-		hHandle = asDetouredHandle(hHandle).trueHandle;
+		trueHandle = asDetouredHandle(hHandle).trueHandle;
 		isProcess = dh.type == HandleType_Process;
 	}
 
-	auto res = True_WaitForSingleObjectEx(hHandle, dwMilliseconds, bAlertable);
+	auto res = True_WaitForSingleObjectEx(trueHandle, dwMilliseconds, bAlertable);
 
 	if (res != WAIT_OBJECT_0 || !isProcess)
 		return res;
@@ -2732,8 +2740,8 @@ DWORD Detoured_WaitForSingleObjectEx(HANDLE hHandle, DWORD dwMilliseconds, BOOL
 	{
 		auto lastError = GetLastError();
 		DWORD exitCode;
-		True_GetExitCodeProcess(hHandle, &exitCode);
-		DEBUG_LOG_DETOURED(L"WaitForSingleObjectEx", L"for process %llu. Exit code: %u", uintptr_t(hHandle), exitCode);
+		True_GetExitCodeProcess(trueHandle, &exitCode);
+		DEBUG_LOG_DETOURED(L"WaitForSingleObjectEx", L"for process %llu (0x%llx). Exit code: %u", hHandle, trueHandle, exitCode);
 		SetLastError(lastError);
 	}
 #endif
@@ -3315,6 +3323,15 @@ BOOL Detoured_CreatePipe(PHANDLE hReadPipe, PHANDLE hWritePipe, LPSECURITY_ATTRI
 	return True_CreatePipe(hReadPipe, hWritePipe, lpPipeAttributes, nSize);
 }
 
+
+BOOL Detoured_SetHandleInformation(HANDLE hObject, DWORD dwMask, DWORD dwFlags)
+{
+	DETOURED_CALL(SetHandleInformation);
+	DEBUG_LOG_TRUE(L"SetHandleInformation", L"%llu", uintptr_t(hObject));
+	UBA_ASSERT(!isDetouredHandle(hObject));
+	return True_SetHandleInformation(hObject, dwMask, dwFlags);
+}
+
 HANDLE Detoured_CreateNamedPipeW(LPCWSTR lpName, DWORD dwOpenMode, DWORD dwPipeMode, DWORD nMaxInstances, DWORD nOutBufferSize, DWORD nInBufferSize, DWORD nDefaultTimeOut, LPSECURITY_ATTRIBUTES lpSecurityAttributes)
 {
 	DETOURED_CALL(CreateNamedPipeW);
diff --git a/Engine/Source/Programs/UnrealBuildAccelerator/Detours/Private/Windows/UbaDetoursFunctionsNtDll.inl b/Engine/Source/Programs/UnrealBuildAccelerator/Detours/Private/Windows/UbaDetoursFunctionsNtDll.inl
--- a/Engine/Source/Programs/UnrealBuildAccelerator/Detours/Private/Windows/UbaDetoursFunctionsNtDll.inl
+++ b/Engine/Source/Programs/UnrealBuildAccelerator/Detours/Private/Windows/UbaDetoursFunctionsNtDll.inl
@@ -524,7 +524,9 @@ NTSTATUS NTAPI Shared_NtCreateFile(bool IsCreateFunc, PHANDLE hFileHandle, ACCES
 		}
 		else if (memcmp(buf, L"\\??\\", 8) == 0)
 		{
-			if (!FixPath(fileName, buf + 4))
+			if (g_isRunningWine && StartsWith(buf+4, L"pipe\\")) // This is where CreatePipe on wine ends up
+				suppressCreateFileDetour = true;
+			else if (!FixPath(fileName, buf + 4))
 				UBA_ASSERTF(false, L"FixPath failed for string '%ls'", buf + 4);
 		}
 		else
diff --git a/Engine/Source/Programs/UnrealBuildAccelerator/Detours/Private/Windows/UbaDetoursFunctionsWin.cpp b/Engine/Source/Programs/UnrealBuildAccelerator/Detours/Private/Windows/UbaDetoursFunctionsWin.cpp
--- a/Engine/Source/Programs/UnrealBuildAccelerator/Detours/Private/Windows/UbaDetoursFunctionsWin.cpp
+++ b/Engine/Source/Programs/UnrealBuildAccelerator/Detours/Private/Windows/UbaDetoursFunctionsWin.cpp
@@ -535,12 +535,31 @@ void CloseCaches()
 
 bool g_exitMessageSent;
 
+void SendExitMessage(DWORD exitCode, u64 startTime);
+void OnModuleLoaded(HMODULE moduleHandle, const wchar_t* name);
+
+// Variables used to communicate state from kernelbase functions to ntdll functions
+thread_local const wchar_t* t_renameFileNewName;
+thread_local const wchar_t* t_createFileFileName;
+
+#include "UbaDetoursFunctionsMiMalloc.inl"
+#include "UbaDetoursFunctionsNtDll.inl"
+#include "UbaDetoursFunctionsKernelBase.inl"
+#include "UbaDetoursFunctionsUcrtBase.inl"
+#include "UbaDetoursFunctionsImagehlp.inl"
+#include "UbaDetoursFunctionsDbgHelp.inl"
+
+extern u32 g_consoleStringIndex;
+
 void SendExitMessage(DWORD exitCode, u64 startTime)
 {
 	if (g_exitMessageSent)
 		return;
 	g_exitMessageSent = true;
 
+	if (g_consoleStringIndex)
+		Shared_WriteConsole(L"\n", 1, 0);
+
 	if (g_trackInputsMem)
 	{
 		u32 left = g_trackInputsBufPos;
@@ -580,19 +599,6 @@ void SendExitMessage(DWORD exitCode, u64 startTime)
 	writer.Flush(g_isChild);
 }
 
-void OnModuleLoaded(HMODULE moduleHandle, const wchar_t* name);
-
-// Variables used to communicate state from kernelbase functions to ntdll functions
-thread_local const wchar_t* t_renameFileNewName;
-thread_local const wchar_t* t_createFileFileName;
-
-#include "UbaDetoursFunctionsMiMalloc.inl"
-#include "UbaDetoursFunctionsNtDll.inl"
-#include "UbaDetoursFunctionsKernelBase.inl"
-#include "UbaDetoursFunctionsUcrtBase.inl"
-#include "UbaDetoursFunctionsImagehlp.inl"
-#include "UbaDetoursFunctionsDbgHelp.inl"
-
 void DetourAttachFunction(void** trueFunc, void* detouredFunc, const char* funcName)
 {
 	if (!*trueFunc)
@@ -814,7 +820,7 @@ void PreInit(const DetoursPayload& payload)
 			FatalError(1349, L"Failed to reserve memory for cl.exe (%u)", GetLastError());
 	}
 
-	// Special link.exe handling.. it seems loading bcrypt.dll can deadlock when using mimalloc so we make sure ti load it here directly instead
+	// Special link.exe handling.. it seems loading bcrypt.dll can deadlock when using mimalloc so we make sure to load it here directly instead
 	// There is a setting to disable bcrypt dll loading inside mimalloc but with that change mimalloc does not work with older versions of windows
 	if (payload.rulesIndex == 2)
 	{
diff --git a/Engine/Source/Programs/UnrealBuildAccelerator/Detours/Private/Windows/UbaDetoursFunctionsWin.h b/Engine/Source/Programs/UnrealBuildAccelerator/Detours/Private/Windows/UbaDetoursFunctionsWin.h
--- a/Engine/Source/Programs/UnrealBuildAccelerator/Detours/Private/Windows/UbaDetoursFunctionsWin.h
+++ b/Engine/Source/Programs/UnrealBuildAccelerator/Detours/Private/Windows/UbaDetoursFunctionsWin.h
@@ -251,6 +251,7 @@
 	DETOURED_FUNCTION(CreateWaitableTimerExW) \
 	DETOURED_FUNCTION(CreateIoCompletionPort) \
 	DETOURED_FUNCTION(CreatePipe) \
+	DETOURED_FUNCTION(SetHandleInformation) \
 	DETOURED_FUNCTION(CreateNamedPipeW) \
 	DETOURED_FUNCTION(PeekNamedPipe) \
 	DETOURED_FUNCTION(GetKernelObjectSecurity) \
diff --git a/Engine/Source/Programs/UnrealBuildAccelerator/Test/Private/UbaTestAll.h b/Engine/Source/Programs/UnrealBuildAccelerator/Test/Private/UbaTestAll.h
--- a/Engine/Source/Programs/UnrealBuildAccelerator/Test/Private/UbaTestAll.h
+++ b/Engine/Source/Programs/UnrealBuildAccelerator/Test/Private/UbaTestAll.h
@@ -8,6 +8,7 @@
 #include "UbaTestSession.h"
 #include "UbaTestStorage.h"
 #include "UbaTestStress.h"
+#include "UbaTestStdOut.h"
 
 namespace uba
 {
@@ -22,6 +23,7 @@ namespace uba
 		UBA_TEST(TestLocalScheduleReuse) \
 		UBA_TEST(TestDetouredTouch) \
 		UBA_TEST(TestRemoteScheduleReuse) \
+		UBA_TEST(TestStdOutLocal) \
 
 #endif 
 
@@ -81,6 +83,7 @@ namespace uba
 			filter = argv[1];
 
 		//UBA_TEST(TestStress) // This can not be submitted.. it depends on CoordinatorHorde and credentials
+		//UBA_TEST(TestStdOutRemote) // This can not be submitted.. depends on a running UbaAgent
 		UBA_TESTS
 
 		logger.Info(TC("Tests finished successfully!"));
diff --git a/Engine/Source/Programs/UnrealBuildAccelerator/Test/Private/UbaTestStdOut.h b/Engine/Source/Programs/UnrealBuildAccelerator/Test/Private/UbaTestStdOut.h
new file mode 100644
--- /dev/null
+++ b/Engine/Source/Programs/UnrealBuildAccelerator/Test/Private/UbaTestStdOut.h
@@ -0,0 +1,87 @@
+// Copyright Epic Games, Inc. All Rights Reserved.
+
+#pragma once
+
+#include "UbaSessionServer.h"
+#include "UbaStorageServer.h"
+
+namespace uba
+{
+	bool TestStdOut(LoggerWithWriter& logger, const StringBufferBase& testRootDir, bool remote)
+	{
+		LogWriter& logWriter = logger.m_writer;
+		NetworkBackendTcp networkBackend(logWriter);
+
+		bool ctorSuccess = true;
+		NetworkServer networkServer(ctorSuccess, { logWriter });
+
+		StringBuffer<> rootDir;
+		rootDir.Append(testRootDir).Append(TC("Uba"));
+		if (!DeleteAllFiles(logger, rootDir.data))
+			return false;
+
+		StorageServerCreateInfo storageServerInfo(networkServer, rootDir.data, logWriter);
+		storageServerInfo.casCapacityBytes = 1024ull * 1024 * 1024;
+		StorageServer storageServer(storageServerInfo);
+
+		SessionServerCreateInfo sessionServerInfo(storageServer, networkServer, logWriter);
+		sessionServerInfo.checkMemory = false;
+		sessionServerInfo.rootDir = rootDir.data;
+		sessionServerInfo.traceEnabled = true;
+		//sessionServerInfo.remoteLogEnabled = true;
+		SessionServer sessionServer(sessionServerInfo);
+
+		auto sg = MakeGuard([&]() { networkServer.DisconnectClients(); });
+
+		StringBuffer<> workingDir;
+		workingDir.Append(testRootDir).Append(TC("WorkingDir"));
+		if (!DeleteAllFiles(logger, workingDir.data))
+			return false;
+		if (!storageServer.CreateDirectory(workingDir.data))
+			return false;
+		if (!DeleteAllFiles(logger, workingDir.data, false))
+			return false;
+		workingDir.EnsureEndsWithSlash();
+
+		if (!networkServer.StartListen(networkBackend))
+			return logger.Error(TC("Failed to listen"));
+
+		StringBuffer<> testApp;
+		GetTestAppPath(logger, testApp);
+
+		ProcessStartInfo pi;
+		pi.application = testApp.data;
+		pi.arguments = TC("-stdout=rootprocess");
+		pi.workingDir = workingDir.data;
+		pi.description = TC("StdOutDesc");
+		pi.logFile = TC("Log");
+
+		auto ph = remote ? sessionServer.RunProcessRemote(pi) : sessionServer.RunProcess(pi);
+		if (!ph.WaitForExit(5000))
+			return logger.Error(TC("Timed out waiting for process"));
+		if (auto ec = ph.GetExitCode())
+			return logger.Error(TC("Process exited with error code %u"), ec);
+
+		networkBackend.StopListen();
+		networkServer.DisconnectClients();
+		sessionServer.WaitOnAllTasks();
+
+		auto& logLines = ph.GetLogLines();
+		if (logLines.size() != 1)
+			return false;
+		if (!Equals(logLines[0].text.c_str(), TC("rootprocess")))
+			return false;
+
+		return true;
+	}
+
+	bool TestStdOutLocal(LoggerWithWriter& logger, const StringBufferBase& testRootDir)
+	{
+		return TestStdOut(logger, testRootDir, false);
+	}
+
+	bool TestStdOutRemote(LoggerWithWriter& logger, const StringBufferBase& testRootDir)
+	{
+		return TestStdOut(logger, testRootDir, true);
+	}
+}
diff --git a/Engine/Source/Programs/UnrealBuildAccelerator/TestApp/Private/Windows/UbaTestAppWin.cpp b/Engine/Source/Programs/UnrealBuildAccelerator/TestApp/Private/Windows/UbaTestAppWin.cpp
--- a/Engine/Source/Programs/UnrealBuildAccelerator/TestApp/Private/Windows/UbaTestAppWin.cpp
+++ b/Engine/Source/Programs/UnrealBuildAccelerator/TestApp/Private/Windows/UbaTestAppWin.cpp
@@ -181,6 +181,60 @@ int wmain(int argc, wchar_t* argv[])
 
 		return 0;
 	}
+	else if (wcsncmp(argv[1], L"-stdout=", 8) == 0)
+	{
+		const wchar_t* str = argv[1] + 8;
+		if (wcscmp(str, L"rootprocess") == 0)
+		{
+			STARTUPINFOW si;
+			memset(&si, 0, sizeof(si));
+			PROCESS_INFORMATION pi;
+			memset(&pi, 0, sizeof(pi));
+			wchar_t arg[1024];
+			wcscpy_s(arg, 1024, argv[0]);
+			wcscat_s(arg, 1024, L" -stdout=childprocess");
+			//wcscpy_s(arg, 1024, L"\"c:\\sdk\\AutoSDK/HostWin64/Win64/MetalDeveloperTools/4.1/metal/macos/bin/metal.exe\" -v --target=air64-apple-darwin18.7.0 16384");
+
+			SECURITY_ATTRIBUTES saAttr;
+			saAttr.nLength = sizeof(SECURITY_ATTRIBUTES);
+			saAttr.bInheritHandle = TRUE;
+			saAttr.lpSecurityDescriptor = NULL;
+			HANDLE readPipe;
+			HANDLE writePipe;
+			if (!CreatePipe(&readPipe, &writePipe, &saAttr, 0))
+				return 1;
+
+			if (!SetHandleInformation(readPipe, HANDLE_FLAG_INHERIT, 0))
+				return 2;
+
+			si.dwFlags = STARTF_USESTDHANDLES;
+			si.hStdInput = GetStdHandle(STD_INPUT_HANDLE);
+			si.hStdOutput = writePipe;
+			si.hStdError = writePipe;
+
+			DWORD flags = 0;//CREATE_NO_WINDOW;
+			if (!CreateProcessW(nullptr, arg, nullptr, nullptr, TRUE, flags, nullptr, nullptr, &si, &pi))
+				return 3;
+			CloseHandle(pi.hThread);
+			CloseHandle(writePipe);
+
+			char buf[4096] = { 0 };
+			DWORD readCount = 0;
+			if (!::ReadFile(readPipe, buf, sizeof(buf) - 1, &readCount, NULL))
+			{
+				LogError(L"Failed to read pipe %u %u", GetLastError(), readCount);
+				return 4;
+			}
+			buf[readCount] = 0;
+			if (strncmp(buf, "childprocess", 12) != 0)
+				return 5;
+
+			if (WaitForSingleObject(pi.hProcess, INFINITE) != WAIT_OBJECT_0)
+				return 6;
+			CloseHandle(pi.hProcess);
+		}
+		wprintf(L"%s\n", str);
+	}
 	else
 	{
 		using u32 = unsigned int;
