[UBA]
* Removed TestStdOutLocal on linux since it fails and the code path in TestAppPosix is not implemented

diff --git a/Engine/Source/Programs/UnrealBuildAccelerator/Test/Private/UbaTestStdOut.h b/Engine/Source/Programs/UnrealBuildAccelerator/Test/Private/UbaTestStdOut.h
--- a/Engine/Source/Programs/UnrealBuildAccelerator/Test/Private/UbaTestStdOut.h
+++ b/Engine/Source/Programs/UnrealBuildAccelerator/Test/Private/UbaTestStdOut.h
@@ -77,6 +77,8 @@ namespace uba
 
 	bool TestStdOutLocal(LoggerWithWriter& logger, const StringBufferBase& testRootDir)
 	{
+		if (!IsWindows)
+			return true;
 		return TestStdOut(logger, testRootDir, false);
 	}
 
