[UBA]
* Fixed stdout logging when process is executed through cmd.exe
* Added support for more detoured code paths to handle python and nodejs execution

diff --git a/Engine/Source/Programs/UnrealBuildAccelerator/Core/Public/UbaPathUtils.h b/Engine/Source/Programs/UnrealBuildAccelerator/Core/Public/UbaPathUtils.h
--- a/Engine/Source/Programs/UnrealBuildAccelerator/Core/Public/UbaPathUtils.h
+++ b/Engine/Source/Programs/UnrealBuildAccelerator/Core/Public/UbaPathUtils.h
@@ -300,6 +300,8 @@ namespace uba
 
 	inline void FixPath(const tchar* fileName, const tchar* workingDir, u64 workingDirCharLen, StringBufferBase& buffer)
 	{
-		FixPath2(fileName, workingDir, workingDirCharLen, buffer.data, buffer.capacity, &buffer.count);
+		u32 count = buffer.count;
+		FixPath2(fileName, workingDir, workingDirCharLen, buffer.data + count, buffer.capacity - count, &buffer.count);
+		buffer.count += count;
 	}
 }
\ No newline at end of file
diff --git a/Engine/Source/Programs/UnrealBuildAccelerator/Detours/Private/Windows/UbaDetoursFunctionsKernelBase.inl b/Engine/Source/Programs/UnrealBuildAccelerator/Detours/Private/Windows/UbaDetoursFunctionsKernelBase.inl
--- a/Engine/Source/Programs/UnrealBuildAccelerator/Detours/Private/Windows/UbaDetoursFunctionsKernelBase.inl
+++ b/Engine/Source/Programs/UnrealBuildAccelerator/Detours/Private/Windows/UbaDetoursFunctionsKernelBase.inl
@@ -2222,12 +2222,16 @@ DWORD Detoured_GetFinalPathNameByHandleW(HANDLE hFile, LPTSTR lpszFilePath, DWOR
 		UBA_ASSERT(fo && fo->fileInfo->originalName);
 		const wchar_t* fileName = fo->fileInfo->originalName;
 
-		if (dwFlags == 0)
+		if (dwFlags == 0 || dwFlags == 2)
 		{
 			if (!fo->newName.empty())
 				fileName = fo->newName.c_str();
 			StringBuffer<> buffer;
-			FixPath2(fileName, g_virtualWorkingDir.data, g_virtualWorkingDir.count, buffer.data, buffer.capacity, &buffer.count);
+
+			if (dwFlags == 2)
+				buffer.Append(L"\\??\\");
+
+			FixPath(fileName, g_virtualWorkingDir.data, g_virtualWorkingDir.count, buffer);
 
 			if (cchFilePath <= buffer.count)
 			{
@@ -2239,6 +2243,8 @@ DWORD Detoured_GetFinalPathNameByHandleW(HANDLE hFile, LPTSTR lpszFilePath, DWOR
 			// Unfortunately casing can be wrong here.. and we need to fix that. Let's use the directory table for that
 			// Note, this really only matters when building linux target from windows.. then there is path validation that errors if this is not properly fixed
 			StringBuffer<> buffer2;
+			if (dwFlags == 2)
+				buffer2.Append(L"\\??\\");
 			g_directoryTable.GetFinalPath(buffer2, fileName);
 			UBA_ASSERT(buffer2.count == buffer.count);
 
@@ -2491,10 +2497,12 @@ BOOL Detoured_GetConsoleMode(HANDLE hConsoleHandle, LPDWORD lpMode)
 BOOL Detoured_SetConsoleMode(HANDLE hConsoleHandle, DWORD mode)
 {
 	DETOURED_CALL(SetConsoleMode);
-	DEBUG_LOG_DETOURED(L"SetConsoleMode", L"(%u)", mode);
+	DEBUG_LOG_DETOURED(L"SetConsoleMode", L"%llu (%u)", u64(hConsoleHandle), mode);
 
-	g_echoOn = (mode & ~503) != 0; // TODO: This might be wrong. Trying to figure out how echo off in batch files work in terms of win32 calls
+	if (hConsoleHandle == g_stdHandle[1])
 	{
+		g_echoOn = (mode & ~503) != 0; // TODO: This might be wrong. Trying to figure out how echo off in batch files work in terms of win32 calls
+
 		SCOPED_WRITE_LOCK(g_communicationLock, pcs);
 		BinaryWriter writer;
 		writer.WriteByte(MessageType_EchoOn);
@@ -2516,7 +2524,7 @@ BOOL Detoured_CreateProcessW(LPCWSTR lpApplicationName, LPWSTR lpCommandLine, LP
 	DWORD dwCreationFlags, LPVOID lpEnvironment, LPCWSTR lpCurrentDirectory, LPSTARTUPINFOW lpStartupInfo, LPPROCESS_INFORMATION lpProcessInformation)
 {
 	DETOURED_CALL(CreateProcessW);
-	DEBUG_LOG_DETOURED(L"CreateProcessW", L"%ls %ls %u %u %llu", lpApplicationName, lpCommandLine ? lpCommandLine : L"", dwCreationFlags, lpStartupInfo->dwFlags, u64(lpStartupInfo->hStdInput));
+	DEBUG_LOG_DETOURED(L"CreateProcessW", L"%ls %ls CreationFlags: 0x%x StartupFlags: 0x%u Stdin: %llu", lpApplicationName, lpCommandLine ? lpCommandLine : L"", dwCreationFlags, lpStartupInfo->dwFlags, u64(lpStartupInfo->hStdInput));
 
 	if ((!lpApplicationName || !*lpApplicationName) && (!lpCommandLine || !*lpCommandLine))
 	{
@@ -3328,8 +3336,7 @@ BOOL Detoured_SetHandleInformation(HANDLE hObject, DWORD dwMask, DWORD dwFlags)
 {
 	DETOURED_CALL(SetHandleInformation);
 	DEBUG_LOG_TRUE(L"SetHandleInformation", L"%llu", uintptr_t(hObject));
-	UBA_ASSERT(!isDetouredHandle(hObject));
-	return True_SetHandleInformation(hObject, dwMask, dwFlags);
+	return True_SetHandleInformation(hObject, dwMask, dwFlags); // Calls NtQueryObject and NtSetInformationObject internally
 }
 
 HANDLE Detoured_CreateNamedPipeW(LPCWSTR lpName, DWORD dwOpenMode, DWORD dwPipeMode, DWORD nMaxInstances, DWORD nOutBufferSize, DWORD nInBufferSize, DWORD nDefaultTimeOut, LPSECURITY_ATTRIBUTES lpSecurityAttributes)
diff --git a/Engine/Source/Programs/UnrealBuildAccelerator/Detours/Private/Windows/UbaDetoursFunctionsNtDll.inl b/Engine/Source/Programs/UnrealBuildAccelerator/Detours/Private/Windows/UbaDetoursFunctionsNtDll.inl
--- a/Engine/Source/Programs/UnrealBuildAccelerator/Detours/Private/Windows/UbaDetoursFunctionsNtDll.inl
+++ b/Engine/Source/Programs/UnrealBuildAccelerator/Detours/Private/Windows/UbaDetoursFunctionsNtDll.inl
@@ -34,7 +34,26 @@ NTSTATUS Detoured_NtQueryVolumeInformationFile(HANDLE FileHandle, PIO_STATUS_BLO
 			}
 		}
 		TrueHandle = dh.trueHandle;
-		UBA_ASSERTF(TrueHandle != INVALID_HANDLE_VALUE, L"NtQueryVolumeInformationFile using class %u not handled (%ls)", FsInformationClass, HandleToName(FileHandle));
+		if (TrueHandle == INVALID_HANDLE_VALUE)
+		{
+			if (FsInformationClass == 1) // FileFsVolumeInformation
+			{
+				// TODO This code path is here to handle nodejs queries..
+				
+				auto& info = *(FILE_FS_VOLUME_INFORMATION*)FsInformation;
+				UBA_ASSERT(dh.dirTableOffset != ~0u);
+				DirectoryTable::EntryInformation entryInfo;
+				g_directoryTable.GetEntryInformation(entryInfo, dh.dirTableOffset);
+				UBA_ASSERT(entryInfo.attributes != 0);
+				info.VolumeCreationTime.QuadPart = 0;
+				info.VolumeSerialNumber = entryInfo.volumeSerial;
+				info.VolumeLabelLength = 0;
+				info.SupportsObjects = false;
+				info.VolumeLabel[0] = 0;
+				return STATUS_SUCCESS;
+			}
+			UBA_ASSERTF(false, L"NtQueryVolumeInformationFile using class %u not handled %ls (%ls)", FsInformationClass, dh.fileObject->fileInfo->name, dh.fileObject->fileInfo->originalName);
+		}
 	}
 	else if (isListDirectoryHandle(FileHandle))
 	{
@@ -114,27 +133,49 @@ NTSTATUS Detoured_NtQueryInformationFile(HANDLE FileHandle, PIO_STATUS_BLOCK IoS
 	if (isDetouredHandle(FileHandle))
 	{
 		auto& dh = asDetouredHandle(FileHandle);
-		/*
-		if (FileInformationClass == 9) // FileNameInformation
-		{
-			const wchar_t* name = dh.fileObject->fileInfo->originalName;
-			auto& info = *(FILE_NAME_INFORMATION*)FileInformation;
-			info.FileName[0] = '\\';
-			info.FileNameLength = 2;
-			//return STATUS_SUCCESS;
-			u32 nameLen = u32(wcslen(name));
-			//UBA_ASSERT(info.FileNameLength/2 > nameLen);
-			//memcpy(info.FileName, name, nameLen*2+2);
-			//info.FileNameLength = nameLen;
-			memcpy(info.FileName, L"\\\\", 4);
-			info.FileNameLength = 2;
-			return STATUS_SUCCESS;
-		}
-		else
-		*/
+		TrueHandle = dh.trueHandle;
+
+		if (TrueHandle == INVALID_HANDLE_VALUE)
 		{
-			TrueHandle = dh.trueHandle;
-			UBA_ASSERTF(TrueHandle != INVALID_HANDLE_VALUE, L"NtQueryInformationFile (%u) failed using detoured handle %ls (%ls)", FileInformationClass, dh.fileObject->fileInfo->name, dh.fileObject->fileInfo->originalName);
+			if (FileInformationClass == 18) // FileAllInformation 
+			{
+				UBA_ASSERT(dh.dirTableOffset != ~0u);
+				DirectoryTable::EntryInformation entryInfo;
+				g_directoryTable.GetEntryInformation(entryInfo, dh.dirTableOffset);
+				UBA_ASSERT(entryInfo.attributes != 0);
+
+				// TODO This code path is here to handle nodejs queries.. Is not properly implemented and miss things
+				auto& info = *(FILE_ALL_INFORMATION*)FileInformation;
+				info.BasicInformation.CreationTime.QuadPart = entryInfo.lastWrite;
+				info.BasicInformation.LastAccessTime.QuadPart = entryInfo.lastWrite;
+				info.BasicInformation.LastWriteTime.QuadPart = entryInfo.lastWrite;
+				info.BasicInformation.ChangeTime.QuadPart = entryInfo.lastWrite;
+				info.BasicInformation.FileAttributes = entryInfo.attributes;
+				info.StandardInformation.AllocationSize.QuadPart = entryInfo.size;
+				info.StandardInformation.EndOfFile.QuadPart = entryInfo.size;
+				info.StandardInformation.NumberOfLinks = 0;
+				info.StandardInformation.DeletePending = false;
+				info.StandardInformation.Directory = false;
+				info.InternalInformation.IndexNumber.QuadPart = entryInfo.fileIndex;
+				return STATUS_SUCCESS;
+			}
+
+			//if (FileInformationClass == 9) // FileNameInformation
+			//{
+			//	const wchar_t* name = dh.fileObject->fileInfo->originalName;
+			//	auto& info = *(FILE_NAME_INFORMATION*)FileInformation;
+			//	info.FileName[0] = '\\';
+			//	info.FileNameLength = 2;
+			//	//return STATUS_SUCCESS;
+			//	u32 nameLen = u32(wcslen(name));
+			//	//UBA_ASSERT(info.FileNameLength/2 > nameLen);
+			//	//memcpy(info.FileName, name, nameLen*2+2);
+			//	//info.FileNameLength = nameLen;
+			//	memcpy(info.FileName, L"\\\\", 4);
+			//	info.FileNameLength = 2;
+			//	return STATUS_SUCCESS;
+			//}
+			UBA_ASSERTF(false, L"NtQueryInformationFile (%u) failed using detoured handle %ls (%ls)", FileInformationClass, dh.fileObject->fileInfo->name, dh.fileObject->fileInfo->originalName);
 		}
 	}
 
@@ -398,6 +439,19 @@ NTSTATUS Detoured_NtSetInformationFile(HANDLE FileHandle, PIO_STATUS_BLOCK IoSta
 	return res;
 }
 
+NTSTATUS NTAPI Detoured_NtSetInformationObject(HANDLE ObjectHandle, OBJECT_INFORMATION_CLASS ObjectInformationClass, PVOID ObjectInformation, ULONG Length)
+{
+	if (isDetouredHandle(ObjectHandle))
+	{
+		DetouredHandle& h = asDetouredHandle(ObjectHandle);
+		ObjectHandle = h.trueHandle;
+		UBA_ASSERT(ObjectHandle != INVALID_HANDLE_VALUE);
+	}
+	auto res = True_NtSetInformationObject(ObjectHandle, ObjectInformationClass, ObjectInformation, Length);
+	DEBUG_LOG_TRUE(L"NtSetInformationObject", L"(%u) %llu (%ls) -> %ls", ObjectInformationClass, uintptr_t(ObjectHandle), HandleToName(ObjectHandle), ToString(res));
+	return res;
+}
+
 NTSTATUS NTAPI Detoured_NtCreateSection(PHANDLE SectionHandle, ACCESS_MASK DesiredAccess, POBJECT_ATTRIBUTES ObjectAttributes, PLARGE_INTEGER MaximumSize, ULONG SectionPageProtection, ULONG AllocationAttributes, HANDLE FileHandle)
 {
 	DETOURED_CALL(NtCreateSection);
diff --git a/Engine/Source/Programs/UnrealBuildAccelerator/Detours/Private/Windows/UbaDetoursFunctionsWin.h b/Engine/Source/Programs/UnrealBuildAccelerator/Detours/Private/Windows/UbaDetoursFunctionsWin.h
--- a/Engine/Source/Programs/UnrealBuildAccelerator/Detours/Private/Windows/UbaDetoursFunctionsWin.h
+++ b/Engine/Source/Programs/UnrealBuildAccelerator/Detours/Private/Windows/UbaDetoursFunctionsWin.h
@@ -128,6 +128,7 @@
 	DETOURED_FUNCTION(NtQueryObject) \
 	DETOURED_FUNCTION(NtQueryInformationProcess) \
 	DETOURED_FUNCTION(NtSetInformationFile) \
+	DETOURED_FUNCTION(NtSetInformationObject) \
 	DETOURED_FUNCTION(NtCreateSection) \
 	DETOURED_FUNCTION(RtlSizeHeap) \
 	DETOURED_FUNCTION(RtlFreeHeap) \
@@ -341,6 +342,7 @@ extern "C" {
 	NTSTATUS NTAPI NtFlushBuffersFileEx(HANDLE FileHandle, ULONG Flags, PVOID Parameters, ULONG ParametersSize, PIO_STATUS_BLOCK IoStatusBlock);
 	NTSTATUS NTAPI NtReadFile(HANDLE FileHandle, HANDLE Event, PIO_APC_ROUTINE ApcRoutine, PVOID ApcContext, PIO_STATUS_BLOCK IoStatusBlock, PVOID Buffer, ULONG Length, PLARGE_INTEGER ByteOffset, PULONG Key);
 	NTSTATUS NTAPI NtSetInformationFile(HANDLE FileHandle, PIO_STATUS_BLOCK IoStatusBlock, PVOID FileInformation, ULONG Length, FILE_INFORMATION_CLASS FileInformationClass);
+	NTSTATUS NTAPI NtSetInformationObject(HANDLE ObjectHandle, OBJECT_INFORMATION_CLASS ObjectInformationClass, PVOID ObjectInformation, ULONG Length);
 	NTSTATUS NTAPI NtCreateSection(PHANDLE SectionHandle, ACCESS_MASK DesiredAccess, POBJECT_ATTRIBUTES ObjectAttributes, PLARGE_INTEGER MaximumSize, ULONG SectionPageProtection, ULONG AllocationAttributes, HANDLE FileHandle);
 	NTSTATUS NTAPI NtCreateIoCompletion(PHANDLE IoCompletionHandle, ACCESS_MASK DesiredAccess, POBJECT_ATTRIBUTES ObjectAttributes, ULONG Count);
 	NTSTATUS NTAPI ZwCreateFile(PHANDLE FileHandle, ACCESS_MASK DesiredAccess, POBJECT_ATTRIBUTES ObjectAttributes, PIO_STATUS_BLOCK IoStatusBlock, PLARGE_INTEGER AllocationSize, ULONG FileAttributes, ULONG ShareAccess, ULONG CreateDisposition, ULONG CreateOptions, PVOID EaBuffer, ULONG EaLength);
@@ -404,17 +406,57 @@ struct FILE_RENAME_INFORMATION {
 };
 
 struct FILE_IS_REMOTE_DEVICE_INFORMATION {
-  BOOLEAN IsRemote;
+	BOOLEAN IsRemote;
+	};
+
+	struct FILE_ID_INFORMATION {
+	ULONGLONG   VolumeSerialNumber;
+	FILE_ID_128 FileId;
+	};
+
+	struct FILE_NAME_INFORMATION {
+	ULONG FileNameLength;
+	WCHAR FileName[1];
+	};
+
+	struct FILE_BASIC_INFORMATION {
+	LARGE_INTEGER CreationTime;
+	LARGE_INTEGER LastAccessTime;
+	LARGE_INTEGER LastWriteTime;
+	LARGE_INTEGER ChangeTime;
+	DWORD FileAttributes;
+	};
+
+	struct FILE_STANDARD_INFORMATION {
+	LARGE_INTEGER AllocationSize;
+	LARGE_INTEGER EndOfFile;
+	ULONG         NumberOfLinks;
+	BOOLEAN       DeletePending;
+	BOOLEAN       Directory;
+};
+
+struct FILE_INTERNAL_INFORMATION {
+	LARGE_INTEGER IndexNumber;
 };
 
-struct FILE_ID_INFORMATION {
-  ULONGLONG   VolumeSerialNumber;
-  FILE_ID_128 FileId;
+struct FILE_ALL_INFORMATION {
+	FILE_BASIC_INFORMATION     BasicInformation;
+	FILE_STANDARD_INFORMATION  StandardInformation;
+	FILE_INTERNAL_INFORMATION  InternalInformation;
+	//FILE_EA_INFORMATION        EaInformation;
+	//FILE_ACCESS_INFORMATION    AccessInformation;
+	//FILE_POSITION_INFORMATION  PositionInformation;
+	//FILE_MODE_INFORMATION      ModeInformation;
+	//FILE_ALIGNMENT_INFORMATION AlignmentInformation;
+	//FILE_NAME_INFORMATION      NameInformation;
 };
 
-struct FILE_NAME_INFORMATION {
-  ULONG FileNameLength;
-  WCHAR FileName[1];
+struct FILE_FS_VOLUME_INFORMATION {
+	LARGE_INTEGER VolumeCreationTime;
+	ULONG         VolumeSerialNumber;
+	ULONG         VolumeLabelLength;
+	BOOLEAN       SupportsObjects;
+	WCHAR         VolumeLabel[1];
 };
 
 namespace uba
diff --git a/Engine/Source/Programs/UnrealBuildAccelerator/Test/Private/UbaTestAll.h b/Engine/Source/Programs/UnrealBuildAccelerator/Test/Private/UbaTestAll.h
--- a/Engine/Source/Programs/UnrealBuildAccelerator/Test/Private/UbaTestAll.h
+++ b/Engine/Source/Programs/UnrealBuildAccelerator/Test/Private/UbaTestAll.h
@@ -12,22 +12,8 @@
 
 namespace uba
 {
-#if PLATFORM_MAC
-	#define UBA_EXTRA_TESTS
-#else
-	#define UBA_EXTRA_TESTS \
-		UBA_TEST(TestMultipleDetouredProcesses) \
-		UBA_TEST(TestLogLines) \
-		UBA_TEST(TestLogLinesNoDetour) \
-		UBA_TEST(TestLocalSchedule) \
-		UBA_TEST(TestLocalScheduleReuse) \
-		UBA_TEST(TestDetouredTouch) \
-		UBA_TEST(TestRemoteScheduleReuse) \
-		UBA_TEST(TestStdOutLocal) \
-
-#endif 
 
-	#define UBA_TESTS \
+#define UBA_ALLPLATFORM_TESTS \
 		UBA_TEST(TestTime) \
 		UBA_TEST(TestEvents) \
 		UBA_TEST(TestPaths) \
@@ -41,10 +27,61 @@ namespace uba
 		UBA_TEST(TestClientServerMem) \
 		UBA_TEST(TestStorage) \
 		UBA_TEST(TestDetouredTestApp) \
-		UBA_TEST(TestDetouredClang) \
 		UBA_TEST(TestRemoteDetouredTestApp) \
+
+
+#define UBA_POSIX_TESTS \
+		UBA_TEST(TestDetouredClang) \
+		UBA_TEST(TestRemoteDetouredClang) \
+		UBA_TEST(TestDetouredTouch) \
+
+
+#define UBA_NONMAC_TESTS \
+		UBA_TEST(TestMultipleDetouredProcesses) \
+		UBA_TEST(TestLogLines) \
+		UBA_TEST(TestLogLinesNoDetour) \
+		UBA_TEST(TestLocalSchedule) \
+		UBA_TEST(TestLocalScheduleReuse) \
+		UBA_TEST(TestRemoteScheduleReuse) \
+
+#define UBA_WINDOWS_TESTS \
+		UBA_NONMAC_TESTS \
 		UBA_TEST(TestCustomService) \
-		UBA_EXTRA_TESTS
+		UBA_TEST(TestStdOutLocal) \
+		UBA_TEST(TestStdOutViaCmd) \
+
+
+#define UBA_LINUX_TESTS \
+		UBA_NONMAC_TESTS \
+		UBA_POSIX_TESTS \
+
+
+#define UBA_MAC_TESTS \
+		UBA_POSIX_TESTS
+
+
+
+
+#if !PLATFORM_WINDOWS
+#undef UBA_WINDOWS_TESTS
+#define UBA_WINDOWS_TESTS
+#endif
+
+#if !PLATFORM_LINUX
+#undef UBA_LINUX_TESTS
+#define UBA_LINUX_TESTS
+#endif
+
+#if !PLATFORM_MAC
+#undef  UBA_MAC_TESTS
+#define UBA_MAC_TESTS
+#endif
+
+#define UBA_TESTS \
+		UBA_ALLPLATFORM_TESTS \
+		UBA_WINDOWS_TESTS \
+		UBA_LINUX_TESTS \
+		UBA_MAC_TESTS \
 
 
 	#define UBA_TEST(x) \
@@ -87,7 +124,7 @@ namespace uba
 		UBA_TESTS
 
 		logger.Info(TC("Tests finished successfully!"));
-		Sleep(3000);
+		Sleep(2000);
 
 		return true;
 	}
diff --git a/Engine/Source/Programs/UnrealBuildAccelerator/Test/Private/UbaTestScheduler.h b/Engine/Source/Programs/UnrealBuildAccelerator/Test/Private/UbaTestScheduler.h
--- a/Engine/Source/Programs/UnrealBuildAccelerator/Test/Private/UbaTestScheduler.h
+++ b/Engine/Source/Programs/UnrealBuildAccelerator/Test/Private/UbaTestScheduler.h
@@ -33,9 +33,6 @@ namespace uba
 
 	bool TestLocalScheduleReuse(LoggerWithWriter& logger, const StringBufferBase& testRootDir)
 	{
-		if (!IsWindows)
-			return true;
-
 		return RunLocal(logger, testRootDir, [](LoggerWithWriter& logger, SessionServer& session, const tchar* workingDir, const RunProcessFunction& runProcess)
 			{
 				SchedulerCreateInfo info(session);
@@ -64,9 +61,6 @@ namespace uba
 
 	bool TestRemoteScheduleReuse(LoggerWithWriter& logger, const StringBufferBase& testRootDir)
 	{
-		if (!IsWindows)
-			return true;
-
 		return RunRemote(logger, testRootDir, [](LoggerWithWriter& logger, SessionServer& session, const tchar* workingDir, const RunProcessFunction& runProcess)
 			{
 				SchedulerCreateInfo info(session);
diff --git a/Engine/Source/Programs/UnrealBuildAccelerator/Test/Private/UbaTestSession.h b/Engine/Source/Programs/UnrealBuildAccelerator/Test/Private/UbaTestSession.h
--- a/Engine/Source/Programs/UnrealBuildAccelerator/Test/Private/UbaTestSession.h
+++ b/Engine/Source/Programs/UnrealBuildAccelerator/Test/Private/UbaTestSession.h
@@ -289,30 +289,21 @@ namespace uba
 
 	bool TestCustomService(LoggerWithWriter& logger, const StringBufferBase& testRootDir)
 	{
-		if (!IsWindows)
-			return true;
-
 		return RunRemote(logger, testRootDir, RunCustomService);
 	}
 
 	bool TestDetouredClang(LoggerWithWriter& logger, const StringBufferBase& testRootDir)
 	{
-		if (IsWindows)
-			return true;
 		return RunLocal(logger, testRootDir, RunClang);
 	}
 
 	bool TestRemoteDetouredClang(LoggerWithWriter& logger, const StringBufferBase& testRootDir)
 	{
-		if (IsWindows)
-			return true;
 		return RunRemote(logger, testRootDir, RunClang);
 	}
 
 	bool TestDetouredTouch(LoggerWithWriter& logger, const StringBufferBase& testRootDir)
 	{
-		if (IsWindows)
-			return true;
 		return RunLocal(logger, testRootDir, [](LoggerWithWriter& logger, SessionServer& session, const tchar* workingDir, const RunProcessFunction& runProcess)
 			{
 				StringBuffer<> file;
diff --git a/Engine/Source/Programs/UnrealBuildAccelerator/Test/Private/UbaTestStdOut.h b/Engine/Source/Programs/UnrealBuildAccelerator/Test/Private/UbaTestStdOut.h
--- a/Engine/Source/Programs/UnrealBuildAccelerator/Test/Private/UbaTestStdOut.h
+++ b/Engine/Source/Programs/UnrealBuildAccelerator/Test/Private/UbaTestStdOut.h
@@ -7,7 +7,7 @@
 
 namespace uba
 {
-	bool TestStdOut(LoggerWithWriter& logger, const StringBufferBase& testRootDir, bool remote)
+	bool TestStdOut(LoggerWithWriter& logger, const StringBufferBase& testRootDir, bool remote, const tchar* app = nullptr, const tchar* arg = nullptr, const tchar* expectedOut = nullptr)
 	{
 		LogWriter& logWriter = logger.m_writer;
 		NetworkBackendTcp networkBackend(logWriter);
@@ -47,11 +47,20 @@ namespace uba
 			return logger.Error(TC("Failed to listen"));
 
 		StringBuffer<> testApp;
-		GetTestAppPath(logger, testApp);
+		if (!app)
+		{
+			GetTestAppPath(logger, testApp);
+			app = testApp.data;
+		}
+		if (!arg)
+		{
+			arg = TC("-stdout=rootprocess");
+			expectedOut = TC("rootprocess");
+		}
 
 		ProcessStartInfo pi;
-		pi.application = testApp.data;
-		pi.arguments = TC("-stdout=rootprocess");
+		pi.application = app;
+		pi.arguments = arg;
 		pi.workingDir = workingDir.data;
 		pi.description = TC("StdOutDesc");
 		pi.logFile = TC("Log");
@@ -68,9 +77,9 @@ namespace uba
 
 		auto& logLines = ph.GetLogLines();
 		if (logLines.size() != 1)
-			return false;
-		if (!Equals(logLines[0].text.c_str(), TC("rootprocess")))
-			return false;
+			return logger.Error(TC("Application %s produced %u log line(s) but expected 1"), app, u32(logLines.size()));
+		if (!Equals(logLines[0].text.c_str(), expectedOut))
+			return logger.Error(TC("Application %s produced non-matching log line: %s"), app, logLines[0].text.c_str());
 
 		return true;
 	}
@@ -86,4 +95,15 @@ namespace uba
 	{
 		return TestStdOut(logger, testRootDir, true);
 	}
+
+	bool TestStdOutViaCmd(LoggerWithWriter& logger, const StringBufferBase& testRootDir)
+	{
+		if (!IsWindows)
+			return true;
+		StringBuffer<> args;
+		args.Append(TC("/c \""));
+		GetTestAppPath(logger, args);
+		args.Append(TC(" -stdout=foo\""));
+		return TestStdOut(logger, testRootDir, false, TC("c:\\windows\\system32\\cmd.exe"), args.data, TC("foo"));
+	}
 }
diff --git a/Engine/Source/Programs/UnrealBuildAccelerator/TestApp/Private/Windows/UbaTestAppWin.cpp b/Engine/Source/Programs/UnrealBuildAccelerator/TestApp/Private/Windows/UbaTestAppWin.cpp
--- a/Engine/Source/Programs/UnrealBuildAccelerator/TestApp/Private/Windows/UbaTestAppWin.cpp
+++ b/Engine/Source/Programs/UnrealBuildAccelerator/TestApp/Private/Windows/UbaTestAppWin.cpp
@@ -19,20 +19,21 @@ int LogError(const wchar_t* format, ...)
 int wmain(int argc, wchar_t* argv[])
 {
 	HMODULE detoursHandle = GetModuleHandleW(L"UbaDetours.dll");
-	if (!detoursHandle)
-		return LogError(L"Did not find UbaDetours.dll in process!!!\n");
-
-	using UbaRunningRemoteFunc = bool();
-	UbaRunningRemoteFunc* runningRemoteFunc = (UbaRunningRemoteFunc*)GetProcAddress(detoursHandle, "UbaRunningRemote");
-	if (!runningRemoteFunc)
-		return LogError(L"Couldn't find UbaRunningRemote function in UbaDetours.dll");
-	bool runningRemote = (*runningRemoteFunc)();
 
 	using UbaRequestNextProcessFunc = bool(unsigned int prevExitCode, wchar_t* outArguments, unsigned int outArgumentsCapacity);
 	static UbaRequestNextProcessFunc* requestNextProcess = (UbaRequestNextProcessFunc*)(void*)GetProcAddress(detoursHandle, "UbaRequestNextProcess");
 
 	if (argc == 1)
 	{
+		if (!detoursHandle)
+			return LogError(L"Did not find UbaDetours.dll in process!!!\n");
+
+		using UbaRunningRemoteFunc = bool();
+		UbaRunningRemoteFunc* runningRemoteFunc = (UbaRunningRemoteFunc*)GetProcAddress(detoursHandle, "UbaRunningRemote");
+		if (!runningRemoteFunc)
+			return LogError(L"Couldn't find UbaRunningRemote function in UbaDetours.dll");
+		bool runningRemote = (*runningRemoteFunc)();
+
 		HMODULE modules[] = { 0, detoursHandle, GetModuleHandleW(L"UbaTestApp.exe") };
 		for (HMODULE module : modules)
 		{
@@ -237,6 +238,9 @@ int wmain(int argc, wchar_t* argv[])
 	}
 	else
 	{
+		if (!detoursHandle)
+			return LogError(L"Did not find UbaDetours.dll in process!!!\n");
+
 		using u32 = unsigned int;
 		using UbaSendCustomMessageFunc = u32(const void* send, u32 sendSize, void* recv, u32 recvCapacity);
 
