[UBA]
* Added support for GetFileInformationByHandleEx with information class FileAttributeTagInfo

diff --git a/Engine/Source/Programs/UnrealBuildAccelerator/Detours/Private/Windows/UbaDetoursFunctionsKernelBase.inl b/Engine/Source/Programs/UnrealBuildAccelerator/Detours/Private/Windows/UbaDetoursFunctionsKernelBase.inl
--- a/Engine/Source/Programs/UnrealBuildAccelerator/Detours/Private/Windows/UbaDetoursFunctionsKernelBase.inl
+++ b/Engine/Source/Programs/UnrealBuildAccelerator/Detours/Private/Windows/UbaDetoursFunctionsKernelBase.inl
@@ -1780,6 +1780,13 @@ BOOL Detoured_GetFileInformationByHandleEx(HANDLE hFile, FILE_INFO_BY_HANDLE_CLA
 			return FALSE;
 			*/
 		}
+		else if (fileInformationClass == FileAttributeTagInfo)
+		{
+			auto& data = *(FILE_ATTRIBUTE_TAG_INFO*)lpFileInformation;
+			data.FileAttributes = entryInfo.attributes;
+			data.ReparseTag = 0;
+			return TRUE;
+		}
 		else
 		{
 			UBA_ASSERTF(trueHandle != INVALID_HANDLE_VALUE, L"GetFileInformationByHandleEx with class %u not Implemented (%ls)", fileInformationClass, HandleToName(hFile));
