[UBA]
* Fixes for new intel compiler not working when compiling remotely. CreateDirectory/RemoveDirectory was not properly implemented for temp directory paths

diff --git a/Engine/Source/Programs/UnrealBuildAccelerator/Detours/Private/Windows/UbaDetoursFunctionsKernelBase.inl b/Engine/Source/Programs/UnrealBuildAccelerator/Detours/Private/Windows/UbaDetoursFunctionsKernelBase.inl
--- a/Engine/Source/Programs/UnrealBuildAccelerator/Detours/Private/Windows/UbaDetoursFunctionsKernelBase.inl
+++ b/Engine/Source/Programs/UnrealBuildAccelerator/Detours/Private/Windows/UbaDetoursFunctionsKernelBase.inl
@@ -170,11 +170,22 @@ HANDLE Detoured_CreateFileA(LPCSTR lpFileName, DWORD dwDesiredAccess, DWORD dwSh
 BOOL Detoured_CreateDirectoryW(LPCWSTR lpPathName, LPSECURITY_ATTRIBUTES lpSecurityAttributes)
 {
 	DETOURED_CALL(CreateDirectoryW);
+
+	StringBuffer<> pathName;
+	FixPath(pathName, lpPathName);
+	if (pathName.StartsWith(g_systemTemp.data))
+	{
+		SuppressCreateFileDetourScope s;
+		BOOL res = True_CreateDirectoryW(lpPathName, lpSecurityAttributes);
+		DEBUG_LOG_TRUE(L"CreateDirectoryW", L"%ls -> %ls", lpPathName, ToString(res));
+		return res;
+	}
+
 	BOOL res;
+	u32 errorCode = 0;
+	StringKey pathNameKey = ToStringKeyLower(pathName);
+
 	{
-		StringBuffer<> pathName;
-		FixPath(pathName, lpPathName);
-		StringKey pathNameKey = ToStringKeyLower(pathName);
 		TimerScope ts(g_stats.createFile);
 		SCOPED_WRITE_LOCK(g_communicationLock, pcs);
 		BinaryWriter writer;
@@ -184,23 +195,33 @@ BOOL Detoured_CreateDirectoryW(LPCWSTR lpPathName, LPSECURITY_ATTRIBUTES lpSecur
 		writer.Flush();
 		BinaryReader reader;
 		res = reader.ReadBool();
-		u32 errorCode = reader.ReadU32();
-		SetLastError(errorCode);
+		errorCode = reader.ReadU32();
 	}
-	DEBUG_LOG_DETOURED(L"CreateDirectoryW", L"%ls -> %ls (%u)", lpPathName, ToString(res), GetLastError());
+
+	SetLastError(errorCode);
+	DEBUG_LOG_DETOURED(L"CreateDirectoryW", L"%ls -> %ls (%u)", lpPathName, ToString(res), errorCode);
 	return res;
 }
 
 BOOL Detoured_RemoveDirectoryW(LPCWSTR lpPathName)
 {
 	DETOURED_CALL(RemoveDirectoryW);
-	UBA_ASSERT(!g_runningRemote);
+
+	StringBuffer<> pathName;
+	FixPath(pathName, lpPathName);
 	BOOL res;
+	if (!g_runningRemote || pathName.StartsWith(g_systemTemp.data))
 	{
 		SuppressCreateFileDetourScope s; // TODO: Revisit this.. will not work remotely
 		res = True_RemoveDirectoryW(lpPathName);
 	}
+	else
+	{
+		UBA_ASSERTF(!g_runningRemote, L"RemoveDirectory is not implemented for remote (removing %s)", lpPathName);
+		res = false;
+	}
 	DEBUG_LOG_TRUE(L"RemoveDirectoryW", L"%ls -> %ls", lpPathName, ToString(res));
+
 	return res;
 }
 
@@ -3645,7 +3666,7 @@ LPCWSTR Detoured_PathFindFileNameW(LPCWSTR pszPath) // This is called by Ps4Symb
 
 BOOL Detoured_PathIsRelativeW(LPCWSTR pszPath)
 {
-	UBA_ASSERTF(!g_runningRemote, L"%ls", pszPath);
+	//UBA_ASSERTF(!g_runningRemote, L"%ls", pszPath); // intel compiler uses PathIsRelativeW.. don't know if this function touches file system but will comment out this assert for now
 	auto res = True_PathIsRelativeW(pszPath);
 	DEBUG_LOG_TRUE(L"PathIsRelativeW", L"(%ls) -> %ls", pszPath, res);
 	return res;
