[UBA]
* Added support for storing .obj files compressed by cl.exe and then decompressed in-memory by link.exe. Note this is not properly implemented and only work for cl.exe+link.exe combo

diff --git a/Engine/Source/Programs/UnrealBuildAccelerator/Common/Private/UbaFile.cpp b/Engine/Source/Programs/UnrealBuildAccelerator/Common/Private/UbaFile.cpp
--- a/Engine/Source/Programs/UnrealBuildAccelerator/Common/Private/UbaFile.cpp
+++ b/Engine/Source/Programs/UnrealBuildAccelerator/Common/Private/UbaFile.cpp
@@ -266,8 +266,11 @@ namespace uba
 			HMODULE hm = NULL;
 			if (!GetModuleHandleEx(GET_MODULE_HANDLE_EX_FLAG_FROM_ADDRESS | GET_MODULE_HANDLE_EX_FLAG_UNCHANGED_REFCOUNT, (LPCWSTR)&GetDirectoryOfCurrentModule, &hm))
 				return logger.Error(TC("GetModuleHandleEx failed (%s)"), LastErrorToText().data);
-			if (!GetModuleFileNameW(hm, out.data, out.capacity))
+			u32 len = GetModuleFileNameW(hm, out.data + out.count, out.capacity - out.count);
+			if (!len)
 				return logger.Error(TC("GetModuleFileNameW failed (%s)"), LastErrorToText().data);
+			out.count += len;
+			UBA_ASSERTF(GetLastError() == ERROR_SUCCESS, TC("GetModuleFileNameW failed (%s)"), LastErrorToText().data);
 			const tchar* lastSlash = out.Last('\\');
 			out.Resize(lastSlash - out.data);
 			return true;
diff --git a/Engine/Source/Programs/UnrealBuildAccelerator/Common/Private/UbaProcess.cpp b/Engine/Source/Programs/UnrealBuildAccelerator/Common/Private/UbaProcess.cpp
--- a/Engine/Source/Programs/UnrealBuildAccelerator/Common/Private/UbaProcess.cpp
+++ b/Engine/Source/Programs/UnrealBuildAccelerator/Common/Private/UbaProcess.cpp
@@ -1394,6 +1394,7 @@ namespace uba
 			payload.trackInputs = m_startInfo.trackInputs;
 			payload.useCustomAllocator = m_startInfo.useCustomAllocator && GetApplicationRules()[m_rulesIndex].rules->AllowMiMalloc();
 			payload.isRunningWine = IsRunningWine();
+			payload.storeObjFilesCompressed = m_session.m_storeObjFilesCompressed;
 			payload.uiLanguage = m_startInfo.uiLanguage;
 			if (*m_startInfo.logFile)
 			{
diff --git a/Engine/Source/Programs/UnrealBuildAccelerator/Common/Private/UbaSession.cpp b/Engine/Source/Programs/UnrealBuildAccelerator/Common/Private/UbaSession.cpp
--- a/Engine/Source/Programs/UnrealBuildAccelerator/Common/Private/UbaSession.cpp
+++ b/Engine/Source/Programs/UnrealBuildAccelerator/Common/Private/UbaSession.cpp
@@ -960,6 +960,8 @@ namespace uba
 		m_shouldWriteToDisk = info.shouldWriteToDisk;
 		UBA_ASSERTF(m_shouldWriteToDisk || m_allowMemoryMaps, TC("Can't disable both should write to disk and allow memory maps"));
 
+		m_storeObjFilesCompressed = info.storeObjFilesCompressed;
+
 		m_detailedTrace = info.detailedTrace;
 		m_logToFile = info.logToFile;
 		if (info.extraInfo)
@@ -1978,6 +1980,13 @@ namespace uba
 
 				auto memClose = MakeGuard([&](){ UnmapViewOfFile(mem, fileSize, file.name.c_str()); });
 
+
+				if (m_storeObjFilesCompressed && EndsWith(file.name.c_str(), file.name.size(), TC(".obj")))
+				{
+					Storage::WriteResult res;
+					return m_storage.WriteCompressed(res, TC("MemoryMap"), InvalidFileHandle, mem, fileSize, file.name.c_str());
+				}
+
 				// Seems like best combo (for windows at least) is to use writes with overlap and max 16 at the same time.
 				// On one machine we get twice as fast without overlap if no bottleneck. On another machine (ntfs compression on) we get twice as slow without overlap
 				// Both machines behaves well with overlap AND bottleneck. Both machine are 128 logical core thread rippers.
diff --git a/Engine/Source/Programs/UnrealBuildAccelerator/Common/Private/UbaSessionClient.cpp b/Engine/Source/Programs/UnrealBuildAccelerator/Common/Private/UbaSessionClient.cpp
--- a/Engine/Source/Programs/UnrealBuildAccelerator/Common/Private/UbaSessionClient.cpp
+++ b/Engine/Source/Programs/UnrealBuildAccelerator/Common/Private/UbaSessionClient.cpp
@@ -1077,6 +1077,7 @@ namespace uba
 
 		m_sessionId = reader.ReadU32();
 		m_uiLanguage = reader.ReadU32();
+		m_storeObjFilesCompressed = reader.ReadBool();
 		m_detailedTrace = reader.ReadBool();
 		m_shouldSendLogToServer = reader.ReadBool();
 		m_shouldSendTraceToServer = reader.ReadBool();
diff --git a/Engine/Source/Programs/UnrealBuildAccelerator/Common/Private/UbaSessionServer.cpp b/Engine/Source/Programs/UnrealBuildAccelerator/Common/Private/UbaSessionServer.cpp
--- a/Engine/Source/Programs/UnrealBuildAccelerator/Common/Private/UbaSessionServer.cpp
+++ b/Engine/Source/Programs/UnrealBuildAccelerator/Common/Private/UbaSessionServer.cpp
@@ -618,6 +618,7 @@ namespace uba
 				writer.WriteBool(m_resetCas);
 				writer.WriteU32(sessionId);
 				writer.WriteU32(m_uiLanguage);
+				writer.WriteBool(m_storeObjFilesCompressed);
 				writer.WriteBool(m_detailedTrace);
 				writer.WriteBool(m_remoteLogEnabled);
 				writer.WriteBool(m_remoteTraceEnabled);
@@ -807,7 +808,8 @@ namespace uba
 
 					if (ShouldWriteToDisk(destination.data, destination.count))
 					{
-						success = m_storage.CopyOrLink(casKey, destination.data, attributes);
+						bool writeCompressed = m_storeObjFilesCompressed && destination.EndsWith(TC(".obj"));
+						success = m_storage.CopyOrLink(casKey, destination.data, attributes, writeCompressed);
 						if (!success)
 							m_logger.Error(TC("Failed to copy cas from %s to %s"), CasKeyString(casKey).str, destination.data);
 					}
diff --git a/Engine/Source/Programs/UnrealBuildAccelerator/Common/Private/UbaStorage.cpp b/Engine/Source/Programs/UnrealBuildAccelerator/Common/Private/UbaStorage.cpp
--- a/Engine/Source/Programs/UnrealBuildAccelerator/Common/Private/UbaStorage.cpp
+++ b/Engine/Source/Programs/UnrealBuildAccelerator/Common/Private/UbaStorage.cpp
@@ -1896,7 +1896,7 @@ namespace uba
 		return out != CasKeyZero;
 	}
 
-	bool StorageImpl::CopyOrLink(const CasKey& casKey, const tchar* destination, u32 fileAttributes)
+	bool StorageImpl::CopyOrLink(const CasKey& casKey, const tchar* destination, u32 fileAttributes, bool writeCompressed)
 	{
 		UBA_ASSERT(casKey != CasKeyZero);
 		UBA_ASSERT(fileAttributes);
@@ -1938,9 +1938,9 @@ namespace uba
 
 		TimerScope ts(stats.copyOrLink);
 
-		CasKey actualKey = AsCompressed(casKey, false);
+		CasKey actualKey = AsCompressed(casKey, writeCompressed);
 
-		bool testCompressed = true;
+		bool testCompressed = !writeCompressed;
 		while (true)
 		{
 			CasEntry* casEntry = nullptr;
@@ -1979,6 +1979,15 @@ namespace uba
 					if (!compressedData)
 						return m_logger.Error(TC("Failed to map view of mapping %s (%s)"), casFile.data, LastErrorToText().data);
 
+					if (writeCompressed)
+					{
+						FileAccessor destinationFile(m_logger, destination);
+						if (!destinationFile.CreateMemoryWrite(false, fileAttributes, mappedView.size, m_tempPath.data))
+							return false;
+						memcpy(destinationFile.GetData(), mappedView.memory, mappedView.size);
+						return destinationFile.Close();
+					}
+
 					decompressedSize = *(u64*)compressedData;
 					readData = compressedData + sizeof(u64);
 				}
@@ -1990,6 +1999,9 @@ namespace uba
 #else
 					UBA_ASSERT(false);
 #endif
+					if (writeCompressed)
+						return m_logger.Error(TC("NOT IMPLEMENTED %s"), destination);
+
 					if (!OpenFileSequentialRead(m_logger, casFile.data, readHandle))
 						return m_logger.Error(TC("Failed to open file %s for read (%s)"), casFile.data, LastErrorToText().data);
 
@@ -2443,12 +2455,12 @@ namespace uba
 			while (left)
 			{
 				u32 sizes[2];
-				if (!ReadFile(m_logger, TC(""), fileHandle, sizes, sizeof(u32) * 2))
+				if (!ReadFile(m_logger, fileName, fileHandle, sizes, sizeof(u32) * 2))
 					return false;
 				u32 compressedBlockSize = sizes[0];
 				u32 decompressedBlockSize = sizes[1];
 
-				if (!ReadFile(m_logger, TC(""), fileHandle, readBuffer, compressedBlockSize))
+				if (!ReadFile(m_logger, fileName, fileHandle, readBuffer, compressedBlockSize))
 					return false;
 				TimerScope ts(stats.decompressToMem);
 				OO_SINTa decompLen = OodleLZ_Decompress(readBuffer, (OO_SINTa)compressedBlockSize, writePos, (OO_SINTa)decompressedBlockSize);
diff --git a/Engine/Source/Programs/UnrealBuildAccelerator/Common/Public/UbaNetwork.h b/Engine/Source/Programs/UnrealBuildAccelerator/Common/Public/UbaNetwork.h
--- a/Engine/Source/Programs/UnrealBuildAccelerator/Common/Public/UbaNetwork.h
+++ b/Engine/Source/Programs/UnrealBuildAccelerator/Common/Public/UbaNetwork.h
@@ -18,7 +18,7 @@ namespace uba
 	
 	static constexpr u32 SystemNetworkVersion = 1339;
 	static constexpr u32 StorageNetworkVersion = 3;
-	static constexpr u32 SessionNetworkVersion = 30;
+	static constexpr u32 SessionNetworkVersion = 31;
 
 
 	// Messages used over network between client/server (system, storage and session)
diff --git a/Engine/Source/Programs/UnrealBuildAccelerator/Common/Public/UbaSession.h b/Engine/Source/Programs/UnrealBuildAccelerator/Common/Public/UbaSession.h
--- a/Engine/Source/Programs/UnrealBuildAccelerator/Common/Public/UbaSession.h
+++ b/Engine/Source/Programs/UnrealBuildAccelerator/Common/Public/UbaSession.h
@@ -165,6 +165,7 @@ namespace uba
 		bool m_shouldWriteToDisk;
 		bool m_detailedTrace;
 		bool m_logToFile;
+		bool m_storeObjFilesCompressed;
 
 		u64 m_keepOutputFileMemoryMapsThreshold;
 
diff --git a/Engine/Source/Programs/UnrealBuildAccelerator/Common/Public/UbaSessionCreateInfo.h b/Engine/Source/Programs/UnrealBuildAccelerator/Common/Public/UbaSessionCreateInfo.h
--- a/Engine/Source/Programs/UnrealBuildAccelerator/Common/Public/UbaSessionCreateInfo.h
+++ b/Engine/Source/Programs/UnrealBuildAccelerator/Common/Public/UbaSessionCreateInfo.h
@@ -26,6 +26,7 @@ namespace uba
 		bool shouldWriteToDisk = true;			// Set to false to skip writing output files to disk
 		bool traceEnabled = false;				// Set to true to always create in-memory trace data. Is not needed if traceName, traceOutputFile or launchVisualizer is set
 		bool detailedTrace = false;				// Enable detailed trace to include jobs, individual file I/O etc in trace dump
+		bool storeObjFilesCompressed = false;	// Compiler will write .obj files compressed to disk and linker will decompress when reading
 		u64 deleteSessionsOlderThanSeconds = 12 * 60 * 60; // Delete session folders older than 12 hours by default . Set to 0 to not delete or 1 to delete all
 		u64 keepOutputFileMemoryMapsThreshold = 256 * 1024; // If allowMemoryMaps is true, output files will be kept in memory if smaller than this size
 	};
diff --git a/Engine/Source/Programs/UnrealBuildAccelerator/Common/Public/UbaStorage.h b/Engine/Source/Programs/UnrealBuildAccelerator/Common/Public/UbaStorage.h
--- a/Engine/Source/Programs/UnrealBuildAccelerator/Common/Public/UbaStorage.h
+++ b/Engine/Source/Programs/UnrealBuildAccelerator/Common/Public/UbaStorage.h
@@ -49,7 +49,7 @@ namespace uba
 		virtual bool StoreCasFile(CasKey& out, StringKey fileNameKey, const tchar* fileName, FileMappingHandle mappingHandle, u64 mappingOffset, u64 fileSize, const tchar* hint, bool deferCreation = false, bool keepMappingInMemory = false) = 0;
 		virtual bool DropCasFile(const CasKey& casKey, bool forceDelete, const tchar* hint) = 0;
 		virtual bool CalculateCasKey(CasKey& out, const tchar* fileName) = 0;
-		virtual bool CopyOrLink(const CasKey& casKey, const tchar* destination, u32 fileAttributes) = 0;
+		virtual bool CopyOrLink(const CasKey& casKey, const tchar* destination, u32 fileAttributes, bool writeCompressed = false) = 0;
 		virtual bool FakeCopy(const CasKey& casKey, const tchar* destination) = 0;
 #if !UBA_USE_SPARSEFILE
 		virtual bool GetCasFileName(StringBufferBase& out, const CasKey& casKey) = 0;
@@ -66,6 +66,9 @@ namespace uba
 
 		virtual void SetTrace(Trace* trace, bool detailed) {}
 		virtual void Ping() {}
+
+		struct WriteResult { FileMappingHandle mappingHandle; u64 size = InvalidValue; u64 offset = InvalidValue; };
+		virtual bool WriteCompressed(WriteResult& out, const tchar* from, FileHandle readHandle, u8* readMem, u64 fileSize, const tchar* toFile) = 0;
 	};
 
 	struct StorageCreateInfo
@@ -106,7 +109,7 @@ namespace uba
 		virtual bool StoreCasFile(CasKey& out, StringKey fileNameKey, const tchar* fileName, FileMappingHandle mappingHandle, u64 mappingOffset, u64 fileSize, const tchar* hint, bool deferCreation = false, bool keepMappingInMemory = false) override;
 		virtual bool DropCasFile(const CasKey& casKey, bool forceDelete, const tchar* hint) override;
 		virtual bool CalculateCasKey(CasKey& out, const tchar* fileName) override;
-		virtual bool CopyOrLink(const CasKey& casKey, const tchar* destination, u32 fileAttributes) override;
+		virtual bool CopyOrLink(const CasKey& casKey, const tchar* destination, u32 fileAttributes, bool writeCompressed = false) override;
 		virtual bool FakeCopy(const CasKey& casKey, const tchar* destination) override;
 #if !UBA_USE_SPARSEFILE
 		virtual bool GetCasFileName(StringBufferBase& out, const CasKey& casKey) override;
@@ -126,9 +129,8 @@ namespace uba
 		CasKey CalculateCasKey(const tchar* fileName, FileHandle fileHandle, u64 fileSize, bool storeCompressed);
 		CasKey CalculateCasKey(u8* fileMem, u64 fileSize, bool storeCompressed);
 
-		struct WriteResult { FileMappingHandle mappingHandle; u64 size = InvalidValue; u64 offset = InvalidValue; };
 		virtual bool WriteCompressed(WriteResult& out, const tchar* from, const tchar* toFile);
-		bool WriteCompressed(WriteResult& out, const tchar* from, FileHandle readHandle, u8* readMem, u64 fileSize, const tchar* toFile);
+		virtual bool WriteCompressed(WriteResult& out, const tchar* from, FileHandle readHandle, u8* readMem, u64 fileSize, const tchar* toFile) final;
 		bool WriteMemToCompressedFile(FileAccessor& destination, u32 workCount, const u8* uncompressedData, u64 fileSize, u64 maxUncompressedBlock, u64& totalWritten);
 		bool WriteCasFileNoCheck(WriteResult& out, const tchar* fileName, const tchar* casFile, bool storeCompressed);
 		bool WriteCasFile(WriteResult& out, const tchar* fileName, const CasKey& casKey);
diff --git a/Engine/Source/Programs/UnrealBuildAccelerator/Core/Public/UbaDetoursPayload.h b/Engine/Source/Programs/UnrealBuildAccelerator/Core/Public/UbaDetoursPayload.h
--- a/Engine/Source/Programs/UnrealBuildAccelerator/Core/Public/UbaDetoursPayload.h
+++ b/Engine/Source/Programs/UnrealBuildAccelerator/Core/Public/UbaDetoursPayload.h
@@ -23,6 +23,7 @@ namespace uba
 		bool useCustomAllocator = true;
 		bool isRunningWine = false;
 		bool isChild = false;
+		bool storeObjFilesCompressed = false;
 		int uiLanguage = 0;
 		StringBuffer<256> logFile;
 	};
diff --git a/Engine/Source/Programs/UnrealBuildAccelerator/Detours/Private/Windows/UbaDetoursFunctionsKernelBase.inl b/Engine/Source/Programs/UnrealBuildAccelerator/Detours/Private/Windows/UbaDetoursFunctionsKernelBase.inl
--- a/Engine/Source/Programs/UnrealBuildAccelerator/Detours/Private/Windows/UbaDetoursFunctionsKernelBase.inl
+++ b/Engine/Source/Programs/UnrealBuildAccelerator/Detours/Private/Windows/UbaDetoursFunctionsKernelBase.inl
@@ -3681,7 +3681,7 @@ BOOL Detoured_PathIsDirectoryEmptyW(LPCWSTR pszPath)
 
 HRESULT Detoured_SHCreateStreamOnFileW(LPCWSTR pszFile, DWORD grfMode, IStream** ppstm)
 {
-	UBA_ASSERTF(!g_runningRemote, L"%ls", pszFile);
+	//UBA_ASSERTF(!g_runningRemote, L"%ls", pszFile);
 	return True_SHCreateStreamOnFileW(pszFile, grfMode, ppstm);
 }
 
diff --git a/Engine/Source/Programs/UnrealBuildAccelerator/Detours/Private/Windows/UbaDetoursFunctionsNtDll.inl b/Engine/Source/Programs/UnrealBuildAccelerator/Detours/Private/Windows/UbaDetoursFunctionsNtDll.inl
--- a/Engine/Source/Programs/UnrealBuildAccelerator/Detours/Private/Windows/UbaDetoursFunctionsNtDll.inl
+++ b/Engine/Source/Programs/UnrealBuildAccelerator/Detours/Private/Windows/UbaDetoursFunctionsNtDll.inl
@@ -884,6 +884,8 @@ NTSTATUS NTAPI Shared_NtCreateFile(bool IsCreateFunc, PHANDLE hFileHandle, ACCES
 	}
 	else
 	{
+		g_objFilesPreloader.Wait(fileNameKey); // Wait for on-going preload if existing
+
 		if (!info.originalName)
 			info.originalName = g_memoryBlock.Strdup(fileName.data);
 		if (isWrite) //(info.lastDesiredAccess != dwDesiredAccess)
@@ -997,10 +999,6 @@ NTSTATUS NTAPI Shared_NtCreateFile(bool IsCreateFunc, PHANDLE hFileHandle, ACCES
 
 	if (keepInMemory || info.memoryFile)
 	{
-		auto fileObject = new FileObject();
-		DetouredHandle* dh = new DetouredHandle(HandleType_File);
-		dh->fileObject = fileObject;
-
 		if (!info.memoryFile)
 		{
 			if (NeedsSharedMemory(fileName.data))
@@ -1069,6 +1067,9 @@ NTSTATUS NTAPI Shared_NtCreateFile(bool IsCreateFunc, PHANDLE hFileHandle, ACCES
 		}
 		_.Leave();
 
+		auto fileObject = new FileObject();
+		DetouredHandle* dh = new DetouredHandle(HandleType_File);
+		dh->fileObject = fileObject;
 		dh->dirTableOffset = dirTableOffset;
 		dh->fileObject->desiredAccess = dwDesiredAccess;
 		dh->fileObject->closeId = closeId;
@@ -1082,6 +1083,22 @@ NTSTATUS NTAPI Shared_NtCreateFile(bool IsCreateFunc, PHANDLE hFileHandle, ACCES
 		return STATUS_SUCCESS;
 	}
 
+	if (lpFileName[0] == ':')
+	{
+		auto fileObject = new FileObject();
+		DetouredHandle* dh = new DetouredHandle(HandleType_File);
+		dh->fileObject = fileObject;
+		dh->dirTableOffset = dirTableOffset;
+		dh->fileObject->desiredAccess = dwDesiredAccess;
+		dh->fileObject->closeId = closeId;
+		dh->fileObject->fileInfo = &info;
+		dh->fileObject->deleteOnClose = isDeleteOnClose;
+		*hFileHandle = makeDetouredHandle(dh);
+
+		DEBUG_LOG_DETOURED(funcName, L"(PRELOAD)%ls %llu (%ls) (%ls) -> Success", isWriteStr, uintptr_t(*hFileHandle), lpFileName, (fileName.data != lpFileName ? fileName.data : L""));
+		return STATUS_SUCCESS;
+	}
+
 	const wchar_t* tempFileName = lpFileName;
 	if (tempFileName[0] == '#')
 		tempFileName = fileName.data;
diff --git a/Engine/Source/Programs/UnrealBuildAccelerator/Detours/Private/Windows/UbaDetoursFunctionsWin.cpp b/Engine/Source/Programs/UnrealBuildAccelerator/Detours/Private/Windows/UbaDetoursFunctionsWin.cpp
--- a/Engine/Source/Programs/UnrealBuildAccelerator/Detours/Private/Windows/UbaDetoursFunctionsWin.cpp
+++ b/Engine/Source/Programs/UnrealBuildAccelerator/Detours/Private/Windows/UbaDetoursFunctionsWin.cpp
@@ -41,8 +41,8 @@ DETOURED_FUNCTIONS_MEMORY
 #include "UbaWinBinDependencyParser.h"
 #include "UbaDetoursPayload.h"
 #include "UbaApplicationRules.h"
-
 #include "UbaDetoursShared.h"
+#include "UbaDetoursObjFilesPreloader.h"
 
 #include "Shlwapi.h"
 #include <detours/detours.h>
@@ -127,112 +127,9 @@ void TrackInput(const wchar_t* file)
 	}
 }
 
-struct MemoryFile
-{
-	MemoryFile(u8* data = nullptr, bool localOnly = true) : baseAddress(data), isLocalOnly(localOnly) {}
-	MemoryFile(bool localOnly, u64 reserveSize_) : isLocalOnly(localOnly)
-	{
-		Reserve(reserveSize_);
-	}
-
-	void Reserve(u64 reserveSize_)
-	{
-		reserveSize = reserveSize_;
-		if (isLocalOnly)
-		{
-			baseAddress = (u8*)VirtualAlloc(NULL, reserveSize, MEM_RESERVE, PAGE_READWRITE);
-			if (!baseAddress)
-				FatalError(1354, L"VirtualAlloc failed trying to reserve %llu. (Error code: %u)", reserveSize, GetLastError());
-			mappedSize = reserveSize;
-		}
-		else
-		{
-			mappedSize = 32 * 1024 * 1024;
-			mappingHandle = True_CreateFileMappingW(INVALID_HANDLE_VALUE, NULL, PAGE_READWRITE | SEC_RESERVE, ToHigh(reserveSize), ToLow(reserveSize), NULL);
-			if (!mappingHandle)
-				FatalError(1348, L"CreateFileMappingW failed trying to reserve %llu. (Error code: %u)", reserveSize, GetLastError());
-			baseAddress = (u8*)True_MapViewOfFile(mappingHandle, FILE_MAP_READ | FILE_MAP_WRITE, 0, 0, mappedSize);
-			if (!baseAddress)
-				FatalError(1353, L"MapViewOfFile failed trying to map %llu. ReservedSize: %llu (Error code: %u)", mappedSize, reserveSize, GetLastError());
-		}
-	}
-
-	void Unreserve()
-	{
-		if (isLocalOnly)
-		{
-			VirtualFree(baseAddress, 0, MEM_RELEASE);
-		}
-		else
-		{
-			True_UnmapViewOfFile(baseAddress);
-			CloseHandle(mappingHandle);
-			mappingHandle = nullptr;
-		}
-		baseAddress = nullptr;
-		committedSize = 0;
-	}
-
-	void Write(struct DetouredHandle& handle, LPCVOID lpBuffer, u64 nNumberOfBytesToWrite);
-	void EnsureCommited(struct DetouredHandle& handle, u64 size);
-
-	u64 fileIndex = ~u64(0);
-	u64 fileTime = ~u64(0);
-	u32 volumeSerial = 0;
-
-	HANDLE mappingHandle = nullptr;
-	u8* baseAddress;
-	u64 reserveSize = 0;
-	u64 mappedSize = 0;
-	u64 committedSize = 0;
-	u64 writtenSize = 0;
-	bool isLocalOnly;
-	bool isReported = false;
-};
 u8 g_emptyMemoryFileMem;
 MemoryFile& g_emptyMemoryFile = *new MemoryFile(&g_emptyMemoryFileMem, true);
 
-struct FileObject
-{
-	void* operator new(size_t size);
-	void operator delete(void* p);
-	FileInfo* fileInfo = nullptr;
-	u32 refCount = 1;
-	u32 closeId = 0;
-	u32 desiredAccess = 0;
-	bool deleteOnClose = false;
-	bool ownsFileInfo = false;
-	TString newName;
-};
-BlockAllocator<FileObject> g_fileObjectAllocator(g_memoryBlock);
-void* FileObject::operator new(size_t size) { return g_fileObjectAllocator.Allocate(); }
-void FileObject::operator delete(void* p) { g_fileObjectAllocator.Free(p); }
-
-
-enum HandleType
-{
-	HandleType_File,
-	HandleType_FileMapping,
-	HandleType_Process,
-	HandleType_Std,
-};
-
-struct DetouredHandle
-{
-	void* operator new(size_t size);
-	void operator delete(void* p);
-
-	DetouredHandle(HandleType t, HANDLE th = INVALID_HANDLE_VALUE) : trueHandle(th), type(t) {}
-
-	HANDLE trueHandle;
-	u32 dirTableOffset = ~u32(0);
-	HandleType type;
-
-	// Only for files
-	FileObject* fileObject = nullptr;
-    u64 pos = 0;
-};
-
 constexpr u64 DetouredHandleMaxCount = 200*1024; // ~200000 handles enough?
 constexpr u64 DetouredHandleStart = 300000; // Let's hope noone uses the handles starting at 300000! :)
 constexpr u64 DetouredHandleEnd = DetouredHandleStart + DetouredHandleMaxCount;
@@ -296,6 +193,9 @@ ReaderWriterLock g_loadedModulesLock;
 UnorderedMap<HMODULE, TString> g_loadedModules;
 u64 g_memoryFileIndexCounter = ~u64(0) - 1000000; // I really hope this will not collide with anything
 
+ObjFilesPreloader g_objFilesPreloader;
+
+
 struct SuppressCreateFileDetourScope
 {
 	SuppressCreateFileDetourScope() { ++t_disallowCreateFileDetour; }
@@ -965,12 +865,17 @@ void Init(const DetoursPayload& payload, u64 startTime)
 	UBA_ASSERT(directoryTableMem);
 	g_directoryTable.Init(directoryTableMem, directoryTableCount, directoryTableSize);
 
+	if (payload.storeObjFilesCompressed && g_rulesIndex == 2) 	// Rules index 2 is link.exe.. compressed obj files is experimental so solution is a bit hacky
+		g_objFilesPreloader.Start(cmdLine);
+
 	g_stats.attach.time += GetTime() - startTime;
 	g_stats.attach.count = 1;
 }
 
 void Deinit(u64 startTime)
 {
+	g_objFilesPreloader.Stop();
+
 	if (g_isRunningWine) // mt.exe etc fails if detaching is not done during shutdown
 	{
 		DetourTransactionBegin();
diff --git a/Engine/Source/Programs/UnrealBuildAccelerator/Detours/Private/Windows/UbaDetoursObjFilesPreloader.h b/Engine/Source/Programs/UnrealBuildAccelerator/Detours/Private/Windows/UbaDetoursObjFilesPreloader.h
new file mode 100644
--- /dev/null
+++ b/Engine/Source/Programs/UnrealBuildAccelerator/Detours/Private/Windows/UbaDetoursObjFilesPreloader.h
@@ -0,0 +1,310 @@
+// Copyright Epic Games, Inc. All Rights Reserved.
+
+#pragma once
+
+#include "UbaDetoursUtilsWin.h"
+#include "UbaEvent.h"
+#include <oodle2.h>
+#define _NTDEF_
+#include <ntsecapi.h>
+
+namespace uba
+{
+	constexpr bool UseLargePages = false;
+
+	u64 GetLargePageSize()
+	{
+		HINSTANCE hDll = LoadLibrary(TEXT("kernel32.dll"));
+		if (!hDll)
+			return 0;
+		auto lg = MakeGuard([hDll]() { FreeLibrary(hDll); });
+
+		using GetLargePageMinimumType = int(void);
+		auto getLargePageMinimumFunc = (GetLargePageMinimumType*)GetProcAddress(hDll, "GetLargePageMinimum");
+		if (getLargePageMinimumFunc == NULL)
+			return 0;
+
+		DWORD size = (*getLargePageMinimumFunc)();
+		if (!size)
+			return 0;
+
+		HANDLE hToken;
+		if (!OpenProcessToken(GetCurrentProcess(), TOKEN_ADJUST_PRIVILEGES | TOKEN_QUERY, &hToken))
+			return 0;
+		auto tg = MakeGuard([hToken]() { CloseHandle(hToken); });
+
+		TOKEN_PRIVILEGES tp;
+		tp.PrivilegeCount = 1;
+		tp.Privileges[0].Attributes = SE_PRIVILEGE_ENABLED;
+		if (!LookupPrivilegeValue(NULL, L"SeLockMemoryPrivilege", &tp.Privileges[0].Luid))
+			return 0;
+
+		if (!AdjustTokenPrivileges(hToken, FALSE, &tp, 0, (PTOKEN_PRIVILEGES)NULL, 0))
+			return 0;
+
+		// It is possible for AdjustTokenPrivileges to return TRUE and still not succeed. So always check for the last error value.
+		if (GetLastError() != ERROR_SUCCESS)
+			return 0;
+
+		return size;
+	}
+
+	class ObjFilesPreloader
+	{
+	public:
+		void Start(const wchar_t* cmdLine)
+		{
+			const wchar_t* at = wcschr(cmdLine, '@');
+			if (!at)
+				return;
+
+			const wchar_t* end = at + wcslen(at);
+			StringBuffer<> rspFile;
+			++at;
+			if (*at == '\"')
+			{
+				++at;
+				--end;
+			}
+			rspFile.Append(at, end - at);
+
+			HANDLE rspFileHandle = CreateFileW(rspFile.data, GENERIC_READ, FILE_SHARE_READ, NULL, OPEN_EXISTING, FILE_ATTRIBUTE_NORMAL, NULL);
+			if (rspFileHandle == INVALID_HANDLE_VALUE)
+				return;
+			LARGE_INTEGER rspFileSize;
+			if (!GetFileSizeEx(rspFileHandle, &rspFileSize))
+			{
+				CloseHandle(rspFileHandle);
+				return;
+			}
+			HANDLE rspFileMappingHandle = CreateFileMappingW(rspFileHandle, NULL, PAGE_READONLY, rspFileSize.HighPart, rspFileSize.LowPart, NULL);
+			CloseHandle(rspFileHandle);
+			if (!rspFileMappingHandle)
+				return;
+			auto rspMem = (char*)MapViewOfFile(rspFileMappingHandle, FILE_MAP_READ, 0, 0, rspFileSize.QuadPart);
+			CloseHandle(rspFileMappingHandle);
+			if (!rspMem)
+				return;
+
+			auto lineBegin = rspMem;
+			while (*lineBegin)
+			{
+				auto lineEnd = strchr(lineBegin, '\n');
+				if (!lineEnd)
+					lineEnd = lineBegin + strlen(lineBegin);
+				if (*lineBegin != '/')
+				{
+					char line[512];
+					u64 lineLen = lineEnd - lineBegin;
+					memcpy(line, lineBegin, lineLen);
+					line[lineLen] = 0;
+
+					if (char* ext = strstr(line, ".obj"))
+					{
+						ext[4] = 0;
+
+						char* objFile = line;
+						if (line[0] == '\"')
+							++objFile;
+
+						StringBuffer<> objFile2;
+						objFile2.Append(objFile);
+
+						StringBuffer<> objFileFull;
+						FixPath(objFileFull, objFile2.data);
+
+						HANDLE objFileHandle = CreateFileW(objFileFull.data, GENERIC_READ, FILE_SHARE_READ, NULL, OPEN_EXISTING, FILE_ATTRIBUTE_NORMAL, NULL);
+						UBA_ASSERT(objFileHandle != INVALID_HANDLE_VALUE);
+
+						StringKey fileNameKey = ToStringKey(objFileFull.MakeLower());
+						auto insres = m_preloadedObjFiles.try_emplace(fileNameKey);
+						UBA_ASSERT(insres.second);
+						insres.first->second.handle = objFileHandle;
+						insres.first->second.event.Create(true);
+
+						SCOPED_READ_LOCK(g_mappedFileTable.m_lookupLock, _);
+						auto findIt = g_mappedFileTable.m_lookup.find(fileNameKey);
+						insres.first->second.fileInfo = &findIt->second;
+					}
+				}
+
+				if (!*lineEnd)
+					break;
+				lineBegin = lineEnd + 1;
+			}
+
+			if (!m_preloadedObjFiles.empty())
+				m_threadHandle = CreateThread(NULL, 0, [](LPVOID p) -> DWORD { ((ObjFilesPreloader*)p)->ThreadPreload(); return 0; }, this, 0, NULL);
+
+			UnmapViewOfFile(rspMem);
+		}
+
+		void Stop()
+		{
+			if (!m_threadHandle)
+				return;
+			WaitForSingleObject(m_threadHandle, INFINITE);
+			CloseHandle(m_threadHandle);
+			WaitForMultipleObjects(sizeof_array(m_helperThreadHandles), m_helperThreadHandles, true, INFINITE);
+			for (u32 i=0;i!=sizeof_array(m_helperThreadHandles); ++i)
+				CloseHandle(m_helperThreadHandles[i]);
+		}
+
+		void ThreadPreload()
+		{
+			u64 totalMemSize = 0;
+			for (auto& kv : m_preloadedObjFiles)
+			{
+				Preload& preload = kv.second;
+				HANDLE objFileHandle = preload.handle;
+				if (!objFileHandle)
+					continue;
+
+				LARGE_INTEGER objFileSize;
+				BOOL res = GetFileSizeEx(objFileHandle, &objFileSize);
+				UBA_ASSERT(res);(void)res;
+
+				HANDLE objFileMappingHandle = CreateFileMappingW(objFileHandle, NULL, PAGE_READONLY, objFileSize.HighPart, objFileSize.LowPart, NULL);
+				UBA_ASSERT(objFileMappingHandle);
+				if (!objFileMappingHandle)
+					continue;
+				CloseHandle(objFileHandle);
+				preload.handle = NULL;
+				preload.objMem = (u8*)MapViewOfFile(objFileMappingHandle, FILE_MAP_READ, 0, 0, objFileSize.QuadPart);
+				UBA_ASSERT(preload.objMem);
+				if (!preload.objMem)
+					continue;
+				preload.objCompressedSize = objFileSize.QuadPart;
+				preload.objDecompressedSize = *(u64*)preload.objMem;
+				preload.objReadOffset = 8;
+				preload.objWriteOffset = 0;
+				preload.objLeft = preload.objDecompressedSize;
+				CloseHandle(objFileMappingHandle);
+				totalMemSize += preload.objDecompressedSize;
+			}
+
+			u64 alignedMem = AlignUp(totalMemSize, 64*1024);
+			DWORD allocationType = MEM_RESERVE | MEM_COMMIT;
+
+			if (UseLargePages)
+			{
+				if (u64 largePageSize = GetLargePageSize())
+				{
+					alignedMem = AlignUp(totalMemSize, largePageSize);
+					allocationType |= MEM_LARGE_PAGES;
+				}
+			}
+			m_totalMem = (u8*)VirtualAlloc(NULL, alignedMem, allocationType, PAGE_READWRITE);
+			m_it = m_preloadedObjFiles.begin();
+
+			for (u32 i=0;i!=sizeof_array(m_helperThreadHandles); ++i)
+				m_helperThreadHandles[i] = CreateThread(NULL, 0, [](LPVOID p) -> DWORD { ((ObjFilesPreloader*)p)->ThreadHelper(); return 0; }, this, 0, NULL);
+			ThreadHelper();
+		}
+
+		void ThreadHelper()
+		{
+			OO_SINTa decoredMemSize = OodleLZDecoder_MemorySizeNeeded(OodleLZ_Compressor_Kraken);
+			void* decoderMem = malloc(decoredMemSize);
+			auto mg = MakeGuard([decoderMem]() { free(decoderMem); });
+
+			while (true)
+			{
+				SCOPED_WRITE_LOCK(m_threadLock, lock);
+				if (m_it == m_preloadedObjFiles.end())
+					break;
+				Preload& preload = m_it->second;
+				if (preload.objReadOffset == preload.objCompressedSize)
+				{
+					++m_it;
+					continue;
+				}
+
+				FileInfo& info = *preload.fileInfo;
+
+				if (preload.objWriteOffset == 0)
+				{
+					info.fileMapMem = m_totalMem + m_totalMemOffset;
+					m_totalMemOffset += preload.objDecompressedSize;
+				}
+
+				BinaryReader reader(preload.objMem, preload.objReadOffset, preload.objCompressedSize);
+				u32 compressedBlockSize = reader.ReadU32();
+				u32 decompressedBlockSize = reader.ReadU32();
+
+				u8* destMem = info.fileMapMem + preload.objWriteOffset;
+
+				preload.objReadOffset += compressedBlockSize + 8;
+				preload.objWriteOffset += decompressedBlockSize;
+
+				lock.Leave();
+
+				OO_SINTa decompLen = OodleLZ_Decompress(reader.GetPositionData(), (OO_SINTa)compressedBlockSize, destMem, (OO_SINTa)decompressedBlockSize,
+					OodleLZ_FuzzSafe_Yes, OodleLZ_CheckCRC_No, OodleLZ_Verbosity_None, NULL, 0, NULL, NULL, decoderMem, decoredMemSize);
+				UBA_ASSERT(decompLen == decompressedBlockSize);(void)decompLen;
+
+				bool isDone = preload.objLeft.fetch_sub(decompressedBlockSize) == decompressedBlockSize;
+				if (!isDone)
+					continue;
+
+				info.size = preload.objDecompressedSize;
+				info.memoryFile = nullptr;
+				info.name = L":";
+				info.isFileMap = true;
+				info.fileMapMemEnd = info.fileMapMem + info.size;
+				info.trueFileMapHandle = nullptr;
+				info.trueFileMapOffset = 0;
+
+				{
+					SCOPED_WRITE_LOCK(g_mappedFileTable.m_memLookupLock, lock3);
+					auto insres = g_mappedFileTable.m_memLookup.try_emplace(destMem);
+					UBA_ASSERT(insres.second);
+					insres.first->second = 1;
+				}
+
+				preload.event.Set();
+
+				UnmapViewOfFile(preload.objMem);
+			}
+		}
+
+		void Wait(const StringKey& key)
+		{
+			auto findIt = m_preloadedObjFiles.find(key);
+			if (findIt == m_preloadedObjFiles.end())
+				return;
+			//u64 start = GetTime();
+			findIt->second.event.IsSet();
+			//u64 waitTime = GetTime() - start;
+			//if (TimeToMs(waitTime) > 2)
+			//{
+			//	StringBuffer b;
+			//	b.Appendf(L"Waited %s for preload", TimeToText(waitTime).str);
+			//	Rpc_WriteLog(b.data, b.count, true, false);
+			//}
+		}
+
+		struct Preload
+		{
+			Event event;
+			HANDLE handle;
+			u8* objMem;
+			u64 objCompressedSize;
+			u64 objDecompressedSize;
+			FileInfo* fileInfo;
+			u64 objReadOffset;
+			u64 objWriteOffset;
+			Atomic<u64> objLeft;
+		};
+
+		using PreloadedObjFiles = Map<StringKey, Preload>;
+		PreloadedObjFiles m_preloadedObjFiles;
+		HANDLE m_threadHandle = 0;
+		HANDLE m_helperThreadHandles[5];
+		u8* m_totalMem = nullptr;
+		u64 m_totalMemOffset = 0;
+
+		ReaderWriterLock m_threadLock;
+		PreloadedObjFiles::iterator m_it;
+	};
+}
diff --git a/Engine/Source/Programs/UnrealBuildAccelerator/Detours/Private/Windows/UbaDetoursUtilsWin.h b/Engine/Source/Programs/UnrealBuildAccelerator/Detours/Private/Windows/UbaDetoursUtilsWin.h
new file mode 100644
--- /dev/null
+++ b/Engine/Source/Programs/UnrealBuildAccelerator/Detours/Private/Windows/UbaDetoursUtilsWin.h
@@ -0,0 +1,115 @@
+// Copyright Epic Games, Inc. All Rights Reserved.
+
+#pragma once
+
+#include "UbaDetoursFunctionsWin.h"
+#include "UbaDetoursFileMappingTable.h"
+
+namespace uba
+{
+	struct FileObject
+	{
+		void* operator new(size_t size);
+		void operator delete(void* p);
+		FileInfo* fileInfo = nullptr;
+		u32 refCount = 1;
+		u32 closeId = 0;
+		u32 desiredAccess = 0;
+		bool deleteOnClose = false;
+		bool ownsFileInfo = false;
+		TString newName;
+	};
+	BlockAllocator<FileObject> g_fileObjectAllocator(g_memoryBlock);
+	void* FileObject::operator new(size_t size) { return g_fileObjectAllocator.Allocate(); }
+	void FileObject::operator delete(void* p) { g_fileObjectAllocator.Free(p); }
+
+
+	enum HandleType
+	{
+		HandleType_File,
+		HandleType_FileMapping,
+		HandleType_Process,
+		HandleType_Std,
+	};
+
+
+	struct DetouredHandle
+	{
+		void* operator new(size_t size);
+		void operator delete(void* p);
+
+		DetouredHandle(HandleType t, HANDLE th = INVALID_HANDLE_VALUE) : trueHandle(th), type(t) {}
+
+		HANDLE trueHandle;
+		u32 dirTableOffset = ~u32(0);
+		HandleType type;
+
+		// Only for files
+		FileObject* fileObject = nullptr;
+		u64 pos = 0;
+	};
+
+
+	struct MemoryFile
+	{
+		MemoryFile(u8* data = nullptr, bool localOnly = true) : baseAddress(data), isLocalOnly(localOnly) {}
+		MemoryFile(bool localOnly, u64 reserveSize_) : isLocalOnly(localOnly)
+		{
+			Reserve(reserveSize_);
+		}
+
+		void Reserve(u64 reserveSize_)
+		{
+			reserveSize = reserveSize_;
+			if (isLocalOnly)
+			{
+				baseAddress = (u8*)VirtualAlloc(NULL, reserveSize, MEM_RESERVE, PAGE_READWRITE);
+				if (!baseAddress)
+					FatalError(1354, L"VirtualAlloc failed trying to reserve %llu. (Error code: %u)", reserveSize, GetLastError());
+				mappedSize = reserveSize;
+			}
+			else
+			{
+				mappedSize = 32 * 1024 * 1024;
+				mappingHandle = True_CreateFileMappingW(INVALID_HANDLE_VALUE, NULL, PAGE_READWRITE | SEC_RESERVE, ToHigh(reserveSize), ToLow(reserveSize), NULL);
+				if (!mappingHandle)
+					FatalError(1348, L"CreateFileMappingW failed trying to reserve %llu. (Error code: %u)", reserveSize, GetLastError());
+				baseAddress = (u8*)True_MapViewOfFile(mappingHandle, FILE_MAP_READ | FILE_MAP_WRITE, 0, 0, mappedSize);
+				if (!baseAddress)
+					FatalError(1353, L"MapViewOfFile failed trying to map %llu. ReservedSize: %llu (Error code: %u)", mappedSize, reserveSize, GetLastError());
+			}
+		}
+
+		void Unreserve()
+		{
+			if (isLocalOnly)
+			{
+				VirtualFree(baseAddress, 0, MEM_RELEASE);
+			}
+			else
+			{
+				True_UnmapViewOfFile(baseAddress);
+				CloseHandle(mappingHandle);
+				mappingHandle = nullptr;
+			}
+			baseAddress = nullptr;
+			committedSize = 0;
+		}
+
+		void Write(struct DetouredHandle& handle, LPCVOID lpBuffer, u64 nNumberOfBytesToWrite);
+		void EnsureCommited(struct DetouredHandle& handle, u64 size);
+
+		u64 fileIndex = ~u64(0);
+		u64 fileTime = ~u64(0);
+		u32 volumeSerial = 0;
+
+		HANDLE mappingHandle = nullptr;
+		u8* baseAddress;
+		u64 reserveSize = 0;
+		u64 mappedSize = 0;
+		u64 committedSize = 0;
+		u64 writtenSize = 0;
+		bool isLocalOnly;
+		bool isReported = false;
+	};
+}
\ No newline at end of file
diff --git a/Engine/Source/Programs/UnrealBuildAccelerator/Detours/UbaDetours.Build.cs b/Engine/Source/Programs/UnrealBuildAccelerator/Detours/UbaDetours.Build.cs
--- a/Engine/Source/Programs/UnrealBuildAccelerator/Detours/UbaDetours.Build.cs
+++ b/Engine/Source/Programs/UnrealBuildAccelerator/Detours/UbaDetours.Build.cs
@@ -19,6 +19,7 @@ public class UbaDetours : ModuleRules
 		{
 			PrivateDependencyModuleNames.AddRange(new string[] {
 				"Detours",
+				"OodleDataCompression",
 			});
 
 			PublicSystemLibraries.AddRange(new string[] {
diff --git a/Engine/Source/Programs/UnrealBuildAccelerator/Host/Private/UbaExports.cpp b/Engine/Source/Programs/UnrealBuildAccelerator/Host/Private/UbaExports.cpp
--- a/Engine/Source/Programs/UnrealBuildAccelerator/Host/Private/UbaExports.cpp
+++ b/Engine/Source/Programs/UnrealBuildAccelerator/Host/Private/UbaExports.cpp
@@ -322,7 +322,7 @@ uba::StorageClient* CreateStorageClient(uba::NetworkClient& client, const uba::t
 	}
 */
 
-	uba::SessionServerCreateInfo* CreateSessionServerCreateInfo(uba::Storage& storage, uba::NetworkServer& client, uba::LogWriter& writer, const uba::tchar* rootDir, const uba::tchar* traceOutputFile, bool disableCustomAllocator, bool launchVisualizer, bool resetCas, bool writeToDisk, bool detailedTrace, bool allowWaitOnMem, bool allowKillOnMem)
+	uba::SessionServerCreateInfo* CreateSessionServerCreateInfo(uba::Storage& storage, uba::NetworkServer& client, uba::LogWriter& writer, const uba::tchar* rootDir, const uba::tchar* traceOutputFile, bool disableCustomAllocator, bool launchVisualizer, bool resetCas, bool writeToDisk, bool detailedTrace, bool allowWaitOnMem, bool allowKillOnMem, bool storeObjFilesCompressed)
 	{
 		auto info = new uba::SessionServerCreateInfo(storage, client, writer);
 		info->rootDir = TStrdup(rootDir);
@@ -334,6 +334,7 @@ uba::StorageClient* CreateStorageClient(uba::NetworkClient& client, const uba::t
 		info->detailedTrace = detailedTrace;
 		info->allowWaitOnMem = allowWaitOnMem;
 		info->allowKillOnMem = allowKillOnMem;
+		info->storeObjFilesCompressed = storeObjFilesCompressed;
 		//info->remoteTraceEnabled = true;
 		//info->remoteLogEnabled = true;
 		return info;
diff --git a/Engine/Source/Programs/UnrealBuildAccelerator/Host/Public/UbaExports.h b/Engine/Source/Programs/UnrealBuildAccelerator/Host/Public/UbaExports.h
--- a/Engine/Source/Programs/UnrealBuildAccelerator/Host/Public/UbaExports.h
+++ b/Engine/Source/Programs/UnrealBuildAccelerator/Host/Public/UbaExports.h
@@ -81,7 +81,7 @@ extern "C"
 	using SessionServer_CustomServiceFunction = uba::u32(uba::ProcessHandle* handle, const void* recv, uba::u32 recvSize, void* send, uba::u32 sendCapacity, void* userData);
 
 	UBA_API uba::SessionServerCreateInfo* CreateSessionServerCreateInfo(uba::Storage& storage, uba::NetworkServer& client, uba::LogWriter& writer, const uba::tchar* rootDir, const uba::tchar* traceOutputFile,
-		bool disableCustomAllocator, bool launchVisualizer, bool resetCas, bool writeToDisk, bool detailedTrace, bool allowWaitOnMem = false, bool allowKillOnMem = false);
+		bool disableCustomAllocator, bool launchVisualizer, bool resetCas, bool writeToDisk, bool detailedTrace, bool allowWaitOnMem = false, bool allowKillOnMem = false, bool storeObjFilesCompressed = false);
 	UBA_API void DestroySessionServerCreateInfo(uba::SessionServerCreateInfo* info);
 	UBA_API uba::SessionServer* CreateSessionServer(const uba::SessionServerCreateInfo& info);
 	UBA_API void SessionServer_SetRemoteProcessAvailable(uba::SessionServer* server, SessionServer_RemoteProcessAvailableCallback* available, void* userData);
