UBA: ::CreateFileMappingW fails for zero-length files

#rnx

diff --git a/Engine/Source/Programs/UnrealBuildAccelerator/Common/Private/UbaSession.cpp b/Engine/Source/Programs/UnrealBuildAccelerator/Common/Private/UbaSession.cpp
--- a/Engine/Source/Programs/UnrealBuildAccelerator/Common/Private/UbaSession.cpp
+++ b/Engine/Source/Programs/UnrealBuildAccelerator/Common/Private/UbaSession.cpp
@@ -1990,7 +1990,7 @@ namespace uba
 				// Seems like best combo (for windows at least) is to use writes with overlap and max 16 at the same time.
 				// On one machine we get twice as fast without overlap if no bottleneck. On another machine (ntfs compression on) we get twice as slow without overlap
 				// Both machines behaves well with overlap AND bottleneck. Both machine are 128 logical core thread rippers.
-				constexpr bool useFileMapForWrite = true;
+				bool useFileMapForWrite = fileSize > 0; // ::CreateFileMappingW does not work for zero-length files
 				bool useOverlap = false;//fileSize > 8 * 1024 * 1024;
 
 
