[Uba]
* Fixed a bug where path environment variable missed last path in colon-separated path list

diff --git a/Engine/Source/Programs/UnrealBuildAccelerator/Common/Private/UbaSession.cpp b/Engine/Source/Programs/UnrealBuildAccelerator/Common/Private/UbaSession.cpp
--- a/Engine/Source/Programs/UnrealBuildAccelerator/Common/Private/UbaSession.cpp
+++ b/Engine/Source/Programs/UnrealBuildAccelerator/Common/Private/UbaSession.cpp
@@ -1606,12 +1606,17 @@ namespace uba
 
 			const char* start = env + 5;
 			const char* it = start;
-			while (*it)
+			bool isLast = false;
+			while (!isLast)
 			{
 				if (*it != ':')
 				{
-					++it;
-					continue;
+					if (*it)
+					{
+						++it;
+						continue;
+					}
+					isLast = true;
 				}
 
 				const char* s = start;
