UBA: Fix up generated solution to show common and core

#rnx

diff --git a/Engine/Source/Programs/UnrealBuildAccelerator/UbaCommon.Target.cs b/Engine/Source/Programs/UnrealBuildAccelerator/UbaCommon.Target.cs
new file mode 100644
--- /dev/null
+++ b/Engine/Source/Programs/UnrealBuildAccelerator/UbaCommon.Target.cs
@@ -0,0 +1,14 @@
+// Copyright Epic Games, Inc. All Rights Reserved.
+
+using UnrealBuildTool;
+
+[SupportedPlatforms()]
+public class UbaCommonTarget : TargetRules
+{
+	public UbaCommonTarget(TargetInfo Target) : base(Target)
+	{
+		LaunchModuleName = "UbaCommon";
+		bShouldCompileAsDLL = true;
+		UbaAgentTarget.CommonUbaSettings(this, Target);
+	}
+}
diff --git a/Engine/Source/Programs/UnrealBuildAccelerator/UbaCore.Target.cs b/Engine/Source/Programs/UnrealBuildAccelerator/UbaCore.Target.cs
new file mode 100644
--- /dev/null
+++ b/Engine/Source/Programs/UnrealBuildAccelerator/UbaCore.Target.cs
@@ -0,0 +1,14 @@
+// Copyright Epic Games, Inc. All Rights Reserved.
+
+using UnrealBuildTool;
+
+[SupportedPlatforms()]
+public class UbaCoreTarget : TargetRules
+{
+	public UbaCoreTarget(TargetInfo Target) : base(Target)
+	{
+		LaunchModuleName = "UbaCore";
+		bShouldCompileAsDLL = true;
+		UbaAgentTarget.CommonUbaSettings(this, Target);
+	}
+}
