[UBA]
* Added so process input dependencies can be tracked on remote compiles and sent back to host.
* Changed so ApplicationRules is a pointer on the ProcessStartInfo instead.. this in preparation of being able to have custom application rules but also be able to query rules before process instance is created
* Added support for suppressing error log if FileAccessor open file fails
* Fixed so BytesToText show bytes when under 1kb
* Fixed potential shutdown race related to process exiting and session instance being destroyed
* Some minor fixes here and there related to logging
* Changed process IsDetoured to GetExecutionType enum

diff --git a/Engine/Source/Programs/UnrealBuildAccelerator/Cli/Private/UbaCli.cpp b/Engine/Source/Programs/UnrealBuildAccelerator/Cli/Private/UbaCli.cpp
--- a/Engine/Source/Programs/UnrealBuildAccelerator/Cli/Private/UbaCli.cpp
+++ b/Engine/Source/Programs/UnrealBuildAccelerator/Cli/Private/UbaCli.cpp
@@ -658,7 +658,7 @@ namespace uba
 					StringBuffer<128> extra;
 					if (ph.IsRemote())
 						extra.Append(TC(" [RemoteExecutor: ")).Append(ph.GetExecutingHost()).Append(']');
-					else if (!ph.IsDetoured())
+					else if (ph.GetExecutionType() == ProcessExecutionType_Native)
 						extra.Append(TC(" (Not detoured)"));
 					logger.Info(TC("[%u/%u] %s%s"), c, queued, desc, extra.data);
 					for (auto& line : ph.GetLogLines())
diff --git a/Engine/Source/Programs/UnrealBuildAccelerator/Common/Private/UbaFileAccessor.cpp b/Engine/Source/Programs/UnrealBuildAccelerator/Common/Private/UbaFileAccessor.cpp
--- a/Engine/Source/Programs/UnrealBuildAccelerator/Common/Private/UbaFileAccessor.cpp
+++ b/Engine/Source/Programs/UnrealBuildAccelerator/Common/Private/UbaFileAccessor.cpp
@@ -321,10 +321,10 @@ namespace uba
 		return false;
 	}
 
-	bool FileAccessor::OpenMemoryRead(u64 offset)
+	bool FileAccessor::OpenMemoryRead(u64 offset, bool errorOnFail)
 	{
 		if (!OpenFileSequentialRead(m_logger, m_fileName, m_fileHandle))
-			return m_logger.Error(TC("Failed to open file %s for read"), m_fileName);
+			return errorOnFail ? m_logger.Error(TC("Failed to open file %s for read"), m_fileName) : false;
 
 		FileInformation info;
 		if (!GetFileInformationByHandle(info, m_logger, m_fileName, m_fileHandle))
diff --git a/Engine/Source/Programs/UnrealBuildAccelerator/Common/Private/UbaLogger.cpp b/Engine/Source/Programs/UnrealBuildAccelerator/Common/Private/UbaLogger.cpp
--- a/Engine/Source/Programs/UnrealBuildAccelerator/Common/Private/UbaLogger.cpp
+++ b/Engine/Source/Programs/UnrealBuildAccelerator/Common/Private/UbaLogger.cpp
@@ -303,7 +303,9 @@ namespace uba
 
 	BytesToText::BytesToText(u64 bytes)
 	{
-		if (bytes < 1000 * 1000)
+		if (bytes < 1000)
+			TSprintf_s(str, 32, TC("%ub"), u32(bytes));
+		else if (bytes < 1000 * 1000)
 			TSprintf_s(str, 32, TC("%.1fkb"), double(bytes) / 1000ull);
 		else if (bytes < 1000ull * 1000 * 1000)
 			TSprintf_s(str, 32, TC("%.1fmb"), double(bytes) / (1000ull * 1000));
diff --git a/Engine/Source/Programs/UnrealBuildAccelerator/Common/Private/UbaProcess.cpp b/Engine/Source/Programs/UnrealBuildAccelerator/Common/Private/UbaProcess.cpp
--- a/Engine/Source/Programs/UnrealBuildAccelerator/Common/Private/UbaProcess.cpp
+++ b/Engine/Source/Programs/UnrealBuildAccelerator/Common/Private/UbaProcess.cpp
@@ -180,7 +180,7 @@ namespace uba
 		if (m_parentProcess)
 			m_waitForParent.Create(true);
 
-		SetRulesIndex(startInfo);
+		UBA_ASSERT(m_startInfo.rules);
 
 		m_session.ProcessAdded(*this, 0);
 
@@ -510,8 +510,6 @@ namespace uba
 
 		UBA_ASSERT(!m_parentProcess || !m_parentProcess->m_hasExited);
 
-		m_session.ProcessExited(*this, m_processStats.wallTime);
-
 		m_hasExited = true;
 
 		{
@@ -539,6 +537,9 @@ namespace uba
 			exitedFunc(userData, h);
 			h.m_process = nullptr;
 		}
+
+		// Must be done last to make sure shutdown is not racing
+		m_session.ProcessExited(*this, m_processStats.wallTime);
 	}
 
 	bool ProcessImpl::HandleMessage(BinaryReader& reader, BinaryWriter& writer)
@@ -788,7 +789,7 @@ namespace uba
 		u32 detoursLibLen = u32(m_session.m_detoursLibrary.size());
 
 		writer.WriteU32(childProcessId);
-		writer.WriteU32(process.m_rulesIndex);
+		writer.WriteU32(info.rules->index);
 		writer.WriteU32(detoursLibLen);
 		writer.WriteBytes(detoursLib, detoursLibLen);
 
@@ -924,7 +925,7 @@ namespace uba
 		m_processStats.Add(stats);
 
 		if (!IsCancelled())
-			if (m_startInfo.writeOutputFilesOnFail || GetApplicationRules()[m_rulesIndex].rules->IsExitCodeSuccess(m_nativeProcessExitCode))
+			if (m_startInfo.writeOutputFilesOnFail || m_startInfo.rules->IsExitCodeSuccess(m_nativeProcessExitCode))
 				m_messageSuccess = WriteFilesToDisk() && m_messageSuccess;
 
 		if (m_parentProcess)
@@ -1061,72 +1062,6 @@ namespace uba
 		return true;
 	}
 
-	void ProcessImpl::SetRulesIndex(const ProcessStartInfo& si)
-	{
-		u32 exeNameStart = 0;
-		u32 exeNameEnd = u32(m_virtualApplication.size());
-		size_t lastSeparator = m_virtualApplication.find_last_of(PathSeparator);
-		if (lastSeparator != -1)
-			exeNameStart = u32(lastSeparator + 1);
-		else if (m_virtualApplication[exeNameStart] == '"')
-			++exeNameStart;
-		if (m_virtualApplication[exeNameEnd - 1] == '"')
-			--exeNameEnd;
-		StringBuffer<128> exeName;
-		exeName.Append(m_virtualApplication.c_str() + exeNameStart, exeNameEnd - exeNameStart);
-		
-		auto rules = GetApplicationRules();
-		
-		while (true)
-		{
-			for (u32 i = 1;; ++i)
-			{
-				const tchar* app = rules[i].app;
-				if (!app)
-					break;
-				if (!exeName.Equals(app))
-					continue;
-				m_rulesIndex = i;
-				return;
-			}
-
-			if (!exeName.Equals(TC("dotnet.exe")))
-				return;
-			
-			u32 firstArgumentStart = 0;
-			u32 firstArgumentEnd = 0;
-			bool quoted = false;
-			for (u32 i = 0, e = u32(m_arguments.size()); i != e; ++i)
-			{
-				tchar c = m_arguments[i];
-				if (firstArgumentEnd)
-				{
-					if (c == '\\')
-						firstArgumentStart = i + 1;
-					if ((quoted && c != '"') || (!quoted && c != ' ' && c != '\t'))
-						continue;
-					firstArgumentEnd = i;
-					break;
-				}
-				else
-				{
-					if (c == ' ' || c == '\t')
-					{
-						++firstArgumentStart;
-						continue;
-					}
-					if (c == '"')
-					{
-						++firstArgumentStart;
-						quoted = true;
-					}
-					firstArgumentEnd = firstArgumentStart + 1;
-				}
-			}
-			exeName.Clear().Append(m_arguments.data() + firstArgumentStart, firstArgumentEnd - firstArgumentStart);
-		}
-	}
-
 	bool ProcessImpl::WriteFilesToDisk()
 	{
 		Vector<WrittenFile*> files;
@@ -1141,7 +1076,20 @@ namespace uba
 				files.push_back(&kv.second);
 		}
 
-		return m_session.WriteFilesToDisk(*this, files.data(), u32(files.size()));
+		if (!m_session.WriteFilesToDisk(*this, files.data(), u32(files.size())))
+			return false;
+
+		if (m_startInfo.trackInputs)
+		{
+			u64 totalBytes = 0;
+			for (auto& kv : m_writtenFiles)
+				totalBytes += GetStringWriteSize(kv.second.name.c_str());
+			m_trackedOutputs.resize(totalBytes);
+			BinaryWriter writer(m_trackedOutputs.data(), 0, totalBytes);
+			for (auto& kv : m_writtenFiles)
+				writer.WriteString(kv.second.name);
+		}
+		return true;
 	}
 
 	const tchar* ProcessImpl::InternalGetChildLogFile(StringBufferBase& temp)
@@ -1235,7 +1183,7 @@ namespace uba
 				return ProcessCancelExitCode;
 			}
 
-			bool isDetachedProcess = GetApplicationRules()[m_rulesIndex].rules->AllowDetach() && m_detourEnabled;
+			bool isDetachedProcess = m_startInfo.rules->AllowDetach() && m_detourEnabled;
 
 			HANDLE hJob = CreateJobObject(nullptr, nullptr);
 			JOBOBJECT_EXTENDED_LIMIT_INFORMATION info = { };
@@ -1387,12 +1335,12 @@ namespace uba
 			payload.readEvent = m_readEvent.GetHandle();
 			payload.communicationHandle = communicationHandle.handle;
 			payload.communicationOffset = communicationOffset;
-			payload.rulesIndex = m_rulesIndex;
+			payload.rulesIndex = m_startInfo.rules->index;
 			payload.version = ProcessMessageVersion;
 			payload.runningRemote = runningRemote;
 			payload.isChild = m_parentProcess != nullptr;
 			payload.trackInputs = m_startInfo.trackInputs;
-			payload.useCustomAllocator = m_startInfo.useCustomAllocator && GetApplicationRules()[m_rulesIndex].rules->AllowMiMalloc();
+			payload.useCustomAllocator = m_startInfo.useCustomAllocator && m_startInfo.rules->AllowMiMalloc();
 			payload.isRunningWine = IsRunningWine();
 			payload.storeObjFilesCompressed = m_session.m_storeObjFilesCompressed;
 			payload.uiLanguage = m_startInfo.uiLanguage;
@@ -1560,7 +1508,7 @@ namespace uba
 
 				comIdVar.Append("UBA_COMID=").AppendValue(communicationHandle.uid).Append('+').AppendValue(communicationOffset);
 				workingDir.Append("UBA_CWD=").Append(m_realWorkingDir);
-				rulesStr.Append("UBA_RULES=").AppendValue(m_rulesIndex);
+				rulesStr.Append("UBA_RULES=").AppendValue(m_startInfo.rules->index);
 
 				if (*m_startInfo.logFile)
 				{
diff --git a/Engine/Source/Programs/UnrealBuildAccelerator/Common/Private/UbaProcessStartInfoHolder.h b/Engine/Source/Programs/UnrealBuildAccelerator/Common/Private/UbaProcessStartInfoHolder.h
--- a/Engine/Source/Programs/UnrealBuildAccelerator/Common/Private/UbaProcessStartInfoHolder.h
+++ b/Engine/Source/Programs/UnrealBuildAccelerator/Common/Private/UbaProcessStartInfoHolder.h
@@ -60,6 +60,7 @@ namespace uba
 			writer.WriteString(workingDir);
 			writer.WriteString(logFile);
 			writer.WriteU32(*(u32*)&weight);
+			writer.WriteBool(startInfo.trackInputs);
 			writer.WriteBool(startInfo.writeOutputFilesOnFail);
 			writer.WriteU64(startInfo.outputStatsThresholdMs);
 		}
@@ -77,8 +78,9 @@ namespace uba
 			u32 weight32 = reader.ReadU32();
 			weight = *(float*)&weight32;
 			
-			startInfo.outputStatsThresholdMs = reader.ReadU64();
+			startInfo.trackInputs = reader.ReadBool();
 			startInfo.writeOutputFilesOnFail = reader.ReadBool();
+			startInfo.outputStatsThresholdMs = reader.ReadU64();
 
 			startInfo.description = description.c_str();
 			startInfo.application = application.c_str();
diff --git a/Engine/Source/Programs/UnrealBuildAccelerator/Common/Private/UbaScheduler.cpp b/Engine/Source/Programs/UnrealBuildAccelerator/Common/Private/UbaScheduler.cpp
--- a/Engine/Source/Programs/UnrealBuildAccelerator/Common/Private/UbaScheduler.cpp
+++ b/Engine/Source/Programs/UnrealBuildAccelerator/Common/Private/UbaScheduler.cpp
@@ -1,6 +1,7 @@
 // Copyright Epic Games, Inc. All Rights Reserved.
 
 #include "UbaScheduler.h"
+#include "UbaApplicationRules.h"
 #include "UbaProcess.h"
 #include "UbaProcessStartInfoHolder.h"
 #include "UbaSessionServer.h"
@@ -45,15 +46,15 @@ namespace uba
 		virtual u32 GetExitCode() override { return ProcessCancelExitCode; }
 		virtual bool HasExited() override { return true; }
 		virtual bool WaitForExit(u32 millisecondsTimeout) override{ return true; }
-		virtual const ProcessStartInfo& GetStartInfo() override { return holder.startInfo; }
-		virtual const Vector<ProcessLogLine>& GetLogLines() override { static Vector<ProcessLogLine> v{ProcessLogLine{TC("Skipped"), LogEntryType_Warning}}; return v; }
-		virtual const Vector<u8>& GetTrackedInputs() override { static Vector<u8> v; return v;}
+		virtual const ProcessStartInfo& GetStartInfo() const override { return holder.startInfo; }
+		virtual const Vector<ProcessLogLine>& GetLogLines() const override { static Vector<ProcessLogLine> v{ProcessLogLine{TC("Skipped"), LogEntryType_Warning}}; return v; }
+		virtual const Vector<u8>& GetTrackedInputs() const override { static Vector<u8> v; return v;}
+		virtual const Vector<u8>& GetTrackedOutputs() const override { static Vector<u8> v; return v;}
 		virtual bool IsRemote() const override { return false; }
-		virtual bool IsDetoured() const { return false; }
+		virtual ProcessExecutionType GetExecutionType() const override { return ProcessExecutionType_Native; }
 		ProcessStartInfoHolder holder;
 	};
 
-	
 	Scheduler::Scheduler(const SchedulerCreateInfo& info)
 	:	m_session(info.session)
 	,	m_maxLocalProcessors(info.maxLocalProcessors != ~0u ? info.maxLocalProcessors : GetLogicalProcessorCount())
@@ -125,6 +126,9 @@ namespace uba
 		auto info2 = new ProcessStartInfo2(info.info, ki, info.knownInputsCount);
 		info2->weight = info.weight;
 
+		const ApplicationRules* rules = m_session.GetRules(info2->startInfo);
+		info2->startInfo.rules = rules;
+
 		SCOPED_WRITE_LOCK(m_processEntriesLock, lock);
 		u32 index = u32(m_processEntries.size());
 		auto& entry = m_processEntries.emplace_back();
diff --git a/Engine/Source/Programs/UnrealBuildAccelerator/Common/Private/UbaSession.cpp b/Engine/Source/Programs/UnrealBuildAccelerator/Common/Private/UbaSession.cpp
--- a/Engine/Source/Programs/UnrealBuildAccelerator/Common/Private/UbaSession.cpp
+++ b/Engine/Source/Programs/UnrealBuildAccelerator/Common/Private/UbaSession.cpp
@@ -163,10 +163,10 @@ namespace uba
 		UBA_ASSERT(m_process);
 		return m_process->IsRemote();
 	}
-	bool ProcessHandle::IsDetoured() const
+	ProcessExecutionType ProcessHandle::GetExecutionType() const
 	{
 		UBA_ASSERT(m_process);
-		return m_process->IsDetoured();
+		return m_process->GetExecutionType();
 	}
 	ProcessHandle::ProcessHandle(Process* process)
 	{
@@ -1208,6 +1208,9 @@ namespace uba
 			si.logFile = logFile.data;
 		}
 
+		if (!si.rules)
+			si.rules = GetRules(si);
+
 		void* env = GetProcessEnvironmentVariables();
 		u32 id = ++m_processIdCounter;
 		auto process = new ProcessImpl(*this, id, parent);
@@ -1304,6 +1307,72 @@ namespace uba
 	Logger& Session::GetLogger() { return m_logger; }
 	LogWriter& Session::GetLogWriter() { return m_logger.m_writer; }
 
+	const ApplicationRules* Session::GetRules(const ProcessStartInfo& si)
+	{
+		u32 exeNameStart = 0;
+		u32 exeNameEnd = TStrlen(si.application);
+		const tchar* lastSeparator = TStrrchr(si.application, PathSeparator);
+		if (lastSeparator)
+			exeNameStart = u32(lastSeparator - si.application + 1);
+		else if (si.application[exeNameStart] == '"')
+			++exeNameStart;
+		if (si.application[exeNameEnd - 1] == '"')
+			--exeNameEnd;
+		StringBuffer<128> exeName;
+		exeName.Append(si.application + exeNameStart, exeNameEnd - exeNameStart);
+		
+		auto rules = GetApplicationRules();
+		
+		while (true)
+		{
+			for (u32 i = 1;; ++i)
+			{
+				const tchar* app = rules[i].app;
+				if (!app)
+					break;
+				if (!exeName.Equals(app))
+					continue;
+				return GetApplicationRules()[i].rules;
+			}
+
+			if (!exeName.Equals(TC("dotnet.exe")))
+				return GetApplicationRules()[0].rules;
+			
+			u32 firstArgumentStart = 0;
+			u32 firstArgumentEnd = 0;
+			bool quoted = false;
+			for (u32 i = 0, e = TStrlen(si.arguments); i != e; ++i)
+			{
+				tchar c = si.arguments[i];
+				if (firstArgumentEnd)
+				{
+					if (c == '\\')
+						firstArgumentStart = i + 1;
+					if ((quoted && c != '"') || (!quoted && c != ' ' && c != '\t'))
+						continue;
+					firstArgumentEnd = i;
+					break;
+				}
+				else
+				{
+					if (c == ' ' || c == '\t')
+					{
+						++firstArgumentStart;
+						continue;
+					}
+					if (c == '"')
+					{
+						++firstArgumentStart;
+						quoted = true;
+					}
+					firstArgumentEnd = firstArgumentStart + 1;
+				}
+			}
+			exeName.Clear().Append(si.arguments + firstArgumentStart, firstArgumentEnd - firstArgumentStart);
+		}
+		return GetApplicationRules()[0].rules;
+	}
+
 	void Session::ProcessAdded(Process& process, u32 sessionId)
 	{
 		u32 processId = process.GetId();
@@ -1500,12 +1569,12 @@ namespace uba
 
 	bool Session::IsRarelyRead(ProcessImpl& process, const StringBufferBase& fileName) const
 	{
-		return GetApplicationRules()[process.m_rulesIndex].rules->IsRarelyRead(fileName);
+		return process.m_startInfo.rules->IsRarelyRead(fileName);
 	}
 
 	bool Session::IsRarelyReadAfterWritten(ProcessImpl& process, const tchar* fileName, u64 fileNameLen) const
 	{
-		return GetApplicationRules()[process.m_rulesIndex].rules->IsRarelyReadAfterWritten(fileName, fileNameLen);
+		return process.m_startInfo.rules->IsRarelyReadAfterWritten(fileName, fileNameLen);
 	}
 
 	bool Session::IsKnownSystemFile(const tchar* applicationName)
diff --git a/Engine/Source/Programs/UnrealBuildAccelerator/Common/Private/UbaSessionClient.cpp b/Engine/Source/Programs/UnrealBuildAccelerator/Common/Private/UbaSessionClient.cpp
--- a/Engine/Source/Programs/UnrealBuildAccelerator/Common/Private/UbaSessionClient.cpp
+++ b/Engine/Source/Programs/UnrealBuildAccelerator/Common/Private/UbaSessionClient.cpp
@@ -434,7 +434,7 @@ namespace uba
 			if (m_useStorage || memoryMapAlignment == 0)
 			{
 				bool storeUncompressed = memoryMapAlignment == 0;
-				bool allowProxy = GetApplicationRules()[msg.process.m_rulesIndex].rules->AllowStorageProxy(fileName);
+				bool allowProxy = msg.process.m_startInfo.rules->AllowStorageProxy(fileName);
 				if (!RetrieveCasFile(newCasKey, fileSize, casKey, fileName.data, storeUncompressed, allowProxy))
 					return m_logger.Error(TC("Error retrieving cas entry %s (%s)"), CasKeyString(casKey).str, fileName.data);
 
@@ -475,7 +475,7 @@ namespace uba
 					casKey = AsCompressed(casKey, false);
 					entry.handled = true;
 					Storage::RetrieveResult result;
-					bool allowProxy = GetApplicationRules()[msg.process.m_rulesIndex].rules->AllowStorageProxy(fileName);
+					bool allowProxy = msg.process.m_startInfo.rules->AllowStorageProxy(fileName);
 					if (!m_storage.RetrieveCasFile(result, casKey, fileName.data, &m_fileMappingBuffer, memoryMapAlignment, allowProxy))
 						return m_logger.Error(TC("Error retrieving cas entry %s (%s)"), CasKeyString(casKey).str, fileName.data);
 					entry.success = true;
@@ -1229,6 +1229,38 @@ namespace uba
 			return;
 	}
 
+	bool SessionClient::SendProcessInputs(ProcessImpl& process)
+	{
+		StackBinaryWriter<SendMaxSize> writer;
+		NetworkMessage msg(m_client, ServiceId, SessionMessageType_ProcessInputs, writer);
+		writer.WriteU32(process.m_id);
+		auto inputs = process.GetTrackedInputs();
+		writer.WriteBytes(inputs.data(), inputs.size());
+		StackBinaryReader<32> reader;
+		return msg.Send(reader);
+	}
+
+	bool SessionClient::SendProcessFinished(ProcessImpl& process, u32 exitCode)
+	{
+		StackBinaryWriter<SendMaxSize> writer;
+		NetworkMessage msg(m_client, ServiceId, SessionMessageType_ProcessFinished, writer);
+		writer.WriteU32(process.m_id);
+		writer.WriteU32(exitCode);
+		writer.WriteU32(CountLogLines(process));
+		WriteLogLines(writer, process);
+
+		// Must be written last
+		process.m_processStats.Write(writer);
+		process.m_sessionStats.Write(writer);
+		process.m_storageStats.Write(writer);
+		process.m_systemStats.Write(writer);
+
+		StackBinaryReader<16> reader;
+		if (!msg.Send(reader, m_stats.procFinishedMsg) && m_loop)
+			return m_logger.Error(TC("Failed to send ProcessFinished message!"));
+		return true;
+	}
+
 	bool SessionClient::SendUpdateDirectoryTable(StackBinaryReader<SendMaxSize>& reader)
 	{
 		UBA_ASSERT(reader.GetPosition() == 0);
@@ -1503,6 +1535,7 @@ namespace uba
 					startInfo.priorityClass = m_defaultPriorityClass;
 					startInfo.useCustomAllocator = !m_disableCustomAllocator;
 					startInfo.outputStatsThresholdMs = m_outputStatsThresholdMs != 0 ? m_outputStatsThresholdMs : startInfo.outputStatsThresholdMs;
+					startInfo.rules = GetRules(startInfo);
 
 					StringBuffer<> logFile;
 					if (m_logToFile)
@@ -1617,22 +1650,10 @@ namespace uba
 							return;
 						}
 
-						StackBinaryWriter<SendMaxSize> writer;
-						NetworkMessage msg(session.m_client, ServiceId, SessionMessageType_ProcessFinished, writer);
-						writer.WriteU32(process.m_id);
-						writer.WriteU32(exitCode);
-						writer.WriteU32(session.CountLogLines(process));
-						session.WriteLogLines(writer, process);
-
-						// Must be written last
-						process.m_processStats.Write(writer);
-						process.m_sessionStats.Write(writer);
-						process.m_storageStats.Write(writer);
-						process.m_systemStats.Write(writer);
-
-						StackBinaryReader<16> reader;
-						if (!msg.Send(reader, session.m_stats.procFinishedMsg) && session.m_loop)
-							session.m_logger.Error(TC("Failed to send ProcessFinished message!"));
+						if (startInfo.trackInputs)
+							session.SendProcessInputs(process);
+
+						session.SendProcessFinished(process, exitCode);
 
 						// TODO: These should be removed and instead added in TraceReader (so it will update over time)
 						session.m_stats.stats.Add(process.m_sessionStats);
diff --git a/Engine/Source/Programs/UnrealBuildAccelerator/Common/Private/UbaSessionServer.cpp b/Engine/Source/Programs/UnrealBuildAccelerator/Common/Private/UbaSessionServer.cpp
--- a/Engine/Source/Programs/UnrealBuildAccelerator/Common/Private/UbaSessionServer.cpp
+++ b/Engine/Source/Programs/UnrealBuildAccelerator/Common/Private/UbaSessionServer.cpp
@@ -25,15 +25,16 @@ namespace uba
 			delete[] m_knownInputs;
 		}
 
-		virtual const ProcessStartInfo& GetStartInfo() override { return startInfo; }
+		virtual const ProcessStartInfo& GetStartInfo() const override { return startInfo; }
 		virtual u32 GetId() override { return m_processId; }
 		virtual u32 GetExitCode() override { UBA_ASSERT(m_done.IsSet(0)); return m_exitCode; }
 		virtual bool HasExited() override { return m_done.IsSet(0); }
 		virtual bool WaitForExit(u32 millisecondsTimeout) override { return m_done.IsSet(millisecondsTimeout); }
 		virtual u64 GetTotalProcessorTime() const override { return m_processorTime; }
 		virtual u64 GetTotalWallTime() const override { return m_wallTime; }
-		virtual const Vector<ProcessLogLine>& GetLogLines() override { return m_logLines; }
-		virtual const Vector<u8>& GetTrackedInputs() override { return m_trackedInputs; }
+		virtual const Vector<ProcessLogLine>& GetLogLines() const override { return m_logLines; }
+		virtual const Vector<u8>& GetTrackedInputs() const override { return m_trackedInputs; }
+		virtual const Vector<u8>& GetTrackedOutputs() const override { return m_trackedOutputs; };
 		virtual void Cancel(bool terminate) override
 		{
 			if (m_cancelled)
@@ -53,7 +54,7 @@ namespace uba
 
 		virtual const tchar* GetExecutingHost() const override { return m_executingHost.c_str(); }
 		virtual bool IsRemote() const override { return true; }
-		virtual bool IsDetoured() const { return true; }
+		virtual ProcessExecutionType GetExecutionType() const override { return ProcessExecutionType_Detoured; }
 		virtual bool IsChild() override { return false; }
 
 		void CallProcessExit(ProcessHandle& h)
@@ -77,6 +78,7 @@ namespace uba
 		Event m_done;
 		Vector<ProcessLogLine> m_logLines;
 		Vector<u8> m_trackedInputs;
+		Vector<u8> m_trackedOutputs;
 		bool m_cancelled = false;
 		u32 m_clientId = ~0u;
 		u32 m_sessionId = 0;
@@ -409,10 +411,37 @@ namespace uba
 		{
 			ScopedCriticalSection lock(m_remoteProcessAndSessionLock);
 			if (m_activeRemoteProcesses.empty() && m_queuedRemoteProcesses.empty())
-				return;
+				break;
 			lock.Leave();
 			Sleep(200);
 		}
+
+		bool isEmpty = false;
+		while (!isEmpty)
+		{
+			Vector<ProcessHandle> processes;
+			{
+				SCOPED_WRITE_LOCK(m_processesLock, lock);
+				isEmpty = m_processes.empty();
+				processes.reserve(m_processes.size());
+				for (auto& pair : m_processes)
+					processes.push_back(pair.second);
+			}
+
+			#if PLATFORM_WINDOWS
+			if (m_processJobObject != NULL)
+			{
+				SCOPED_WRITE_LOCK(m_processJobObjectLock, lock);
+				CloseHandle(m_processJobObject);
+				m_processJobObject = NULL;
+			}
+			#endif
+
+			for (auto& process : processes)
+				process.WaitForExit(100000);
+		}
+
+		FlushDeadProcesses();
 	}
 
 	void SessionServer::SetMaxRemoteProcessCount(u32 count)
@@ -767,7 +796,7 @@ namespace uba
 			}
 			case SessionMessageType_SendFileToServer:
 			{
-				u32 processId = reader.ReadU32(); (void)processId;
+				u32 processId = reader.ReadU32();
 				StringBuffer<> destination;
 				reader.ReadString(destination);
 				StringKey destinationKey = reader.ReadStringKey();
@@ -828,6 +857,24 @@ namespace uba
 				{
 					m_storage.DropCasFile(casKey, false, destination.data);
 					RegisterCreateFileForWrite(StringKeyZero, destination.data, destination.count, true);
+
+
+					SCOPED_WRITE_LOCK(m_processesLock, lock);
+					auto findIt = m_processes.find(processId);
+					if (findIt != m_processes.end())
+					{
+						ProcessHandle h(findIt->second);
+						lock.Leave();
+						auto& process = *(RemoteProcess*)h.m_process;
+						if (process.startInfo.trackInputs)
+						{
+							u64 bytes = GetStringWriteSize(destination.data);
+							u64 prevSize = process.m_trackedOutputs.size();
+							process.m_trackedOutputs.resize(prevSize + bytes);
+							BinaryWriter w2(process.m_trackedOutputs.data(), prevSize, prevSize + bytes);
+							w2.WriteString(destination);
+						}
+					}
 				}
 				return true;
 			}
@@ -1018,6 +1065,21 @@ namespace uba
 
 				return true;
 			}
+			case SessionMessageType_ProcessInputs:
+			{
+				u32 processId = reader.ReadU32();
+				SCOPED_WRITE_LOCK(m_processesLock, lock);
+				auto findIt = m_processes.find(processId);
+				if (findIt == m_processes.end())
+					return m_logger.Error(TC("Failed to find process for id %u when receiving custom message"), processId);
+				ProcessHandle h(findIt->second);
+				lock.Leave();
+				auto& process = *(RemoteProcess*)h.m_process;
+				u64 size = reader.GetLeft();
+				process.m_trackedInputs.resize(size);
+				reader.ReadBytes(process.m_trackedInputs.data(), size);
+				return true;
+			}
 			case SessionMessageType_ProcessFinished:
 			{
 				u32 processId = reader.ReadU32();
diff --git a/Engine/Source/Programs/UnrealBuildAccelerator/Common/Private/UbaStorage.cpp b/Engine/Source/Programs/UnrealBuildAccelerator/Common/Private/UbaStorage.cpp
--- a/Engine/Source/Programs/UnrealBuildAccelerator/Common/Private/UbaStorage.cpp
+++ b/Engine/Source/Programs/UnrealBuildAccelerator/Common/Private/UbaStorage.cpp
@@ -1494,7 +1494,7 @@ namespace uba
 								mem += sizeof(u64);
 								u8* dest = new u8[decompressedSize];
 								auto g = MakeGuard([dest]() { delete[] dest; });
-								if (!DecompressMemoryToMemory(mem, dest, decompressedSize, TC("")))
+								if (!DecompressMemoryToMemory(mem, dest, decompressedSize, filePath.c_str()))
 									return;
 								checkedKey = CalculateCasKey(dest, decompressedSize, true);
 							}
diff --git a/Engine/Source/Programs/UnrealBuildAccelerator/Common/Private/UbaStorageClient.cpp b/Engine/Source/Programs/UnrealBuildAccelerator/Common/Private/UbaStorageClient.cpp
--- a/Engine/Source/Programs/UnrealBuildAccelerator/Common/Private/UbaStorageClient.cpp
+++ b/Engine/Source/Programs/UnrealBuildAccelerator/Common/Private/UbaStorageClient.cpp
@@ -598,7 +598,7 @@ namespace uba
 
 							if (!mappingBuffer)
 							{
-								if (!destinationFile.Write(decompressBuffer, uncompressedSize))
+								if (!destinationFile.Write(decompressBuffer, uncompressedSize, actualSize - leftUncompressed))
 									return false;
 							}
 							else
diff --git a/Engine/Source/Programs/UnrealBuildAccelerator/Common/Private/UbaStorageServer.cpp b/Engine/Source/Programs/UnrealBuildAccelerator/Common/Private/UbaStorageServer.cpp
--- a/Engine/Source/Programs/UnrealBuildAccelerator/Common/Private/UbaStorageServer.cpp
+++ b/Engine/Source/Programs/UnrealBuildAccelerator/Common/Private/UbaStorageServer.cpp
@@ -1,6 +1,7 @@
 // Copyright Epic Games, Inc. All Rights Reserved.
 
 #include "UbaStorageServer.h"
+#include "UbaFileAccessor.h"
 #include "UbaNetworkServer.h"
 #include "UbaTrace.h"
 
@@ -11,6 +12,7 @@ namespace uba
 	,	m_server(info.server)
 	{
 		m_zone = info.zone;
+		m_writeRecievedCasFilesToDisk = info.writeRecievedCasFilesToDisk;
 
 		if (!CreateGuid(m_uid))
 			UBA_ASSERT(false);
@@ -153,6 +155,8 @@ namespace uba
 						m_trace->FileEndStore(clientId, store.casEntry->key);
 				}
 
+				delete store.fileAccessor;
+
 				m_casDataBuffer.UnmapView(store.mappedView, TC("OnDisconnected"));
 				it = m_activeStores.erase(it);
 			}
@@ -383,7 +387,7 @@ namespace uba
 
 				CasEntry* casEntry = nullptr;
 				bool has = HasCasFile(casKey, &casEntry); // HasCasFile also writes deferred cas entries if in queue
-				if (!has)
+				if (!has && !EnsureCasFile(casKey, nullptr))
 				{
 					// Last resort.. use hint to load file into cas (hint should be renamed since it is now a critical parameter)
 					// We better check the caskey first to make sure it is matching on the server
@@ -725,11 +729,31 @@ namespace uba
 					return false;
 				}
 
-				auto mappedView = m_casDataBuffer.AllocAndMapView(MappedView_Transient, fileSize, 1, CasKeyString(casKey).str);
-				if (!mappedView.memory)
+				MappedView mappedView;
+				FileAccessor* fileAccessor = nullptr;
+				
+				if (m_writeRecievedCasFilesToDisk)
 				{
-					casEntry.verified = false;
-					return false;
+					StringBuffer<> casKeyName;
+					GetCasFileName(casKeyName, casKey);
+					fileAccessor = new FileAccessor(m_logger, casKeyName.data);
+					if (!fileAccessor->CreateMemoryWrite(false, DefaultAttributes(), fileSize, m_tempPath.data))
+					{
+						delete fileAccessor;
+						m_logger.Error(TC("Failed to create cas file %s"), casKeyName.data);
+						casEntry.verified = false;
+						return false;
+					}
+					mappedView.memory = fileAccessor->GetData();
+				}
+				else
+				{
+					mappedView = m_casDataBuffer.AllocAndMapView(MappedView_Transient, fileSize, 1, CasKeyString(casKey).str);
+					if (!mappedView.memory)
+					{
+						casEntry.verified = false;
+						return false;
+					}
 				}
 
 				casEntry.beingWritten = true;
@@ -742,6 +766,7 @@ namespace uba
 				firstStore->fileSize = fileSize;
 				firstStore->actualSize = actualSize;
 				firstStore->mappedView = mappedView;
+				firstStore->fileAccessor = fileAccessor;
 				firstStore->recvCasTime = GetTime() - start;
 
 				if (m_trace)
@@ -781,6 +806,14 @@ namespace uba
 				{
 					m_casDataBuffer.UnmapView(activeStore.mappedView, TC("StoreDone"));
 
+					if (activeStore.fileAccessor)
+					{
+						bool success = activeStore.fileAccessor->Close();
+						delete activeStore.fileAccessor;
+						if (!success)
+							return m_logger.Error(TC("REVISIT THIS!"));
+					}
+
 					CasEntry& casEntry = *activeStore.casEntry;
 					{
 						SCOPED_WRITE_LOCK(casEntry.lock, entryLock);
@@ -844,6 +877,7 @@ namespace uba
 
 					s.fileSize = firstStore->fileSize;
 					s.mappedView = firstStore->mappedView;
+					s.fileAccessor = firstStore->fileAccessor;
 					s.casEntry = firstStore->casEntry;
 					s.totalWritten = firstStore->totalWritten.load();
 					s.recvCasTime = firstStore->recvCasTime.load();
diff --git a/Engine/Source/Programs/UnrealBuildAccelerator/Common/Public/UbaFileAccessor.h b/Engine/Source/Programs/UnrealBuildAccelerator/Common/Public/UbaFileAccessor.h
--- a/Engine/Source/Programs/UnrealBuildAccelerator/Common/Public/UbaFileAccessor.h
+++ b/Engine/Source/Programs/UnrealBuildAccelerator/Common/Public/UbaFileAccessor.h
@@ -21,7 +21,7 @@ namespace uba
 		bool Write(const void* data, u64 dataLen, u64 offset = 0);
 
 		bool OpenRead();
-		bool OpenMemoryRead(u64 offset = 0);
+		bool OpenMemoryRead(u64 offset = 0, bool errorOnFail = true);
 
 
 		const tchar* GetFileName() { return m_fileName; }
diff --git a/Engine/Source/Programs/UnrealBuildAccelerator/Common/Public/UbaNetwork.h b/Engine/Source/Programs/UnrealBuildAccelerator/Common/Public/UbaNetwork.h
--- a/Engine/Source/Programs/UnrealBuildAccelerator/Common/Public/UbaNetwork.h
+++ b/Engine/Source/Programs/UnrealBuildAccelerator/Common/Public/UbaNetwork.h
@@ -18,7 +18,7 @@ namespace uba
 	
 	static constexpr u32 SystemNetworkVersion = 1339;
 	static constexpr u32 StorageNetworkVersion = 3;
-	static constexpr u32 SessionNetworkVersion = 31;
+	static constexpr u32 SessionNetworkVersion = 32;
 
 
 	// Messages used over network between client/server (system, storage and session)
@@ -72,6 +72,7 @@ namespace uba
 		UBA_SESSION_MESSAGE(GetDirectoriesFromServer) \
 		UBA_SESSION_MESSAGE(GetNameToHashFromServer) \
 		UBA_SESSION_MESSAGE(ProcessAvailable) \
+		UBA_SESSION_MESSAGE(ProcessInputs) \
 		UBA_SESSION_MESSAGE(ProcessFinished) \
 		UBA_SESSION_MESSAGE(ProcessReturned) \
 		UBA_SESSION_MESSAGE(VirtualAllocFailed) \
diff --git a/Engine/Source/Programs/UnrealBuildAccelerator/Common/Public/UbaProcess.h b/Engine/Source/Programs/UnrealBuildAccelerator/Common/Public/UbaProcess.h
--- a/Engine/Source/Programs/UnrealBuildAccelerator/Common/Public/UbaProcess.h
+++ b/Engine/Source/Programs/UnrealBuildAccelerator/Common/Public/UbaProcess.h
@@ -10,19 +10,20 @@ namespace uba
 	class Process
 	{
 	public:
-		virtual const ProcessStartInfo& GetStartInfo() = 0;
+		virtual const ProcessStartInfo& GetStartInfo() const = 0;
 		virtual u32 GetId() { UBA_ASSERT(false); return 0; }
 		virtual u32 GetExitCode() { UBA_ASSERT(false); return ~0u; }
 		virtual bool HasExited()  { UBA_ASSERT(false); return false; }
 		virtual bool WaitForExit(u32 millisecondsTimeout) { UBA_ASSERT(false); return false; }
 		virtual u64 GetTotalProcessorTime() const { UBA_ASSERT(false); return 0; }
 		virtual u64 GetTotalWallTime() const { UBA_ASSERT(false); return 0; }
-		virtual const Vector<ProcessLogLine>& GetLogLines() = 0;
-		virtual const Vector<u8>& GetTrackedInputs() = 0;
+		virtual const Vector<ProcessLogLine>& GetLogLines() const = 0;
+		virtual const Vector<u8>& GetTrackedInputs() const = 0;
+		virtual const Vector<u8>& GetTrackedOutputs() const = 0;
 		virtual void Cancel(bool terminate) { UBA_ASSERT(false); }
 		virtual const tchar* GetExecutingHost() const { UBA_ASSERT(false); return nullptr; }
 		virtual bool IsRemote() const { UBA_ASSERT(false); return false; }
-		virtual bool IsDetoured() const { UBA_ASSERT(false); return false; }
+		virtual ProcessExecutionType GetExecutionType() const { UBA_ASSERT(false); return ProcessExecutionType_Native; }
 		virtual bool IsChild() { UBA_ASSERT(false); return false; }
 
 	protected:
@@ -37,19 +38,20 @@ namespace uba
 	{
 	public:
 
-		virtual const ProcessStartInfo& GetStartInfo() override { return m_startInfo; }
+		virtual const ProcessStartInfo& GetStartInfo() const override { return m_startInfo; }
 		virtual u32 GetId() override { return m_id; }
 		virtual u32 GetExitCode() override { return m_exitCode; }
 		virtual bool HasExited() override  { return m_hasExited; }
 		virtual bool WaitForExit(u32 millisecondsTimeout) override;
 		virtual u64 GetTotalProcessorTime() const override;
 		virtual u64 GetTotalWallTime() const override;
-		virtual const Vector<ProcessLogLine>& GetLogLines() override { return m_logLines; }
-		virtual const Vector<u8>& GetTrackedInputs() override { return m_trackedInputs; }
+		virtual const Vector<ProcessLogLine>& GetLogLines() const override { return m_logLines; }
+		virtual const Vector<u8>& GetTrackedInputs() const override { return m_trackedInputs; }
+		virtual const Vector<u8>& GetTrackedOutputs() const override { return m_trackedOutputs; }
 		virtual void Cancel(bool terminate) override;
 		virtual const tchar* GetExecutingHost() const override { return TC(""); }
 		virtual bool IsRemote() const override { return false; }
-		virtual bool IsDetoured() const { return m_detourEnabled; }
+		virtual ProcessExecutionType GetExecutionType() const override { return m_detourEnabled ? ProcessExecutionType_Detoured : ProcessExecutionType_Native; }
 		virtual bool IsChild() override { return m_parentProcess != nullptr; }
 
 		ProcessImpl(Session& session, u32 id, ProcessImpl* parent);
@@ -98,7 +100,6 @@ namespace uba
 		bool OpenTempFile(BinaryReader& reader, BinaryWriter& writer, const tchar* application);
 		bool WriteFilesToDisk();
 
-		void SetRulesIndex(const ProcessStartInfo& si);
 		const tchar* InternalGetChildLogFile(StringBufferBase& temp);
 		
 		#if !PLATFORM_WINDOWS
@@ -143,7 +144,6 @@ namespace uba
 		u32 m_nativeProcessId = 0;
 		u32 m_nativeProcessExitCode = ~0u;
 		u32 m_exitCode = ~0u;
-		u32 m_rulesIndex = 0;
 		u32 m_messageCount = 0;
 		Atomic<bool> m_hasExited;
 		bool m_messageSuccess = true;
@@ -164,6 +164,7 @@ namespace uba
 		ReaderWriterLock m_logLinesLock;
 		Vector<ProcessLogLine> m_logLines;
 		Vector<u8> m_trackedInputs;
+		Vector<u8> m_trackedOutputs;
 		Vector<ProcessHandle> m_childProcesses;
 		ReaderWriterLock& m_writtenFilesLock;
 		UnorderedMap<TString, WrittenFile>& m_writtenFiles;
diff --git a/Engine/Source/Programs/UnrealBuildAccelerator/Common/Public/UbaProcessHandle.h b/Engine/Source/Programs/UnrealBuildAccelerator/Common/Public/UbaProcessHandle.h
--- a/Engine/Source/Programs/UnrealBuildAccelerator/Common/Public/UbaProcessHandle.h
+++ b/Engine/Source/Programs/UnrealBuildAccelerator/Common/Public/UbaProcessHandle.h
@@ -21,6 +21,11 @@ namespace uba
 		LogEntryType type = LogEntryType_Info;
 	};
 
+	enum ProcessExecutionType : u8
+	{
+		ProcessExecutionType_Native,
+		ProcessExecutionType_Detoured,
+	};
 
 	class ProcessHandle
 	{
@@ -38,7 +43,7 @@ namespace uba
 		void Cancel(bool terminate) const;					// Request to cancel process. Set terminate to true to quickly kill process
 		const tchar* GetExecutingHost() const;				// Host that is executing process.
 		bool IsRemote() const;								// Returns true if process is a remote process
-		bool IsDetoured() const;							// Returns true if process is detoured
+		ProcessExecutionType GetExecutionType() const;		// Returns execution type
 
 		ProcessHandle();
 		ProcessHandle(const ProcessHandle& o);
diff --git a/Engine/Source/Programs/UnrealBuildAccelerator/Common/Public/UbaProcessStartInfo.h b/Engine/Source/Programs/UnrealBuildAccelerator/Common/Public/UbaProcessStartInfo.h
--- a/Engine/Source/Programs/UnrealBuildAccelerator/Common/Public/UbaProcessStartInfo.h
+++ b/Engine/Source/Programs/UnrealBuildAccelerator/Common/Public/UbaProcessStartInfo.h
@@ -6,6 +6,7 @@
 
 namespace uba
 {
+	class ApplicationRules;
 	class ProcessHandle;
 
 	struct ProcessStartInfo
@@ -33,6 +34,7 @@ namespace uba
 		ExitedCallback* exitedFunc = nullptr;	// Callback for when process is done (it has already exited)
 		void* userData = nullptr;				// User data provided to exit callback
 
+		const ApplicationRules* rules = nullptr;// Internal use (for now)
 		int uiLanguage = 1033;					// Internal use
 	};
 }
diff --git a/Engine/Source/Programs/UnrealBuildAccelerator/Common/Public/UbaSession.h b/Engine/Source/Programs/UnrealBuildAccelerator/Common/Public/UbaSession.h
--- a/Engine/Source/Programs/UnrealBuildAccelerator/Common/Public/UbaSession.h
+++ b/Engine/Source/Programs/UnrealBuildAccelerator/Common/Public/UbaSession.h
@@ -68,6 +68,7 @@ namespace uba
 		Storage& GetStorage();		// Storage (only used when remote machines are connected)
 		Logger& GetLogger();			// Logger used for logging 
 		LogWriter& GetLogWriter();		// LogWriter used by logger
+		const ApplicationRules* GetRules(const ProcessStartInfo& si); // Get application rules used for process
 
 		virtual ~Session();
 
diff --git a/Engine/Source/Programs/UnrealBuildAccelerator/Common/Public/UbaSessionClient.h b/Engine/Source/Programs/UnrealBuildAccelerator/Common/Public/UbaSessionClient.h
--- a/Engine/Source/Programs/UnrealBuildAccelerator/Common/Public/UbaSessionClient.h
+++ b/Engine/Source/Programs/UnrealBuildAccelerator/Common/Public/UbaSessionClient.h
@@ -84,6 +84,8 @@ namespace uba
 		void BuildEnvironmentVariables(BinaryReader& reader);
 		bool SendProcessAvailable(Vector<InternalProcessStartInfo>& out, float availableWeight);
 		void SendReturnProcess(u32 processId, const tchar* reason);
+		bool SendProcessInputs(ProcessImpl& process);
+		bool SendProcessFinished(ProcessImpl& process, u32 exitCode);
 		void SendPing(u64 memAvail, u64 memTotal);
 		void SendLogFileToServer(ProcessImpl& pi);
 		void GetLogFileName(StringBufferBase& out, const tchar* logFile, const tchar* arguments);
diff --git a/Engine/Source/Programs/UnrealBuildAccelerator/Common/Public/UbaStorageServer.h b/Engine/Source/Programs/UnrealBuildAccelerator/Common/Public/UbaStorageServer.h
--- a/Engine/Source/Programs/UnrealBuildAccelerator/Common/Public/UbaStorageServer.h
+++ b/Engine/Source/Programs/UnrealBuildAccelerator/Common/Public/UbaStorageServer.h
@@ -11,6 +11,7 @@ namespace uba
 	{
 		StorageServerCreateInfo(NetworkServer& s, const tchar* rootDir_, LogWriter& writer) : StorageCreateInfo(rootDir_, writer), server(s) { workManager = &server; }
 		NetworkServer& server;
+		bool writeRecievedCasFilesToDisk = false;
 		const tchar* zone = TC("");
 	};
 
@@ -61,6 +62,7 @@ namespace uba
 		{
 			u32 clientId = ~0u;
 			MappedView mappedView;
+			FileAccessor* fileAccessor = nullptr;
 			CasEntry* casEntry = nullptr;
 			Atomic<u64> totalWritten;
 			Atomic<u64> recvCasTime;
@@ -133,5 +135,7 @@ namespace uba
 		Trace* m_trace = nullptr;
 
 		Vector<TString> m_disallowedPaths;
+
+		bool m_writeRecievedCasFilesToDisk;
 	};
 }
diff --git a/Engine/Source/Programs/UnrealBuildAccelerator/Core/Private/UbaApplicationRules.cpp b/Engine/Source/Programs/UnrealBuildAccelerator/Core/Private/UbaApplicationRules.cpp
--- a/Engine/Source/Programs/UnrealBuildAccelerator/Core/Private/UbaApplicationRules.cpp
+++ b/Engine/Source/Programs/UnrealBuildAccelerator/Core/Private/UbaApplicationRules.cpp
@@ -9,12 +9,12 @@ namespace uba
 	{
 		using Super = ApplicationRules;
 
-		virtual bool AllowDetach() override
+		virtual bool AllowDetach() const override
 		{
 			return true;
 		}
 
-		virtual u64 FileTypeMaxSize(const StringBufferBase& file, bool isSystemOrTempFile) override
+		virtual u64 FileTypeMaxSize(const StringBufferBase& file, bool isSystemOrTempFile) const override
 		{
 			if (file.EndsWith(TC(".pdb")))
 				return 14ull * 1024 * 1024 * 1024; // This is ridiculous
@@ -24,17 +24,17 @@ namespace uba
 				return 1024 * 1024 * 1024;
 			return Super::FileTypeMaxSize(file, isSystemOrTempFile);
 		}
-		virtual bool CanDetour(const tchar* file) override
+		virtual bool CanDetour(const tchar* file) const override
 		{
 			if (file[0] == '\\' && file[1] == '\\' && !Contains(file, TC("vctip_"))) // This might be too aggressive but will cover pipes etc.. might need revisit
 				return false;
 			return true;
 		}
-		virtual bool IsThrowAway(const tchar* fileName, u32 fileNameLen) override
+		virtual bool IsThrowAway(const tchar* fileName, u32 fileNameLen) const override
 		{
 			return Contains(fileName, TC("vctip_")) || Super::IsThrowAway(fileName, fileNameLen);
 		}
-		virtual bool KeepInMemory(const tchar* fileName, u32 fileNameLen, const tchar* systemTemp) override
+		virtual bool KeepInMemory(const tchar* fileName, u32 fileNameLen, const tchar* systemTemp) const override
 		{
 			if (Contains(fileName, TC("\\vctip_")) != 0)
 				return true;
@@ -43,7 +43,7 @@ namespace uba
 			return Super::KeepInMemory(fileName, fileNameLen, systemTemp);
 		}
 
-		virtual bool IsExitCodeSuccess(u32 exitCode) override
+		virtual bool IsExitCodeSuccess(u32 exitCode) const override
 		{
 			return exitCode == 0;
 		}
@@ -52,7 +52,7 @@ namespace uba
 	class ApplicationRulesClExe : public ApplicationRulesVC
 	{
 	public:
-		virtual bool IsOutputFile(const tchar* file, u64 fileLen) override
+		virtual bool IsOutputFile(const tchar* file, u64 fileLen) const override
 		{
 			return EndsWith(file, fileLen, TC(".h.obj"))
 				|| EndsWith(file, fileLen, TC(".c.obj"))
@@ -63,7 +63,7 @@ namespace uba
 				;// || EndsWith(file, fileLen, TC(".h.pch") // Not tested enough
 		}
 
-		virtual bool IsRarelyRead(const StringBufferBase& file) override
+		virtual bool IsRarelyRead(const StringBufferBase& file) const override
 		{
 			return file.EndsWith(TC(".cpp"))
 				|| file.EndsWith(TC(".obj.rsp"));
@@ -76,7 +76,7 @@ namespace uba
 				|| EndsWith(fileName, fileNameLen, TC(".dll"));
 		}
 
-		virtual bool NeedsSharedMemory(const tchar* file) override
+		virtual bool NeedsSharedMemory(const tchar* file) const override
 		{
 			return Contains(file, TC("\\_cl_")); // This file is needed when cl.exe spawns link.exe
 		}
@@ -87,7 +87,7 @@ namespace uba
 	{
 		using Super = ApplicationRulesVC;
 	public:
-		virtual bool IsOutputFile(const tchar* file, u64 fileLen) override
+		virtual bool IsOutputFile(const tchar* file, u64 fileLen) const override
 		{
 			return EndsWith(file, fileLen, TC(".lib"))
 				|| EndsWith(file, fileLen, TC(".exp"))
@@ -97,22 +97,22 @@ namespace uba
 				|| EndsWith(file, fileLen, TC(".rc2.res")); // Not really an obj file.. 
 		}
 
-		virtual bool IsThrowAway(const tchar* fileName, u32 fileNameLen) override
+		virtual bool IsThrowAway(const tchar* fileName, u32 fileNameLen) const override
 		{
 			return Contains(fileName, TC(".sup.")); // .sup.lib/exp are throw-away files that we don't want created
 		}
 
-		virtual bool CanExist(const tchar* file) override
+		virtual bool CanExist(const tchar* file) const override
 		{
 			return Contains(file, TC("vctip.exe")) == false; // This is hacky but we don't want to start vctip.exe
 		}
-		virtual bool NeedsSharedMemory(const tchar* file) override
+		virtual bool NeedsSharedMemory(const tchar* file) const override
 		{
 			return Contains(file, TC("lnk{")) // This file is shared from link.exe to mt.exe and rc.exe so we need to put it shared memory
 			 	|| Contains(file, TC("\\_cl_")); // When link.exe is spawned by cl.exe we might use this which is in shared memory
 		}
 
-		virtual bool IsRarelyRead(const StringBufferBase& file) override
+		virtual bool IsRarelyRead(const StringBufferBase& file) const override
 		{
 			return file.EndsWith(TC(".exp"))
 				|| file.EndsWith(TC(".dll.rsp"))
@@ -121,7 +121,7 @@ namespace uba
 				|| file.EndsWith(TC(".pdb"));
 		}
 
-		virtual bool AllowStorageProxy(const StringBufferBase& file) override
+		virtual bool AllowStorageProxy(const StringBufferBase& file) const override
 		{
 			if (file.EndsWith(TC(".obj")))
 				return false;
@@ -141,12 +141,12 @@ namespace uba
 	class ApplicationRulesClang : public ApplicationRules
 	{
 	public:
-		virtual bool EnableVectoredExceptionHandler() override
+		virtual bool EnableVectoredExceptionHandler() const override
 		{
 			return true;
 		}
 
-		virtual bool IsExitCodeSuccess(u32 exitCode) override
+		virtual bool IsExitCodeSuccess(u32 exitCode) const override
 		{
 			return exitCode == 0;
 		}
@@ -155,12 +155,12 @@ namespace uba
 	class ApplicationRulesClangPlusPlusExe : public ApplicationRulesClang
 	{
 	public:
-		virtual bool AllowDetach() override
+		virtual bool AllowDetach() const override
 		{
 			return true;
 		}
 
-		virtual bool IsOutputFile(const tchar* file, u64 fileLen) override
+		virtual bool IsOutputFile(const tchar* file, u64 fileLen) const override
 		{
 			return EndsWith(file, fileLen, TC(".c.d"))
 				|| EndsWith(file, fileLen, TC(".h.d"))
@@ -171,7 +171,7 @@ namespace uba
 				;// || EndsWith(file, fileLen, TC(".gch.tmp")); // Need to fix "has been modified since the precompiled header"
 		}
 
-		virtual bool IsRarelyRead(const StringBufferBase& file) override
+		virtual bool IsRarelyRead(const StringBufferBase& file) const override
 		{
 			return file.EndsWith(TC(".cpp"))
 				|| file.EndsWith(TC(".o.rsp"));
@@ -182,7 +182,7 @@ namespace uba
 			return EndsWith(fileName, fileNameLen, TC(".d"));
 		}
 
-		virtual bool AllowMiMalloc() override
+		virtual bool AllowMiMalloc() const override
 		{
 			return true;
 		}
@@ -192,17 +192,17 @@ namespace uba
 	{
 		using Super = ApplicationRulesClang;
 
-		virtual bool IsOutputFile(const tchar* file, u64 fileLen) override
+		virtual bool IsOutputFile(const tchar* file, u64 fileLen) const override
 		{
 			return Contains(file, TC(".tmp")); // both .so.tmp and .tmp123456
 		}
 
-		virtual bool IsRarelyRead(const StringBufferBase& file) override
+		virtual bool IsRarelyRead(const StringBufferBase& file) const override
 		{
 			return file.EndsWith(TC(".so.rsp"));
 		}
 
-		virtual u64 FileTypeMaxSize(const StringBufferBase& file, bool isSystemOrTempFile) override
+		virtual u64 FileTypeMaxSize(const StringBufferBase& file, bool isSystemOrTempFile) const override
 		{
 			return 14ull * 1024 * 1024 * 1024; // This is ridiculous (needed for asan targets)
 		}
@@ -212,12 +212,12 @@ namespace uba
 	{
 		using Super = ApplicationRulesClang;
 
-		virtual bool IsOutputFile(const tchar* file, u64 fileLen) override
+		virtual bool IsOutputFile(const tchar* file, u64 fileLen) const override
 		{
 			return Contains(file, TC(".temp-stream-"));
 		}
 
-		virtual u64 FileTypeMaxSize(const StringBufferBase& file, bool isSystemOrTempFile) override
+		virtual u64 FileTypeMaxSize(const StringBufferBase& file, bool isSystemOrTempFile) const override
 		{
 			if (IsOutputFile(file.data, file.count))
 				return 14ull * 1024 * 1024 * 1024; // This is ridiculous (needed for asan targets)
@@ -227,7 +227,7 @@ namespace uba
 
 	class ApplicationRulesDumpSymsExe : public ApplicationRulesClang
 	{
-		virtual bool IsOutputFile(const tchar* file, u64 fileLen) override
+		virtual bool IsOutputFile(const tchar* file, u64 fileLen) const override
 		{
 			return false; // EndsWith(file, fileLen, TC(".psym")); With psym as output file the BreakpadSymbolEncoder fails to output a .sym file
 		}
@@ -237,7 +237,7 @@ namespace uba
 	{
 		using Super = ApplicationRulesClangPlusPlusExe;
 
-		virtual bool IsThrowAway(const tchar* fileName, u32 fileNameLen) override
+		virtual bool IsThrowAway(const tchar* fileName, u32 fileNameLen) const override
 		{
 			return EndsWith(fileName, fileNameLen, TC("-telemetry.json")) || Super::IsThrowAway(fileName, fileNameLen);
 		}
@@ -252,16 +252,16 @@ namespace uba
 	{
 		using Super = ApplicationRules;
 
-		//virtual bool IsOutputFile(const tchar* file, u64 fileLen) override
+		//virtual bool IsOutputFile(const tchar* file, u64 fileLen) const override
 		//{
 		//	return EndsWith(file, fileLen, TC(".self")) || Equals(file, TC("Symbols.map"));
 		//}
-		virtual bool KeepInMemory(const tchar* fileName, u32 fileNameLen, const tchar* systemTemp) override
+		virtual bool KeepInMemory(const tchar* fileName, u32 fileNameLen, const tchar* systemTemp) const override
 		{
 			return Super::KeepInMemory(fileName, fileNameLen, systemTemp)
 				|| Contains(fileName, TC("thinlto-"));// Used by a clang based platform's link time optimization pass. Shared from lto process back to linker process
 		}
-		virtual bool NeedsSharedMemory(const tchar* file) override
+		virtual bool NeedsSharedMemory(const tchar* file) const override
 		{
 			return Contains(file, TC("thinlto-")); // Used by a clang based platform's link time optimization pass. Shared from lto process back to linker process
 		}
@@ -271,12 +271,12 @@ namespace uba
 	{
 		using Super = ApplicationRulesClangPlusPlusExe;
 
-		virtual bool IsOutputFile(const tchar* file, u64 fileLen) override
+		virtual bool IsOutputFile(const tchar* file, u64 fileLen) const override
 		{
 			return Contains(file, TC(".self")) || Super::IsOutputFile(file, fileLen);
 		}
 
-		virtual bool IsThrowAway(const tchar* fileName, u32 fileNameLen) override
+		virtual bool IsThrowAway(const tchar* fileName, u32 fileNameLen) const override
 		{
 			return EndsWith(fileName, fileNameLen, TC("-telemetry.json")) || Super::IsThrowAway(fileName, fileNameLen);
 		}
@@ -291,12 +291,12 @@ namespace uba
 	{
 		using Super = ApplicationRules;
 
-		virtual bool IsOutputFile(const tchar* file, u64 fileLen) override
+		virtual bool IsOutputFile(const tchar* file, u64 fileLen) const override
 		{
 			return Contains(file, TC(".self"));
 		}
 
-		virtual bool IsThrowAway(const tchar* fileName, u32 fileNameLen) override
+		virtual bool IsThrowAway(const tchar* fileName, u32 fileNameLen) const override
 		{
 			return EndsWith(fileName, fileNameLen, TC("-telemetry.json")) || Super::IsThrowAway(fileName, fileNameLen);
 		}
@@ -311,12 +311,12 @@ namespace uba
 
 	class ApplicationRulesISPCExe : public ApplicationRules
 	{
-		virtual bool AllowDetach() override
+		virtual bool AllowDetach() const override
 		{
 			return true;
 		}
 
-		virtual bool IsOutputFile(const tchar* file, u64 fileLen) override
+		virtual bool IsOutputFile(const tchar* file, u64 fileLen) const override
 		{
 			return EndsWith(file, fileLen, TC(".dummy.h"))
 				|| EndsWith(file, fileLen, TC(".ispc.bc"))
@@ -327,7 +327,7 @@ namespace uba
 
 	class ApplicationRulesUBTDll : public ApplicationRules
 	{
-		virtual bool IsOutputFile(const tchar* file, u64 fileLen) override
+		virtual bool IsOutputFile(const tchar* file, u64 fileLen) const override
 		{
 			return false;
 			// TODO: These does not work when UnrealBuildTool creates these files multiple times in a row (building multiple targets)
@@ -340,14 +340,14 @@ namespace uba
 
 	class ApplicationRulesPVSStudio : public ApplicationRules
 	{
-		virtual bool IsOutputFile(const tchar* file, u64 fileLen) override
+		virtual bool IsOutputFile(const tchar* file, u64 fileLen) const override
 		{
 			return EndsWith(file, fileLen, TC(".PVS-Studio.log"))
 				|| EndsWith(file, fileLen, TC(".pvslog"))
 				|| EndsWith(file, fileLen, TC(".stacktrace.txt"));
 		}
 		
-		virtual bool IsRarelyRead(const StringBufferBase& file) override
+		virtual bool IsRarelyRead(const StringBufferBase& file) const override
 		{
 			return file.EndsWith(TC(".i"))
 				|| file.EndsWith(TC(".PVS-Studio.log"))
@@ -356,7 +356,7 @@ namespace uba
 		}
 
 #if PLATFORM_WINDOWS
-		virtual void RepairMalformedLibPath(const tchar* path) override
+		virtual void RepairMalformedLibPath(const tchar* path) const override
 		{
 			// There is a bug where the path passed into wsplitpath_s is malformed and not null terminated correctly
 			const tchar* pext = TStrstr(path, TC(".dll"));
@@ -370,7 +370,7 @@ namespace uba
 
 	class ApplicationRulesShaderCompileWorker : public ApplicationRules
 	{
-		virtual bool IsRarelyRead(const StringBufferBase& file) override
+		virtual bool IsRarelyRead(const StringBufferBase& file) const override
 		{
 			return file.Contains(TC(".uba."));
 		}
@@ -411,6 +411,14 @@ namespace uba
 			{ nullptr, nullptr }
 		};
 
+		static bool initRules = []()
+			{
+				u32 index = 0;
+				for (RulesRec* rec = rules; rec->app; ++rec)
+					rec->rules->index = index++;
+				return true;
+			}();
+
 		return rules;
 	}
 }
\ No newline at end of file
diff --git a/Engine/Source/Programs/UnrealBuildAccelerator/Core/Private/UbaStringBuffer.cpp b/Engine/Source/Programs/UnrealBuildAccelerator/Core/Private/UbaStringBuffer.cpp
--- a/Engine/Source/Programs/UnrealBuildAccelerator/Core/Private/UbaStringBuffer.cpp
+++ b/Engine/Source/Programs/UnrealBuildAccelerator/Core/Private/UbaStringBuffer.cpp
@@ -240,6 +240,16 @@ namespace uba
 		data[count] = 0;
 		return *this;
 	}
+	
+	StringBufferBase& StringBufferBase::Append(const char* str, u32 charCount)
+	{
+		u32 capacityEnd = capacity - 1;
+		for (const char* i = str; *i && charCount; ++i, --charCount)
+			if (count < capacityEnd)
+				data[count++] = *i;
+		data[count] = 0;
+		return *this;
+	}
 	#endif
 
 	StringBufferBase& StringBufferBase::Resize(u64 newSize)
diff --git a/Engine/Source/Programs/UnrealBuildAccelerator/Core/Public/UbaApplicationRules.h b/Engine/Source/Programs/UnrealBuildAccelerator/Core/Public/UbaApplicationRules.h
--- a/Engine/Source/Programs/UnrealBuildAccelerator/Core/Public/UbaApplicationRules.h
+++ b/Engine/Source/Programs/UnrealBuildAccelerator/Core/Public/UbaApplicationRules.h
@@ -12,14 +12,14 @@ namespace uba
 	public:
 		// This means that process can run entirely without console. (win32 flag DETACHED_PROCESS)
 		// Uba stub out console interaction. This is an optimization that is entirely optional
-		virtual bool AllowDetach()
+		virtual bool AllowDetach() const
 		{
 			return false;
 		}
 
 		// Some kernel functions are not detoured and can't handle detoured handles (pipe system for example)
 		// This function can be used to prevent uba from detouring handles. Defaults to not detour pipes
-		virtual bool CanDetour(const tchar* file)
+		virtual bool CanDetour(const tchar* file) const
 		{
 			if constexpr (IsWindows)
 			{
@@ -34,46 +34,46 @@ namespace uba
 		}
 
 		// Throw-away means that the file is temporary and will not be used after process exists. (By default these are kept in memory and never touch disk)
-		virtual bool IsThrowAway(const tchar* fileName, u32 fileNameLen)
+		virtual bool IsThrowAway(const tchar* fileName, u32 fileNameLen) const
 		{
 			return false;
 		}
 
 		// Keep file in memory
 		// If this returns true it means that file will be kept in memory and never touch disk.
-		virtual bool KeepInMemory(const tchar* fileName, u32 fileNameLen, const tchar* systemTemp)
+		virtual bool KeepInMemory(const tchar* fileName, u32 fileNameLen, const tchar* systemTemp) const
 		{
 			return IsThrowAway(fileName, fileNameLen);
 		}
 
 		// For files that are kept in memory but shared between process (temporary files where one process write and another read)
 		// This only works if file is never read outside of the same process hierarchy
-		virtual bool NeedsSharedMemory(const tchar* file)
+		virtual bool NeedsSharedMemory(const tchar* file) const
 		{
 			return false;
 		}
 
 		// Max file size if using memory files.
-		virtual u64 FileTypeMaxSize(const StringBufferBase& file, bool isSystemOrTempFile)
+		virtual u64 FileTypeMaxSize(const StringBufferBase& file, bool isSystemOrTempFile) const
 		{
 			return 8ull * 1024 * 1024 * 1024;
 		}
 
 		// Outputfile means that it is kept in memory and then sent back to session process which can decide to write it to disk or send it over network
-		virtual bool IsOutputFile(const tchar* file, u64 fileLen)
+		virtual bool IsOutputFile(const tchar* file, u64 fileLen) const
 		{
 			return false;
 		}
 
 		// If returns false this means that all GetFileAttribute etc will return file-not-found on this file
-		virtual bool CanExist(const tchar* file)
+		virtual bool CanExist(const tchar* file) const
 		{
 			return true;
 		}
 
 		// Return true if the file is only read by this process or very rarely read more than once
 		// This is an optimization to not store the file in the mapping table since it will not be read again and would just take up space
-		virtual bool IsRarelyRead(const StringBufferBase& file)
+		virtual bool IsRarelyRead(const StringBufferBase& file) const
 		{
 			return false;
 		}
@@ -85,44 +85,46 @@ namespace uba
 			return false;
 		}
 
-		virtual bool AllowStorageProxy(const StringBufferBase& file)
+		virtual bool AllowStorageProxy(const StringBufferBase& file) const
 		{
 			return !IsRarelyRead(file);
 		}
 
 		// Enable vectored exception handler. This can't be enabled for all processes because some of them actually allow
 		// access violations etc and then catch them.
-		virtual bool EnableVectoredExceptionHandler()
+		virtual bool EnableVectoredExceptionHandler() const
 		{
 			return false;
 		}
 
-		virtual bool AllowMiMalloc()
+		virtual bool AllowMiMalloc() const
 		{
 			return true;// !IsRunningWine();
 		}
 
-		virtual bool AllowLoadLibrary(const tchar* libraryName)
+		virtual bool AllowLoadLibrary(const tchar* libraryName) const
 		{
 			if (Contains(libraryName, TC("nvinject.dll")))
 				return false;
 			return true;
 		}
 
-		virtual bool SuppressLogLine(const tchar* logLine, u32 logLineLen)
+		virtual bool SuppressLogLine(const tchar* logLine, u32 logLineLen) const
 		{
 			return false;
 		}
 
-		virtual bool IsExitCodeSuccess(u32 exitCode)
+		virtual bool IsExitCodeSuccess(u32 exitCode) const
 		{
 			return true;
 		}
 
-		virtual void RepairMalformedLibPath(const tchar* path)
+		virtual void RepairMalformedLibPath(const tchar* path) const
 		{
 			// Do nothing
 		}
+
+		u32 index = ~0u;
 	};
 
 
diff --git a/Engine/Source/Programs/UnrealBuildAccelerator/Core/Public/UbaBinaryReaderWriter.h b/Engine/Source/Programs/UnrealBuildAccelerator/Core/Public/UbaBinaryReaderWriter.h
--- a/Engine/Source/Programs/UnrealBuildAccelerator/Core/Public/UbaBinaryReaderWriter.h
+++ b/Engine/Source/Programs/UnrealBuildAccelerator/Core/Public/UbaBinaryReaderWriter.h
@@ -144,8 +144,9 @@ namespace uba
 
 	inline u64 GetStringWriteSize(const tchar* str)
 	{
-		u64 actualBytes = GetWrittenBytes(str, TStrlen(str));
-		return Get7BitEncodedCount(actualBytes) + actualBytes;
+		u32 strLen = TStrlen(str);
+		u64 actualBytes = GetWrittenBytes(str, strLen);
+		return Get7BitEncodedCount(strLen) + actualBytes;
 	}
 
 	void BinaryWriter::WriteString(const tchar* str, u64 strLen)
@@ -250,7 +251,7 @@ namespace uba
 
 	void BinaryWriter::Write7BitEncoded(u64 value)
 	{
-		UBA_ASSERT_WRITE(5);
+		UBA_ASSERT_WRITE(Get7BitEncodedCount(value));
 		do
 		{
 			u8 HasMoreBytes = (u8)((value > u64(0x7F)) << 7);
diff --git a/Engine/Source/Programs/UnrealBuildAccelerator/Core/Public/UbaStringBuffer.h b/Engine/Source/Programs/UnrealBuildAccelerator/Core/Public/UbaStringBuffer.h
--- a/Engine/Source/Programs/UnrealBuildAccelerator/Core/Public/UbaStringBuffer.h
+++ b/Engine/Source/Programs/UnrealBuildAccelerator/Core/Public/UbaStringBuffer.h
@@ -37,6 +37,7 @@ namespace uba
 
 		#if PLATFORM_WINDOWS
 		StringBufferBase& Append(const char* str);
+		StringBufferBase& Append(const char* str, u32 charCount);
 		#endif
 
 		tchar operator[](u64 i) const { return data[i]; }
diff --git a/Engine/Source/Programs/UnrealBuildAccelerator/Detours/Private/Windows/UbaDetoursFunctionsKernelBase.inl b/Engine/Source/Programs/UnrealBuildAccelerator/Detours/Private/Windows/UbaDetoursFunctionsKernelBase.inl
--- a/Engine/Source/Programs/UnrealBuildAccelerator/Detours/Private/Windows/UbaDetoursFunctionsKernelBase.inl
+++ b/Engine/Source/Programs/UnrealBuildAccelerator/Detours/Private/Windows/UbaDetoursFunctionsKernelBase.inl
@@ -2340,9 +2340,12 @@ HMODULE Recursive_LoadLibraryExW(LPCWSTR lpLibFileName, LPCWSTR originalName, DW
 		if (importedModule.isKnown && !g_isRunningWine)
 			continue;
 
-		HMODULE checkModule = GetModuleHandleW(importedModule.name);
-		if (checkModule)
-			continue;
+		{
+			SuppressCreateFileDetourScope cfs;
+			HMODULE checkModule = GetModuleHandleW(importedModule.name); // This function ends up in NtCreateFile when running in wine
+			if (checkModule)
+				continue;
+		}
 
 		if (importedModule.isKnown) // We need to catch dbghelp.dll and imagehlp.dll
 		{
diff --git a/Engine/Source/Programs/UnrealBuildAccelerator/Detours/Private/Windows/UbaDetoursFunctionsNtDll.inl b/Engine/Source/Programs/UnrealBuildAccelerator/Detours/Private/Windows/UbaDetoursFunctionsNtDll.inl
--- a/Engine/Source/Programs/UnrealBuildAccelerator/Detours/Private/Windows/UbaDetoursFunctionsNtDll.inl
+++ b/Engine/Source/Programs/UnrealBuildAccelerator/Detours/Private/Windows/UbaDetoursFunctionsNtDll.inl
@@ -684,6 +684,8 @@ NTSTATUS NTAPI Shared_NtCreateFile(bool IsCreateFunc, PHANDLE hFileHandle, ACCES
 			return res;
 		if (!isSystemFile && !isWrite && !t_disallowDetour)
 			TrackInput(fileName.data);
+		else
+			SkipTrackInput(fileName.data);
 		return res;
 	}
 
@@ -929,10 +931,17 @@ NTSTATUS NTAPI Shared_NtCreateFile(bool IsCreateFunc, PHANDLE hFileHandle, ACCES
 
 	auto TrackFileInput = [&]()
 		{
-			if (!info.tracked && dwDesiredAccess && (dwDesiredAccess & GENERIC_WRITE) == 0)
+			if (!keepInMemory && dwDesiredAccess && (dwDesiredAccess & GENERIC_WRITE) == 0)
+			{
+				if (!info.tracked)
+				{
+					info.tracked = true;
+					TrackInput(fileName.data);
+				}
+			}
+			else
 			{
-				info.tracked = true;
-				TrackInput(fileName.data);
+				SkipTrackInput(fileName.data);
 			}
 		};
 
diff --git a/Engine/Source/Programs/UnrealBuildAccelerator/Detours/Private/Windows/UbaDetoursFunctionsWin.cpp b/Engine/Source/Programs/UnrealBuildAccelerator/Detours/Private/Windows/UbaDetoursFunctionsWin.cpp
--- a/Engine/Source/Programs/UnrealBuildAccelerator/Detours/Private/Windows/UbaDetoursFunctionsWin.cpp
+++ b/Engine/Source/Programs/UnrealBuildAccelerator/Detours/Private/Windows/UbaDetoursFunctionsWin.cpp
@@ -126,6 +126,10 @@ void TrackInput(const wchar_t* file)
 		g_trackInputsBufPos = u32(w.GetPosition());
 	}
 }
+void SkipTrackInput(const wchar_t* file)
+{
+	// Just here to easily log out what we are ignoring in terms of input
+}
 
 u8 g_emptyMemoryFileMem;
 MemoryFile& g_emptyMemoryFile = *new MemoryFile(&g_emptyMemoryFileMem, true);
