[UBA]
* Fixed so tracked inputs can be bigger than 512kb (ouch)
* Added to BinaryWriter so you can always write utf8 even though platform is not utf8
* Fixed bug in IsKnownSystemFile (there was a missing comma on one line). Added unit test
* Added Swap function to MemoryBlock which can be used to swap content of memoryblocks
* Added Parse function to StringBuffer that can parse a wchar stringbuffer into a char array
* Fixed TimeToText for times over 24 hours
* Added some accessors to Session
* Fixed so some temp files are not tracked as inputs for detoured process
* Added assert when message is created without body (it is a special case scenario)
* Removed final keyword on NetworkClient to be able to subclass in export logic (it is just easier)

diff --git a/Engine/Source/Programs/UnrealBuildAccelerator/Common/Private/UbaNetworkClient.cpp b/Engine/Source/Programs/UnrealBuildAccelerator/Common/Private/UbaNetworkClient.cpp
--- a/Engine/Source/Programs/UnrealBuildAccelerator/Common/Private/UbaNetworkClient.cpp
+++ b/Engine/Source/Programs/UnrealBuildAccelerator/Common/Private/UbaNetworkClient.cpp
@@ -504,7 +504,9 @@ namespace uba
 		u32 sendSize = u32(writer.GetPosition());
 		u8* data = writer.GetData();
 		data[1] = messageId >> 8;
-		*(u32*)(data + 2) = (sendSize - 6) | u32(messageId) << 24;
+		u32 dataSize = sendSize - 6;
+		UBA_ASSERTF(dataSize, TC("NetworkMessage must have data size of at least 1."));
+		*(u32*)(data + 2) = dataSize | u32(messageId) << 24;
 
 		//m_logger.Debug(TC("Send: %u, %u, %u, %u"), data[0], data[1], data[2], sendSize - 7);
 
diff --git a/Engine/Source/Programs/UnrealBuildAccelerator/Common/Private/UbaProcess.cpp b/Engine/Source/Programs/UnrealBuildAccelerator/Common/Private/UbaProcess.cpp
--- a/Engine/Source/Programs/UnrealBuildAccelerator/Common/Private/UbaProcess.cpp
+++ b/Engine/Source/Programs/UnrealBuildAccelerator/Common/Private/UbaProcess.cpp
@@ -899,11 +899,10 @@ namespace uba
 	bool ProcessImpl::HandleInputDependencies(BinaryReader& reader, BinaryWriter& writer)
 	{
 		UBA_ASSERT(m_startInfo.trackInputs);
-		if (m_trackedInputs.empty())
-		{
-			u32 trackedInputsSize = reader.ReadU32();
-			m_trackedInputs.reserve(trackedInputsSize);
-		}
+
+		if (u64 reserveSize = reader.Read7BitEncoded())
+			m_trackedInputs.reserve(m_trackedInputs.size() + reserveSize);
+
 		u32 toRead = reader.ReadU32();
 		u8* pos = m_trackedInputs.data() + m_trackedInputs.size();
 		m_trackedInputs.resize(m_trackedInputs.size() + toRead);
@@ -1083,7 +1082,7 @@ namespace uba
 		{
 			u64 totalBytes = 0;
 			for (auto& kv : m_writtenFiles)
-				totalBytes += GetStringWriteSize(kv.second.name.c_str());
+				totalBytes += GetStringWriteSize(kv.second.name.c_str(), kv.second.name.size());
 			m_trackedOutputs.resize(totalBytes);
 			BinaryWriter writer(m_trackedOutputs.data(), 0, totalBytes);
 			for (auto& kv : m_writtenFiles)
diff --git a/Engine/Source/Programs/UnrealBuildAccelerator/Common/Private/UbaSession.cpp b/Engine/Source/Programs/UnrealBuildAccelerator/Common/Private/UbaSession.cpp
--- a/Engine/Source/Programs/UnrealBuildAccelerator/Common/Private/UbaSession.cpp
+++ b/Engine/Source/Programs/UnrealBuildAccelerator/Common/Private/UbaSession.cpp
@@ -138,6 +138,11 @@ namespace uba
 		UBA_ASSERT(m_process);
 		return m_process->GetTrackedInputs();
 	}
+	const Vector<u8>& ProcessHandle::GetTrackedOutputs() const
+	{
+		UBA_ASSERT(m_process);
+		return m_process->GetTrackedOutputs();
+	}
 	u64 ProcessHandle::GetTotalProcessorTime() const
 	{
 		UBA_ASSERT(m_process);
@@ -1276,12 +1281,12 @@ namespace uba
 		return ToStringKey(hasher, baseFileNameForHash.data, baseFileNameForHash.count);
 	}
 
-	void Session::RegisterNewFile(const tchar* filePath)
+	bool Session::RegisterNewFile(const tchar* filePath)
 	{
 		UBA_ASSERT(!m_runningRemote);
 		StringBuffer<> fixedFilePath;
 		auto key = GetKeyAndFixedName(fixedFilePath, filePath);
-		RegisterCreateFileForWrite(key, fixedFilePath.data, fixedFilePath.count, true);
+		return RegisterCreateFileForWrite(key, fixedFilePath.data, fixedFilePath.count, true);
 	}
 
 	void Session::RegisterDeleteFile(const tchar* filePath)
@@ -1373,6 +1378,11 @@ namespace uba
 		return GetApplicationRules()[0].rules;
 	}
 
+	const tchar* Session::GetTempPath()
+	{
+		return m_tempPath.data;
+	}
+
 	void Session::ProcessAdded(Process& process, u32 sessionId)
 	{
 		u32 processId = process.GetId();
@@ -2012,7 +2022,7 @@ namespace uba
 
 	bool Session::GetFullFileName(GetFullFileNameResponse& out, const GetFullFileNameMessage& msg)
 	{
-		UBA_ASSERTF(false, TC("SHOULD NOT HAPPEN (only remote)"));
+		UBA_ASSERTF(false, TC("SHOULD NOT HAPPEN (only remote).. %s"), msg.fileName.data);
 		return SearchPathForFile(m_logger, out.fileName, msg.fileName.data, msg.process.m_virtualApplicationDir.c_str());
 	}
 
diff --git a/Engine/Source/Programs/UnrealBuildAccelerator/Common/Private/UbaSessionClient.cpp b/Engine/Source/Programs/UnrealBuildAccelerator/Common/Private/UbaSessionClient.cpp
--- a/Engine/Source/Programs/UnrealBuildAccelerator/Common/Private/UbaSessionClient.cpp
+++ b/Engine/Source/Programs/UnrealBuildAccelerator/Common/Private/UbaSessionClient.cpp
@@ -1231,13 +1231,26 @@ namespace uba
 
 	bool SessionClient::SendProcessInputs(ProcessImpl& process)
 	{
-		StackBinaryWriter<SendMaxSize> writer;
-		NetworkMessage msg(m_client, ServiceId, SessionMessageType_ProcessInputs, writer);
-		writer.WriteU32(process.m_id);
 		auto inputs = process.GetTrackedInputs();
-		writer.WriteBytes(inputs.data(), inputs.size());
-		StackBinaryReader<32> reader;
-		return msg.Send(reader);
+		u32 left = u32(inputs.size());
+		u32 capacityToAdd = left;
+		u8* readPos = inputs.data();
+		while (left)
+		{
+			StackBinaryWriter<SendMaxSize> writer;
+			NetworkMessage msg(m_client, ServiceId, SessionMessageType_ProcessInputs, writer);
+			writer.Write7BitEncoded(process.m_id);
+			writer.Write7BitEncoded(capacityToAdd);
+			capacityToAdd = 0;
+			u32 toWrite = Min(left, u32(writer.GetCapacityLeft()));
+			writer.WriteBytes(readPos, toWrite);
+			StackBinaryReader<32> reader;
+			if (!msg.Send(reader))
+				return false;
+			readPos += toWrite;
+			left -= toWrite;
+		}
+		return true;
 	}
 
 	bool SessionClient::SendProcessFinished(ProcessImpl& process, u32 exitCode)
diff --git a/Engine/Source/Programs/UnrealBuildAccelerator/Common/Private/UbaSessionServer.cpp b/Engine/Source/Programs/UnrealBuildAccelerator/Common/Private/UbaSessionServer.cpp
--- a/Engine/Source/Programs/UnrealBuildAccelerator/Common/Private/UbaSessionServer.cpp
+++ b/Engine/Source/Programs/UnrealBuildAccelerator/Common/Private/UbaSessionServer.cpp
@@ -868,7 +868,7 @@ namespace uba
 						auto& process = *(RemoteProcess*)h.m_process;
 						if (process.startInfo.trackInputs)
 						{
-							u64 bytes = GetStringWriteSize(destination.data);
+							u64 bytes = GetStringWriteSize(destination.data, destination.count);
 							u64 prevSize = process.m_trackedOutputs.size();
 							process.m_trackedOutputs.resize(prevSize + bytes);
 							BinaryWriter w2(process.m_trackedOutputs.data(), prevSize, prevSize + bytes);
@@ -1067,7 +1067,7 @@ namespace uba
 			}
 			case SessionMessageType_ProcessInputs:
 			{
-				u32 processId = reader.ReadU32();
+				u32 processId = u32(reader.Read7BitEncoded());
 				SCOPED_WRITE_LOCK(m_processesLock, lock);
 				auto findIt = m_processes.find(processId);
 				if (findIt == m_processes.end())
@@ -1075,9 +1075,13 @@ namespace uba
 				ProcessHandle h(findIt->second);
 				lock.Leave();
 				auto& process = *(RemoteProcess*)h.m_process;
-				u64 size = reader.GetLeft();
-				process.m_trackedInputs.resize(size);
-				reader.ReadBytes(process.m_trackedInputs.data(), size);
+				auto& inputs = process.m_trackedInputs;
+				u64 size = inputs.size();
+				if (u64 addCapacity = reader.Read7BitEncoded())
+					inputs.reserve(size + addCapacity);
+				u64 toRead = reader.GetLeft();
+				inputs.resize(size + toRead);
+				reader.ReadBytes(inputs.data() + size, toRead);
 				return true;
 			}
 			case SessionMessageType_ProcessFinished:
diff --git a/Engine/Source/Programs/UnrealBuildAccelerator/Common/Private/UbaStorageClient.cpp b/Engine/Source/Programs/UnrealBuildAccelerator/Common/Private/UbaStorageClient.cpp
--- a/Engine/Source/Programs/UnrealBuildAccelerator/Common/Private/UbaStorageClient.cpp
+++ b/Engine/Source/Programs/UnrealBuildAccelerator/Common/Private/UbaStorageClient.cpp
@@ -872,7 +872,7 @@ namespace uba
 		StorageStats& stats = Stats();
 		TimerScope ts(stats.sendCas);
 
-		u64 firstMessageOverHead = (sizeof(CasKey) + sizeof(u64)*2 + GetStringWriteSize(hint));
+		u64 firstMessageOverHead = (sizeof(CasKey) + sizeof(u64)*2 + GetStringWriteSize(hint, TStrlen(hint)));
 
 		u64 messageHeader = client.GetMessageHeaderSize();
 		u64 messageHeaderMaxSize = messageHeader + firstMessageOverHead;
diff --git a/Engine/Source/Programs/UnrealBuildAccelerator/Common/Public/UbaNetworkClient.h b/Engine/Source/Programs/UnrealBuildAccelerator/Common/Public/UbaNetworkClient.h
--- a/Engine/Source/Programs/UnrealBuildAccelerator/Common/Public/UbaNetworkClient.h
+++ b/Engine/Source/Programs/UnrealBuildAccelerator/Common/Public/UbaNetworkClient.h
@@ -27,7 +27,7 @@ namespace uba
 		const u8* cryptoKey128 = nullptr;
 	};
 
-	class NetworkClient final : public WorkManagerImpl
+	class NetworkClient : public WorkManagerImpl
 	{
 	public:
 		NetworkClient(bool& outCtorSuccess, const NetworkClientCreateInfo& info = {}, const tchar* name = TC("UbaClient"));
diff --git a/Engine/Source/Programs/UnrealBuildAccelerator/Common/Public/UbaProcessHandle.h b/Engine/Source/Programs/UnrealBuildAccelerator/Common/Public/UbaProcessHandle.h
--- a/Engine/Source/Programs/UnrealBuildAccelerator/Common/Public/UbaProcessHandle.h
+++ b/Engine/Source/Programs/UnrealBuildAccelerator/Common/Public/UbaProcessHandle.h
@@ -25,6 +25,7 @@ namespace uba
 	{
 		ProcessExecutionType_Native,
 		ProcessExecutionType_Detoured,
+		ProcessExecutionType_FromCache,
 	};
 
 	class ProcessHandle
@@ -38,6 +39,7 @@ namespace uba
 		bool WaitForExit(u32 millisecondsTimeout) const;	// Wait for process until it exists or time reach timeout
 		const Vector<ProcessLogLine>& GetLogLines() const;	// Log lines produced by the process
 		const Vector<u8>& GetTrackedInputs() const;			// If ProcessStartInfo.trackInputs was true, this will return a buffer to wchar/char strings 
+		const Vector<u8>& GetTrackedOutputs() const;		// If ProcessStartInfo.trackInputs was true, this will return a buffer to wchar/char strings 
 		u64 GetTotalProcessorTime() const;					// Total used cpu time (in time units. Use TimeToMs etc)
 		u64 GetTotalWallTime() const;						// Total wall time (in time units. Use TimeToMs etc)
 		void Cancel(bool terminate) const;					// Request to cancel process. Set terminate to true to quickly kill process
diff --git a/Engine/Source/Programs/UnrealBuildAccelerator/Common/Public/UbaSession.h b/Engine/Source/Programs/UnrealBuildAccelerator/Common/Public/UbaSession.h
--- a/Engine/Source/Programs/UnrealBuildAccelerator/Common/Public/UbaSession.h
+++ b/Engine/Source/Programs/UnrealBuildAccelerator/Common/Public/UbaSession.h
@@ -54,7 +54,7 @@ namespace uba
 
 		void PrintSummary(Logger& logger); // Print summary stats of session
 		void RefreshDirectory(const tchar* dirName); // Tell uba a directory on disk has been changed by some other system while session is running
-		void RegisterNewFile(const tchar* filePath); // Tell uba a new file on disk has been added by some other system while session is running
+		bool RegisterNewFile(const tchar* filePath); // Tell uba a new file on disk has been added by some other system while session is running
 		void RegisterDeleteFile(const tchar* filePath); // Tell uba a file on disk has been deleted by some other system while session is running
 
 		using CustomServiceFunction = Function<u32(Process& handle, const void* recv, u32 recvSize, void* send, u32 sendCapacity)>;
@@ -63,12 +63,14 @@ namespace uba
 		using GetNextProcessFunction = Function<bool(Process& handle, NextProcessInfo& outNextProcess, u32 prevExitCode)>;
 		void RegisterGetNextProcess(GetNextProcessFunction&& function); // Register a custom service (that can be communicated with from the remote agents)
 
-		const tchar* GetId();			// Id for session. Will be "yymmdd_hhmmss" unless SessionCreateInfo.useUniqueId is set to false
-		u32 GetActiveProcessCount();	// Current active processes running inside session
-		Storage& GetStorage();		// Storage (only used when remote machines are connected)
-		Logger& GetLogger();			// Logger used for logging 
-		LogWriter& GetLogWriter();		// LogWriter used by logger
+		const tchar* GetId(); // Id for session. Will be "yymmdd_hhmmss" unless SessionCreateInfo.useUniqueId is set to false
+		u32 GetActiveProcessCount(); // Current active processes running inside session
+		Storage& GetStorage(); // Storage (only used when remote machines are connected)
+		Logger& GetLogger(); // Logger used for logging 
+		LogWriter& GetLogWriter(); // LogWriter used by logger
 		const ApplicationRules* GetRules(const ProcessStartInfo& si); // Get application rules used for process
+		const tchar* GetTempPath(); // Path for temp files used for current session
+		bool ShouldStoreObjFilesCompressed() { return m_storeObjFilesCompressed; }
 
 		virtual ~Session();
 
diff --git a/Engine/Source/Programs/UnrealBuildAccelerator/Core/Private/UbaMemory.cpp b/Engine/Source/Programs/UnrealBuildAccelerator/Core/Private/UbaMemory.cpp
--- a/Engine/Source/Programs/UnrealBuildAccelerator/Core/Private/UbaMemory.cpp
+++ b/Engine/Source/Programs/UnrealBuildAccelerator/Core/Private/UbaMemory.cpp
@@ -105,4 +105,21 @@ namespace uba
 		return (tchar*)mem;
 	}
 
+	void MemoryBlock::Swap(MemoryBlock& other)
+	{
+		u8* m = memory;
+		u64 rs = reserveSize;
+		u64 ws = writtenSize;
+		u64 ms = mappedSize;
+
+		memory = other.memory;
+		reserveSize = other.reserveSize;
+		writtenSize = other.writtenSize;
+		mappedSize = other.mappedSize;
+
+		other.memory = m;
+		other.reserveSize = rs;
+		other.writtenSize = ws;
+		other.mappedSize = ms;
+	}
 }
diff --git a/Engine/Source/Programs/UnrealBuildAccelerator/Core/Private/UbaStringBuffer.cpp b/Engine/Source/Programs/UnrealBuildAccelerator/Core/Private/UbaStringBuffer.cpp
--- a/Engine/Source/Programs/UnrealBuildAccelerator/Core/Private/UbaStringBuffer.cpp
+++ b/Engine/Source/Programs/UnrealBuildAccelerator/Core/Private/UbaStringBuffer.cpp
@@ -244,7 +244,7 @@ namespace uba
 	StringBufferBase& StringBufferBase::Append(const char* str, u32 charCount)
 	{
 		u32 capacityEnd = capacity - 1;
-		for (const char* i = str; *i && charCount; ++i, --charCount)
+		for (const char* i = str; charCount && *i; ++i, --charCount)
 			if (count < capacityEnd)
 				data[count++] = *i;
 		data[count] = 0;
@@ -360,4 +360,14 @@ namespace uba
 		#endif
 		return out != 0 || Equals(TC("0"));
 	}
+
+	#if PLATFORM_WINDOWS
+	u32 StringBufferBase::Parse(char* out, u64 outCapacity)
+	{
+		size_t destLen;
+		if (wcstombs_s(&destLen, out, outCapacity, data, outCapacity-1) != 0)
+			return 0;
+		return (u32)destLen;
+	}
+	#endif
 }
diff --git a/Engine/Source/Programs/UnrealBuildAccelerator/Core/Private/UbaTimer.cpp b/Engine/Source/Programs/UnrealBuildAccelerator/Core/Private/UbaTimer.cpp
--- a/Engine/Source/Programs/UnrealBuildAccelerator/Core/Private/UbaTimer.cpp
+++ b/Engine/Source/Programs/UnrealBuildAccelerator/Core/Private/UbaTimer.cpp
@@ -66,7 +66,12 @@ namespace uba
 			u32 totalMin = totalSec / 60;
 			u32 min = totalMin % 60;
 			u32 sec = totalSec % 60;
-			if (u32 hour = totalMin / 60)
+			u32 hour = totalMin / 60;
+			u32 days = hour / 24;
+			hour -= days*24;
+			if (days)
+				TSprintf_s(str, 32, TC("%ud%uh%um"), (unsigned int)days, (unsigned int)hour, (unsigned int)min);
+			else if (hour)
 				TSprintf_s(str, 32, TC("%uh%um%us"), (unsigned int)hour, (unsigned int)min, (unsigned int)sec);
 			else
 				TSprintf_s(str, 32, TC("%um%us"), (unsigned int)min, (unsigned int)sec);
diff --git a/Engine/Source/Programs/UnrealBuildAccelerator/Core/Public/UbaBinaryReaderWriter.h b/Engine/Source/Programs/UnrealBuildAccelerator/Core/Public/UbaBinaryReaderWriter.h
--- a/Engine/Source/Programs/UnrealBuildAccelerator/Core/Public/UbaBinaryReaderWriter.h
+++ b/Engine/Source/Programs/UnrealBuildAccelerator/Core/Public/UbaBinaryReaderWriter.h
@@ -16,7 +16,7 @@ namespace uba
 		inline void WriteU32(u32 value);
 		inline void WriteU64(u64 value);
 		inline void WriteString(const tchar* str);
-		inline void WriteString(const tchar* str, u64 strLen);
+		inline void WriteString(const tchar* str, u64 len);
 		inline void WriteString(const StringBufferBase& str);
 		inline void WriteStringf(const tchar* format, ...);
 		inline void WriteString(const TString& str);
@@ -25,6 +25,10 @@ namespace uba
 		inline void WriteCasKey(const CasKey& casKey);
 		inline void Write7BitEncoded(u64 value);
 		inline void WriteBool(bool value) { WriteByte(value ? 1 : 0); }
+
+		template<typename CharType>
+		inline void WriteUtf8String(const CharType* str, u64 len);
+
 		inline u8* AllocWrite(u64 bytes);
 		inline u64 GetPosition() const { return u64(m_pos - m_begin); }
 		inline u64 GetCapacityLeft() const { return u64(m_end - m_pos); }
@@ -115,10 +119,11 @@ namespace uba
 		WriteString(str, u64(TStrlen(str)));
 	}
 
-	inline u64 GetWrittenBytes(const tchar* str, u64 strlen)
+	template<typename CharType>
+	inline u64 GetUtf8WrittenBytes(const CharType* str, u64 strlen)
 	{
 		u64 actualBytes = 0;
-		for (const tchar* i = str, *e = str + strlen; i != e; ++i)
+		for (const CharType* i = str, *e = str + strlen; i != e; ++i)
 		{
 			u16 c = u16(*i);
 			if (c < 128)
@@ -142,23 +147,24 @@ namespace uba
 		return count;
 	}
 
-	inline u64 GetStringWriteSize(const tchar* str)
+	inline u64 GetStringWriteSize(const tchar* str, u64 len)
 	{
-		u32 strLen = TStrlen(str);
-		u64 actualBytes = GetWrittenBytes(str, strLen);
-		return Get7BitEncodedCount(strLen) + actualBytes;
+		#if PLATFORM_WINDOWS
+		u64 actualBytes = GetUtf8WrittenBytes(str, len);
+		#else
+		u64 actualBytes = len;
+		#endif
+		return Get7BitEncodedCount(len) + actualBytes;
 	}
 
-	void BinaryWriter::WriteString(const tchar* str, u64 strLen)
+	template<typename CharType>
+	void BinaryWriter::WriteUtf8String(const CharType* str, u64 len)
 	{
 		UBA_ASSERT(str);
 
-		#if PLATFORM_WINDOWS
-		Write7BitEncoded(strLen);
-
-		UBA_ASSERT_WRITE(GetWrittenBytes(str, strLen));
+		UBA_ASSERT_WRITE(GetUtf8WrittenBytes(str, len));
 
-		for (const tchar* i = str, *e = str + strLen; i != e; ++i)
+		for (const CharType* i = str, *e = str + len; i != e; ++i)
 		{
 			int c = *i;
 			if (c < 128)
@@ -175,9 +181,19 @@ namespace uba
 				*m_pos++ = u8((c % 64) + 128);
 			}
 		}
+	}
+
+	void BinaryWriter::WriteString(const tchar* str, u64 len)
+	{
+		UBA_ASSERT(str);
+
+		Write7BitEncoded(len);
+
+		#if PLATFORM_WINDOWS
+		UBA_ASSERT_WRITE(GetUtf8WrittenBytes(str, len));
+		WriteUtf8String(str, len);
 		#else
-		Write7BitEncoded(strLen);
-		WriteBytes(str, strLen);
+		WriteBytes(str, len);
 		#endif
 	}
 
diff --git a/Engine/Source/Programs/UnrealBuildAccelerator/Core/Public/UbaMemory.h b/Engine/Source/Programs/UnrealBuildAccelerator/Core/Public/UbaMemory.h
--- a/Engine/Source/Programs/UnrealBuildAccelerator/Core/Public/UbaMemory.h
+++ b/Engine/Source/Programs/UnrealBuildAccelerator/Core/Public/UbaMemory.h
@@ -101,6 +101,7 @@ namespace uba
 		void* AllocateNoLock(u64 bytes, u64 alignment, const tchar* hint);
 		void Free(void* p);
 		tchar* Strdup(const tchar* str);
+		void Swap(MemoryBlock& other);
 		
 
 		ReaderWriterLock lock;
diff --git a/Engine/Source/Programs/UnrealBuildAccelerator/Core/Public/UbaStringBuffer.h b/Engine/Source/Programs/UnrealBuildAccelerator/Core/Public/UbaStringBuffer.h
--- a/Engine/Source/Programs/UnrealBuildAccelerator/Core/Public/UbaStringBuffer.h
+++ b/Engine/Source/Programs/UnrealBuildAccelerator/Core/Public/UbaStringBuffer.h
@@ -63,6 +63,10 @@ namespace uba
 		bool Parse(u16& out);
 		bool Parse(float& out);
 
+		#if PLATFORM_WINDOWS
+		u32 Parse(char* out, u64 outCapacity);
+		#endif
+
 		u32 count;
 		u32 capacity;
 		tchar data[1];
diff --git a/Engine/Source/Programs/UnrealBuildAccelerator/Core/Public/UbaWinBinDependencyParser.h b/Engine/Source/Programs/UnrealBuildAccelerator/Core/Public/UbaWinBinDependencyParser.h
--- a/Engine/Source/Programs/UnrealBuildAccelerator/Core/Public/UbaWinBinDependencyParser.h
+++ b/Engine/Source/Programs/UnrealBuildAccelerator/Core/Public/UbaWinBinDependencyParser.h
@@ -14,75 +14,74 @@
 
 namespace uba
 {
-	inline bool IsKnownSystemFile(const tchar* applicationName)
+	// These dlls should _always_ exist on all machines so we can filter them out from the list of imports to copy
+	static const wchar_t* g_knownSystemFiles[] = // Sorted lowercase
 	{
-		// These dlls should _always_ exist on all machines so we can filter them out from the list of imports to copy
-
-		static const wchar_t* knownSystemFiles[] = // Sorted lowercase
-		{
-			L"advapi32.dll",
-			L"bcrypt.dll",
-			L"bcryptprimitives.dll",
-			L"combase.dll",
-			L"concrt140.dll",
-			L"crypt32.dll",
-			L"cryptbase.dll",
-			L"dbghelp.dll",
-			L"dnsapi.dll",
-			L"dwmapi.dll",
-			L"dxgi.dll",
-			L"fwpuclnt.dll",
-			L"gdi32.dll",
-			L"gdi32full.dll",
-			L"glu32.dll",
-			L"imagehlp.dll",
-			L"imm32.dll",
-			L"iphlpapi.dll",
-			L"kernel32.dll",
-			L"kernelbase.dll",
-			L"msvcp_win.dll",
-			L"msvcrt.dll",
-			L"mswsock.dll",
-			L"ncrypt.dll",
-			L"netapi32.dll",
-			L"nsi.dll",
-			L"ntasn1.dll",
-			L"ntdll.dll",
-			L"ole32.dll",
-			L"oleaut32.dll",
-			L"ondemandconnroutehelper.dll",
-			L"opengl32.dll",
-			L"powrprof.dll",
-			L"psapi.dll",
-			L"rasadhlp.dll",
-			L"rpcrt4.dll",
-			L"rstrtmgr.dll",
-			L"sechost.dll",
-			L"setupapi.dll"
-			L"shell32.dll",
-			L"sspicli.dll",
-			L"ucrtbase.dll",
-			L"umpdc.dll",
-			L"umppc17706.dll",
-			L"user32.dll",
-			L"userenv.dll",
-			L"uxtheme.dll",
-			L"version.dll",
-			L"webio.dll",
-			L"win32u.dll",
-			L"winhttp.dll",
-			L"winmm.dll",
-			L"winnsi.dll",
-			L"ws2_32.dll",
-		};
+		L"advapi32.dll",
+		L"bcrypt.dll",
+		L"bcryptprimitives.dll",
+		L"combase.dll",
+		L"concrt140.dll",
+		L"crypt32.dll",
+		L"cryptbase.dll",
+		L"dbghelp.dll",
+		L"dnsapi.dll",
+		L"dwmapi.dll",
+		L"dxgi.dll",
+		L"fwpuclnt.dll",
+		L"gdi32.dll",
+		L"gdi32full.dll",
+		L"glu32.dll",
+		L"imagehlp.dll",
+		L"imm32.dll",
+		L"iphlpapi.dll",
+		L"kernel32.dll",
+		L"kernelbase.dll",
+		L"msvcp_win.dll",
+		L"msvcrt.dll",
+		L"mswsock.dll",
+		L"ncrypt.dll",
+		L"netapi32.dll",
+		L"nsi.dll",
+		L"ntasn1.dll",
+		L"ntdll.dll",
+		L"ole32.dll",
+		L"oleaut32.dll",
+		L"ondemandconnroutehelper.dll",
+		L"opengl32.dll",
+		L"powrprof.dll",
+		L"psapi.dll",
+		L"rasadhlp.dll",
+		L"rpcrt4.dll",
+		L"rstrtmgr.dll",
+		L"sechost.dll",
+		L"setupapi.dll",
+		L"shell32.dll",
+		L"sspicli.dll",
+		L"ucrtbase.dll",
+		L"umpdc.dll",
+		L"umppc17706.dll",
+		L"user32.dll",
+		L"userenv.dll",
+		L"uxtheme.dll",
+		L"version.dll",
+		L"webio.dll",
+		L"win32u.dll",
+		L"winhttp.dll",
+		L"winmm.dll",
+		L"winnsi.dll",
+		L"ws2_32.dll",
+	};
 
+	inline bool IsKnownSystemFile(const tchar* fileName)
+	{
 		StringBuffer<> nameLower;
-		if (auto lastSeparator = TStrrchr(applicationName, PathSeparator))
+		if (auto lastSeparator = TStrrchr(fileName, PathSeparator))
 			nameLower.Append(lastSeparator + 1);
 		else
-			nameLower.Append(applicationName);
+			nameLower.Append(fileName);
 		nameLower.MakeLower();
-		return std::binary_search(knownSystemFiles, knownSystemFiles + sizeof_array(knownSystemFiles), nameLower.data, [](const wchar_t* a, const wchar_t* b) { return TStrcmp(a, b) < 0; });
+		return std::binary_search(g_knownSystemFiles, g_knownSystemFiles + sizeof_array(g_knownSystemFiles), nameLower.data, [](const wchar_t* a, const wchar_t* b) { return TStrcmp(a, b) < 0; });
 	}
 
 	inline DWORD GetOffset(DWORD rva, PIMAGE_SECTION_HEADER psh, PIMAGE_NT_HEADERS pnt)
diff --git a/Engine/Source/Programs/UnrealBuildAccelerator/Detours/Private/Windows/UbaDetoursFunctionsNtDll.inl b/Engine/Source/Programs/UnrealBuildAccelerator/Detours/Private/Windows/UbaDetoursFunctionsNtDll.inl
--- a/Engine/Source/Programs/UnrealBuildAccelerator/Detours/Private/Windows/UbaDetoursFunctionsNtDll.inl
+++ b/Engine/Source/Programs/UnrealBuildAccelerator/Detours/Private/Windows/UbaDetoursFunctionsNtDll.inl
@@ -682,7 +682,7 @@ NTSTATUS NTAPI Shared_NtCreateFile(bool IsCreateFunc, PHANDLE hFileHandle, ACCES
 
 		if (NT_ERROR(res))
 			return res;
-		if (!isSystemFile && !isWrite && !t_disallowDetour)
+		if (!isSystemOrTempFile && !isWrite && !t_disallowDetour && fileName[fileName.count-1] != ':')
 			TrackInput(fileName.data);
 		else
 			SkipTrackInput(fileName.data);
diff --git a/Engine/Source/Programs/UnrealBuildAccelerator/Detours/Private/Windows/UbaDetoursFunctionsWin.cpp b/Engine/Source/Programs/UnrealBuildAccelerator/Detours/Private/Windows/UbaDetoursFunctionsWin.cpp
--- a/Engine/Source/Programs/UnrealBuildAccelerator/Detours/Private/Windows/UbaDetoursFunctionsWin.cpp
+++ b/Engine/Source/Programs/UnrealBuildAccelerator/Detours/Private/Windows/UbaDetoursFunctionsWin.cpp
@@ -117,14 +117,41 @@ wchar_t* g_virtualCommandLine;
 constexpr u32 TrackInputsMemCapacity = 512 * 1024;
 u8* g_trackInputsMem;
 u32 g_trackInputsBufPos;
-void TrackInput(const wchar_t* file)
+
+void SendInput()
 {
-	if (g_trackInputsMem)
+	u32 left = g_trackInputsBufPos;
+	u32 reserveSize = left;
+	u32 pos = 0;
+	while (left)
 	{
-		BinaryWriter w(g_trackInputsMem, g_trackInputsBufPos, TrackInputsMemCapacity);
-		w.WriteString(file);
-		g_trackInputsBufPos = u32(w.GetPosition());
+		SCOPED_WRITE_LOCK(g_communicationLock, pcs);
+		BinaryWriter writer;
+		writer.WriteByte(MessageType_InputDependencies);
+		writer.Write7BitEncoded(reserveSize);
+		reserveSize = 0;
+		u32 toWrite = Min(left, u32(writer.GetCapacityLeft() - sizeof(u32)));
+		writer.WriteU32(toWrite);
+		writer.WriteBytes(g_trackInputsMem + pos, toWrite);
+		writer.Flush();
+		left -= toWrite;
+		pos += toWrite;
 	}
+	g_trackInputsBufPos = 0;
+}
+
+
+void TrackInput(const wchar_t* file)
+{
+	if (!g_trackInputsMem)
+		return;
+
+	if (g_trackInputsBufPos > TrackInputsMemCapacity - 2048)
+		SendInput();
+
+	BinaryWriter w(g_trackInputsMem, g_trackInputsBufPos, TrackInputsMemCapacity);
+	w.WriteString(file);
+	g_trackInputsBufPos = u32(w.GetPosition());
 }
 void SkipTrackInput(const wchar_t* file)
 {
@@ -465,24 +492,7 @@ void SendExitMessage(DWORD exitCode, u64 startTime)
 		Shared_WriteConsole(L"\n", 1, 0);
 
 	if (g_trackInputsMem)
-	{
-		u32 left = g_trackInputsBufPos;
-		u32 pos = 0;
-		while (left)
-		{
-			u32 toWrite = Min(left, u32(30 * 1024));
-			SCOPED_WRITE_LOCK(g_communicationLock, pcs);
-			BinaryWriter writer;
-			writer.WriteByte(MessageType_InputDependencies);
-			if (pos == 0)
-				writer.WriteU32(left);
-			writer.WriteU32(toWrite);
-			writer.WriteBytes(g_trackInputsMem + pos, toWrite);
-			writer.Flush();
-			left -= toWrite;
-			pos += toWrite;
-		}
-	}
+		SendInput();
 
 	g_stats.usedMemory = u32(g_memoryBlock.writtenSize);
 
diff --git a/Engine/Source/Programs/UnrealBuildAccelerator/Test/Private/UbaTestAll.h b/Engine/Source/Programs/UnrealBuildAccelerator/Test/Private/UbaTestAll.h
--- a/Engine/Source/Programs/UnrealBuildAccelerator/Test/Private/UbaTestAll.h
+++ b/Engine/Source/Programs/UnrealBuildAccelerator/Test/Private/UbaTestAll.h
@@ -44,6 +44,7 @@ namespace uba
 
 #define UBA_WINDOWS_TESTS \
 		UBA_NONMAC_TESTS \
+		UBA_TEST(TestKnownSystemFiles) \
 		UBA_TEST(TestCustomService) \
 		UBA_TEST(TestStdOutLocal) \
 		UBA_TEST(TestStdOutViaCmd) \
diff --git a/Engine/Source/Programs/UnrealBuildAccelerator/Test/Private/UbaTestBasics.h b/Engine/Source/Programs/UnrealBuildAccelerator/Test/Private/UbaTestBasics.h
--- a/Engine/Source/Programs/UnrealBuildAccelerator/Test/Private/UbaTestBasics.h
+++ b/Engine/Source/Programs/UnrealBuildAccelerator/Test/Private/UbaTestBasics.h
@@ -16,6 +16,10 @@
 #include "UbaTimer.h"
 #include "UbaDirectoryIterator.h"
 
+#if PLATFORM_WINDOWS
+#include "UbaWinBinDependencyParser.h"
+#endif
+
 #define VA_ARGS(...) , ##__VA_ARGS__
 #define UBA_TEST_CHECK(expr, fmt, ...) if (!(expr)) return logger.Error(TC(fmt) VA_ARGS(__VA_ARGS__));
 
@@ -291,4 +295,16 @@ namespace uba
 
 		return true;
 	}
+
+	#if PLATFORM_WINDOWS
+	bool TestKnownSystemFiles(Logger& logger, const StringBufferBase& rootDir)
+	{
+		for (auto systemFile : g_knownSystemFiles)
+			if (!IsKnownSystemFile(systemFile))
+				return logger.Error(TC("IsKnownSystemFile returned false for %s which is a system file"), systemFile);
+		if (IsKnownSystemFile(TC("Fooo.dll")))
+			return logger.Error(TC("IsKnownSystemFile returned true for Fooo.dll which is not a system file"));
+		return true;
+	}
+	#endif
 }
