[UBA]
* Added artifact cache support. It is now possible to skip running process and instead just download outputs if inputs matches cache entry.

diff --git a/Engine/Source/Programs/UnrealBuildAccelerator/CacheService/Private/UbaCacheService.cpp b/Engine/Source/Programs/UnrealBuildAccelerator/CacheService/Private/UbaCacheService.cpp
new file mode 100644
--- /dev/null
+++ b/Engine/Source/Programs/UnrealBuildAccelerator/CacheService/Private/UbaCacheService.cpp
@@ -0,0 +1,263 @@
+// Copyright Epic Games, Inc. All Rights Reserved.
+
+#include "UbaCacheServer.h"
+#include "UbaFile.h"
+#include "UbaNetworkBackendTcp.h"
+#include "UbaNetworkServer.h"
+#include "UbaPlatform.h"
+#include "UbaProtocol.h"
+#include "UbaStorageServer.h"
+#include "UbaVersion.h"
+
+#if PLATFORM_WINDOWS
+#include <dbghelp.h>
+#include <io.h>
+#pragma comment (lib, "Dbghelp.lib")
+#endif
+
+namespace uba
+{
+	const tchar*	Version = GetVersionString();
+	u32				DefaultCapacityGb = 500;
+	const tchar*	DefaultRootDir = []() {
+		static tchar buf[256];
+		if (IsWindows)
+			ExpandEnvironmentStringsW(TC("%ProgramData%\\Epic\\" UE_APP_NAME), buf, sizeof(buf));
+		else
+			GetFullPathNameW(TC("~/" UE_APP_NAME), sizeof_array(buf), buf, nullptr);
+		return buf;
+		}();
+	u32				DefaultProcessorCount = []() { return GetLogicalProcessorCount(); }();
+
+	int PrintHelp(const tchar* message)
+	{
+		LoggerWithWriter logger(g_consoleLogWriter, TC(""));
+		if (*message)
+		{
+			logger.Info(TC(""));
+			logger.Error(TC("%s"), message);
+		}
+		logger.Info(TC(""));
+		logger.Info(TC("-------------------------------------------"));
+		logger.Info(TC("   UbaCacheService v%s"), Version);
+		logger.Info(TC("-------------------------------------------"));
+		logger.Info(TC(""));
+		logger.Info(TC("  UbaCacheService.exe [options...]"));
+		logger.Info(TC(""));
+		logger.Info(TC("  Options:"));
+		logger.Info(TC("   -dir=<rootdir>          The directory used to store data. Defaults to \"%s\""), DefaultRootDir);
+		logger.Info(TC("   -port=[<host>:]<port>   The ip/name and port (default: %u) to listen for clients on"), DefaultCachePort);
+		logger.Info(TC(""));
+		return -1;
+	}
+
+	StorageServer* g_storageServer;
+	CacheServer* g_cacheServer;
+	
+	void CtrlBreakPressed()
+	{
+		if (g_storageServer)
+		{
+			g_storageServer->SaveCasTable(true);
+			LoggerWithWriter(g_consoleLogWriter).Info(TC("CAS table saved..."));
+
+			g_cacheServer->Save();
+		}
+		abort();
+	}
+
+	#if PLATFORM_WINDOWS
+	BOOL ConsoleHandler(DWORD signal)
+	{
+		if (signal == CTRL_C_EVENT)
+			CtrlBreakPressed();
+		return FALSE;
+	}
+	#else
+	void ConsoleHandler(int sig)
+	{
+		CtrlBreakPressed();
+	}
+	#endif
+	
+	StringBuffer<> g_rootDir(DefaultRootDir);
+
+	//LONG WINAPI UbaUnhandledExceptionFilter(EXCEPTION_POINTERS* ExceptionInfo)
+	//{
+	//	time_t rawtime;
+	//	time(&rawtime);
+	//	tm ti;
+	//	localtime_s(&ti, &rawtime);
+	//
+	//	StringBuffer<> dumpFile;
+	//	dumpFile.Append(g_rootDir).EnsureEndsWithSlash().Appendf(TC("UbaCliCrash_%02u%02u%02u_%02u%02u%02u.dmp"), ti.tm_year - 100, ti.tm_mon + 1, ti.tm_mday, ti.tm_hour, ti.tm_min, ti.tm_sec);
+	//
+	//	wprintf(TC("Unhandled exception - Writing minidump %s\n"), dumpFile.data);
+	//	HANDLE hFile = ::CreateFileW(dumpFile.data, GENERIC_READ | GENERIC_WRITE, FILE_SHARE_DELETE | FILE_SHARE_READ | FILE_SHARE_WRITE, NULL, CREATE_ALWAYS, FILE_ATTRIBUTE_NORMAL, NULL);
+	//	MINIDUMP_EXCEPTION_INFORMATION mei;
+	//	mei.ThreadId = GetCurrentThreadId();
+	//	mei.ClientPointers = TRUE;
+	//	mei.ExceptionPointers = ExceptionInfo;
+	//	MiniDumpWriteDump(GetCurrentProcess(), ::GetCurrentProcessId(), hFile, MiniDumpNormal, &mei, NULL, NULL);
+	//	return EXCEPTION_EXECUTE_HANDLER;
+	//}
+
+
+	int WrappedMain(int argc, tchar* argv[])
+	{
+		using namespace uba;
+		//SetUnhandledExceptionFilter(UbaUnhandledExceptionFilter);
+
+		u32 storageCapacityGb = DefaultCapacityGb;
+		StringBuffer<256> workDir;
+		StringBuffer<128> listenIp;
+		u16 port = DefaultCachePort;
+		bool quiet = false;
+		bool storeCompressed = true;
+
+		for (int i=1; i!=argc; ++i)
+		{
+			StringBuffer<> name;
+			StringBuffer<> value;
+
+			if (const tchar* equals = TStrchr(argv[i],'='))
+			{
+				name.Append(argv[i], equals - argv[i]);
+				value.Append(equals+1);
+			}
+			else
+			{
+				name.Append(argv[i]);
+			}
+
+			if (name.Equals(TC("-port")))
+			{
+				if (const tchar* portIndex = value.First(':'))
+				{
+					StringBuffer<> portStr(portIndex + 1);
+					if (!portStr.Parse(port))
+						return PrintHelp(TC("Invalid value for port in -port"));
+					listenIp.Append(value.data, portIndex - value.data);
+				}
+				else
+				{
+					if (!value.Parse(port))
+						return PrintHelp(TC("Invalid value for -port"));
+				}
+			}
+			else if (name.Equals(TC("-dir")))
+			{
+				if (value.IsEmpty())
+					return PrintHelp(TC("-dir needs a value"));
+				if ((g_rootDir.count = GetFullPathNameW(value.Replace('/', PathSeparator).data, g_rootDir.capacity, g_rootDir.data, nullptr)) == 0)
+					return PrintHelp(StringBuffer<>().Appendf(TC("-dir has invalid path %s"), g_rootDir.data).data);
+			}
+			else if (name.Equals(TC("-?")))
+			{
+				return PrintHelp(TC(""));
+			}
+			else
+			{
+				StringBuffer<> msg;
+				msg.Appendf(TC("Unknown argument '%s'"), name.data);
+				return PrintHelp(msg.data);
+			}
+		}
+
+		FilteredLogWriter logWriter(g_consoleLogWriter, quiet ? LogEntryType_Info : LogEntryType_Detail);
+		LoggerWithWriter logger(logWriter, TC(""));
+
+		const tchar* dbgStr = TC("");
+		#if UBA_DEBUG
+		dbgStr = TC(" (DEBUG)");
+		#endif
+		logger.Info(TC("UbaCacheService v%s%s (Rootdir: \"%s\", StoreCapacity: %uGb)\n"), Version, dbgStr, g_rootDir.data, storageCapacityGb);
+
+		u64 storageCapacity = u64(storageCapacityGb)*1000*1000*1000;
+
+		StringBuffer<512> currentDir;
+		GetCurrentDirectoryW(currentDir);
+
+		if (workDir.IsEmpty())
+			workDir.Append(currentDir);
+
+		// TODO: Change workdir to make it full
+
+
+		StringBuffer<> logFile;
+		#if UBA_DEBUG
+		logFile.count = GetFullPathNameW(g_rootDir.data, logFile.capacity, logFile.data, nullptr);
+		logFile.EnsureEndsWithSlash().Append(TC("DebugLog.log"));
+		#endif
+
+		#if PLATFORM_WINDOWS
+		SetConsoleCtrlHandler(ConsoleHandler, TRUE);
+		#else
+		signal(SIGINT, ConsoleHandler);
+		#endif
+
+		NetworkBackendTcp networkBackend(logWriter);
+		NetworkServerCreateInfo nsci(logWriter);
+		//nsci.workerCount = 4;
+		bool ctorSuccess = true;
+		NetworkServer networkServer(ctorSuccess, nsci);
+		if (!ctorSuccess)
+			return -1;
+
+		StorageServerCreateInfo storageInfo(networkServer, g_rootDir.data, logWriter);
+		storageInfo.casCapacityBytes = storageCapacity;
+		storageInfo.storeCompressed = storeCompressed;
+		storageInfo.allowFallback = false;
+		storageInfo.writeRecievedCasFilesToDisk = true;
+		StorageServer storageServer(storageInfo);
+
+		if (!storageServer.LoadCasTable(true))
+			return -1;
+
+		CacheServer cacheServer(logWriter, g_rootDir.data, networkServer, storageServer);
+
+		if (!cacheServer.Load())
+			return -1;
+
+		if (!cacheServer.RunMaintenance(true))
+			return -1;
+
+		{
+			g_storageServer = &storageServer;
+			g_cacheServer = &cacheServer;
+
+			auto stopListen = MakeGuard([&]() { networkBackend.StopListen(); });
+			auto stopServer = MakeGuard([&]() { networkServer.DisconnectClients(); });
+
+			if (!networkServer.StartListen(networkBackend, port, listenIp.data))
+				return -1;
+
+			while (!IsEscapePressed())
+			{
+				Sleep(1000);
+				if (!cacheServer.RunMaintenance(false))
+					break;
+			}
+
+			g_cacheServer = nullptr;
+			g_storageServer = nullptr;
+		}
+
+		cacheServer.Save();
+		return 0;
+	}
+}
+
+#if PLATFORM_WINDOWS
+int wmain(int argc, wchar_t* argv[])
+{
+	int res = uba::WrappedMain(argc, argv);
+	Sleep(2000);
+	return res;
+}
+#else
+int main(int argc, char* argv[])
+{
+	return uba::WrappedMain(argc, argv);
+}
+#endif
diff --git a/Engine/Source/Programs/UnrealBuildAccelerator/CacheService/UbaCacheService.Build.cs b/Engine/Source/Programs/UnrealBuildAccelerator/CacheService/UbaCacheService.Build.cs
new file mode 100644
--- /dev/null
+++ b/Engine/Source/Programs/UnrealBuildAccelerator/CacheService/UbaCacheService.Build.cs
@@ -0,0 +1,23 @@
+// Copyright Epic Games, Inc. All Rights Reserved.
+
+using UnrealBuildTool;
+
+[SupportedPlatforms(UnrealPlatformClass.Desktop)]
+public class UbaCacheService : ModuleRules
+{
+	public UbaCacheService(ReadOnlyTargetRules Target) : base(Target)
+	{
+		PrivatePCHHeaderFile = "../Core/Public/UbaCorePch.h";
+
+		UnsafeTypeCastWarningLevel = WarningLevel.Error;
+		StaticAnalyzerDisabledCheckers.Clear();
+
+		PrivateDependencyModuleNames.AddRange(new string[] {
+			"UbaCommon",
+		});
+
+		PrivateDefinitions.AddRange(new string[] {
+			"_CONSOLE",
+		});
+	}
+}
diff --git a/Engine/Source/Programs/UnrealBuildAccelerator/Cli/Private/UbaCli.cpp b/Engine/Source/Programs/UnrealBuildAccelerator/Cli/Private/UbaCli.cpp
--- a/Engine/Source/Programs/UnrealBuildAccelerator/Cli/Private/UbaCli.cpp
+++ b/Engine/Source/Programs/UnrealBuildAccelerator/Cli/Private/UbaCli.cpp
@@ -1,5 +1,7 @@
 // Copyright Epic Games, Inc. All Rights Reserved.
 
+#include "UbaCacheClient.h"
+#include "UbaCacheServer.h"
 #include "UbaClient.h"
 #include "UbaCoordinatorWrapper.h"
 #include "UbaFileAccessor.h"
@@ -12,6 +14,7 @@
 #include "UbaStorageClient.h"
 #include "UbaStorageProxy.h"
 #include "UbaStorageServer.h"
+#include "UbaStorageUtils.h"
 #include "UbaVersion.h"
 
 #include "UbaAWS.h"
@@ -75,6 +78,9 @@ namespace uba
 		logger.Info(TC("   -maxcpu=<number>        Max number of processes that can be started. Defaults to \"%u\" on this machine"), DefaultProcessorCount);
 		logger.Info(TC("   -visualizer             Spawn a visualizer that visualizes progress"));
 		logger.Info(TC("   -coordinator=<name>     Load a UbaCoordinator<name>.dll to instantiate a coordinator to get helpers"));
+		logger.Info(TC("   -cache=<host>[:<port>]  Connect to cache server. Will fetch from cache unless -populatecache is set"));
+		logger.Info(TC("   -populatecache          Populate cache server if connected to one"));
+		logger.Info(TC("   -writecachesummary      Write cache summary file about connected cache server"));
 		logger.Info(TC(""));
 		logger.Info(TC("  CoordinatorOptions (if coordinator set):"));
 		logger.Info(TC("   -uri=<address>          Uri to coordinator"));
@@ -142,6 +148,7 @@ namespace uba
 		u32 storageCapacityGb = DefaultCapacityGb;
 		StringBuffer<256> workDir;
 		StringBuffer<128> listenIp;
+		StringBuffer<128> cacheHost;
 		TString coordinatorName;
 		TString coordinatorUri;
 		TString coordinatorPool;
@@ -160,6 +167,10 @@ namespace uba
 		bool deleteCas = false;
 		bool enableStdOut = true;
 		bool printSummary = false;
+		bool populateCache = false;
+		bool writeCacheSummary = false;
+		TString cacheFilterString;
+
 		u32 loopCount = 1;
 
 		enum CommandType
@@ -331,6 +342,21 @@ namespace uba
 			{
 				printSummary = true;
 			}
+			else if (name.Equals(TC("-cache")))
+			{
+				if (value.IsEmpty())
+					return PrintHelp(TC("-cache needs a value"));
+				cacheHost.Append(value);
+			}
+			else if (name.Equals(TC("-populatecache")))
+			{
+				populateCache = true;
+			}
+			else if (name.Equals(TC("-writecachesummary")))
+			{
+				writeCacheSummary = true;
+				cacheFilterString = value.data;
+			}
 			else if (name.Equals(TC("-storeraw")))
 			{
 				storeCompressed = false;
@@ -402,7 +428,7 @@ namespace uba
 				return false;
 			bool success = true;
 			WorkManagerImpl workManager(DefaultProcessorCount);
-			storageServer.TraverseAllCasFiles([&](const CasKey& casKey)
+			storageServer.TraverseAllCasFiles([&](const CasKey& casKey, u64 size)
 				{
 					workManager.AddWork([&, casKey]()
 						{
@@ -461,8 +487,6 @@ namespace uba
 
 		if (getCas)
 		{
-			CasKey key;
-
 			FileAccessor fa(logger, application.c_str());
 			if (!fa.OpenMemoryRead())
 				return logger.Error(TC("Failed to open file %s"), application.c_str());
@@ -470,9 +494,7 @@ namespace uba
 			u8* data = fa.GetData();
 			bool is64Bit = true;
 
-			CasKeyHasher hasher;
-			hasher.Update(data, fileSize);
-			key = ToCasKey(hasher, false);
+			CasKey key = CalculateCasKey(data, fileSize, false, nullptr);
 
 			if (data[0] != 'M' || data[1] != 'Z')
 				is64Bit = false;
@@ -518,21 +540,21 @@ namespace uba
 		NetworkServerCreateInfo nsci(logWriter);
 		//nsci.workerCount = 4;
 		bool ctorSuccess = true;
-		NetworkServer* server = new NetworkServer(ctorSuccess, nsci);
-		auto destroyServer = MakeGuard([&]() { delete server; });
+		NetworkServer& networkServer = *new NetworkServer(ctorSuccess, nsci);
+		auto destroyServer = MakeGuard([&]() { delete &networkServer; });
 		if (!ctorSuccess)
 			return -1;
 
 		bool isRemote = commandType == CommandType_Remote || commandType == CommandType_Agent;
 		bool useScheduler = EndsWith(application.c_str(), application.size(), TC(".yaml"));
 
-		StorageServerCreateInfo storageInfo(*server, g_rootDir.data, logWriter);
+		StorageServerCreateInfo storageInfo(networkServer, g_rootDir.data, logWriter);
 		storageInfo.casCapacityBytes = storageCapacity;
 		storageInfo.storeCompressed = storeCompressed;
-		StorageServer* storage = new StorageServer(storageInfo);
-		auto destroyStorage = MakeGuard([&]() { delete storage; });
+		StorageServer& storageServer = *new StorageServer(storageInfo);
+		auto destroyStorage = MakeGuard([&]() { delete &storageServer; });
 
-		SessionServerCreateInfo info(*storage, *server);
+		SessionServerCreateInfo info(storageServer, networkServer);
 		info.useUniqueId = useScheduler;
 		info.traceEnabled = true;
 		//info.detailedTrace = true;
@@ -541,23 +563,61 @@ namespace uba
 		//info.shouldWriteToDisk = shouldWriteToDisk;
 		info.rootDir = g_rootDir.data;
 		//info.traceName.Append(TC("TESTTRACE"));
+		//info.storeObjFilesCompressed = true;
 		#if UBA_DEBUG_LOG_ENABLED
 		info.remoteLogEnabled = true;
 		#endif
 		//info.remoteTraceEnabled = true;
 
 		info.deleteSessionsOlderThanSeconds = 1;
-		auto session = new SessionServer(info);
-		auto destroySession = MakeGuard([&]() { delete session; });
+		SessionServer& sessionServer = *new SessionServer(info);
+		auto destroySession = MakeGuard([&]() { delete &sessionServer; });
+
+		CacheClient* cacheClient = nullptr;
+		auto ccg = MakeGuard([&]() { if (!cacheClient) return; auto& nc = cacheClient->GetClient(); nc.Disconnect(); delete cacheClient; delete &nc; });
+
+		auto CreateCacheClient = [&]()
+			{
+				auto nc = new NetworkClient(ctorSuccess);
+				cacheClient = new CacheClient(logWriter, storageServer, *nc, sessionServer);
+			};
+
+		if (cacheHost.count)
+		{
+			CreateCacheClient();
+			if (!cacheClient->GetClient().Connect(networkBackend, cacheHost.data, DefaultCachePort))
+			{
+				logger.Error(TC("Failed to connect to cache server"));
+				return -1;
+			}
+
+			if (writeCacheSummary)
+			{
+				StringBuffer<> tempFile(sessionServer.GetTempPath());
+				Guid guid;
+				CreateGuid(guid);
+				tempFile.Append(GuidToString(guid).str).Append(TC(".txt"));
+				if (!cacheClient->WriteCacheSummary(tempFile.data, cacheFilterString.data()))
+					return -1;
+				logger.Info(TC("Cache status summary written to %s"), tempFile.data);
+
+				#if PLATFORM_WINDOWS
+				ShellExecuteW(NULL, L"open", tempFile.data, NULL, NULL, SW_SHOW);
+				#endif
+				return 0;
+			}
+		}
+
+
 
 		if (isRemote)
 		{
-			if (!storage->LoadCasTable(true))
+			if (!storageServer.LoadCasTable(true))
 				return -1;
-			if (!server->StartListen(networkBackend, port, listenIp.data))
+			if (!networkServer.StartListen(networkBackend, port, listenIp.data))
 				return -1;
 		}
-		auto stopServer = MakeGuard([&]() { server->DisconnectClients(); });
+		auto stopServer = MakeGuard([&]() { networkServer.DisconnectClients(); });
 
 		auto stopListen = MakeGuard([&]() { networkBackend.StopListen(); });
 
@@ -575,7 +635,7 @@ namespace uba
 				pinfo.logLineFunc = [](void* userData, const tchar* line, u32 length, LogEntryType type) { ((Logger*)userData)->Log(type, line, length); };
 			pinfo.trackInputs = trackInputs;
 			logger.Info(TC("Running %s %s"), app.c_str(), arg.c_str());
-			ProcessHandle process = session->RunProcess(pinfo, false, enableDetour);
+			ProcessHandle process = sessionServer.RunProcess(pinfo, false, enableDetour);
 			if (process.GetExitCode() != 0)
 				return logger.Error(TC("Error exit code: %u"), process.GetExitCode());
 			u64 time = GetTime() - start;
@@ -596,7 +656,7 @@ namespace uba
 			if (enableStdOut)
 				pinfo.logLineFunc = [](void* userData, const tchar* line, u32 length, LogEntryType type) { ((Logger*)userData)->Log(type, line, length); };
 			logger.Info(TC("Running %s %s"), app.c_str(), arg.c_str());
-			ProcessHandle process = session->RunProcessRemote(pinfo);
+			ProcessHandle process = sessionServer.RunProcessRemote(pinfo);
 			process.WaitForExit(~0u);
 			if (process.GetExitCode() != 0)
 				return logger.Error(TC("Error exit code: %u"), process.GetExitCode());
@@ -609,11 +669,11 @@ namespace uba
 			{
 				Vector<Client> clients;
 				auto slg = MakeGuard([&]() { networkBackend.StopListen(); });
-				clients.resize(4);
+				clients.resize(maxProcessCount == 1 ? 1 : 4);
 				u32 clientIndex = 0;
 				for (auto& c : clients)
 				{
-					ClientInitInfo cii { logWriter, networkBackend, g_rootDir.data, TC("127.0.0.1"), port, TC("DummyZone"), maxProcessCount/4, clientIndex++};
+					ClientInitInfo cii { logWriter, networkBackend, g_rootDir.data, TC("127.0.0.1"), port, TC("DummyZone"), maxProcessCount/u32(clients.size()), clientIndex++};
 					if (!c.Init(cii))
 						return false;
 				}
@@ -629,10 +689,14 @@ namespace uba
 
 		auto RunScheduler = [&](const tchar* yamlFile)
 		{
-			SchedulerCreateInfo info(*session);
+			auto g = MakeGuard([&]() { if (cacheClient) cacheClient->GetClient().Disconnect(); });
+
+			SchedulerCreateInfo info(sessionServer);
 			info.forceRemote = isRemote;
 			info.forceNative = commandType == CommandType_Native;
 			info.maxLocalProcessors = maxProcessCount;
+			info.cacheClient = cacheClient;
+			info.writeToCache = populateCache;
 			Scheduler scheduler(info);
 
 			if (!scheduler.EnqueueFromFile(yamlFile))
@@ -647,10 +711,11 @@ namespace uba
 
 			scheduler.SetProcessFinishedCallback([&](const ProcessHandle& ph)
 				{
-					const tchar* desc = ph.GetStartInfo().description;
+					auto& si = ph.GetStartInfo();
+					const tchar* desc = si.description;
 					if (ph.GetExitCode() != 0 && ph.GetExitCode() != ProcessCancelExitCode)
 					{
-						logger.Error(TC("%s - Error exit code: %u"), desc, ph.GetExitCode());
+						logger.Error(TC("%s - Error exit code: %u (%s %s)"), desc, ph.GetExitCode(), si.application, si.arguments);
 						success = false;
 					}
 					u32 c = ++counter;
@@ -660,6 +725,8 @@ namespace uba
 						extra.Append(TC(" [RemoteExecutor: ")).Append(ph.GetExecutingHost()).Append(']');
 					else if (ph.GetExecutionType() == ProcessExecutionType_Native)
 						extra.Append(TC(" (Not detoured)"));
+					else if (ph.GetExecutionType() == ProcessExecutionType_FromCache)
+						extra.Append(TC(" (From cache)"));
 					logger.Info(TC("[%u/%u] %s%s"), c, queued, desc, extra.data);
 					for (auto& line : ph.GetLogLines())
 						if (line.text != desc && !StartsWith(line.text.c_str(), TC("   Creating library")))
@@ -709,7 +776,7 @@ namespace uba
 			cinfo.oidc = coordinatorOidc.c_str();
 			cinfo.maxCoreCount = 400;
 			cinfo.logging = true;
-			if (!coordinator.Create(logger, coordinatorName.c_str(), cinfo, networkBackend, *server))
+			if (!coordinator.Create(logger, coordinatorName.c_str(), cinfo, networkBackend, networkServer))
 				return false;
 		}
 		auto cg = MakeGuard([&]() { coordinator.Destroy(); });
@@ -747,9 +814,9 @@ namespace uba
 		logger.BeginScope();
 		if (printSummary)
 		{
-			session->PrintSummary(logger);
-			storage->PrintSummary(logger);
-			server->PrintSummary(logger);
+			sessionServer.PrintSummary(logger);
+			storageServer.PrintSummary(logger);
+			networkServer.PrintSummary(logger);
 			SystemStats::GetGlobal().Print(logger, true);
 		}
 		logger.EndScope();
@@ -761,7 +828,9 @@ namespace uba
 #if PLATFORM_WINDOWS
 int wmain(int argc, wchar_t* argv[])
 {
-	return uba::WrappedMain(argc, argv);
+	int res = uba::WrappedMain(argc, argv);
+	Sleep(1); // Here to be able to put a breakpoint just before exit :-)
+	return res;
 }
 #else
 int main(int argc, char* argv[])
diff --git a/Engine/Source/Programs/UnrealBuildAccelerator/Common/Private/UbaCacheClient.cpp b/Engine/Source/Programs/UnrealBuildAccelerator/Common/Private/UbaCacheClient.cpp
new file mode 100644
--- /dev/null
+++ b/Engine/Source/Programs/UnrealBuildAccelerator/Common/Private/UbaCacheClient.cpp
@@ -0,0 +1,796 @@
+// Copyright Epic Games, Inc. All Rights Reserved.
+
+#include "UbaCacheClient.h"
+#include "UbaFileAccessor.h"
+#include "UbaNetworkMessage.h"
+#include "UbaProcess.h"
+#include "UbaStorage.h"
+#include "UbaStorageUtils.h"
+
+#if PLATFORM_WINDOWS
+#include <shlobj_core.h>
+#endif
+
+namespace uba
+{
+	static constexpr u8 RootStartByte = ' ';
+
+	struct CacheClient::Root
+	{
+		TString path;
+		StringKey shortestPathKey;
+		bool includeInKey;
+	};
+
+	CacheClient::CacheClient(LogWriter& writer, StorageImpl& storage, NetworkClient& client, Session& session)
+	:	m_logger(writer, TC("UbaCacheClient"))
+	,	m_storage(storage)
+	,	m_client(client)
+	,	m_session(session)
+	,	m_serverPathTable(CachePathTableMaxSize)
+	,	m_serverCasKeyTable(CacheCasKeyTableMaxSize)
+	,	m_sendPathTable(8*1024*1024)
+	,	m_sendCasKeyTable(8*1024*1024)
+	{
+		m_client.RegisterOnConnected([this]()
+			{
+				StackBinaryWriter<1024> writer;
+				NetworkMessage msg(m_client, CacheServiceId, CacheMessageType_Connect, writer);
+				writer.WriteU32(CacheNetworkVersion);
+				StackBinaryReader<1024> reader;
+				if (!msg.Send(reader))
+				{
+					m_logger.Error(TC("Failed to connect to cache server. Version mismatch?"));
+					return;
+				}
+				m_connected = true;
+			});
+
+		m_client.RegisterOnDisconnected([this]()
+			{
+				m_connected = false;
+			});
+	}
+
+	CacheClient::~CacheClient() = default;
+
+	bool CacheClient::RegisterRoot(const tchar* rootPath, bool includeInKey)
+	{
+		// Register rootPath both with single path separators and double path separators because text files store them with double path separators
+
+		StringBuffer<> doubleSlash;
+		for (const tchar* it=rootPath; *it; ++it)
+		{
+			doubleSlash.Append(*it);
+			if (*it == PathSeparator)
+				doubleSlash.Append(PathSeparator);
+		}
+
+		const tchar* rootPaths[] = { rootPath, doubleSlash.data };
+		for (const tchar* rp : rootPaths)
+		{
+			if (m_roots.size() == '~' - ' ') // This is not really true.. as long as value is under 256 we're good
+				return m_logger.Error(TC("Too many roots added (%llu)"), m_roots.size());
+
+			auto& root = m_roots.emplace_back(rp);
+
+			ToLower(root.path.data());
+			if (root.path[root.path.size()-1] != PathSeparator)
+				return m_logger.Error(TC("Root path must end with separator"));
+
+			root.includeInKey = includeInKey;
+
+			m_longestRoot = Max(u32(root.path.size()), m_longestRoot);
+
+			if (!m_shortestRoot || root.path.size() < m_shortestRoot)
+			{
+				m_shortestRoot = u32(root.path.size());
+				for (auto& r : m_roots)
+					r.shortestPathKey = ToStringKeyNoCheck(r.path.data(), m_shortestRoot);
+			}
+			else
+				root.shortestPathKey = ToStringKeyNoCheck(root.path.data(), m_shortestRoot);
+		}
+		return true;
+	}
+
+	bool CacheClient::RegisterSystemRoots()
+	{
+		#if PLATFORM_WINDOWS
+		StringBuffer<MaxPath> dir;
+		dir.count = GetSystemDirectory(dir.data, dir.capacity);
+		RegisterRoot(dir.EnsureEndsWithSlash().data, false); // Ignore files from here.. we do expect them not to affect the output of a process
+		
+		dir.count = GetEnvironmentVariable(TC("ProgramW6432"), dir.Clear().data, dir.capacity);
+		RegisterRoot(dir.EnsureEndsWithSlash().data, true);
+
+		dir.count = GetEnvironmentVariable(TC("ProgramFiles(x86)"), dir.Clear().data, dir.capacity);
+		RegisterRoot(dir.EnsureEndsWithSlash().data, true);
+
+		dir.count = GetEnvironmentVariable(TC("ProgramFiles(x86)"), dir.Clear().data, dir.capacity);
+		RegisterRoot(dir.EnsureEndsWithSlash().data, true);
+
+		PWSTR path;
+		if (!SUCCEEDED(SHGetKnownFolderPath(FOLDERID_ProgramData, 0, NULL, &path)))
+			return false;
+		RegisterRoot(dir.Clear().Append(path).EnsureEndsWithSlash().data, true);
+		CoTaskMemFree(path);
+
+		#else
+		UBA_ASSERT(false);
+		#endif
+		return true;
+	}
+
+	bool CacheClient::WriteToCache(const ProcessHandle& process)
+	{
+		if (!m_connected)
+			return false;
+
+		auto& si = process.GetStartInfo();
+		if (!si.trackInputs)
+			return false;
+
+		CasKey cmdKey = GetCmdKey(si);
+		if (cmdKey == CasKeyZero)
+			return false;
+
+		const Vector<u8>& inputs = process.GetTrackedInputs();
+		BinaryReader inputsReader(inputs.data(), 0, inputs.size());
+
+		const Vector<u8>& outputs = process.GetTrackedOutputs();
+		BinaryReader outputsReader(outputs.data(), 0, outputs.size());
+
+		Map<u32, u32> inputsStringToCasKey;
+		Map<u32, u32> outputsStringToCasKey;
+		u32 requiredPathTableSize = 0;
+		u32 requiredCasTableSize = 0;
+		bool success = true;
+
+		// Traverse all inputs and outputs. to create cache entry that we can send to server
+		while (true)
+		{
+			CasKey casKey;
+
+			StringBuffer<512> path;
+			bool isOutput = outputsReader.GetLeft();
+			if (isOutput)
+				outputsReader.ReadString(path);
+			else if (inputsReader.GetLeft())
+				inputsReader.ReadString(path);
+			else
+				break;
+			
+			TString temp2 = path.data;
+
+			// For .exe and .dll we sometimes get relative paths so we need to expand them to full
+			if (path[1] != ':' && (path.EndsWith(TC(".dll")) || path.EndsWith(TC(".exe"))))
+			{
+				tchar temp[512];
+				bool res = SearchPathW(NULL, path.data, NULL, 512, temp, NULL);
+				path.Clear().Append(temp);
+				if (!res)
+				{
+					m_logger.Info(TC("Can't find file: %s"), path.data);
+					return false;
+				}
+			}
+			else if (path.EndsWith(TC(".rsp"))) // Paths can be absolute in rsp files so we need to normalize those paths
+			{
+				casKey = AsCompressed(NormalizeAndHashFile(path.data), true);
+			}
+			else if (path[path.count-1] == ':')
+			{
+				m_logger.Info(TC("GOT UNKNOWN RELATIVE PATH: %s"), path.data);
+				success = false;
+				continue;
+			}
+
+			// Find root for path in order to be able to normalize it.
+			u32 rootIndex = FindRootIndex(path);
+			if (rootIndex == ~0u)
+			{
+				m_logger.Info(TC("FILE WITHOUT ROOT: %s"), path.data);
+				success = false;
+				continue;
+			}
+
+			auto& root = m_roots[rootIndex];
+			if (!root.includeInKey)
+				continue;
+
+			u32 rootLen = u32(root.path.size());
+			TString qualifiedPath = path.data + rootLen - 2;
+			qualifiedPath[0] = tchar(RootStartByte + rootIndex);
+
+			u32 pathOffset = m_sendPathTable.Add(qualifiedPath.c_str(), u32(qualifiedPath.size()), &requiredPathTableSize);
+
+			if (!isOutput) // Output files should be removed from input files.. For example when cl.exe compiles pch it reads previous pch file and we don't want it to be input
+				if (outputsStringToCasKey.find(pathOffset) != outputsStringToCasKey.end())
+					continue;
+
+			auto insres = (isOutput ? outputsStringToCasKey : inputsStringToCasKey).try_emplace(pathOffset);
+			
+			if (!insres.second)
+				continue;
+
+			// .dep.json contains absolute paths, need to normalize file
+			if (isOutput && path.EndsWith(TC(".dep.json"))) // TODO: More data driven approach. Also, hash does not match content atm.
+				casKey = AsCompressed(NormalizeAndHashFile(path.data), true);
+
+			// Get file caskey using storage
+			if (casKey == CasKeyZero)
+			{
+				bool deferCreation = true;
+				if (!m_storage.StoreCasFile(casKey, path.data, CasKeyZero, deferCreation))
+					return false;
+				if (casKey == CasKeyZero) // If file is not found it was a temporary file that was deleted and is not really an output
+					continue; // m_logger.Info(TC("This should never happen! (%s)"), path.data);
+			}
+
+			UBA_ASSERT(IsCompressed(casKey));
+			insres.first->second = m_sendCasKeyTable.Add(casKey, pathOffset, &requiredCasTableSize);
+		}
+
+		if (!success)
+			return false;
+
+		if (outputsStringToCasKey.empty())
+			m_logger.Warning(TC("NO OUTPUTS FROM process %s"), process.GetStartInfo().description); 
+
+		// Make sure server has enough of the path table to be able to resolve offsets from cache entry
+		if (!SendPathTable(requiredPathTableSize))
+			return false;
+
+		// Make sure server has enough of the cas table to be able to resolve offsets from cache entry
+		if (!SendCasTable(requiredCasTableSize))
+			return false;
+
+		// actual cache entry now when we know server has the needed tables
+		if (!SendCacheEntry(cmdKey, inputsStringToCasKey, outputsStringToCasKey))
+			return false;
+
+		return true;
+	}
+
+	bool CacheClient::FetchFromCache(const ProcessStartInfo& info)
+	{
+		if (!m_connected)
+			return false;
+
+		CasKey cmdKey = GetCmdKey(info);
+		if (cmdKey == CasKeyZero)
+			return false;
+
+		StackBinaryReader<SendMaxSize> reader;
+
+		{
+			// Fetch entries.. server will provide as many as fits. TODO: Should it be possible to ask for more entries?
+			StackBinaryWriter<32> writer;
+			NetworkMessage msg(m_client, CacheServiceId, CacheMessageType_FetchEntries, writer);
+			writer.WriteCasKey(cmdKey);
+			if (!msg.Send(reader))
+				return false;
+		}
+
+		// Traverse entries and test inputs against local machine
+		u32 entryCount = reader.ReadU16();
+		for (u32 i=0; i!=entryCount; ++i)
+		{
+			bool isMatch = true;
+			u64 inputSize = reader.Read7BitEncoded();
+			const u8* inputEnd = reader.GetPositionData() + inputSize;
+			while (reader.GetPositionData() != inputEnd)
+			{
+				u32 casKeyOffset = u32(reader.Read7BitEncoded());
+
+				if (casKeyOffset >= m_serverCasKeyTable.GetSize())
+					if (!FetchCasTable())
+						return false;
+
+				StringBuffer<MaxPath> path;
+				CasKey cacheCasKey;
+				if (!GetLocalPathAndCasKey(path, cacheCasKey, m_serverCasKeyTable, m_serverPathTable, casKeyOffset))
+					return false;
+				UBA_ASSERT(IsCompressed(cacheCasKey));
+
+				CasKey localCasKey;
+				if (path.EndsWith(TC(".rsp")) || path.EndsWith(TC(".dep.json"))) // Need to normalize caskey for these files since they contain absolute paths
+				{
+					localCasKey = AsCompressed(NormalizeAndHashFile(path.data), true);
+				}
+				else
+				{
+					bool deferCreation = true;
+					m_storage.StoreCasFile(localCasKey, path.data, CasKeyZero, deferCreation);
+					UBA_ASSERT(localCasKey == CasKeyZero || IsCompressed(localCasKey));
+				}
+
+				if (localCasKey != cacheCasKey)
+				{
+					reader.Skip(inputEnd -  reader.GetPositionData());
+					isMatch = false;
+					break;
+				}
+			}
+
+			u64 outputSize = reader.Read7BitEncoded();
+
+			// No match, test next entry
+			if (!isMatch)
+			{
+				reader.Skip(outputSize);
+				continue;
+			}
+
+			// Fetch output files from cache (and some files need to be "denormalized" before written to disk
+
+			const u8* outputEnd = reader.GetPositionData() + outputSize;
+			while (reader.GetPositionData() != outputEnd)
+			{
+				u32 casKeyOffset = u32(reader.Read7BitEncoded());
+				if (casKeyOffset >= m_serverCasKeyTable.GetSize())
+					if (!FetchCasTable())
+						return false;
+
+				StringBuffer<MaxPath> path;
+				CasKey casKey;
+				if (!GetLocalPathAndCasKey(path, casKey, m_serverCasKeyTable, m_serverPathTable, casKeyOffset))
+					return false;
+				UBA_ASSERT(IsCompressed(casKey));
+
+				FileFetcher fetcher { m_storage.m_bufferSlots };
+
+				if (path.EndsWith(TC(".dep.json")))
+				{
+					// Fetch into memory, file is in special format without absolute paths
+					MemoryBlock normalizedBlock(1*1024*1024);
+					if (!fetcher.RetrieveFile(m_logger, m_client, casKey, path.data, &normalizedBlock))
+						return false;
+					u32 rootOffsets = *(u32*)normalizedBlock.memory;
+					char* fileStart = (char*)normalizedBlock.memory + sizeof(u32);
+
+					// "denormalize" fetched file into another memory block that will be written to disk
+					MemoryBlock localBlock(2*1024*1024);
+
+					u64 lastWritten = 0;
+					BinaryReader reader2(normalizedBlock.memory, rootOffsets, normalizedBlock.writtenSize);
+					while (reader2.GetLeft())
+					{
+						u64 rootOffset = reader2.Read7BitEncoded();
+						if (u64 toWrite = rootOffset - lastWritten)
+							memcpy(localBlock.Allocate(toWrite, 1, TC("")), fileStart + lastWritten, toWrite);
+						u8 rootIndex = fileStart[rootOffset] - RootStartByte;
+						auto& root = m_roots[rootIndex];
+
+						#if PLATFORM_WINDOWS
+						StringBuffer<> pathTemp;
+						pathTemp.Append(root.path);
+						char rootPath[512];
+						u32 rootPathLen = pathTemp.Parse(rootPath, sizeof_array(rootPath));
+						#else
+						const char* rootPath = root.path.data();
+						u32 rootPathLen = root.path.size();
+						#endif
+
+						if (u32 toWrite = rootPathLen - 1)
+							memcpy(localBlock.Allocate(toWrite, 1, TC("")), rootPath, toWrite);
+						lastWritten = rootOffset + 2;
+					}
+
+					u64 fileSize = rootOffsets - sizeof(u32);
+					if (u64 toWrite = fileSize - lastWritten)
+						memcpy(localBlock.Allocate(toWrite, 1, TC("")), fileStart + lastWritten, toWrite);
+
+					FileAccessor destFile(m_logger, path.data);
+					if (!destFile.CreateWrite())
+						return false;
+					if (!destFile.Write(localBlock.memory, localBlock.writtenSize))
+						return false;
+					if (!destFile.Close(&fetcher.m_lastWritten))
+						return false;
+
+					fetcher.m_size = fileSize;
+					casKey = CalculateCasKey(localBlock.memory, localBlock.writtenSize, false, nullptr);
+				}
+				else
+				{
+					UBA_ASSERT(!m_session.ShouldStoreObjFilesCompressed()); // TODO: Implement Retrieve .obj that is not decompressing
+					if (!fetcher.RetrieveFile(m_logger, m_client, casKey, path.data))
+						return false;
+				}
+				if (!m_storage.FakeCopy(casKey, path.data, fetcher.m_size, fetcher.m_lastWritten, false))
+					return false;
+				if (!m_session.RegisterNewFile(path.data))
+					return false;
+			}
+
+			return true;
+		}
+
+		return false;
+	}
+
+	bool CacheClient::WriteCacheSummary(const tchar* destinationFile, const tchar* filterString)
+	{
+		StackBinaryWriter<1024> writer;
+		NetworkMessage msg(m_client, CacheServiceId, CacheMessageType_CreateStatusFile, writer);
+		writer.WriteString(filterString ? filterString : TC(""));
+		StackBinaryReader<512> reader;
+		if (!msg.Send(reader))
+			return false;
+		CasKey statusFileCasKey = reader.ReadCasKey();
+		if (statusFileCasKey == CasKeyZero)
+			return false;
+
+		FileFetcher fetcher { m_storage.m_bufferSlots };
+		if (!fetcher.RetrieveFile(m_logger, m_client, statusFileCasKey, destinationFile))
+			return false;
+		return true;
+	}
+
+	bool CacheClient::SendPathTable(u32 requiredPathTableSize)
+	{
+		SCOPED_WRITE_LOCK(m_pathTableNetworkLock, lock);
+		if (requiredPathTableSize <= m_pathTableSizeSent)
+			return true;
+
+		u32 left = requiredPathTableSize - m_pathTableSizeSent;
+		while (left)
+		{
+			StackBinaryWriter<SendMaxSize> writer;
+			NetworkMessage msg(m_client, CacheServiceId, CacheMessageType_StorePathTable, writer);
+			u32 toSend = Min(requiredPathTableSize - m_pathTableSizeSent, u32(m_client.GetMessageMaxSize() - 32));
+			left -= toSend;
+			writer.WriteBytes(m_sendPathTable.GetMemory() + m_pathTableSizeSent, toSend);
+			m_pathTableSizeSent += toSend;
+
+			StackBinaryReader<16> reader;
+			if (!msg.Send(reader))
+				return false;
+		}
+		return true;
+	}
+
+	bool CacheClient::SendCasTable(u32 requiredCasTableSize)
+	{
+		SCOPED_WRITE_LOCK(m_casKeyTableNetworkLock, lock);
+		if (requiredCasTableSize <= m_casKeyTableSizeSent)
+			return true;
+
+		u32 left = requiredCasTableSize - m_casKeyTableSizeSent;
+		while (left)
+		{
+			StackBinaryWriter<SendMaxSize> writer;
+			NetworkMessage msg(m_client, CacheServiceId, CacheMessageType_StoreCasTable, writer);
+			u32 toSend = Min(requiredCasTableSize - m_casKeyTableSizeSent, u32(m_client.GetMessageMaxSize() - 32));
+			left -= toSend;
+			writer.WriteBytes(m_sendCasKeyTable.GetMemory() + m_casKeyTableSizeSent, toSend);
+			m_casKeyTableSizeSent += toSend;
+
+			StackBinaryReader<16> reader;
+			if (!msg.Send(reader))
+				return false;
+		}
+		return true;
+	}
+
+	bool CacheClient::SendCacheEntry(const CasKey& cmdKey, const Map<u32, u32>& inputsStringToCasKey, const Map<u32, u32>& outputsStringToCasKey)
+	{
+		StackBinaryReader<1024> reader;
+		{
+			StackBinaryWriter<SendMaxSize> writer;
+
+			NetworkMessage msg(m_client, CacheServiceId, CacheMessageType_StoreEntry, writer);
+			writer.WriteCasKey(cmdKey);
+
+			writer.Write7BitEncoded(outputsStringToCasKey.size());
+			for (auto& kv : outputsStringToCasKey)
+				writer.Write7BitEncoded(kv.second);
+
+			for (auto& kv : inputsStringToCasKey)
+				writer.Write7BitEncoded(kv.second);
+
+			if (!msg.Send(reader))
+				return false;
+		}
+
+		// Server has all content for caskeys.. upload is done
+		if (!reader.GetLeft())
+			return true;
+
+		// There is content we need to upload to server
+		while (reader.GetLeft())
+		{
+			u32 casKeyOffset = u32(reader.Read7BitEncoded());
+
+			StringBuffer<MaxPath> path;
+			CasKey casKey;
+			if (!GetLocalPathAndCasKey(path, casKey, m_sendCasKeyTable, m_sendPathTable, casKeyOffset))
+				return false;
+
+			casKey = AsCompressed(casKey, true);
+
+			StorageImpl::CasEntry* casEntry;
+			if (m_storage.HasCasFile(casKey, &casEntry))
+			{
+				StringBuffer<> casKeyFileName;
+				if (!m_storage.GetCasFileName(casKeyFileName, casKey))
+					return false;
+
+				const u8* fileData;
+				u64 fileSize;
+
+				MappedView mappedView;
+				auto mapViewGuard = MakeGuard([&](){ m_storage.m_casDataBuffer.UnmapView(mappedView, path.data); });
+
+				FileAccessor file(m_logger, casKeyFileName.data);
+
+				if (casEntry->mappingHandle.IsValid()) // If file was created by helper it will be in the transient mapped memory
+				{
+					mappedView = m_storage.m_casDataBuffer.MapView(casEntry->mappingHandle, casEntry->mappingOffset, casEntry->mappingSize, path.data);
+					fileData = mappedView.memory;
+					fileSize = mappedView.size;
+				}
+				else
+				{
+					if (!file.OpenMemoryRead())
+						return false;
+					fileData = file.GetData();
+					fileSize = file.GetSize();
+				}
+
+				if (!SendFile(m_logger, m_client, casKey, fileData, fileSize, casKeyFileName.data))
+					return false;
+			}
+			else // If we don't have the cas key it should be one of the normalized files.... otherwise there is a bug
+			{
+				if (path.EndsWith(TC(".dep.json")))
+				{
+					FileAccessor file(m_logger, path.data);
+					if (!file.OpenMemoryRead())
+						return false;
+					MemoryBlock block(AlignUp(file.GetSize(), 64*1024));
+					u32& rootOffsetsStart = *(u32*)block.Allocate(sizeof(u32), 1, TC(""));
+					rootOffsetsStart = 0;
+					Vector<u32> rootOffsets;
+					u32 rootOffsetsSize = 0;
+
+					auto handleString = [&](const char* str, u64 strLen, u32 rootPos)
+						{
+							void* mem = block.Allocate(strLen, 1, TC(""));
+							memcpy(mem, str, strLen);
+							if (rootPos != ~0u)
+							{
+								rootOffsets.push_back(rootPos);
+								rootOffsetsSize += Get7BitEncodedCount(rootPos);
+							}
+						};
+
+					if (!NormalizeString<char>((const char*)file.GetData(), file.GetSize(), handleString, path.data))
+						return false;
+
+					if (!rootOffsets.empty())
+					{
+						u8* mem = (u8*)block.Allocate(rootOffsetsSize, 1, TC(""));
+						rootOffsetsStart = u32(mem - block.memory);
+						BinaryWriter writer(mem, 0, rootOffsetsSize);
+						for (u32 rootOffset : rootOffsets)
+							writer.Write7BitEncoded(rootOffset);
+					}
+
+					auto& s = m_storage;
+					FileSender sender { m_logger, m_client, s.m_bufferSlots, s.Stats(), m_sendOneAtTheTimeLock, s.m_casCompressor, s.m_casCompressionLevel };
+					if (!sender.SendFileCompressed(casKey, path.data, block.memory, block.writtenSize, TC("SendCacheEntry")))
+						return m_logger.Error(TC("Failed to send cas content for file %s"), path.data);
+				}
+				else
+				{
+					return m_logger.Error(TC("Can't find output file %s to send to cache server"), path.data);
+				}
+			}
+
+		}
+
+		// Send done.. confirm to server
+		StackBinaryWriter<SendMaxSize> writer;
+		NetworkMessage msg(m_client, CacheServiceId, CacheMessageType_StoreEntryDone, writer);
+		writer.WriteCasKey(cmdKey);
+		if (!msg.Send(reader))
+			return false;
+
+		return true;
+	}
+
+	bool CacheClient::FetchCasTable()
+	{
+		SCOPED_WRITE_LOCK(m_casKeyTableNetworkLock, lock);
+
+		StackBinaryReader<SendMaxSize> reader;
+		{
+			//SCOPED_WRITE_LOCK(m_pathTableNetworkLock, lock);
+			u32 targetSize = ~0u;
+			while (m_serverCasKeyTable.GetSize() < targetSize)
+			{
+				StackBinaryWriter<16> writer;
+				NetworkMessage msg(m_client, CacheServiceId, CacheMessageType_FetchCasTable, writer);
+				writer.WriteU32(m_serverCasKeyTable.GetSize());
+
+				reader.Reset();
+				if (!msg.Send(reader))
+					return false;
+				u32 size = reader.ReadU32();
+				if (targetSize == ~0u)
+					targetSize = size;
+
+				m_serverCasKeyTable.ReadMem(reader, false);
+			}
+		}
+		{
+			u32 targetSize = ~0u;
+			while (m_serverPathTable.GetSize() < targetSize)
+			{
+				StackBinaryWriter<16> writer;
+				NetworkMessage msg(m_client, CacheServiceId, CacheMessageType_FetchPathTable, writer);
+				writer.WriteU32(m_serverPathTable.GetSize());
+
+				reader.Reset();
+				if (!msg.Send(reader))
+					return false;
+				u32 size = reader.ReadU32();
+				if (targetSize == ~0u)
+					targetSize = size;
+
+				m_serverPathTable.ReadMem(reader, false);
+			}
+		}
+		return true;
+	}
+
+	CasKey CacheClient::GetCmdKey(const ProcessStartInfo& info)
+	{
+		CasKeyHasher hasher;
+
+		// Add hash of application binary to key
+		CasKey applicationCasKey;
+		bool deferCreation = true;
+		if (!m_storage.StoreCasFile(applicationCasKey, info.application, CasKeyZero, deferCreation))
+			return CasKeyZero;
+		hasher.Update(&applicationCasKey, sizeof(CasKey));
+
+		// Add arguments list to key
+		auto hashString = [&](const tchar* str, u64 strLen, u32 rootPos) { hasher.Update(str, strLen*sizeof(tchar)); };
+		if (!NormalizeString(info.arguments, TStrlen(info.arguments), hashString, TC("")))
+			return CasKeyZero;
+
+		// Add content of rsp file to key (This will cost a bit of perf since we need to normalize.. should this be part of key?)
+		if (auto rspStart = TStrchr(info.arguments, '@'))
+		{
+			if (rspStart[1] == '"')
+			{
+				rspStart += 2;
+				if (auto rspEnd = TStrchr(rspStart, '"'))
+				{
+					StringBuffer<> rsp;
+					if (rspStart[1] != ':')
+						rsp.Append(info.workingDir).EnsureEndsWithSlash();
+					rsp.Append(rspStart, rspEnd - rspStart);
+					CasKey rspCasKey = NormalizeAndHashFile(rsp.data);
+					hasher.Update(&rspCasKey, sizeof(CasKey));
+				}
+			}
+		}
+
+		return ToCasKey(hasher, false);
+	}
+
+	u32 CacheClient::FindRootIndex(const StringBufferBase& path)
+	{
+		if (path.count < m_shortestRoot)
+			return ~0u;
+
+		StringBuffer<MaxPath> shortPath;
+		shortPath.Append(path.data, m_shortestRoot).MakeLower();
+		StringKey key = ToStringKeyNoCheck(shortPath.data, m_shortestRoot);
+		for (u32 i=0, e=u32(m_roots.size()); i!=e; ++i)
+		{
+			auto& root = m_roots[i];
+			if (key != root.shortestPathKey)
+				continue;
+			if (!path.StartsWith(root.path.c_str()))
+				continue;
+			return i;
+		}
+		return ~0u;
+	}
+
+	template<typename CharType, typename Func>
+	bool CacheClient::NormalizeString(const CharType* str, u64 strLen, const Func& func, const tchar* hint)
+	{
+		auto strEnd = str + strLen;
+		auto searchPos = str;
+
+		u32 destPos = 0;
+
+		while (true)
+		{
+			auto absPathChars = searchPos;
+			CharType lastChar = 0;
+			while (absPathChars < strEnd && !(lastChar == ':' && *absPathChars == '\\'))
+			{
+				lastChar = *absPathChars;
+				++absPathChars;
+			}
+		
+			if (absPathChars == strEnd)
+			{
+				func(searchPos, strEnd - searchPos, ~0u);
+				return true;
+			}
+
+			auto pathStart = absPathChars - 2;
+
+			auto pathEndOrMore = pathStart;
+			while (pathEndOrMore < strEnd && *pathEndOrMore != '\n')
+				++pathEndOrMore;
+
+			u32 lenOrMore = u32(pathEndOrMore - pathStart);
+			u32 toCopy = Min(lenOrMore, m_longestRoot);
+			StringBuffer<512> path;
+			path.Append(pathStart, toCopy);
+
+			u32 rootIndex = FindRootIndex(path);
+			if (rootIndex == ~0u)
+			{
+				m_logger.Info(TC("PATH WITHOUT ROOT: %s (inside file %s)"), path.data, hint);
+				return false;
+			}
+
+			if (u32 len = u32(pathStart - searchPos))
+			{
+				destPos += len;
+				func(searchPos, len, ~0u);
+			}
+			CharType temp = RootStartByte + CharType(rootIndex);
+			func(&temp, 1, destPos);
+			destPos += 1;
+
+			searchPos = pathStart + u32(m_roots[rootIndex].path.size()) - 1;
+		}
+	}
+
+	CasKey CacheClient::NormalizeAndHashFile(const tchar* filename)
+	{
+		FileAccessor file(m_logger, filename);
+		if (!file.OpenMemoryRead())
+			return CasKeyZero;
+
+		CasKeyHasher hasher;
+		auto hashString = [&](const char* str, u64 strLen, u32 rootPos) { hasher.Update(str, strLen); };
+		if (!NormalizeString<char>((const char*)file.GetData(), file.GetSize(), hashString, filename))
+			return CasKeyZero;
+
+		return ToCasKey(hasher, false);
+	}
+
+	bool CacheClient::GetLocalPathAndCasKey(StringBufferBase& outPath, CasKey& outKey, CompactCasKeyTable& casKeyTable, CompactPathTable& pathTable, u32 offset)
+	{
+		if (!m_connected)
+			return false;
+
+		SCOPED_READ_LOCK(m_casKeyTableNetworkLock, lock); // TODO: Is this needed?
+
+		StringBuffer<MaxPath> normalizedPath;
+		casKeyTable.GetPathAndKey(normalizedPath, outKey, pathTable, offset);
+		UBA_ASSERT(normalizedPath.count);
+
+		StringBuffer<8> rootIndexStr;
+		rootIndexStr.Append(normalizedPath.data, normalizedPath.First(PathSeparator) - normalizedPath.data);
+		u32 rootIndex = rootIndexStr[0] - RootStartByte;
+		auto& root = m_roots[rootIndex];
+
+		StringBuffer<MaxPath> path;
+		outPath.Append(root.path).Append(normalizedPath.data + rootIndexStr.count + 1);
+		return true;
+	}
+}
diff --git a/Engine/Source/Programs/UnrealBuildAccelerator/Common/Private/UbaCacheServer.cpp b/Engine/Source/Programs/UnrealBuildAccelerator/Common/Private/UbaCacheServer.cpp
new file mode 100644
--- /dev/null
+++ b/Engine/Source/Programs/UnrealBuildAccelerator/Common/Private/UbaCacheServer.cpp
@@ -0,0 +1,939 @@
+// Copyright Epic Games, Inc. All Rights Reserved.
+
+#include "UbaCacheServer.h"
+#include "UbaBinaryReaderWriter.h"
+#include "UbaFileAccessor.h"
+#include "UbaNetworkServer.h"
+#include "UbaStorageServer.h"
+
+namespace uba
+{
+	static constexpr u32 CacheFileVersion = 1;
+
+	struct CacheServer::Connection
+	{
+		Connection() : m_pathTable(8*1024*1024), m_casKeyTable(8*1024*1024) {}
+		CompactPathTable m_pathTable;
+		CompactCasKeyTable m_casKeyTable;
+
+		ReaderWriterLock cacheEntryLookupLock;
+		UnorderedMap<CasKey, CacheEntry> cacheEntryLookup;
+	};
+
+	const tchar* ToString(CacheMessageType type)
+	{
+		switch (type)
+		{
+			#define UBA_CACHE_MESSAGE(x) case CacheMessageType_##x: return TC("")#x;
+			UBA_CACHE_MESSAGES
+			#undef UBA_CACHE_MESSAGE
+		default:
+			return TC("Unknown"); // Should never happen
+		}
+	}
+
+	CacheServer::CacheServer(LogWriter& writer, const tchar* rootDir, NetworkServer& server, StorageServer& storage)
+	:	m_logger(writer, TC("UbaCacheServer"))
+	,	m_server(server)
+	,	m_storage(storage)
+	,	m_pathTable(CachePathTableMaxSize)
+	,	m_casKeyTable(CacheCasKeyTableMaxSize)
+	{
+		m_rootDir.count = GetFullPathNameW(rootDir, m_rootDir.capacity, m_rootDir.data, NULL);
+		m_rootDir.Replace('/', PathSeparator).EnsureEndsWithSlash();
+
+		m_server.RegisterService(CacheServiceId,
+			[this](const ConnectionInfo& connectionInfo, MessageInfo& messageInfo, BinaryReader& reader, BinaryWriter& writer)
+			{
+				return HandleMessage(connectionInfo, messageInfo.type, reader, writer);
+			},
+			[](u8 messageType)
+			{
+				return ToString(CacheMessageType(messageType));
+			}
+		);
+
+		m_server.RegisterOnClientDisconnected(CacheServiceId, [this](const Guid& clientUid, u32 clientId)
+			{
+				OnDisconnected(clientId);
+			});
+	}
+
+	CacheServer::~CacheServer()
+	{
+	}
+
+	bool CacheServer::Load()
+	{
+		u64 startTime = GetTime();
+
+		StringBuffer<> fileName(m_rootDir);
+		fileName.EnsureEndsWithSlash().Append(TC("cachedb"));
+
+		FileAccessor file(m_logger, fileName.data);
+		if (!file.OpenMemoryRead(0, false))
+			return true;
+		BinaryReader reader(file.GetData(), 0, file.GetSize());
+
+		u32 version = reader.ReadU32();
+		if (version != CacheFileVersion)
+			return true;
+
+		u32 pathTableSize = reader.ReadU32();
+		if (pathTableSize)
+		{
+			BinaryReader pathTableReader(reader.GetPositionData(), 0, pathTableSize);
+			m_pathTable.ReadMem(pathTableReader, true);
+			reader.Skip(pathTableSize);
+		}
+
+		u32 casKeyTableSize = reader.ReadU32();
+		if (casKeyTableSize)
+		{
+			BinaryReader casKeyTableReader(reader.GetPositionData(), 0, casKeyTableSize);
+			m_casKeyTable.ReadMem(casKeyTableReader, true);
+			reader.Skip(casKeyTableSize);
+		}
+
+		u32 entryLookupCount = reader.ReadU32();
+		m_cacheEntryLookup.reserve(entryLookupCount);
+
+		u64 totalCacheEntryCount = 0;
+		while (reader.GetLeft())
+		{
+			auto insres = m_cacheEntryLookup.try_emplace(reader.ReadCasKey());
+			UBA_ASSERT(insres.second);
+			auto& cacheEntries = insres.first->second;
+			u32 cacheEntryCount = reader.ReadU32();
+			totalCacheEntryCount += cacheEntryCount;
+
+			while (cacheEntryCount--)
+			{
+				auto& cacheEntry = cacheEntries.entries.emplace_back();
+
+				cacheEntry.creationTime = reader.ReadU64();
+
+				u32 inputSize = reader.ReadU32();
+				cacheEntry.inputCasKeyOffsets.resize(inputSize);
+				reader.ReadBytes(cacheEntry.inputCasKeyOffsets.data(), inputSize);
+
+				u32 outputSize = reader.ReadU32();
+				cacheEntry.outputCasKeyOffsets.resize(outputSize);
+				reader.ReadBytes(cacheEntry.outputCasKeyOffsets.data(), outputSize);
+			}
+		}
+
+		u64 duration = GetTime() - startTime;
+		m_logger.Detail(TC("Database loaded from %s in %s (contained %s paths, %s keys, %llu cache entries)"), fileName.data, TimeToText(duration).str, BytesToText(pathTableSize).str, BytesToText(casKeyTableSize).str, totalCacheEntryCount);
+		return true;
+	}
+
+	bool CacheServer::Save()
+	{
+		if (m_addsSinceMaintenance == 0)
+			return true;
+
+		SCOPED_WRITE_LOCK(m_maintenanceLock, lock);
+		return SaveNoLock();
+	}
+
+	bool CacheServer::SaveNoLock()
+	{
+		StringBuffer<MaxPath> fileName(m_rootDir);
+		fileName.EnsureEndsWithSlash().Append(TC("cachedb"));
+
+		StringBuffer<MaxPath> tempFileName;
+		tempFileName.Append(fileName).Append(TC(".tmp"));
+
+		FileAccessor file(m_logger, tempFileName.data);
+		if (!file.CreateWrite())
+			return false;
+
+
+		constexpr u64 tempBufferSize = 1024*1024;
+		u8* tempBuffer = (u8*)malloc(tempBufferSize);
+		auto g = MakeGuard([tempBuffer](){ free(tempBuffer); });
+		u64 tempBufferPos = 0;
+
+		bool success = true;
+		auto WriteBytes = [&](const void* data, u64 size)
+			{
+				u8* readPos = (u8*)data;
+				u64 left = size;
+				while (left)
+				{
+					if (tempBufferPos != tempBufferSize)
+					{
+						u64 toWrite = Min(tempBufferSize - tempBufferPos, left);
+						memcpy(tempBuffer+tempBufferPos, readPos, toWrite);
+						tempBufferPos += toWrite;
+						left -= toWrite;
+						readPos += toWrite;
+					}
+					else
+					{
+						success &= file.Write(tempBuffer, tempBufferPos);
+						tempBufferPos = 0;
+					}
+				}
+			};
+
+		auto Write = [&](auto v) { WriteBytes(&v, sizeof(v)); };
+
+		Write(CacheFileVersion);
+
+		u32 pathTableSize = m_pathTable.GetSize();
+		Write(pathTableSize);
+		WriteBytes(m_pathTable.GetMemory(), pathTableSize);
+
+		u32 casKeyTableSize = m_casKeyTable.GetSize();
+		Write(casKeyTableSize);
+		WriteBytes(m_casKeyTable.GetMemory(), casKeyTableSize);
+
+		u32 entryLookupCount = u32(m_cacheEntryLookup.size());
+		Write(entryLookupCount);
+
+		for (auto& kv : m_cacheEntryLookup)
+		{
+			Write(kv.first);
+
+			u32 cacheEntryCount = u32(kv.second.entries.size());
+			Write(cacheEntryCount);
+
+			for (CacheEntry& entry : kv.second.entries)
+			{
+				WriteBytes(&entry.creationTime, sizeof(entry.creationTime));
+
+				u32 inputSize = u32(entry.inputCasKeyOffsets.size());
+				Write(inputSize);
+				WriteBytes(entry.inputCasKeyOffsets.data(), inputSize);
+
+				u32 outputSize = u32(entry.outputCasKeyOffsets.size());
+				Write(outputSize);
+				WriteBytes(entry.outputCasKeyOffsets.data(), outputSize);
+			}
+		}
+
+		success &= file.Write(tempBuffer, tempBufferPos);
+
+		if (!success)
+			return false;
+
+		if (!file.Close())
+			return false;
+
+		if (!MoveFileExW(tempFileName.data, fileName.data, MOVEFILE_REPLACE_EXISTING))
+			return m_logger.Error(TC("Can't move file from %s to %s (%s)"), tempFileName.data, fileName.data, LastErrorToText().data);
+
+		return true;
+	}
+
+	bool CacheServer::RunMaintenance(bool force)
+	{
+		SCOPED_WRITE_LOCK(m_maintenanceLock, lock);
+		SCOPED_READ_LOCK(m_connectionsLock, lock2);
+		if (!m_connections.empty())
+			return true;
+		lock2.Leave();
+
+		if (m_addsSinceMaintenance == 0 && !force)
+			return true;
+
+		bool forceAllSteps = true;
+
+		m_logger.Info(TC("Maintenance started after %u added cache entries"), m_addsSinceMaintenance.load());
+		
+		u64 startTime = GetTime();
+
+		bool entriesAdded = m_addsSinceMaintenance != 0;
+		m_addsSinceMaintenance = 0;
+
+		Set<CasKey> deletedCasFiles;
+		m_storage.HandleOverflow(&deletedCasFiles);
+		u64 deletedCasCount = deletedCasFiles.size();
+
+		struct CasFileInfo { u64 size; u64 refCount; };
+		Map<CasKey, CasFileInfo> existingCas;
+		u64 totalCasSize = 0;
+		{
+			u64 traverseStartTime = GetTime();
+			m_storage.TraverseAllCasFiles([&](const CasKey& casKey, u64 size) { totalCasSize += size; existingCas.try_emplace(casKey, CasFileInfo{size, 0ull}); });
+			m_logger.Detail(TC("  Found %llu cas files (%s)"), existingCas.size(), TimeToText(GetTime() - traverseStartTime).str);
+		}
+		u64 totalCasCount = existingCas.size();
+
+		Vector<u64*> touchedCas;
+
+		u64 now = GetSystemTimeAsFileTime();
+		u64 oldest = now;
+
+
+		u64 totalEntryCount = 0;
+		u64 totalEntrySize = 0;
+		u64 deleteEntryCount = 0;
+
+		u64 deleteCacheEntriesStartTime = GetTime();
+		do
+		{
+			oldest = now;
+			totalEntryCount = 0;
+			totalEntrySize = 0;
+
+			for (auto li=m_cacheEntryLookup.begin(), le=m_cacheEntryLookup.end(); li!=le;)
+			{
+				// There is currently no idea saving more than 256kb worth of entries per lookup key (because that is what fetch max returns).. so let's wipe out
+				// all the entries that overflow that number
+				u64 capacityLeft = SendMaxSize - 32;
+
+				CacheEntries& entries = li->second;
+				for (auto i=entries.entries.begin(), e=entries.entries.end(); i!=e;)
+				{
+					auto& entry = *i;
+					auto& inputs = entry.inputCasKeyOffsets;
+					auto& outputs = entry.outputCasKeyOffsets;
+
+					bool deleteEntry = false;
+
+					u64 neededSize = Get7BitEncodedCount(inputs.size()) + inputs.size() + Get7BitEncodedCount(outputs.size()) + outputs.size();
+					if (neededSize > capacityLeft)
+					{
+						deleteEntry = true;
+						capacityLeft = 0;
+					}
+					else
+					{
+						// Traverse outputs and check if cas files exists for each output, if not, delete entry.
+						touchedCas.clear();
+						BinaryReader reader2(outputs.data(), 0, outputs.size());
+						while (reader2.GetLeft())
+						{
+							u64 offset = reader2.Read7BitEncoded();
+							CasKey casKey;
+							m_casKeyTable.GetKey(casKey, offset);
+							UBA_ASSERT(IsCompressed(casKey));
+							auto findIt = existingCas.find(casKey);
+							if (findIt == existingCas.end())
+								deleteEntry = true;
+							else
+								touchedCas.push_back(&findIt->second.refCount);
+						}
+					}
+
+					// Check if there are keys that use removed caskey as input and delete those too
+					if (!deleteEntry && !deletedCasFiles.empty())
+					{
+						BinaryReader reader2(inputs.data(), 0, inputs.size());
+						while (reader2.GetLeft())
+						{
+							u64 offset = reader2.Read7BitEncoded();
+							CasKey casKey;
+							m_casKeyTable.GetKey(casKey, offset);
+							if (deletedCasFiles.find(casKey) != deletedCasFiles.end())
+							{
+								deleteEntry = true;
+								break;
+							}
+						}
+					}
+
+					// Remove entry from entries list and skip increasing ref count of cas files
+					if (deleteEntry)
+					{
+						++deleteEntryCount;
+						i = entries.entries.erase(i);
+						e = entries.entries.end();
+						continue;
+					}
+
+					++totalEntryCount;
+					totalEntrySize += neededSize;
+
+					capacityLeft -= neededSize;
+
+					if (!oldest || entry.creationTime < oldest)
+						oldest = entry.creationTime;
+
+					for (u64* v : touchedCas)
+						++(*v);
+					++i;
+				}
+
+				if (!entries.entries.empty())
+				{
+					++li;
+					continue;
+				}
+
+				li = m_cacheEntryLookup.erase(li);
+				le = m_cacheEntryLookup.end();
+			}
+
+			// Reset deleted cas files and update it again..
+			deletedCasFiles.clear();
+
+			for (auto i=existingCas.begin(), e=existingCas.end(); i!=e;)
+			{
+				if (i->second.refCount != 0)
+				{
+					i->second.refCount = 0;
+					++i;
+					continue;
+				}
+				deletedCasFiles.insert(i->first);
+				++deletedCasCount;
+				totalCasSize -= i->second.size;
+				i = existingCas.erase(i);
+				e = existingCas.end();
+			}
+
+			for (auto& casKey : deletedCasFiles)
+				m_storage.DropCasFile(casKey, true, TC(""));
+		}
+		while (!deletedCasFiles.empty()); // if cas files are deleted we need to do another loop and check cache entry inputs to see if files were inputs
+
+		m_logger.Detail(TC("  Deleted %llu cas files and %llu cache entries (%s)"), deletedCasCount, deleteEntryCount, TimeToText(GetTime() - deleteCacheEntriesStartTime).str);
+
+		if ( deleteEntryCount || forceAllSteps)
+		{
+			UnorderedSet<u32> usedCasKeyOffsets;
+
+			// Collect all caskeys that are used by cache entries.
+			for (auto& kv : m_cacheEntryLookup)
+			{
+				for (auto& entry : kv.second.entries)
+				{
+					auto collectUsedCasKeyOffsets = [&](const Vector<u8>& offsets)
+						{
+							BinaryReader reader2(offsets.data(), 0, offsets.size());
+							while (reader2.GetLeft())
+								usedCasKeyOffsets.insert(u32(reader2.Read7BitEncoded()));
+						};
+					collectUsedCasKeyOffsets(entry.inputCasKeyOffsets);
+					collectUsedCasKeyOffsets(entry.outputCasKeyOffsets);
+				}
+			}
+
+			u64 recreatePathTableStart = GetTime();
+
+			// Traverse all caskeys in caskey table and figure out which ones we can delete
+			UnorderedSet<u32> usedPathOffsets;
+			usedPathOffsets.reserve(usedCasKeyOffsets.size());
+
+			for (u32 casKeyOffset : usedCasKeyOffsets)
+			{
+				BinaryReader reader2(m_casKeyTable.GetMemory(), casKeyOffset, m_casKeyTable.GetSize());
+				u32 pathOffset = u32(reader2.Read7BitEncoded());
+				usedPathOffsets.insert(pathOffset);
+			}
+
+			// Build new path table based on used offsets
+			UnorderedMap<u32, u32> oldToNewPathOffset;
+			u32 oldSize = m_pathTable.GetSize();
+			{
+				u64 reservedPathOffsets = usedPathOffsets.size()*2; // Trying to estimate offset count
+				CompactPathTable newPathTable(CachePathTableMaxSize, reservedPathOffsets);
+				oldToNewPathOffset.reserve(usedPathOffsets.size());
+
+				for (u32 pathOffset : usedPathOffsets)
+				{
+					StringBuffer<> temp;
+					m_pathTable.GetString(temp, pathOffset);
+					auto res = oldToNewPathOffset.try_emplace(pathOffset, newPathTable.Add(temp.data, temp.count));
+					UBA_ASSERT(res.second);(void)res;
+				}
+				m_pathTable.Swap(newPathTable);
+			}
+			m_logger.Detail(TC("  Recreated path table. %s -> %s (%s)"), BytesToText(oldSize).str, BytesToText(m_pathTable.GetSize()).str, TimeToText(GetTime() - recreatePathTableStart).str);
+
+
+			// Build new caskey table based on used offsets
+			u64 recreateCasKeyTableStart = GetTime();
+			UnorderedMap<u32, u32> oldToNewCasKeyOffset;
+			oldSize = m_casKeyTable.GetSize();
+			{
+				oldToNewCasKeyOffset.reserve(usedCasKeyOffsets.size());
+				CompactCasKeyTable newCasKeyTable(CacheCasKeyTableMaxSize, usedCasKeyOffsets.size());
+				for (u32 casKeyOffset : usedCasKeyOffsets)
+				{
+					BinaryReader reader2(m_casKeyTable.GetMemory(), casKeyOffset, m_casKeyTable.GetSize());
+					u32 oldPathOffset = u32(reader2.Read7BitEncoded());
+					CasKey casKey = reader2.ReadCasKey();
+					auto findIt = oldToNewPathOffset.find(oldPathOffset);
+					UBA_ASSERT(findIt != oldToNewPathOffset.end());
+					auto res = oldToNewCasKeyOffset.try_emplace(casKeyOffset, newCasKeyTable.Add(casKey, findIt->second));
+					UBA_ASSERT(res.second);(void)res;
+				}
+				m_casKeyTable.Swap(newCasKeyTable);
+			}
+			m_logger.Detail(TC("  Recreated caskey table. %s -> %s (%s)"), BytesToText(oldSize).str, BytesToText(m_casKeyTable.GetSize()).str, TimeToText(GetTime() - recreateCasKeyTableStart).str);
+
+
+			// Update all casKeyOffsets
+			u64 updateEntriesStart = GetTime();
+			for (auto& kv : m_cacheEntryLookup)
+			{
+				for (auto& entry : kv.second.entries)
+				{
+					auto updateCasKeyOffsets = [&](Vector<u8>& offsets)
+						{
+							Set<u32> newOffsets; // Want these sorted since that is what we get from uploads
+							u32 newOffsetsSize = 0;
+							BinaryReader reader2(offsets.data(), 0, offsets.size());
+							while (reader2.GetLeft())
+							{
+								u32 oldOffset = u32(reader2.Read7BitEncoded());
+								auto findIt = oldToNewCasKeyOffset.find(oldOffset);
+								UBA_ASSERT(findIt != oldToNewCasKeyOffset.end());
+								u32 newOffset = findIt->second;
+								if (newOffsets.insert(newOffset).second)
+									newOffsetsSize += Get7BitEncodedCount(newOffset);
+								else
+								{
+									StringBuffer<> temp;
+									CasKey key;
+									m_casKeyTable.GetPathAndKey(temp, key, m_pathTable, newOffset);
+									m_logger.Info(TC("Found duplicate of %s in cache entry"), temp.data);
+								}
+							}
+
+							offsets.resize(newOffsetsSize);
+							BinaryWriter writer2(offsets.data(), 0, newOffsetsSize);
+							for (u32 offset : newOffsets)
+								writer2.Write7BitEncoded(offset);
+							UBA_ASSERT(writer2.GetPosition() == newOffsetsSize);
+						};
+					updateCasKeyOffsets(entry.inputCasKeyOffsets);
+					updateCasKeyOffsets(entry.outputCasKeyOffsets);
+				}
+			}
+
+			m_logger.Detail(TC("  Updated cache entries with new tables (%s)"), TimeToText(GetTime() - updateEntriesStart).str);
+		}
+
+		if (entriesAdded || deletedCasCount || deleteEntryCount || forceAllSteps)
+		{
+			u64 saveStart = GetTime();
+			m_storage.SaveCasTable(false, false);
+			SaveNoLock();
+			m_logger.Detail(TC("  Saved to disk (%s)"), TimeToText(GetTime() - saveStart).str);
+		}
+
+		u64 oldestTime = MsToTime(GetFileTimeAsSeconds(now - oldest)*1000);
+		m_logger.Info(TC("Maintenance done! (%s). CacheEntries: %llu (%s) CasFiles: %llu (%s) PathTable: %s CasTable: %s OldestEntry: %s"), TimeToText(GetTime() - startTime).str, totalEntryCount, BytesToText(totalEntrySize).str, totalCasCount, BytesToText(totalCasSize).str, BytesToText(m_pathTable.GetSize()).str, BytesToText(m_casKeyTable.GetSize()).str, TimeToText(oldestTime, true).str);
+
+		return true;
+	}
+
+	void CacheServer::OnDisconnected(u32 clientId)
+	{
+		SCOPED_WRITE_LOCK(m_connectionsLock, lock);
+		m_connections.erase(clientId);
+		lock.Leave();
+	}
+
+	bool CacheServer::HandleMessage(const ConnectionInfo& connectionInfo, u8 messageType, BinaryReader& reader, BinaryWriter& writer)
+	{
+		switch (messageType)
+		{
+		case CacheMessageType_Connect:
+		{
+			u32 clientVersion = reader.ReadU32();
+			if (clientVersion != CacheNetworkVersion)
+				return m_logger.Error(TC("Different network versions. Client: %u, Server: %u. Disconnecting"), clientVersion, CacheNetworkVersion);
+			SCOPED_READ_LOCK(m_maintenanceLock, lock);
+			SCOPED_WRITE_LOCK(m_connectionsLock, lock2);
+			m_connections[connectionInfo.GetId()];
+			return true;
+		}
+		case CacheMessageType_StorePathTable:
+		{
+			SCOPED_WRITE_LOCK(m_connectionsLock, lock);
+			auto& connection = m_connections[connectionInfo.GetId()];
+			lock.Leave();
+			connection.m_pathTable.ReadMem(reader, false);
+			return true;
+		}
+		case CacheMessageType_StoreCasTable:
+		{
+			SCOPED_WRITE_LOCK(m_connectionsLock, lock);
+			auto& connection = m_connections[connectionInfo.GetId()];
+			lock.Leave();
+			connection.m_casKeyTable.ReadMem(reader, false);
+			return true;
+		}
+		case CacheMessageType_StoreEntry:
+		{
+			SCOPED_WRITE_LOCK(m_connectionsLock, lock);
+			auto& connection = m_connections[connectionInfo.GetId()];
+			lock.Leave();
+			return HandleStoreEntry(connection, reader, writer);
+		}
+		case CacheMessageType_StoreEntryDone:
+		{
+			SCOPED_WRITE_LOCK(m_connectionsLock, lock);
+			auto& connection = m_connections[connectionInfo.GetId()];
+			lock.Leave();
+			CasKey cmdKey = reader.ReadCasKey();
+
+			SCOPED_WRITE_LOCK(connection.cacheEntryLookupLock, lock2);
+			auto findIt = connection.cacheEntryLookup.find(cmdKey);
+			if (findIt != connection.cacheEntryLookup.end())
+			{
+				SCOPED_WRITE_LOCK(m_cacheEntryLookupLock, lock3);
+				auto insres = m_cacheEntryLookup.try_emplace(cmdKey);
+				auto& cacheEntries = insres.first->second;
+				lock2.Leave();
+
+				SCOPED_WRITE_LOCK(cacheEntries.lock, lock4);
+				auto& cacheEntry = cacheEntries.entries.emplace_front(std::move(findIt->second));
+				(void)cacheEntry;
+				//m_logger.Info(TC("Added new cache entry"));
+			}
+			return true;
+		}
+		case CacheMessageType_FetchPathTable:
+			return HandleFetchPathTable(reader, writer);
+
+		case CacheMessageType_FetchCasTable:
+			return HandleFetchCasTable(reader, writer);
+
+		case CacheMessageType_FetchEntries:
+			return HandleFetchEntries(reader, writer);
+
+		case CacheMessageType_CreateStatusFile:
+			return HandleCreateStatusFile(reader, writer);
+
+		default:
+			return false;
+		}
+	}
+
+	bool CacheServer::HandleStoreEntry(Connection& connection, BinaryReader& reader, BinaryWriter& writer)
+	{
+		CasKey cmdKey = reader.ReadCasKey();
+
+		u64 outputCount = reader.Read7BitEncoded();
+		u64 index = 0;
+
+		Set<u32> inputs;
+		u64 bytesForInput = 0;
+
+		u64 outputStartOffset = reader.GetPosition();
+
+		while (reader.GetLeft())
+		{
+			u32 offset = u32(reader.Read7BitEncoded());
+			bool isInput = index++ >= outputCount;
+			if (!isInput)
+				continue;
+
+			CasKey casKey;
+			StringBuffer<> path;
+			connection.m_casKeyTable.GetPathAndKey(path, casKey, connection.m_pathTable, offset);
+			UBA_ASSERT(path.count);
+
+			u32 pathOffset = m_pathTable.Add(path.data, path.count);
+			u32 casKeyOffset = m_casKeyTable.Add(casKey, pathOffset);
+			inputs.insert(casKeyOffset);
+			bytesForInput += Get7BitEncodedCount(casKeyOffset);
+
+			//m_logger.Info(TC("%s - %s%s"), path.data, CasKeyString(casKey).str, isInput ? TC("") : TC(" (output)"));
+		}
+
+		Vector<u8> inputCasKeyOffsets;
+		{
+			inputCasKeyOffsets.resize(bytesForInput);
+			BinaryWriter w2(inputCasKeyOffsets.data(), 0, inputCasKeyOffsets.size());
+			for (u32 input : inputs)
+				w2.Write7BitEncoded(input);
+		}
+
+		SCOPED_WRITE_LOCK(m_cacheEntryLookupLock, lock);
+		auto insres = m_cacheEntryLookup.try_emplace(cmdKey);
+		auto& cacheEntries = insres.first->second;
+		lock.Leave();
+
+		SCOPED_WRITE_LOCK(cacheEntries.lock, lock2);
+		
+		List<CacheEntry>::iterator matchingEntry = cacheEntries.entries.end();
+		for (auto i=cacheEntries.entries.begin(), e=cacheEntries.entries.end(); i!=e; ++i)
+		{
+			if (i->inputCasKeyOffsets != inputCasKeyOffsets)
+				continue;
+			matchingEntry = i;
+			break;
+		}
+
+		// Already exists
+		if (matchingEntry != cacheEntries.entries.end())
+		{
+			bool shouldOverwrite = false;
+			Map<TString, CasKey> existing;
+
+			BinaryReader r2(matchingEntry->outputCasKeyOffsets.data(), 0, matchingEntry->outputCasKeyOffsets.size());
+			while (r2.GetLeft())
+			{
+				u32 existingOffset = u32(r2.Read7BitEncoded());
+				CasKey casKey;
+				StringBuffer<> path;
+				m_casKeyTable.GetPathAndKey(path, casKey, m_pathTable, existingOffset);
+				existing.try_emplace(path.data, casKey);
+			}
+
+			reader.SetPosition(outputStartOffset);
+			u64 left = outputCount;
+			while (left--)
+			{
+				u32 outputOffset = u32(reader.Read7BitEncoded());
+				CasKey casKey;
+				StringBuffer<> path;
+				connection.m_casKeyTable.GetPathAndKey(path, casKey, connection.m_pathTable, outputOffset);
+
+				auto findIt = existing.find(path.data);
+				if (findIt == existing.end())
+				{
+					m_logger.Warning(TC("Existing cache entry matches input but does not match output (output file %s did not exist in existing cache entry)"), path.data);
+					cacheEntries.entries.erase(matchingEntry);
+					shouldOverwrite = true;
+					break;
+				}
+				if (findIt->second != casKey)
+				{
+					m_logger.Warning(TC("Existing cache entry matches input but does not match output (%s has different caskey)"), path.data);
+					cacheEntries.entries.erase(matchingEntry);
+					shouldOverwrite = true;
+					break;
+				}
+			}
+			if (!shouldOverwrite)
+				return true;
+		}
+
+		// Add new entry
+		CacheEntry newEntry;
+		newEntry.inputCasKeyOffsets.swap(inputCasKeyOffsets);
+
+		Set<u32> outputs;
+		u64 bytesForOutput = 0;
+
+		bool hasAllContent = true;
+		reader.SetPosition(outputStartOffset);
+		u64 left = outputCount;
+		while (left--)
+		{
+			u32 outputOffset = u32(reader.Read7BitEncoded());
+			CasKey casKey;
+			StringBuffer<> path;
+			connection.m_casKeyTable.GetPathAndKey(path, casKey, connection.m_pathTable, outputOffset);
+			u32 pathOffset = m_pathTable.Add(path.data, path.count);
+			u32 casKeyOffset = m_casKeyTable.Add(casKey, pathOffset);
+			outputs.insert(casKeyOffset);
+			bytesForOutput += Get7BitEncodedCount(casKeyOffset);
+
+			if (!m_storage.EnsureCasFile(casKey, nullptr))
+			{
+				writer.Write7BitEncoded(outputOffset);
+				hasAllContent = false;
+			}
+		}
+
+		newEntry.outputCasKeyOffsets.resize(bytesForOutput);
+		BinaryWriter w2(newEntry.outputCasKeyOffsets.data(), 0, newEntry.outputCasKeyOffsets.size());
+		for (u32 output : outputs)
+			w2.Write7BitEncoded(output);
+
+
+		newEntry.creationTime = GetSystemTimeAsFileTime();
+
+		// If cache server has all content we can put the new cache entry directly in the lookup.. otherwise we'll have to wait until client has uploaded content
+		if (hasAllContent)
+		{
+			cacheEntries.entries.emplace_front(std::move(newEntry));
+		}
+		else
+		{
+			SCOPED_WRITE_LOCK(connection.cacheEntryLookupLock, lock3);
+			bool res = connection.cacheEntryLookup.try_emplace(cmdKey, std::move(newEntry)).second;
+			UBA_ASSERT(res);(void)res;
+		}
+
+		//m_logger.Info(TC("Added new cache entry (%u inputs and %u outputs)"), u32(inputs.size()), outputCount);
+
+		++m_addsSinceMaintenance;
+
+		return true;
+	}
+
+	bool CacheServer::HandleFetchPathTable(BinaryReader& reader, BinaryWriter& writer)
+	{
+		u32 haveSize = reader.ReadU32();
+		u32 size = m_pathTable.GetSize();
+		writer.WriteU32(size);
+		u32 toSend = Min(u32(writer.GetCapacityLeft()), size - haveSize);
+		writer.WriteBytes(m_pathTable.GetMemory() + haveSize, toSend);
+		return true;
+	}
+
+	bool CacheServer::HandleFetchCasTable(BinaryReader& reader, BinaryWriter& writer)
+	{
+		u32 haveSize = reader.ReadU32();
+		u32 size = m_casKeyTable.GetSize();
+		writer.WriteU32(size);
+		u32 toSend = Min(u32(writer.GetCapacityLeft()), size - haveSize);
+		writer.WriteBytes(m_casKeyTable.GetMemory() + haveSize, toSend);
+		return true;
+	}
+
+	bool CacheServer::HandleFetchEntries(BinaryReader& reader, BinaryWriter& writer)
+	{
+		CasKey cmdKey = reader.ReadCasKey();
+
+		u16& entryCount = *(u16*)writer.AllocWrite(2);
+		entryCount = 0;
+
+		SCOPED_READ_LOCK(m_cacheEntryLookupLock, lock);
+		auto findIt = m_cacheEntryLookup.find(cmdKey);
+		if (findIt == m_cacheEntryLookup.end())
+			return true;
+		auto& cacheEntries = findIt->second;
+		lock.Leave();
+
+		SCOPED_READ_LOCK(cacheEntries.lock, lock2);
+
+		for (auto& entry : cacheEntries.entries)
+		{
+			auto& inputs = entry.inputCasKeyOffsets;
+			auto& outputs = entry.outputCasKeyOffsets;
+
+			u64 neededSize = Get7BitEncodedCount(inputs.size()) + inputs.size() + Get7BitEncodedCount(outputs.size()) + outputs.size();
+			if (neededSize > writer.GetCapacityLeft())
+				break;
+
+			writer.Write7BitEncoded(inputs.size());
+			writer.WriteBytes(inputs.data(), inputs.size());
+			writer.Write7BitEncoded(outputs.size());
+			writer.WriteBytes(outputs.data(), outputs.size());
+			++entryCount;
+		}
+		return true;
+	}
+
+	bool CacheServer::HandleCreateStatusFile(BinaryReader& reader, BinaryWriter& writer)
+	{
+		StringBuffer<> filterString;
+		reader.ReadString(filterString);
+
+		SCOPED_READ_LOCK(m_maintenanceLock, lock);
+		SCOPED_READ_LOCK(m_cacheEntryLookupLock, lock2);
+
+		StringBuffer<> tempFile(m_storage.GetTempPath());
+		Guid guid;
+		CreateGuid(guid);
+		tempFile.Append(GuidToString(guid).str);
+
+		FileAccessor file(m_logger, tempFile.data);
+		if (!file.CreateWrite())
+			return false;
+
+		bool success = true;
+		auto Write = [&](const void* data, u64 size) { success &= file.Write(data, size); };
+
+		u8 bom[] = {0xEF,0xBB,0xBF}; 
+		Write(bom, sizeof(bom));
+
+		auto writeLine = [&](const tchar* text)
+			{
+				u8 buffer[1024];
+				BinaryWriter w(buffer, 0, sizeof(buffer));
+				w.WriteUtf8String(text, TStrlen(text));
+				w.WriteUtf8String(TC("\n"), 1);
+				Write(buffer, w.GetPosition());
+			};
+
+		writeLine(TC("UbaCache server summary"));
+
+		u64 now = GetSystemTimeAsFileTime();
+
+		for (auto& kv : m_cacheEntryLookup)
+		{
+			CacheEntries& entries = kv.second;
+			SCOPED_READ_LOCK(entries.lock, lock3);
+
+			Set<u32> visibleIndices;
+			if (filterString.count)
+			{
+				u32 index = 0;
+				for (auto& entry : entries.entries)
+				{
+					auto findString = [&](const Vector<u8>& offsets)
+						{
+							BinaryReader reader2(offsets.data(), 0, offsets.size());
+							while (reader2.GetLeft())
+							{
+								u64 offset = reader2.Read7BitEncoded();
+								CasKey casKey;
+								StringBuffer<> path;
+								m_casKeyTable.GetPathAndKey(path, casKey, m_pathTable, offset);
+								if (path.Contains(filterString.data))
+									return true;
+							}
+							return false;
+						};
+
+					if (findString(entry.inputCasKeyOffsets) || findString(entry.outputCasKeyOffsets))
+						visibleIndices.insert(index);
+					++index;
+				}
+				if (visibleIndices.empty())
+					continue;
+			}
+
+
+			StringBuffer<> line;
+			writeLine(CasKeyString(kv.first).str);
+			u32 index = 0;
+			for (auto& entry : entries.entries)
+			{
+				if (!visibleIndices.empty() && visibleIndices.find(index) == visibleIndices.end())
+				{
+					++index;
+					continue;
+				}
+
+				u64 age = MsToTime(GetFileTimeAsSeconds(now - entry.creationTime)*1000);
+				writeLine(line.Clear().Appendf(TC("  #%u (%s ago)"), index, TimeToText(age, true).str).data);
+
+				auto writeOffsets = [&](const Vector<u8>& offsets)
+					{
+						BinaryReader reader2(offsets.data(), 0, offsets.size());
+						while (reader2.GetLeft())
+						{
+							u64 offset = reader2.Read7BitEncoded();
+							CasKey casKey;
+							StringBuffer<> path;
+							m_casKeyTable.GetPathAndKey(path, casKey, m_pathTable, offset);
+							writeLine(line.Clear().Appendf(TC("    %s - %s"), path.data, CasKeyString(casKey).str).data);
+						}
+					};
+
+				writeLine(line.Clear().Append (TC("   Inputs:")).data);
+				writeOffsets(entry.inputCasKeyOffsets);
+				writeLine(line.Clear().Append (TC("   Outputs:")).data);
+				writeOffsets(entry.outputCasKeyOffsets);
+				++index;
+			}
+		}
+
+		if (!success || !file.Close())
+			return false;
+
+		CasKey key;
+		if (!m_storage.StoreCasFile(key, tempFile.data))
+			return false;
+
+		writer.WriteCasKey(key);
+
+		DeleteFileW(tempFile.data);
+		return true;
+	}
+}
diff --git a/Engine/Source/Programs/UnrealBuildAccelerator/Common/Private/UbaCompactTables.cpp b/Engine/Source/Programs/UnrealBuildAccelerator/Common/Private/UbaCompactTables.cpp
new file mode 100644
--- /dev/null
+++ b/Engine/Source/Programs/UnrealBuildAccelerator/Common/Private/UbaCompactTables.cpp
@@ -0,0 +1,232 @@
+// Copyright Epic Games, Inc. All Rights Reserved.
+
+#include "UbaCompactTables.h"
+#include "UbaBinaryReaderWriter.h"
+
+namespace uba
+{
+	CompactPathTable::CompactPathTable(u64 reserveSize, u64 reserveOffsetsCount)
+	{
+		m_reserveSize = reserveSize;
+		if (reserveOffsetsCount)
+			m_offsets.reserve(reserveOffsetsCount);
+	}
+
+	u32 CompactPathTable::Add(const tchar* str, u32 strLen, u32* outRequiredCasTableSize)
+	{
+		SCOPED_WRITE_LOCK(m_lock, lock);
+		if (!m_mem.memory)
+			m_mem.Init(m_reserveSize);
+		u32 res = InternalAdd(str, strLen);
+		if (outRequiredCasTableSize)
+			*outRequiredCasTableSize = u32(m_mem.writtenSize);
+		return res;
+	}
+
+	u32 CompactPathTable::InternalAdd(const tchar* str, u32 strLen)
+	{
+		const tchar* stringKeyString = str;
+
+		StringBuffer<MaxPath> tempStringKeyStr;
+		if (CaseInsensitiveFs)
+			stringKeyString = tempStringKeyStr.Append(str).MakeLower().data;
+
+		StringKey key = ToStringKey(stringKeyString, strLen);
+		auto insres = m_offsets.try_emplace(key);
+		if (!insres.second)
+			return insres.first->second;
+			
+		const tchar* seg = str;
+		u32 parentOffset = 0;
+		if (const tchar* lastSeparator = TStrrchr(str, PathSeparator))
+		{
+			StringBuffer<> temp;
+			temp.Append(str, lastSeparator - str);
+			parentOffset = InternalAdd(temp.data, temp.count);
+			seg = lastSeparator + 1;
+		}
+
+		u64 segLen = strLen - (seg - str);
+		u8 bytesForParent = Get7BitEncodedCount(parentOffset);
+		u64 bytesForString = GetStringWriteSize(seg, segLen);
+
+		if (!m_mem.writtenSize)
+			m_mem.Allocate(1, 1, TC(""));
+
+		u8* mem = (u8*)m_mem.Allocate(bytesForParent + bytesForString, 1, TC(""));
+		BinaryWriter writer(mem, 0, 10000);
+		writer.Write7BitEncoded(parentOffset);
+		writer.WriteString(seg, segLen);
+		insres.first->second = u32(mem - m_mem.memory);
+		return insres.first->second;
+	}
+
+	void CompactPathTable::GetString(StringBufferBase& out, u64 offset) const
+	{
+		u32 offsets[256];
+		offsets[0] = u32(offset);
+		u32 offsetCount = 0;
+		while (offset)
+		{
+			++offsetCount;
+			UBA_ASSERT(offsetCount < sizeof_array(offsets));
+			BinaryReader reader(m_mem.memory, offset, 4);
+			offset = (u32)reader.Read7BitEncoded();
+			offsets[offsetCount] = u32(offset);
+		}
+
+		bool isFirst = true;
+		for (u32 i=offsetCount;i; --i)
+		{
+			BinaryReader reader(m_mem.memory, offsets[i-1], 1024);
+			reader.Read7BitEncoded();
+
+			if (!isFirst)
+				out.Append(PathSeparator);
+			isFirst = false;
+			reader.ReadString(out);
+		}
+	}
+
+	u8* CompactPathTable::GetMemory()
+	{
+		return m_mem.memory;
+	}
+
+	u32 CompactPathTable::GetSize()
+	{
+		SCOPED_READ_LOCK(m_lock, lock2)
+		return u32(m_mem.writtenSize);
+	}
+
+	void CompactPathTable::ReadMem(BinaryReader& reader, bool populateLookup)
+	{
+		if (!m_mem.memory)
+			m_mem.Init(m_reserveSize);
+
+		u64 writtenSize = m_mem.writtenSize;
+		u64 left = reader.GetLeft();
+		void* mem = m_mem.Allocate(left, 1, TC(""));
+		reader.ReadBytes(mem, left);
+
+		if (!populateLookup)
+			return;
+
+		BinaryReader reader2(m_mem.memory, writtenSize, m_mem.writtenSize);
+		if (!writtenSize)
+			reader2.Skip(1);
+
+		while (reader2.GetLeft())
+		{
+			u32 offset = u32(reader2.GetPosition());
+			reader2.Read7BitEncoded();
+			reader2.SkipString();
+			StringBuffer<> str;
+			GetString(str, offset);
+			if (CaseInsensitiveFs)
+				str.MakeLower();
+			m_offsets.try_emplace(ToStringKey(str), offset);
+		}
+	}
+
+	void CompactPathTable::Swap(CompactPathTable& other)
+	{
+		m_offsets.swap(other.m_offsets);
+		m_mem.Swap(other.m_mem);
+		u64 rs = m_reserveSize;
+		m_reserveSize = other.m_reserveSize;
+		other.m_reserveSize = rs;
+	}
+
+	CompactCasKeyTable::CompactCasKeyTable(u64 reserveSize, u64 reserveOffsetsCount)
+	{
+		m_reserveSize = reserveSize;
+		if (reserveOffsetsCount)
+			m_offsets.reserve(reserveOffsetsCount);
+	}
+
+	u32 CompactCasKeyTable::Add(const CasKey& casKey, u64 stringOffset, u32* outRequiredCasTableSize)
+	{
+		SCOPED_WRITE_LOCK(m_lock, lock2)
+		if (!m_mem.memory)
+			m_mem.Init(m_reserveSize);
+		auto insres = m_offsets.try_emplace({ casKey, u32(stringOffset)});
+		if (insres.second)
+		{
+			u8 bytesForStringOffset = Get7BitEncodedCount(stringOffset);
+			u8* mem = (u8*)m_mem.Allocate(bytesForStringOffset + sizeof(CasKey), 1, TC(""));
+			BinaryWriter writer(mem, 0, 1000);
+			writer.Write7BitEncoded(stringOffset);
+			writer.WriteCasKey(casKey);
+			insres.first->second = u32(mem - m_mem.memory);
+			if (outRequiredCasTableSize)
+				*outRequiredCasTableSize = (u32)m_mem.writtenSize;
+		}
+		else if (outRequiredCasTableSize)
+		{
+			BinaryReader reader(m_mem.memory, insres.first->second, 1000);
+			reader.Read7BitEncoded();
+			*outRequiredCasTableSize = Max(*outRequiredCasTableSize, u32(reader.GetPosition() + sizeof(CasKey)));
+		}
+		return insres.first->second;
+	}
+
+	void CompactCasKeyTable::GetKey(CasKey& outKey, u64 offset) const
+	{
+		BinaryReader reader(m_mem.memory, offset, 1000);
+		reader.Read7BitEncoded();
+		outKey = reader.ReadCasKey();
+	}
+
+	void CompactCasKeyTable::GetPathAndKey(StringBufferBase& outPath, CasKey& outKey, const CompactPathTable& pathTable, u64 offset) const
+	{
+		BinaryReader reader(m_mem.memory, offset, 1000);
+		u32 stringOffset = (u32)reader.Read7BitEncoded();
+		outKey = reader.ReadCasKey();
+		pathTable.GetString(outPath, stringOffset);
+	}
+
+	u8* CompactCasKeyTable::GetMemory()
+	{
+		return m_mem.memory;
+	}
+
+	u32 CompactCasKeyTable::GetSize()
+	{
+		SCOPED_READ_LOCK(m_lock, lock2)
+		return u32(m_mem.writtenSize);
+	}
+
+	void CompactCasKeyTable::ReadMem(BinaryReader& reader, bool populateLookup)
+	{
+		if (!m_mem.memory)
+			m_mem.Init(m_reserveSize);
+
+		u64 writtenSize = m_mem.writtenSize;
+
+		u64 left = reader.GetLeft();
+		void* mem = m_mem.Allocate(left, 1, TC(""));
+		reader.ReadBytes(mem, left);
+
+		if (!populateLookup)
+			return;
+
+		BinaryReader reader2(m_mem.memory, writtenSize, m_mem.writtenSize);
+		while (reader2.GetLeft())
+		{
+			u32 offset = u32(reader2.GetPosition());
+			u64 stringOffset = reader2.Read7BitEncoded();
+			CasKey casKey = reader2.ReadCasKey();
+			m_offsets.try_emplace({casKey, u32(stringOffset)}, offset);
+		}
+	}
+
+	void CompactCasKeyTable::Swap(CompactCasKeyTable& other)
+	{
+		m_offsets.swap(other.m_offsets);
+		m_mem.Swap(other.m_mem);
+		u64 rs = m_reserveSize;
+		m_reserveSize = other.m_reserveSize;
+		other.m_reserveSize = rs;
+	}
+}
diff --git a/Engine/Source/Programs/UnrealBuildAccelerator/Common/Private/UbaScheduler.cpp b/Engine/Source/Programs/UnrealBuildAccelerator/Common/Private/UbaScheduler.cpp
--- a/Engine/Source/Programs/UnrealBuildAccelerator/Common/Private/UbaScheduler.cpp
+++ b/Engine/Source/Programs/UnrealBuildAccelerator/Common/Private/UbaScheduler.cpp
@@ -2,6 +2,8 @@
 
 #include "UbaScheduler.h"
 #include "UbaApplicationRules.h"
+#include "UbaCacheClient.h"
+#include "UbaNetworkServer.h"
 #include "UbaProcess.h"
 #include "UbaProcessStartInfoHolder.h"
 #include "UbaSessionServer.h"
@@ -26,6 +28,7 @@ namespace uba
 		const u8* knownInputs;
 		u32 knownInputsCount;
 		float weight = 1.0f;
+		u64 queryCacheTime = 0;
 	};
 
 
@@ -55,6 +58,24 @@ namespace uba
 		ProcessStartInfoHolder holder;
 	};
 
+	class CachedProcess : public Process
+	{
+	public:
+		CachedProcess(const ProcessStartInfo& i) : holder(i) {}
+		virtual u32 GetExitCode() override { return 0; }
+		virtual bool HasExited() override { return true; }
+		virtual bool WaitForExit(u32 millisecondsTimeout) override{ return true; }
+		virtual const ProcessStartInfo& GetStartInfo() const override { return holder.startInfo; }
+		virtual const Vector<ProcessLogLine>& GetLogLines() const override { return logLines; }
+		virtual const Vector<u8>& GetTrackedInputs() const override { static Vector<u8> v; return v;}
+		virtual const Vector<u8>& GetTrackedOutputs() const override { static Vector<u8> v; return v;}
+		virtual bool IsRemote() const override { return false; }
+		virtual ProcessExecutionType GetExecutionType() const override { return ProcessExecutionType_FromCache; }
+		ProcessStartInfoHolder holder;
+		Vector<ProcessLogLine> logLines;
+	};
+
+
 	Scheduler::Scheduler(const SchedulerCreateInfo& info)
 	:	m_session(info.session)
 	,	m_maxLocalProcessors(info.maxLocalProcessors != ~0u ? info.maxLocalProcessors : GetLogicalProcessorCount())
@@ -62,6 +83,8 @@ namespace uba
 	,	m_enableProcessReuse(info.enableProcessReuse)
 	,	m_forceRemote(info.forceRemote)
 	,	m_forceNative(info.forceNative)
+	,	m_cacheClient(info.cacheClient)
+	,	m_writeToCache(info.writeToCache && info.cacheClient)
 	{
 		m_session.RegisterGetNextProcess([this](Process& process, NextProcessInfo& outNextProcess, u32 prevExitCode)
 			{
@@ -129,13 +152,15 @@ namespace uba
 		const ApplicationRules* rules = m_session.GetRules(info2->startInfo);
 		info2->startInfo.rules = rules;
 
+		bool useCache = m_cacheClient && !m_writeToCache && rules->IsCacheable();
+
 		SCOPED_WRITE_LOCK(m_processEntriesLock, lock);
 		u32 index = u32(m_processEntries.size());
 		auto& entry = m_processEntries.emplace_back();
 		entry.info = info2;
 		entry.dependencies = dep;
 		entry.dependencyCount = info.dependencyCount;
-		entry.status = ProcessStatus_Queued;
+		entry.status = useCache ? ProcessStatus_QueuedForCache : ProcessStatus_QueuedForRun;
 		entry.canDetour = info.canDetour;
 		entry.canExecuteRemotely = info.canExecuteRemotely && info.canDetour;
 		lock.Leave();
@@ -186,8 +211,29 @@ namespace uba
 		SCOPED_WRITE_LOCK(m_processEntriesLock, lock);
 		if (m_processEntries[processIndex].status != ProcessStatus_Running)
 			return;
-		m_processEntries[processIndex].status = ProcessStatus_Queued;
+		m_processEntries[processIndex].status = ProcessStatus_QueuedForRun;
+		m_processEntriesStart = Min(m_processEntriesStart, processIndex);
+		lock.Leave();
+
+		UpdateQueueCounter(1);
+		UpdateActiveProcessCounter(false, -1);
+		m_updateThreadLoop.Set();
+	}
+
+	void Scheduler::HandleCacheMissed(ExitProcessInfo* ei)
+	{
+		u32 processIndex = ei->processIndex;
+		delete ei;
+
+		if (processIndex == ~0u)
+			return;
+
+		SCOPED_WRITE_LOCK(m_processEntriesLock, lock);
+		if (m_processEntries[processIndex].status != ProcessStatus_Running)
+			return;
+		m_processEntries[processIndex].status = ProcessStatus_QueuedForRun;
 		m_processEntriesStart = Min(m_processEntriesStart, processIndex);
+		--m_activeCacheQueries;
 		lock.Leave();
 
 		UpdateQueueCounter(1);
@@ -214,15 +260,20 @@ namespace uba
 			return;
 		}
 
-		ExitProcess(*info, *handle.m_process, handle.m_process->GetExitCode());
+		if (si->queryCacheTime) // A bit hacky but we know this is a local process
+		{
+			Timer& timer = ((ProcessImpl*)handle.m_process)->m_processStats.queryCache;
+			timer.count = 1;
+			timer.time = si->queryCacheTime;
+		}
+
+		ExitProcess(*info, *handle.m_process, handle.m_process->GetExitCode(), false);
 	}
 
-	u32 Scheduler::PopProcess(bool isLocal)
+	u32 Scheduler::PopProcess(bool isLocal, ProcessStatus& outPrevStatus)
 	{
-		if (isLocal)
-			if (m_activeLocalProcessWeight >= float(m_maxLocalProcessors))
-				return ~0u;
-
+		bool atMaxLocalWeight = m_activeLocalProcessWeight >= float(m_maxLocalProcessors);
+		bool atMaxCacheQueries = m_activeCacheQueries >= 16;
 		auto processEntries = m_processEntries.data();
 		bool allFinished = true;
 
@@ -230,7 +281,7 @@ namespace uba
 		{
 			auto& entry = processEntries[i];
 			auto status = entry.status;
-			if (status != ProcessStatus_Queued)
+			if (status != ProcessStatus_QueuedForCache && status != ProcessStatus_QueuedForRun)
 			{
 				if (allFinished)
 				{
@@ -243,17 +294,32 @@ namespace uba
 			}
 			allFinished = false;
 
-			if (!isLocal && !entry.canExecuteRemotely)
-				continue;
-
-			if (isLocal && m_forceRemote && entry.canExecuteRemotely)
-				continue;
+			if (isLocal)
+			{
+				if (m_forceRemote && entry.canExecuteRemotely)
+					continue;
+				if (status == ProcessStatus_QueuedForRun && atMaxLocalWeight)
+					continue;
+				if (status == ProcessStatus_QueuedForCache && atMaxCacheQueries)
+					continue;
+			}
+			else
+			{
+				if (!entry.canExecuteRemotely)
+					continue;
+				if (status == ProcessStatus_QueuedForCache)
+					continue;
+			}
 
 			bool canRun = true;
 			for (u32 j=0, je=entry.dependencyCount; j!=je; ++j)
 			{
 				auto depIndex = entry.dependencies[j];
-				UBA_ASSERTF(depIndex < m_processEntries.size(), TC("Found dependency on index %u but there are only %u processes registered"), depIndex, u32(m_processEntries.size()));
+				if (depIndex >= m_processEntries.size())
+				{
+					m_session.GetLogger().Error(TC("Found dependency on index %u but there are only %u processes registered"), depIndex, u32(m_processEntries.size()));
+					return ~0u;
+				}
 				auto depStatus = processEntries[depIndex].status;
 				if (depStatus == ProcessStatus_Failed || depStatus == ProcessStatus_Skipped)
 				{
@@ -271,8 +337,14 @@ namespace uba
 				continue;
 
 			if (isLocal)
-				m_activeLocalProcessWeight += entry.info->weight;
+			{
+				if (status == ProcessStatus_QueuedForRun)
+					m_activeLocalProcessWeight += entry.info->weight;
+				else
+					++m_activeCacheQueries;
+			}
 
+			outPrevStatus = entry.status;
 			entry.status = ProcessStatus_Running;
 			return i;
 		}
@@ -284,8 +356,9 @@ namespace uba
 	{
 		while (true)
 		{
+			ProcessStatus prevStatus;
 			SCOPED_WRITE_LOCK(m_processEntriesLock, lock);
-			u32 indexToRun = PopProcess(isLocal);
+			u32 indexToRun = PopProcess(isLocal, prevStatus);
 			if (indexToRun == ~0u)
 				return false;
 
@@ -319,6 +392,33 @@ namespace uba
 					ei->scheduler->ProcessExited(ei, handle);
 				};
 
+			UBA_ASSERT(si.rules);
+			if (m_writeToCache && si.rules->IsCacheable())
+				si.trackInputs = true;
+			else if (prevStatus == ProcessStatus_QueuedForCache)
+			{
+				// TODO: This should not use work manager since it is mostly waiting on network
+				m_session.GetServer().AddWork([this, exitInfo]()
+					{
+						ProcessStartInfo& si = exitInfo->startInfo->startInfo;
+						u64 startTime = GetTime();
+
+						if (m_cacheClient->FetchFromCache(si))
+						{
+							auto process = new CachedProcess(si);
+							ProcessHandle ph(process);
+							exitInfo->startInfo->queryCacheTime = GetTime() - startTime;
+							ExitProcess(*exitInfo, *process, 0, true);
+						}
+						else
+						{
+							exitInfo->startInfo->queryCacheTime = GetTime() - startTime;
+							HandleCacheMissed(exitInfo);
+						}
+					}, 1, TC("DownloadCache"));
+				return true;
+			}
+
 			if (isLocal)
 				m_session.RunProcess(si, true, canDetour);
 			else
@@ -337,7 +437,7 @@ namespace uba
 		if (!ei) // If null, process has already exited from some other thread
 			return false;
 
-		ExitProcess(*ei, process, prevExitCode);
+		ExitProcess(*ei, process, prevExitCode, false);
 
 		ei->startInfo = nullptr;
 		ei->processIndex = ~0u;
@@ -348,10 +448,12 @@ namespace uba
 
 		while (true)
 		{
+			ProcessStatus prevStatus;
 			SCOPED_WRITE_LOCK(m_processEntriesLock, lock);
-			u32 indexToRun = PopProcess(isLocal);
+			u32 indexToRun = PopProcess(isLocal, prevStatus);
 			if (indexToRun == ~0u)
 				return false;
+			UBA_ASSERT(prevStatus != ProcessStatus_QueuedForCache);
 			auto& processEntry = m_processEntries[indexToRun];
 			auto newInfo = processEntry.info;
 			bool wasSkipped = processEntry.status == ProcessStatus_Skipped;
@@ -380,7 +482,7 @@ namespace uba
 		}
 	}
 
-	void Scheduler::ExitProcess(ExitProcessInfo& info, Process& process, u32 exitCode)
+	void Scheduler::ExitProcess(ExitProcessInfo& info, Process& process, u32 exitCode, bool fromCache)
 	{
 		ProcessHandle ph;
 		ph.m_process = &process;
@@ -396,7 +498,12 @@ namespace uba
 		entry.info = nullptr;
 		entry.dependencies = nullptr;
 		if (info.isLocal)
-			m_activeLocalProcessWeight -= si->weight;
+		{
+			if (fromCache)
+				--m_activeCacheQueries;
+			else
+				m_activeLocalProcessWeight -= si->weight;
+		}
 		lock.Leave();
 
 		UpdateActiveProcessCounter(info.isLocal, -1);
@@ -405,6 +512,9 @@ namespace uba
 		delete[] dependencies;
 		delete si;
 
+		if (m_writeToCache && exitCode == 0 && process.GetStartInfo().trackInputs)
+			m_cacheClient->WriteToCache(ph);
+
 		ph.m_process = nullptr;
 	}
 
@@ -533,6 +643,15 @@ namespace uba
 				weight = 1.0f;
 			};
 
+		enum InsideArray
+		{
+			InsideArray_None,
+			InsideArray_CacheRoots,
+			InsideArray_Processes,
+		};
+
+		InsideArray insideArray = InsideArray_None;
+
 		auto readLine = [&](const TString& line)
 			{
 				const tchar* keyStart = line.c_str();
@@ -540,69 +659,115 @@ namespace uba
 					++keyStart;
 				if (!*keyStart)
 					return true;
-				const tchar* colon = TStrchr(keyStart, ':');
-				if (!colon)
-					return false;
+				u32 indentation = u32(keyStart - line.c_str());
+
+				if (insideArray != InsideArray_None && !indentation)
+					insideArray = InsideArray_None;
+
+				StringBuffer<32> key;
+				const tchar* valueStart = nullptr;
+
 				if (*keyStart == '-')
 				{
-					keyStart += 2;
-					if (!app.empty())
-						enqueueProcess();
+					UBA_ASSERT(insideArray != InsideArray_None);
+					valueStart = keyStart + 2;
+				}
+				else
+				{
+					const tchar* colon = TStrchr(keyStart, ':');
+					if (!colon)
+						return false;
+					key.Append(keyStart, colon - keyStart);
+					valueStart = colon + 1;
+					while (*valueStart && *valueStart == ' ')
+						++valueStart;
 				}
 
-				StringBuffer<32> key;
-				key.Append(keyStart, colon - keyStart);
-				const tchar* valueStart = colon + 1;
-				while (*valueStart && *valueStart == ' ')
-					++valueStart;
-
-				if (key.Equals(TC("environment")))
+				switch (insideArray)
 				{
-					#if PLATFORM_WINDOWS
-					SetEnvironmentVariable(TC("PATH"), valueStart);
-					#endif
+				case InsideArray_None:
+				{
+					if (key.Equals(TC("environment")))
+					{
+						#if PLATFORM_WINDOWS
+						SetEnvironmentVariable(TC("PATH"), valueStart);
+						#endif
+						return true;
+					}
+					if (key.Equals(TC("cacheroots")))
+					{
+						insideArray = InsideArray_CacheRoots;
+						return true;
+					}
+					if (key.Equals(TC("processes")))
+					{
+						insideArray = InsideArray_Processes;
+						return true;
+					}
+					return true;
 				}
-				else if (key.Equals(TC("processes")))
+				case InsideArray_CacheRoots:
+				{
+					if (m_cacheClient)
+					{
+						if (Equals(valueStart, TC("SystemRoots")))
+							m_cacheClient->RegisterSystemRoots();
+						else
+							m_cacheClient->RegisterRoot(valueStart);
+					}
 					return true;
-				else if (key.Equals(TC("app")))
-					app = valueStart;
-				else if (key.Equals(TC("arg")))
-					arg = valueStart;
-				else if (key.Equals(TC("dir")))
-					dir = valueStart;
-				else if (key.Equals(TC("desc")))
-					desc = valueStart;
-				else if (key.Equals(TC("detour")))
-					allowDetour = !Equals(valueStart, TC("false"));
-				else if (key.Equals(TC("remote")))
-					allowRemote = !Equals(valueStart, TC("false"));
-				else if (key.Equals(TC("weight")))
-					StringBuffer<32>(valueStart).Parse(weight);
-				else if (key.Equals(TC("dep")))
+				}
+				case InsideArray_Processes:
 				{
-					const tchar* depStart = TStrchr(valueStart, '[');
-					if (!depStart)
-						return false;
-					++depStart;
-					StringBuffer<32> depStr;
-					for (const tchar* it = depStart; *it; ++it)
+					if (*keyStart == '-')
+					{
+						keyStart += 2;
+						if (!app.empty())
+							enqueueProcess();
+					}
+
+					if (key.Equals(TC("app")))
+						app = valueStart;
+					else if (key.Equals(TC("arg")))
+						arg = valueStart;
+					else if (key.Equals(TC("dir")))
+						dir = valueStart;
+					else if (key.Equals(TC("desc")))
+						desc = valueStart;
+					else if (key.Equals(TC("detour")))
+						allowDetour = !Equals(valueStart, TC("false"));
+					else if (key.Equals(TC("remote")))
+						allowRemote = !Equals(valueStart, TC("false"));
+					else if (key.Equals(TC("weight")))
+						StringBuffer<32>(valueStart).Parse(weight);
+					else if (key.Equals(TC("dep")))
 					{
-						if (*it != ']' && *it != ',')
+						const tchar* depStart = TStrchr(valueStart, '[');
+						if (!depStart)
+							return false;
+						++depStart;
+						StringBuffer<32> depStr;
+						for (const tchar* it = depStart; *it; ++it)
 						{
-							if (*it != ' ')
-								depStr.Append(*it);
-							continue;
+							if (*it != ']' && *it != ',')
+							{
+								if (*it != ' ')
+									depStr.Append(*it);
+								continue;
+							}
+							u32 depIndex;
+							if (!depStr.Parse(depIndex))
+								return false;
+							depStr.Clear();
+							deps.push_back(depIndex);
+
+							if (!*it)
+								break;
+							depStart = it + 1;
 						}
-						u32 depIndex;
-						if (!depStr.Parse(depIndex))
-							return false;
-						depStr.Clear();
-						deps.push_back(depIndex);
-
-						if (!*it)
-							break;
-						depStart = it + 1;
 					}
+					return true;
+				}
 				}
 				return true;
 			};
diff --git a/Engine/Source/Programs/UnrealBuildAccelerator/Common/Private/UbaStorage.cpp b/Engine/Source/Programs/UnrealBuildAccelerator/Common/Private/UbaStorage.cpp
--- a/Engine/Source/Programs/UnrealBuildAccelerator/Common/Private/UbaStorage.cpp
+++ b/Engine/Source/Programs/UnrealBuildAccelerator/Common/Private/UbaStorage.cpp
@@ -5,7 +5,9 @@
 #include "UbaFileAccessor.h"
 #include "UbaBinaryReaderWriter.h"
 #include "UbaDirectoryIterator.h"
+#include "UbaStorageUtils.h"
 #include "UbaWorkManager.h"
+#include <oodle2.h>
 
 namespace uba
 {
@@ -23,6 +25,32 @@ namespace uba
 		out.Append('^').AppendHex(mappingHandle.ToU64()).Append('-').AppendHex(offset);
 	}
 
+	u8* BufferSlots::Pop()
+	{
+		SCOPED_WRITE_LOCK(m_slotsLock, lock);
+		if (!m_slots.empty())
+		{
+			auto back = m_slots.back();
+			m_slots.pop_back();
+			return back;
+		}
+		return (u8*)malloc(BufferSlotSize);
+	}
+
+	void BufferSlots::Push(u8* slot)
+	{
+		if (!slot)
+			return;
+		SCOPED_WRITE_LOCK(m_slotsLock, lock);
+		m_slots.push_back(slot);
+	}
+
+	BufferSlots::~BufferSlots()
+	{
+		for (u8* slot : m_slots)
+			free(slot);
+	}
+
 	void StorageImpl::CasEntryAccessed(CasEntry& entry)
 	{
 		bool hasMapping;
@@ -215,7 +243,7 @@ namespace uba
 
 		u64 totalWritten = 0;
 
-		u64 diff = (u64)OodleLZ_GetCompressedBufferSizeNeeded(m_sendCasCompressor, BufferSlotHalfSize) - BufferSlotHalfSize;
+		u64 diff = (u64)OodleLZ_GetCompressedBufferSizeNeeded((OodleLZ_Compressor)m_casCompressor, BufferSlotHalfSize) - BufferSlotHalfSize;
 		u64 maxUncompressedBlock = BufferSlotHalfSize - diff - 8; // 8 bytes for the little header
 		u32 workCount = u32((fileSize + maxUncompressedBlock - 1) / maxUncompressedBlock);
 
@@ -278,8 +306,8 @@ namespace uba
 		}
 		else
 		{
-			u8* slot = PopBufferSlot();
-			auto _ = MakeGuard([&](){ PushBufferSlot(slot); });
+			u8* slot = m_bufferSlots.Pop();
+			auto _ = MakeGuard([&](){ m_bufferSlots.Push(slot); });
 			u8* uncompressedData = slot;
 			u8* compressBuffer = slot + BufferSlotHalfSize;
 
@@ -299,7 +327,7 @@ namespace uba
 				OO_SINTa compressedBlockSize;
 				{
 					TimerScope cts(stats.compressWrite);
-					compressedBlockSize = OodleLZ_Compress(m_createCasCompressor, uncompressedData, (OO_SINTa)uncompressedBlockSize, destBuf + 8, m_createCasCompressionLevel);
+					compressedBlockSize = OodleLZ_Compress((OodleLZ_Compressor)m_casCompressor, uncompressedData, (OO_SINTa)uncompressedBlockSize, destBuf + 8, (OodleLZ_CompressionLevel)m_casCompressionLevel);
 					if (compressedBlockSize == OODLELZ_FAILED)
 						return m_logger.Error(TC("Failed to compress %llu bytes for %s"), uncompressedBlockSize, from);
 				}
@@ -399,8 +427,8 @@ namespace uba
 
 		auto work = [this, rec, &stats]()
 		{
-			u8* slot = PopBufferSlot();
-			auto _ = MakeGuard([&](){ PushBufferSlot(slot); });
+			u8* slot = m_bufferSlots.Pop();
+			auto _ = MakeGuard([&](){ m_bufferSlots.Push(slot); });
 			u8* compressSlotBuffer = slot + BufferSlotHalfSize;
 			while (true)
 			{
@@ -417,7 +445,7 @@ namespace uba
 				OO_SINTa compressedBlockSize;
 				{
 					TimerScope cts(stats.compressWrite);
-					compressedBlockSize = OodleLZ_Compress(m_createCasCompressor, uncompressedDataSlot, uncompressedBlockSize, compressSlotBuffer + 8, m_createCasCompressionLevel);
+					compressedBlockSize = OodleLZ_Compress((OodleLZ_Compressor)m_casCompressor, uncompressedDataSlot, uncompressedBlockSize, compressSlotBuffer + 8, (OodleLZ_CompressionLevel)m_casCompressionLevel);
 					if (compressedBlockSize == OODLELZ_FAILED)
 					{
 						m_logger.Error(TC("Failed to compress %llu bytes for %s"), u64(uncompressedBlockSize), rec->destination->GetFileName());
@@ -487,8 +515,8 @@ namespace uba
 			if (!destinationFile.CreateWrite(false, DefaultAttributes(), 0, m_tempPath.data))
 				return false;
 
-			u8* slot= PopBufferSlot();
-			auto _ = MakeGuard([&](){ PushBufferSlot(slot); });
+			u8* slot= m_bufferSlots.Pop();
+			auto _ = MakeGuard([&](){ m_bufferSlots.Push(slot); });
 			u64 left = fileSize;
 			while (left)
 			{
@@ -684,8 +712,8 @@ namespace uba
 		StorageStats& stats = Stats();
 		u8* readPos = compressedData;
 
-		u8* slot = PopBufferSlot();
-		auto _ = MakeGuard([&](){ PushBufferSlot(slot); });
+		u8* slot = m_bufferSlots.Pop();
+		auto _ = MakeGuard([&](){ m_bufferSlots.Push(slot); });
 
 		u64 left = decompressedSize;
 		u64 overflow = 0;
@@ -868,13 +896,13 @@ namespace uba
 			});
 	}
 
-	void StorageImpl::TraverseAllCasFiles(const Function<void(const CasKey& key)>& func)
+	void StorageImpl::TraverseAllCasFiles(const Function<void(const CasKey& key, u64 size)>& func)
 	{
 		StringBuffer<> casRoot;
 		casRoot.Append(m_rootDir.data, m_rootDir.count - 1);
 		TraverseAllCasFiles(casRoot.data, 0, [&](const StringBufferBase& fullPath, const DirectoryEntry& e)
 			{
-				func(CasKeyFromString(e.name));
+				func(CasKeyFromString(e.name), e.size);
 			});
 	}
 
@@ -956,7 +984,7 @@ namespace uba
 #endif
 	}
 
-	void StorageImpl::HandleOverflow()
+	void StorageImpl::HandleOverflow(Set<CasKey>* outDeletedFiles)
 	{
 		if (!m_casCapacityBytes)
 			return;
@@ -971,6 +999,8 @@ namespace uba
 				break;
 			}
 			DropCasFile(entry->key, true, TC("HandleOverflow"));
+			if (outDeletedFiles)
+				outDeletedFiles->insert(entry->key);
 			DetachEntry(*entry);
 			m_casLookup.erase(entry->key);
 		}
@@ -1071,6 +1101,9 @@ namespace uba
 
 		m_maxParallelCopyOrLink = info.maxParallelCopyOrLink;
 
+		m_casCompressor = DefaultCompressor;
+		m_casCompressionLevel = DefaultCompressionLevel;
+
 		#if UBA_USE_MIMALLOC
 		OodleCore_Plugins_SetAllocators(Oodle_MallocAligned, Oodle_Free);
 		#endif
@@ -1079,8 +1112,6 @@ namespace uba
 	StorageImpl::~StorageImpl()
 	{
 		SaveCasTable(true, true);
-		for (u8* slot : m_compSlots)
-			free(slot);
 	}
 
 	bool CheckExclusive(Logger& logger, const StringBufferBase& rootDir)
@@ -1263,7 +1294,7 @@ namespace uba
 			CheckAllCasFiles();
 			resave = true;
 		}
-		HandleOverflow();
+		HandleOverflow(nullptr);
 
 		if (resave)
 		{
@@ -1275,7 +1306,7 @@ namespace uba
 		if (logStats)
 		{
 			u64 duration = GetTime() - startTime;
-			m_logger.Detail(TC("Database loaded from %s in %s (contained %llu entries)"), fileName.data, TimeToText(duration).str, m_casLookup.size());
+			m_logger.Detail(TC("Database loaded from %s in %s (contained %llu entries estimated to %s)"), fileName.data, TimeToText(duration).str, m_casLookup.size(), BytesToText(m_casTotalBytes).str);
 		}
 
 		return true;
@@ -1528,6 +1559,11 @@ namespace uba
 		return true;
 	}
 
+	const tchar* StorageImpl::GetTempPath()
+	{
+		return m_tempPath.data;
+	}
+
 	u64 StorageImpl::GetStorageCapacity()
 	{
 		return m_casCapacityBytes;
@@ -1838,7 +1874,25 @@ namespace uba
 		auto findIt = m_casLookup.find(casKey);
 		bool foundEntry = findIt != m_casLookup.end();
 		if (!foundEntry)
+		{
+			if (forceDelete)
+			{
+				StringBuffer<> casFile;
+				#if !UBA_USE_SPARSEFILE
+				if (!StorageImpl::GetCasFileName(casFile, casKey))
+					return false;
+				#else
+				UBA_ASSERT(false);
+				#endif
+				if (DeleteFileW(casFile.data) == 0)
+				{
+					u32 lastError = GetLastError();
+					if (lastError != ERROR_FILE_NOT_FOUND && lastError != ERROR_PATH_NOT_FOUND)
+						return m_logger.Error(TC("Failed to drop cas %s (%s) (%s)"), casFile.data, hint, LastErrorToText(lastError).data);
+				}
+			}
 			return true;
+		}
 		CasEntry& casEntry = findIt->second;
 		lookupLock.Leave();
 	
@@ -1847,12 +1901,13 @@ namespace uba
 		if (forceDelete)
 		{
 			StringBuffer<> casFile;
-#if !UBA_USE_SPARSEFILE
+			#if !UBA_USE_SPARSEFILE
 			if (!StorageImpl::GetCasFileName(casFile, casKey))
 				return false;
-#else
+			#else
 			UBA_ASSERT(false);
-#endif
+			#endif
+
 			u64 sizeDeleted = 0;
 			if (DeleteFileW(casFile.data) == 0)
 			{
@@ -2113,10 +2168,10 @@ namespace uba
 
 	}
 
-	bool StorageImpl::FakeCopy(const CasKey& casKey, const tchar* destination)
+	bool StorageImpl::FakeCopy(const CasKey& casKey, const tchar* destination, u64 size, u64 lastWritten, bool deleteExisting)
 	{
-		// Delete existing file to make sure it is not picked up (since it is out of date)
-		DeleteFileW(destination);
+		if (deleteExisting) 
+			DeleteFileW(destination);
 
 		StringBuffer<> forKey;
 		forKey.Append(destination);
@@ -2127,8 +2182,8 @@ namespace uba
 		auto insres = m_fileTableLookup.try_emplace(key);
 		FileEntry& entry = insres.first->second;
 		entry.casKey = casKey;
-		entry.lastWritten = 0;
-		entry.size = 0;
+		entry.lastWritten = lastWritten;
+		entry.size = size;
 		entry.verified = true;
 		return true;
 	}
@@ -2191,7 +2246,7 @@ namespace uba
 			return;
 		}
 		
-		logger.Info(TC("  WorkMemoryBuffers    %6u %9s"), u32(m_compSlots.size()), BytesToText(m_compSlots.size() * BufferSlotSize).str);
+		logger.Info(TC("  WorkMemoryBuffers    %6u %9s"), u32(m_bufferSlots.m_slots.size()), BytesToText(m_bufferSlots.m_slots.size() * BufferSlotSize).str);
 		logger.Info(TC("  FileTable            %6u"), u32(m_fileTableLookup.size()));
 
 		StorageStats& stats = Stats();
@@ -2214,94 +2269,7 @@ namespace uba
 	{
 		StorageStats& stats = Stats();
 		TimerScope ts(stats.calculateCasKey);
-
-		CasKeyHasher hasher;
-
-		if (fileSize == 0)
-			return ToCasKey(hasher, storeCompressed);
-
-		#ifndef __clang_analyzer__
-
-		if (fileSize > BufferSlotSize)
-		{
-			struct WorkRec
-			{
-				Atomic<u64> refCount;
-				Atomic<u64> counter;
-				Atomic<u64> doneCounter;
-				u8* fileMem = nullptr;
-				u64 workCount = 0;
-				u64 fileSize = 0;
-				bool error = false;
-				Vector<CasKey> keys;
-				Event done;
-			};
-
-			u32 workCount = u32((fileSize + BufferSlotSize - 1) / BufferSlotSize);
-
-			WorkRec* rec = new WorkRec();
-			rec->fileMem = fileMem;
-			rec->workCount = workCount;
-			rec->fileSize = fileSize;
-			rec->keys.resize(workCount);
-			rec->done.Create(true);
-			rec->refCount = 2;
-
-			auto work = [rec]()
-			{
-				while (true)
-				{
-					u64 index = rec->counter++;
-					if (index >= rec->workCount)
-					{
-						if (!--rec->refCount)
-							delete rec;
-						return 0;
-					}
-
-					u64 startOffset = BufferSlotSize*index;
-					u64 toRead = Min(BufferSlotSize, rec->fileSize - startOffset);
-					u8* slot = rec->fileMem + startOffset;
-					CasKeyHasher hasher;
-					hasher.Update(slot, toRead);
-					rec->keys[index] = ToCasKey(hasher, false);
-
-					if (++rec->doneCounter == rec->workCount)
-						rec->done.Set();
-				}
-				return 0;
-			};
-
-			u32 workerCount = 0;
-			if (m_workManager)
-			{
-				workerCount = Min(workCount, m_workManager->GetWorkerCount()-1); // We are a worker ourselves
-				workerCount = Min(workerCount, MaxWorkItemsPerAction); // Cap this to not starve other things
-				rec->refCount += workerCount;
-				m_workManager->AddWork(work, workerCount, TC("CalculateKey"));
-			}
-
-			work();
-			rec->done.IsSet();
-
-			hasher.Update(rec->keys.data(), rec->keys.size()*sizeof(CasKey));
-
-			bool error = rec->error;
-
-			if (!--rec->refCount)
-				delete rec;
-
-			if (error)
-				return CasKeyZero;
-		}
-		else
-		{
-			hasher.Update(fileMem, fileSize);
-		}
-
-		#endif // __clang_analyzer__
-
-		return ToCasKey(hasher, storeCompressed);
+		return uba::CalculateCasKey(fileMem, fileSize, storeCompressed, m_workManager);
 	}
 
 
@@ -2314,7 +2282,7 @@ namespace uba
 
 		#ifndef __clang_analyzer__
 
-		if (m_workManager && fileSize > BufferSlotSize)
+		if (fileSize > BufferSlotSize) // Note that when filesize is larger than BufferSlotSize the hash becomes a hash of hashes
 		{
 			FileMappingHandle fileMapping = uba::CreateFileMappingW(fileHandle, PAGE_READONLY, fileSize, fileName);
 			if (!fileMapping.IsValid())
@@ -2331,10 +2299,11 @@ namespace uba
 			}
 			auto udg = MakeGuard([&]() { 
 					#if UBA_EXPERIMENTAL
-					m_workManager->AddWork([=, fn = TString(fileName)]() { UnmapViewOfFile(fileData, fileSize, fn.c_str()); }, 1, TC("UnmapFile"));
-					#else
-					UnmapViewOfFile(fileData, fileSize, fileName);
+					if (m_workManager
+						m_workManager->AddWork([=, fn = TString(fileName)]() { UnmapViewOfFile(fileData, fileSize, fn.c_str()); }, 1, TC("UnmapFile"));
+					else
 					#endif
+						UnmapViewOfFile(fileData, fileSize, fileName);
 				});
 
 			struct WorkRec
@@ -2386,11 +2355,19 @@ namespace uba
 				return 0;
 			};
 
-			u32 workerCount = Min(workCount, m_workManager->GetWorkerCount());
-			workerCount = Min(workerCount, MaxWorkItemsPerAction); // Cap this to not starve other things
+			if (m_workManager)
+			{
+				u32 workerCount = Min(workCount, m_workManager->GetWorkerCount());
+				workerCount = Min(workerCount, MaxWorkItemsPerAction); // Cap this to not starve other things
+
+				rec->refCount = workerCount + 1; // We need to keep refcount up 1 to make sure it is not deleted before we read rec->written
+				m_workManager->AddWork(work, workerCount-1, TC("CalculateKey")); // We are a worker ourselves
+			}
+			else
+			{
+				rec->refCount = 2;
+			}
 
-			rec->refCount = workerCount + 1; // We need to keep refcount up 1 to make sure it is not deleted before we read rec->written
-			m_workManager->AddWork(work, workerCount-1, TC("CalculateKey")); // We are a worker ourselves
 			work();
 			rec->done.IsSet();
 
@@ -2406,8 +2383,8 @@ namespace uba
 		}
 		else
 		{
-			u8* slot = PopBufferSlot();
-			auto _ = MakeGuard([&](){ PushBufferSlot(slot); });
+			u8* slot = m_bufferSlots.Pop();
+			auto _ = MakeGuard([&](){ m_bufferSlots.Push(slot); });
 			u64 left = fileSize;
 			while (left)
 			{
@@ -2446,8 +2423,8 @@ namespace uba
 		else
 		{
 			StorageStats& stats = Stats();
-			u8* slot = PopBufferSlot();
-			auto _ = MakeGuard([&]() { PushBufferSlot(slot); });
+			u8* slot = m_bufferSlots.Pop();
+			auto _ = MakeGuard([&]() { m_bufferSlots.Push(slot); });
 
 			u8* readBuffer = slot;
 			u8* writePos = dest;
@@ -2472,26 +2449,6 @@ namespace uba
 		return true;
 	}
 
-	u8* StorageImpl::PopBufferSlot()
-	{
-		SCOPED_WRITE_LOCK(m_compSlotsLock, lock);
-		if (!m_compSlots.empty())
-		{
-			auto back = m_compSlots.back();
-			m_compSlots.pop_back();
-			return back;
-		}
-		return (u8*)malloc(BufferSlotSize);
-	}
-
-	void StorageImpl::PushBufferSlot(u8* slot)
-	{
-		if (!slot)
-			return;
-		SCOPED_WRITE_LOCK(m_compSlotsLock, lock);
-		m_compSlots.push_back(slot);
-	}
-
 	bool StorageImpl::CreateDirectory(const tchar* dir)
 	{
 		return m_dirCache.CreateDirectory(m_logger, dir);
diff --git a/Engine/Source/Programs/UnrealBuildAccelerator/Common/Private/UbaStorageClient.cpp b/Engine/Source/Programs/UnrealBuildAccelerator/Common/Private/UbaStorageClient.cpp
--- a/Engine/Source/Programs/UnrealBuildAccelerator/Common/Private/UbaStorageClient.cpp
+++ b/Engine/Source/Programs/UnrealBuildAccelerator/Common/Private/UbaStorageClient.cpp
@@ -6,6 +6,7 @@
 #include "UbaNetworkClient.h"
 #include "UbaNetworkMessage.h"
 #include "UbaNetworkBackendTcp.h"
+#include "UbaStorageUtils.h"
 
 namespace uba
 {
@@ -56,6 +57,9 @@ namespace uba
 					return;
 
 				m_storageServerUid = reader.ReadGuid();
+				m_casCompressor = reader.ReadByte();
+				m_casCompressionLevel = reader.ReadByte();
+
 			});
 
 		m_client.RegisterOnDisconnected([this]() { m_logger.isMuted = true; });
@@ -213,8 +217,8 @@ namespace uba
 		casFile.Append(CasKeyString(casKey).str);
 		#endif
 
-		u8* slot = PopBufferSlot();
-		auto slotGuard = MakeGuard([&](){ PushBufferSlot(slot); });
+		u8* slot = m_bufferSlots.Pop();
+		auto slotGuard = MakeGuard([&](){ m_bufferSlots.Push(slot); });
 
 		MappedView mappedView;
 		auto mvg = MakeGuard([&]() { if (mappingBuffer) mappingBuffer->UnmapView(mappedView, hint); });
@@ -802,223 +806,10 @@ namespace uba
 			m_proxyClient->client.PrintSummary(logger);
 	}
 
-	bool StorageClient::SendBatchMessages(Logger& logger, NetworkClient& client, u16 fetchId, u8* slot, u64 capacity, u64 left, u32 messageMaxSize, u32& readIndex, u32& responseSize)
-	{
-		responseSize = 0;
-
-		struct Entry
-		{
-			Entry(u8* slot, u32 i, u32 messageMaxSize) : reader(slot + i * messageMaxSize, 0, SendMaxSize), done(true) {}
-			NetworkMessage message;
-			BinaryReader reader;
-			Event done;
-		};
-
-		u64 sendCountCapacity = capacity / messageMaxSize;
-		u64 sendCount = left/messageMaxSize;
-
-		if (sendCount > sendCountCapacity)
-			sendCount = sendCountCapacity;
-		else if (sendCount < sendCountCapacity && (left - sendCount * messageMaxSize) > 0)
-			++sendCount;
-
-
-		u64 entriesMem[sizeof(Entry) * 8];
-		UBA_ASSERT(sizeof(Entry)*sendCount <= sizeof(entriesMem));
-
-		Entry* entries = (Entry*)entriesMem;
-
-		bool success = true;
-		u32 inFlightCount = u32(sendCount);
-		for (u32 i=0; i!=sendCount; ++i)
-		{
-			auto& entry = *new (entries + i) Entry(slot, i, messageMaxSize);
-			StackBinaryWriter<32> writer;
-			entry.message.Init(client, ServiceId, StorageMessageType_FetchSegment, writer);
-			writer.WriteU16(fetchId);
-			writer.WriteU32(readIndex + i + 1);
-			if (entry.message.SendAsync(entry.reader, [](bool error, void* userData) { ((Event*)userData)->Set(); }, &entry.done))
-				continue;
-			entry.~Entry();
-			inFlightCount = i;
-			success = false;
-			break;
-		}
-
-		for (u32 i=0; i!=inFlightCount; ++i)
-		{
-			Entry& entry = entries[i];
-			if (!entry.done.IsSet(5*60*1000))
-				logger.Error(TC("SendBatchMessages timed out getting async message response"));
-			if (!entry.message.ProcessAsyncResults(entry.reader))
-				success = false;
-			else
-				responseSize += u32(entry.reader.GetLeft());
-		}
-
-		for (u32 i=0; i!=inFlightCount; ++i)
-			entries[i].~Entry();
-
-		readIndex += u32(sendCount);
-		return success;
-	}
-	
 	bool StorageClient::SendFile(const CasKey& casKey, const tchar* fileName, u8* sourceMem, u64 sourceSize, const tchar* hint)
 	{
-		UBA_ASSERT(casKey != CasKeyZero);
-
-		NetworkClient& client = m_client; // Don't use proxy
-
-		StorageStats& stats = Stats();
-		TimerScope ts(stats.sendCas);
-
-		u64 firstMessageOverHead = (sizeof(CasKey) + sizeof(u64)*2 + GetStringWriteSize(hint, TStrlen(hint)));
-
-		u64 messageHeader = client.GetMessageHeaderSize();
-		u64 messageHeaderMaxSize = messageHeader + firstMessageOverHead;
-
-		MemoryBlock memoryBlock(sourceSize + messageHeaderMaxSize + 1024);
-		{
-			u8* uncompressedData = sourceMem;
-			u8* compressBufferStart = memoryBlock.memory + messageHeaderMaxSize;
-			u8* compressBuffer = compressBufferStart;
-			u64 totalWritten = messageHeaderMaxSize; // Make sure there is room for msg header in the memory since we are using it to send
-			u64 left = sourceSize;
-
-			compressBuffer += 8;
-			totalWritten += 8;
-			memoryBlock.Allocate(totalWritten, 1, hint);
-
-			u64 diff = u64(OodleLZ_GetCompressedBufferSizeNeeded(m_sendCasCompressor, BufferSlotHalfSize)) - BufferSlotHalfSize;
-			u64 maxUncompressedBlock = BufferSlotHalfSize - diff - totalWritten - 8; // 8 bytes block header
-
-			OodleLZ_CompressOptions oodleOptions = *OodleLZ_CompressOptions_GetDefault();
-			while (left)
-			{
-				u32 uncompressedBlockSize = (u32)Min(left, maxUncompressedBlock);
-
-				u64 reserveSize = totalWritten + uncompressedBlockSize + diff + 8;
-				if (reserveSize > memoryBlock.mappedSize)
-				{
-					u64 toAllocate = reserveSize - memoryBlock.writtenSize;
-					memoryBlock.Allocate(toAllocate, 1, hint);
-				}
-
-				u8* destBuf = compressBuffer;
-				u32 compressedBlockSize;
-				{
-					TimerScope cts(stats.compressSend);
-					compressedBlockSize = (u32)OodleLZ_Compress(m_sendCasCompressor, uncompressedData, (int)uncompressedBlockSize, destBuf + 8, m_sendCasCompressionLevel, &oodleOptions);
-					if (compressedBlockSize == OODLELZ_FAILED)
-						return m_logger.Error(TC("Failed to compress %u bytes at %llu for %s (%s) (%s) (uncompressed size: %llu)"), uncompressedBlockSize, totalWritten, fileName, CasKeyString(casKey).str, hint, sourceSize);
-				}
-
-				*(u32*)destBuf =  u32(compressedBlockSize);
-				*(u32*)(destBuf+4) =  u32(uncompressedBlockSize);
-
-				u32 writeBytes = u32(compressedBlockSize) + 8;
-
-				totalWritten += writeBytes;
-				memoryBlock.writtenSize = totalWritten;
-
-				left -= uncompressedBlockSize;
-				uncompressedData += uncompressedBlockSize;
-				compressBuffer += writeBytes;
-			}
-
-			*(u64*)compressBufferStart = sourceSize;
-		}
-
-
-		u8* readData = memoryBlock.memory + messageHeaderMaxSize;
-		u64 fileSize = memoryBlock.writtenSize - messageHeaderMaxSize;
-
-		u16 storeId = 0;
-		bool isFirst = true;
-		bool sendEnd = false;
-		u64 sendLeft = fileSize;
-		u64 sendPos = 0;
-
-		bool hasSendOneAtTheTimeLock = false;
-		auto lockGuard = MakeGuard([&]() { if (hasSendOneAtTheTimeLock) m_sendOneAtTheTimeLock.LeaveWrite(); });
-
-		while (sendLeft)
-		{
-			u64 writerStartOffset = messageHeader + (isFirst ? firstMessageOverHead : (sizeof(u16) + sizeof(u64)));
-			BinaryWriter writer(readData + sendPos - writerStartOffset, 0, client.GetMessageMaxSize());
-			NetworkMessage msg(client, ServiceId, isFirst ? StorageMessageType_StoreBegin : StorageMessageType_StoreSegment, writer);
-			if (isFirst)
-			{
-				writer.WriteCasKey(casKey);
-				writer.WriteU64(fileSize);
-				writer.WriteU64(sourceSize);
-				writer.WriteString(hint);
-			}
-			else
-			{
-				UBA_ASSERT(storeId != 0);
-				writer.WriteU16(storeId);
-				writer.WriteU64(sendPos);
-			}
-
-			u64 capacityLeft = writer.GetCapacityLeft();
-			u64 toWrite = Min(sendLeft, capacityLeft);
-			writer.AllocWrite(toWrite);
-
-			sendLeft -= toWrite;
-			sendPos += toWrite;
-
-			bool isDone = sendLeft == 0;
-
-			if (isFirst && !isDone)
-			{
-				m_sendOneAtTheTimeLock.EnterWrite();
-				hasSendOneAtTheTimeLock = true;
-			}
-
-			if (isFirst) // First message must always be acknowledged (provide a reader) to make sure there is an entry on server that can be waited on.
-			{
-				StackBinaryReader<128> reader;
-				if (!msg.Send(reader))
-					return false;
-				storeId = reader.ReadU16();
-				sendEnd = reader.ReadBool();
-				if (isDone)
-					break;
-
-				if (!storeId) // Zero means error
-					return m_logger.Error(TC("Server failed to start storing file %s (%s)"), CasKeyString(casKey).str, hint);
-
-				if (storeId == u16(~0)) // File already exists on server
-				{
-					m_logger.Info(TC("Server already has file %s (%s)"), CasKeyString(casKey).str, hint);
-					return true;
-				}
-
-				isFirst = false;
-			}
-			else
-			{
-				if (!msg.Send())
-					return false;
-				if (isDone)
-					break;
-			}
-		}
-
-		if (sendEnd)
-		{
-			StackBinaryWriter<128> writer;
-			NetworkMessage msg(client, ServiceId, StorageMessageType_StoreEnd, writer);
-			writer.WriteCasKey(casKey);
-			if (!msg.Send())
-				return false;
-		}
-
-		stats.sendCasBytesRaw += sourceSize;
-		stats.sendCasBytesComp += fileSize;
-
-		return true;
+		FileSender sender { m_logger, m_client, m_bufferSlots, Stats(), m_sendOneAtTheTimeLock, m_casCompressor, m_casCompressionLevel };
+		return sender.SendFileCompressed(casKey, fileName, sourceMem, sourceSize, hint);
 	}
 
 	bool StorageClient::PopulateCasFromDirsRecursive(const tchar* dir, WorkManager& workManager, UnorderedSet<u64>& seenIds, ReaderWriterLock& seenIdsLock)
diff --git a/Engine/Source/Programs/UnrealBuildAccelerator/Common/Private/UbaStorageServer.cpp b/Engine/Source/Programs/UnrealBuildAccelerator/Common/Private/UbaStorageServer.cpp
--- a/Engine/Source/Programs/UnrealBuildAccelerator/Common/Private/UbaStorageServer.cpp
+++ b/Engine/Source/Programs/UnrealBuildAccelerator/Common/Private/UbaStorageServer.cpp
@@ -12,6 +12,7 @@ namespace uba
 	,	m_server(info.server)
 	{
 		m_zone = info.zone;
+		m_allowFallback = info.allowFallback;
 		m_writeRecievedCasFilesToDisk = info.writeRecievedCasFilesToDisk;
 
 		if (!CreateGuid(m_uid))
@@ -119,7 +120,7 @@ namespace uba
 				server.m_casDataBuffer.UnmapView(mappedView, TC("OnDisconnected"));
 		}
 		else
-			server.PushBufferSlot(memoryBegin);
+			server.m_bufferSlots.Push(memoryBegin);
 	}
 
 	void StorageServer::OnDisconnected(u32 clientId)
@@ -305,6 +306,8 @@ namespace uba
 				info.proxyPort = proxyPort;
 
 				writer.WriteGuid(m_uid);
+				writer.WriteByte(m_casCompressor);
+				writer.WriteByte(m_casCompressionLevel);
 				return true;
 			}
 
@@ -387,52 +390,59 @@ namespace uba
 
 				CasEntry* casEntry = nullptr;
 				bool has = HasCasFile(casKey, &casEntry); // HasCasFile also writes deferred cas entries if in queue
-				if (!has && !EnsureCasFile(casKey, nullptr))
+				if (!has)
 				{
-					// Last resort.. use hint to load file into cas (hint should be renamed since it is now a critical parameter)
-					// We better check the caskey first to make sure it is matching on the server
-
-					CasKey checkedCasKey;
+					if (!EnsureCasFile(casKey, nullptr) && m_allowFallback)
 					{
-						StringKey fileNameKey = CaseInsensitiveFs ? ToStringKeyLower(hint) : ToStringKey(hint);
-						SCOPED_READ_LOCK(m_fileTableLookupLock, lookupLock);
-						auto findIt = m_fileTableLookup.find(fileNameKey);
-						if (findIt != m_fileTableLookup.end())
+						// Last resort.. use hint to load file into cas (hint should be renamed since it is now a critical parameter)
+						// We better check the caskey first to make sure it is matching on the server
+
+						CasKey checkedCasKey;
 						{
-							FileEntry& fileEntry = findIt->second;
-							lookupLock.Leave();
-							SCOPED_READ_LOCK(fileEntry.lock, entryLock);
-							if (fileEntry.verified)
-								checkedCasKey = fileEntry.casKey;
+							StringKey fileNameKey = CaseInsensitiveFs ? ToStringKeyLower(hint) : ToStringKey(hint);
+							SCOPED_READ_LOCK(m_fileTableLookupLock, lookupLock);
+							auto findIt = m_fileTableLookup.find(fileNameKey);
+							if (findIt != m_fileTableLookup.end())
+							{
+								FileEntry& fileEntry = findIt->second;
+								lookupLock.Leave();
+								SCOPED_READ_LOCK(fileEntry.lock, entryLock);
+								if (fileEntry.verified)
+									checkedCasKey = fileEntry.casKey;
+							}
 						}
-					}
-					if (checkedCasKey == CasKeyZero)
-					{
-						m_logger.Info(TC("Server did not find cas for %s in file table lookup. Recalculating cas key"), hint.data);
-						if (!CalculateCasKey(checkedCasKey, hint.data))
+						if (checkedCasKey == CasKeyZero)
 						{
-							m_logger.Error(TC("FetchBegin failed for cas file %s (%s) requested by %s. Can't calculate cas key for file"), CasKeyString(casKey).str, hint.data, GuidToString(connectionInfo.GetUid()).str);
+							m_logger.Info(TC("Server did not find cas for %s in file table lookup. Recalculating cas key"), hint.data);
+							if (!CalculateCasKey(checkedCasKey, hint.data))
+							{
+								m_logger.Error(TC("FetchBegin failed for cas file %s (%s) requested by %s. Can't calculate cas key for file"), CasKeyString(casKey).str, hint.data, GuidToString(connectionInfo.GetUid()).str);
+								writer.WriteU16(0);
+								return false;
+							}
+						}
+
+						if (checkedCasKey != casKey)
+						{
+							m_logger.Error(TC("FetchBegin failed for cas file %s (%s). Server has a source file"), CasKeyString(casKey).str, hint.data);
 							writer.WriteU16(0);
 							return false;
 						}
-					}
 
-					if (checkedCasKey != casKey)
-					{
-						m_logger.Error(TC("FetchBegin failed for cas file %s (%s). Server has a source file"), CasKeyString(casKey).str, hint.data);
-						writer.WriteU16(0);
-						return false;
+						if (!AddCasFile(hint.data, casKey, false))
+						{
+							m_logger.Error(TC("FetchBegin failed for cas file %s (%s). Can't add cas file to database"), CasKeyString(casKey).str, hint.data);
+							writer.WriteU16(0);
+							return true;
+						}
 					}
-
-					if (!AddCasFile(hint.data, casKey, false))
+					SCOPED_WRITE_LOCK(m_casLookupLock, lookupLock);
+					auto findIt = m_casLookup.find(casKey);
+					if (findIt == m_casLookup.end())
 					{
-						m_logger.Error(TC("FetchBegin failed for cas file %s (%s). Can't add cas file to database"), CasKeyString(casKey).str, hint.data);
 						writer.WriteU16(0);
 						return true;
 					}
-					SCOPED_WRITE_LOCK(m_casLookupLock, lookupLock);
-					auto findIt = m_casLookup.find(casKey);
-					UBA_ASSERT(findIt != m_casLookup.end());
 					casEntry = &findIt->second;
 				}
 
@@ -528,7 +538,7 @@ namespace uba
 				}
 				else
 				{
-					memoryBegin = PopBufferSlot();
+					memoryBegin = m_bufferSlots.Pop();
 					memoryPos = memoryBegin;
 					u32 toRead = u32(Min(left, BufferSlotSize));
  					if (!ReadFile(m_logger, casFile.data, readFileHandle, memoryBegin, toRead))
diff --git a/Engine/Source/Programs/UnrealBuildAccelerator/Common/Private/UbaStorageUtils.cpp b/Engine/Source/Programs/UnrealBuildAccelerator/Common/Private/UbaStorageUtils.cpp
new file mode 100644
--- /dev/null
+++ b/Engine/Source/Programs/UnrealBuildAccelerator/Common/Private/UbaStorageUtils.cpp
@@ -0,0 +1,612 @@
+// Copyright Epic Games, Inc. All Rights Reserved.
+
+#include "UbaStorageUtils.h"
+#include "UbaStorage.h"
+#include "UbaFileAccessor.h"
+#include "UbaStats.h"
+#include <oodle2.h>
+
+namespace uba
+{
+	CasKey CalculateCasKey(u8* fileMem, u64 fileSize, bool storeCompressed, WorkManager* workManager)
+	{
+		constexpr u32 MaxWorkItemsPerAction = 128; // Cap this to not starve other things
+		CasKeyHasher hasher;
+
+		if (fileSize == 0)
+			return ToCasKey(hasher, storeCompressed);
+
+		#ifndef __clang_analyzer__
+
+		if (fileSize > BufferSlotSize) // Note that when filesize is larger than BufferSlotSize the hash becomes a hash of hashes
+		{
+			struct WorkRec
+			{
+				Atomic<u64> refCount;
+				Atomic<u64> counter;
+				Atomic<u64> doneCounter;
+				u8* fileMem = nullptr;
+				u64 workCount = 0;
+				u64 fileSize = 0;
+				bool error = false;
+				Vector<CasKey> keys;
+				Event done;
+			};
+
+			u32 workCount = u32((fileSize + BufferSlotSize - 1) / BufferSlotSize);
+
+			WorkRec* rec = new WorkRec();
+			rec->fileMem = fileMem;
+			rec->workCount = workCount;
+			rec->fileSize = fileSize;
+			rec->keys.resize(workCount);
+			rec->done.Create(true);
+			rec->refCount = 2;
+
+			auto work = [rec]()
+			{
+				while (true)
+				{
+					u64 index = rec->counter++;
+					if (index >= rec->workCount)
+					{
+						if (!--rec->refCount)
+							delete rec;
+						return 0;
+					}
+
+					u64 startOffset = BufferSlotSize*index;
+					u64 toRead = Min(BufferSlotSize, rec->fileSize - startOffset);
+					u8* slot = rec->fileMem + startOffset;
+					CasKeyHasher hasher;
+					hasher.Update(slot, toRead);
+					rec->keys[index] = ToCasKey(hasher, false);
+
+					if (++rec->doneCounter == rec->workCount)
+						rec->done.Set();
+				}
+				return 0;
+			};
+
+			u32 workerCount = 0;
+			if (workManager)
+			{
+				workerCount = Min(workCount, workManager->GetWorkerCount()-1); // We are a worker ourselves
+				workerCount = Min(workerCount, MaxWorkItemsPerAction); // Cap this to not starve other things
+				rec->refCount += workerCount;
+				workManager->AddWork(work, workerCount, TC("CalculateKey"));
+			}
+
+			work();
+			rec->done.IsSet();
+
+			hasher.Update(rec->keys.data(), rec->keys.size()*sizeof(CasKey));
+
+			bool error = rec->error;
+
+			if (!--rec->refCount)
+				delete rec;
+
+			if (error)
+				return CasKeyZero;
+		}
+		else
+		{
+			hasher.Update(fileMem, fileSize);
+		}
+
+		#endif // __clang_analyzer__
+
+		return ToCasKey(hasher, storeCompressed);
+	}
+
+	bool SendBatchMessages(Logger& logger, NetworkClient& client, u16 fetchId, u8* slot, u64 capacity, u64 left, u32 messageMaxSize, u32& readIndex, u32& responseSize)
+	{
+		responseSize = 0;
+
+		struct Entry
+		{
+			Entry(u8* slot, u32 i, u32 messageMaxSize) : reader(slot + i * messageMaxSize, 0, SendMaxSize), done(true) {}
+			NetworkMessage message;
+			BinaryReader reader;
+			Event done;
+		};
+
+		u64 sendCountCapacity = capacity / messageMaxSize;
+		u64 sendCount = left/messageMaxSize;
+
+		if (sendCount > sendCountCapacity)
+			sendCount = sendCountCapacity;
+		else if (sendCount < sendCountCapacity && (left - sendCount * messageMaxSize) > 0)
+			++sendCount;
+
+
+		u64 entriesMem[sizeof(Entry) * 8];
+		UBA_ASSERT(sizeof(Entry)*sendCount <= sizeof(entriesMem));
+
+		Entry* entries = (Entry*)entriesMem;
+
+		bool success = true;
+		u32 inFlightCount = u32(sendCount);
+		for (u32 i=0; i!=sendCount; ++i)
+		{
+			auto& entry = *new (entries + i) Entry(slot, i, messageMaxSize);
+			StackBinaryWriter<32> writer;
+			entry.message.Init(client, StorageServiceId, StorageMessageType_FetchSegment, writer);
+			writer.WriteU16(fetchId);
+			writer.WriteU32(readIndex + i + 1);
+			if (entry.message.SendAsync(entry.reader, [](bool error, void* userData) { ((Event*)userData)->Set(); }, &entry.done))
+				continue;
+			entry.~Entry();
+			inFlightCount = i;
+			success = false;
+			break;
+		}
+
+		for (u32 i=0; i!=inFlightCount; ++i)
+		{
+			Entry& entry = entries[i];
+			if (!entry.done.IsSet(5*60*1000))
+				logger.Error(TC("SendBatchMessages timed out getting async message response"));
+			if (!entry.message.ProcessAsyncResults(entry.reader))
+				success = false;
+			else
+				responseSize += u32(entry.reader.GetLeft());
+		}
+
+		for (u32 i=0; i!=inFlightCount; ++i)
+			entries[i].~Entry();
+
+		readIndex += u32(sendCount);
+		return success;
+	}
+
+	bool SendFile(Logger& logger, NetworkClient& client, const CasKey& casKey, const u8* sourceMem, u64 sourceSize, const tchar* hint)
+	{
+		UBA_ASSERT(casKey != CasKeyZero);
+
+		const u8* readData = sourceMem;
+		u64 fileSize = sourceSize;
+
+		u16 storeId = 0;
+		bool isFirst = true;
+		bool sendEnd = false;
+		u64 sendLeft = fileSize;
+		u64 sendPos = 0;
+
+		while (sendLeft)
+		{
+			StackBinaryWriter<SendMaxSize> writer;
+			NetworkMessage msg(client, StorageServiceId, isFirst ? StorageMessageType_StoreBegin : StorageMessageType_StoreSegment, writer);
+			if (isFirst)
+			{
+				writer.WriteCasKey(casKey);
+				writer.WriteU64(fileSize);
+				writer.WriteU64(sourceSize);
+				writer.WriteString(hint);
+			}
+			else
+			{
+				UBA_ASSERT(storeId != 0);
+				writer.WriteU16(storeId);
+				writer.WriteU64(sendPos);
+			}
+
+			u64 capacityLeft = writer.GetCapacityLeft();
+			u64 toWrite = Min(sendLeft, capacityLeft);
+			writer.WriteBytes(readData, toWrite);
+
+			readData += toWrite;
+			sendLeft -= toWrite;
+			sendPos += toWrite;
+
+			bool isDone = sendLeft == 0;
+
+			if (isFirst) // First message must always be acknowledged (provide a reader) to make sure there is an entry on server that can be waited on.
+			{
+				StackBinaryReader<128> reader;
+				if (!msg.Send(reader))
+					return false;
+				storeId = reader.ReadU16();
+				sendEnd = reader.ReadBool();
+				if (isDone)
+					break;
+
+				if (!storeId) // Zero means error
+					return logger.Error(TC("Server failed to start storing file %s (%s)"), CasKeyString(casKey).str, hint);
+
+				if (storeId == u16(~0)) // File already exists on server
+				{
+					logger.Info(TC("Server already has file %s (%s)"), CasKeyString(casKey).str, hint);
+					return true;
+				}
+
+				isFirst = false;
+			}
+			else
+			{
+				if (!msg.Send())
+					return false;
+				if (isDone)
+					break;
+			}
+		}
+
+		if (sendEnd)
+		{
+			StackBinaryWriter<128> writer;
+			NetworkMessage msg(client, StorageServiceId, StorageMessageType_StoreEnd, writer);
+			writer.WriteCasKey(casKey);
+			if (!msg.Send())
+				return false;
+		}
+
+		return true;
+	}
+
+	bool FileSender::SendFileCompressed(const CasKey& casKey, const tchar* fileName, const u8* sourceMem, u64 sourceSize, const tchar* hint)
+	{
+		UBA_ASSERT(casKey != CasKeyZero);
+
+		NetworkClient& client = m_client;
+
+		TimerScope ts(m_stats.sendCas);
+
+		u64 firstMessageOverHead = (sizeof(CasKey) + sizeof(u64)*2 + GetStringWriteSize(hint, TStrlen(hint)));
+
+		u64 messageHeader = client.GetMessageHeaderSize();
+		u64 messageHeaderMaxSize = messageHeader + firstMessageOverHead;
+
+		MemoryBlock memoryBlock(sourceSize + messageHeaderMaxSize + 1024);
+		{
+			const u8* uncompressedData = sourceMem;
+			u8* compressBufferStart = memoryBlock.memory + messageHeaderMaxSize;
+			u8* compressBuffer = compressBufferStart;
+			u64 totalWritten = messageHeaderMaxSize; // Make sure there is room for msg header in the memory since we are using it to send
+			u64 left = sourceSize;
+
+			compressBuffer += 8;
+			totalWritten += 8;
+			memoryBlock.Allocate(totalWritten, 1, hint);
+
+			u64 diff = u64(OodleLZ_GetCompressedBufferSizeNeeded((OodleLZ_Compressor)m_casCompressor, BufferSlotHalfSize)) - BufferSlotHalfSize;
+			u64 maxUncompressedBlock = BufferSlotHalfSize - diff - totalWritten - 8; // 8 bytes block header
+
+			OodleLZ_CompressOptions oodleOptions = *OodleLZ_CompressOptions_GetDefault();
+			while (left)
+			{
+				u32 uncompressedBlockSize = (u32)Min(left, maxUncompressedBlock);
+
+				u64 reserveSize = totalWritten + uncompressedBlockSize + diff + 8;
+				if (reserveSize > memoryBlock.mappedSize)
+				{
+					u64 toAllocate = reserveSize - memoryBlock.writtenSize;
+					memoryBlock.Allocate(toAllocate, 1, hint);
+				}
+
+				u8* destBuf = compressBuffer;
+				u32 compressedBlockSize;
+				{
+					TimerScope cts(m_stats.compressSend);
+					compressedBlockSize = (u32)OodleLZ_Compress((OodleLZ_Compressor)m_casCompressor, uncompressedData, (int)uncompressedBlockSize, destBuf + 8, (OodleLZ_CompressionLevel)m_casCompressionLevel, &oodleOptions);
+					if (compressedBlockSize == OODLELZ_FAILED)
+						return m_logger.Error(TC("Failed to compress %u bytes at %llu for %s (%s) (%s) (uncompressed size: %llu)"), uncompressedBlockSize, totalWritten, fileName, CasKeyString(casKey).str, hint, sourceSize);
+				}
+
+				*(u32*)destBuf =  u32(compressedBlockSize);
+				*(u32*)(destBuf+4) =  u32(uncompressedBlockSize);
+
+				u32 writeBytes = u32(compressedBlockSize) + 8;
+
+				totalWritten += writeBytes;
+				memoryBlock.writtenSize = totalWritten;
+
+				left -= uncompressedBlockSize;
+				uncompressedData += uncompressedBlockSize;
+				compressBuffer += writeBytes;
+			}
+
+			*(u64*)compressBufferStart = sourceSize;
+		}
+
+
+		u8* readData = memoryBlock.memory + messageHeaderMaxSize;
+		u64 fileSize = memoryBlock.writtenSize - messageHeaderMaxSize;
+
+		u16 storeId = 0;
+		bool isFirst = true;
+		bool sendEnd = false;
+		u64 sendLeft = fileSize;
+		u64 sendPos = 0;
+
+		bool hasSendOneAtTheTimeLock = false;
+		auto lockGuard = MakeGuard([&]() { if (hasSendOneAtTheTimeLock) m_sendOneAtTheTimeLock.LeaveWrite(); });
+
+		while (sendLeft)
+		{
+			u64 writerStartOffset = messageHeader + (isFirst ? firstMessageOverHead : (sizeof(u16) + sizeof(u64)));
+			BinaryWriter writer(readData + sendPos - writerStartOffset, 0, client.GetMessageMaxSize());
+			NetworkMessage msg(client, StorageServiceId, isFirst ? StorageMessageType_StoreBegin : StorageMessageType_StoreSegment, writer);
+			if (isFirst)
+			{
+				writer.WriteCasKey(casKey);
+				writer.WriteU64(fileSize);
+				writer.WriteU64(sourceSize);
+				writer.WriteString(hint);
+			}
+			else
+			{
+				UBA_ASSERT(storeId != 0);
+				writer.WriteU16(storeId);
+				writer.WriteU64(sendPos);
+			}
+
+			u64 capacityLeft = writer.GetCapacityLeft();
+			u64 toWrite = Min(sendLeft, capacityLeft);
+			writer.AllocWrite(toWrite);
+
+			sendLeft -= toWrite;
+			sendPos += toWrite;
+
+			bool isDone = sendLeft == 0;
+
+			if (isFirst && !isDone)
+			{
+				m_sendOneAtTheTimeLock.EnterWrite();
+				hasSendOneAtTheTimeLock = true;
+			}
+
+			if (isFirst) // First message must always be acknowledged (provide a reader) to make sure there is an entry on server that can be waited on.
+			{
+				StackBinaryReader<128> reader;
+				if (!msg.Send(reader))
+					return false;
+				storeId = reader.ReadU16();
+				sendEnd = reader.ReadBool();
+				if (isDone)
+					break;
+
+				if (!storeId) // Zero means error
+					return m_logger.Error(TC("Server failed to start storing file %s (%s)"), CasKeyString(casKey).str, hint);
+
+				if (storeId == u16(~0)) // File already exists on server
+				{
+					m_logger.Info(TC("Server already has file %s (%s)"), CasKeyString(casKey).str, hint);
+					return true;
+				}
+
+				isFirst = false;
+			}
+			else
+			{
+				if (!msg.Send())
+					return false;
+				if (isDone)
+					break;
+			}
+		}
+
+		if (sendEnd)
+		{
+			StackBinaryWriter<128> writer;
+			NetworkMessage msg(client, StorageServiceId, StorageMessageType_StoreEnd, writer);
+			writer.WriteCasKey(casKey);
+			if (!msg.Send())
+				return false;
+		}
+
+		m_stats.sendCasBytesRaw += sourceSize;
+		m_stats.sendCasBytesComp += fileSize;
+
+		return true;
+	}
+
+	bool FileFetcher::RetrieveFile(Logger& logger, NetworkClient& client, const CasKey& casKey, const tchar* destination, MemoryBlock* destinationMem)
+	{
+		u8* slot = m_bufferSlots.Pop();
+		auto sg = MakeGuard([&](){ m_bufferSlots.Push(slot); });
+
+		u64 fileSize = 0;
+		u64 actualSize = 0;
+		u64 sizeOnDisk = 0;
+
+		u8* readBuffer = nullptr;
+		u8* readPosition = nullptr;
+
+		u16 fetchId = 0;
+		u32 responseSize = 0;
+		bool isCompressed = false;
+		bool sendEnd = false;
+		u64 left = ~u64(0);
+
+		u32 sizeOfFirstMessage = 0;
+
+		{
+			StackBinaryWriter<1024> writer;
+			NetworkMessage msg(client, StorageServiceId, StorageMessageType_FetchBegin, writer);
+			writer.WriteBool(false);
+			writer.WriteCasKey(casKey);
+			writer.WriteString(destination);
+			BinaryReader reader(slot, 0, SendMaxSize);
+			if (!msg.Send(reader))
+				return logger.Error(TC("Failed to send fetch begin message for cas %s (%s)"), CasKeyString(casKey).str, destination);
+			sizeOfFirstMessage = u32(reader.GetLeft());
+			fetchId = reader.ReadU16();
+			if (fetchId == 0)
+			{
+				logger.Logf(m_errorOnFail ? LogEntryType_Error : LogEntryType_Detail, TC("Failed to fetch cas %s (%s)"), CasKeyString(casKey).str, destination);
+				return false;
+			}
+
+			fileSize = reader.Read7BitEncoded();
+
+			u8 flags = reader.ReadByte();
+
+			isCompressed = (flags >> 0) & 1;
+			sendEnd = (flags >> 1) & 1;
+
+			left = fileSize;
+
+			responseSize = u32(reader.GetLeft());
+			readBuffer = (u8*)reader.GetPositionData();
+			readPosition = readBuffer;
+
+			actualSize = fileSize;
+			if (isCompressed)
+				actualSize = *(u64*)readBuffer;
+		}
+
+		sizeOnDisk = IsCompressed(casKey) ? fileSize : actualSize;
+
+		FileAccessor destinationFile(logger, destination);
+		if (!destinationMem)
+			if (!destinationFile.CreateWrite(false, DefaultAttributes(), sizeOnDisk, m_tempPath.data))
+				return false;
+
+		//u8* writePos = writeMem;
+
+		// This is here just to prevent server from getting a million messages at the same time.
+		// In theory we could have 10 clients with 48 processes each where each one of the processes asks for a large file (64 messages in flight)
+		// So worst case in that scenario would be 10*48*64 = 30000 messages.
+		bool oneAtTheTime = false;//left > client->GetMessageMaxSize() * 2;
+		if (oneAtTheTime)
+			m_retrieveOneBatchAtTheTimeLock.EnterWrite();
+		auto oatg = MakeGuard([&]() { if (oneAtTheTime) m_retrieveOneBatchAtTheTimeLock.LeaveWrite(); });
+
+
+		bool sendSegmentMessage = responseSize == 0;
+		u64 leftUncompressed = actualSize;
+		readBuffer += sizeof(u64); // Size is stored first
+		u64 maxReadSize = BufferSlotHalfSize - sizeof(u64);
+
+		if (actualSize)
+		{
+			u64 leftCompressed = fileSize - responseSize;
+			u32 readIndex = 0;
+			do
+			{
+				// First read in a full decompressable block
+				bool isFirstInBlock = true;
+				u32 compressedSize = ~u32(0);
+				u32 uncompressedSize = ~u32(0);
+				left = 0;
+				u32 overflow = 0;
+				do
+				{
+					if (sendSegmentMessage)
+					{
+						if (fetchId == u16(~0))
+							return logger.Error(TC("Cas content error (2). Server believes %s was only one segment but client sees more. UncompressedSize: %llu LeftUncompressed: %llu Size: %llu Left to read: %llu ResponseSize: %u. (%s)"), destination, actualSize, leftUncompressed, fileSize, left, responseSize, CasKeyString(casKey).str);
+						if (!SendBatchMessages(logger, client, fetchId, readPosition, maxReadSize - u32(readPosition - readBuffer), leftCompressed, sizeOfFirstMessage, readIndex, responseSize))
+							return logger.Error(TC("Failed to send batched messages to server (%s)"), CasKeyString(casKey).str);
+						leftCompressed -= responseSize;
+					}
+					else
+					{
+						sendSegmentMessage = true;
+					}
+
+					if (isFirstInBlock)
+					{
+						if (readPosition - readBuffer < sizeof(u32) * 2)
+							return logger.Error(TC("Received less than minimum amount of data. Most likely corrupt cas file %s (Available: %u UncompressedSize: %llu LeftUncompressed: %llu)"), CasKeyString(casKey).str, u32(readPosition - readBuffer), actualSize, leftUncompressed);
+						isFirstInBlock = false;
+						u32* blockSize = (u32*)readBuffer;
+						compressedSize = blockSize[0];
+						uncompressedSize = blockSize[1];
+						readBuffer += sizeof(u32) * 2;
+						maxReadSize = BufferSlotHalfSize - sizeof(u32) * 2;
+						u32 read = (responseSize + u32(readPosition - readBuffer));
+						//UBA_ASSERTF(read <= compressedSize, TC("Error in datastream fetching cas. Read size: %u CompressedSize: %u %s (%s)"), read, compressedSize, CasKeyString(casKey).str, destination);
+						if (read > compressedSize)
+						{
+							//UBA_ASSERT(!responseSize); // TODO: This has not really been tested
+							left = 0;
+							overflow = read - compressedSize;
+							sendSegmentMessage = false;
+						}
+						else
+						{
+							left = compressedSize - read;
+						}
+						readPosition += responseSize;
+					}
+					else
+					{
+						readPosition += responseSize;
+						if (responseSize > left)
+						{
+							overflow = responseSize - u32(left);
+							UBA_ASSERTF(overflow < BufferSlotHalfSize, TC("Something went wrong. Overflow: %u responseSize: %u, left: %u"), overflow, responseSize, left);
+							if (overflow >= 8)
+							{
+								responseSize = 0;
+								sendSegmentMessage = false;
+							}
+							left = 0;
+						}
+						else
+						{
+							left -= responseSize;
+						}
+					}
+				} while (left);
+
+				// Then decompress
+				{
+					u8* decompressBuffer = slot + BufferSlotHalfSize;
+
+					//TimerScope ts(stats.decompressRecv);
+					OO_SINTa decompLen = OodleLZ_Decompress(readBuffer, int(compressedSize), decompressBuffer, int(uncompressedSize));
+					if (decompLen != uncompressedSize)
+						return logger.Error(TC("Expected %u but got %i when decompressing %u bytes for file %s"), uncompressedSize, int(decompLen), compressedSize, destination);
+
+					if (!destinationMem)
+					{
+						if (!destinationFile.Write(decompressBuffer, uncompressedSize, actualSize - leftUncompressed))
+							return false;
+					}
+					else
+					{
+						void* mem = destinationMem->Allocate(uncompressedSize, 1, TC(""));
+						memcpy(mem, decompressBuffer, uncompressedSize);
+					}
+					leftUncompressed -= uncompressedSize;
+				}
+
+				// Move overflow back to the beginning of the buffer and start the next block (if there is one)
+				readBuffer = slot;
+				maxReadSize = BufferSlotHalfSize;
+				UBA_ASSERTF(readPosition - overflow >= readBuffer, TC("ReadPosition - overflow is before beginning of buffer (overflow: %u) for file %s"), overflow, destination);
+				UBA_ASSERTF(readPosition <= readBuffer + BufferSlotHalfSize, TC("ReadPosition is outside readBuffer size (pos: %llu, overflow: %u) for file %s"), readPosition - readBuffer, overflow, destination);
+				memmove(readBuffer, readPosition - overflow, overflow);
+				readPosition = readBuffer + overflow;
+				if (overflow)
+					responseSize = 0;
+			} while (leftUncompressed);
+		}
+
+		if (sendEnd)
+		{
+			StackBinaryWriter<128> writer;
+			NetworkMessage msg(client, StorageServiceId, StorageMessageType_FetchEnd, writer);
+			writer.WriteCasKey(casKey);
+			if (!msg.Send())
+				return false;
+		}
+
+		if (!destinationMem)
+			if (!destinationFile.Close(&m_lastWritten))
+				return false;
+
+		m_size = actualSize;
+
+		//stats.recvCasBytesRaw += actualSize;
+		//stats.recvCasBytesComp += fileSize;
+		//out.size = actualSize;
+
+		return true;
+	}
+
+
+}
diff --git a/Engine/Source/Programs/UnrealBuildAccelerator/Common/Public/UbaCacheClient.h b/Engine/Source/Programs/UnrealBuildAccelerator/Common/Public/UbaCacheClient.h
new file mode 100644
--- /dev/null
+++ b/Engine/Source/Programs/UnrealBuildAccelerator/Common/Public/UbaCacheClient.h
@@ -0,0 +1,74 @@
+// Copyright Epic Games, Inc. All Rights Reserved.
+
+#pragma once
+
+#include "UbaCompactTables.h"
+#include "UbaLogger.h"
+
+namespace uba
+{
+	class NetworkClient;
+	class ProcessHandle;
+	class Session;
+	class StorageImpl;
+	struct ProcessStartInfo;
+
+	class CacheClient
+	{
+	public:
+		CacheClient(LogWriter& writer, StorageImpl& storage, NetworkClient& client, Session& session);
+		~CacheClient();
+
+		bool RegisterRoot(const tchar* rootPath, bool includeInKey = true);
+		bool RegisterSystemRoots();
+
+		bool WriteToCache(const ProcessHandle& process);
+		bool FetchFromCache(const ProcessStartInfo& info);
+
+		bool WriteCacheSummary(const tchar* destinationFile, const tchar* filterString = nullptr);
+
+		inline NetworkClient& GetClient() { return m_client; }
+		inline StorageImpl& GetStorage() { return m_storage; }
+
+	private:
+		bool SendPathTable(u32 requiredPathTableSize);
+		bool SendCasTable(u32 requiredCasTableSize);
+		bool SendCacheEntry(const CasKey& cmdKey, const Map<u32, u32>& inputsStringToCasKey, const Map<u32, u32>& outputsStringToCasKey);
+		bool FetchCasTable();
+
+		CasKey GetCmdKey(const ProcessStartInfo& info);
+		u32 FindRootIndex(const StringBufferBase& path);
+
+		template<typename CharType, typename Func>
+		bool NormalizeString(const CharType* str, u64 strLen, const Func& func, const tchar* hint);
+
+		CasKey NormalizeAndHashFile(const tchar* rspFile);
+		bool GetLocalPathAndCasKey(StringBufferBase& outPath, CasKey& outKey, CompactCasKeyTable& casKeyTable, CompactPathTable& pathTable, u32 offset);
+
+		MutableLogger m_logger;
+		StorageImpl& m_storage;
+		NetworkClient& m_client;
+		Session& m_session;
+
+		Atomic<bool> m_connected;
+
+		struct Root;
+		Vector<Root> m_roots;
+		u32 m_shortestRoot = 0;
+		u32 m_longestRoot = 0;
+
+		CompactPathTable m_serverPathTable;
+		CompactCasKeyTable m_serverCasKeyTable;
+
+		CompactPathTable m_sendPathTable;
+		CompactCasKeyTable m_sendCasKeyTable;
+
+		ReaderWriterLock m_pathTableNetworkLock;
+		u32 m_pathTableSizeSent = 0;
+
+		ReaderWriterLock m_casKeyTableNetworkLock;
+		u32 m_casKeyTableSizeSent = 0;
+
+		ReaderWriterLock m_sendOneAtTheTimeLock;
+	};
+}
\ No newline at end of file
diff --git a/Engine/Source/Programs/UnrealBuildAccelerator/Common/Public/UbaCacheServer.h b/Engine/Source/Programs/UnrealBuildAccelerator/Common/Public/UbaCacheServer.h
new file mode 100644
--- /dev/null
+++ b/Engine/Source/Programs/UnrealBuildAccelerator/Common/Public/UbaCacheServer.h
@@ -0,0 +1,71 @@
+// Copyright Epic Games, Inc. All Rights Reserved.
+
+#pragma once
+
+#include "UbaCompactTables.h"
+#include "UbaHash.h"
+#include "UbaLogger.h"
+
+namespace uba
+{
+	class NetworkServer;
+	class StorageServer;
+	struct BinaryReader;
+	struct BinaryWriter;
+	struct ConnectionInfo;
+
+	class CacheServer
+	{
+	public:
+		CacheServer(LogWriter& writer, const tchar* rootDir, NetworkServer& server, StorageServer& storage);
+		~CacheServer();
+
+		bool Load();
+		bool Save();
+
+		bool RunMaintenance(bool force);
+
+	private:
+		bool SaveNoLock();
+		void OnDisconnected(u32 clientId);
+
+		struct Connection;
+		bool HandleMessage(const ConnectionInfo& connectionInfo, u8 messageType, BinaryReader& reader, BinaryWriter& writer);
+		bool HandleStoreEntry(Connection& connection, BinaryReader& reader, BinaryWriter& writer);
+		bool HandleFetchPathTable(BinaryReader& reader, BinaryWriter& writer);
+		bool HandleFetchCasTable(BinaryReader& reader, BinaryWriter& writer);
+		bool HandleFetchEntries(BinaryReader& reader, BinaryWriter& writer);
+		bool HandleCreateStatusFile(BinaryReader& reader, BinaryWriter& writer);
+
+		MutableLogger m_logger;
+		NetworkServer& m_server;
+		StorageServer& m_storage;
+
+		StringBuffer<MaxPath> m_rootDir;
+
+		ReaderWriterLock m_maintenanceLock;
+		Atomic<u32> m_addsSinceMaintenance;
+
+		struct CacheEntry
+		{
+			u64 creationTime;
+			Vector<u8> inputCasKeyOffsets;
+			Vector<u8> outputCasKeyOffsets;
+		};
+
+		struct CacheEntries
+		{
+			ReaderWriterLock lock;
+			List<CacheEntry> entries;
+		};
+
+		ReaderWriterLock m_cacheEntryLookupLock;
+		UnorderedMap<CasKey, CacheEntries> m_cacheEntryLookup;
+
+		CompactPathTable m_pathTable;
+		CompactCasKeyTable m_casKeyTable;
+
+		ReaderWriterLock m_connectionsLock;
+		Map<u32, Connection> m_connections;
+	};
+}
\ No newline at end of file
diff --git a/Engine/Source/Programs/UnrealBuildAccelerator/Common/Public/UbaCompactTables.h b/Engine/Source/Programs/UnrealBuildAccelerator/Common/Public/UbaCompactTables.h
new file mode 100644
--- /dev/null
+++ b/Engine/Source/Programs/UnrealBuildAccelerator/Common/Public/UbaCompactTables.h
@@ -0,0 +1,66 @@
+// Copyright Epic Games, Inc. All Rights Reserved.
+
+#pragma once
+
+#include "UbaHash.h"
+#include "UbaMemory.h"
+
+namespace uba
+{
+	struct BinaryReader;
+	struct BinaryWriter;
+	struct StringKey;
+
+	class CompactPathTable
+	{
+	public:
+		CompactPathTable(u64 reserveSize, u64 reserveOffsetsCount = 0);
+
+		u32 Add(const tchar* str, u32 strLen, u32* outRequiredCasTableSize = nullptr);
+
+		void GetString(StringBufferBase& out, u64 offset) const;
+
+		u8* GetMemory();
+		u32 GetSize();
+
+		void ReadMem(BinaryReader& reader, bool populateLookup);
+		void Swap(CompactPathTable& other);
+
+	private:
+		u32 InternalAdd(const tchar* str, u32 strLen);
+		ReaderWriterLock m_lock;
+		MemoryBlock m_mem;
+		UnorderedMap<StringKey, u32> m_offsets;
+		u64 m_reserveSize;
+	};
+
+	class CompactCasKeyTable
+	{
+	public:
+		CompactCasKeyTable(u64 reserveSize, u64 reserveOffsetsCount = 0);
+
+		u32 Add(const CasKey& casKey, u64 stringOffset, u32* outRequiredCasTableSize = nullptr);
+
+		void GetKey(CasKey& outKey, u64 offset) const;
+		void GetPathAndKey(StringBufferBase& outPath, CasKey& outKey, const CompactPathTable& pathTable, u64 offset) const;
+
+		u8* GetMemory();
+		u32 GetSize();
+		ReaderWriterLock& GetLock() { return m_lock; }
+
+		void ReadMem(BinaryReader& reader, bool populateLookup);
+		void Swap(CompactCasKeyTable& other);
+
+	private:
+		ReaderWriterLock m_lock;
+		MemoryBlock m_mem;
+		struct Key
+		{
+			bool operator==(const Key& o) const { return ck == o.ck && offset == o.offset; }
+			CasKey ck; u32 offset;
+		};
+		struct KeyHash { size_t operator()(const Key& k) const { return k.ck.a ^ k.ck.b ^ k.ck.c ^ k.offset; } };
+		UnorderedMap<Key, u32, KeyHash> m_offsets;
+		u64 m_reserveSize;
+	};
+}
diff --git a/Engine/Source/Programs/UnrealBuildAccelerator/Common/Public/UbaNetwork.h b/Engine/Source/Programs/UnrealBuildAccelerator/Common/Public/UbaNetwork.h
--- a/Engine/Source/Programs/UnrealBuildAccelerator/Common/Public/UbaNetwork.h
+++ b/Engine/Source/Programs/UnrealBuildAccelerator/Common/Public/UbaNetwork.h
@@ -9,17 +9,22 @@ namespace uba
 	// Default constants
 	static constexpr u16 DefaultPort = 1345;
 	static constexpr u16 DefaultStorageProxyPort = DefaultPort + 1;
+	static constexpr u16 DefaultCachePort = DefaultPort + 2;
 	static constexpr u32 SendMaxSize = 256*1024;
 	static constexpr u32 SendDefaultSize = 256*1024;
 
 	static constexpr u8 SystemServiceId = 0;
 	static constexpr u8 StorageServiceId = 1;
 	static constexpr u8 SessionServiceId = 2;
+	static constexpr u8 CacheServiceId = 3;
 	
 	static constexpr u32 SystemNetworkVersion = 1339;
-	static constexpr u32 StorageNetworkVersion = 3;
+	static constexpr u32 StorageNetworkVersion = 4;
 	static constexpr u32 SessionNetworkVersion = 32;
+	static constexpr u32 CacheNetworkVersion = 1;
 
+	static constexpr u32 CachePathTableMaxSize = 32*1024*1024;
+	static constexpr u32 CacheCasKeyTableMaxSize = 32*1024*1024;
 
 	// Messages used over network between client/server (system, storage and session)
 
@@ -100,5 +105,23 @@ namespace uba
 		SessionProcessAvailableResponse_RemoteExecutionDisabled = ~u32(0) - 1,
 	};
 
+	#define UBA_CACHE_MESSAGES \
+		UBA_CACHE_MESSAGE(Connect) \
+		UBA_CACHE_MESSAGE(StorePathTable) \
+		UBA_CACHE_MESSAGE(StoreCasTable) \
+		UBA_CACHE_MESSAGE(StoreEntry) \
+		UBA_CACHE_MESSAGE(StoreEntryDone) \
+		UBA_CACHE_MESSAGE(FetchPathTable) \
+		UBA_CACHE_MESSAGE(FetchCasTable) \
+		UBA_CACHE_MESSAGE(FetchEntries) \
+		UBA_CACHE_MESSAGE(CreateStatusFile) \
+
+	enum CacheMessageType : u8
+	{
+		#define UBA_CACHE_MESSAGE(x) CacheMessageType_##x,
+		UBA_CACHE_MESSAGES
+		#undef UBA_CACHE_MESSAGE
+	};
+
 	inline constexpr const char EncryptionHandshakeString[] = "This is a test string used to check so encryption keys matches between client and server. This string is 128 characters long...";
 }
diff --git a/Engine/Source/Programs/UnrealBuildAccelerator/Common/Public/UbaScheduler.h b/Engine/Source/Programs/UnrealBuildAccelerator/Common/Public/UbaScheduler.h
--- a/Engine/Source/Programs/UnrealBuildAccelerator/Common/Public/UbaScheduler.h
+++ b/Engine/Source/Programs/UnrealBuildAccelerator/Common/Public/UbaScheduler.h
@@ -8,6 +8,7 @@
 
 namespace uba
 {
+	class CacheClient;
 	class Process;
 	class SessionServer;
 	struct NextProcessInfo;
@@ -18,10 +19,12 @@ namespace uba
 		SchedulerCreateInfo(SessionServer& s) : session(s) {}
 
 		SessionServer& session;
+		CacheClient* cacheClient = nullptr; // Set cache client for scheduler to use when building
 		u32 maxLocalProcessors = ~0u; // Max local processors to use. ~0u means it will use all processors
 		bool enableProcessReuse = false; // If this is true, the system will allow processes to be reused when they're asking for it.
 		bool forceRemote = false; // Force all processes that can run remotely to run remotely.
 		bool forceNative = false; // Force all processes to run native (not detoured)
+		bool writeToCache = false; // Set to true in combination with setting cacheClient to populate cache
 	};
 
 	struct EnqueueProcessInfo
@@ -64,14 +67,25 @@ namespace uba
 		struct ExitProcessInfo;
 		struct ProcessStartInfo2;
 
+		enum ProcessStatus : u8
+		{
+			ProcessStatus_QueuedForCache,
+			ProcessStatus_QueuedForRun,
+			ProcessStatus_Running,
+			ProcessStatus_Success,
+			ProcessStatus_Failed,
+			ProcessStatus_Skipped,
+		};
+
 		void ThreadLoop();
 		void RemoteProcessReturned(Process& process);
+		void HandleCacheMissed(ExitProcessInfo* ei);
 		void RemoteSlotAvailable();
 		void ProcessExited(ExitProcessInfo* info, const ProcessHandle& handle);
-		u32 PopProcess(bool isLocal);
+		u32 PopProcess(bool isLocal, ProcessStatus& outPrevStatus);
 		bool RunQueuedProcess(bool isLocal);
 		bool HandleReuseMessage(Process& process, NextProcessInfo& outNextProcess, u32 prevExitCode);
-		void ExitProcess(ExitProcessInfo& info, Process& process, u32 exitCode);
+		void ExitProcess(ExitProcessInfo& info, Process& process, u32 exitCode, bool fromCache);
 		void SkipProcess(ProcessStartInfo2& info);
 		void UpdateQueueCounter(int offset);
 		void UpdateActiveProcessCounter(bool isLocal, int offset);
@@ -79,16 +93,6 @@ namespace uba
 
 		SessionServer& m_session;
 		u32 m_maxLocalProcessors;
-	
-		enum ProcessStatus : u8
-		{
-			ProcessStatus_Queued,
-			ProcessStatus_Running,
-			ProcessStatus_Success,
-			ProcessStatus_Failed,
-			ProcessStatus_Skipped,
-		};
-
 		
 		struct ProcessEntry
 		{
@@ -114,12 +118,16 @@ namespace uba
 		bool m_forceNative;
 
 		float m_activeLocalProcessWeight = 0.0f;
+		u32 m_activeCacheQueries = 0;
 
 		Atomic<u32> m_queuedProcesses;
 		Atomic<u32> m_activeLocalProcesses;
 		Atomic<u32> m_activeRemoteProcesses;
 		Atomic<u32> m_finishedProcesses;
 
+		CacheClient* m_cacheClient;
+		bool m_writeToCache;
+
 		Scheduler(const Scheduler&) = delete;
 		void operator=(const Scheduler&) = delete;
 	};
diff --git a/Engine/Source/Programs/UnrealBuildAccelerator/Common/Public/UbaStorage.h b/Engine/Source/Programs/UnrealBuildAccelerator/Common/Public/UbaStorage.h
--- a/Engine/Source/Programs/UnrealBuildAccelerator/Common/Public/UbaStorage.h
+++ b/Engine/Source/Programs/UnrealBuildAccelerator/Common/Public/UbaStorage.h
@@ -10,7 +10,6 @@
 #include "UbaMemory.h"
 #include "UbaPathUtils.h"
 #include "UbaStats.h"
-#include <oodle2.h>
 
 #define UBA_USE_SPARSEFILE 0
 
@@ -50,7 +49,7 @@ namespace uba
 		virtual bool DropCasFile(const CasKey& casKey, bool forceDelete, const tchar* hint) = 0;
 		virtual bool CalculateCasKey(CasKey& out, const tchar* fileName) = 0;
 		virtual bool CopyOrLink(const CasKey& casKey, const tchar* destination, u32 fileAttributes, bool writeCompressed = false) = 0;
-		virtual bool FakeCopy(const CasKey& casKey, const tchar* destination) = 0;
+		virtual bool FakeCopy(const CasKey& casKey, const tchar* destination, u64 size = 0, u64 lastWritten = 0, bool deleteExisting = true) = 0;
 #if !UBA_USE_SPARSEFILE
 		virtual bool GetCasFileName(StringBufferBase& out, const CasKey& casKey) = 0;
 #endif
@@ -82,6 +81,20 @@ namespace uba
 		WorkManager* workManager = nullptr;
 	};
 
+	struct BufferSlots
+	{
+		u8* Pop();
+		void Push(u8* slot);
+
+		~BufferSlots();
+
+		ReaderWriterLock m_slotsLock;
+		Vector<u8*> m_slots;
+	};
+
+	static constexpr u64 BufferSlotSize = 16*1024*1024;
+	static constexpr u64 BufferSlotHalfSize = BufferSlotSize/2; // This must be three times a msg size or more.
+
 	class StorageImpl : public Storage
 	{
 	public:
@@ -90,6 +103,7 @@ namespace uba
 
 		bool LoadCasTable(bool logStats = true);
 		bool CheckCasContent(u32 workerCount);
+		const tchar* GetTempPath();
 
 		virtual bool SaveCasTable(bool deleteIsRunningfile, bool deleteDropped = true) override;
 		virtual u64 GetStorageCapacity() override;
@@ -110,7 +124,7 @@ namespace uba
 		virtual bool DropCasFile(const CasKey& casKey, bool forceDelete, const tchar* hint) override;
 		virtual bool CalculateCasKey(CasKey& out, const tchar* fileName) override;
 		virtual bool CopyOrLink(const CasKey& casKey, const tchar* destination, u32 fileAttributes, bool writeCompressed = false) override;
-		virtual bool FakeCopy(const CasKey& casKey, const tchar* destination) override;
+		virtual bool FakeCopy(const CasKey& casKey, const tchar* destination, u64 size = 0, u64 lastWritten = 0, bool deleteExisting = true) override;
 #if !UBA_USE_SPARSEFILE
 		virtual bool GetCasFileName(StringBufferBase& out, const CasKey& casKey) override;
 #endif
@@ -147,9 +161,9 @@ namespace uba
 		void AttachEntry(CasEntry& entry);
 		void DetachEntry(CasEntry& entry);
 		void TraverseAllCasFiles(const tchar* dir, u32 recursion, const Function<void(const StringBufferBase& fullPath, const DirectoryEntry& e)>& func);
-		void TraverseAllCasFiles(const Function<void(const CasKey& key)>& func);
+		void TraverseAllCasFiles(const Function<void(const CasKey& key, u64 size)>& func);
 		void CheckAllCasFiles();
-		void HandleOverflow();
+		void HandleOverflow(Set<CasKey>* outDeletedFiles);
 		bool OpenCasDataFile(u32 index, u64 size);
 		bool CreateCasDataFiles();
 
@@ -213,19 +227,10 @@ namespace uba
 
 		DirectoryCache m_dirCache;
 
-		static constexpr u64 BufferSlotSize = 16*1024*1024;
-		static constexpr u64 BufferSlotHalfSize = BufferSlotSize/2; // This must be three times a msg size or more.
-		u8* PopBufferSlot();
-		void PushBufferSlot(u8* slot);
-
-		ReaderWriterLock m_compSlotsLock;
-		Vector<u8*> m_compSlots;
-
-		OodleLZ_Compressor m_createCasCompressor = OodleLZ_Compressor_Kraken;
-		OodleLZ_CompressionLevel m_createCasCompressionLevel = OodleLZ_CompressionLevel_SuperFast;
+		BufferSlots m_bufferSlots;
 
-		OodleLZ_Compressor m_sendCasCompressor = OodleLZ_Compressor_Kraken;
-		OodleLZ_CompressionLevel m_sendCasCompressionLevel = OodleLZ_CompressionLevel_SuperFast;
+		u8 m_casCompressor;
+		u8 m_casCompressionLevel;
 
 		StorageStats m_stats;
 
diff --git a/Engine/Source/Programs/UnrealBuildAccelerator/Common/Public/UbaStorageClient.h b/Engine/Source/Programs/UnrealBuildAccelerator/Common/Public/UbaStorageClient.h
--- a/Engine/Source/Programs/UnrealBuildAccelerator/Common/Public/UbaStorageClient.h
+++ b/Engine/Source/Programs/UnrealBuildAccelerator/Common/Public/UbaStorageClient.h
@@ -55,8 +55,6 @@ namespace uba
 		virtual void Ping() override;
 		virtual void PrintSummary(Logger& logger) override;
 
-		static bool SendBatchMessages(Logger& logger, NetworkClient& client, u16 fetchId, u8* slot, u64 capacity, u64 left, u32 messageMaxSize, u32& readIndex, u32& responseSize);
-
 	private:
 		bool SendFile(const CasKey& casKey, const tchar* fileName, u8* sourceMem, u64 sourceSize, const tchar* hint);
 		bool PopulateCasFromDirsRecursive(const tchar* dir, WorkManager& workManager, UnorderedSet<u64>& seenIds, ReaderWriterLock& seenIdsLock);
diff --git a/Engine/Source/Programs/UnrealBuildAccelerator/Common/Public/UbaStorageServer.h b/Engine/Source/Programs/UnrealBuildAccelerator/Common/Public/UbaStorageServer.h
--- a/Engine/Source/Programs/UnrealBuildAccelerator/Common/Public/UbaStorageServer.h
+++ b/Engine/Source/Programs/UnrealBuildAccelerator/Common/Public/UbaStorageServer.h
@@ -11,6 +11,7 @@ namespace uba
 	{
 		StorageServerCreateInfo(NetworkServer& s, const tchar* rootDir_, LogWriter& writer) : StorageCreateInfo(rootDir_, writer), server(s) { workManager = &server; }
 		NetworkServer& server;
+		bool allowFallback = true; // Will fallback to file system to recalculate cas if cas content does not exist anymore
 		bool writeRecievedCasFilesToDisk = false;
 		const tchar* zone = TC("");
 	};
@@ -33,6 +34,7 @@ namespace uba
 		void OnDisconnected(u32 clientId);
 		bool HandleMessage(const ConnectionInfo& connectionInfo, u8 messageType, BinaryReader& reader, BinaryWriter& writer);
 		bool WaitForWritten(CasEntry& casEntry, ScopedWriteLock& entryLock, const tchar* hint);
+		using StorageImpl::StoreCasFile;
 
 		u16 PopId();
 		void PushId(u16 id);
@@ -136,6 +138,7 @@ namespace uba
 
 		Vector<TString> m_disallowedPaths;
 
+		bool m_allowFallback;
 		bool m_writeRecievedCasFilesToDisk;
 	};
 }
diff --git a/Engine/Source/Programs/UnrealBuildAccelerator/Common/Public/UbaStorageUtils.h b/Engine/Source/Programs/UnrealBuildAccelerator/Common/Public/UbaStorageUtils.h
new file mode 100644
--- /dev/null
+++ b/Engine/Source/Programs/UnrealBuildAccelerator/Common/Public/UbaStorageUtils.h
@@ -0,0 +1,46 @@
+// Copyright Epic Games, Inc. All Rights Reserved.
+
+#include "UbaBinaryReaderWriter.h"
+#include "UbaHash.h"
+#include "UbaNetworkMessage.h"
+#include <oodle2.h>
+
+namespace uba
+{
+	struct BufferSlots;
+	struct StorageStats;
+
+	static constexpr u8 DefaultCompressor = OodleLZ_Compressor_Kraken;
+	static constexpr u8 DefaultCompressionLevel = OodleLZ_CompressionLevel_SuperFast;
+
+	CasKey CalculateCasKey(u8* fileMem, u64 fileSize, bool storeCompressed, WorkManager* workManager);
+	bool SendBatchMessages(Logger& logger, NetworkClient& client, u16 fetchId, u8* slot, u64 capacity, u64 left, u32 messageMaxSize, u32& readIndex, u32& responseSize);
+	bool SendFile(Logger& logger, NetworkClient& client, const CasKey& casKey, const u8* sourceMem, u64 sourceSize, const tchar* hint);
+
+	struct FileSender
+	{
+		bool SendFileCompressed(const CasKey& casKey, const tchar* fileName, const u8* sourceMem, u64 sourceSize, const tchar* hint);
+
+		Logger& m_logger;
+		NetworkClient& m_client;
+		BufferSlots& m_bufferSlots;
+		StorageStats& m_stats;
+		ReaderWriterLock& m_sendOneAtTheTimeLock;
+
+		u8 m_casCompressor = DefaultCompressor;
+		u8 m_casCompressionLevel = DefaultCompressionLevel;
+	};
+
+	struct FileFetcher
+	{
+		bool RetrieveFile(Logger& logger, NetworkClient& client, const CasKey& casKey, const tchar* destination, MemoryBlock* destinationMem = nullptr);
+
+		BufferSlots& m_bufferSlots;
+		StringBuffer<> m_tempPath;
+		ReaderWriterLock m_retrieveOneBatchAtTheTimeLock;
+		bool m_errorOnFail = true;
+
+		u64 m_lastWritten = 0;
+		u64 m_size = 0;
+	};
+}
\ No newline at end of file
diff --git a/Engine/Source/Programs/UnrealBuildAccelerator/Common/Public/UbaTrace.h b/Engine/Source/Programs/UnrealBuildAccelerator/Common/Public/UbaTrace.h
--- a/Engine/Source/Programs/UnrealBuildAccelerator/Common/Public/UbaTrace.h
+++ b/Engine/Source/Programs/UnrealBuildAccelerator/Common/Public/UbaTrace.h
@@ -57,7 +57,7 @@ namespace uba
 	using Color = u32;
 	inline Color toColor(u8 r, u8 g, u8 b) { return (r << 16) + (g << 8) + b; }
 
-	static constexpr u32 TraceVersion = 23;
+	static constexpr u32 TraceVersion = 24;
 	static constexpr u32 TraceReadCompatibilityVersion = 6;
 
 	class Trace : public WorkTracker
diff --git a/Engine/Source/Programs/UnrealBuildAccelerator/Core/Private/UbaApplicationRules.cpp b/Engine/Source/Programs/UnrealBuildAccelerator/Core/Private/UbaApplicationRules.cpp
--- a/Engine/Source/Programs/UnrealBuildAccelerator/Core/Private/UbaApplicationRules.cpp
+++ b/Engine/Source/Programs/UnrealBuildAccelerator/Core/Private/UbaApplicationRules.cpp
@@ -81,6 +81,10 @@ namespace uba
 			return Contains(file, TC("\\_cl_")); // This file is needed when cl.exe spawns link.exe
 		}
 
+		virtual bool IsCacheable() const override
+		{
+			return true;
+		}
 	};
 
 	class ApplicationRulesLinkExe : public ApplicationRulesVC
@@ -134,6 +138,11 @@ namespace uba
 				|| EndsWith(fileName, fileNameLen, TC(".exe"))
 				|| EndsWith(fileName, fileNameLen, TC(".dll"));
 		}
+
+		virtual bool IsCacheable() const override
+		{
+			return true;
+		}
 	};
 
 	// ==== Clang tool chain ====
@@ -186,6 +195,11 @@ namespace uba
 		{
 			return true;
 		}
+
+		virtual bool IsCacheable() const override
+		{
+			return true;
+		}
 	};
 
 	class ApplicationRulesLdLLdExe : public ApplicationRulesClang
@@ -318,11 +332,16 @@ namespace uba
 
 		virtual bool IsOutputFile(const tchar* file, u64 fileLen) const override
 		{
-			return EndsWith(file, fileLen, TC(".dummy.h"))
+			return Contains(file, TC(".generated.dummy"))
 				|| EndsWith(file, fileLen, TC(".ispc.bc"))
 				|| EndsWith(file, fileLen, TC(".ispc.txt"))
 				|| EndsWith(file, fileLen, TC(".obj"));
 		}
+
+		virtual bool IsCacheable() const override
+		{
+			return true;
+		}
 	};
 
 	class ApplicationRulesUBTDll : public ApplicationRules
diff --git a/Engine/Source/Programs/UnrealBuildAccelerator/Core/Public/UbaApplicationRules.h b/Engine/Source/Programs/UnrealBuildAccelerator/Core/Public/UbaApplicationRules.h
--- a/Engine/Source/Programs/UnrealBuildAccelerator/Core/Public/UbaApplicationRules.h
+++ b/Engine/Source/Programs/UnrealBuildAccelerator/Core/Public/UbaApplicationRules.h
@@ -124,6 +124,11 @@ namespace uba
 			// Do nothing
 		}
 
+		virtual bool IsCacheable() const
+		{
+			return false;
+		}
+
 		u32 index = ~0u;
 	};
 
diff --git a/Engine/Source/Programs/UnrealBuildAccelerator/Core/Public/UbaProcessStats.h b/Engine/Source/Programs/UnrealBuildAccelerator/Core/Public/UbaProcessStats.h
--- a/Engine/Source/Programs/UnrealBuildAccelerator/Core/Public/UbaProcessStats.h
+++ b/Engine/Source/Programs/UnrealBuildAccelerator/Core/Public/UbaProcessStats.h
@@ -40,6 +40,7 @@ namespace uba
 		UBA_PROCESS_STAT(log, 0) \
 		UBA_PROCESS_STAT(sendFiles, 0) \
 		UBA_PROCESS_STAT(writeFiles, 19) \
+		UBA_PROCESS_STAT(queryCache, 24) \
 
 
 		#define UBA_PROCESS_STAT(T, ver) Timer T;
diff --git a/Engine/Source/Programs/UnrealBuildAccelerator/Deploy-Windows.bat b/Engine/Source/Programs/UnrealBuildAccelerator/Deploy-Windows.bat
--- a/Engine/Source/Programs/UnrealBuildAccelerator/Deploy-Windows.bat
+++ b/Engine/Source/Programs/UnrealBuildAccelerator/Deploy-Windows.bat
@@ -20,7 +20,8 @@ call Engine/Build/BatchFiles/RunUBT.bat -NoUba -NoUbaLocal -NoSNDBS -NoXGE -Arch
 	-Target="UbaDetours Win64 %config%" ^
 	-Target="UbaHost Win64 %config%" ^
 	-Target="UbaStorageProxy Win64 %config%" ^
-	-Target="UbaVisualizer Win64 %config%"
+	-Target="UbaVisualizer Win64 %config%" ^
+	-Target="UbaCacheService Win64 %config%"
 
 if %ERRORLEVEL% NEQ 0 exit /b %ERRORLEVEL%
 popd
diff --git a/Engine/Source/Programs/UnrealBuildAccelerator/Host/Private/UbaExports.cpp b/Engine/Source/Programs/UnrealBuildAccelerator/Host/Private/UbaExports.cpp
--- a/Engine/Source/Programs/UnrealBuildAccelerator/Host/Private/UbaExports.cpp
+++ b/Engine/Source/Programs/UnrealBuildAccelerator/Host/Private/UbaExports.cpp
@@ -1,9 +1,11 @@
 // Copyright Epic Games, Inc. All Rights Reserved.
 
 #include "UbaExports.h"
+#include "UbaAWS.h"
+#include "UbaCacheClient.h"
 #include "UbaNetworkBackendQuic.h"
 #include "UbaNetworkBackendTcp.h"
-#include "UbaAWS.h"
+#include "UbaNetworkClient.h"
 #include "UbaProcess.h"
 #include "UbaScheduler.h"
 #include "UbaStorageServer.h"
@@ -46,11 +48,22 @@ namespace uba
 	{
 	public:
 		NetworkServerWithBackend(bool& outSuccess, const NetworkServerCreateInfo& info, NetworkBackend* nb)
-		: NetworkServer(outSuccess, info), networkBackend(nb)
+		: NetworkServer(outSuccess, info), backend(nb)
 		{
 		}
 
-		NetworkBackend* networkBackend;
+		NetworkBackend* backend;
+	};
+
+	class NetworkClientWithBackend : public NetworkClient
+	{
+	public:
+		NetworkClientWithBackend(bool& outSuccess, const NetworkClientCreateInfo& info, NetworkBackend* nb)
+		: NetworkClient(outSuccess, info), backend(nb)
+		{
+		}
+
+		NetworkBackend* backend;
 	};
 
 	#define UBA_USE_SIGNALHANDLER 0//PLATFORM_LINUX // It might be that we can't use signal handlers in c# processes.. so don't set this to 1
@@ -120,7 +133,7 @@ extern "C"
 	void DestroyServer(uba::NetworkServer* server)
 	{
 		auto s = (uba::NetworkServerWithBackend*)server;
-		auto networkBackend = s->networkBackend;
+		auto networkBackend = s->backend;
 		delete s;
 		delete networkBackend;
 	}
@@ -138,13 +151,13 @@ extern "C"
 		}
 
 		auto s = (NetworkServerWithBackend*)server;
-		return s->StartListen(*s->networkBackend, u16(port), ip, crypto128);
+		return s->StartListen(*s->backend, u16(port), ip, crypto128);
 	}
 
 	void Server_Stop(uba::NetworkServer* server)
 	{
 		auto s = (uba::NetworkServerWithBackend*)server;
-		auto networkBackend = s->networkBackend;
+		auto networkBackend = s->backend;
 		networkBackend->StopListen();
 		server->DisconnectClients();
 	}
@@ -163,7 +176,7 @@ extern "C"
 		}
 
 		auto s = (NetworkServerWithBackend*)server;
-		return s->AddClient(*s->networkBackend, ip, u16(port), crypto128);
+		return s->AddClient(*s->backend, ip, u16(port), crypto128);
 	}
 
 	/*
@@ -438,7 +451,7 @@ uba::StorageClient* CreateStorageClient(uba::NetworkClient& client, const uba::t
 		if (server)
 		{
 			auto& s = (uba::NetworkServerWithBackend&)server->GetServer();
-			s.networkBackend->StopListen();
+			s.backend->StopListen();
 			s.DisconnectClients();
 		}
 		delete server;
@@ -511,6 +524,65 @@ uba::StorageClient* CreateStorageClient(uba::NetworkClient& client, const uba::t
 		scheduler->GetStats(outQueued, outActiveLocal, outActiveRemote, outFinished);
 	}
 
+	uba::CacheClient* CacheClient_Create(uba::SessionServer* session)
+	{
+		using namespace uba;
+		LogWriter& writer = session->GetLogWriter();
+		StorageImpl& storage = (StorageImpl&)session->GetStorage();
+
+		auto networkBackend = new NetworkBackendTcp(writer);
+
+		NetworkClientCreateInfo ncci(writer);
+		ncci.receiveTimeoutSeconds = 60;
+		bool ctorSuccess = false;
+		auto networkClient = new NetworkClientWithBackend(ctorSuccess, ncci, networkBackend);
+		if (!ctorSuccess)
+			return nullptr;
+		return new CacheClient(writer, storage, *networkClient, *session);
+	}
+
+	bool CacheClient_RegisterRoot(uba::CacheClient* cacheClient, const uba::tchar* root, bool includeInKey)
+	{
+		return cacheClient->RegisterRoot(root, includeInKey);
+	}
+
+	bool CacheClient_RegisterSystemRoots(uba::CacheClient* cacheClient)
+	{
+		return cacheClient->RegisterSystemRoots();
+	}
+
+	bool CacheClient_Connect(uba::CacheClient* cacheClient, const uba::tchar* host, int port)
+	{
+		using namespace uba;
+		auto& networkClient = (NetworkClientWithBackend&)cacheClient->GetClient();
+
+		if (!networkClient.Connect(*networkClient.backend, host, u16(port)))
+			return false;
+		cacheClient->GetStorage().LoadCasTable();
+		return true;
+	}
+
+	bool CacheClient_WriteToCache(uba::CacheClient* cacheClient, const uba::ProcessHandle* process)
+	{
+		if (!process->GetStartInfo().trackInputs)
+			return false;
+		return cacheClient->WriteToCache(*process);
+	}
+
+	bool CacheClient_FetchFromCache(uba::CacheClient* cacheClient, const uba::ProcessStartInfo& info)
+	{
+		return cacheClient->FetchFromCache(info);
+	}
+
+	void CacheClient_Destroy(uba::CacheClient* cacheClient)
+	{
+		using namespace uba;
+		auto& networkClient = (NetworkClientWithBackend&)cacheClient->GetClient();
+		networkClient.Disconnect();
+		delete cacheClient;
+		delete &networkClient;
+	}
+
 	void Uba_SetCustomAssertHandler(Uba_CustomAssertHandler* handler)
 	{
 		uba::SetCustomAssertHandler(handler);
diff --git a/Engine/Source/Programs/UnrealBuildAccelerator/Host/Public/UbaExports.h b/Engine/Source/Programs/UnrealBuildAccelerator/Host/Public/UbaExports.h
--- a/Engine/Source/Programs/UnrealBuildAccelerator/Host/Public/UbaExports.h
+++ b/Engine/Source/Programs/UnrealBuildAccelerator/Host/Public/UbaExports.h
@@ -8,6 +8,7 @@
 
 namespace uba
 {
+	class CacheClient;
 	class Process;
 	class ProcessHandle;
 	class Session;
@@ -121,6 +122,15 @@ extern "C"
 	UBA_API void Scheduler_Destroy(uba::Scheduler* scheduler);
 	UBA_API void Scheduler_GetStats(uba::Scheduler* scheduler, uba::u32& outQueued, uba::u32& outActiveLocal, uba::u32& outActiveRemote, uba::u32& outFinished);
 
+	// Cache
+	UBA_API uba::CacheClient* CacheClient_Create(uba::SessionServer* session);
+	UBA_API bool CacheClient_RegisterRoot(uba::CacheClient* cacheClient, const uba::tchar* root, bool includeInKey);
+	UBA_API bool CacheClient_RegisterSystemRoots(uba::CacheClient* cacheClient);
+	UBA_API bool CacheClient_Connect(uba::CacheClient* cacheClient, const uba::tchar* host, int port);
+	UBA_API bool CacheClient_WriteToCache(uba::CacheClient* cacheClient, const uba::ProcessHandle* process);
+	UBA_API bool CacheClient_FetchFromCache(uba::CacheClient* cacheClient, const uba::ProcessStartInfo& info);
+	UBA_API void CacheClient_Destroy(uba::CacheClient* cacheClient);
+
 	// Misc
 	using Uba_CustomAssertHandler = void(const uba::tchar* text);
 	UBA_API void Uba_SetCustomAssertHandler(Uba_CustomAssertHandler* handler);
diff --git a/Engine/Source/Programs/UnrealBuildAccelerator/UbaCacheService.Target.cs b/Engine/Source/Programs/UnrealBuildAccelerator/UbaCacheService.Target.cs
new file mode 100644
--- /dev/null
+++ b/Engine/Source/Programs/UnrealBuildAccelerator/UbaCacheService.Target.cs
@@ -0,0 +1,13 @@
+// Copyright Epic Games, Inc. All Rights Reserved.
+
+using UnrealBuildTool;
+
+[SupportedPlatforms(UnrealPlatformClass.Desktop)]
+public class UbaCacheServiceTarget : TargetRules
+{
+	public UbaCacheServiceTarget(TargetInfo Target) : base(Target)
+	{
+		LaunchModuleName = "UbaCacheService";
+		UbaAgentTarget.CommonUbaSettings(this, Target);
+	}
+}
