[Uba]
* Added missing #pragma once

diff --git a/Engine/Source/Programs/UnrealBuildAccelerator/Common/Public/UbaStorageUtils.h b/Engine/Source/Programs/UnrealBuildAccelerator/Common/Public/UbaStorageUtils.h
--- a/Engine/Source/Programs/UnrealBuildAccelerator/Common/Public/UbaStorageUtils.h
+++ b/Engine/Source/Programs/UnrealBuildAccelerator/Common/Public/UbaStorageUtils.h
@@ -1,5 +1,7 @@
 // Copyright Epic Games, Inc. All Rights Reserved.
 
+#pragma once
+
 #include "UbaBinaryReaderWriter.h"
 #include "UbaHash.h"
 #include "UbaNetworkMessage.h"
