[UBA]
* Fixed unity compile error

diff --git a/Engine/Source/Programs/UnrealBuildAccelerator/Common/Private/UbaStorageUtils.cpp b/Engine/Source/Programs/UnrealBuildAccelerator/Common/Private/UbaStorageUtils.cpp
--- a/Engine/Source/Programs/UnrealBuildAccelerator/Common/Private/UbaStorageUtils.cpp
+++ b/Engine/Source/Programs/UnrealBuildAccelerator/Common/Private/UbaStorageUtils.cpp
@@ -10,7 +10,7 @@ namespace uba
 {
 	CasKey CalculateCasKey(u8* fileMem, u64 fileSize, bool storeCompressed, WorkManager* workManager)
 	{
-		constexpr u32 MaxWorkItemsPerAction = 128; // Cap this to not starve other things
+		constexpr u32 MaxWorkItemsPerAction2 = 128; // Cap this to not starve other things
 		CasKeyHasher hasher;
 
 		if (fileSize == 0)
@@ -72,7 +72,7 @@ namespace uba
 			if (workManager)
 			{
 				workerCount = Min(workCount, workManager->GetWorkerCount()-1); // We are a worker ourselves
-				workerCount = Min(workerCount, MaxWorkItemsPerAction); // Cap this to not starve other things
+				workerCount = Min(workerCount, MaxWorkItemsPerAction2); // Cap this to not starve other things
 				rec->refCount += workerCount;
 				workManager->AddWork(work, workerCount, TC("CalculateKey"));
 			}
