[Uba]
* Fixed macos compile error
* Fixed race condition in TestRemoteScheduleReuse

diff --git a/Engine/Source/Programs/UnrealBuildAccelerator/Common/Private/UbaCacheClient.cpp b/Engine/Source/Programs/UnrealBuildAccelerator/Common/Private/UbaCacheClient.cpp
--- a/Engine/Source/Programs/UnrealBuildAccelerator/Common/Private/UbaCacheClient.cpp
+++ b/Engine/Source/Programs/UnrealBuildAccelerator/Common/Private/UbaCacheClient.cpp
@@ -72,7 +72,8 @@ namespace uba
 			if (m_roots.size() == '~' - ' ') // This is not really true.. as long as value is under 256 we're good
 				return m_logger.Error(TC("Too many roots added (%llu)"), m_roots.size());
 
-			auto& root = m_roots.emplace_back(rp);
+			auto& root = m_roots.emplace_back();
+			root.path = rp;
 
 			ToLower(root.path.data());
 			if (root.path[root.path.size()-1] != PathSeparator)
diff --git a/Engine/Source/Programs/UnrealBuildAccelerator/Test/Private/UbaTestScheduler.h b/Engine/Source/Programs/UnrealBuildAccelerator/Test/Private/UbaTestScheduler.h
--- a/Engine/Source/Programs/UnrealBuildAccelerator/Test/Private/UbaTestScheduler.h
+++ b/Engine/Source/Programs/UnrealBuildAccelerator/Test/Private/UbaTestScheduler.h
@@ -83,6 +83,8 @@ namespace uba
 				u32 queued, activeLocal, activeRemote, finished;
 				do { scheduler.GetStats(queued, activeLocal, activeRemote, finished); } while (finished != 1);
 
+				session.GetServer().DisconnectClients(); // Must make sure all is disconnected since scheduler goes out of scope
+
 				scheduler.Stop();
 				return true;
 			});
