[UBA]
* Optimized CompactPathTable and added unit test
* Made optimizations for cache lookup (it is likely entries for a cmd key mostly match)

diff --git a/Engine/Source/Programs/UnrealBuildAccelerator/Common/Private/UbaCacheClient.cpp b/Engine/Source/Programs/UnrealBuildAccelerator/Common/Private/UbaCacheClient.cpp
--- a/Engine/Source/Programs/UnrealBuildAccelerator/Common/Private/UbaCacheClient.cpp
+++ b/Engine/Source/Programs/UnrealBuildAccelerator/Common/Private/UbaCacheClient.cpp
@@ -274,6 +274,9 @@ namespace uba
 				return false;
 		}
 
+
+		UnorderedMap<u32, bool> offsetIsMatch;
+
 		// Traverse entries and test inputs against local machine
 		u32 entryCount = reader.ReadU16();
 		for (u32 i=0; i!=entryCount; ++i)
@@ -285,29 +288,35 @@ namespace uba
 			{
 				u32 casKeyOffset = u32(reader.Read7BitEncoded());
 
-				if (casKeyOffset >= m_serverCasKeyTable.GetSize())
-					if (!FetchCasTable())
+				auto insres = offsetIsMatch.try_emplace(casKeyOffset);
+				if (insres.second)
+				{
+					if (casKeyOffset >= m_serverCasKeyTable.GetSize())
+						if (!FetchCasTable())
+							return false;
+
+					StringBuffer<MaxPath> path;
+					CasKey cacheCasKey;
+					if (!GetLocalPathAndCasKey(path, cacheCasKey, m_serverCasKeyTable, m_serverPathTable, casKeyOffset))
 						return false;
+					UBA_ASSERT(IsCompressed(cacheCasKey));
 
-				StringBuffer<MaxPath> path;
-				CasKey cacheCasKey;
-				if (!GetLocalPathAndCasKey(path, cacheCasKey, m_serverCasKeyTable, m_serverPathTable, casKeyOffset))
-					return false;
-				UBA_ASSERT(IsCompressed(cacheCasKey));
+					CasKey localCasKey;
+					if (path.EndsWith(TC(".rsp")) || path.EndsWith(TC(".dep.json"))) // Need to normalize caskey for these files since they contain absolute paths
+					{
+						localCasKey = AsCompressed(NormalizeAndHashFile(path.data), true);
+					}
+					else
+					{
+						bool deferCreation = true;
+						m_storage.StoreCasFile(localCasKey, path.data, CasKeyZero, deferCreation);
+						UBA_ASSERT(localCasKey == CasKeyZero || IsCompressed(localCasKey));
+					}
 
-				CasKey localCasKey;
-				if (path.EndsWith(TC(".rsp")) || path.EndsWith(TC(".dep.json"))) // Need to normalize caskey for these files since they contain absolute paths
-				{
-					localCasKey = AsCompressed(NormalizeAndHashFile(path.data), true);
-				}
-				else
-				{
-					bool deferCreation = true;
-					m_storage.StoreCasFile(localCasKey, path.data, CasKeyZero, deferCreation);
-					UBA_ASSERT(localCasKey == CasKeyZero || IsCompressed(localCasKey));
+					insres.first->second = localCasKey == cacheCasKey;
 				}
 
-				if (localCasKey != cacheCasKey)
+				if (!insres.first->second)
 				{
 					reader.Skip(inputEnd -  reader.GetPositionData());
 					isMatch = false;
diff --git a/Engine/Source/Programs/UnrealBuildAccelerator/Common/Private/UbaCacheServer.cpp b/Engine/Source/Programs/UnrealBuildAccelerator/Common/Private/UbaCacheServer.cpp
--- a/Engine/Source/Programs/UnrealBuildAccelerator/Common/Private/UbaCacheServer.cpp
+++ b/Engine/Source/Programs/UnrealBuildAccelerator/Common/Private/UbaCacheServer.cpp
@@ -438,7 +438,7 @@ namespace uba
 				{
 					StringBuffer<> temp;
 					m_pathTable.GetString(temp, pathOffset);
-					auto res = oldToNewPathOffset.try_emplace(pathOffset, newPathTable.Add(temp.data, temp.count));
+					auto res = oldToNewPathOffset.try_emplace(pathOffset, newPathTable.AddNoLock(temp.data, temp.count));
 					UBA_ASSERT(res.second);(void)res;
 				}
 				m_pathTable.Swap(newPathTable);
diff --git a/Engine/Source/Programs/UnrealBuildAccelerator/Common/Private/UbaCompactTables.cpp b/Engine/Source/Programs/UnrealBuildAccelerator/Common/Private/UbaCompactTables.cpp
--- a/Engine/Source/Programs/UnrealBuildAccelerator/Common/Private/UbaCompactTables.cpp
+++ b/Engine/Source/Programs/UnrealBuildAccelerator/Common/Private/UbaCompactTables.cpp
@@ -12,25 +12,32 @@ namespace uba
 			m_offsets.reserve(reserveOffsetsCount);
 	}
 
-	u32 CompactPathTable::Add(const tchar* str, u32 strLen, u32* outRequiredCasTableSize)
+	u32 CompactPathTable::Add(const tchar* str, u64 strLen, u32* outRequiredCasTableSize)
 	{
 		SCOPED_WRITE_LOCK(m_lock, lock);
-		if (!m_mem.memory)
-			m_mem.Init(m_reserveSize);
-		u32 res = InternalAdd(str, strLen);
+		u32 res = AddNoLock(str, strLen);
 		if (outRequiredCasTableSize)
 			*outRequiredCasTableSize = u32(m_mem.writtenSize);
 		return res;
 	}
 
-	u32 CompactPathTable::InternalAdd(const tchar* str, u32 strLen)
+	u32 CompactPathTable::AddNoLock(const tchar* str, u64 strLen)
 	{
-		const tchar* stringKeyString = str;
+		if (!m_mem.memory)
+			m_mem.Init(m_reserveSize);
+		if (!m_mem.writtenSize)
+			m_mem.AllocateNoLock(1, 1, TC(""));
 
+		const tchar* stringKeyString = str;
 		StringBuffer<MaxPath> tempStringKeyStr;
 		if (CaseInsensitiveFs)
 			stringKeyString = tempStringKeyStr.Append(str).MakeLower().data;
 
+		return InternalAdd(str, stringKeyString, strLen);
+	}
+
+	u32 CompactPathTable::InternalAdd(const tchar* str, const tchar* stringKeyString, u64 strLen)
+	{
 		StringKey key = ToStringKey(stringKeyString, strLen);
 		auto insres = m_offsets.try_emplace(key);
 		if (!insres.second)
@@ -38,23 +45,24 @@ namespace uba
 			
 		const tchar* seg = str;
 		u32 parentOffset = 0;
-		if (const tchar* lastSeparator = TStrrchr(str, PathSeparator))
+		
+		for (const tchar* it = str + strLen - 1; it > str; --it)
 		{
-			StringBuffer<> temp;
-			temp.Append(str, lastSeparator - str);
-			parentOffset = InternalAdd(temp.data, temp.count);
-			seg = lastSeparator + 1;
+			if (*it != PathSeparator)
+				continue;
+
+			parentOffset = InternalAdd(str, stringKeyString, it - str);
+			seg = it + 1;
+			break;
 		}
 
 		u64 segLen = strLen - (seg - str);
 		u8 bytesForParent = Get7BitEncodedCount(parentOffset);
 		u64 bytesForString = GetStringWriteSize(seg, segLen);
 
-		if (!m_mem.writtenSize)
-			m_mem.Allocate(1, 1, TC(""));
-
-		u8* mem = (u8*)m_mem.Allocate(bytesForParent + bytesForString, 1, TC(""));
-		BinaryWriter writer(mem, 0, 10000);
+		u64 memSize = bytesForParent + bytesForString;
+		u8* mem = (u8*)m_mem.AllocateNoLock(memSize, 1, TC(""));
+		BinaryWriter writer(mem, 0, memSize);
 		writer.Write7BitEncoded(parentOffset);
 		writer.WriteString(seg, segLen);
 		insres.first->second = u32(mem - m_mem.memory);
@@ -66,11 +74,14 @@ namespace uba
 		u32 offsets[256];
 		offsets[0] = u32(offset);
 		u32 offsetCount = 0;
+
+		BinaryReader reader(m_mem.memory, offset, m_mem.writtenSize);
+
 		while (offset)
 		{
 			++offsetCount;
 			UBA_ASSERT(offsetCount < sizeof_array(offsets));
-			BinaryReader reader(m_mem.memory, offset, 4);
+			reader.SetPosition(offset);
 			offset = (u32)reader.Read7BitEncoded();
 			offsets[offsetCount] = u32(offset);
 		}
@@ -78,7 +89,7 @@ namespace uba
 		bool isFirst = true;
 		for (u32 i=offsetCount;i; --i)
 		{
-			BinaryReader reader(m_mem.memory, offsets[i-1], 1024);
+			reader.SetPosition(offsets[i-1]);
 			reader.Read7BitEncoded();
 
 			if (!isFirst)
@@ -106,7 +117,7 @@ namespace uba
 
 		u64 writtenSize = m_mem.writtenSize;
 		u64 left = reader.GetLeft();
-		void* mem = m_mem.Allocate(left, 1, TC(""));
+		void* mem = m_mem.AllocateNoLock(left, 1, TC(""));
 		reader.ReadBytes(mem, left);
 
 		if (!populateLookup)
@@ -154,7 +165,7 @@ namespace uba
 		if (insres.second)
 		{
 			u8 bytesForStringOffset = Get7BitEncodedCount(stringOffset);
-			u8* mem = (u8*)m_mem.Allocate(bytesForStringOffset + sizeof(CasKey), 1, TC(""));
+			u8* mem = (u8*)m_mem.AllocateNoLock(bytesForStringOffset + sizeof(CasKey), 1, TC(""));
 			BinaryWriter writer(mem, 0, 1000);
 			writer.Write7BitEncoded(stringOffset);
 			writer.WriteCasKey(casKey);
@@ -205,7 +216,7 @@ namespace uba
 		u64 writtenSize = m_mem.writtenSize;
 
 		u64 left = reader.GetLeft();
-		void* mem = m_mem.Allocate(left, 1, TC(""));
+		void* mem = m_mem.AllocateNoLock(left, 1, TC(""));
 		reader.ReadBytes(mem, left);
 
 		if (!populateLookup)
diff --git a/Engine/Source/Programs/UnrealBuildAccelerator/Common/Public/UbaCompactTables.h b/Engine/Source/Programs/UnrealBuildAccelerator/Common/Public/UbaCompactTables.h
--- a/Engine/Source/Programs/UnrealBuildAccelerator/Common/Public/UbaCompactTables.h
+++ b/Engine/Source/Programs/UnrealBuildAccelerator/Common/Public/UbaCompactTables.h
@@ -16,7 +16,8 @@ namespace uba
 	public:
 		CompactPathTable(u64 reserveSize, u64 reserveOffsetsCount = 0);
 
-		u32 Add(const tchar* str, u32 strLen, u32* outRequiredCasTableSize = nullptr);
+		u32 Add(const tchar* str, u64 strLen, u32* outRequiredCasTableSize = nullptr);
+		u32 AddNoLock(const tchar* str, u64 strLen);
 
 		void GetString(StringBufferBase& out, u64 offset) const;
 
@@ -27,7 +28,7 @@ namespace uba
 		void Swap(CompactPathTable& other);
 
 	private:
-		u32 InternalAdd(const tchar* str, u32 strLen);
+		u32 InternalAdd(const tchar* str, const tchar* stringKeyString, u64 strLen);
 		ReaderWriterLock m_lock;
 		MemoryBlock m_mem;
 		UnorderedMap<StringKey, u32> m_offsets;
diff --git a/Engine/Source/Programs/UnrealBuildAccelerator/Test/Private/UbaTestAll.h b/Engine/Source/Programs/UnrealBuildAccelerator/Test/Private/UbaTestAll.h
--- a/Engine/Source/Programs/UnrealBuildAccelerator/Test/Private/UbaTestAll.h
+++ b/Engine/Source/Programs/UnrealBuildAccelerator/Test/Private/UbaTestAll.h
@@ -28,6 +28,7 @@ namespace uba
 		UBA_TEST(TestStorage) \
 		UBA_TEST(TestDetouredTestApp) \
 		UBA_TEST(TestRemoteDetouredTestApp) \
+		UBA_TEST(TestCompactPathTable) \
 
 
 #define UBA_POSIX_TESTS \
diff --git a/Engine/Source/Programs/UnrealBuildAccelerator/Test/Private/UbaTestBasics.h b/Engine/Source/Programs/UnrealBuildAccelerator/Test/Private/UbaTestBasics.h
--- a/Engine/Source/Programs/UnrealBuildAccelerator/Test/Private/UbaTestBasics.h
+++ b/Engine/Source/Programs/UnrealBuildAccelerator/Test/Private/UbaTestBasics.h
@@ -5,6 +5,7 @@
 #define Local_GetLongPathNameW uba::GetLongPathNameW
 
 #include "UbaBinaryReaderWriter.h"
+#include "UbaCompactTables.h"
 #include "UbaFileAccessor.h"
 #include "UbaProcess.h"
 #include "UbaPathUtils.h"
@@ -307,4 +308,25 @@ namespace uba
 		return true;
 	}
 	#endif
+
+	bool TestCompactPathTable(Logger& logger, const StringBufferBase& rootDir)
+	{
+		CompactPathTable table(64*1024);
+
+		StringBuffer<> str;
+		str.Append("foo").EnsureEndsWithSlash().Append("bar");
+		u32 offset = table.Add(str.data, str.count);
+
+		StringBuffer<> str2;
+		table.GetString(str2, offset);
+		if (!str.Equals(str2.data))
+			return false;
+
+		str.Clear().Append(PathSeparator).Append("foo").Append(PathSeparator).Append("bar");
+		offset = table.Add(str.data, str.count);
+		table.GetString(str2.Clear(), offset);
+		if (!str.Equals(str2.data))
+			return false;
+		return true;
+	}
 }
