[UBA]
* Implemented V1 of CompactPathTable which instead stores string segments separately.. Reduces the size table by ~30%.. slightly slower perf
* Added "RootPaths" which is a class that holds a list of root paths needed to normalize paths into non-environment dependent paths
* Added stats for cache that visualizes properly in visualizer
* Changed hash functions to have better quality
* Cleaned up code visualizing stats in visualizer.. now boxes popping up should have a nice matching size to strings
* Cleaned up UbaExports and formatted all functions the same way.. old are deprecated and will be removed once new binaries are in

diff --git a/Engine/Source/Programs/UnrealBuildAccelerator/Common/Private/UbaCacheClient.cpp b/Engine/Source/Programs/UnrealBuildAccelerator/Common/Private/UbaCacheClient.cpp
--- a/Engine/Source/Programs/UnrealBuildAccelerator/Common/Private/UbaCacheClient.cpp
+++ b/Engine/Source/Programs/UnrealBuildAccelerator/Common/Private/UbaCacheClient.cpp
@@ -4,32 +4,20 @@
 #include "UbaFileAccessor.h"
 #include "UbaNetworkMessage.h"
 #include "UbaProcess.h"
+#include "UbaRootPaths.h"
 #include "UbaStorage.h"
 #include "UbaStorageUtils.h"
 
-#if PLATFORM_WINDOWS
-#include <shlobj_core.h>
-#endif
-
 namespace uba
 {
-	static constexpr u8 RootStartByte = ' ';
-
-	struct CacheClient::Root
-	{
-		TString path;
-		StringKey shortestPathKey;
-		bool includeInKey;
-	};
-
 	CacheClient::CacheClient(LogWriter& writer, StorageImpl& storage, NetworkClient& client, Session& session)
 	:	m_logger(writer, TC("UbaCacheClient"))
 	,	m_storage(storage)
 	,	m_client(client)
 	,	m_session(session)
-	,	m_serverPathTable(CachePathTableMaxSize)
+	,	m_serverPathTable(CachePathTableMaxSize, CompactPathTable::V1)
 	,	m_serverCasKeyTable(CacheCasKeyTableMaxSize)
-	,	m_sendPathTable(8*1024*1024)
+	,	m_sendPathTable(8*1024*1024, CompactPathTable::V1)
 	,	m_sendCasKeyTable(8*1024*1024)
 	{
 		m_client.RegisterOnConnected([this]()
@@ -54,76 +42,7 @@ namespace uba
 
 	CacheClient::~CacheClient() = default;
 
-	bool CacheClient::RegisterRoot(const tchar* rootPath, bool includeInKey)
-	{
-		// Register rootPath both with single path separators and double path separators because text files store them with double path separators
-
-		StringBuffer<> doubleSlash;
-		for (const tchar* it=rootPath; *it; ++it)
-		{
-			doubleSlash.Append(*it);
-			if (*it == PathSeparator)
-				doubleSlash.Append(PathSeparator);
-		}
-
-		const tchar* rootPaths[] = { rootPath, doubleSlash.data };
-		for (const tchar* rp : rootPaths)
-		{
-			if (m_roots.size() == '~' - ' ') // This is not really true.. as long as value is under 256 we're good
-				return m_logger.Error(TC("Too many roots added (%llu)"), m_roots.size());
-
-			auto& root = m_roots.emplace_back();
-			root.path = rp;
-
-			ToLower(root.path.data());
-			if (root.path[root.path.size()-1] != PathSeparator)
-				return m_logger.Error(TC("Root path must end with separator"));
-
-			root.includeInKey = includeInKey;
-
-			m_longestRoot = Max(u32(root.path.size()), m_longestRoot);
-
-			if (!m_shortestRoot || root.path.size() < m_shortestRoot)
-			{
-				m_shortestRoot = u32(root.path.size());
-				for (auto& r : m_roots)
-					r.shortestPathKey = ToStringKeyNoCheck(r.path.data(), m_shortestRoot);
-			}
-			else
-				root.shortestPathKey = ToStringKeyNoCheck(root.path.data(), m_shortestRoot);
-		}
-		return true;
-	}
-
-	bool CacheClient::RegisterSystemRoots()
-	{
-		#if PLATFORM_WINDOWS
-		StringBuffer<MaxPath> dir;
-		dir.count = GetSystemDirectory(dir.data, dir.capacity);
-		RegisterRoot(dir.EnsureEndsWithSlash().data, false); // Ignore files from here.. we do expect them not to affect the output of a process
-		
-		dir.count = GetEnvironmentVariable(TC("ProgramW6432"), dir.Clear().data, dir.capacity);
-		RegisterRoot(dir.EnsureEndsWithSlash().data, true);
-
-		dir.count = GetEnvironmentVariable(TC("ProgramFiles(x86)"), dir.Clear().data, dir.capacity);
-		RegisterRoot(dir.EnsureEndsWithSlash().data, true);
-
-		dir.count = GetEnvironmentVariable(TC("ProgramFiles(x86)"), dir.Clear().data, dir.capacity);
-		RegisterRoot(dir.EnsureEndsWithSlash().data, true);
-
-		PWSTR path;
-		if (!SUCCEEDED(SHGetKnownFolderPath(FOLDERID_ProgramData, 0, NULL, &path)))
-			return false;
-		RegisterRoot(dir.Clear().Append(path).EnsureEndsWithSlash().data, true);
-		CoTaskMemFree(path);
-
-		#else
-		UBA_ASSERT(false);
-		#endif
-		return true;
-	}
-
-	bool CacheClient::WriteToCache(const ProcessHandle& process)
+	bool CacheClient::WriteToCache(const RootPaths& rootPaths, const ProcessHandle& process)
 	{
 		if (!m_connected)
 			return false;
@@ -132,7 +51,7 @@ namespace uba
 		if (!si.trackInputs)
 			return false;
 
-		CasKey cmdKey = GetCmdKey(si);
+		CasKey cmdKey = GetCmdKey(rootPaths, si);
 		if (cmdKey == CasKeyZero)
 			return false;
 
@@ -178,7 +97,7 @@ namespace uba
 			}
 			else if (path.EndsWith(TC(".rsp"))) // Paths can be absolute in rsp files so we need to normalize those paths
 			{
-				casKey = AsCompressed(NormalizeAndHashFile(path.data), true);
+				casKey = AsCompressed(rootPaths.NormalizeAndHashFile(m_logger, path.data), true);
 			}
 			else if (path[path.count-1] == ':')
 			{
@@ -188,21 +107,20 @@ namespace uba
 			}
 
 			// Find root for path in order to be able to normalize it.
-			u32 rootIndex = FindRootIndex(path);
-			if (rootIndex == ~0u)
+			auto root = rootPaths.FindRoot(path);
+			if (!root)
 			{
 				m_logger.Info(TC("FILE WITHOUT ROOT: %s"), path.data);
 				success = false;
 				continue;
 			}
 
-			auto& root = m_roots[rootIndex];
-			if (!root.includeInKey)
+			if (!root->includeInKey)
 				continue;
 
-			u32 rootLen = u32(root.path.size());
-			TString qualifiedPath = path.data + rootLen - 2;
-			qualifiedPath[0] = tchar(RootStartByte + rootIndex);
+			u32 rootLen = u32(root->path.size());
+			TString qualifiedPath = path.data + rootLen - 1;
+			qualifiedPath[0] = tchar(RootPaths::RootStartByte + root->index);
 
 			u32 pathOffset = m_sendPathTable.Add(qualifiedPath.c_str(), u32(qualifiedPath.size()), &requiredPathTableSize);
 
@@ -217,7 +135,7 @@ namespace uba
 
 			// .dep.json contains absolute paths, need to normalize file
 			if (isOutput && path.EndsWith(TC(".dep.json"))) // TODO: More data driven approach. Also, hash does not match content atm.
-				casKey = AsCompressed(NormalizeAndHashFile(path.data), true);
+				casKey = AsCompressed(rootPaths.NormalizeAndHashFile(m_logger, path.data), true);
 
 			// Get file caskey using storage
 			if (casKey == CasKeyZero)
@@ -248,24 +166,50 @@ namespace uba
 			return false;
 
 		// actual cache entry now when we know server has the needed tables
-		if (!SendCacheEntry(cmdKey, inputsStringToCasKey, outputsStringToCasKey))
+		if (!SendCacheEntry(rootPaths, cmdKey, inputsStringToCasKey, outputsStringToCasKey))
 			return false;
 
 		return true;
 	}
 
-	bool CacheClient::FetchFromCache(const ProcessStartInfo& info)
+	bool CacheClient::FetchFromCache(const RootPaths& rootPaths, const ProcessStartInfo& info)
 	{
 		if (!m_connected)
 			return false;
 
-		CasKey cmdKey = GetCmdKey(info);
+		CacheStats cacheStats;
+		StorageStats storageStats;
+		SystemStats systemStats;
+
+		StorageStatsScope __(storageStats);
+		SystemStatsScope _(systemStats);
+
+		CasKey cmdKey = GetCmdKey(rootPaths, info);
 		if (cmdKey == CasKeyZero)
 			return false;
 
-		StackBinaryReader<SendMaxSize> reader;
+		u8 memory[SendMaxSize];
+
+		u32 fetchId = m_session.CreateProcessId();
+		m_session.GetTrace().CacheBeginFetch(fetchId, info.description);
+		bool success = false;
+		auto tg = MakeGuard([&]()
+			{
+				cacheStats.testEntries.time -= cacheStats.fetchCasTable.time;
+				BinaryWriter writer(memory, 0, sizeof_array(memory));
+				cacheStats.Write(writer);
+				if (success)
+				{
+					storageStats.Write(writer);
+					systemStats.Write(writer);
+				}
+				m_session.GetTrace().CacheEndFetch(fetchId, success, memory, writer.GetPosition());
+			});
+
+		BinaryReader reader(memory, 0, sizeof_array(memory));
 
 		{
+			TimerScope ts(cacheStats.fetchEntries);
 			// Fetch entries.. server will provide as many as fits. TODO: Should it be possible to ask for more entries?
 			StackBinaryWriter<32> writer;
 			NetworkMessage msg(m_client, CacheServiceId, CacheMessageType_FetchEntries, writer);
@@ -281,58 +225,65 @@ namespace uba
 		u32 entryCount = reader.ReadU16();
 		for (u32 i=0; i!=entryCount; ++i)
 		{
-			bool isMatch = true;
-			u64 inputSize = reader.Read7BitEncoded();
-			const u8* inputEnd = reader.GetPositionData() + inputSize;
-			while (reader.GetPositionData() != inputEnd)
+			u64 outputSize = 0;
 			{
-				u32 casKeyOffset = u32(reader.Read7BitEncoded());
-
-				auto insres = offsetIsMatch.try_emplace(casKeyOffset);
-				if (insres.second)
+				TimerScope ts(cacheStats.testEntries);
+				bool isMatch = true;
+				u64 inputSize = reader.Read7BitEncoded();
+				const u8* inputEnd = reader.GetPositionData() + inputSize;
+				while (reader.GetPositionData() != inputEnd)
 				{
-					if (casKeyOffset >= m_serverCasKeyTable.GetSize())
-						if (!FetchCasTable())
+					u32 casKeyOffset = u32(reader.Read7BitEncoded());
+
+					auto insres = offsetIsMatch.try_emplace(casKeyOffset);
+					if (insres.second)
+					{
+						if (casKeyOffset >= m_serverCasKeyTable.GetSize())
+						{
+							TimerScope ts2(cacheStats.fetchCasTable);
+							if (!FetchCasTable())
+								return false;
+						}
+
+						StringBuffer<MaxPath> path;
+						CasKey cacheCasKey;
+						if (!GetLocalPathAndCasKey(rootPaths, path, cacheCasKey, m_serverCasKeyTable, m_serverPathTable, casKeyOffset))
 							return false;
+						UBA_ASSERT(IsCompressed(cacheCasKey));
 
-					StringBuffer<MaxPath> path;
-					CasKey cacheCasKey;
-					if (!GetLocalPathAndCasKey(path, cacheCasKey, m_serverCasKeyTable, m_serverPathTable, casKeyOffset))
-						return false;
-					UBA_ASSERT(IsCompressed(cacheCasKey));
+						CasKey localCasKey;
+						if (path.EndsWith(TC(".rsp")) || path.EndsWith(TC(".dep.json"))) // Need to normalize caskey for these files since they contain absolute paths
+						{
+							localCasKey = AsCompressed(rootPaths.NormalizeAndHashFile(m_logger, path.data), true);
+						}
+						else
+						{
+							bool deferCreation = true;
+							m_storage.StoreCasFile(localCasKey, path.data, CasKeyZero, deferCreation);
+							UBA_ASSERT(localCasKey == CasKeyZero || IsCompressed(localCasKey));
+						}
 
-					CasKey localCasKey;
-					if (path.EndsWith(TC(".rsp")) || path.EndsWith(TC(".dep.json"))) // Need to normalize caskey for these files since they contain absolute paths
-					{
-						localCasKey = AsCompressed(NormalizeAndHashFile(path.data), true);
+						insres.first->second = localCasKey == cacheCasKey;
 					}
-					else
+
+					if (!insres.first->second)
 					{
-						bool deferCreation = true;
-						m_storage.StoreCasFile(localCasKey, path.data, CasKeyZero, deferCreation);
-						UBA_ASSERT(localCasKey == CasKeyZero || IsCompressed(localCasKey));
+						reader.Skip(inputEnd -  reader.GetPositionData());
+						isMatch = false;
+						break;
 					}
-
-					insres.first->second = localCasKey == cacheCasKey;
 				}
 
-				if (!insres.first->second)
+				outputSize = reader.Read7BitEncoded();
+
+				// No match, test next entry
+				if (!isMatch)
 				{
-					reader.Skip(inputEnd -  reader.GetPositionData());
-					isMatch = false;
-					break;
+					reader.Skip(outputSize);
+					continue;
 				}
 			}
 
-			u64 outputSize = reader.Read7BitEncoded();
-
-			// No match, test next entry
-			if (!isMatch)
-			{
-				reader.Skip(outputSize);
-				continue;
-			}
-
 			// Fetch output files from cache (and some files need to be "denormalized" before written to disk
 
 			const u8* outputEnd = reader.GetPositionData() + outputSize;
@@ -343,9 +294,11 @@ namespace uba
 					if (!FetchCasTable())
 						return false;
 
+				TimerScope ts(cacheStats.fetchOutput);
+
 				StringBuffer<MaxPath> path;
 				CasKey casKey;
-				if (!GetLocalPathAndCasKey(path, casKey, m_serverCasKeyTable, m_serverPathTable, casKeyOffset))
+				if (!GetLocalPathAndCasKey(rootPaths, path, casKey, m_serverCasKeyTable, m_serverPathTable, casKeyOffset))
 					return false;
 				UBA_ASSERT(IsCompressed(casKey));
 
@@ -370,8 +323,8 @@ namespace uba
 						u64 rootOffset = reader2.Read7BitEncoded();
 						if (u64 toWrite = rootOffset - lastWritten)
 							memcpy(localBlock.Allocate(toWrite, 1, TC("")), fileStart + lastWritten, toWrite);
-						u8 rootIndex = fileStart[rootOffset] - RootStartByte;
-						auto& root = m_roots[rootIndex];
+						u8 rootIndex = fileStart[rootOffset] - RootPaths::RootStartByte;
+						auto& root = rootPaths.GetRoot(rootIndex);
 
 						#if PLATFORM_WINDOWS
 						StringBuffer<> pathTemp;
@@ -415,6 +368,7 @@ namespace uba
 					return false;
 			}
 
+			success = true;
 			return true;
 		}
 
@@ -485,7 +439,7 @@ namespace uba
 		return true;
 	}
 
-	bool CacheClient::SendCacheEntry(const CasKey& cmdKey, const Map<u32, u32>& inputsStringToCasKey, const Map<u32, u32>& outputsStringToCasKey)
+	bool CacheClient::SendCacheEntry(const RootPaths& rootPaths, const CasKey& cmdKey, const Map<u32, u32>& inputsStringToCasKey, const Map<u32, u32>& outputsStringToCasKey)
 	{
 		StackBinaryReader<1024> reader;
 		{
@@ -516,7 +470,7 @@ namespace uba
 
 			StringBuffer<MaxPath> path;
 			CasKey casKey;
-			if (!GetLocalPathAndCasKey(path, casKey, m_sendCasKeyTable, m_sendPathTable, casKeyOffset))
+			if (!GetLocalPathAndCasKey(rootPaths, path, casKey, m_sendCasKeyTable, m_sendPathTable, casKeyOffset))
 				return false;
 
 			casKey = AsCompressed(casKey, true);
@@ -577,7 +531,7 @@ namespace uba
 							}
 						};
 
-					if (!NormalizeString<char>((const char*)file.GetData(), file.GetSize(), handleString, path.data))
+					if (!rootPaths.NormalizeString<char>(m_logger, (const char*)file.GetData(), file.GetSize(), handleString, path.data))
 						return false;
 
 					if (!rootOffsets.empty())
@@ -657,7 +611,7 @@ namespace uba
 		return true;
 	}
 
-	CasKey CacheClient::GetCmdKey(const ProcessStartInfo& info)
+	CasKey CacheClient::GetCmdKey(const RootPaths& rootPaths, const ProcessStartInfo& info)
 	{
 		CasKeyHasher hasher;
 
@@ -670,7 +624,7 @@ namespace uba
 
 		// Add arguments list to key
 		auto hashString = [&](const tchar* str, u64 strLen, u32 rootPos) { hasher.Update(str, strLen*sizeof(tchar)); };
-		if (!NormalizeString(info.arguments, TStrlen(info.arguments), hashString, TC("")))
+		if (!rootPaths.NormalizeString(m_logger, info.arguments, TStrlen(info.arguments), hashString, TC("")))
 			return CasKeyZero;
 
 		// Add content of rsp file to key (This will cost a bit of perf since we need to normalize.. should this be part of key?)
@@ -685,7 +639,7 @@ namespace uba
 					if (rspStart[1] != ':')
 						rsp.Append(info.workingDir).EnsureEndsWithSlash();
 					rsp.Append(rspStart, rspEnd - rspStart);
-					CasKey rspCasKey = NormalizeAndHashFile(rsp.data);
+					CasKey rspCasKey = rootPaths.NormalizeAndHashFile(m_logger, rsp.data);
 					hasher.Update(&rspCasKey, sizeof(CasKey));
 				}
 			}
@@ -694,96 +648,7 @@ namespace uba
 		return ToCasKey(hasher, false);
 	}
 
-	u32 CacheClient::FindRootIndex(const StringBufferBase& path)
-	{
-		if (path.count < m_shortestRoot)
-			return ~0u;
-
-		StringBuffer<MaxPath> shortPath;
-		shortPath.Append(path.data, m_shortestRoot).MakeLower();
-		StringKey key = ToStringKeyNoCheck(shortPath.data, m_shortestRoot);
-		for (u32 i=0, e=u32(m_roots.size()); i!=e; ++i)
-		{
-			auto& root = m_roots[i];
-			if (key != root.shortestPathKey)
-				continue;
-			if (!path.StartsWith(root.path.c_str()))
-				continue;
-			return i;
-		}
-		return ~0u;
-	}
-
-	template<typename CharType, typename Func>
-	bool CacheClient::NormalizeString(const CharType* str, u64 strLen, const Func& func, const tchar* hint)
-	{
-		auto strEnd = str + strLen;
-		auto searchPos = str;
-
-		u32 destPos = 0;
-
-		while (true)
-		{
-			auto absPathChars = searchPos;
-			CharType lastChar = 0;
-			while (absPathChars < strEnd && !(lastChar == ':' && *absPathChars == '\\'))
-			{
-				lastChar = *absPathChars;
-				++absPathChars;
-			}
-		
-			if (absPathChars == strEnd)
-			{
-				func(searchPos, strEnd - searchPos, ~0u);
-				return true;
-			}
-
-			auto pathStart = absPathChars - 2;
-
-			auto pathEndOrMore = pathStart;
-			while (pathEndOrMore < strEnd && *pathEndOrMore != '\n')
-				++pathEndOrMore;
-
-			u32 lenOrMore = u32(pathEndOrMore - pathStart);
-			u32 toCopy = Min(lenOrMore, m_longestRoot);
-			StringBuffer<512> path;
-			path.Append(pathStart, toCopy);
-
-			u32 rootIndex = FindRootIndex(path);
-			if (rootIndex == ~0u)
-			{
-				m_logger.Info(TC("PATH WITHOUT ROOT: %s (inside file %s)"), path.data, hint);
-				return false;
-			}
-
-			if (u32 len = u32(pathStart - searchPos))
-			{
-				destPos += len;
-				func(searchPos, len, ~0u);
-			}
-			CharType temp = RootStartByte + CharType(rootIndex);
-			func(&temp, 1, destPos);
-			destPos += 1;
-
-			searchPos = pathStart + u32(m_roots[rootIndex].path.size()) - 1;
-		}
-	}
-
-	CasKey CacheClient::NormalizeAndHashFile(const tchar* filename)
-	{
-		FileAccessor file(m_logger, filename);
-		if (!file.OpenMemoryRead())
-			return CasKeyZero;
-
-		CasKeyHasher hasher;
-		auto hashString = [&](const char* str, u64 strLen, u32 rootPos) { hasher.Update(str, strLen); };
-		if (!NormalizeString<char>((const char*)file.GetData(), file.GetSize(), hashString, filename))
-			return CasKeyZero;
-
-		return ToCasKey(hasher, false);
-	}
-
-	bool CacheClient::GetLocalPathAndCasKey(StringBufferBase& outPath, CasKey& outKey, CompactCasKeyTable& casKeyTable, CompactPathTable& pathTable, u32 offset)
+	bool CacheClient::GetLocalPathAndCasKey(const RootPaths& rootPaths, StringBufferBase& outPath, CasKey& outKey, CompactCasKeyTable& casKeyTable, CompactPathTable& pathTable, u32 offset)
 	{
 		if (!m_connected)
 			return false;
@@ -794,13 +659,11 @@ namespace uba
 		casKeyTable.GetPathAndKey(normalizedPath, outKey, pathTable, offset);
 		UBA_ASSERT(normalizedPath.count);
 
-		StringBuffer<8> rootIndexStr;
-		rootIndexStr.Append(normalizedPath.data, normalizedPath.First(PathSeparator) - normalizedPath.data);
-		u32 rootIndex = rootIndexStr[0] - RootStartByte;
-		auto& root = m_roots[rootIndex];
+		u32 rootIndex = normalizedPath[0] - RootPaths::RootStartByte;
+		auto& root = rootPaths.GetRoot(rootIndex);
 
 		StringBuffer<MaxPath> path;
-		outPath.Append(root.path).Append(normalizedPath.data + rootIndexStr.count + 1);
+		outPath.Append(root.path).Append(normalizedPath.data + 1);
 		return true;
 	}
 }
diff --git a/Engine/Source/Programs/UnrealBuildAccelerator/Common/Private/UbaCacheServer.cpp b/Engine/Source/Programs/UnrealBuildAccelerator/Common/Private/UbaCacheServer.cpp
--- a/Engine/Source/Programs/UnrealBuildAccelerator/Common/Private/UbaCacheServer.cpp
+++ b/Engine/Source/Programs/UnrealBuildAccelerator/Common/Private/UbaCacheServer.cpp
@@ -12,7 +12,7 @@ namespace uba
 
 	struct CacheServer::Connection
 	{
-		Connection() : m_pathTable(8*1024*1024), m_casKeyTable(8*1024*1024) {}
+		Connection() : m_pathTable(8*1024*1024, CompactPathTable::V1), m_casKeyTable(8*1024*1024) {}
 		CompactPathTable m_pathTable;
 		CompactCasKeyTable m_casKeyTable;
 
@@ -36,7 +36,7 @@ namespace uba
 	:	m_logger(writer, TC("UbaCacheServer"))
 	,	m_server(server)
 	,	m_storage(storage)
-	,	m_pathTable(CachePathTableMaxSize)
+	,	m_pathTable(CachePathTableMaxSize, CompactPathTable::V1)
 	,	m_casKeyTable(CacheCasKeyTableMaxSize)
 	{
 		m_rootDir.count = GetFullPathNameW(rootDir, m_rootDir.capacity, m_rootDir.data, NULL);
@@ -430,15 +430,22 @@ namespace uba
 			UnorderedMap<u32, u32> oldToNewPathOffset;
 			u32 oldSize = m_pathTable.GetSize();
 			{
-				u64 reservedPathOffsets = usedPathOffsets.size()*2; // Trying to estimate offset count
-				CompactPathTable newPathTable(CachePathTableMaxSize, reservedPathOffsets);
+				CompactPathTable newPathTable(CachePathTableMaxSize, CompactPathTable::V1, m_pathTable.GetPathCount(), m_pathTable.GetSegmentCount());
 				oldToNewPathOffset.reserve(usedPathOffsets.size());
 
 				for (u32 pathOffset : usedPathOffsets)
 				{
 					StringBuffer<> temp;
 					m_pathTable.GetString(temp, pathOffset);
-					auto res = oldToNewPathOffset.try_emplace(pathOffset, newPathTable.AddNoLock(temp.data, temp.count));
+					u32 newOffset = newPathTable.AddNoLock(temp.data, temp.count);
+
+					#if 0
+					StringBuffer<> test;
+					newPathTable.GetString(test, newOffset);
+					UBA_ASSERT(test.Equals(temp.data));
+					#endif
+
+					auto res = oldToNewPathOffset.try_emplace(pathOffset, newOffset);
 					UBA_ASSERT(res.second);(void)res;
 				}
 				m_pathTable.Swap(newPathTable);
@@ -633,6 +640,13 @@ namespace uba
 			UBA_ASSERT(path.count);
 
 			u32 pathOffset = m_pathTable.Add(path.data, path.count);
+
+			#if 0
+			StringBuffer<> test;
+			m_pathTable.GetString(test, pathOffset);
+			UBA_ASSERT(test.Equals(path.data));
+			#endif
+
 			u32 casKeyOffset = m_casKeyTable.Add(casKey, pathOffset);
 			inputs.insert(casKeyOffset);
 			bytesForInput += Get7BitEncodedCount(casKeyOffset);
@@ -726,6 +740,13 @@ namespace uba
 			StringBuffer<> path;
 			connection.m_casKeyTable.GetPathAndKey(path, casKey, connection.m_pathTable, outputOffset);
 			u32 pathOffset = m_pathTable.Add(path.data, path.count);
+
+			#if 0
+			StringBuffer<> test;
+			m_pathTable.GetString(test, pathOffset);
+			UBA_ASSERT(test.Equals(path.data));
+			#endif
+
 			u32 casKeyOffset = m_casKeyTable.Add(casKey, pathOffset);
 			outputs.insert(casKeyOffset);
 			bytesForOutput += Get7BitEncodedCount(casKeyOffset);
diff --git a/Engine/Source/Programs/UnrealBuildAccelerator/Common/Private/UbaCompactTables.cpp b/Engine/Source/Programs/UnrealBuildAccelerator/Common/Private/UbaCompactTables.cpp
--- a/Engine/Source/Programs/UnrealBuildAccelerator/Common/Private/UbaCompactTables.cpp
+++ b/Engine/Source/Programs/UnrealBuildAccelerator/Common/Private/UbaCompactTables.cpp
@@ -5,11 +5,14 @@
 
 namespace uba
 {
-	CompactPathTable::CompactPathTable(u64 reserveSize, u64 reserveOffsetsCount)
+	CompactPathTable::CompactPathTable(u64 reserveSize, Version version, u64 reservePathCount, u64 reserveSegmentCount)
 	{
 		m_reserveSize = reserveSize;
-		if (reserveOffsetsCount)
-			m_offsets.reserve(reserveOffsetsCount);
+		if (reservePathCount)
+			m_offsets.reserve(reservePathCount);
+		if (reserveSegmentCount)
+			m_segmentOffsets.reserve(reserveSegmentCount);
+		m_version = version;
 	}
 
 	u32 CompactPathTable::Add(const tchar* str, u64 strLen, u32* outRequiredCasTableSize)
@@ -58,15 +61,55 @@ namespace uba
 
 		u64 segLen = strLen - (seg - str);
 		u8 bytesForParent = Get7BitEncodedCount(parentOffset);
-		u64 bytesForString = GetStringWriteSize(seg, segLen);
-
-		u64 memSize = bytesForParent + bytesForString;
-		u8* mem = (u8*)m_mem.AllocateNoLock(memSize, 1, TC(""));
-		BinaryWriter writer(mem, 0, memSize);
-		writer.Write7BitEncoded(parentOffset);
-		writer.WriteString(seg, segLen);
-		insres.first->second = u32(mem - m_mem.memory);
-		return insres.first->second;
+
+		if (m_version == V0)
+		{
+			u64 bytesForString = GetStringWriteSize(seg, segLen);
+
+			u64 memSize = bytesForParent + bytesForString;
+			u8* mem = (u8*)m_mem.AllocateNoLock(memSize, 1, TC(""));
+			BinaryWriter writer(mem, 0, memSize);
+			writer.Write7BitEncoded(parentOffset);
+			writer.WriteString(seg, segLen);
+			insres.first->second = u32(mem - m_mem.memory);
+			return insres.first->second;
+		}
+		else
+		{
+			StringKey segmentKey = ToStringKeyNoCheck(seg, segLen);
+			auto insres2 = m_segmentOffsets.try_emplace(segmentKey);
+			if (insres2.second)
+			{
+				// Put string directly after current element and set segment offset to 0
+				u64 bytesForString = GetStringWriteSize(seg, segLen);
+				u64 memSize = bytesForParent + 1 + bytesForString;
+				u8* mem = (u8*)m_mem.AllocateNoLock(memSize, 1, TC(""));
+				BinaryWriter writer(mem, 0, memSize);
+				writer.Write7BitEncoded(parentOffset);
+				writer.Write7BitEncoded(0);
+				writer.WriteString(seg, segLen);
+				u32 offset = u32(mem - m_mem.memory);
+				insres.first->second = offset;
+				insres2.first->second = offset + bytesForParent + 1;
+				return offset;
+			}
+
+			#if 0
+			StringBuffer<> temp;
+			BinaryReader reader(m_mem.memory, insres2.first->second, 1000);
+			reader.ReadString(temp);
+			UBA_ASSERT(temp.count == segLen && wcsncmp(temp.data, seg, segLen) == 0);
+			#endif
+
+			u32 strOffset = insres2.first->second;
+			u64 memSize = bytesForParent + Get7BitEncodedCount(strOffset);
+			u8* mem = (u8*)m_mem.AllocateNoLock(memSize, 1, TC(""));
+			BinaryWriter writer(mem, 0, memSize);
+			writer.Write7BitEncoded(parentOffset);
+			writer.Write7BitEncoded(strOffset);
+			insres.first->second = u32(mem - m_mem.memory);
+			return insres.first->second;
+		}
 	}
 
 	void CompactPathTable::GetString(StringBufferBase& out, u64 offset) const
@@ -86,16 +129,36 @@ namespace uba
 			offsets[offsetCount] = u32(offset);
 		}
 
-		bool isFirst = true;
-		for (u32 i=offsetCount;i; --i)
+		if (m_version == V0)
 		{
-			reader.SetPosition(offsets[i-1]);
-			reader.Read7BitEncoded();
-
-			if (!isFirst)
-				out.Append(PathSeparator);
-			isFirst = false;
-			reader.ReadString(out);
+			bool isFirst = true;
+			for (u32 i=offsetCount;i; --i)
+			{
+				reader.SetPosition(offsets[i-1]);
+				reader.Read7BitEncoded();
+
+				if (!isFirst)
+					out.Append(PathSeparator);
+				isFirst = false;
+				reader.ReadString(out);
+			}
+		}
+		else
+		{
+			bool isFirst = true;
+			for (u32 i=offsetCount;i; --i)
+			{
+				reader.SetPosition(offsets[i-1]);
+				reader.Read7BitEncoded();
+				u32 strOffset = u32(reader.Read7BitEncoded());
+				if (strOffset != 0)
+					reader.SetPosition(strOffset);
+
+				if (!isFirst)
+					out.Append(PathSeparator);
+				isFirst = false;
+				reader.ReadString(out);
+			}
 		}
 	}
 
@@ -127,26 +190,54 @@ namespace uba
 		if (!writtenSize)
 			reader2.Skip(1);
 
-		while (reader2.GetLeft())
+		if (m_version == V0)
 		{
-			u32 offset = u32(reader2.GetPosition());
-			reader2.Read7BitEncoded();
-			reader2.SkipString();
-			StringBuffer<> str;
-			GetString(str, offset);
-			if (CaseInsensitiveFs)
-				str.MakeLower();
-			m_offsets.try_emplace(ToStringKey(str), offset);
+			while (reader2.GetLeft())
+			{
+				u32 offset = u32(reader2.GetPosition());
+				reader2.Read7BitEncoded();
+				reader2.SkipString();
+				StringBuffer<> str;
+				GetString(str, offset);
+				if (CaseInsensitiveFs)
+					str.MakeLower();
+				m_offsets.try_emplace(ToStringKeyNoCheck(str.data, str.count), offset);
+			}
+		}
+		else
+		{
+			while (reader2.GetLeft())
+			{
+				u32 offset = u32(reader2.GetPosition());
+				reader2.Read7BitEncoded();
+				u64 stringOffset = reader2.Read7BitEncoded();
+				if (!stringOffset)
+				{
+					u32 strOffset = u32(reader2.GetPosition());
+					StringBuffer<> seg;
+					reader2.ReadString(seg);
+					m_segmentOffsets.try_emplace(ToStringKeyNoCheck(seg.data, seg.count), strOffset);
+				}
+				StringBuffer<> str;
+				GetString(str, offset);
+				if (CaseInsensitiveFs)
+					str.MakeLower();
+				m_offsets.try_emplace(ToStringKeyNoCheck(str.data, str.count), offset);
+			}
 		}
 	}
 
 	void CompactPathTable::Swap(CompactPathTable& other)
 	{
 		m_offsets.swap(other.m_offsets);
+		m_segmentOffsets.swap(other.m_segmentOffsets);
 		m_mem.Swap(other.m_mem);
 		u64 rs = m_reserveSize;
 		m_reserveSize = other.m_reserveSize;
 		other.m_reserveSize = rs;
+		Version v = m_version;
+		m_version = other.m_version;
+		other.m_version = v;
 	}
 
 	CompactCasKeyTable::CompactCasKeyTable(u64 reserveSize, u64 reserveOffsetsCount)
diff --git a/Engine/Source/Programs/UnrealBuildAccelerator/Common/Private/UbaRootPaths.cpp b/Engine/Source/Programs/UnrealBuildAccelerator/Common/Private/UbaRootPaths.cpp
new file mode 100644
--- /dev/null
+++ b/Engine/Source/Programs/UnrealBuildAccelerator/Common/Private/UbaRootPaths.cpp
@@ -0,0 +1,122 @@
+// Copyright Epic Games, Inc. All Rights Reserved.
+
+#include "UbaRootPaths.h"
+#include "UbaFileAccessor.h"
+
+#if PLATFORM_WINDOWS
+#include <shlobj_core.h>
+#endif
+
+namespace uba
+{
+	bool RootPaths::RegisterRoot(Logger& logger, const tchar* rootPath, bool includeInKey)
+	{
+		// Register rootPath both with single path separators and double path separators because text files store them with double path separators
+
+		StringBuffer<> doubleSlash;
+		for (const tchar* it=rootPath; *it; ++it)
+		{
+			doubleSlash.Append(*it);
+			if (*it == PathSeparator)
+				doubleSlash.Append(PathSeparator);
+		}
+
+		const tchar* rootPaths[] = { rootPath, doubleSlash.data };
+		for (const tchar* rp : rootPaths)
+		{
+			u32 index = u32(m_roots.size());
+			if (index == '~' - ' ') // This is not really true.. as long as value is under 256 we're good
+				return logger.Error(TC("Too many roots added (%llu)"), index);
+
+			auto& root = m_roots.emplace_back();
+			root.index = index;
+			root.path = rp;
+
+			ToLower(root.path.data());
+			if (root.path[root.path.size()-1] != PathSeparator)
+				return logger.Error(TC("Root path must end with separator"));
+
+			root.includeInKey = includeInKey;
+
+			m_longestRoot = Max(u32(root.path.size()), m_longestRoot);
+
+			if (!m_shortestRoot || root.path.size() < m_shortestRoot)
+			{
+				m_shortestRoot = u32(root.path.size());
+				for (auto& r : m_roots)
+					r.shortestPathKey = ToStringKeyNoCheck(r.path.data(), m_shortestRoot);
+			}
+			else
+				root.shortestPathKey = ToStringKeyNoCheck(root.path.data(), m_shortestRoot);
+		}
+		return true;
+	}
+
+	bool RootPaths::RegisterSystemRoots(Logger& logger)
+	{
+		#if PLATFORM_WINDOWS
+		StringBuffer<MaxPath> dir;
+		dir.count = GetSystemDirectory(dir.data, dir.capacity);
+		RegisterRoot(logger, dir.EnsureEndsWithSlash().data, false); // Ignore files from here.. we do expect them not to affect the output of a process
+		
+		dir.count = GetEnvironmentVariable(TC("ProgramW6432"), dir.Clear().data, dir.capacity);
+		RegisterRoot(logger, dir.EnsureEndsWithSlash().data, true);
+
+		dir.count = GetEnvironmentVariable(TC("ProgramFiles(x86)"), dir.Clear().data, dir.capacity);
+		RegisterRoot(logger, dir.EnsureEndsWithSlash().data, true);
+
+		dir.count = GetEnvironmentVariable(TC("ProgramFiles(x86)"), dir.Clear().data, dir.capacity);
+		RegisterRoot(logger, dir.EnsureEndsWithSlash().data, true);
+
+		PWSTR path;
+		if (!SUCCEEDED(SHGetKnownFolderPath(FOLDERID_ProgramData, 0, NULL, &path)))
+			return false;
+		RegisterRoot(logger, dir.Clear().Append(path).EnsureEndsWithSlash().data, true);
+		CoTaskMemFree(path);
+
+		#else
+		UBA_ASSERT(false);
+		#endif
+		return true;
+	}
+
+	const RootPaths::Root* RootPaths::FindRoot(const StringBufferBase& path) const
+	{
+		if (path.count < m_shortestRoot)
+			return nullptr;
+
+		StringBuffer<MaxPath> shortPath;
+		shortPath.Append(path.data, m_shortestRoot).MakeLower();
+		StringKey key = ToStringKeyNoCheck(shortPath.data, m_shortestRoot);
+		for (u32 i=0, e=u32(m_roots.size()); i!=e; ++i)
+		{
+			auto& root = m_roots[i];
+			if (key != root.shortestPathKey)
+				continue;
+			if (!path.StartsWith(root.path.c_str()))
+				continue;
+			return &m_roots[i];
+		}
+		return nullptr;
+	}
+
+	const RootPaths::Root& RootPaths::GetRoot(u32 index) const
+	{
+		return m_roots[index];
+	}
+
+	CasKey RootPaths::NormalizeAndHashFile(Logger& logger, const tchar* filename) const
+	{
+		FileAccessor file(logger, filename);
+		if (!file.OpenMemoryRead())
+			return CasKeyZero;
+
+		CasKeyHasher hasher;
+		auto hashString = [&](const char* str, u64 strLen, u32 rootPos) { hasher.Update(str, strLen); };
+		if (!NormalizeString<char>(logger, (const char*)file.GetData(), file.GetSize(), hashString, filename))
+			return CasKeyZero;
+
+		return ToCasKey(hasher, false);
+	}
+
+}
\ No newline at end of file
diff --git a/Engine/Source/Programs/UnrealBuildAccelerator/Common/Private/UbaScheduler.cpp b/Engine/Source/Programs/UnrealBuildAccelerator/Common/Private/UbaScheduler.cpp
--- a/Engine/Source/Programs/UnrealBuildAccelerator/Common/Private/UbaScheduler.cpp
+++ b/Engine/Source/Programs/UnrealBuildAccelerator/Common/Private/UbaScheduler.cpp
@@ -6,6 +6,7 @@
 #include "UbaNetworkServer.h"
 #include "UbaProcess.h"
 #include "UbaProcessStartInfoHolder.h"
+#include "UbaRootPaths.h"
 #include "UbaSessionServer.h"
 #include "UbaStringBuffer.h"
 
@@ -95,6 +96,8 @@ namespace uba
 	Scheduler::~Scheduler()
 	{
 		Stop();
+		for (auto rt : m_rootPaths)
+			delete rt;
 	}
 
 	void Scheduler::Start()
@@ -403,7 +406,7 @@ namespace uba
 						ProcessStartInfo& si = exitInfo->startInfo->startInfo;
 						u64 startTime = GetTime();
 
-						if (m_cacheClient->FetchFromCache(si))
+						if (m_cacheClient->FetchFromCache(*m_rootPaths[0], si))
 						{
 							auto process = new CachedProcess(si);
 							ProcessHandle ph(process);
@@ -513,7 +516,7 @@ namespace uba
 		delete si;
 
 		if (m_writeToCache && exitCode == 0 && process.GetStartInfo().trackInputs)
-			m_cacheClient->WriteToCache(ph);
+			m_cacheClient->WriteToCache(*m_rootPaths[0], ph);
 
 		ph.m_process = nullptr;
 	}
@@ -708,13 +711,11 @@ namespace uba
 				}
 				case InsideArray_CacheRoots:
 				{
-					if (m_cacheClient)
-					{
-						if (Equals(valueStart, TC("SystemRoots")))
-							m_cacheClient->RegisterSystemRoots();
-						else
-							m_cacheClient->RegisterRoot(valueStart);
-					}
+					auto& rootPaths = *m_rootPaths.emplace_back(new RootPaths());
+					if (Equals(valueStart, TC("SystemRoots")))
+						rootPaths.RegisterSystemRoots(logger);
+					else
+						rootPaths.RegisterRoot(logger, valueStart);
 					return true;
 				}
 				case InsideArray_Processes:
diff --git a/Engine/Source/Programs/UnrealBuildAccelerator/Common/Private/UbaSession.cpp b/Engine/Source/Programs/UnrealBuildAccelerator/Common/Private/UbaSession.cpp
--- a/Engine/Source/Programs/UnrealBuildAccelerator/Common/Private/UbaSession.cpp
+++ b/Engine/Source/Programs/UnrealBuildAccelerator/Common/Private/UbaSession.cpp
@@ -1217,7 +1217,7 @@ namespace uba
 			si.rules = GetRules(si);
 
 		void* env = GetProcessEnvironmentVariables();
-		u32 id = ++m_processIdCounter;
+		u32 id = CreateProcessId();
 		auto process = new ProcessImpl(*this, id, parent);
 		ProcessHandle h(process);
 		process->Start(startInfo, realApplication.data, realWorkingDir, m_runningRemote, env, async, enableDetour);
@@ -1311,6 +1311,7 @@ namespace uba
 	Storage& Session::GetStorage() { return m_storage; }
 	Logger& Session::GetLogger() { return m_logger; }
 	LogWriter& Session::GetLogWriter() { return m_logger.m_writer; }
+	Trace& Session::GetTrace() { return m_trace; }
 
 	const ApplicationRules* Session::GetRules(const ProcessStartInfo& si)
 	{
@@ -1382,6 +1383,10 @@ namespace uba
 	{
 		return m_tempPath.data;
 	}
+	u32 Session::CreateProcessId()
+	{
+		return ++m_processIdCounter;
+	}
 
 	void Session::ProcessAdded(Process& process, u32 sessionId)
 	{
diff --git a/Engine/Source/Programs/UnrealBuildAccelerator/Common/Private/UbaSessionServer.cpp b/Engine/Source/Programs/UnrealBuildAccelerator/Common/Private/UbaSessionServer.cpp
--- a/Engine/Source/Programs/UnrealBuildAccelerator/Common/Private/UbaSessionServer.cpp
+++ b/Engine/Source/Programs/UnrealBuildAccelerator/Common/Private/UbaSessionServer.cpp
@@ -253,7 +253,7 @@ namespace uba
 	{
 		FlushDeadProcesses();
 		ValidateStartInfo(startInfo);
-		u32 processId = ++m_processIdCounter;
+		u32 processId = CreateProcessId();
 		RemoteProcess* remoteProcess = new RemoteProcess(this, startInfo, processId, weight);
 		
 		if (knownInputsCount)
@@ -451,7 +451,7 @@ namespace uba
 
 	u32 SessionServer::BeginExternalProcess(const tchar* description)
 	{
-		u32 processId = ++m_processIdCounter;
+		u32 processId = CreateProcessId();
 		m_trace.ProcessAdded(0, processId, description);
 		return processId;
 	}
diff --git a/Engine/Source/Programs/UnrealBuildAccelerator/Common/Private/UbaStats.cpp b/Engine/Source/Programs/UnrealBuildAccelerator/Common/Private/UbaStats.cpp
--- a/Engine/Source/Programs/UnrealBuildAccelerator/Common/Private/UbaStats.cpp
+++ b/Engine/Source/Programs/UnrealBuildAccelerator/Common/Private/UbaStats.cpp
@@ -56,6 +56,14 @@ namespace uba
 			logger.Info(TC(""));
 	}
 
+	bool SystemStats::IsEmpty()
+	{
+		#define UBA_SYSTEM_STAT(T) if (T.count) return false;
+		UBA_SYSTEM_STATS
+		#undef UBA_SYSTEM_STAT
+		return true;
+	}
+
 	void SystemStats::Add(const SystemStats& other)
 	{
 		#define UBA_SYSTEM_STAT(var) var += other.var;
@@ -131,6 +139,14 @@ namespace uba
 			logger.Info(TC("  DecompressToMem      %6u %9s"), decompressToMem.count.load(), TimeToText(decompressToMem.time, false, frequency).str);
 	}
 
+	bool StorageStats::IsEmpty()
+	{
+		#define UBA_STORAGE_STAT(type, var) if (var != type()) return false;
+		UBA_STORAGE_STATS
+		#undef UBA_STORAGE_STAT
+		return true;
+	}
+
 	thread_local StorageStats* t_storageStats;
 
 	StorageStats* StorageStats::GetCurrent()
@@ -189,6 +205,14 @@ namespace uba
 		#undef UBA_SESSION_STAT
 	}
 
+	bool SessionStats::IsEmpty()
+	{
+		#define UBA_SESSION_STAT(type, var, ver) if (var.count) return false;
+		UBA_SESSION_STATS
+		#undef UBA_SESSION_STAT
+		return true;
+	}
+
 	thread_local SessionStats* t_sessionStats;
 
 	SessionStats* SessionStats::GetCurrent()
@@ -233,4 +257,33 @@ namespace uba
 		logger.Info(TC("  ProcessesKilled             %9llu"), killCount);
 		logger.Info(TC(""));
 	}
+
+	void CacheStats::Write(BinaryWriter& writer)
+	{
+		#define UBA_CACHE_STAT(type, var, ver) uba::Write(writer, var);
+		UBA_CACHE_STATS
+		#undef UBA_CACHE_STAT
+	}
+
+	void CacheStats::Read(BinaryReader& reader, u32 version)
+	{
+		#define UBA_CACHE_STAT(type, var, ver) if (ver <= version) uba::Read(reader, var);
+		UBA_CACHE_STATS
+		#undef UBA_CACHE_STAT
+	}
+
+	void CacheStats::Print(Logger& logger, u64 frequency)
+	{
+		#define UBA_CACHE_STAT(type, var, ver) LogStat(logger, #var, var, frequency);
+		UBA_CACHE_STATS
+		#undef UBA_CACHE_STAT
+	}
+
+	bool CacheStats::IsEmpty()
+	{
+		#define UBA_CACHE_STAT(type, var, ver) if (var.count) return false;
+		UBA_CACHE_STATS
+		#undef UBA_CACHE_STAT
+		return true;
+	}
 }
diff --git a/Engine/Source/Programs/UnrealBuildAccelerator/Common/Private/UbaTrace.cpp b/Engine/Source/Programs/UnrealBuildAccelerator/Common/Private/UbaTrace.cpp
--- a/Engine/Source/Programs/UnrealBuildAccelerator/Common/Private/UbaTrace.cpp
+++ b/Engine/Source/Programs/UnrealBuildAccelerator/Common/Private/UbaTrace.cpp
@@ -339,6 +339,21 @@ namespace uba
 		writer.WriteByte(statusType);
 	}
 
+	void Trace::CacheBeginFetch(u32 fetchId, const tchar* description)
+	{
+		BEGIN_TRACE_ENTRY(TraceType_CacheBeginFetch);
+		writer.Write7BitEncoded(fetchId);
+		writer.WriteString(description);
+	}
+
+	void Trace::CacheEndFetch(u32 fetchId, bool success, const u8* data, u64 dataSize)
+	{
+		BEGIN_TRACE_ENTRY(TraceType_CacheEndFetch);
+		writer.Write7BitEncoded(fetchId);
+		writer.WriteBool(success);
+		writer.WriteBytes(data, dataSize);
+	}
+
 	TraceChannel::TraceChannel(Logger& logger) : m_logger(logger)
 	{
 	}
diff --git a/Engine/Source/Programs/UnrealBuildAccelerator/Common/Public/UbaCacheClient.h b/Engine/Source/Programs/UnrealBuildAccelerator/Common/Public/UbaCacheClient.h
--- a/Engine/Source/Programs/UnrealBuildAccelerator/Common/Public/UbaCacheClient.h
+++ b/Engine/Source/Programs/UnrealBuildAccelerator/Common/Public/UbaCacheClient.h
@@ -9,6 +9,7 @@ namespace uba
 {
 	class NetworkClient;
 	class ProcessHandle;
+	class RootPaths;
 	class Session;
 	class StorageImpl;
 	struct ProcessStartInfo;
@@ -19,31 +20,24 @@ namespace uba
 		CacheClient(LogWriter& writer, StorageImpl& storage, NetworkClient& client, Session& session);
 		~CacheClient();
 
-		bool RegisterRoot(const tchar* rootPath, bool includeInKey = true);
-		bool RegisterSystemRoots();
-
-		bool WriteToCache(const ProcessHandle& process);
-		bool FetchFromCache(const ProcessStartInfo& info);
+		bool WriteToCache(const RootPaths& rootPaths, const ProcessHandle& process);
+		bool FetchFromCache(const RootPaths& rootPaths, const ProcessStartInfo& info);
 
 		bool WriteCacheSummary(const tchar* destinationFile, const tchar* filterString = nullptr);
 
+		inline MutableLogger& GetLogger() { return m_logger; }
 		inline NetworkClient& GetClient() { return m_client; }
 		inline StorageImpl& GetStorage() { return m_storage; }
 
 	private:
 		bool SendPathTable(u32 requiredPathTableSize);
 		bool SendCasTable(u32 requiredCasTableSize);
-		bool SendCacheEntry(const CasKey& cmdKey, const Map<u32, u32>& inputsStringToCasKey, const Map<u32, u32>& outputsStringToCasKey);
+		bool SendCacheEntry(const RootPaths& rootPaths, const CasKey& cmdKey, const Map<u32, u32>& inputsStringToCasKey, const Map<u32, u32>& outputsStringToCasKey);
 		bool FetchCasTable();
 
-		CasKey GetCmdKey(const ProcessStartInfo& info);
-		u32 FindRootIndex(const StringBufferBase& path);
-
-		template<typename CharType, typename Func>
-		bool NormalizeString(const CharType* str, u64 strLen, const Func& func, const tchar* hint);
+		CasKey GetCmdKey(const RootPaths& rootPaths, const ProcessStartInfo& info);
 
-		CasKey NormalizeAndHashFile(const tchar* rspFile);
-		bool GetLocalPathAndCasKey(StringBufferBase& outPath, CasKey& outKey, CompactCasKeyTable& casKeyTable, CompactPathTable& pathTable, u32 offset);
+		bool GetLocalPathAndCasKey(const RootPaths& rootPaths, StringBufferBase& outPath, CasKey& outKey, CompactCasKeyTable& casKeyTable, CompactPathTable& pathTable, u32 offset);
 
 		MutableLogger m_logger;
 		StorageImpl& m_storage;
@@ -52,11 +46,6 @@ namespace uba
 
 		Atomic<bool> m_connected;
 
-		struct Root;
-		Vector<Root> m_roots;
-		u32 m_shortestRoot = 0;
-		u32 m_longestRoot = 0;
-
 		CompactPathTable m_serverPathTable;
 		CompactCasKeyTable m_serverCasKeyTable;
 
diff --git a/Engine/Source/Programs/UnrealBuildAccelerator/Common/Public/UbaCompactTables.h b/Engine/Source/Programs/UnrealBuildAccelerator/Common/Public/UbaCompactTables.h
--- a/Engine/Source/Programs/UnrealBuildAccelerator/Common/Public/UbaCompactTables.h
+++ b/Engine/Source/Programs/UnrealBuildAccelerator/Common/Public/UbaCompactTables.h
@@ -14,7 +14,9 @@ namespace uba
 	class CompactPathTable
 	{
 	public:
-		CompactPathTable(u64 reserveSize, u64 reserveOffsetsCount = 0);
+		enum Version : u8 { V0, V1 };
+
+		CompactPathTable(u64 reserveSize, Version version, u64 reservePathCount = 0, u64 reserveSegmentCount = 0);
 
 		u32 Add(const tchar* str, u64 strLen, u32* outRequiredCasTableSize = nullptr);
 		u32 AddNoLock(const tchar* str, u64 strLen);
@@ -27,12 +29,17 @@ namespace uba
 		void ReadMem(BinaryReader& reader, bool populateLookup);
 		void Swap(CompactPathTable& other);
 
+		u64 GetPathCount() { return m_offsets.size(); }
+		u64 GetSegmentCount() { return m_segmentOffsets.size(); }
+
 	private:
 		u32 InternalAdd(const tchar* str, const tchar* stringKeyString, u64 strLen);
 		ReaderWriterLock m_lock;
 		MemoryBlock m_mem;
 		UnorderedMap<StringKey, u32> m_offsets;
+		UnorderedMap<StringKey, u32> m_segmentOffsets;
 		u64 m_reserveSize;
+		Version m_version;
 	};
 
 	class CompactCasKeyTable
diff --git a/Engine/Source/Programs/UnrealBuildAccelerator/Common/Public/UbaRootPaths.h b/Engine/Source/Programs/UnrealBuildAccelerator/Common/Public/UbaRootPaths.h
new file mode 100644
--- /dev/null
+++ b/Engine/Source/Programs/UnrealBuildAccelerator/Common/Public/UbaRootPaths.h
@@ -0,0 +1,97 @@
+// Copyright Epic Games, Inc. All Rights Reserved.
+
+#pragma once
+
+#include "UbaHash.h"
+#include "UbaLogger.h"
+
+namespace uba
+{
+	class RootPaths
+	{
+	public:
+
+		bool RegisterRoot(Logger& logger, const tchar* rootPath, bool includeInKey = true);
+		bool RegisterSystemRoots(Logger& logger);
+
+		struct Root
+		{
+			TString path;
+			StringKey shortestPathKey;
+			u32 index;
+			bool includeInKey;
+		};
+
+		const Root* FindRoot(const StringBufferBase& path) const;
+		const Root& GetRoot(u32 index) const;
+
+		template<typename CharType, typename Func>
+		bool NormalizeString(Logger& logger, const CharType* str, u64 strLen, const Func& func, const tchar* hint) const;
+
+		CasKey NormalizeAndHashFile(Logger& logger, const tchar* filename) const;
+
+		static constexpr u8 RootStartByte = ' ';
+
+	private:
+		Vector<Root> m_roots;
+		u32 m_shortestRoot = 0;
+		u32 m_longestRoot = 0;
+	};
+
+
+
+	template<typename CharType, typename Func>
+	bool RootPaths::NormalizeString(Logger& logger, const CharType* str, u64 strLen, const Func& func, const tchar* hint) const
+	{
+		auto strEnd = str + strLen;
+		auto searchPos = str;
+
+		u32 destPos = 0;
+
+		while (true)
+		{
+			auto absPathChars = searchPos;
+			CharType lastChar = 0;
+			while (absPathChars < strEnd && !(lastChar == ':' && *absPathChars == '\\'))
+			{
+				lastChar = *absPathChars;
+				++absPathChars;
+			}
+		
+			if (absPathChars == strEnd)
+			{
+				func(searchPos, strEnd - searchPos, ~0u);
+				return true;
+			}
+
+			auto pathStart = absPathChars - 2;
+
+			auto pathEndOrMore = pathStart;
+			while (pathEndOrMore < strEnd && *pathEndOrMore != '\n')
+				++pathEndOrMore;
+
+			u32 lenOrMore = u32(pathEndOrMore - pathStart);
+			u32 toCopy = Min(lenOrMore, m_longestRoot);
+			StringBuffer<512> path;
+			path.Append(pathStart, toCopy);
+
+			auto root = FindRoot(path);
+			if (!root)
+			{
+				logger.Info(TC("PATH WITHOUT ROOT: %s (inside file %s)"), path.data, hint);
+				return false;
+			}
+
+			if (u32 len = u32(pathStart - searchPos))
+			{
+				destPos += len;
+				func(searchPos, len, ~0u);
+			}
+			CharType temp = RootStartByte + CharType(root->index);
+			func(&temp, 1, destPos);
+			destPos += 1;
+
+			searchPos = pathStart + root->path.size();
+		}
+	}
+}
diff --git a/Engine/Source/Programs/UnrealBuildAccelerator/Common/Public/UbaScheduler.h b/Engine/Source/Programs/UnrealBuildAccelerator/Common/Public/UbaScheduler.h
--- a/Engine/Source/Programs/UnrealBuildAccelerator/Common/Public/UbaScheduler.h
+++ b/Engine/Source/Programs/UnrealBuildAccelerator/Common/Public/UbaScheduler.h
@@ -10,6 +10,7 @@ namespace uba
 {
 	class CacheClient;
 	class Process;
+	class RootPaths;
 	class SessionServer;
 	struct NextProcessInfo;
 	struct ProcessStartInfoHolder;
@@ -126,6 +127,7 @@ namespace uba
 		Atomic<u32> m_finishedProcesses;
 
 		CacheClient* m_cacheClient;
+		Vector<RootPaths*> m_rootPaths;
 		bool m_writeToCache;
 
 		Scheduler(const Scheduler&) = delete;
diff --git a/Engine/Source/Programs/UnrealBuildAccelerator/Common/Public/UbaSession.h b/Engine/Source/Programs/UnrealBuildAccelerator/Common/Public/UbaSession.h
--- a/Engine/Source/Programs/UnrealBuildAccelerator/Common/Public/UbaSession.h
+++ b/Engine/Source/Programs/UnrealBuildAccelerator/Common/Public/UbaSession.h
@@ -68,12 +68,15 @@ namespace uba
 		Storage& GetStorage(); // Storage (only used when remote machines are connected)
 		Logger& GetLogger(); // Logger used for logging 
 		LogWriter& GetLogWriter(); // LogWriter used by logger
+		Trace& GetTrace(); // Trace written to be session
 		const ApplicationRules* GetRules(const ProcessStartInfo& si); // Get application rules used for process
 		const tchar* GetTempPath(); // Path for temp files used for current session
 		bool ShouldStoreObjFilesCompressed() { return m_storeObjFilesCompressed; }
 
 		virtual ~Session();
 
+		u32 CreateProcessId();
+
 	protected:
 		Session(const SessionCreateInfo& info, const tchar* logPrefix, bool runningRemote, WorkManager* workManager = nullptr);
 		bool Create(const SessionCreateInfo& info);
diff --git a/Engine/Source/Programs/UnrealBuildAccelerator/Common/Public/UbaStats.h b/Engine/Source/Programs/UnrealBuildAccelerator/Common/Public/UbaStats.h
--- a/Engine/Source/Programs/UnrealBuildAccelerator/Common/Public/UbaStats.h
+++ b/Engine/Source/Programs/UnrealBuildAccelerator/Common/Public/UbaStats.h
@@ -41,6 +41,7 @@ namespace uba
 		void Read(BinaryReader& reader);
 		void Add(const StorageStats& other);
 		void Print(Logger& logger, u64 frequency = GetFrequency());
+		bool IsEmpty();
 		static StorageStats* GetCurrent();
 	};
 
@@ -77,6 +78,7 @@ namespace uba
 		void Read(BinaryReader& reader, u32 version);
 		void Add(const SessionStats& other);
 		void Print(Logger& logger, u64 frequency = GetFrequency());
+		bool IsEmpty();
 		static SessionStats* GetCurrent();
 	};
 
@@ -116,4 +118,21 @@ namespace uba
 		void Print(Logger& logger, u64 frequency = GetFrequency());
 	};
 
+	#define UBA_CACHE_STATS \
+		UBA_CACHE_STAT(Timer, fetchEntries, 0) \
+		UBA_CACHE_STAT(Timer, fetchCasTable, 0) \
+		UBA_CACHE_STAT(Timer, testEntries, 0) \
+		UBA_CACHE_STAT(Timer, fetchOutput, 0) \
+
+    struct CacheStats
+	{
+		#define UBA_CACHE_STAT(type, var, ver) type var;
+		UBA_CACHE_STATS
+		#undef UBA_CACHE_STAT
+
+		void Write(BinaryWriter& writer);
+		void Read(BinaryReader& reader, u32 version);
+		void Print(Logger& logger, u64 frequency = GetFrequency());
+		bool IsEmpty();
+	};
 }
diff --git a/Engine/Source/Programs/UnrealBuildAccelerator/Common/Public/UbaTrace.h b/Engine/Source/Programs/UnrealBuildAccelerator/Common/Public/UbaTrace.h
--- a/Engine/Source/Programs/UnrealBuildAccelerator/Common/Public/UbaTrace.h
+++ b/Engine/Source/Programs/UnrealBuildAccelerator/Common/Public/UbaTrace.h
@@ -52,12 +52,14 @@ namespace uba
 		TraceType_FileStoreLight,
 		TraceType_StatusUpdate,
 		TraceType_SessionNotification,
+		TraceType_CacheBeginFetch,
+		TraceType_CacheEndFetch,
 	};
 
 	using Color = u32;
 	inline Color toColor(u8 r, u8 g, u8 b) { return (r << 16) + (g << 8) + b; }
 
-	static constexpr u32 TraceVersion = 24;
+	static constexpr u32 TraceVersion = 25;
 	static constexpr u32 TraceReadCompatibilityVersion = 6;
 
 	class Trace : public WorkTracker
@@ -88,6 +90,9 @@ namespace uba
 		void EndWork(u32 workIndex);
 		void StatusUpdate(u32 statusIndex, u32 statusNameIndent, const tchar* statusName, u32 statusTextIndent, const tchar* statusText, LogEntryType statusType);
 
+		void CacheBeginFetch(u32 fetchId, const tchar* description);
+		void CacheEndFetch(u32 fetchId, bool success, const u8* data, u64 dataSize);
+
 		bool StopWrite(const tchar* writeFileName);
 
 		virtual u32 TrackWorkStart(const tchar* desc) final override;
diff --git a/Engine/Source/Programs/UnrealBuildAccelerator/Core/Public/UbaHash.h b/Engine/Source/Programs/UnrealBuildAccelerator/Core/Public/UbaHash.h
--- a/Engine/Source/Programs/UnrealBuildAccelerator/Core/Public/UbaHash.h
+++ b/Engine/Source/Programs/UnrealBuildAccelerator/Core/Public/UbaHash.h
@@ -195,7 +195,19 @@ namespace uba
 	#endif
 }
 
-template<> struct std::hash<uba::Guid> { size_t operator()(const uba::Guid& g) const { std::hash<uba::u64> hash; return hash(((uba::u64*)&g)[0]) ^ hash(((uba::u64*)&g)[1]); } };
-template<> struct std::hash<uba::StringKey> { size_t operator()(const uba::StringKey& g) const { return g.a ^ g.b; } };
+//template<> struct std::hash<uba::Guid> { size_t operator()(const uba::Guid& g) const { std::hash<uba::u64> hash; return hash(((uba::u64*)&g)[0]) ^ hash(((uba::u64*)&g)[1]); } };
+template<> struct std::hash<uba::StringKey> { size_t operator()(const uba::StringKey& g) const { return g.a; } };
 template<> struct std::hash<uba::CasKey> { size_t operator()(const uba::CasKey& g) const { return g.a ^ g.b ^ g.c; } };
 
+template <class Map> 
+void PrintMapInfo(const char* name, const Map& map)
+{
+    double l = map.size() / double(map.bucket_count());
+    double c = 0.0;
+    for (auto& kv : map)
+        c += map.bucket_size(map.bucket(kv.first));
+    c /= map.size();
+
+    double quality = 1.0 - std::max(0.0, c / (1 + l) - 1);
+	printf("%s Size: %llu Buckets: %llu Quality: %f\r\n", name, map.size(), map.bucket_count(), quality);
+}
\ No newline at end of file
diff --git a/Engine/Source/Programs/UnrealBuildAccelerator/Core/Public/UbaPathUtils.h b/Engine/Source/Programs/UnrealBuildAccelerator/Core/Public/UbaPathUtils.h
--- a/Engine/Source/Programs/UnrealBuildAccelerator/Core/Public/UbaPathUtils.h
+++ b/Engine/Source/Programs/UnrealBuildAccelerator/Core/Public/UbaPathUtils.h
@@ -202,7 +202,7 @@ namespace uba
 			++charLen;
 		}
 
-		UBA_ASSERT(charLen <= bufferCharCapacity);
+		UBA_ASSERTF(charLen <= bufferCharCapacity, TC("Buffer overflow (Capacity %llu) fixing path %s"), fileName);
 
 		if (outBufferCharLen)
 			*outBufferCharLen = u32(charLen - 1); // Remove terminator
diff --git a/Engine/Source/Programs/UnrealBuildAccelerator/Core/Public/UbaProcessStats.h b/Engine/Source/Programs/UnrealBuildAccelerator/Core/Public/UbaProcessStats.h
--- a/Engine/Source/Programs/UnrealBuildAccelerator/Core/Public/UbaProcessStats.h
+++ b/Engine/Source/Programs/UnrealBuildAccelerator/Core/Public/UbaProcessStats.h
@@ -149,6 +149,7 @@ namespace uba
 		void Read(BinaryReader& reader);
 		void Write(BinaryWriter& writer);
 		void Print(Logger& logger, bool writeHeader, u64 frequency = GetFrequency());
+		bool IsEmpty();
 		void Add(const SystemStats& other);
 		static SystemStats& GetCurrent();
 		static SystemStats& GetGlobal();
diff --git a/Engine/Source/Programs/UnrealBuildAccelerator/Core/Public/UbaTimer.h b/Engine/Source/Programs/UnrealBuildAccelerator/Core/Public/UbaTimer.h
--- a/Engine/Source/Programs/UnrealBuildAccelerator/Core/Public/UbaTimer.h
+++ b/Engine/Source/Programs/UnrealBuildAccelerator/Core/Public/UbaTimer.h
@@ -37,6 +37,7 @@ namespace uba
 		Atomic<u32> count;
 		void Add(const Timer& o) { time += o.time; count += o.count; }
 		void operator+=(const Timer& o) { time += o.time; count += o.count; }
+		bool operator==(const Timer& o) const { return time == o.time && count == o.count; }
 	};
 
 	struct ExtendedTimer : Timer
diff --git a/Engine/Source/Programs/UnrealBuildAccelerator/Host/Private/UbaExports.cpp b/Engine/Source/Programs/UnrealBuildAccelerator/Host/Private/UbaExports.cpp
--- a/Engine/Source/Programs/UnrealBuildAccelerator/Host/Private/UbaExports.cpp
+++ b/Engine/Source/Programs/UnrealBuildAccelerator/Host/Private/UbaExports.cpp
@@ -7,6 +7,7 @@
 #include "UbaNetworkBackendTcp.h"
 #include "UbaNetworkClient.h"
 #include "UbaProcess.h"
+#include "UbaRootPaths.h"
 #include "UbaScheduler.h"
 #include "UbaStorageServer.h"
 #include "UbaSessionServer.h"
@@ -66,6 +67,13 @@ namespace uba
 		NetworkBackend* backend;
 	};
 
+	class RootPathsWithLogger : public RootPaths
+	{
+	public:
+		RootPathsWithLogger(LogWriter& writer) : logger(writer) {}
+		LoggerWithWriter logger;
+	};
+
 	#define UBA_USE_SIGNALHANDLER 0//PLATFORM_LINUX // It might be that we can't use signal handlers in c# processes.. so don't set this to 1
 
 	#if UBA_USE_SIGNALHANDLER
@@ -106,7 +114,7 @@ extern "C"
 			delete writer;
 	}
 
-	uba::NetworkServer* CreateServer(uba::LogWriter& writer, uba::u32 workerCount, uba::u32 sendSize, uba::u32 receiveTimeoutSeconds, bool useQuic)
+	uba::NetworkServer* NetworkServer_Create(uba::LogWriter& writer, uba::u32 workerCount, uba::u32 sendSize, uba::u32 receiveTimeoutSeconds, bool useQuic)
 	{
 		using namespace uba;
 		NetworkBackend* networkBackend;
@@ -130,7 +138,7 @@ extern "C"
 		return nullptr;
 	}
 
-	void DestroyServer(uba::NetworkServer* server)
+	void NetworkServer_Destroy(uba::NetworkServer* server)
 	{
 		auto s = (uba::NetworkServerWithBackend*)server;
 		auto networkBackend = s->backend;
@@ -138,7 +146,7 @@ extern "C"
 		delete networkBackend;
 	}
 
-	bool Server_StartListen(uba::NetworkServer* server, int port, const uba::tchar* ip, const uba::tchar* crypto)
+	bool NetworkServer_StartListen(uba::NetworkServer* server, int port, const uba::tchar* ip, const uba::tchar* crypto)
 	{
 		using namespace uba;
 		u8 crypto128Data[16];
@@ -154,7 +162,7 @@ extern "C"
 		return s->StartListen(*s->backend, u16(port), ip, crypto128);
 	}
 
-	void Server_Stop(uba::NetworkServer* server)
+	void NetworkServer_Stop(uba::NetworkServer* server)
 	{
 		auto s = (uba::NetworkServerWithBackend*)server;
 		auto networkBackend = s->backend;
@@ -163,7 +171,7 @@ extern "C"
 	}
 
 
-	bool Server_AddClient(uba::NetworkServer* server, const uba::tchar* ip, int port, const uba::tchar* crypto)
+	bool NetworkServer_AddClient(uba::NetworkServer* server, const uba::tchar* ip, int port, const uba::tchar* crypto)
 	{
 		using namespace uba;
 		u8 crypto128Data[16];
@@ -179,42 +187,7 @@ extern "C"
 		return s->AddClient(*s->backend, ip, u16(port), crypto128);
 	}
 
-	/*
-	uba::NetworkClient* CreateClient(uba::LogWriter& writer, uba::u32 sendSize, uba::u32 receiveTimeoutSeconds)
-	{
-		return new uba::NetworkClient(writer, sendSize, receiveTimeoutSeconds);
-	}
-
-	void DestroyClient(uba::NetworkClient* client)
-	{
-		delete client;
-	}
-	*/
-
-	uba::Storage* CreateStorage(const uba::tchar* rootDir, uba::u64 casCapacityBytes, bool storeCompressed)
-	{
-		uba::StorageCreateInfo info(rootDir, uba::g_consoleLogWriter);
-		info.casCapacityBytes = casCapacityBytes;
-		info.storeCompressed = storeCompressed;
-		return new uba::StorageImpl(info, nullptr);
-	}
-
-	void DestroyStorage(uba::Storage* storage)
-	{
-		delete storage;
-	}
-
-	void Storage_SaveCasTable(uba::Storage* storage)
-	{
-		storage->SaveCasTable(true);
-	}
-
-	void Storage_DeleteFile(uba::Storage* storage, const uba::tchar* file)
-	{
-		storage->DeleteCasForFile(file);
-	}
-
-	uba::Storage* CreateStorageServer(uba::NetworkServer& server, const uba::tchar* rootDir, uba::u64 casCapacityBytes, bool storeCompressed, uba::LogWriter& writer, const uba::tchar* zone)
+	uba::StorageServer* StorageServer_Create(uba::NetworkServer& server, const uba::tchar* rootDir, uba::u64 casCapacityBytes, bool storeCompressed, uba::LogWriter& writer, const uba::tchar* zone)
 	{
 		using namespace uba;
 
@@ -246,12 +219,12 @@ extern "C"
 		return new StorageServer(info);
 	}
 
-	void DestroyStorageServer(uba::Storage* storageServer)
+	void StorageServer_Destroy(uba::StorageServer* storageServer)
 	{
 		delete storageServer;
 	}
 
-	void StorageServer_SaveCasTable(uba::Storage* storageServer)
+	void StorageServer_SaveCasTable(uba::StorageServer* storageServer)
 	{
 		storageServer->SaveCasTable(true);
 	}
@@ -261,21 +234,9 @@ extern "C"
 		storageServer->RegisterDisallowedPath(path);
 	}
 
-	/*
-uba::StorageClient* CreateStorageClient(uba::NetworkClient& client, const uba::tchar* rootDir, uba::u64 casCapacityBytes, bool storeCompressed, bool sendCompressed)
+	void StorageServer_DeleteFile(uba::StorageServer* storageServer, const uba::tchar* file)
 	{
-		uba::StorageClientCreateInfo info(client, rootDir);
-		info.casCapacityBytes = casCapacityBytes;
-		info.storeCompressed = storeCompressed;
-		info.sendCompressed = sendCompressed;
-		info.workManager = &client;
-		return new uba::StorageClient(info);
-	}
-	*/
-
-	void DestroySession(uba::Session* session)
-	{
-		delete session;
+		storageServer->DeleteCasForFile(file);
 	}
 
 	uba::u32 ProcessHandle_GetExitCode(uba::ProcessHandle* handle)
@@ -327,15 +288,7 @@ uba::StorageClient* CreateStorageClient(uba::NetworkClient& client, const uba::t
 		return &process.GetStartInfo();
 	}
 
-/*
-	uba::SessionClient* CreateSessionClient(const uba::SessionClientCreateInfo& info)
-	{
-		using namespace uba;
-		return new SessionClient(info);
-	}
-*/
-
-	uba::SessionServerCreateInfo* CreateSessionServerCreateInfo(uba::Storage& storage, uba::NetworkServer& client, uba::LogWriter& writer, const uba::tchar* rootDir, const uba::tchar* traceOutputFile, bool disableCustomAllocator, bool launchVisualizer, bool resetCas, bool writeToDisk, bool detailedTrace, bool allowWaitOnMem, bool allowKillOnMem, bool storeObjFilesCompressed)
+	uba::SessionServerCreateInfo* SessionServerCreateInfo_Create(uba::StorageServer& storage, uba::NetworkServer& client, uba::LogWriter& writer, const uba::tchar* rootDir, const uba::tchar* traceOutputFile, bool disableCustomAllocator, bool launchVisualizer, bool resetCas, bool writeToDisk, bool detailedTrace, bool allowWaitOnMem, bool allowKillOnMem, bool storeObjFilesCompressed)
 	{
 		auto info = new uba::SessionServerCreateInfo(storage, client, writer);
 		info->rootDir = TStrdup(rootDir);
@@ -353,14 +306,14 @@ uba::StorageClient* CreateStorageClient(uba::NetworkClient& client, const uba::t
 		return info;
 	}
 
-	void DestroySessionServerCreateInfo(uba::SessionServerCreateInfo* info)
+	void SessionServerCreateInfo_Destroy(uba::SessionServerCreateInfo* info)
 	{
 		free((void*)info->traceOutputFile);
 		free((void*)info->rootDir);
 		delete info;
 	}
 
-	uba::SessionServer* CreateSessionServer(const uba::SessionServerCreateInfo& info)
+	uba::SessionServer* SessionServer_Create(const uba::SessionServerCreateInfo& info)
 	{
 		return new uba::SessionServer(info);
 	}
@@ -446,7 +399,7 @@ uba::StorageClient* CreateStorageClient(uba::NetworkClient& client, const uba::t
 			});
 	}
 
-	void DestroySessionServer(uba::SessionServer* server)
+	void SessionServer_Destroy(uba::SessionServer* server)
 	{
 		if (server)
 		{
@@ -457,7 +410,32 @@ uba::StorageClient* CreateStorageClient(uba::NetworkClient& client, const uba::t
 		delete server;
 	}
 
-	uba::ProcessStartInfo* CreateProcessStartInfo(const uba::tchar* application, const uba::tchar* arguments, const uba::tchar* workingDir, const uba::tchar* description, uba::u32 priorityClass, uba::u64 outputStatsThresholdMs, bool trackInputs, const uba::tchar* logFile, ProcessHandle_ExitCallback* exit)
+	uba::RootPaths* RootPaths_Create(uba::LogWriter& writer)
+	{
+		return new uba::RootPathsWithLogger(writer);
+	}
+
+	bool RootPaths_RegisterRoot(uba::RootPaths* rootPaths, const uba::tchar* path, bool includeInKey)
+	{
+		using namespace uba;
+		auto rp = (uba::RootPathsWithLogger*)rootPaths;
+		return rp->RegisterRoot(rp->logger, path, includeInKey);
+	}
+
+	bool RootPaths_RegisterSystemRoots(uba::RootPaths* rootPaths)
+	{
+		using namespace uba;
+		auto rp = (uba::RootPathsWithLogger*)rootPaths;
+		return rp->RegisterSystemRoots(rp->logger);
+	}
+
+	void RootPaths_Destroy(uba::RootPaths* rootPaths)
+	{
+		delete (uba::RootPathsWithLogger*)rootPaths;
+	}
+
+
+	uba::ProcessStartInfo* ProcessStartInfo_Create(const uba::tchar* application, const uba::tchar* arguments, const uba::tchar* workingDir, const uba::tchar* description, uba::u32 priorityClass, uba::u64 outputStatsThresholdMs, bool trackInputs, const uba::tchar* logFile, ProcessHandle_ExitCallback* exit)
 	{
 		auto info = new uba::ProcessStartInfo();
 		info->application = TStrdup(application);
@@ -471,7 +449,8 @@ uba::StorageClient* CreateStorageClient(uba::NetworkClient& client, const uba::t
 		info->exitedFunc = exit;
 		return info;
 	}
-	void DestroyProcessStartInfo(uba::ProcessStartInfo* info)
+
+	void ProcessStartInfo_Destroy(uba::ProcessStartInfo* info)
 	{
 		free((void*)info->application);
 		free((void*)info->arguments);
@@ -541,16 +520,6 @@ uba::StorageClient* CreateStorageClient(uba::NetworkClient& client, const uba::t
 		return new CacheClient(writer, storage, *networkClient, *session);
 	}
 
-	bool CacheClient_RegisterRoot(uba::CacheClient* cacheClient, const uba::tchar* root, bool includeInKey)
-	{
-		return cacheClient->RegisterRoot(root, includeInKey);
-	}
-
-	bool CacheClient_RegisterSystemRoots(uba::CacheClient* cacheClient)
-	{
-		return cacheClient->RegisterSystemRoots();
-	}
-
 	bool CacheClient_Connect(uba::CacheClient* cacheClient, const uba::tchar* host, int port)
 	{
 		using namespace uba;
@@ -562,16 +531,19 @@ uba::StorageClient* CreateStorageClient(uba::NetworkClient& client, const uba::t
 		return true;
 	}
 
-	bool CacheClient_WriteToCache(uba::CacheClient* cacheClient, const uba::ProcessHandle* process)
+	bool CacheClient_WriteToCache(uba::CacheClient* cacheClient, uba::RootPaths* rootPaths, const uba::ProcessHandle* process)
 	{
-		if (!process->GetStartInfo().trackInputs)
+		using namespace uba;
+		auto& si = process->GetStartInfo();
+		if (!si.trackInputs)
 			return false;
-		return cacheClient->WriteToCache(*process);
+		return cacheClient->WriteToCache(*rootPaths, *process);
 	}
 
-	bool CacheClient_FetchFromCache(uba::CacheClient* cacheClient, const uba::ProcessStartInfo& info)
+	bool CacheClient_FetchFromCache(uba::CacheClient* cacheClient, uba::RootPaths* rootPaths, const uba::ProcessStartInfo& info)
 	{
-		return cacheClient->FetchFromCache(info);
+		using namespace uba;
+		return cacheClient->FetchFromCache(*rootPaths, info);
 	}
 
 	void CacheClient_Destroy(uba::CacheClient* cacheClient)
@@ -594,4 +566,25 @@ uba::StorageClient* CreateStorageClient(uba::NetworkClient& client, const uba::t
 		uba::FindImports(binary, [&](const uba::tchar* importName, bool isKnown) { func(importName, userData); });
 #endif
 	}
+
+
+	// Deprecated.. delete as soon as binaries are updated and code is using new api
+	UBA_DEP_API uba::NetworkServer* CreateServer(uba::LogWriter& writer, uba::u32 workerCount, uba::u32 sendSize, uba::u32 receiveTimeoutSeconds, bool useQuic) { return NetworkServer_Create(writer, workerCount, sendSize, receiveTimeoutSeconds, useQuic); }
+	UBA_DEP_API void DestroyServer(uba::NetworkServer* server) { NetworkServer_Destroy(server); }
+	UBA_DEP_API bool Server_StartListen(uba::NetworkServer* server, int port, const uba::tchar* ip, const uba::tchar* crypto) { return NetworkServer_StartListen(server, port, ip, crypto); }
+	UBA_DEP_API void Server_Stop(uba::NetworkServer* server) { NetworkServer_Stop(server); }
+	UBA_DEP_API bool Server_AddClient(uba::NetworkServer* server, const uba::tchar* ip, int port, const uba::tchar* crypto)	{ return NetworkServer_AddClient(server, ip, port, crypto); }
+	UBA_DEP_API uba::ProcessStartInfo* CreateProcessStartInfo(const uba::tchar* application, const uba::tchar* arguments, const uba::tchar* workingDir, const uba::tchar* description, uba::u32 priorityClass, uba::u64 outputStatsThresholdMs, bool trackInputs, const uba::tchar* logFile, ProcessHandle_ExitCallback* exit) { return ProcessStartInfo_Create(application, arguments, workingDir, description, priorityClass, outputStatsThresholdMs, trackInputs, logFile, exit); }
+	UBA_DEP_API void DestroyProcessStartInfo(uba::ProcessStartInfo* info) { ProcessStartInfo_Destroy(info); }
+	UBA_DEP_API uba::Storage* CreateStorageServer(uba::NetworkServer& server, const uba::tchar* rootDir, uba::u64 casCapacityBytes, bool storeCompressed, uba::LogWriter& writer, const uba::tchar* zone) { return StorageServer_Create(server, rootDir, casCapacityBytes, storeCompressed, writer, zone); }
+	UBA_DEP_API void DestroyStorageServer(uba::Storage* storageServer) { StorageServer_Destroy((uba::StorageServer*)storageServer); }
+	UBA_DEP_API uba::SessionServerCreateInfo* CreateSessionServerCreateInfo(uba::StorageServer& storage, uba::NetworkServer& client, uba::LogWriter& writer, const uba::tchar* rootDir, const uba::tchar* traceOutputFile,
+		bool disableCustomAllocator, bool launchVisualizer, bool resetCas, bool writeToDisk, bool detailedTrace, bool allowWaitOnMem = false, bool allowKillOnMem = false, bool storeObjFilesCompressed = false)
+	{
+		return SessionServerCreateInfo_Create(storage, client, writer, rootDir, traceOutputFile, disableCustomAllocator, launchVisualizer, resetCas, writeToDisk, detailedTrace, allowWaitOnMem, allowKillOnMem, storeObjFilesCompressed);
+	}
+	UBA_DEP_API void DestroySessionServerCreateInfo(uba::SessionServerCreateInfo* info) { SessionServerCreateInfo_Destroy(info); }
+	UBA_DEP_API uba::SessionServer* CreateSessionServer(const uba::SessionServerCreateInfo& info) { return SessionServer_Create(info); }
+	UBA_DEP_API void DestroySessionServer(uba::SessionServer* server) { SessionServer_Destroy(server); }
+	UBA_DEP_API void Storage_DeleteFile(uba::Storage* storage, const uba::tchar* file) { storage->DeleteCasForFile(file); }
 }
diff --git a/Engine/Source/Programs/UnrealBuildAccelerator/Host/Public/UbaExports.h b/Engine/Source/Programs/UnrealBuildAccelerator/Host/Public/UbaExports.h
--- a/Engine/Source/Programs/UnrealBuildAccelerator/Host/Public/UbaExports.h
+++ b/Engine/Source/Programs/UnrealBuildAccelerator/Host/Public/UbaExports.h
@@ -11,6 +11,7 @@ namespace uba
 	class CacheClient;
 	class Process;
 	class ProcessHandle;
+	class RootPaths;
 	class Session;
 	class SessionServer;
 	class Scheduler;
@@ -40,33 +41,37 @@ namespace uba
 
 extern "C"
 {
+	// LogWriter
 	UBA_API uba::LogWriter* GetDefaultLogWriter();
 	UBA_API uba::LogWriter* CreateCallbackLogWriter(uba::CallbackLogWriter::BeginScopeCallback begin, uba::CallbackLogWriter::EndScopeCallback end, uba::CallbackLogWriter::LogCallback log);
 	UBA_API void DestroyCallbackLogWriter(uba::LogWriter* writer);
 
-	UBA_API uba::NetworkServer* CreateServer(uba::LogWriter& writer = uba::g_consoleLogWriter, uba::u32 workerCount = 64, uba::u32 sendSize = uba::SendDefaultSize, uba::u32 receiveTimeoutSeconds = 60, bool useQuic = false);
-	UBA_API void DestroyServer(uba::NetworkServer* server);
-	UBA_API bool Server_StartListen(uba::NetworkServer* server, int port, const uba::tchar* ip, const uba::tchar* crypto = nullptr);
-	UBA_API void Server_Stop(uba::NetworkServer* server);
-	UBA_API bool Server_AddClient(uba::NetworkServer* server, const uba::tchar* ip, int port, const uba::tchar* crypto = nullptr);
-
-//	UBA_API uba::NetworkClient* CreateClient(uba::LogWriter& writer = uba::g_consoleLogWriter, uba::u32 sendSize = uba::SendDefaultSize, uba::u32 receiveTimeoutSeconds = 0);
-//	UBA_API void DestroyClient(uba::NetworkClient* client);
-
-	UBA_API uba::Storage* CreateStorage(const uba::tchar* rootDir, uba::u64 casCapacityBytes, bool storeCompressed);
-	UBA_API void DestroyStorage(uba::Storage* storage);
-	UBA_API void Storage_SaveCasTable(uba::Storage* storage);
-	UBA_API void Storage_DeleteFile(uba::Storage* storage, const uba::tchar* file);
-
-	UBA_API uba::Storage* CreateStorageServer(uba::NetworkServer& server, const uba::tchar* rootDir, uba::u64 casCapacityBytes, bool storeCompressed, uba::LogWriter& writer = uba::g_consoleLogWriter, const uba::tchar* zone = TC(""));
-	UBA_API void DestroyStorageServer(uba::Storage* storageServer);
-	UBA_API void StorageServer_SaveCasTable(uba::Storage* storageServer);
+	// NetworkServer
+	UBA_API uba::NetworkServer* NetworkServer_Create(uba::LogWriter& writer = uba::g_consoleLogWriter, uba::u32 workerCount = 64, uba::u32 sendSize = uba::SendDefaultSize, uba::u32 receiveTimeoutSeconds = 60, bool useQuic = false);
+	UBA_API void NetworkServer_Destroy(uba::NetworkServer* server);
+	UBA_API bool NetworkServer_StartListen(uba::NetworkServer* server, int port, const uba::tchar* ip, const uba::tchar* crypto = nullptr);
+	UBA_API void NetworkServer_Stop(uba::NetworkServer* server);
+	UBA_API bool NetworkServer_AddClient(uba::NetworkServer* server, const uba::tchar* ip, int port, const uba::tchar* crypto = nullptr);
+
+	// StorageServer
+	UBA_API uba::StorageServer* StorageServer_Create(uba::NetworkServer& server, const uba::tchar* rootDir, uba::u64 casCapacityBytes, bool storeCompressed, uba::LogWriter& writer = uba::g_consoleLogWriter, const uba::tchar* zone = TC(""));
+	UBA_API void StorageServer_Destroy(uba::StorageServer* storageServer);
+	UBA_API void StorageServer_SaveCasTable(uba::StorageServer* storageServer);
 	UBA_API void StorageServer_RegisterDisallowedPath(uba::StorageServer* storageServer, const uba::tchar* path);
+	UBA_API void StorageServer_DeleteFile(uba::StorageServer* storage, const uba::tchar* file);
 
-	//UBA_API uba::StorageClient* CreateStorageClient(uba::NetworkClient& client, const uba::tchar* rootDir, uba::u64 casCapacityBytes, bool storeCompressed, bool sendCompressed);
+	// RootPaths
+	UBA_API uba::RootPaths* RootPaths_Create(uba::LogWriter& writer);
+	UBA_API bool RootPaths_RegisterRoot(uba::RootPaths* rootPaths, const uba::tchar* path, bool includeInKey);
+	UBA_API bool RootPaths_RegisterSystemRoots(uba::RootPaths* rootPaths);
+	UBA_API void RootPaths_Destroy(uba::RootPaths* rootPaths);
 
-	UBA_API void DestroySession(uba::Session* session);
+	// ProcessStartInfo
+	using ProcessHandle_ExitCallback = void(void* userData, const uba::ProcessHandle&);
+	UBA_API uba::ProcessStartInfo* ProcessStartInfo_Create(const uba::tchar* application, const uba::tchar* arguments, const uba::tchar* workingDir, const uba::tchar* description, uba::u32 priorityClass, uba::u64 outputStatsThresholdMs, bool trackInputs, const uba::tchar* logFile, ProcessHandle_ExitCallback* exit);
+	UBA_API void ProcessStartInfo_Destroy(uba::ProcessStartInfo* info);
 
+	// ProcessHandle
 	UBA_API uba::u32 ProcessHandle_GetExitCode(uba::ProcessHandle* handle);
 	UBA_API const uba::tchar* ProcessHandle_GetExecutingHost(uba::ProcessHandle* handle);
 	UBA_API const uba::tchar* ProcessHandle_GetLogLine(const uba::ProcessHandle* handle, uba::u32 index);
@@ -77,14 +82,15 @@ extern "C"
 	UBA_API void DestroyProcessHandle(uba::ProcessHandle* handle);
 	UBA_API const uba::ProcessStartInfo* Process_GetStartInfo(uba::Process& process);
 
+	// SessionServer
 	using SessionServer_RemoteProcessAvailableCallback = void(void* userData);
 	using SessionServer_RemoteProcessReturnedCallback = void(uba::Process& process, void* userData);
 	using SessionServer_CustomServiceFunction = uba::u32(uba::ProcessHandle* handle, const void* recv, uba::u32 recvSize, void* send, uba::u32 sendCapacity, void* userData);
 
-	UBA_API uba::SessionServerCreateInfo* CreateSessionServerCreateInfo(uba::Storage& storage, uba::NetworkServer& client, uba::LogWriter& writer, const uba::tchar* rootDir, const uba::tchar* traceOutputFile,
+	UBA_API uba::SessionServerCreateInfo* SessionServerCreateInfo_Create(uba::StorageServer& storage, uba::NetworkServer& client, uba::LogWriter& writer, const uba::tchar* rootDir, const uba::tchar* traceOutputFile,
 		bool disableCustomAllocator, bool launchVisualizer, bool resetCas, bool writeToDisk, bool detailedTrace, bool allowWaitOnMem = false, bool allowKillOnMem = false, bool storeObjFilesCompressed = false);
-	UBA_API void DestroySessionServerCreateInfo(uba::SessionServerCreateInfo* info);
-	UBA_API uba::SessionServer* CreateSessionServer(const uba::SessionServerCreateInfo& info);
+	UBA_API void SessionServerCreateInfo_Destroy(uba::SessionServerCreateInfo* info);
+	UBA_API uba::SessionServer* SessionServer_Create(const uba::SessionServerCreateInfo& info);
 	UBA_API void SessionServer_SetRemoteProcessAvailable(uba::SessionServer* server, SessionServer_RemoteProcessAvailableCallback* available, void* userData);
 	UBA_API void SessionServer_SetRemoteProcessReturned(uba::SessionServer* server, SessionServer_RemoteProcessReturnedCallback* returned, void* userData);
 	UBA_API void SessionServer_RefreshDirectory(uba::SessionServer* server, const uba::tchar* directory);
@@ -103,15 +109,7 @@ extern "C"
 	UBA_API void SessionServer_EndExternalProcess(uba::SessionServer* server, uba::u32 id, uba::u32 exitCode);
 	UBA_API void SessionServer_UpdateStatus(uba::SessionServer* server, uba::u32 statusIndex, uba::u32 statusNameIndent, const uba::tchar* statusName, uba::u32 statusTextIndent, const uba::tchar* statusText, uba::LogEntryType statusType);
 	UBA_API void SessionServer_RegisterCustomService(uba::SessionServer* server, SessionServer_CustomServiceFunction* function, void* userData = nullptr);
-	UBA_API void DestroySessionServer(uba::SessionServer* server);
-
-	//UBA_API uba::SessionClient* CreateSessionClient(const uba::SessionClientCreateInfo& info);
-
-	using ProcessHandle_ExitCallback = void(void* userData, const uba::ProcessHandle&);
-
-	UBA_API uba::ProcessStartInfo* CreateProcessStartInfo(const uba::tchar* application, const uba::tchar* arguments, const uba::tchar* workingDir, const uba::tchar* description, uba::u32 priorityClass, uba::u64 outputStatsThresholdMs, bool trackInputs, const uba::tchar* logFile, ProcessHandle_ExitCallback* exit);
-	UBA_API void DestroyProcessStartInfo(uba::ProcessStartInfo* info);
-
+	UBA_API void SessionServer_Destroy(uba::SessionServer* server);
 
 	// Scheduler
 	UBA_API uba::Scheduler* Scheduler_Create(uba::SessionServer* session, uba::u32 maxLocalProcessors = ~0u, bool enableProcessReuse = false);
@@ -124,11 +122,9 @@ extern "C"
 
 	// Cache
 	UBA_API uba::CacheClient* CacheClient_Create(uba::SessionServer* session);
-	UBA_API bool CacheClient_RegisterRoot(uba::CacheClient* cacheClient, const uba::tchar* root, bool includeInKey);
-	UBA_API bool CacheClient_RegisterSystemRoots(uba::CacheClient* cacheClient);
 	UBA_API bool CacheClient_Connect(uba::CacheClient* cacheClient, const uba::tchar* host, int port);
-	UBA_API bool CacheClient_WriteToCache(uba::CacheClient* cacheClient, const uba::ProcessHandle* process);
-	UBA_API bool CacheClient_FetchFromCache(uba::CacheClient* cacheClient, const uba::ProcessStartInfo& info);
+	UBA_API bool CacheClient_WriteToCache(uba::CacheClient* cacheClient, uba::RootPaths* rootPaths, const uba::ProcessHandle* process);
+	UBA_API bool CacheClient_FetchFromCache(uba::CacheClient* cacheClient, uba::RootPaths* rootPaths, const uba::ProcessStartInfo& info);
 	UBA_API void CacheClient_Destroy(uba::CacheClient* cacheClient);
 
 	// Misc
@@ -137,4 +133,23 @@ extern "C"
 
 	using ImportFunc = void(const uba::tchar* importName, void* userData);
 	UBA_API void Uba_FindImports(const uba::tchar* binary, ImportFunc* func, void* userData);
+
+
+
+
+	// Deprecated.. delete as soon as binaries are updated and code is using new api
+	#define UBA_DEP_API UBA_API
+	UBA_DEP_API uba::NetworkServer* CreateServer(uba::LogWriter& writer = uba::g_consoleLogWriter, uba::u32 workerCount = 64, uba::u32 sendSize = uba::SendDefaultSize, uba::u32 receiveTimeoutSeconds = 60, bool useQuic = false);
+	UBA_DEP_API void DestroyServer(uba::NetworkServer* server);
+	UBA_DEP_API bool Server_StartListen(uba::NetworkServer* server, int port, const uba::tchar* ip, const uba::tchar* crypto = nullptr);
+	UBA_DEP_API void Server_Stop(uba::NetworkServer* server);
+	UBA_DEP_API bool Server_AddClient(uba::NetworkServer* server, const uba::tchar* ip, int port, const uba::tchar* crypto = nullptr);
+	UBA_DEP_API uba::ProcessStartInfo* CreateProcessStartInfo(const uba::tchar* application, const uba::tchar* arguments, const uba::tchar* workingDir, const uba::tchar* description, uba::u32 priorityClass, uba::u64 outputStatsThresholdMs, bool trackInputs, const uba::tchar* logFile, ProcessHandle_ExitCallback* exit);
+	UBA_DEP_API void DestroyProcessStartInfo(uba::ProcessStartInfo* info);
+	UBA_DEP_API uba::Storage* CreateStorageServer(uba::NetworkServer& server, const uba::tchar* rootDir, uba::u64 casCapacityBytes, bool storeCompressed, uba::LogWriter& writer = uba::g_consoleLogWriter, const uba::tchar* zone = TC(""));
+	UBA_DEP_API void DestroyStorageServer(uba::Storage* storageServer);
+	UBA_DEP_API void Storage_DeleteFile(uba::Storage* storage, const uba::tchar* file);
+	UBA_DEP_API uba::SessionServer* CreateSessionServer(const uba::SessionServerCreateInfo& info);
+	UBA_DEP_API void DestroySessionServer(uba::SessionServer* server);
 }
+
diff --git a/Engine/Source/Programs/UnrealBuildAccelerator/Test/Private/UbaTestAll.h b/Engine/Source/Programs/UnrealBuildAccelerator/Test/Private/UbaTestAll.h
--- a/Engine/Source/Programs/UnrealBuildAccelerator/Test/Private/UbaTestAll.h
+++ b/Engine/Source/Programs/UnrealBuildAccelerator/Test/Private/UbaTestAll.h
@@ -29,6 +29,7 @@ namespace uba
 		UBA_TEST(TestDetouredTestApp) \
 		UBA_TEST(TestRemoteDetouredTestApp) \
 		UBA_TEST(TestCompactPathTable) \
+		UBA_TEST(TestRootPaths) \
 
 
 #define UBA_POSIX_TESTS \
diff --git a/Engine/Source/Programs/UnrealBuildAccelerator/Test/Private/UbaTestBasics.h b/Engine/Source/Programs/UnrealBuildAccelerator/Test/Private/UbaTestBasics.h
--- a/Engine/Source/Programs/UnrealBuildAccelerator/Test/Private/UbaTestBasics.h
+++ b/Engine/Source/Programs/UnrealBuildAccelerator/Test/Private/UbaTestBasics.h
@@ -13,6 +13,7 @@
 #include "UbaEvent.h"
 #include "UbaFile.h"
 #include "UbaLogger.h"
+#include "UbaRootPaths.h"
 #include "UbaThread.h"
 #include "UbaTimer.h"
 #include "UbaDirectoryIterator.h"
@@ -311,22 +312,75 @@ namespace uba
 
 	bool TestCompactPathTable(Logger& logger, const StringBufferBase& rootDir)
 	{
-		CompactPathTable table(64*1024);
+		for (u32 i=0;i!=2; ++i)
+		{
+			CompactPathTable table(64*1024, CompactPathTable::Version(i));
+
+			StringBuffer<> str;
+			str.Append("foo").EnsureEndsWithSlash().Append("bar");
+			u32 offset = table.Add(str.data, str.count);
+
+			StringBuffer<> str2;
+			table.GetString(str2, offset);
+			if (!str.Equals(str2.data))
+				return false;
+
+			str.Clear().Append(PathSeparator).Append("foo").Append(PathSeparator).Append("bar");
+			offset = table.Add(str.data, str.count);
+			table.GetString(str2.Clear(), offset);
+			if (!str.Equals(str2.data))
+				return false;
+
+			CompactPathTable table2(64*1024, CompactPathTable::Version(i));
+			BinaryReader reader(table.GetMemory(), 0, table.GetSize());
+			table2.ReadMem(reader, true);
+			u32 offset2 = table2.Add(str.data, str.count);
+			if (offset != offset2)
+				return false;
+		}
+		return true;
+	}
 
-		StringBuffer<> str;
-		str.Append("foo").EnsureEndsWithSlash().Append("bar");
-		u32 offset = table.Add(str.data, str.count);
+	bool TestRootPaths(Logger& logger, const StringBufferBase& rootDir)
+	{
+		RootPaths paths;
+		if (!paths.RegisterRoot(logger, TC("c:\\temp\\")))
+			return false;
+		if (!paths.RegisterRoot(logger, TC("e:\\temp\\")))
+			return false;
+
+		tchar str[] = TC("e:\\temp\\foo");
 
-		StringBuffer<> str2;
-		table.GetString(str2, offset);
-		if (!str.Equals(str2.data))
+		bool success = true;
+		StringBuffer<> temp;
+		u32 rootPos = ~0u;
+		paths.NormalizeString(logger, str, sizeof_array(str), [&](const tchar* str, u64 strLen, u32 rp)
+			{
+				if (rp != ~0u)
+				{
+					if (strLen != 1)
+						success = false;
+					if (str[0] != RootPaths::RootStartByte+2)
+						success = false;
+					rootPos = str[0];
+				}
+				else
+				{
+					temp.Append(str, strLen);
+					if (!temp.Equals(TC("foo")))
+						success = false;
+				}
+			}, TC(""));
+
+		if (!success)
 			return false;
 
-		str.Clear().Append(PathSeparator).Append("foo").Append(PathSeparator).Append("bar");
-		offset = table.Add(str.data, str.count);
-		table.GetString(str2.Clear(), offset);
-		if (!str.Equals(str2.data))
+		StringBuffer<> newStr;
+		auto& root = paths.GetRoot(rootPos - RootPaths::RootStartByte);
+		newStr.Append(root.path).Append(temp);
+		if (!newStr.Equals(str))
 			return false;
+
 		return true;
 	}
 }
diff --git a/Engine/Source/Programs/UnrealBuildAccelerator/Visualizer/Private/UbaTraceReader.cpp b/Engine/Source/Programs/UnrealBuildAccelerator/Visualizer/Private/UbaTraceReader.cpp
--- a/Engine/Source/Programs/UnrealBuildAccelerator/Visualizer/Private/UbaTraceReader.cpp
+++ b/Engine/Source/Programs/UnrealBuildAccelerator/Visualizer/Private/UbaTraceReader.cpp
@@ -473,39 +473,7 @@ namespace uba
 				}
 			}
 
-			TraceView::Processor* processor = nullptr;
-			TraceView::Session& session = GetSession(out, sessionIndex);
-
-			u32 processorIndex = 0;
-			for (u32 i=0, e=u32(session.processors.size()); i!=e; ++i)
-			{
-				auto& it = session.processors[i];
-				if (it.processes.back().stop == ~u64(0))
-					continue;
-				processorIndex = i;
-				processor = &it;
-				break;
-			}
-			if (!processor)
-			{
-				processorIndex = u32(session.processors.size());
-				session.processors.emplace_back();
-				processor = &session.processors.back();
-			}
-
-			processor->processes.emplace_back();
-			auto& process = processor->processes.back();
-
-			m_activeProcesses.try_emplace(id, TraceView::ProcessLocation{sessionIndex, processorIndex, u32(processor->processes.size() - 1)});
-			
-			++session.processActiveCount;
-			++out.totalProcessActiveCount;
-
-			process.id = id;
-			process.description = desc.data;
-			process.start = time;
-			process.stop = ~u64(0);
-			process.exitCode = ~0u;
+			ProcessBegin(out, sessionIndex, id, time, desc.data);
 			break;
 		}
 		case TraceType_ProcessExited:
@@ -522,28 +490,32 @@ namespace uba
 				}
 			}
 
-			auto findIt = m_activeProcesses.find(id);
-			if (findIt == m_activeProcesses.end())
-				return false;
-			TraceView::ProcessLocation active = findIt->second;
-			m_activeProcesses.erase(findIt);
+			u32 sessionIndex;
+			TraceView::Process& process = *ProcessEnd(out, sessionIndex, id, time);
 
-			auto& session = GetSession(out, active.sessionIndex);
-			++session.processExitedCount;
-			--session.processActiveCount;
+			process.exitCode = exitCode;
 
-			++out.totalProcessExitedCount;
-			--out.totalProcessActiveCount;
+			ProcessStats processStats;
+			SessionStats sessionStats;
+			StorageStats storageStats;
+			SystemStats systemStats;
 
-			TraceView::Process& process = session.processors[active.processorIndex].processes[active.processIndex];
-			process.processStats.Read(reader, out.version);
-			bool isRemote = active.sessionIndex != 0;
+			const u8* dataStart = reader.GetPositionData();
+			processStats.Read(reader, out.version);
+
+			bool isRemote = sessionIndex != 0;
 			if (isRemote && out.version >= 7)
 			{
-				process.sessionStats.Read(reader, out.version);
-				process.storageStats.Read(reader);
-				process.systemStats.Read(reader);
+				sessionStats.Read(reader, out.version);
+				storageStats.Read(reader);
+				systemStats.Read(reader);
 			}
+			const u8* dataEnd = reader.GetPositionData();
+			process.stats.resize(dataEnd - dataStart);
+			memcpy(process.stats.data(), dataStart, dataEnd - dataStart);
+
+			process.createFilesTime = processStats.createFile.time;
+			process.writeFilesTime = Max(processStats.writeFiles.time, processStats.sendFiles.time);
 
 			if (out.version >= 22)
 			{
@@ -568,9 +540,6 @@ namespace uba
 				}
 			}
 
-			process.exitCode = exitCode;
-			process.stop = time;
-			process.bitmapDirty = true;
 			break;
 		}
 		case TraceType_ProcessEnvironmentUpdated:
@@ -588,10 +557,16 @@ namespace uba
 
 			auto& processes = session.processors[active.processorIndex].processes;
 			TraceView::Process& process = processes[active.processIndex];
-			process.processStats.Read(reader, out.version);
-			process.sessionStats.Read(reader, out.version);
-			process.storageStats.Read(reader);
-			process.systemStats.Read(reader);
+
+			ProcessStats processStats;
+			SessionStats sessionStats;
+			StorageStats storageStats;
+			SystemStats systemStats;
+			processStats.Read(reader, out.version);
+			sessionStats.Read(reader, out.version);
+			storageStats.Read(reader);
+			systemStats.Read(reader);
+
 			process.exitCode = 0u;
 			process.stop = time;
 			process.bitmapDirty = true;
@@ -908,6 +883,39 @@ namespace uba
 			out.strings.push_back(reader.ReadString());
 			break;
 		}
+		case TraceType_CacheBeginFetch:
+		{
+			u32 id = u32(reader.Read7BitEncoded());
+			StringBuffer<> desc;
+			reader.ReadString(desc);
+			ProcessBegin(out, 0, id, time, desc.data)->cacheFetch = true;
+			break;
+		}
+		case TraceType_CacheEndFetch:
+		{
+			u32 id = u32(reader.Read7BitEncoded());
+			bool success = reader.ReadBool();
+
+			u32 sessionIndex;
+			TraceView::Process& process = *ProcessEnd(out, sessionIndex, id, time);
+			process.exitCode = 0;//success ? 0 : -1;
+			process.returned = !success;
+
+			CacheStats cacheStats;
+			SystemStats systemStats;
+			StorageStats storageStats;
+			const u8* dataStart = reader.GetPositionData();
+			cacheStats.Read(reader, out.version);
+			if (success)
+			{
+				storageStats.Read(reader);
+				systemStats.Read(reader);
+			}
+			const u8* dataEnd = reader.GetPositionData();
+			process.stats.resize(dataEnd - dataStart);
+			memcpy(process.stats.data(), dataStart, dataEnd - dataStart);
+			break;
+		}
 		}
 		return true;
 	}
@@ -973,5 +981,66 @@ namespace uba
 		//UBA_ASSERTF(false, L"Failed to get session for clientUid %ls", GuidToString(clientUid).str);
 		return nullptr;
 	}
+
+	TraceView::Process* TraceReader::ProcessBegin(TraceView& out, u32 sessionIndex, u32 id, u64 time, const tchar* description)
+	{
+		TraceView::Processor* processor = nullptr;
+		TraceView::Session& session = GetSession(out, sessionIndex);
+
+		u32 processorIndex = 0;
+		for (u32 i=0, e=u32(session.processors.size()); i!=e; ++i)
+		{
+			auto& it = session.processors[i];
+			if (it.processes.back().stop == ~u64(0))
+				continue;
+			processorIndex = i;
+			processor = &it;
+			break;
+		}
+		if (!processor)
+		{
+			processorIndex = u32(session.processors.size());
+			session.processors.emplace_back();
+			processor = &session.processors.back();
+		}
+
+		processor->processes.emplace_back();
+		auto& process = processor->processes.back();
+
+		m_activeProcesses.try_emplace(id, TraceView::ProcessLocation{sessionIndex, processorIndex, u32(processor->processes.size() - 1)});
+			
+		++session.processActiveCount;
+		++out.totalProcessActiveCount;
+
+		process.id = id;
+		process.description = description;
+		process.start = time;
+		process.stop = ~u64(0);
+		process.exitCode = ~0u;
+		return &process;
+	}
+
+	TraceView::Process* TraceReader::ProcessEnd(TraceView& out, u32& outSessionIndex, u32 id, u64 time)
+	{
+		auto findIt = m_activeProcesses.find(id);
+		if (findIt == m_activeProcesses.end())
+			return nullptr;
+		TraceView::ProcessLocation active = findIt->second;
+		m_activeProcesses.erase(findIt);
+		outSessionIndex = active.sessionIndex;
+		auto& session = GetSession(out, active.sessionIndex);
+		++session.processExitedCount;
+		--session.processActiveCount;
+
+		++out.totalProcessExitedCount;
+		--out.totalProcessActiveCount;
+
+		TraceView::Process& process = session.processors[active.processorIndex].processes[active.processIndex];
+
+		process.stop = time;
+		process.bitmapDirty = true;
+		return &process;
+	}
+
 #endif // PLATFORM_WINDOWS
 }
\ No newline at end of file
diff --git a/Engine/Source/Programs/UnrealBuildAccelerator/Visualizer/Private/UbaTraceReader.h b/Engine/Source/Programs/UnrealBuildAccelerator/Visualizer/Private/UbaTraceReader.h
--- a/Engine/Source/Programs/UnrealBuildAccelerator/Visualizer/Private/UbaTraceReader.h
+++ b/Engine/Source/Programs/UnrealBuildAccelerator/Visualizer/Private/UbaTraceReader.h
@@ -24,10 +24,10 @@ namespace uba
 			u32 bitmapOffset = 0;
 			bool bitmapDirty = true;
 			bool returned = false;
-			ProcessStats processStats;
-			SessionStats sessionStats;
-			StorageStats storageStats;
-			SystemStats systemStats;
+			bool cacheFetch = false;
+			u64 createFilesTime = 0;
+			u64 writeFilesTime = 0;
+			Vector<u8> stats;
 			Vector<ProcessLogLine> logLines;
 		};
 
@@ -166,6 +166,9 @@ namespace uba
 		TraceView::Session& GetSession(TraceView& out, u32 sessionIndex);
 		TraceView::Session* GetSession(TraceView& out, const Guid& clientUid);
 
+		TraceView::Process* ProcessBegin(TraceView& out, u32 sessionIndex, u32 id, u64 time, const tchar* description);
+		TraceView::Process* ProcessEnd(TraceView& out, u32& outSessionIndex, u32 id, u64 time);
+
 		UnorderedMap<u32, TraceView::ProcessLocation> m_activeProcesses;
 
 		struct WorkRecordLocation { u32 track; u32 index; };
diff --git a/Engine/Source/Programs/UnrealBuildAccelerator/Visualizer/Private/UbaVisualizer.cpp b/Engine/Source/Programs/UnrealBuildAccelerator/Visualizer/Private/UbaVisualizer.cpp
--- a/Engine/Source/Programs/UnrealBuildAccelerator/Visualizer/Private/UbaVisualizer.cpp
+++ b/Engine/Source/Programs/UnrealBuildAccelerator/Visualizer/Private/UbaVisualizer.cpp
@@ -26,33 +26,104 @@ enum
 
 namespace uba
 {
-	class LineCountLogger : public Logger
-	{
-	public:
-		virtual void BeginScope() override {}
-		virtual void EndScope() override {}
-		virtual void Log(LogEntryType type, const wchar_t* str, u32 strLen) override { ++lineCount; }
-		u32 lineCount = 0;
-	};
-
 	class DrawTextLogger : public Logger
 	{
 	public:
-		DrawTextLogger(HDC h, const RECT& r, int fh) : hdc(h), rect(r), fontHeight(fh) {}
+		DrawTextLogger(HWND hw, HDC h, int fh, HBRUSH bb)
+		:	hwnd(hw)
+		,	hdc(h)
+		,	fontHeight(fh)
+		,	backgroundBrush(bb)
+		{
+			textColor = GetTextColor(hdc);
+		}
 
 		virtual void BeginScope() override {}
 		virtual void EndScope() override {}
 		virtual void Log(LogEntryType type, const wchar_t* str, u32 strLen) override
 		{
-			DrawTextW(hdc, str, strLen, &rect, DT_SINGLELINE);
-			rect.top += fontHeight;
+			RECT textRect{0,0,0,0};
+			DrawTextW(hdc, str, strLen, &textRect, DT_CALCRECT);
+
+			lines.emplace_back(TString(str, strLen), textOffset, height, textColor);
+			width = Max(width, int(textRect.right + textOffset));
+			height += fontHeight;
+		}
+
+		void AddSpace(int space = 5)
+		{
+			height += space;
+		}
+
+		void AddTextOffset(int offset)
+		{
+			textOffset += offset;
+		}
+
+		void AddWidth(int extra)
+		{
+			extraWidth += extra;
+		}
+
+		void DrawAtPos(int x, int y)
+		{
+			RECT r;
+			r.left = x;
+			r.top = y;
+			r.right = r.left + width;
+			r.bottom = r.top + height;
+
+			RECT clientRect;
+			GetClientRect(hwnd, &clientRect);
+
+			if (r.right > clientRect.right)
+				OffsetRect(&r, -width, 0);
+			if (r.bottom > clientRect.bottom)
+			{
+				OffsetRect(&r, 0, clientRect.bottom - r.bottom);
+				if (r.top < 0)
+					OffsetRect(&r, 0, -r.top);
+			}
+
+			RECT fillRect = r;
+			fillRect.right += 2 + extraWidth;
+			FillRect(hdc, &fillRect, backgroundBrush);
+
+			for (auto& line : lines)
+			{
+				RECT tr = r;
+				tr.left += line.left;
+				tr.top += line.top;
+				SetTextColor(hdc, line.color);
+				DrawTextW(hdc, line.str.data(), u32(line.str.size()), &tr, DT_SINGLELINE);
+			}
 		}
 
-		DrawTextLogger& SetColor(COLORREF c) { SetTextColor(hdc, c); return *this; }
+		void DrawAtCursor()
+		{
+			POINT p;
+			GetCursorPos(&p);
+			ScreenToClient(hwnd, &p);
+			p.x += 3;
+			p.y += 3;
+			DrawAtPos(p.x, p.y);
+		}
 
+		DrawTextLogger& SetColor(COLORREF c) { textColor = c; return *this; }
+
+		int width = 0;
+		int height = 0;
+		int textOffset = 2;
+		int extraWidth = 0;
+		struct Line { TString str; int left; int top; COLORREF color; };
+		Vector<Line> lines;
+
+		HWND hwnd;
 		HDC hdc;
-		RECT rect;
 		int fontHeight;
+		HBRUSH backgroundBrush;
+		COLORREF textColor;
+		bool isFirst = true;
 	};
 
 	class WriteTextLogger : public Logger
@@ -1150,30 +1221,15 @@ namespace uba
 		if (m_processSelected)
 		{
 			TraceView::Process& process = *m_traceView.GetProcess(m_processSelectedLocation);
-			bool hasStorageStats = false;
-			u32 lineCount = 4;
 			u64 duration = 0;
 			
-			int width = 290;
 			Vector<TString> logLines;
 			u32 maxCharCount = 50u;
 
 			bool hasExited = process.stop != ~u64(0);
 			if (hasExited)
 			{
-				LineCountLogger counter;
-				process.processStats.Print(counter, m_traceView.frequency);
-				u32 prevLineCount = counter.lineCount;
-				process.sessionStats.Print(counter, m_traceView.frequency);
-				process.storageStats.Print(counter, m_traceView.frequency);
-				process.systemStats.Print(counter, false, m_traceView.frequency);
-				hasStorageStats = prevLineCount != counter.lineCount;
-				lineCount = 6 + counter.lineCount;
-				if (hasStorageStats)
-					lineCount += 6;
 				duration = process.stop - process.start;
-				if (process.exitCode != 0)
-					++lineCount;
 
 				if (!process.logLines.empty())
 				{
@@ -1191,10 +1247,6 @@ namespace uba
 							left -= toCopy;
 						}
 					}
-
-					width = Max(width, int(lineMaxCount) * 7 + 2*14);
-
-					lineCount += u32(logLines.size()) + 1;
 				}
 			}
 			else
@@ -1202,91 +1254,113 @@ namespace uba
 				duration = playTime - process.start;
 			}
 
-			int height = lineCount*m_popupFontHeight;
-
-			POINT p;
-			GetClientCursorPos(p);
-			RECT r;
-			r.left = p.x;
-			r.top = p.y;
-			r.right = r.left + width;
-			r.bottom = r.top + height;
-
-			if (r.right > clientRect.right)
-				OffsetRect(&r, -width, 0);
-			if (r.bottom > clientRect.bottom)
-			{
-				OffsetRect(&r, 0, clientRect.bottom - r.bottom);
-				if (r.top < 0)
-					OffsetRect(&r, 0, -r.top);
-			}
-
-			FillRect(hdc, &r, m_tooltipBackgroundBrush);
-
-			r.top += 5;
 			SelectObject(hdc, m_popupFont);
-			DrawTextLogger logger(hdc, r, m_popupFontHeight);
+			DrawTextLogger logger(m_hwnd, hdc, m_popupFontHeight, m_tooltipBackgroundBrush);
 
+			logger.AddTextOffset(-10); // Remove spaces in the front
+			logger.AddWidth(3);
+
+			logger.AddSpace(2);
 			logger.Info(L"  %ls", process.description.c_str());
 			logger.Info(L"  Start:     %ls", TimeToText(process.start, true).str);
 			logger.Info(L"  Duration:  %ls", TimeToText(duration, true).str);
 			if (hasExited && process.exitCode != 0)
 				logger.Info(L"  ExitCode:  %u", process.exitCode);
-			logger.Info(L"");
 
 			if (process.stop != ~u64(0))
 			{
-				logger.Info(L"  ----------- Process stats -----------");
-				process.processStats.Print(logger, m_traceView.frequency);
-				if (hasStorageStats)
+				BinaryReader reader(process.stats.data(), 0, process.stats.size());
+				ProcessStats processStats;
+				SessionStats sessionStats;
+				StorageStats storageStats;
+				SystemStats systemStats;
+				CacheStats cacheStats;
+
+				if (process.cacheFetch)
+				{
+					cacheStats.Read(reader, m_traceView.version);
+					if (process.exitCode == 0)
+					{
+						storageStats.Read(reader);
+						systemStats.Read(reader);
+					}
+				}
+				else
+				{
+					processStats.Read(reader, m_traceView.version);
+
+					if (reader.GetLeft())
+					{
+						sessionStats.Read(reader, m_traceView.version);
+						storageStats.Read(reader);
+						systemStats.Read(reader);
+					}
+				}
+
+				if (processStats.hostTotalTime)
+				{
+					logger.Info(L"");
+					logger.Info(L"  ----------- Process stats -----------");
+					processStats.Print(logger, m_traceView.frequency);
+				}
+
+				if (!sessionStats.IsEmpty())
 				{
 					logger.Info(L"");
 					logger.Info(L"  ----------- Session stats -----------");
-					process.sessionStats.Print(logger, m_traceView.frequency);
+					sessionStats.Print(logger, m_traceView.frequency);
+				}
+
+				if (!storageStats.IsEmpty())
+				{
 					logger.Info(L"");
 					logger.Info(L"  ----------- Storage stats -----------");
-					process.storageStats.Print(logger, m_traceView.frequency);
+					storageStats.Print(logger, m_traceView.frequency);
+				}
+
+				if (!cacheStats.IsEmpty())
+				{
+					logger.Info(L"");
+					logger.Info(L"  ------------ Cache stats ------------");
+					cacheStats.Print(logger, m_traceView.frequency);
+				}
+
+				if (!systemStats.IsEmpty())
+				{
 					logger.Info(L"");
 					logger.Info(L"  ----------- System stats ------------");
-					process.systemStats.Print(logger, false, m_traceView.frequency);
+					systemStats.Print(logger, false, m_traceView.frequency);
 				}
 
 				if (!logLines.empty())
 				{
+					logger.Info(L"");
 					logger.Info(L"  ---------------- Log ----------------");
-					logger.rect.left += 14;
+					logger.AddTextOffset(14);
 					for (auto& line : logLines)
 						logger.Log(LogEntryType_Info, line.c_str(), u32(line.size()));
 				}
 			}
+			logger.AddSpace(3);
+			logger.DrawAtCursor();
 		}
 		else if (m_workSelected && selectedWork.description)
 		{
-			int width = 290;
-			int height = 3*m_popupFontHeight;
-			POINT p;
-			GetClientCursorPos(p);
-			RECT r;
-			r.left = p.x;
-			r.top = p.y;
-			r.right = r.left + width;
-			r.bottom = r.top + height;
-
 			u64 duration;
 			if (selectedWork.stop != ~u64(0))
 				duration = selectedWork.stop - selectedWork.start;
 			else
 				duration = playTime - selectedWork.start;
 
-			r.bottom += 10;
-			FillRect(hdc, &r, m_tooltipBackgroundBrush);
-
-			r.top += 5;
 			SelectObject(hdc, m_popupFont);
-			DrawTextLogger logger(hdc, r, m_popupFontHeight);
+			DrawTextLogger logger(m_hwnd, hdc, m_popupFontHeight, m_tooltipBackgroundBrush);
+
+			logger.AddSpace();
 			logger.Info(L"  %ls", selectedWork.description);
 			logger.Info(L"  Start:     %ls", TimeToText(selectedWork.start, true).str);
 			logger.Info(L"  Duration:  %ls", TimeToText(duration, true).str);
+			logger.AddSpace();
+			logger.DrawAtCursor();
 		}
 		else if (m_sessionSelectedIndex != ~0u)
 		{
@@ -1335,39 +1409,17 @@ namespace uba
 		}
 		else if (m_statsSelected)
 		{
-			int width = 160;
-			int lineCount = m_stats.ping ? 5 : 4;
-			int height = lineCount * m_popupFontHeight + 6;
-
-			POINT p;
-			GetCursorPos(&p);
-			ScreenToClient(m_hwnd, &p);
-			RECT r;
-			r.left = p.x;
-			r.top = p.y;
-			r.right = r.left + width;
-			r.bottom = r.top + height;
-
-			if (r.right > clientRect.right)
-				OffsetRect(&r, -width, 0);
-			if (r.bottom > clientRect.bottom)
-			{
-				OffsetRect(&r, 0, -height);
-				if (r.top < 0)
-					OffsetRect(&r, 0, -r.top);
-			}
-			FillRect(hdc, &r, m_tooltipBackgroundBrush);
-
-			r.top += 3;
-
 			SelectObject(hdc, m_popupFont);
-			DrawTextLogger logger(hdc, r, m_popupFontHeight);
+			DrawTextLogger logger(m_hwnd, hdc, m_popupFontHeight, m_tooltipBackgroundBrush);
+			logger.AddSpace(3);
 			logger.SetColor(m_cpuColor).Info(L"  Cpu: %.1f%%", m_stats.cpuLoad * 100.0f);
 			logger.SetColor(m_memColor).Info(L"  Mem: %ls/%ls", BytesToText(m_stats.memTotal - m_stats.memAvail).str, BytesToText(m_stats.memTotal).str);
 			logger.SetColor(m_recvColor).Info(L"  Recv: %ls/s", BytesToText(m_stats.recvBytesPerSecond).str);
 			logger.SetColor(m_sendColor).Info(L"  Send: %ls/s", BytesToText(m_stats.sendBytesPerSecond).str);
 			if (m_stats.ping)
 				logger.Info(L"  Ping: %ls", TimeToText(m_stats.ping, false, m_traceView.frequency).str);
+			logger.AddSpace(3);
+			logger.DrawAtCursor();
 		}
 		else if (m_buttonSelected != ~0u)
 		{
@@ -1381,33 +1433,11 @@ namespace uba
 				L"workers (threads on host taking care of requests from helpers)",
 			};
 
-			int width = int(wcslen(tooltip[m_buttonSelected])+5) * 7 + 5;
-			int lineCount = 1;
-			int height = lineCount * m_popupFontHeight;
-
-			POINT p;
-			GetCursorPos(&p);
-			ScreenToClient(m_hwnd, &p);
-			RECT r;
-			r.left = p.x;
-			r.top = p.y;
-			r.right = r.left + width;
-			r.bottom = r.top + height;
-
-			if (r.right > clientRect.right)
-				OffsetRect(&r, -width, 0);
-			if (r.bottom > clientRect.bottom)
-			{
-				OffsetRect(&r, 0, -height);
-				if (r.top < 0)
-					OffsetRect(&r, 0, -r.top);
-			}
-			FillRect(hdc, &r, m_tooltipBackgroundBrush);
-
 			bool vis = m_visibleComponents[m_buttonSelected];
 			SelectObject(hdc, m_popupFont);
-			DrawTextLogger logger(hdc, r, m_popupFontHeight);
+			DrawTextLogger logger(m_hwnd, hdc, m_popupFontHeight, m_tooltipBackgroundBrush);
 			logger.Info(L"%ls %ls", vis ? L"Hide" : L"Show", tooltip[m_buttonSelected]);
+			logger.DrawAtCursor();
 		}
 		else if (m_timelineSelected)
 		{
@@ -1448,17 +1478,11 @@ namespace uba
 					b.Append('0');
 				b.AppendValue(milliseconds);
 
-				RECT r;
-				r.left = left + 4;
-				r.top = timelineTop - 20;
-				r.right = left + b.count*8;
-				r.bottom = r.top + 15;
-
-				FillRect(hdc, &r, m_tooltipBackgroundBrush);
 				SelectObject(hdc, m_popupFont);
-				DrawTextLogger logger(hdc, r, m_popupFontHeight);
+				DrawTextLogger logger(m_hwnd, hdc, m_popupFontHeight, m_tooltipBackgroundBrush);
 
 				logger.Info(L"%s", b.data);
+				logger.DrawAtPos(left + 4, timelineTop - 20);
 			}
 		}
 		else if (m_fetchedFilesSelected != ~0u)
@@ -1467,32 +1491,13 @@ namespace uba
 			auto& fetchedFiles = session.fetchedFiles;
 			if (!fetchedFiles.empty() && !fetchedFiles[0].hint.empty())
 			{
+				/*
 				int colWidth = 500;
 				int width = colWidth * 2;
 				int height = Min(int(clientRect.bottom), int(fetchedFiles.size() * m_popupFontHeight));
 
-
-				POINT p;
-				GetCursorPos(&p);
-				ScreenToClient(m_hwnd, &p);
-				RECT r;
-				r.left = p.x;
-				r.top = p.y;
-				r.right = r.left + width;
-				r.bottom = r.top + height;
-
-				if (r.right > clientRect.right)
-					OffsetRect(&r, -width, 0);
-				if (r.bottom > clientRect.bottom)
-				{
-					OffsetRect(&r, 0, -height);
-					if (r.top < 0)
-						OffsetRect(&r, 0, -r.top);
-				}
-				FillRect(hdc, &r, m_tooltipBackgroundBrush);
-
 				SelectObject(hdc, m_font);
-				DrawTextLogger logger(hdc, r, FontHeight);
+				DrawTextLogger logger(m_hwnd, hdc, r, FontHeight, m_tooltipBackgroundBrush);
 				for (auto& f : fetchedFiles)
 				{
 					if (f.hint == TC("KnownInput"))
@@ -1509,6 +1514,8 @@ namespace uba
 					}
 					logger.Info(L"%s", f.hint.c_str());
 				}
+				logger.DrawAtCursor();
+				*/
 			}
 		}
 	}
@@ -1528,9 +1535,9 @@ namespace uba
 		else if (process.exitCode != 0)
 			brush = m_processBrushes[selected].error;
 
-		u64 writeFilesTime = Max(process.processStats.writeFiles.time, process.processStats.sendFiles.time);
+		u64 writeFilesTime = process.writeFilesTime;
 
-		if (!done || process.exitCode != 0 || !m_showCreateWriteColors || (TimeToMs(writeFilesTime, m_traceView.frequency) < 300 && TimeToMs(process.processStats.createFile.time, m_traceView.frequency) < 300))
+		if (!done || process.exitCode != 0 || !m_showCreateWriteColors || (TimeToMs(writeFilesTime, m_traceView.frequency) < 300 && TimeToMs(process.createFilesTime, m_traceView.frequency) < 300))
 		{
 			if (writingBitmap)
 				rect.right = 256;
@@ -1544,7 +1551,7 @@ namespace uba
 		RECT main = rect;
 		int width = rect.right - rect.left;
 
-		double recvPart = (double(ConvertTime(m_traceView, process.processStats.createFile.time)) / duration);
+		double recvPart = (double(ConvertTime(m_traceView, process.createFilesTime)) / duration);
 		if (int headSize = int(recvPart * width))
 		{
 			UBA_ASSERT(headSize > 0);
@@ -2082,23 +2089,38 @@ namespace uba
 			out.Info(L"  Duration:  %ls", TimeToText(process.stop - process.start, true).str);
 		if (hasExited && process.exitCode != 0)
 			out.Info(L"  ExitCode:  %u", process.exitCode);
-		out.Info(L"");
 
 		if (process.stop != ~u64(0))
 		{
+			out.Info(L"");
+
+			BinaryReader reader(process.stats.data(), 0, process.stats.size());
+			ProcessStats processStats;
+			SessionStats sessionStats;
+			StorageStats storageStats;
+			SystemStats systemStats;
+
+			processStats.Read(reader, m_traceView.version);
+			if (reader.GetLeft())
+			{
+				sessionStats.Read(reader, m_traceView.version);
+				storageStats.Read(reader);
+				systemStats.Read(reader);
+			}
+
 			out.Info(L"  ----------- Process stats -----------");
-			process.processStats.Print(out, m_traceView.frequency);
+			processStats.Print(out, m_traceView.frequency);
 			if (hasStorageStats)
 			{
 				out.Info(L"");
 				out.Info(L"  ----------- Session stats -----------");
-				process.sessionStats.Print(out, m_traceView.frequency);
+				sessionStats.Print(out, m_traceView.frequency);
 				out.Info(L"");
 				out.Info(L"  ----------- Storage stats -----------");
-				process.storageStats.Print(out, m_traceView.frequency);
+				storageStats.Print(out, m_traceView.frequency);
 				out.Info(L"");
 				out.Info(L"  ----------- System stats ------------");
-				process.systemStats.Print(out, false, m_traceView.frequency);
+				systemStats.Print(out, false, m_traceView.frequency);
 			}
 		}
 	}
@@ -2368,7 +2390,7 @@ namespace uba
 				for (auto& session : m_traceView.sessions)
 					for (auto& processor : session.processors)
 						for (auto& process : processor.processes)
-							if (TimeToMs(Max(process.processStats.writeFiles.time, process.processStats.sendFiles.time), m_traceView.frequency) >= 300 || TimeToMs(process.processStats.createFile.time, m_traceView.frequency) >= 300)
+							if (TimeToMs(process.writeFilesTime, m_traceView.frequency) >= 300 || TimeToMs(process.createFilesTime, m_traceView.frequency) >= 300)
 								process.bitmapDirty = true;
 
 			UpdateScrollbars(true);
