[UBA]
* Fixed so RootPaths does not create double backslash version on non-windows when registering root paths
* Disable root paths unit test on linux since it doesn't work and require some work

diff --git a/Engine/Source/Programs/UnrealBuildAccelerator/Common/Private/UbaRootPaths.cpp b/Engine/Source/Programs/UnrealBuildAccelerator/Common/Private/UbaRootPaths.cpp
--- a/Engine/Source/Programs/UnrealBuildAccelerator/Common/Private/UbaRootPaths.cpp
+++ b/Engine/Source/Programs/UnrealBuildAccelerator/Common/Private/UbaRootPaths.cpp
@@ -11,8 +11,8 @@ namespace uba
 {
 	bool RootPaths::RegisterRoot(Logger& logger, const tchar* rootPath, bool includeInKey)
 	{
-		// Register rootPath both with single path separators and double path separators because text files store them with double path separators
-
+		// Register rootPath both with single path separators and double path separators on windows because text files store them with double path separators
+		#if PLATFORM_WINDOWS
 		StringBuffer<> doubleSlash;
 		for (const tchar* it=rootPath; *it; ++it)
 		{
@@ -22,6 +22,10 @@ namespace uba
 		}
 
 		const tchar* rootPaths[] = { rootPath, doubleSlash.data };
+		#else
+		const tchar* rootPaths[] = { rootPath };
+		#endif
+
 		for (const tchar* rp : rootPaths)
 		{
 			u32 index = u32(m_roots.size());
diff --git a/Engine/Source/Programs/UnrealBuildAccelerator/Test/Private/UbaTestBasics.h b/Engine/Source/Programs/UnrealBuildAccelerator/Test/Private/UbaTestBasics.h
--- a/Engine/Source/Programs/UnrealBuildAccelerator/Test/Private/UbaTestBasics.h
+++ b/Engine/Source/Programs/UnrealBuildAccelerator/Test/Private/UbaTestBasics.h
@@ -343,14 +343,22 @@ namespace uba
 
 	bool TestRootPaths(Logger& logger, const StringBufferBase& rootDir)
 	{
+		#if PLATFORM_WINDOWS
+		const tchar root1[] = TC("c:\\temp\\");
+		const tchar root2[] = TC("e:\\temp\\");
+		const tchar str[] = TC("e:\\temp\\foo");
+		#else
+		const tchar root1[] = TC("/mnt/c/");
+		const tchar root2[] = TC("/mnt/e/");
+		const tchar str[] = TC("/mnt/e/foo");
+		#endif
+
 		RootPaths paths;
-		if (!paths.RegisterRoot(logger, TC("c:\\temp\\")))
+		if (!paths.RegisterRoot(logger, root1))
 			return false;
-		if (!paths.RegisterRoot(logger, TC("e:\\temp\\")))
+		if (!paths.RegisterRoot(logger, root2))
 			return false;
 
-		tchar str[] = TC("e:\\temp\\foo");
-
 		bool success = true;
 		StringBuffer<> temp;
 		u32 rootPos = ~0u;
@@ -360,7 +368,7 @@ namespace uba
 				{
 					if (strLen != 1)
 						success = false;
-					if (str[0] != RootPaths::RootStartByte+2)
+					if (str[0] != RootPaths::RootStartByte + (IsWindows ? 2 : 1))
 						success = false;
 					rootPos = str[0];
 				}
