UBA: Fix assert

#rnx

diff --git a/Engine/Source/Programs/UnrealBuildAccelerator/Core/Public/UbaPlatform.h b/Engine/Source/Programs/UnrealBuildAccelerator/Core/Public/UbaPlatform.h
--- a/Engine/Source/Programs/UnrealBuildAccelerator/Core/Public/UbaPlatform.h
+++ b/Engine/Source/Programs/UnrealBuildAccelerator/Core/Public/UbaPlatform.h
@@ -107,7 +107,7 @@ namespace uba
 	inline constexpr bool CaseInsensitiveFs = true;
 	inline constexpr tchar PathSeparator = '\\';
 	inline constexpr tchar NonPathSeparator = '/';
-	inline constexpr u32 MaxPath = 512;
+	inline constexpr u32 MaxPath = 1024;
 	inline DWORD ToLow(u64 v) { LARGE_INTEGER li; li.QuadPart = (LONGLONG)v; return li.LowPart; }
 	inline LONG ToHigh(u64 v) { LARGE_INTEGER li; li.QuadPart = (LONGLONG)v; return li.HighPart; }
 	inline LARGE_INTEGER ToLargeInteger(u64 v) { LARGE_INTEGER li; li.QuadPart = (LONGLONG)v; return li; }
@@ -125,7 +125,7 @@ namespace uba
 #else
 	inline constexpr tchar PathSeparator = '/';
 	inline constexpr tchar NonPathSeparator = '\\';
-	inline constexpr u32 MaxPath = 512;
+	inline constexpr u32 MaxPath = 1024;
 	#define TStrlen(s) u32(strlen(s))
 	#define TStrchr(a, b) strchr(a, b)
 	#define TStrrchr(a, b) strrchr(a, b)
diff --git a/Engine/Source/Programs/UnrealBuildAccelerator/Detours/Private/Windows/UbaDetoursFunctionsKernelBase.inl b/Engine/Source/Programs/UnrealBuildAccelerator/Detours/Private/Windows/UbaDetoursFunctionsKernelBase.inl
--- a/Engine/Source/Programs/UnrealBuildAccelerator/Detours/Private/Windows/UbaDetoursFunctionsKernelBase.inl
+++ b/Engine/Source/Programs/UnrealBuildAccelerator/Detours/Private/Windows/UbaDetoursFunctionsKernelBase.inl
@@ -934,7 +934,7 @@ BOOL Detoured_GetFileAttributesExW(LPCWSTR lpFileName, GET_FILEEX_INFO_LEVELS fI
 		return True_GetFileAttributesExW(lpFileName, fInfoLevelId, lpFileInformation);
 	}
 
-	StringBuffer<> fixedName;
+	StringBuffer<MaxPath> fixedName;
 	FixPath(fixedName, lpFileName);
 
 	if (!g_rules->CanExist(fixedName.data))
@@ -956,7 +956,7 @@ DWORD Detoured_GetFileAttributesW(LPCWSTR lpFileName)
 		return res;
 	}
 
-	StringBuffer<> fixedPath;
+	StringBuffer<MaxPath> fixedPath;
 	if (!FixPath(fixedPath, lpFileName))
 		return INVALID_FILE_ATTRIBUTES;
 
