[Uba]
* Bumped table sizes to prevent oom when building our largest non unity builds
* Added bucket support in cache system
* Added ParallelFor function to WorkManager
* Fixed so annoying logging when pressing ctrl-c is muted

diff --git a/Engine/Source/Programs/UnrealBuildAccelerator/Common/Private/UbaCacheClient.cpp b/Engine/Source/Programs/UnrealBuildAccelerator/Common/Private/UbaCacheClient.cpp
--- a/Engine/Source/Programs/UnrealBuildAccelerator/Common/Private/UbaCacheClient.cpp
+++ b/Engine/Source/Programs/UnrealBuildAccelerator/Common/Private/UbaCacheClient.cpp
@@ -1,6 +1,7 @@
 // Copyright Epic Games, Inc. All Rights Reserved.
 
 #include "UbaCacheClient.h"
+#include "UbaCompactTables.h"
 #include "UbaFileAccessor.h"
 #include "UbaNetworkMessage.h"
 #include "UbaProcess.h"
@@ -10,15 +11,37 @@
 
 namespace uba
 {
+	struct CacheClient::Bucket
+	{
+		Bucket(u32 id_)
+		:	id(id_)
+		,	serverPathTable(CachePathTableMaxSize, CompactPathTable::V1)
+		,	serverCasKeyTable(CacheCasKeyTableMaxSize)
+		,	sendPathTable(CachePathTableMaxSize, CompactPathTable::V1)
+		,	sendCasKeyTable(CacheCasKeyTableMaxSize)
+		{
+		}
+
+		u32 id = 0;
+
+		CompactPathTable serverPathTable;
+		CompactCasKeyTable serverCasKeyTable;
+
+		CompactPathTable sendPathTable;
+		CompactCasKeyTable sendCasKeyTable;
+
+		ReaderWriterLock pathTableNetworkLock;
+		u32 pathTableSizeSent = 0;
+
+		ReaderWriterLock casKeyTableNetworkLock;
+		u32 casKeyTableSizeSent = 0;
+	};
+
 	CacheClient::CacheClient(LogWriter& writer, StorageImpl& storage, NetworkClient& client, Session& session)
 	:	m_logger(writer, TC("UbaCacheClient"))
 	,	m_storage(storage)
 	,	m_client(client)
 	,	m_session(session)
-	,	m_serverPathTable(CachePathTableMaxSize, CompactPathTable::V1)
-	,	m_serverCasKeyTable(CacheCasKeyTableMaxSize)
-	,	m_sendPathTable(8*1024*1024, CompactPathTable::V1)
-	,	m_sendCasKeyTable(8*1024*1024)
 	{
 		m_client.RegisterOnConnected([this]()
 			{
@@ -42,7 +65,7 @@ namespace uba
 
 	CacheClient::~CacheClient() = default;
 
-	bool CacheClient::WriteToCache(const RootPaths& rootPaths, const ProcessHandle& process)
+	bool CacheClient::WriteToCache(const RootPaths& rootPaths, u32 bucketId, const ProcessHandle& process)
 	{
 		if (!m_connected)
 			return false;
@@ -67,6 +90,10 @@ namespace uba
 		u32 requiredCasTableSize = 0;
 		bool success = true;
 
+		SCOPED_WRITE_LOCK(m_bucketsLock, bucketsLock);
+		Bucket& bucket = m_buckets.try_emplace(bucketId, bucketId).first->second;
+		bucketsLock.Leave();
+
 		// Traverse all inputs and outputs. to create cache entry that we can send to server
 		while (true)
 		{
@@ -122,7 +149,7 @@ namespace uba
 			TString qualifiedPath = path.data + rootLen - 1;
 			qualifiedPath[0] = tchar(RootPaths::RootStartByte + root->index);
 
-			u32 pathOffset = m_sendPathTable.Add(qualifiedPath.c_str(), u32(qualifiedPath.size()), &requiredPathTableSize);
+			u32 pathOffset = bucket.sendPathTable.Add(qualifiedPath.c_str(), u32(qualifiedPath.size()), &requiredPathTableSize);
 
 			if (!isOutput) // Output files should be removed from input files.. For example when cl.exe compiles pch it reads previous pch file and we don't want it to be input
 				if (outputsStringToCasKey.find(pathOffset) != outputsStringToCasKey.end())
@@ -148,7 +175,7 @@ namespace uba
 			}
 
 			UBA_ASSERT(IsCompressed(casKey));
-			insres.first->second = m_sendCasKeyTable.Add(casKey, pathOffset, &requiredCasTableSize);
+			insres.first->second = bucket.sendCasKeyTable.Add(casKey, pathOffset, &requiredCasTableSize);
 		}
 
 		if (!success)
@@ -158,21 +185,21 @@ namespace uba
 			m_logger.Warning(TC("NO OUTPUTS FROM process %s"), process.GetStartInfo().description); 
 
 		// Make sure server has enough of the path table to be able to resolve offsets from cache entry
-		if (!SendPathTable(requiredPathTableSize))
+		if (!SendPathTable(bucket, requiredPathTableSize))
 			return false;
 
 		// Make sure server has enough of the cas table to be able to resolve offsets from cache entry
-		if (!SendCasTable(requiredCasTableSize))
+		if (!SendCasTable(bucket, requiredCasTableSize))
 			return false;
 
 		// actual cache entry now when we know server has the needed tables
-		if (!SendCacheEntry(rootPaths, cmdKey, inputsStringToCasKey, outputsStringToCasKey))
+		if (!SendCacheEntry(bucket, rootPaths, cmdKey, inputsStringToCasKey, outputsStringToCasKey))
 			return false;
 
 		return true;
 	}
 
-	bool CacheClient::FetchFromCache(const RootPaths& rootPaths, const ProcessStartInfo& info)
+	bool CacheClient::FetchFromCache(const RootPaths& rootPaths, u32 bucketId, const ProcessStartInfo& info)
 	{
 		if (!m_connected)
 			return false;
@@ -208,11 +235,16 @@ namespace uba
 
 		BinaryReader reader(memory, 0, sizeof_array(memory));
 
+		SCOPED_WRITE_LOCK(m_bucketsLock, bucketsLock);
+		Bucket& bucket = m_buckets.try_emplace(bucketId, bucketId).first->second;
+		bucketsLock.Leave();
+
 		{
 			TimerScope ts(cacheStats.fetchEntries);
 			// Fetch entries.. server will provide as many as fits. TODO: Should it be possible to ask for more entries?
 			StackBinaryWriter<32> writer;
 			NetworkMessage msg(m_client, CacheServiceId, CacheMessageType_FetchEntries, writer);
+			writer.Write7BitEncoded(bucket.id);
 			writer.WriteCasKey(cmdKey);
 			if (!msg.Send(reader))
 				return false;
@@ -238,16 +270,16 @@ namespace uba
 					auto insres = offsetIsMatch.try_emplace(casKeyOffset);
 					if (insres.second)
 					{
-						if (casKeyOffset >= m_serverCasKeyTable.GetSize())
+						if (casKeyOffset >= bucket.serverCasKeyTable.GetSize())
 						{
 							TimerScope ts2(cacheStats.fetchCasTable);
-							if (!FetchCasTable())
+							if (!FetchCasTable(bucket))
 								return false;
 						}
 
 						StringBuffer<MaxPath> path;
 						CasKey cacheCasKey;
-						if (!GetLocalPathAndCasKey(rootPaths, path, cacheCasKey, m_serverCasKeyTable, m_serverPathTable, casKeyOffset))
+						if (!GetLocalPathAndCasKey(bucket, rootPaths, path, cacheCasKey, bucket.serverCasKeyTable, bucket.serverPathTable, casKeyOffset))
 							return false;
 						UBA_ASSERT(IsCompressed(cacheCasKey));
 
@@ -290,15 +322,15 @@ namespace uba
 			while (reader.GetPositionData() != outputEnd)
 			{
 				u32 casKeyOffset = u32(reader.Read7BitEncoded());
-				if (casKeyOffset >= m_serverCasKeyTable.GetSize())
-					if (!FetchCasTable())
+				if (casKeyOffset >= bucket.serverCasKeyTable.GetSize())
+					if (!FetchCasTable(bucket))
 						return false;
 
 				TimerScope ts(cacheStats.fetchOutput);
 
 				StringBuffer<MaxPath> path;
 				CasKey casKey;
-				if (!GetLocalPathAndCasKey(rootPaths, path, casKey, m_serverCasKeyTable, m_serverPathTable, casKeyOffset))
+				if (!GetLocalPathAndCasKey(bucket, rootPaths, path, casKey, bucket.serverCasKeyTable, bucket.serverPathTable, casKeyOffset))
 					return false;
 				UBA_ASSERT(IsCompressed(casKey));
 
@@ -393,21 +425,22 @@ namespace uba
 		return true;
 	}
 
-	bool CacheClient::SendPathTable(u32 requiredPathTableSize)
+	bool CacheClient::SendPathTable(Bucket& bucket, u32 requiredPathTableSize)
 	{
-		SCOPED_WRITE_LOCK(m_pathTableNetworkLock, lock);
-		if (requiredPathTableSize <= m_pathTableSizeSent)
+		SCOPED_WRITE_LOCK(bucket.pathTableNetworkLock, lock);
+		if (requiredPathTableSize <= bucket.pathTableSizeSent)
 			return true;
 
-		u32 left = requiredPathTableSize - m_pathTableSizeSent;
+		u32 left = requiredPathTableSize - bucket.pathTableSizeSent;
 		while (left)
 		{
 			StackBinaryWriter<SendMaxSize> writer;
 			NetworkMessage msg(m_client, CacheServiceId, CacheMessageType_StorePathTable, writer);
-			u32 toSend = Min(requiredPathTableSize - m_pathTableSizeSent, u32(m_client.GetMessageMaxSize() - 32));
+			writer.Write7BitEncoded(bucket.id);
+			u32 toSend = Min(requiredPathTableSize - bucket.pathTableSizeSent, u32(m_client.GetMessageMaxSize() - 32));
 			left -= toSend;
-			writer.WriteBytes(m_sendPathTable.GetMemory() + m_pathTableSizeSent, toSend);
-			m_pathTableSizeSent += toSend;
+			writer.WriteBytes(bucket.sendPathTable.GetMemory() + bucket.pathTableSizeSent, toSend);
+			bucket.pathTableSizeSent += toSend;
 
 			StackBinaryReader<16> reader;
 			if (!msg.Send(reader))
@@ -416,21 +449,22 @@ namespace uba
 		return true;
 	}
 
-	bool CacheClient::SendCasTable(u32 requiredCasTableSize)
+	bool CacheClient::SendCasTable(Bucket& bucket, u32 requiredCasTableSize)
 	{
-		SCOPED_WRITE_LOCK(m_casKeyTableNetworkLock, lock);
-		if (requiredCasTableSize <= m_casKeyTableSizeSent)
+		SCOPED_WRITE_LOCK(bucket.casKeyTableNetworkLock, lock);
+		if (requiredCasTableSize <= bucket.casKeyTableSizeSent)
 			return true;
 
-		u32 left = requiredCasTableSize - m_casKeyTableSizeSent;
+		u32 left = requiredCasTableSize - bucket.casKeyTableSizeSent;
 		while (left)
 		{
 			StackBinaryWriter<SendMaxSize> writer;
 			NetworkMessage msg(m_client, CacheServiceId, CacheMessageType_StoreCasTable, writer);
-			u32 toSend = Min(requiredCasTableSize - m_casKeyTableSizeSent, u32(m_client.GetMessageMaxSize() - 32));
+			writer.Write7BitEncoded(bucket.id);
+			u32 toSend = Min(requiredCasTableSize - bucket.casKeyTableSizeSent, u32(m_client.GetMessageMaxSize() - 32));
 			left -= toSend;
-			writer.WriteBytes(m_sendCasKeyTable.GetMemory() + m_casKeyTableSizeSent, toSend);
-			m_casKeyTableSizeSent += toSend;
+			writer.WriteBytes(bucket.sendCasKeyTable.GetMemory() + bucket.casKeyTableSizeSent, toSend);
+			bucket.casKeyTableSizeSent += toSend;
 
 			StackBinaryReader<16> reader;
 			if (!msg.Send(reader))
@@ -439,13 +473,14 @@ namespace uba
 		return true;
 	}
 
-	bool CacheClient::SendCacheEntry(const RootPaths& rootPaths, const CasKey& cmdKey, const Map<u32, u32>& inputsStringToCasKey, const Map<u32, u32>& outputsStringToCasKey)
+	bool CacheClient::SendCacheEntry(Bucket& bucket, const RootPaths& rootPaths, const CasKey& cmdKey, const Map<u32, u32>& inputsStringToCasKey, const Map<u32, u32>& outputsStringToCasKey)
 	{
 		StackBinaryReader<1024> reader;
 		{
 			StackBinaryWriter<SendMaxSize> writer;
 
 			NetworkMessage msg(m_client, CacheServiceId, CacheMessageType_StoreEntry, writer);
+			writer.Write7BitEncoded(bucket.id);
 			writer.WriteCasKey(cmdKey);
 
 			writer.Write7BitEncoded(outputsStringToCasKey.size());
@@ -470,7 +505,7 @@ namespace uba
 
 			StringBuffer<MaxPath> path;
 			CasKey casKey;
-			if (!GetLocalPathAndCasKey(rootPaths, path, casKey, m_sendCasKeyTable, m_sendPathTable, casKeyOffset))
+			if (!GetLocalPathAndCasKey(bucket, rootPaths, path, casKey, bucket.sendCasKeyTable, bucket.sendPathTable, casKeyOffset))
 				return false;
 
 			casKey = AsCompressed(casKey, true);
@@ -559,6 +594,7 @@ namespace uba
 		// Send done.. confirm to server
 		StackBinaryWriter<SendMaxSize> writer;
 		NetworkMessage msg(m_client, CacheServiceId, CacheMessageType_StoreEntryDone, writer);
+		writer.Write7BitEncoded(bucket.id);
 		writer.WriteCasKey(cmdKey);
 		if (!msg.Send(reader))
 			return false;
@@ -566,19 +602,20 @@ namespace uba
 		return true;
 	}
 
-	bool CacheClient::FetchCasTable()
+	bool CacheClient::FetchCasTable(Bucket& bucket)
 	{
-		SCOPED_WRITE_LOCK(m_casKeyTableNetworkLock, lock);
+		SCOPED_WRITE_LOCK(bucket.casKeyTableNetworkLock, lock);
 
 		StackBinaryReader<SendMaxSize> reader;
 		{
-			//SCOPED_WRITE_LOCK(m_pathTableNetworkLock, lock);
+			//SCOPED_WRITE_LOCK(bucket.pathTableNetworkLock, lock);
 			u32 targetSize = ~0u;
-			while (m_serverCasKeyTable.GetSize() < targetSize)
+			while (bucket.serverCasKeyTable.GetSize() < targetSize)
 			{
 				StackBinaryWriter<16> writer;
 				NetworkMessage msg(m_client, CacheServiceId, CacheMessageType_FetchCasTable, writer);
-				writer.WriteU32(m_serverCasKeyTable.GetSize());
+				writer.Write7BitEncoded(bucket.id);
+				writer.WriteU32(bucket.serverCasKeyTable.GetSize());
 
 				reader.Reset();
 				if (!msg.Send(reader))
@@ -587,16 +624,17 @@ namespace uba
 				if (targetSize == ~0u)
 					targetSize = size;
 
-				m_serverCasKeyTable.ReadMem(reader, false);
+				bucket.serverCasKeyTable.ReadMem(reader, false);
 			}
 		}
 		{
 			u32 targetSize = ~0u;
-			while (m_serverPathTable.GetSize() < targetSize)
+			while (bucket.serverPathTable.GetSize() < targetSize)
 			{
 				StackBinaryWriter<16> writer;
 				NetworkMessage msg(m_client, CacheServiceId, CacheMessageType_FetchPathTable, writer);
-				writer.WriteU32(m_serverPathTable.GetSize());
+				writer.Write7BitEncoded(bucket.id);
+				writer.WriteU32(bucket.serverPathTable.GetSize());
 
 				reader.Reset();
 				if (!msg.Send(reader))
@@ -605,7 +643,7 @@ namespace uba
 				if (targetSize == ~0u)
 					targetSize = size;
 
-				m_serverPathTable.ReadMem(reader, false);
+				bucket.serverPathTable.ReadMem(reader, false);
 			}
 		}
 		return true;
@@ -648,12 +686,12 @@ namespace uba
 		return ToCasKey(hasher, false);
 	}
 
-	bool CacheClient::GetLocalPathAndCasKey(const RootPaths& rootPaths, StringBufferBase& outPath, CasKey& outKey, CompactCasKeyTable& casKeyTable, CompactPathTable& pathTable, u32 offset)
+	bool CacheClient::GetLocalPathAndCasKey(Bucket& bucket, const RootPaths& rootPaths, StringBufferBase& outPath, CasKey& outKey, CompactCasKeyTable& casKeyTable, CompactPathTable& pathTable, u32 offset)
 	{
 		if (!m_connected)
 			return false;
 
-		SCOPED_READ_LOCK(m_casKeyTableNetworkLock, lock); // TODO: Is this needed?
+		SCOPED_READ_LOCK(bucket.casKeyTableNetworkLock, lock); // TODO: Is this needed?
 
 		StringBuffer<MaxPath> normalizedPath;
 		casKeyTable.GetPathAndKey(normalizedPath, outKey, pathTable, offset);
diff --git a/Engine/Source/Programs/UnrealBuildAccelerator/Common/Private/UbaCacheServer.cpp b/Engine/Source/Programs/UnrealBuildAccelerator/Common/Private/UbaCacheServer.cpp
--- a/Engine/Source/Programs/UnrealBuildAccelerator/Common/Private/UbaCacheServer.cpp
+++ b/Engine/Source/Programs/UnrealBuildAccelerator/Common/Private/UbaCacheServer.cpp
@@ -1,6 +1,7 @@
 // Copyright Epic Games, Inc. All Rights Reserved.
 
 #include "UbaCacheServer.h"
+#include "UbaCompactTables.h"
 #include "UbaBinaryReaderWriter.h"
 #include "UbaFileAccessor.h"
 #include "UbaNetworkServer.h"
@@ -10,14 +11,45 @@ namespace uba
 {
 	static constexpr u32 CacheFileVersion = 1;
 
+	struct CacheServer::CacheEntry
+	{
+		u64 creationTime;
+		Vector<u8> inputCasKeyOffsets;
+		Vector<u8> outputCasKeyOffsets;
+	};
+
+	struct CacheServer::ConnectionBucket
+	{
+		ConnectionBucket() : pathTable(CachePathTableMaxSize, CompactPathTable::V1), casKeyTable(CacheCasKeyTableMaxSize) {}
+		CompactPathTable pathTable;
+		CompactCasKeyTable casKeyTable;
+
+		ReaderWriterLock cacheEntryLookupLock;
+		UnorderedMap<CasKey, CacheEntry> cacheEntryLookup;
+	};
+
 	struct CacheServer::Connection
 	{
-		Connection() : m_pathTable(8*1024*1024, CompactPathTable::V1), m_casKeyTable(8*1024*1024) {}
+		UnorderedMap<u32, ConnectionBucket> buckets;
+	};
+
+	struct CacheServer::CacheEntries
+	{
+		ReaderWriterLock lock;
+		List<CacheEntry> entries;
+	};
+
+	struct CacheServer::Bucket
+	{
+		Bucket() : m_pathTable(CachePathTableMaxSize, CompactPathTable::V1), m_casKeyTable(CacheCasKeyTableMaxSize) {}
+		ReaderWriterLock m_cacheEntryLookupLock;
+		UnorderedMap<CasKey, CacheEntries> m_cacheEntryLookup;
+
 		CompactPathTable m_pathTable;
 		CompactCasKeyTable m_casKeyTable;
 
-		ReaderWriterLock cacheEntryLookupLock;
-		UnorderedMap<CasKey, CacheEntry> cacheEntryLookup;
+		u64 totalEntryCount = 0;
+		u64 totalEntrySize = 0;
 	};
 
 	const tchar* ToString(CacheMessageType type)
@@ -36,8 +68,6 @@ namespace uba
 	:	m_logger(writer, TC("UbaCacheServer"))
 	,	m_server(server)
 	,	m_storage(storage)
-	,	m_pathTable(CachePathTableMaxSize, CompactPathTable::V1)
-	,	m_casKeyTable(CacheCasKeyTableMaxSize)
 	{
 		m_rootDir.count = GetFullPathNameW(rootDir, m_rootDir.capacity, m_rootDir.data, NULL);
 		m_rootDir.Replace('/', PathSeparator).EnsureEndsWithSlash();
@@ -72,59 +102,74 @@ namespace uba
 
 		FileAccessor file(m_logger, fileName.data);
 		if (!file.OpenMemoryRead(0, false))
+		{
+			m_logger.Detail(TC("No database found. Starting a new one at %s"), fileName.data);
 			return true;
+		}
 		BinaryReader reader(file.GetData(), 0, file.GetSize());
 
 		u32 version = reader.ReadU32();
 		if (version != CacheFileVersion)
 			return true;
 
-		u32 pathTableSize = reader.ReadU32();
-		if (pathTableSize)
-		{
-			BinaryReader pathTableReader(reader.GetPositionData(), 0, pathTableSize);
-			m_pathTable.ReadMem(pathTableReader, true);
-			reader.Skip(pathTableSize);
-		}
+		u32 totalPathTableSize = 0;
+		u32 totalCasKeyTableSize = 0;
+		u64 totalCacheEntryCount = 0;
 
-		u32 casKeyTableSize = reader.ReadU32();
-		if (casKeyTableSize)
+		u32 bucketCount = reader.ReadU32();
+		while (bucketCount--)
 		{
-			BinaryReader casKeyTableReader(reader.GetPositionData(), 0, casKeyTableSize);
-			m_casKeyTable.ReadMem(casKeyTableReader, true);
-			reader.Skip(casKeyTableSize);
-		}
+			u32 bucketId = reader.ReadU32();
+			Bucket& bucket = m_buckets[bucketId];
 
-		u32 entryLookupCount = reader.ReadU32();
-		m_cacheEntryLookup.reserve(entryLookupCount);
+			u32 pathTableSize = reader.ReadU32();
+			if (pathTableSize)
+			{
+				BinaryReader pathTableReader(reader.GetPositionData(), 0, pathTableSize);
+				bucket.m_pathTable.ReadMem(pathTableReader, true);
+				reader.Skip(pathTableSize);
+			}
+			totalPathTableSize += pathTableSize;
 
-		u64 totalCacheEntryCount = 0;
-		while (reader.GetLeft())
-		{
-			auto insres = m_cacheEntryLookup.try_emplace(reader.ReadCasKey());
-			UBA_ASSERT(insres.second);
-			auto& cacheEntries = insres.first->second;
-			u32 cacheEntryCount = reader.ReadU32();
-			totalCacheEntryCount += cacheEntryCount;
+			u32 casKeyTableSize = reader.ReadU32();
+			if (casKeyTableSize)
+			{
+				BinaryReader casKeyTableReader(reader.GetPositionData(), 0, casKeyTableSize);
+				bucket.m_casKeyTable.ReadMem(casKeyTableReader, true);
+				reader.Skip(casKeyTableSize);
+			}
+			totalCasKeyTableSize += casKeyTableSize;
+
+			u32 entryLookupCount = reader.ReadU32();
+			bucket.m_cacheEntryLookup.reserve(entryLookupCount);
 
-			while (cacheEntryCount--)
+			while (entryLookupCount--)
 			{
-				auto& cacheEntry = cacheEntries.entries.emplace_back();
+				auto insres = bucket.m_cacheEntryLookup.try_emplace(reader.ReadCasKey());
+				UBA_ASSERT(insres.second);
+				auto& cacheEntries = insres.first->second;
+				u32 cacheEntryCount = reader.ReadU32();
+				totalCacheEntryCount += cacheEntryCount;
 
-				cacheEntry.creationTime = reader.ReadU64();
+				while (cacheEntryCount--)
+				{
+					auto& cacheEntry = cacheEntries.entries.emplace_back();
+
+					cacheEntry.creationTime = reader.ReadU64();
 
-				u32 inputSize = reader.ReadU32();
-				cacheEntry.inputCasKeyOffsets.resize(inputSize);
-				reader.ReadBytes(cacheEntry.inputCasKeyOffsets.data(), inputSize);
+					u32 inputSize = reader.ReadU32();
+					cacheEntry.inputCasKeyOffsets.resize(inputSize);
+					reader.ReadBytes(cacheEntry.inputCasKeyOffsets.data(), inputSize);
 
-				u32 outputSize = reader.ReadU32();
-				cacheEntry.outputCasKeyOffsets.resize(outputSize);
-				reader.ReadBytes(cacheEntry.outputCasKeyOffsets.data(), outputSize);
+					u32 outputSize = reader.ReadU32();
+					cacheEntry.outputCasKeyOffsets.resize(outputSize);
+					reader.ReadBytes(cacheEntry.outputCasKeyOffsets.data(), outputSize);
+				}
 			}
 		}
 
 		u64 duration = GetTime() - startTime;
-		m_logger.Detail(TC("Database loaded from %s in %s (contained %s paths, %s keys, %llu cache entries)"), fileName.data, TimeToText(duration).str, BytesToText(pathTableSize).str, BytesToText(casKeyTableSize).str, totalCacheEntryCount);
+		m_logger.Detail(TC("Database loaded from %s in %s (%llu bucket(s) containing %s paths, %s keys, %llu cache entries)"), fileName.data, TimeToText(duration).str, m_buckets.size(), BytesToText(totalPathTableSize).str, BytesToText(totalCasKeyTableSize).str, totalCacheEntryCount);
 		return true;
 	}
 
@@ -182,35 +227,44 @@ namespace uba
 
 		Write(CacheFileVersion);
 
-		u32 pathTableSize = m_pathTable.GetSize();
-		Write(pathTableSize);
-		WriteBytes(m_pathTable.GetMemory(), pathTableSize);
+		Write(u32(m_buckets.size()));
 
-		u32 casKeyTableSize = m_casKeyTable.GetSize();
-		Write(casKeyTableSize);
-		WriteBytes(m_casKeyTable.GetMemory(), casKeyTableSize);
-
-		u32 entryLookupCount = u32(m_cacheEntryLookup.size());
-		Write(entryLookupCount);
-
-		for (auto& kv : m_cacheEntryLookup)
+		for (auto& kv : m_buckets)
 		{
+			Bucket& bucket = kv.second;
+
 			Write(kv.first);
 
-			u32 cacheEntryCount = u32(kv.second.entries.size());
-			Write(cacheEntryCount);
+			u32 pathTableSize = bucket.m_pathTable.GetSize();
+			Write(pathTableSize);
+			WriteBytes(bucket.m_pathTable.GetMemory(), pathTableSize);
 
-			for (CacheEntry& entry : kv.second.entries)
+			u32 casKeyTableSize = bucket.m_casKeyTable.GetSize();
+			Write(casKeyTableSize);
+			WriteBytes(bucket.m_casKeyTable.GetMemory(), casKeyTableSize);
+
+			u32 entryLookupCount = u32(bucket.m_cacheEntryLookup.size());
+			Write(entryLookupCount);
+
+			for (auto& kv2 : bucket.m_cacheEntryLookup)
 			{
-				WriteBytes(&entry.creationTime, sizeof(entry.creationTime));
+				Write(kv2.first);
+
+				u32 cacheEntryCount = u32(kv2.second.entries.size());
+				Write(cacheEntryCount);
+
+				for (CacheEntry& entry : kv2.second.entries)
+				{
+					WriteBytes(&entry.creationTime, sizeof(entry.creationTime));
 
-				u32 inputSize = u32(entry.inputCasKeyOffsets.size());
-				Write(inputSize);
-				WriteBytes(entry.inputCasKeyOffsets.data(), inputSize);
+					u32 inputSize = u32(entry.inputCasKeyOffsets.size());
+					Write(inputSize);
+					WriteBytes(entry.inputCasKeyOffsets.data(), inputSize);
 
-				u32 outputSize = u32(entry.outputCasKeyOffsets.size());
-				Write(outputSize);
-				WriteBytes(entry.outputCasKeyOffsets.data(), outputSize);
+					u32 outputSize = u32(entry.outputCasKeyOffsets.size());
+					Write(outputSize);
+					WriteBytes(entry.outputCasKeyOffsets.data(), outputSize);
+				}
 			}
 		}
 
@@ -262,111 +316,119 @@ namespace uba
 		}
 		u64 totalCasCount = existingCas.size();
 
-		Vector<u64*> touchedCas;
-
 		u64 now = GetSystemTimeAsFileTime();
 		u64 oldest = now;
 
 
-		u64 totalEntryCount = 0;
-		u64 totalEntrySize = 0;
-		u64 deleteEntryCount = 0;
+		Atomic<u64> deleteEntryCount;
 
 		u64 deleteCacheEntriesStartTime = GetTime();
 		do
 		{
 			oldest = now;
-			totalEntryCount = 0;
-			totalEntrySize = 0;
 
-			for (auto li=m_cacheEntryLookup.begin(), le=m_cacheEntryLookup.end(); li!=le;)
+			ReaderWriterLock existingCasLock;
+
+			m_server.ParallelFor(16, m_buckets, [&](auto& it)
 			{
-				// There is currently no idea saving more than 256kb worth of entries per lookup key (because that is what fetch max returns).. so let's wipe out
-				// all the entries that overflow that number
-				u64 capacityLeft = SendMaxSize - 32;
+				Vector<u64*> touchedCas;
 
-				CacheEntries& entries = li->second;
-				for (auto i=entries.entries.begin(), e=entries.entries.end(); i!=e;)
-				{
-					auto& entry = *i;
-					auto& inputs = entry.inputCasKeyOffsets;
-					auto& outputs = entry.outputCasKeyOffsets;
+				Bucket& bucket = it->second;
+				bucket.totalEntryCount = 0;
+				bucket.totalEntrySize = 0;
 
-					bool deleteEntry = false;
+				for (auto li=bucket.m_cacheEntryLookup.begin(), le=bucket.m_cacheEntryLookup.end(); li!=le;)
+				{
+					// There is currently no idea saving more than 256kb worth of entries per lookup key (because that is what fetch max returns).. so let's wipe out
+					// all the entries that overflow that number
+					u64 capacityLeft = SendMaxSize - 32;
 
-					u64 neededSize = Get7BitEncodedCount(inputs.size()) + inputs.size() + Get7BitEncodedCount(outputs.size()) + outputs.size();
-					if (neededSize > capacityLeft)
-					{
-						deleteEntry = true;
-						capacityLeft = 0;
-					}
-					else
+					CacheEntries& entries = li->second;
+					for (auto i=entries.entries.begin(), e=entries.entries.end(); i!=e;)
 					{
-						// Traverse outputs and check if cas files exists for each output, if not, delete entry.
-						touchedCas.clear();
-						BinaryReader reader2(outputs.data(), 0, outputs.size());
-						while (reader2.GetLeft())
+						auto& entry = *i;
+						auto& inputs = entry.inputCasKeyOffsets;
+						auto& outputs = entry.outputCasKeyOffsets;
+
+						bool deleteEntry = false;
+
+						u64 neededSize = Get7BitEncodedCount(inputs.size()) + inputs.size() + Get7BitEncodedCount(outputs.size()) + outputs.size();
+						if (neededSize > capacityLeft)
 						{
-							u64 offset = reader2.Read7BitEncoded();
-							CasKey casKey;
-							m_casKeyTable.GetKey(casKey, offset);
-							UBA_ASSERT(IsCompressed(casKey));
-							auto findIt = existingCas.find(casKey);
-							if (findIt == existingCas.end())
-								deleteEntry = true;
-							else
-								touchedCas.push_back(&findIt->second.refCount);
+							deleteEntry = true;
+							capacityLeft = 0;
+						}
+						else
+						{
+							// Traverse outputs and check if cas files exists for each output, if not, delete entry.
+							touchedCas.clear();
+							BinaryReader reader2(outputs.data(), 0, outputs.size());
+							while (reader2.GetLeft())
+							{
+								u64 offset = reader2.Read7BitEncoded();
+								CasKey casKey;
+								bucket.m_casKeyTable.GetKey(casKey, offset);
+								UBA_ASSERT(IsCompressed(casKey));
+								auto findIt = existingCas.find(casKey);
+								if (findIt == existingCas.end())
+									deleteEntry = true;
+								else
+									touchedCas.push_back(&findIt->second.refCount);
+							}
 						}
-					}
 
-					// Check if there are keys that use removed caskey as input and delete those too
-					if (!deleteEntry && !deletedCasFiles.empty())
-					{
-						BinaryReader reader2(inputs.data(), 0, inputs.size());
-						while (reader2.GetLeft())
+						// Check if there are keys that use removed caskey as input and delete those too
+						if (!deleteEntry && !deletedCasFiles.empty())
 						{
-							u64 offset = reader2.Read7BitEncoded();
-							CasKey casKey;
-							m_casKeyTable.GetKey(casKey, offset);
-							if (deletedCasFiles.find(casKey) != deletedCasFiles.end())
+							BinaryReader reader2(inputs.data(), 0, inputs.size());
+							while (reader2.GetLeft())
 							{
-								deleteEntry = true;
-								break;
+								u64 offset = reader2.Read7BitEncoded();
+								CasKey casKey;
+								bucket.m_casKeyTable.GetKey(casKey, offset);
+								if (deletedCasFiles.find(casKey) != deletedCasFiles.end())
+								{
+									deleteEntry = true;
+									break;
+								}
 							}
 						}
-					}
 
-					// Remove entry from entries list and skip increasing ref count of cas files
-					if (deleteEntry)
-					{
-						++deleteEntryCount;
-						i = entries.entries.erase(i);
-						e = entries.entries.end();
-						continue;
-					}
+						// Remove entry from entries list and skip increasing ref count of cas files
+						if (deleteEntry)
+						{
+							++deleteEntryCount;
+							i = entries.entries.erase(i);
+							e = entries.entries.end();
+							continue;
+						}
 
-					++totalEntryCount;
-					totalEntrySize += neededSize;
+						++bucket.totalEntryCount;
+						bucket.totalEntrySize += neededSize;
 
-					capacityLeft -= neededSize;
+						capacityLeft -= neededSize;
 
-					if (!oldest || entry.creationTime < oldest)
-						oldest = entry.creationTime;
+						SCOPED_WRITE_LOCK(existingCasLock, l);
+						if (!oldest || entry.creationTime < oldest)
+							oldest = entry.creationTime;
 
-					for (u64* v : touchedCas)
-						++(*v);
-					++i;
-				}
+						for (u64* v : touchedCas)
+							++(*v);
+						l.Leave();
 
-				if (!entries.entries.empty())
-				{
-					++li;
-					continue;
-				}
+						++i;
+					}
 
-				li = m_cacheEntryLookup.erase(li);
-				le = m_cacheEntryLookup.end();
-			}
+					if (!entries.entries.empty())
+					{
+						++li;
+						continue;
+					}
+
+					li = bucket.m_cacheEntryLookup.erase(li);
+					le = bucket.m_cacheEntryLookup.end();
+				}
+			});
 
 			// Reset deleted cas files and update it again..
 			deletedCasFiles.clear();
@@ -391,131 +453,145 @@ namespace uba
 		}
 		while (!deletedCasFiles.empty()); // if cas files are deleted we need to do another loop and check cache entry inputs to see if files were inputs
 
-		m_logger.Detail(TC("  Deleted %llu cas files and %llu cache entries (%s)"), deletedCasCount, deleteEntryCount, TimeToText(GetTime() - deleteCacheEntriesStartTime).str);
+		m_logger.Detail(TC("  Deleted %llu cas files and %llu cache entries (%s)"), deletedCasCount, deleteEntryCount.load(), TimeToText(GetTime() - deleteCacheEntriesStartTime).str);
 
-		if ( deleteEntryCount || forceAllSteps)
+		Atomic<u32> bucketCounter;
+		m_server.ParallelFor(16, m_buckets, [&](auto& it)
 		{
-			UnorderedSet<u32> usedCasKeyOffsets;
+			u64 bucketStartTime = GetTime();
+
+			Bucket& bucket = it->second;
+			u32 bucketIndex = bucketCounter++;
 
-			// Collect all caskeys that are used by cache entries.
-			for (auto& kv : m_cacheEntryLookup)
+			if (deleteEntryCount || forceAllSteps)
 			{
-				for (auto& entry : kv.second.entries)
+				UnorderedSet<u32> usedCasKeyOffsets;
+
+				// Collect all caskeys that are used by cache entries.
+				for (auto& kv2 : bucket.m_cacheEntryLookup)
 				{
-					auto collectUsedCasKeyOffsets = [&](const Vector<u8>& offsets)
-						{
-							BinaryReader reader2(offsets.data(), 0, offsets.size());
-							while (reader2.GetLeft())
-								usedCasKeyOffsets.insert(u32(reader2.Read7BitEncoded()));
-						};
-					collectUsedCasKeyOffsets(entry.inputCasKeyOffsets);
-					collectUsedCasKeyOffsets(entry.outputCasKeyOffsets);
+					for (auto& entry : kv2.second.entries)
+					{
+						auto collectUsedCasKeyOffsets = [&](const Vector<u8>& offsets)
+							{
+								BinaryReader reader2(offsets.data(), 0, offsets.size());
+								while (reader2.GetLeft())
+									usedCasKeyOffsets.insert(u32(reader2.Read7BitEncoded()));
+							};
+						collectUsedCasKeyOffsets(entry.inputCasKeyOffsets);
+						collectUsedCasKeyOffsets(entry.outputCasKeyOffsets);
+					}
 				}
-			}
 
-			u64 recreatePathTableStart = GetTime();
+				u64 recreatePathTableStart = GetTime();
 
-			// Traverse all caskeys in caskey table and figure out which ones we can delete
-			UnorderedSet<u32> usedPathOffsets;
-			usedPathOffsets.reserve(usedCasKeyOffsets.size());
+				// Traverse all caskeys in caskey table and figure out which ones we can delete
+				UnorderedSet<u32> usedPathOffsets;
+				usedPathOffsets.reserve(usedCasKeyOffsets.size());
 
-			for (u32 casKeyOffset : usedCasKeyOffsets)
-			{
-				BinaryReader reader2(m_casKeyTable.GetMemory(), casKeyOffset, m_casKeyTable.GetSize());
-				u32 pathOffset = u32(reader2.Read7BitEncoded());
-				usedPathOffsets.insert(pathOffset);
-			}
-
-			// Build new path table based on used offsets
-			UnorderedMap<u32, u32> oldToNewPathOffset;
-			u32 oldSize = m_pathTable.GetSize();
-			{
-				CompactPathTable newPathTable(CachePathTableMaxSize, CompactPathTable::V1, m_pathTable.GetPathCount(), m_pathTable.GetSegmentCount());
-				oldToNewPathOffset.reserve(usedPathOffsets.size());
+				for (u32 casKeyOffset : usedCasKeyOffsets)
+				{
+					BinaryReader reader2(bucket.m_casKeyTable.GetMemory(), casKeyOffset, bucket.m_casKeyTable.GetSize());
+					u32 pathOffset = u32(reader2.Read7BitEncoded());
+					usedPathOffsets.insert(pathOffset);
+				}
 
-				for (u32 pathOffset : usedPathOffsets)
+				// Build new path table based on used offsets
+				UnorderedMap<u32, u32> oldToNewPathOffset;
+				u32 oldSize = bucket.m_pathTable.GetSize();
 				{
-					StringBuffer<> temp;
-					m_pathTable.GetString(temp, pathOffset);
-					u32 newOffset = newPathTable.AddNoLock(temp.data, temp.count);
-
-					#if 0
-					StringBuffer<> test;
-					newPathTable.GetString(test, newOffset);
-					UBA_ASSERT(test.Equals(temp.data));
-					#endif
-
-					auto res = oldToNewPathOffset.try_emplace(pathOffset, newOffset);
-					UBA_ASSERT(res.second);(void)res;
+					CompactPathTable newPathTable(CachePathTableMaxSize, CompactPathTable::V1, bucket.m_pathTable.GetPathCount(), bucket.m_pathTable.GetSegmentCount());
+					oldToNewPathOffset.reserve(usedPathOffsets.size());
+
+					for (u32 pathOffset : usedPathOffsets)
+					{
+						StringBuffer<> temp;
+						bucket.m_pathTable.GetString(temp, pathOffset);
+						u32 newOffset = newPathTable.AddNoLock(temp.data, temp.count);
+
+						#if 0
+						StringBuffer<> test;
+						newPathTable.GetString(test, newOffset);
+						UBA_ASSERT(test.Equals(temp.data));
+						#endif
+
+						auto res = oldToNewPathOffset.try_emplace(pathOffset, newOffset);
+						UBA_ASSERT(res.second);(void)res;
+					}
+					bucket.m_pathTable.Swap(newPathTable);
 				}
-				m_pathTable.Swap(newPathTable);
-			}
-			m_logger.Detail(TC("  Recreated path table. %s -> %s (%s)"), BytesToText(oldSize).str, BytesToText(m_pathTable.GetSize()).str, TimeToText(GetTime() - recreatePathTableStart).str);
+				m_logger.Detail(TC("    Bucket %u Recreated path table. %s -> %s (%s)"), bucketIndex, BytesToText(oldSize).str, BytesToText(bucket.m_pathTable.GetSize()).str, TimeToText(GetTime() - recreatePathTableStart).str);
 
 
-			// Build new caskey table based on used offsets
-			u64 recreateCasKeyTableStart = GetTime();
-			UnorderedMap<u32, u32> oldToNewCasKeyOffset;
-			oldSize = m_casKeyTable.GetSize();
-			{
-				oldToNewCasKeyOffset.reserve(usedCasKeyOffsets.size());
-				CompactCasKeyTable newCasKeyTable(CacheCasKeyTableMaxSize, usedCasKeyOffsets.size());
-				for (u32 casKeyOffset : usedCasKeyOffsets)
+				// Build new caskey table based on used offsets
+				u64 recreateCasKeyTableStart = GetTime();
+				UnorderedMap<u32, u32> oldToNewCasKeyOffset;
+				oldSize = bucket.m_casKeyTable.GetSize();
 				{
-					BinaryReader reader2(m_casKeyTable.GetMemory(), casKeyOffset, m_casKeyTable.GetSize());
-					u32 oldPathOffset = u32(reader2.Read7BitEncoded());
-					CasKey casKey = reader2.ReadCasKey();
-					auto findIt = oldToNewPathOffset.find(oldPathOffset);
-					UBA_ASSERT(findIt != oldToNewPathOffset.end());
-					auto res = oldToNewCasKeyOffset.try_emplace(casKeyOffset, newCasKeyTable.Add(casKey, findIt->second));
-					UBA_ASSERT(res.second);(void)res;
+					oldToNewCasKeyOffset.reserve(usedCasKeyOffsets.size());
+					CompactCasKeyTable newCasKeyTable(CacheCasKeyTableMaxSize, usedCasKeyOffsets.size());
+					for (u32 casKeyOffset : usedCasKeyOffsets)
+					{
+						BinaryReader reader2(bucket.m_casKeyTable.GetMemory(), casKeyOffset, bucket.m_casKeyTable.GetSize());
+						u32 oldPathOffset = u32(reader2.Read7BitEncoded());
+						CasKey casKey = reader2.ReadCasKey();
+						auto findIt = oldToNewPathOffset.find(oldPathOffset);
+						UBA_ASSERT(findIt != oldToNewPathOffset.end());
+						auto res = oldToNewCasKeyOffset.try_emplace(casKeyOffset, newCasKeyTable.Add(casKey, findIt->second));
+						UBA_ASSERT(res.second);(void)res;
+					}
+					bucket.m_casKeyTable.Swap(newCasKeyTable);
 				}
-				m_casKeyTable.Swap(newCasKeyTable);
-			}
-			m_logger.Detail(TC("  Recreated caskey table. %s -> %s (%s)"), BytesToText(oldSize).str, BytesToText(m_casKeyTable.GetSize()).str, TimeToText(GetTime() - recreateCasKeyTableStart).str);
+				m_logger.Detail(TC("    Bucket %u Recreated caskey table. %s -> %s (%s)"), bucketIndex, BytesToText(oldSize).str, BytesToText(bucket.m_casKeyTable.GetSize()).str, TimeToText(GetTime() - recreateCasKeyTableStart).str);
 
 
-			// Update all casKeyOffsets
-			u64 updateEntriesStart = GetTime();
-			for (auto& kv : m_cacheEntryLookup)
-			{
-				for (auto& entry : kv.second.entries)
-				{
-					auto updateCasKeyOffsets = [&](Vector<u8>& offsets)
+				// Update all casKeyOffsets
+				u64 updateEntriesStart = GetTime();
+
+				auto updateCasKeyOffsets = [&](Vector<u8>& offsets)
+					{
+						Set<u32> newOffsets; // Want these sorted since that is what we get from uploads
+						u32 newOffsetsSize = 0;
+						BinaryReader reader2(offsets.data(), 0, offsets.size());
+						while (reader2.GetLeft())
 						{
-							Set<u32> newOffsets; // Want these sorted since that is what we get from uploads
-							u32 newOffsetsSize = 0;
-							BinaryReader reader2(offsets.data(), 0, offsets.size());
-							while (reader2.GetLeft())
+							u32 oldOffset = u32(reader2.Read7BitEncoded());
+							auto findIt = oldToNewCasKeyOffset.find(oldOffset);
+							UBA_ASSERT(findIt != oldToNewCasKeyOffset.end());
+							u32 newOffset = findIt->second;
+							if (newOffsets.insert(newOffset).second)
+								newOffsetsSize += Get7BitEncodedCount(newOffset);
+							else
 							{
-								u32 oldOffset = u32(reader2.Read7BitEncoded());
-								auto findIt = oldToNewCasKeyOffset.find(oldOffset);
-								UBA_ASSERT(findIt != oldToNewCasKeyOffset.end());
-								u32 newOffset = findIt->second;
-								if (newOffsets.insert(newOffset).second)
-									newOffsetsSize += Get7BitEncodedCount(newOffset);
-								else
-								{
-									StringBuffer<> temp;
-									CasKey key;
-									m_casKeyTable.GetPathAndKey(temp, key, m_pathTable, newOffset);
-									m_logger.Info(TC("Found duplicate of %s in cache entry"), temp.data);
-								}
+								StringBuffer<> temp;
+								CasKey key;
+								bucket.m_casKeyTable.GetPathAndKey(temp, key, bucket.m_pathTable, newOffset);
+								m_logger.Info(TC("Found duplicate of %s in cache entry"), temp.data);
 							}
+						}
 
-							offsets.resize(newOffsetsSize);
-							BinaryWriter writer2(offsets.data(), 0, newOffsetsSize);
-							for (u32 offset : newOffsets)
-								writer2.Write7BitEncoded(offset);
-							UBA_ASSERT(writer2.GetPosition() == newOffsetsSize);
-						};
-					updateCasKeyOffsets(entry.inputCasKeyOffsets);
-					updateCasKeyOffsets(entry.outputCasKeyOffsets);
-				}
+						offsets.resize(newOffsetsSize);
+						BinaryWriter writer2(offsets.data(), 0, newOffsetsSize);
+						for (u32 offset : newOffsets)
+							writer2.Write7BitEncoded(offset);
+						UBA_ASSERT(writer2.GetPosition() == newOffsetsSize);
+					};
+
+				m_server.ParallelFor(16, bucket.m_cacheEntryLookup, [&](auto& it)
+					{
+						CacheEntries& entries = it->second;(void)entries;
+						for (auto& entry : entries.entries)
+						{
+							updateCasKeyOffsets(entry.inputCasKeyOffsets);
+							updateCasKeyOffsets(entry.outputCasKeyOffsets);
+						}
+					});
+
+				m_logger.Detail(TC("    Bucket %u Updated cache entries with new tables (%s)"), bucketIndex, TimeToText(GetTime() - updateEntriesStart).str);
 			}
 
-			m_logger.Detail(TC("  Updated cache entries with new tables (%s)"), TimeToText(GetTime() - updateEntriesStart).str);
-		}
+			m_logger.Info(TC("    Bucket %u Done (%s). CacheEntries: %llu (%s) PathTable: %s CasTable: %s"), bucketIndex, TimeToText(GetTime() - bucketStartTime).str, bucket.totalEntryCount, BytesToText(bucket.totalEntrySize).str, BytesToText(bucket.m_pathTable.GetSize()).str, BytesToText(bucket.m_casKeyTable.GetSize()).str);
+		});
 
 		if (entriesAdded || deletedCasCount || deleteEntryCount || forceAllSteps)
 		{
@@ -526,7 +602,7 @@ namespace uba
 		}
 
 		u64 oldestTime = MsToTime(GetFileTimeAsSeconds(now - oldest)*1000);
-		m_logger.Info(TC("Maintenance done! (%s). CacheEntries: %llu (%s) CasFiles: %llu (%s) PathTable: %s CasTable: %s OldestEntry: %s"), TimeToText(GetTime() - startTime).str, totalEntryCount, BytesToText(totalEntrySize).str, totalCasCount, BytesToText(totalCasSize).str, BytesToText(m_pathTable.GetSize()).str, BytesToText(m_casKeyTable.GetSize()).str, TimeToText(oldestTime, true).str);
+		m_logger.Info(TC("Maintenance done! (%s) CasFiles: %llu (%s) OldestEntry: %s"), TimeToText(GetTime() - startTime).str, totalCasCount - deletedCasCount, BytesToText(totalCasSize).str, TimeToText(oldestTime, true).str);
 
 		return true;
 	}
@@ -554,40 +630,52 @@ namespace uba
 		}
 		case CacheMessageType_StorePathTable:
 		{
+			u32 bucketId = u32(reader.Read7BitEncoded());
 			SCOPED_WRITE_LOCK(m_connectionsLock, lock);
 			auto& connection = m_connections[connectionInfo.GetId()];
+			auto& connectionBucket = connection.buckets[bucketId];
 			lock.Leave();
-			connection.m_pathTable.ReadMem(reader, false);
+			connectionBucket.pathTable.ReadMem(reader, false);
 			return true;
 		}
 		case CacheMessageType_StoreCasTable:
 		{
+			u32 bucketId = u32(reader.Read7BitEncoded());
 			SCOPED_WRITE_LOCK(m_connectionsLock, lock);
 			auto& connection = m_connections[connectionInfo.GetId()];
+			auto& connectionBucket = connection.buckets[bucketId];
 			lock.Leave();
-			connection.m_casKeyTable.ReadMem(reader, false);
+			connectionBucket.casKeyTable.ReadMem(reader, false);
 			return true;
 		}
 		case CacheMessageType_StoreEntry:
 		{
+			u32 bucketId = u32(reader.Read7BitEncoded());
 			SCOPED_WRITE_LOCK(m_connectionsLock, lock);
 			auto& connection = m_connections[connectionInfo.GetId()];
+			auto& connectionBucket = connection.buckets[bucketId];
 			lock.Leave();
-			return HandleStoreEntry(connection, reader, writer);
+			return HandleStoreEntry(bucketId, connectionBucket, reader, writer);
 		}
 		case CacheMessageType_StoreEntryDone:
 		{
+			u32 bucketId = u32(reader.Read7BitEncoded());
 			SCOPED_WRITE_LOCK(m_connectionsLock, lock);
 			auto& connection = m_connections[connectionInfo.GetId()];
+			auto& connectionBucket = connection.buckets[bucketId];
 			lock.Leave();
 			CasKey cmdKey = reader.ReadCasKey();
 
-			SCOPED_WRITE_LOCK(connection.cacheEntryLookupLock, lock2);
-			auto findIt = connection.cacheEntryLookup.find(cmdKey);
-			if (findIt != connection.cacheEntryLookup.end())
+			SCOPED_WRITE_LOCK(connectionBucket.cacheEntryLookupLock, lock2);
+			auto findIt = connectionBucket.cacheEntryLookup.find(cmdKey);
+			if (findIt != connectionBucket.cacheEntryLookup.end())
 			{
-				SCOPED_WRITE_LOCK(m_cacheEntryLookupLock, lock3);
-				auto insres = m_cacheEntryLookup.try_emplace(cmdKey);
+				SCOPED_WRITE_LOCK(m_bucketsLock, bucketsLock);
+				Bucket& bucket = m_buckets[bucketId];
+				bucketsLock.Leave();
+
+				SCOPED_WRITE_LOCK(bucket.m_cacheEntryLookupLock, lock3);
+				auto insres = bucket.m_cacheEntryLookup.try_emplace(cmdKey);
 				auto& cacheEntries = insres.first->second;
 				lock2.Leave();
 
@@ -615,7 +703,7 @@ namespace uba
 		}
 	}
 
-	bool CacheServer::HandleStoreEntry(Connection& connection, BinaryReader& reader, BinaryWriter& writer)
+	bool CacheServer::HandleStoreEntry(u32 bucketId, ConnectionBucket& connectionBucket, BinaryReader& reader, BinaryWriter& writer)
 	{
 		CasKey cmdKey = reader.ReadCasKey();
 
@@ -627,6 +715,10 @@ namespace uba
 
 		u64 outputStartOffset = reader.GetPosition();
 
+		SCOPED_WRITE_LOCK(m_bucketsLock, bucketsLock);
+		Bucket& bucket = m_buckets[bucketId];
+		bucketsLock.Leave();
+
 		while (reader.GetLeft())
 		{
 			u32 offset = u32(reader.Read7BitEncoded());
@@ -636,18 +728,18 @@ namespace uba
 
 			CasKey casKey;
 			StringBuffer<> path;
-			connection.m_casKeyTable.GetPathAndKey(path, casKey, connection.m_pathTable, offset);
+			connectionBucket.casKeyTable.GetPathAndKey(path, casKey, connectionBucket.pathTable, offset);
 			UBA_ASSERT(path.count);
 
-			u32 pathOffset = m_pathTable.Add(path.data, path.count);
+			u32 pathOffset = bucket.m_pathTable.Add(path.data, path.count);
 
 			#if 0
 			StringBuffer<> test;
-			m_pathTable.GetString(test, pathOffset);
+			bucket.m_pathTable.GetString(test, pathOffset);
 			UBA_ASSERT(test.Equals(path.data));
 			#endif
 
-			u32 casKeyOffset = m_casKeyTable.Add(casKey, pathOffset);
+			u32 casKeyOffset = bucket.m_casKeyTable.Add(casKey, pathOffset);
 			inputs.insert(casKeyOffset);
 			bytesForInput += Get7BitEncodedCount(casKeyOffset);
 
@@ -662,8 +754,8 @@ namespace uba
 				w2.Write7BitEncoded(input);
 		}
 
-		SCOPED_WRITE_LOCK(m_cacheEntryLookupLock, lock);
-		auto insres = m_cacheEntryLookup.try_emplace(cmdKey);
+		SCOPED_WRITE_LOCK(bucket.m_cacheEntryLookupLock, lock);
+		auto insres = bucket.m_cacheEntryLookup.try_emplace(cmdKey);
 		auto& cacheEntries = insres.first->second;
 		lock.Leave();
 
@@ -690,7 +782,7 @@ namespace uba
 				u32 existingOffset = u32(r2.Read7BitEncoded());
 				CasKey casKey;
 				StringBuffer<> path;
-				m_casKeyTable.GetPathAndKey(path, casKey, m_pathTable, existingOffset);
+				bucket.m_casKeyTable.GetPathAndKey(path, casKey, bucket.m_pathTable, existingOffset);
 				existing.try_emplace(path.data, casKey);
 			}
 
@@ -701,7 +793,7 @@ namespace uba
 				u32 outputOffset = u32(reader.Read7BitEncoded());
 				CasKey casKey;
 				StringBuffer<> path;
-				connection.m_casKeyTable.GetPathAndKey(path, casKey, connection.m_pathTable, outputOffset);
+				connectionBucket.casKeyTable.GetPathAndKey(path, casKey, connectionBucket.pathTable, outputOffset);
 
 				auto findIt = existing.find(path.data);
 				if (findIt == existing.end())
@@ -738,16 +830,16 @@ namespace uba
 			u32 outputOffset = u32(reader.Read7BitEncoded());
 			CasKey casKey;
 			StringBuffer<> path;
-			connection.m_casKeyTable.GetPathAndKey(path, casKey, connection.m_pathTable, outputOffset);
-			u32 pathOffset = m_pathTable.Add(path.data, path.count);
+			connectionBucket.casKeyTable.GetPathAndKey(path, casKey, connectionBucket.pathTable, outputOffset);
+			u32 pathOffset = bucket.m_pathTable.Add(path.data, path.count);
 
 			#if 0
 			StringBuffer<> test;
-			m_pathTable.GetString(test, pathOffset);
+			bucket.m_pathTable.GetString(test, pathOffset);
 			UBA_ASSERT(test.Equals(path.data));
 			#endif
 
-			u32 casKeyOffset = m_casKeyTable.Add(casKey, pathOffset);
+			u32 casKeyOffset = bucket.m_casKeyTable.Add(casKey, pathOffset);
 			outputs.insert(casKeyOffset);
 			bytesForOutput += Get7BitEncodedCount(casKeyOffset);
 
@@ -773,8 +865,8 @@ namespace uba
 		}
 		else
 		{
-			SCOPED_WRITE_LOCK(connection.cacheEntryLookupLock, lock3);
-			bool res = connection.cacheEntryLookup.try_emplace(cmdKey, std::move(newEntry)).second;
+			SCOPED_WRITE_LOCK(connectionBucket.cacheEntryLookupLock, lock3);
+			bool res = connectionBucket.cacheEntryLookup.try_emplace(cmdKey, std::move(newEntry)).second;
 			UBA_ASSERT(res);(void)res;
 		}
 
@@ -787,34 +879,52 @@ namespace uba
 
 	bool CacheServer::HandleFetchPathTable(BinaryReader& reader, BinaryWriter& writer)
 	{
+		u32 bucketId = u32(reader.Read7BitEncoded());
+
+		SCOPED_WRITE_LOCK(m_bucketsLock, bucketsLock);
+		Bucket& bucket = m_buckets[bucketId];
+		bucketsLock.Leave();
+
 		u32 haveSize = reader.ReadU32();
-		u32 size = m_pathTable.GetSize();
+		u32 size = bucket.m_pathTable.GetSize();
 		writer.WriteU32(size);
 		u32 toSend = Min(u32(writer.GetCapacityLeft()), size - haveSize);
-		writer.WriteBytes(m_pathTable.GetMemory() + haveSize, toSend);
+		writer.WriteBytes(bucket.m_pathTable.GetMemory() + haveSize, toSend);
 		return true;
 	}
 
 	bool CacheServer::HandleFetchCasTable(BinaryReader& reader, BinaryWriter& writer)
 	{
+		u32 bucketId = u32(reader.Read7BitEncoded());
+
+		SCOPED_WRITE_LOCK(m_bucketsLock, bucketsLock);
+		Bucket& bucket = m_buckets[bucketId];
+		bucketsLock.Leave();
+
 		u32 haveSize = reader.ReadU32();
-		u32 size = m_casKeyTable.GetSize();
+		u32 size = bucket.m_casKeyTable.GetSize();
 		writer.WriteU32(size);
 		u32 toSend = Min(u32(writer.GetCapacityLeft()), size - haveSize);
-		writer.WriteBytes(m_casKeyTable.GetMemory() + haveSize, toSend);
+		writer.WriteBytes(bucket.m_casKeyTable.GetMemory() + haveSize, toSend);
 		return true;
 	}
 
 	bool CacheServer::HandleFetchEntries(BinaryReader& reader, BinaryWriter& writer)
 	{
+		u32 bucketId = u32(reader.Read7BitEncoded());
+
+		SCOPED_WRITE_LOCK(m_bucketsLock, bucketsLock);
+		Bucket& bucket = m_buckets[bucketId];
+		bucketsLock.Leave();
+
 		CasKey cmdKey = reader.ReadCasKey();
 
 		u16& entryCount = *(u16*)writer.AllocWrite(2);
 		entryCount = 0;
 
-		SCOPED_READ_LOCK(m_cacheEntryLookupLock, lock);
-		auto findIt = m_cacheEntryLookup.find(cmdKey);
-		if (findIt == m_cacheEntryLookup.end())
+		SCOPED_READ_LOCK(bucket.m_cacheEntryLookupLock, lock);
+		auto findIt = bucket.m_cacheEntryLookup.find(cmdKey);
+		if (findIt == bucket.m_cacheEntryLookup.end())
 			return true;
 		auto& cacheEntries = findIt->second;
 		lock.Leave();
@@ -845,7 +955,6 @@ namespace uba
 		reader.ReadString(filterString);
 
 		SCOPED_READ_LOCK(m_maintenanceLock, lock);
-		SCOPED_READ_LOCK(m_cacheEntryLookupLock, lock2);
 
 		StringBuffer<> tempFile(m_storage.GetTempPath());
 		Guid guid;
@@ -875,18 +984,62 @@ namespace uba
 
 		u64 now = GetSystemTimeAsFileTime();
 
-		for (auto& kv : m_cacheEntryLookup)
+		SCOPED_READ_LOCK(m_bucketsLock, bucketsLock);
+		for (auto& kv : m_buckets)
 		{
-			CacheEntries& entries = kv.second;
-			SCOPED_READ_LOCK(entries.lock, lock3);
+			Bucket& bucket = kv.second;
+			SCOPED_READ_LOCK(bucket.m_cacheEntryLookupLock, lock2);
 
-			Set<u32> visibleIndices;
-			if (filterString.count)
+			for (auto& kv2 : bucket.m_cacheEntryLookup)
 			{
+				CacheEntries& entries = kv2.second;
+				SCOPED_READ_LOCK(entries.lock, lock3);
+
+				Set<u32> visibleIndices;
+				if (filterString.count)
+				{
+					u32 index = 0;
+					for (auto& entry : entries.entries)
+					{
+						auto findString = [&](const Vector<u8>& offsets)
+							{
+								BinaryReader reader2(offsets.data(), 0, offsets.size());
+								while (reader2.GetLeft())
+								{
+									u64 offset = reader2.Read7BitEncoded();
+									CasKey casKey;
+									StringBuffer<> path;
+									bucket.m_casKeyTable.GetPathAndKey(path, casKey, bucket.m_pathTable, offset);
+									if (path.Contains(filterString.data))
+										return true;
+								}
+								return false;
+							};
+
+						if (findString(entry.inputCasKeyOffsets) || findString(entry.outputCasKeyOffsets))
+							visibleIndices.insert(index);
+						++index;
+					}
+					if (visibleIndices.empty())
+						continue;
+				}
+
+
+				StringBuffer<> line;
+				writeLine(CasKeyString(kv2.first).str);
 				u32 index = 0;
 				for (auto& entry : entries.entries)
 				{
-					auto findString = [&](const Vector<u8>& offsets)
+					if (!visibleIndices.empty() && visibleIndices.find(index) == visibleIndices.end())
+					{
+						++index;
+						continue;
+					}
+
+					u64 age = MsToTime(GetFileTimeAsSeconds(now - entry.creationTime)*1000);
+					writeLine(line.Clear().Appendf(TC("  #%u (%s ago)"), index, TimeToText(age, true).str).data);
+
+					auto writeOffsets = [&](const Vector<u8>& offsets)
 						{
 							BinaryReader reader2(offsets.data(), 0, offsets.size());
 							while (reader2.GetLeft())
@@ -894,54 +1047,17 @@ namespace uba
 								u64 offset = reader2.Read7BitEncoded();
 								CasKey casKey;
 								StringBuffer<> path;
-								m_casKeyTable.GetPathAndKey(path, casKey, m_pathTable, offset);
-								if (path.Contains(filterString.data))
-									return true;
+								bucket.m_casKeyTable.GetPathAndKey(path, casKey, bucket.m_pathTable, offset);
+								writeLine(line.Clear().Appendf(TC("    %s - %s"), path.data, CasKeyString(casKey).str).data);
 							}
-							return false;
 						};
 
-					if (findString(entry.inputCasKeyOffsets) || findString(entry.outputCasKeyOffsets))
-						visibleIndices.insert(index);
+					writeLine(line.Clear().Append (TC("   Inputs:")).data);
+					writeOffsets(entry.inputCasKeyOffsets);
+					writeLine(line.Clear().Append (TC("   Outputs:")).data);
+					writeOffsets(entry.outputCasKeyOffsets);
 					++index;
 				}
-				if (visibleIndices.empty())
-					continue;
-			}
-
-
-			StringBuffer<> line;
-			writeLine(CasKeyString(kv.first).str);
-			u32 index = 0;
-			for (auto& entry : entries.entries)
-			{
-				if (!visibleIndices.empty() && visibleIndices.find(index) == visibleIndices.end())
-				{
-					++index;
-					continue;
-				}
-
-				u64 age = MsToTime(GetFileTimeAsSeconds(now - entry.creationTime)*1000);
-				writeLine(line.Clear().Appendf(TC("  #%u (%s ago)"), index, TimeToText(age, true).str).data);
-
-				auto writeOffsets = [&](const Vector<u8>& offsets)
-					{
-						BinaryReader reader2(offsets.data(), 0, offsets.size());
-						while (reader2.GetLeft())
-						{
-							u64 offset = reader2.Read7BitEncoded();
-							CasKey casKey;
-							StringBuffer<> path;
-							m_casKeyTable.GetPathAndKey(path, casKey, m_pathTable, offset);
-							writeLine(line.Clear().Appendf(TC("    %s - %s"), path.data, CasKeyString(casKey).str).data);
-						}
-					};
-
-				writeLine(line.Clear().Append (TC("   Inputs:")).data);
-				writeOffsets(entry.inputCasKeyOffsets);
-				writeLine(line.Clear().Append (TC("   Outputs:")).data);
-				writeOffsets(entry.outputCasKeyOffsets);
-				++index;
 			}
 		}
 
diff --git a/Engine/Source/Programs/UnrealBuildAccelerator/Common/Private/UbaNetworkServer.cpp b/Engine/Source/Programs/UnrealBuildAccelerator/Common/Private/UbaNetworkServer.cpp
--- a/Engine/Source/Programs/UnrealBuildAccelerator/Common/Private/UbaNetworkServer.cpp
+++ b/Engine/Source/Programs/UnrealBuildAccelerator/Common/Private/UbaNetworkServer.cpp
@@ -804,11 +804,23 @@ namespace uba
 		}
 	}
 
+	void NetworkServer::DoWork(u32 count)
+	{
+		while (count--)
+			if (!DoAdditionalWork())
+				return;
+	}
+
 	u32 NetworkServer::GetWorkerCount()
 	{
 		return m_maxWorkerCount;
 	}
 
+	MutableLogger& NetworkServer::GetLogger()
+	{
+		return m_logger;
+	}
+
 	u64 NetworkServer::GetTotalSentBytes()
 	{
 		return m_sendBytes;
diff --git a/Engine/Source/Programs/UnrealBuildAccelerator/Common/Private/UbaScheduler.cpp b/Engine/Source/Programs/UnrealBuildAccelerator/Common/Private/UbaScheduler.cpp
--- a/Engine/Source/Programs/UnrealBuildAccelerator/Common/Private/UbaScheduler.cpp
+++ b/Engine/Source/Programs/UnrealBuildAccelerator/Common/Private/UbaScheduler.cpp
@@ -406,7 +406,7 @@ namespace uba
 						ProcessStartInfo& si = exitInfo->startInfo->startInfo;
 						u64 startTime = GetTime();
 
-						if (m_cacheClient->FetchFromCache(*m_rootPaths[0], si))
+						if (m_cacheClient->FetchFromCache(*m_rootPaths[0], 0, si))
 						{
 							auto process = new CachedProcess(si);
 							ProcessHandle ph(process);
@@ -516,7 +516,7 @@ namespace uba
 		delete si;
 
 		if (m_writeToCache && exitCode == 0 && process.GetStartInfo().trackInputs)
-			m_cacheClient->WriteToCache(*m_rootPaths[0], ph);
+			m_cacheClient->WriteToCache(*m_rootPaths[0], 0, ph);
 
 		ph.m_process = nullptr;
 	}
diff --git a/Engine/Source/Programs/UnrealBuildAccelerator/Common/Private/UbaSession.cpp b/Engine/Source/Programs/UnrealBuildAccelerator/Common/Private/UbaSession.cpp
--- a/Engine/Source/Programs/UnrealBuildAccelerator/Common/Private/UbaSession.cpp
+++ b/Engine/Source/Programs/UnrealBuildAccelerator/Common/Private/UbaSession.cpp
@@ -1129,8 +1129,6 @@ namespace uba
 
 	void Session::CancelAllProcessesAndWait(bool terminate)
 	{
-		m_logger.isMuted = true;
-
 		bool isEmpty = false;
 		bool isFirst = true;
 		while (!isEmpty)
@@ -1149,6 +1147,7 @@ namespace uba
 				isFirst = false;
 				if (!processes.empty())
 					m_logger.Info(TC("Cancelling %llu processes and wait for them to exit"), processes.size());
+				++m_logger.isMuted;
 			}
 
 			for (auto& process : processes)
@@ -1167,7 +1166,7 @@ namespace uba
 				process.WaitForExit(100000);
 		}
 
-		m_logger.isMuted = false;
+		--m_logger.isMuted;
 	}
 
 	ProcessHandle Session::RunProcess(const ProcessStartInfo& startInfo, bool async, bool enableDetour)
@@ -1309,7 +1308,7 @@ namespace uba
 
 	const tchar* Session::GetId() { return m_id.data; }
 	Storage& Session::GetStorage() { return m_storage; }
-	Logger& Session::GetLogger() { return m_logger; }
+	MutableLogger& Session::GetLogger() { return m_logger; }
 	LogWriter& Session::GetLogWriter() { return m_logger.m_writer; }
 	Trace& Session::GetTrace() { return m_trace; }
 
diff --git a/Engine/Source/Programs/UnrealBuildAccelerator/Common/Public/UbaCacheClient.h b/Engine/Source/Programs/UnrealBuildAccelerator/Common/Public/UbaCacheClient.h
--- a/Engine/Source/Programs/UnrealBuildAccelerator/Common/Public/UbaCacheClient.h
+++ b/Engine/Source/Programs/UnrealBuildAccelerator/Common/Public/UbaCacheClient.h
@@ -2,16 +2,19 @@
 
 #pragma once
 
-#include "UbaCompactTables.h"
 #include "UbaLogger.h"
+#include "UbaMemory.h"
 
 namespace uba
 {
+	class CompactCasKeyTable;
+	class CompactPathTable;
 	class NetworkClient;
 	class ProcessHandle;
 	class RootPaths;
 	class Session;
 	class StorageImpl;
+	struct CasKey;
 	struct ProcessStartInfo;
 
 	class CacheClient
@@ -20,8 +23,8 @@ namespace uba
 		CacheClient(LogWriter& writer, StorageImpl& storage, NetworkClient& client, Session& session);
 		~CacheClient();
 
-		bool WriteToCache(const RootPaths& rootPaths, const ProcessHandle& process);
-		bool FetchFromCache(const RootPaths& rootPaths, const ProcessStartInfo& info);
+		bool WriteToCache(const RootPaths& rootPaths, u32 bucketId, const ProcessHandle& process);
+		bool FetchFromCache(const RootPaths& rootPaths, u32 bucketId, const ProcessStartInfo& info);
 
 		bool WriteCacheSummary(const tchar* destinationFile, const tchar* filterString = nullptr);
 
@@ -30,14 +33,16 @@ namespace uba
 		inline StorageImpl& GetStorage() { return m_storage; }
 
 	private:
-		bool SendPathTable(u32 requiredPathTableSize);
-		bool SendCasTable(u32 requiredCasTableSize);
-		bool SendCacheEntry(const RootPaths& rootPaths, const CasKey& cmdKey, const Map<u32, u32>& inputsStringToCasKey, const Map<u32, u32>& outputsStringToCasKey);
-		bool FetchCasTable();
+		struct Bucket;
+
+		bool SendPathTable(Bucket& bucket, u32 requiredPathTableSize);
+		bool SendCasTable(Bucket& bucket, u32 requiredCasTableSize);
+		bool SendCacheEntry(Bucket& bucket, const RootPaths& rootPaths, const CasKey& cmdKey, const Map<u32, u32>& inputsStringToCasKey, const Map<u32, u32>& outputsStringToCasKey);
+		bool FetchCasTable(Bucket& bucket);
 
 		CasKey GetCmdKey(const RootPaths& rootPaths, const ProcessStartInfo& info);
 
-		bool GetLocalPathAndCasKey(const RootPaths& rootPaths, StringBufferBase& outPath, CasKey& outKey, CompactCasKeyTable& casKeyTable, CompactPathTable& pathTable, u32 offset);
+		bool GetLocalPathAndCasKey(Bucket& bucket, const RootPaths& rootPaths, StringBufferBase& outPath, CasKey& outKey, CompactCasKeyTable& casKeyTable, CompactPathTable& pathTable, u32 offset);
 
 		MutableLogger m_logger;
 		StorageImpl& m_storage;
@@ -46,18 +51,12 @@ namespace uba
 
 		Atomic<bool> m_connected;
 
-		CompactPathTable m_serverPathTable;
-		CompactCasKeyTable m_serverCasKeyTable;
-
-		CompactPathTable m_sendPathTable;
-		CompactCasKeyTable m_sendCasKeyTable;
-
-		ReaderWriterLock m_pathTableNetworkLock;
-		u32 m_pathTableSizeSent = 0;
-
-		ReaderWriterLock m_casKeyTableNetworkLock;
-		u32 m_casKeyTableSizeSent = 0;
+		ReaderWriterLock m_bucketsLock;
+		UnorderedMap<u32, Bucket> m_buckets;
 
 		ReaderWriterLock m_sendOneAtTheTimeLock;
+
+		CacheClient(const CacheClient&) = delete;
+		CacheClient& operator=(const CacheClient&) = delete;
 	};
 }
\ No newline at end of file
diff --git a/Engine/Source/Programs/UnrealBuildAccelerator/Common/Public/UbaCacheServer.h b/Engine/Source/Programs/UnrealBuildAccelerator/Common/Public/UbaCacheServer.h
--- a/Engine/Source/Programs/UnrealBuildAccelerator/Common/Public/UbaCacheServer.h
+++ b/Engine/Source/Programs/UnrealBuildAccelerator/Common/Public/UbaCacheServer.h
@@ -2,9 +2,8 @@
 
 #pragma once
 
-#include "UbaCompactTables.h"
-#include "UbaHash.h"
 #include "UbaLogger.h"
+#include "UbaStringBuffer.h"
 
 namespace uba
 {
@@ -30,8 +29,9 @@ namespace uba
 		void OnDisconnected(u32 clientId);
 
 		struct Connection;
+		struct ConnectionBucket;
 		bool HandleMessage(const ConnectionInfo& connectionInfo, u8 messageType, BinaryReader& reader, BinaryWriter& writer);
-		bool HandleStoreEntry(Connection& connection, BinaryReader& reader, BinaryWriter& writer);
+		bool HandleStoreEntry(u32 bucketId, ConnectionBucket& bucket, BinaryReader& reader, BinaryWriter& writer);
 		bool HandleFetchPathTable(BinaryReader& reader, BinaryWriter& writer);
 		bool HandleFetchCasTable(BinaryReader& reader, BinaryWriter& writer);
 		bool HandleFetchEntries(BinaryReader& reader, BinaryWriter& writer);
@@ -46,26 +46,17 @@ namespace uba
 		ReaderWriterLock m_maintenanceLock;
 		Atomic<u32> m_addsSinceMaintenance;
 
-		struct CacheEntry
-		{
-			u64 creationTime;
-			Vector<u8> inputCasKeyOffsets;
-			Vector<u8> outputCasKeyOffsets;
-		};
+		struct CacheEntry;
+		struct CacheEntries;
+		struct Bucket;
 
-		struct CacheEntries
-		{
-			ReaderWriterLock lock;
-			List<CacheEntry> entries;
-		};
-
-		ReaderWriterLock m_cacheEntryLookupLock;
-		UnorderedMap<CasKey, CacheEntries> m_cacheEntryLookup;
-
-		CompactPathTable m_pathTable;
-		CompactCasKeyTable m_casKeyTable;
+		ReaderWriterLock m_bucketsLock;
+		Map<u32, Bucket> m_buckets;
 
 		ReaderWriterLock m_connectionsLock;
 		Map<u32, Connection> m_connections;
+
+		CacheServer(const CacheServer&) = delete;
+		CacheServer& operator=(const CacheServer&) = delete;
 	};
 }
\ No newline at end of file
diff --git a/Engine/Source/Programs/UnrealBuildAccelerator/Common/Public/UbaDefaultConstants.h b/Engine/Source/Programs/UnrealBuildAccelerator/Common/Public/UbaDefaultConstants.h
new file mode 100644
--- /dev/null
+++ b/Engine/Source/Programs/UnrealBuildAccelerator/Common/Public/UbaDefaultConstants.h
@@ -0,0 +1,14 @@
+// Copyright Epic Games, Inc. All Rights Reserved.
+
+#pragma once
+
+#include "UbaBase.h"
+
+namespace uba
+{
+	// Default constants
+	static constexpr u16 DefaultPort = 1345;
+	static constexpr u16 DefaultStorageProxyPort = DefaultPort + 1;
+	static constexpr u16 DefaultCachePort = DefaultPort + 2;
+	static constexpr u32 SendDefaultSize = 256*1024;
+}
diff --git a/Engine/Source/Programs/UnrealBuildAccelerator/Common/Public/UbaLogger.h b/Engine/Source/Programs/UnrealBuildAccelerator/Common/Public/UbaLogger.h
--- a/Engine/Source/Programs/UnrealBuildAccelerator/Common/Public/UbaLogger.h
+++ b/Engine/Source/Programs/UnrealBuildAccelerator/Common/Public/UbaLogger.h
@@ -51,7 +51,7 @@ namespace uba
 	{
 		MutableLogger(LogWriter& writer, const tchar* prefix) : LoggerWithWriter(writer, prefix) {}
 		virtual void Log(LogEntryType type, const tchar* str, u32 strLen) override { if (!isMuted) LoggerWithWriter::Log(type, str, strLen); }
-		Atomic<bool> isMuted;
+		Atomic<u32> isMuted;
 	};
 
 	class FilteredLogWriter : public LogWriter
diff --git a/Engine/Source/Programs/UnrealBuildAccelerator/Common/Public/UbaNetwork.h b/Engine/Source/Programs/UnrealBuildAccelerator/Common/Public/UbaNetwork.h
--- a/Engine/Source/Programs/UnrealBuildAccelerator/Common/Public/UbaNetwork.h
+++ b/Engine/Source/Programs/UnrealBuildAccelerator/Common/Public/UbaNetwork.h
@@ -2,16 +2,11 @@
 
 #pragma once
 
-#include "UbaBase.h"
+#include "UbaDefaultConstants.h"
 
 namespace uba
 {
-	// Default constants
-	static constexpr u16 DefaultPort = 1345;
-	static constexpr u16 DefaultStorageProxyPort = DefaultPort + 1;
-	static constexpr u16 DefaultCachePort = DefaultPort + 2;
 	static constexpr u32 SendMaxSize = 256*1024;
-	static constexpr u32 SendDefaultSize = 256*1024;
 
 	static constexpr u8 SystemServiceId = 0;
 	static constexpr u8 StorageServiceId = 1;
diff --git a/Engine/Source/Programs/UnrealBuildAccelerator/Common/Public/UbaNetworkServer.h b/Engine/Source/Programs/UnrealBuildAccelerator/Common/Public/UbaNetworkServer.h
--- a/Engine/Source/Programs/UnrealBuildAccelerator/Common/Public/UbaNetworkServer.h
+++ b/Engine/Source/Programs/UnrealBuildAccelerator/Common/Public/UbaNetworkServer.h
@@ -73,6 +73,7 @@ namespace uba
 		//void UnregisterOnConnection(u8 id);
 
 		virtual void AddWork(const Function<void()>& work, u32 count, const tchar* desc) override final;
+		virtual void DoWork(u32 count = 1) override final;
 		virtual u32 GetWorkerCount() override final;
 
 		struct ClientStats
@@ -82,6 +83,7 @@ namespace uba
 			u32 connectionCount = 0;
 		};
 
+		MutableLogger& GetLogger();
 		u64 GetTotalSentBytes();
 		u64 GetTotalRecvBytes();
 		u32 GetConnectionCount();
@@ -108,7 +110,7 @@ namespace uba
 
 		void RemoveDisconnectedConnections();
 
-		LoggerWithWriter m_logger;
+		MutableLogger m_logger;
 
 		CryptoKey m_listenCrypto = InvalidCryptoKey;
 		Guid m_uid;
diff --git a/Engine/Source/Programs/UnrealBuildAccelerator/Common/Public/UbaSession.h b/Engine/Source/Programs/UnrealBuildAccelerator/Common/Public/UbaSession.h
--- a/Engine/Source/Programs/UnrealBuildAccelerator/Common/Public/UbaSession.h
+++ b/Engine/Source/Programs/UnrealBuildAccelerator/Common/Public/UbaSession.h
@@ -66,7 +66,7 @@ namespace uba
 		const tchar* GetId(); // Id for session. Will be "yymmdd_hhmmss" unless SessionCreateInfo.useUniqueId is set to false
 		u32 GetActiveProcessCount(); // Current active processes running inside session
 		Storage& GetStorage(); // Storage (only used when remote machines are connected)
-		Logger& GetLogger(); // Logger used for logging 
+		MutableLogger& GetLogger(); // Logger used for logging 
 		LogWriter& GetLogWriter(); // LogWriter used by logger
 		Trace& GetTrace(); // Trace written to be session
 		const ApplicationRules* GetRules(const ProcessStartInfo& si); // Get application rules used for process
diff --git a/Engine/Source/Programs/UnrealBuildAccelerator/Common/Public/UbaWorkManager.h b/Engine/Source/Programs/UnrealBuildAccelerator/Common/Public/UbaWorkManager.h
--- a/Engine/Source/Programs/UnrealBuildAccelerator/Common/Public/UbaWorkManager.h
+++ b/Engine/Source/Programs/UnrealBuildAccelerator/Common/Public/UbaWorkManager.h
@@ -18,6 +18,7 @@ namespace uba
 	public:
 		virtual void AddWork(const Function<void()>& work, u32 count, const tchar* desc) = 0;
 		virtual u32 GetWorkerCount() = 0;
+		virtual void DoWork(u32 count = 1) = 0;
 
 		u32 TrackWorkStart(const tchar* desc);
 		void TrackWorkEnd(u32 id);
@@ -25,6 +26,9 @@ namespace uba
 		void SetWorkTracker(WorkTracker* workTracker) { m_workTracker = workTracker; }
 		WorkTracker* GetWorkTracker() { return m_workTracker; }
 
+		template<typename TContainer, typename TFunc>
+		void ParallelFor(u32 workCount, TContainer& container, const TFunc& func, const tchar* description = TC(""));
+
 	protected:
 		Atomic<WorkTracker*> m_workTracker;
 	};
@@ -37,7 +41,7 @@ namespace uba
 		virtual ~WorkManagerImpl();
 		virtual void AddWork(const Function<void()>& work, u32 count, const tchar* desc) override;
 		virtual u32 GetWorkerCount() override;
-		void DoWork(u32 count = 1);
+		virtual void DoWork(u32 count = 1) override;
 		void FlushWork();
 
 	private:
@@ -64,4 +68,36 @@ namespace uba
 		WorkManager& workManager;
 		u32 workIndex;
 	};
+
+
+
+	template<typename TContainer, typename TFunc>
+	void WorkManager::ParallelFor(u32 workCount, TContainer& container, const TFunc& func, const tchar* description)
+	{
+		auto it = container.begin();
+		auto end = container.end();
+		ReaderWriterLock lock;
+		Atomic<u32> workLeft = workCount;
+
+		AddWork([&]()
+			{
+				while (true)
+				{
+					SCOPED_WRITE_LOCK(lock, l);
+					if (it == end)
+					{
+						l.Leave();
+						--workLeft;
+						return;
+					}
+					auto it2 = it++;
+					l.Leave();
+
+					func(it2);
+				}
+			}, workCount, description);
+
+		while (workLeft)
+			DoWork();
+	}
 }
diff --git a/Engine/Source/Programs/UnrealBuildAccelerator/Host/Private/UbaExports.cpp b/Engine/Source/Programs/UnrealBuildAccelerator/Host/Private/UbaExports.cpp
--- a/Engine/Source/Programs/UnrealBuildAccelerator/Host/Private/UbaExports.cpp
+++ b/Engine/Source/Programs/UnrealBuildAccelerator/Host/Private/UbaExports.cpp
@@ -369,6 +369,8 @@ extern "C"
 	}
 	void SessionServer_CancelAll(uba::SessionServer* server)
 	{
+		++server->GetServer().GetLogger().isMuted; // Mute forever
+		++server->GetLogger().isMuted; // Mute forever
 		server->CancelAllProcessesAndWait();
 	}
 	void SessionServer_SetCustomCasKeyFromTrackedInputs(uba::SessionServer* server, uba::ProcessHandle* handle, const uba::tchar* fileName, const uba::tchar* workingDir)
@@ -531,19 +533,19 @@ extern "C"
 		return true;
 	}
 
-	bool CacheClient_WriteToCache(uba::CacheClient* cacheClient, uba::RootPaths* rootPaths, const uba::ProcessHandle* process)
+	bool CacheClient_WriteToCache(uba::CacheClient* cacheClient, uba::RootPaths* rootPaths, uba::u32 bucket, const uba::ProcessHandle* process)
 	{
 		using namespace uba;
 		auto& si = process->GetStartInfo();
 		if (!si.trackInputs)
 			return false;
-		return cacheClient->WriteToCache(*rootPaths, *process);
+		return cacheClient->WriteToCache(*rootPaths, bucket, *process);
 	}
 
-	bool CacheClient_FetchFromCache(uba::CacheClient* cacheClient, uba::RootPaths* rootPaths, const uba::ProcessStartInfo& info)
+	bool CacheClient_FetchFromCache(uba::CacheClient* cacheClient, uba::RootPaths* rootPaths, uba::u32 bucket, const uba::ProcessStartInfo& info)
 	{
 		using namespace uba;
-		return cacheClient->FetchFromCache(*rootPaths, info);
+		return cacheClient->FetchFromCache(*rootPaths, bucket, info);
 	}
 
 	void CacheClient_Destroy(uba::CacheClient* cacheClient)
@@ -566,25 +568,4 @@ extern "C"
 		uba::FindImports(binary, [&](const uba::tchar* importName, bool isKnown) { func(importName, userData); });
 #endif
 	}
-
-
-	// Deprecated.. delete as soon as binaries are updated and code is using new api
-	UBA_DEP_API uba::NetworkServer* CreateServer(uba::LogWriter& writer, uba::u32 workerCount, uba::u32 sendSize, uba::u32 receiveTimeoutSeconds, bool useQuic) { return NetworkServer_Create(writer, workerCount, sendSize, receiveTimeoutSeconds, useQuic); }
-	UBA_DEP_API void DestroyServer(uba::NetworkServer* server) { NetworkServer_Destroy(server); }
-	UBA_DEP_API bool Server_StartListen(uba::NetworkServer* server, int port, const uba::tchar* ip, const uba::tchar* crypto) { return NetworkServer_StartListen(server, port, ip, crypto); }
-	UBA_DEP_API void Server_Stop(uba::NetworkServer* server) { NetworkServer_Stop(server); }
-	UBA_DEP_API bool Server_AddClient(uba::NetworkServer* server, const uba::tchar* ip, int port, const uba::tchar* crypto)	{ return NetworkServer_AddClient(server, ip, port, crypto); }
-	UBA_DEP_API uba::ProcessStartInfo* CreateProcessStartInfo(const uba::tchar* application, const uba::tchar* arguments, const uba::tchar* workingDir, const uba::tchar* description, uba::u32 priorityClass, uba::u64 outputStatsThresholdMs, bool trackInputs, const uba::tchar* logFile, ProcessHandle_ExitCallback* exit) { return ProcessStartInfo_Create(application, arguments, workingDir, description, priorityClass, outputStatsThresholdMs, trackInputs, logFile, exit); }
-	UBA_DEP_API void DestroyProcessStartInfo(uba::ProcessStartInfo* info) { ProcessStartInfo_Destroy(info); }
-	UBA_DEP_API uba::Storage* CreateStorageServer(uba::NetworkServer& server, const uba::tchar* rootDir, uba::u64 casCapacityBytes, bool storeCompressed, uba::LogWriter& writer, const uba::tchar* zone) { return StorageServer_Create(server, rootDir, casCapacityBytes, storeCompressed, writer, zone); }
-	UBA_DEP_API void DestroyStorageServer(uba::Storage* storageServer) { StorageServer_Destroy((uba::StorageServer*)storageServer); }
-	UBA_DEP_API uba::SessionServerCreateInfo* CreateSessionServerCreateInfo(uba::StorageServer& storage, uba::NetworkServer& client, uba::LogWriter& writer, const uba::tchar* rootDir, const uba::tchar* traceOutputFile,
-		bool disableCustomAllocator, bool launchVisualizer, bool resetCas, bool writeToDisk, bool detailedTrace, bool allowWaitOnMem = false, bool allowKillOnMem = false, bool storeObjFilesCompressed = false)
-	{
-		return SessionServerCreateInfo_Create(storage, client, writer, rootDir, traceOutputFile, disableCustomAllocator, launchVisualizer, resetCas, writeToDisk, detailedTrace, allowWaitOnMem, allowKillOnMem, storeObjFilesCompressed);
-	}
-	UBA_DEP_API void DestroySessionServerCreateInfo(uba::SessionServerCreateInfo* info) { SessionServerCreateInfo_Destroy(info); }
-	UBA_DEP_API uba::SessionServer* CreateSessionServer(const uba::SessionServerCreateInfo& info) { return SessionServer_Create(info); }
-	UBA_DEP_API void DestroySessionServer(uba::SessionServer* server) { SessionServer_Destroy(server); }
-	UBA_DEP_API void Storage_DeleteFile(uba::Storage* storage, const uba::tchar* file) { storage->DeleteCasForFile(file); }
 }
diff --git a/Engine/Source/Programs/UnrealBuildAccelerator/Host/Public/UbaExports.h b/Engine/Source/Programs/UnrealBuildAccelerator/Host/Public/UbaExports.h
--- a/Engine/Source/Programs/UnrealBuildAccelerator/Host/Public/UbaExports.h
+++ b/Engine/Source/Programs/UnrealBuildAccelerator/Host/Public/UbaExports.h
@@ -2,20 +2,19 @@
 
 #pragma once
 
-#include "UbaNetwork.h"
-#include "UbaSessionServerCreateInfo.h"
-#include "UbaProcessStartInfo.h"
+#include "UbaLogWriter.h"
+#include "UbaDefaultConstants.h"
 
 namespace uba
 {
 	class CacheClient;
+	class NetworkServer;
 	class Process;
 	class ProcessHandle;
 	class RootPaths;
 	class Session;
 	class SessionServer;
 	class Scheduler;
-	class Storage;
 	class StorageServer;
 	struct ProcessStartInfo;
 	struct SessionServerCreateInfo;
@@ -49,9 +48,9 @@ extern "C"
 	// NetworkServer
 	UBA_API uba::NetworkServer* NetworkServer_Create(uba::LogWriter& writer = uba::g_consoleLogWriter, uba::u32 workerCount = 64, uba::u32 sendSize = uba::SendDefaultSize, uba::u32 receiveTimeoutSeconds = 60, bool useQuic = false);
 	UBA_API void NetworkServer_Destroy(uba::NetworkServer* server);
-	UBA_API bool NetworkServer_StartListen(uba::NetworkServer* server, int port, const uba::tchar* ip, const uba::tchar* crypto = nullptr);
+	UBA_API bool NetworkServer_StartListen(uba::NetworkServer* server, int port = uba::DefaultPort, const uba::tchar* ip = nullptr, const uba::tchar* crypto = nullptr);
 	UBA_API void NetworkServer_Stop(uba::NetworkServer* server);
-	UBA_API bool NetworkServer_AddClient(uba::NetworkServer* server, const uba::tchar* ip, int port, const uba::tchar* crypto = nullptr);
+	UBA_API bool NetworkServer_AddClient(uba::NetworkServer* server, const uba::tchar* ip, int port = uba::DefaultPort, const uba::tchar* crypto = nullptr);
 
 	// StorageServer
 	UBA_API uba::StorageServer* StorageServer_Create(uba::NetworkServer& server, const uba::tchar* rootDir, uba::u64 casCapacityBytes, bool storeCompressed, uba::LogWriter& writer = uba::g_consoleLogWriter, const uba::tchar* zone = TC(""));
@@ -123,8 +122,8 @@ extern "C"
 	// Cache
 	UBA_API uba::CacheClient* CacheClient_Create(uba::SessionServer* session);
 	UBA_API bool CacheClient_Connect(uba::CacheClient* cacheClient, const uba::tchar* host, int port);
-	UBA_API bool CacheClient_WriteToCache(uba::CacheClient* cacheClient, uba::RootPaths* rootPaths, const uba::ProcessHandle* process);
-	UBA_API bool CacheClient_FetchFromCache(uba::CacheClient* cacheClient, uba::RootPaths* rootPaths, const uba::ProcessStartInfo& info);
+	UBA_API bool CacheClient_WriteToCache(uba::CacheClient* cacheClient, uba::RootPaths* rootPaths, uba::u32 bucket, const uba::ProcessHandle* process);
+	UBA_API bool CacheClient_FetchFromCache(uba::CacheClient* cacheClient, uba::RootPaths* rootPaths, uba::u32 bucket, const uba::ProcessStartInfo& info);
 	UBA_API void CacheClient_Destroy(uba::CacheClient* cacheClient);
 
 	// Misc
@@ -133,23 +132,4 @@ extern "C"
 
 	using ImportFunc = void(const uba::tchar* importName, void* userData);
 	UBA_API void Uba_FindImports(const uba::tchar* binary, ImportFunc* func, void* userData);
-
-
-
-
-	// Deprecated.. delete as soon as binaries are updated and code is using new api
-	#define UBA_DEP_API UBA_API
-	UBA_DEP_API uba::NetworkServer* CreateServer(uba::LogWriter& writer = uba::g_consoleLogWriter, uba::u32 workerCount = 64, uba::u32 sendSize = uba::SendDefaultSize, uba::u32 receiveTimeoutSeconds = 60, bool useQuic = false);
-	UBA_DEP_API void DestroyServer(uba::NetworkServer* server);
-	UBA_DEP_API bool Server_StartListen(uba::NetworkServer* server, int port, const uba::tchar* ip, const uba::tchar* crypto = nullptr);
-	UBA_DEP_API void Server_Stop(uba::NetworkServer* server);
-	UBA_DEP_API bool Server_AddClient(uba::NetworkServer* server, const uba::tchar* ip, int port, const uba::tchar* crypto = nullptr);
-	UBA_DEP_API uba::ProcessStartInfo* CreateProcessStartInfo(const uba::tchar* application, const uba::tchar* arguments, const uba::tchar* workingDir, const uba::tchar* description, uba::u32 priorityClass, uba::u64 outputStatsThresholdMs, bool trackInputs, const uba::tchar* logFile, ProcessHandle_ExitCallback* exit);
-	UBA_DEP_API void DestroyProcessStartInfo(uba::ProcessStartInfo* info);
-	UBA_DEP_API uba::Storage* CreateStorageServer(uba::NetworkServer& server, const uba::tchar* rootDir, uba::u64 casCapacityBytes, bool storeCompressed, uba::LogWriter& writer = uba::g_consoleLogWriter, const uba::tchar* zone = TC(""));
-	UBA_DEP_API void DestroyStorageServer(uba::Storage* storageServer);
-	UBA_DEP_API void Storage_DeleteFile(uba::Storage* storage, const uba::tchar* file);
-	UBA_DEP_API uba::SessionServer* CreateSessionServer(const uba::SessionServerCreateInfo& info);
-	UBA_DEP_API void DestroySessionServer(uba::SessionServer* server);
 }
-
