[Uba]
* Fixed option to manually handle storage overflow
* Enabled this option in cache server
* Fixed bug where cas files received from client was not added to written files linked list used to decide which files to evict first

diff --git a/Engine/Source/Programs/UnrealBuildAccelerator/CacheService/Private/UbaCacheService.cpp b/Engine/Source/Programs/UnrealBuildAccelerator/CacheService/Private/UbaCacheService.cpp
--- a/Engine/Source/Programs/UnrealBuildAccelerator/CacheService/Private/UbaCacheService.cpp
+++ b/Engine/Source/Programs/UnrealBuildAccelerator/CacheService/Private/UbaCacheService.cpp
@@ -42,11 +42,9 @@ namespace uba
 		logger.Info(TC("   UbaCacheService v%s"), Version);
 		logger.Info(TC("-------------------------------------------"));
 		logger.Info(TC(""));
-		logger.Info(TC("  UbaCacheService.exe [options...]"));
-		logger.Info(TC(""));
-		logger.Info(TC("  Options:"));
-		logger.Info(TC("   -dir=<rootdir>          The directory used to store data. Defaults to \"%s\""), DefaultRootDir);
-		logger.Info(TC("   -port=[<host>:]<port>   The ip/name and port (default: %u) to listen for clients on"), DefaultCachePort);
+		logger.Info(TC("  -dir=<rootdir>          The directory used to store data. Defaults to \"%s\""), DefaultRootDir);
+		logger.Info(TC("  -port=[<host>:]<port>   The ip/name and port (default: %u) to listen for clients on"), DefaultCachePort);
+		logger.Info(TC("  -capacity=<gigaby>      Capacity of local store. Defaults to %u gigabytes"), DefaultCapacityGb);
 		logger.Info(TC(""));
 		return -1;
 	}
@@ -152,6 +150,11 @@ namespace uba
 				if ((g_rootDir.count = GetFullPathNameW(value.Replace('/', PathSeparator).data, g_rootDir.capacity, g_rootDir.data, nullptr)) == 0)
 					return PrintHelp(StringBuffer<>().Appendf(TC("-dir has invalid path %s"), g_rootDir.data).data);
 			}
+			else if (name.Equals(TC("-capacity")))
+			{
+				if (!value.Parse(storageCapacityGb))
+					return PrintHelp(TC("Invalid value for -capacity"));
+			}
 			else if (name.Equals(TC("-?")))
 			{
 				return PrintHelp(TC(""));
@@ -208,6 +211,7 @@ namespace uba
 		storageInfo.casCapacityBytes = storageCapacity;
 		storageInfo.storeCompressed = storeCompressed;
 		storageInfo.allowFallback = false;
+		storageInfo.manuallyHandleOverflow = true;
 		storageInfo.writeRecievedCasFilesToDisk = true;
 		StorageServer storageServer(storageInfo);
 
diff --git a/Engine/Source/Programs/UnrealBuildAccelerator/Common/Private/UbaCacheServer.cpp b/Engine/Source/Programs/UnrealBuildAccelerator/Common/Private/UbaCacheServer.cpp
--- a/Engine/Source/Programs/UnrealBuildAccelerator/Common/Private/UbaCacheServer.cpp
+++ b/Engine/Source/Programs/UnrealBuildAccelerator/Common/Private/UbaCacheServer.cpp
@@ -314,7 +314,7 @@ namespace uba
 			m_storage.TraverseAllCasFiles([&](const CasKey& casKey, u64 size) { totalCasSize += size; existingCas.try_emplace(casKey, CasFileInfo{size, 0ull}); });
 			m_logger.Detail(TC("  Found %llu cas files (%s)"), existingCas.size(), TimeToText(GetTime() - traverseStartTime).str);
 		}
-		u64 totalCasCount = existingCas.size();
+		u64 totalCasCount = existingCas.size() + deletedCasCount;
 
 		u64 now = GetSystemTimeAsFileTime();
 		u64 oldest = now;
diff --git a/Engine/Source/Programs/UnrealBuildAccelerator/Common/Private/UbaStorage.cpp b/Engine/Source/Programs/UnrealBuildAccelerator/Common/Private/UbaStorage.cpp
--- a/Engine/Source/Programs/UnrealBuildAccelerator/Common/Private/UbaStorage.cpp
+++ b/Engine/Source/Programs/UnrealBuildAccelerator/Common/Private/UbaStorage.cpp
@@ -99,7 +99,7 @@ namespace uba
 		UBA_ASSERT(!entry.mappingHandle.IsValid());
 #endif
 
-		if (!m_casCapacityBytes || m_overflowReported || m_casTotalBytes <= m_casCapacityBytes)
+		if (!m_casCapacityBytes || m_overflowReported || m_casTotalBytes <= m_casCapacityBytes || m_manuallyHandleOverflow)
 			return;
 
 		#if UBA_USE_SPARSEFILE
@@ -1087,6 +1087,7 @@ namespace uba
 	{
 		m_casCapacityBytes = info.casCapacityBytes;
 		m_storeCompressed = info.storeCompressed;
+		m_manuallyHandleOverflow = info.manuallyHandleOverflow;
 		m_rootDir.count = GetFullPathNameW(info.rootDir, m_rootDir.capacity, m_rootDir.data, NULL);
 		m_rootDir.Replace('/', PathSeparator).EnsureEndsWithSlash();
 
@@ -1294,7 +1295,9 @@ namespace uba
 			CheckAllCasFiles();
 			resave = true;
 		}
-		HandleOverflow(nullptr);
+
+		if (!m_manuallyHandleOverflow)
+			HandleOverflow(nullptr);
 
 		if (resave)
 		{
diff --git a/Engine/Source/Programs/UnrealBuildAccelerator/Common/Private/UbaStorageServer.cpp b/Engine/Source/Programs/UnrealBuildAccelerator/Common/Private/UbaStorageServer.cpp
--- a/Engine/Source/Programs/UnrealBuildAccelerator/Common/Private/UbaStorageServer.cpp
+++ b/Engine/Source/Programs/UnrealBuildAccelerator/Common/Private/UbaStorageServer.cpp
@@ -834,7 +834,7 @@ namespace uba
 						casEntry.beingWritten = false;
 					}
 
-					bool isPersistentStore = false;
+					bool isPersistentStore = m_writeRecievedCasFilesToDisk;
 					if (isPersistentStore)
 						CasEntryWritten(*activeStore.casEntry, totalWritten);
 
diff --git a/Engine/Source/Programs/UnrealBuildAccelerator/Common/Public/UbaStorage.h b/Engine/Source/Programs/UnrealBuildAccelerator/Common/Public/UbaStorage.h
--- a/Engine/Source/Programs/UnrealBuildAccelerator/Common/Public/UbaStorage.h
+++ b/Engine/Source/Programs/UnrealBuildAccelerator/Common/Public/UbaStorage.h
@@ -78,6 +78,7 @@ namespace uba
 		u64 casCapacityBytes = 20llu * 1024 * 1024 * 1024;
 		u32 maxParallelCopyOrLink = 1000;
 		bool storeCompressed = true;
+		bool manuallyHandleOverflow = false;
 		WorkManager* workManager = nullptr;
 	};
 
@@ -210,6 +211,7 @@ namespace uba
 		u32 m_casDroppedCount = 0;
 		bool m_overflowReported = false;
 		bool m_storeCompressed = false;
+		bool m_manuallyHandleOverflow = false;
 
 		u32 m_maxParallelCopyOrLink;
 		ReaderWriterLock m_activeCopyOrLinkLock;
