[UBA]
* Fixed bug in text normalization code.. was missing one character
* Fixed so cache processes are colored differently in visualizer
* Improved assert information when pressing ctrl-c during build and an assert triggers

diff --git a/Engine/Source/Programs/UnrealBuildAccelerator/Common/Private/UbaCacheClient.cpp b/Engine/Source/Programs/UnrealBuildAccelerator/Common/Private/UbaCacheClient.cpp
--- a/Engine/Source/Programs/UnrealBuildAccelerator/Common/Private/UbaCacheClient.cpp
+++ b/Engine/Source/Programs/UnrealBuildAccelerator/Common/Private/UbaCacheClient.cpp
@@ -9,6 +9,8 @@
 #include "UbaStorage.h"
 #include "UbaStorageUtils.h"
 
+#define UBA_LOG_CACHE_INFO 0
+
 namespace uba
 {
 	struct CacheClient::Bucket
@@ -196,6 +198,11 @@ namespace uba
 		if (!SendCacheEntry(bucket, rootPaths, cmdKey, inputsStringToCasKey, outputsStringToCasKey))
 			return false;
 
+
+		#if UBA_LOG_CACHE_INFO
+		m_logger.Info(TC("WRITECACHE: %s -> %u %s"), si.description, bucketId, CasKeyString(cmdKey).str);
+		#endif
+
 		return true;
 	}
 
@@ -255,6 +262,11 @@ namespace uba
 
 		// Traverse entries and test inputs against local machine
 		u32 entryCount = reader.ReadU16();
+
+		#if UBA_LOG_CACHE_INFO
+		auto mg = MakeGuard([&]() { m_logger.Info(TC("FETCHCACHE %s: %s -> %u %s (%u)"), success ? TC("SUCC") : TC("FAIL"), info.description, bucketId, CasKeyString(cmdKey).str, entryCount); });
+		#endif
+
 		for (u32 i=0; i!=entryCount; ++i)
 		{
 			u64 outputSize = 0;
diff --git a/Engine/Source/Programs/UnrealBuildAccelerator/Common/Private/UbaNetworkServer.cpp b/Engine/Source/Programs/UnrealBuildAccelerator/Common/Private/UbaNetworkServer.cpp
--- a/Engine/Source/Programs/UnrealBuildAccelerator/Common/Private/UbaNetworkServer.cpp
+++ b/Engine/Source/Programs/UnrealBuildAccelerator/Common/Private/UbaNetworkServer.cpp
@@ -865,8 +865,11 @@ namespace uba
 			if (m_createdWorkerCount != m_maxWorkerCount)
 				return false;
 			lock2.Leave();
+
 			auto worker = t_worker;
-			UBA_ASSERT(worker);
+			if (!worker)
+				return false;
+
 			auto oldContext = worker->m_context;
 			WorkerContext context(*this);
 			worker->m_context = &context;
diff --git a/Engine/Source/Programs/UnrealBuildAccelerator/Common/Private/UbaSessionServer.cpp b/Engine/Source/Programs/UnrealBuildAccelerator/Common/Private/UbaSessionServer.cpp
--- a/Engine/Source/Programs/UnrealBuildAccelerator/Common/Private/UbaSessionServer.cpp
+++ b/Engine/Source/Programs/UnrealBuildAccelerator/Common/Private/UbaSessionServer.cpp
@@ -517,7 +517,7 @@ namespace uba
 			m_trace.SessionDisconnect(sessionId);
 
 			sessionName.Append(s.name);
-			UBA_ASSERT(s.usedSlotCount == returnCount);
+			UBA_ASSERTF(s.usedSlotCount == returnCount, TC("Used slot count different than return count (%u vs %u)"), s.usedSlotCount, returnCount);
 			s.usedSlotCount -= returnCount;
 
 			if (s.enabled)
diff --git a/Engine/Source/Programs/UnrealBuildAccelerator/Common/Public/UbaRootPaths.h b/Engine/Source/Programs/UnrealBuildAccelerator/Common/Public/UbaRootPaths.h
--- a/Engine/Source/Programs/UnrealBuildAccelerator/Common/Public/UbaRootPaths.h
+++ b/Engine/Source/Programs/UnrealBuildAccelerator/Common/Public/UbaRootPaths.h
@@ -91,7 +91,7 @@ namespace uba
 			func(&temp, 1, destPos);
 			destPos += 1;
 
-			searchPos = pathStart + root->path.size();
+			searchPos = pathStart + root->path.size() - 1;
 		}
 	}
 }
diff --git a/Engine/Source/Programs/UnrealBuildAccelerator/Visualizer/Private/UbaVisualizer.cpp b/Engine/Source/Programs/UnrealBuildAccelerator/Visualizer/Private/UbaVisualizer.cpp
--- a/Engine/Source/Programs/UnrealBuildAccelerator/Visualizer/Private/UbaVisualizer.cpp
+++ b/Engine/Source/Programs/UnrealBuildAccelerator/Visualizer/Private/UbaVisualizer.cpp
@@ -345,6 +345,8 @@ namespace uba
 			m_processBrushes[1].success = CreateSolidBrush(RGB(10, 140, 10));
 			m_processBrushes[0].send = CreateSolidBrush(RGB(10, 115, 10));
 			m_processBrushes[1].send = CreateSolidBrush(RGB(10, 145, 10));
+			m_processBrushes[0].cacheFetch = CreateSolidBrush(RGB(24, 112, 110));
+			m_processBrushes[1].cacheFetch = CreateSolidBrush(RGB(31, 143, 138));
 
 			m_workBrush = CreateSolidBrush(RGB(70, 70, 100));
 
@@ -380,6 +382,9 @@ namespace uba
 			m_processBrushes[0].send = CreateSolidBrush(RGB(80, 210, 80));
 			m_processBrushes[1].send = CreateSolidBrush(RGB(90, 250, 90));
 
+			m_processBrushes[0].cacheFetch = CreateSolidBrush(RGB(150, 150, 200));
+			m_processBrushes[1].cacheFetch = CreateSolidBrush(RGB(170, 170, 200));
+
 			m_workBrush = CreateSolidBrush(RGB(150, 150, 200));
 
 			m_backgroundBrush = GetSysColorBrush(0);
@@ -1532,6 +1537,8 @@ namespace uba
 			brush = m_processBrushes[selected].returned;
 		else if (!done)
 			brush = m_processBrushes[selected].inProgress;
+		else if (process.cacheFetch)
+			brush = m_processBrushes[selected].cacheFetch;
 		else if (process.exitCode != 0)
 			brush = m_processBrushes[selected].error;
 
diff --git a/Engine/Source/Programs/UnrealBuildAccelerator/Visualizer/Private/UbaVisualizer.h b/Engine/Source/Programs/UnrealBuildAccelerator/Visualizer/Private/UbaVisualizer.h
--- a/Engine/Source/Programs/UnrealBuildAccelerator/Visualizer/Private/UbaVisualizer.h
+++ b/Engine/Source/Programs/UnrealBuildAccelerator/Visualizer/Private/UbaVisualizer.h
@@ -82,6 +82,7 @@ namespace uba
 			HBRUSH returned = 0;
 			HBRUSH recv = 0;
 			HBRUSH send = 0;
+			HBRUSH cacheFetch = 0;
 		};
 
 		ProcessBrushes m_processBrushes[2]; // Non-selected and selected
