[Uba]
* Proper fix for the normalized paths which also keeps unit tests working

diff --git a/Engine/Source/Programs/UnrealBuildAccelerator/Common/Private/UbaCacheClient.cpp b/Engine/Source/Programs/UnrealBuildAccelerator/Common/Private/UbaCacheClient.cpp
--- a/Engine/Source/Programs/UnrealBuildAccelerator/Common/Private/UbaCacheClient.cpp
+++ b/Engine/Source/Programs/UnrealBuildAccelerator/Common/Private/UbaCacheClient.cpp
@@ -382,7 +382,7 @@ namespace uba
 
 						if (u32 toWrite = rootPathLen - 1)
 							memcpy(localBlock.Allocate(toWrite, 1, TC("")), rootPath, toWrite);
-						lastWritten = rootOffset + 2;
+						lastWritten = rootOffset + 1;
 					}
 
 					u64 fileSize = rootOffsets - sizeof(u32);
diff --git a/Engine/Source/Programs/UnrealBuildAccelerator/Common/Public/UbaRootPaths.h b/Engine/Source/Programs/UnrealBuildAccelerator/Common/Public/UbaRootPaths.h
--- a/Engine/Source/Programs/UnrealBuildAccelerator/Common/Public/UbaRootPaths.h
+++ b/Engine/Source/Programs/UnrealBuildAccelerator/Common/Public/UbaRootPaths.h
@@ -91,7 +91,7 @@ namespace uba
 			func(&temp, 1, destPos);
 			destPos += 1;
 
-			searchPos = pathStart + root->path.size() - 1;
+			searchPos = pathStart + root->path.size();
 		}
 	}
 }
