[UBA]
* Fixed serious bug related to child processes on helper machines getting process ids that collides with ids of root processes coming from server
* Fixed bug in trace reader causing reusable processes to get no stats
* Added good-to-have asserts

diff --git a/Engine/Source/Programs/UnrealBuildAccelerator/Common/Private/UbaNetworkClient.cpp b/Engine/Source/Programs/UnrealBuildAccelerator/Common/Private/UbaNetworkClient.cpp
--- a/Engine/Source/Programs/UnrealBuildAccelerator/Common/Private/UbaNetworkClient.cpp
+++ b/Engine/Source/Programs/UnrealBuildAccelerator/Common/Private/UbaNetworkClient.cpp
@@ -318,6 +318,7 @@ namespace uba
 	{
 		StackBinaryWriter<64> writer;
 		NetworkMessage msg(*this, SystemServiceId, SystemMessageType_KeepAlive, writer);
+		writer.WriteByte(0); // Need to have a body
 		return msg.Send();
 	}
 
diff --git a/Engine/Source/Programs/UnrealBuildAccelerator/Common/Private/UbaProcess.cpp b/Engine/Source/Programs/UnrealBuildAccelerator/Common/Private/UbaProcess.cpp
--- a/Engine/Source/Programs/UnrealBuildAccelerator/Common/Private/UbaProcess.cpp
+++ b/Engine/Source/Programs/UnrealBuildAccelerator/Common/Private/UbaProcess.cpp
@@ -108,6 +108,7 @@ namespace uba
 
 	ProcessImpl::~ProcessImpl()
 	{
+		UBA_ASSERT(m_refCount == 0);
 		{
 			#if !PLATFORM_WINDOWS
 			SCOPED_WRITE_LOCK(m_comMemoryLock, lock);
@@ -538,6 +539,8 @@ namespace uba
 			h.m_process = nullptr;
 		}
 
+		UBA_ASSERT(m_refCount);
+
 		// Must be done last to make sure shutdown is not racing
 		m_session.ProcessExited(*this, m_processStats.wallTime);
 	}
diff --git a/Engine/Source/Programs/UnrealBuildAccelerator/Common/Private/UbaSession.cpp b/Engine/Source/Programs/UnrealBuildAccelerator/Common/Private/UbaSession.cpp
--- a/Engine/Source/Programs/UnrealBuildAccelerator/Common/Private/UbaSession.cpp
+++ b/Engine/Source/Programs/UnrealBuildAccelerator/Common/Private/UbaSession.cpp
@@ -1395,7 +1395,8 @@ namespace uba
 			m_trace.ProcessAdded(sessionId, processId, process.GetStartInfo().description);
 
 		SCOPED_WRITE_LOCK(m_processesLock, lock);
-		m_processes.try_emplace(processId, ProcessHandle(&process));
+		bool success = m_processes.try_emplace(processId, ProcessHandle(&process)).second;
+		UBA_ASSERT(success);(void)success;
 	}
 
 	void Session::ProcessExited(ProcessImpl& process, u64 executionTime)
@@ -1426,7 +1427,8 @@ namespace uba
 		auto& stats = m_applicationStats[applicationName.data];
 		stats.count++;
 		stats.time += executionTime;
-		m_processes.erase(id);
+		auto count = m_processes.erase(id);
+		UBA_ASSERT(count == 1);(void)count;
 	}
 
 	void Session::FlushDeadProcesses()
diff --git a/Engine/Source/Programs/UnrealBuildAccelerator/Common/Private/UbaSessionClient.cpp b/Engine/Source/Programs/UnrealBuildAccelerator/Common/Private/UbaSessionClient.cpp
--- a/Engine/Source/Programs/UnrealBuildAccelerator/Common/Private/UbaSessionClient.cpp
+++ b/Engine/Source/Programs/UnrealBuildAccelerator/Common/Private/UbaSessionClient.cpp
@@ -31,6 +31,8 @@ namespace uba
 		m_memKillLoadPercent = info.memKillLoadPercent;
 		m_processFinished = info.processFinished;
 
+		m_processIdCounter = ~0u / 2; // We set this value to a very high value.. because it will be used by child processes and we don't want id from server and child process id to collide
+
 		if (m_name.IsEmpty())
 		{
 			tchar buf[256];
diff --git a/Engine/Source/Programs/UnrealBuildAccelerator/Visualizer/Private/UbaTraceReader.cpp b/Engine/Source/Programs/UnrealBuildAccelerator/Visualizer/Private/UbaTraceReader.cpp
--- a/Engine/Source/Programs/UnrealBuildAccelerator/Visualizer/Private/UbaTraceReader.cpp
+++ b/Engine/Source/Programs/UnrealBuildAccelerator/Visualizer/Private/UbaTraceReader.cpp
@@ -558,6 +558,8 @@ namespace uba
 			auto& processes = session.processors[active.processorIndex].processes;
 			TraceView::Process& process = processes[active.processIndex];
 
+			const u8* dataStart = reader.GetPositionData();
+
 			ProcessStats processStats;
 			SessionStats sessionStats;
 			StorageStats storageStats;
@@ -567,6 +569,10 @@ namespace uba
 			storageStats.Read(reader);
 			systemStats.Read(reader);
 
+			const u8* dataEnd = reader.GetPositionData();
+			process.stats.resize(dataEnd - dataStart);
+			memcpy(process.stats.data(), dataStart, dataEnd - dataStart);
+
 			process.exitCode = 0u;
 			process.stop = time;
 			process.bitmapDirty = true;
diff --git a/Engine/Source/Programs/UnrealBuildAccelerator/Visualizer/Private/UbaVisualizer.cpp b/Engine/Source/Programs/UnrealBuildAccelerator/Visualizer/Private/UbaVisualizer.cpp
--- a/Engine/Source/Programs/UnrealBuildAccelerator/Visualizer/Private/UbaVisualizer.cpp
+++ b/Engine/Source/Programs/UnrealBuildAccelerator/Visualizer/Private/UbaVisualizer.cpp
@@ -1272,7 +1272,7 @@ namespace uba
 			if (hasExited && process.exitCode != 0)
 				logger.Info(L"  ExitCode:  %u", process.exitCode);
 
-			if (process.stop != ~u64(0))
+			if (process.stop != ~u64(0) && !process.stats.empty())
 			{
 				BinaryReader reader(process.stats.data(), 0, process.stats.size());
 				ProcessStats processStats;
