[UBA]
* Fixed so absolute paths (paths with ':' as second character) are prefixed with \\?\ on windows. This to hopefully solve problems with systems that does not have long paths enabled but still end up with long paths

diff --git a/Engine/Source/Programs/UnrealBuildAccelerator/Common/Private/UbaFile.cpp b/Engine/Source/Programs/UnrealBuildAccelerator/Common/Private/UbaFile.cpp
--- a/Engine/Source/Programs/UnrealBuildAccelerator/Common/Private/UbaFile.cpp
+++ b/Engine/Source/Programs/UnrealBuildAccelerator/Common/Private/UbaFile.cpp
@@ -72,6 +72,16 @@ namespace uba
 	{
 		return (HANDLE)(fh == InvalidFileHandle ? InvalidFileHandle : (fh & FileHandleFlagMask));
 	}
+
+	#define MAKE_LONG_FILENAME(fileName) \
+		UBA_ASSERT(TStrlen(fileName) < 508); \
+		StringBuffer<512> STRING_JOIN(longName, __LINE__); \
+		if (fileName && fileName[0] && fileName[1] == ':') \
+		{ \
+			STRING_JOIN(longName, __LINE__).Append(TC("\\\\?\\")).Append(fileName); \
+			fileName = STRING_JOIN(longName, __LINE__).data; \
+		}
+
 #else
 	int asFileDescriptor(FileHandle fh)
 	{
@@ -124,6 +134,7 @@ namespace uba
 		u32 dwFlagsAndAttributes = DefaultAttributes();
 		#if PLATFORM_WINDOWS
 		dwFlagsAndAttributes |= (overlapped ? FILE_FLAG_OVERLAPPED : FILE_FLAG_SEQUENTIAL_SCAN);
+		MAKE_LONG_FILENAME(fileName);
 		#endif
 
 		outHandle = uba::CreateFileW(fileName, GENERIC_READ, FILE_SHARE_READ|FILE_SHARE_WRITE, OPEN_EXISTING, dwFlagsAndAttributes);
@@ -180,6 +191,7 @@ namespace uba
 	bool GetFileInformation(FileInformation& out, Logger& logger, const tchar* fileName)
 	{
 #if PLATFORM_WINDOWS
+		MAKE_LONG_FILENAME(fileName);
 		FileHandle h = uba::CreateFileW(fileName, 0, 0, OPEN_EXISTING, FILE_FLAG_BACKUP_SEMANTICS);
 		if (h == InvalidFileHandle)
 			return false;// logger.Error(TC("GetFileInformation: CreateFile failed for file %s (%s)"), fileName, LastErrorToText().data);
@@ -208,6 +220,7 @@ namespace uba
 	bool FileExists(Logger& logger, const tchar* fileName, u64* outSize, u32* outAttributes)
 	{
 #if PLATFORM_WINDOWS
+		MAKE_LONG_FILENAME(fileName);
 		WIN32_FILE_ATTRIBUTE_DATA data;
 		if (!::GetFileAttributesExW(fileName, GetFileExInfoStandard, &data))
 		{
@@ -396,6 +409,7 @@ namespace uba
 	{
 		ExtendedTimerScope ts(SystemStats::GetCurrent().createFile);
 	#if PLATFORM_WINDOWS
+		MAKE_LONG_FILENAME(fileName);
 		return (FileHandle)(u64)::CreateFileW(fileName, desiredAccess, shareMode, NULL, createDisp, flagsAndAttributes, NULL);
 	#else
 		int flags = O_CLOEXEC;
@@ -469,6 +483,7 @@ namespace uba
 	bool CreateDirectoryW(const tchar* pathName)
 	{
 #if PLATFORM_WINDOWS
+		MAKE_LONG_FILENAME(pathName);
 		return ::CreateDirectoryW(pathName, NULL);
 #else
 		if (mkdir(pathName, 0777) == 0)
@@ -494,6 +509,7 @@ namespace uba
 	bool RemoveDirectoryW(const tchar* pathName)
 	{
 #if PLATFORM_WINDOWS
+		MAKE_LONG_FILENAME(pathName);
 		if (::RemoveDirectoryW(pathName))
 			return true;
 		return false;
@@ -514,6 +530,7 @@ namespace uba
 	bool DeleteFileW(const tchar* fileName)
 	{
 #if PLATFORM_WINDOWS
+		MAKE_LONG_FILENAME(fileName);
 		if (::DeleteFileW(fileName))
 			return true;
 		return false;
@@ -534,6 +551,8 @@ namespace uba
 	bool CopyFileW(const tchar* existingFileName, const tchar* newFileName, bool bFailIfExists)
 	{
 #if PLATFORM_WINDOWS
+		MAKE_LONG_FILENAME(existingFileName);
+		MAKE_LONG_FILENAME(newFileName);
 		return ::CopyFileW(existingFileName, newFileName, bFailIfExists);
 #else
 
@@ -585,6 +604,8 @@ namespace uba
 	bool MoveFileExW(const tchar* existingFileName, const tchar* newFileName, u32 dwFlags)
 	{
 #if PLATFORM_WINDOWS
+		MAKE_LONG_FILENAME(existingFileName);
+		MAKE_LONG_FILENAME(newFileName);
 		return ::MoveFileExW(existingFileName, newFileName, dwFlags);
 #else
 		int res = rename(existingFileName, newFileName);
@@ -623,6 +644,7 @@ namespace uba
 	u32 GetFileAttributesW(const tchar* fileName)
 	{
 #if PLATFORM_WINDOWS
+		MAKE_LONG_FILENAME(fileName);
 		return ::GetFileAttributesW(fileName);
 #else
 		struct stat attr;
@@ -670,6 +692,8 @@ namespace uba
 	bool CreateHardLinkW(const tchar* newFileName, const tchar* existingFileName)
 	{
 #if PLATFORM_WINDOWS
+		MAKE_LONG_FILENAME(newFileName);
+		MAKE_LONG_FILENAME(existingFileName);
 		return ::CreateHardLinkW(newFileName, existingFileName, NULL);
 #else
 #if 1//PLATFORM_MAC
