[Uba]
* Fixed so code adding \\?\ also fix path separator type from / to \ if forward slashes exists

diff --git a/Engine/Source/Programs/UnrealBuildAccelerator/Common/Private/UbaFile.cpp b/Engine/Source/Programs/UnrealBuildAccelerator/Common/Private/UbaFile.cpp
--- a/Engine/Source/Programs/UnrealBuildAccelerator/Common/Private/UbaFile.cpp
+++ b/Engine/Source/Programs/UnrealBuildAccelerator/Common/Private/UbaFile.cpp
@@ -78,7 +78,7 @@ namespace uba
 		StringBuffer<512> STRING_JOIN(longName, __LINE__); \
 		if (fileName && fileName[0] && fileName[1] == ':') \
 		{ \
-			STRING_JOIN(longName, __LINE__).Append(TC("\\\\?\\")).Append(fileName); \
+			STRING_JOIN(longName, __LINE__).Append(TC("\\\\?\\")).Append(fileName).Replace('/', '\\'); \
 			fileName = STRING_JOIN(longName, __LINE__).data; \
 		}
 
