[UBA]
* Improved MAKE_LONG_FILENAME to also take care of \.\ and \..\ properly

diff --git a/Engine/Source/Programs/UnrealBuildAccelerator/Common/Private/UbaFile.cpp b/Engine/Source/Programs/UnrealBuildAccelerator/Common/Private/UbaFile.cpp
--- a/Engine/Source/Programs/UnrealBuildAccelerator/Common/Private/UbaFile.cpp
+++ b/Engine/Source/Programs/UnrealBuildAccelerator/Common/Private/UbaFile.cpp
@@ -74,12 +74,14 @@ namespace uba
 	}
 
 	#define MAKE_LONG_FILENAME(fileName) \
-		UBA_ASSERT(TStrlen(fileName) < 508); \
-		StringBuffer<512> STRING_JOIN(longName, __LINE__); \
+		UBA_ASSERT(TStrlen(fileName) < MaxPath); \
+		StringBuffer<MaxPath> STRING_JOIN(longName, __LINE__); \
 		if (fileName && fileName[0] && fileName[1] == ':') \
 		{ \
-			STRING_JOIN(longName, __LINE__).Append(TC("\\\\?\\")).Append(fileName).Replace('/', '\\'); \
-			fileName = STRING_JOIN(longName, __LINE__).data; \
+			auto& lsb = STRING_JOIN(longName, __LINE__); \
+			lsb.Append(TC("\\\\?\\")); \
+			FixPath(fileName, nullptr, 0, lsb); \
+			fileName = lsb.data; \
 		}
 
 #else
