[UBA]
* Fixed so session reports updated file info to storage

diff --git a/Engine/Source/Programs/UnrealBuildAccelerator/Common/Private/UbaSession.cpp b/Engine/Source/Programs/UnrealBuildAccelerator/Common/Private/UbaSession.cpp
--- a/Engine/Source/Programs/UnrealBuildAccelerator/Common/Private/UbaSession.cpp
+++ b/Engine/Source/Programs/UnrealBuildAccelerator/Common/Private/UbaSession.cpp
@@ -2074,48 +2074,56 @@ namespace uba
 				if (m_storeObjFilesCompressed && EndsWith(file.name.c_str(), file.name.size(), TC(".obj")))
 				{
 					Storage::WriteResult res;
-					return m_storage.WriteCompressed(res, TC("MemoryMap"), InvalidFileHandle, mem, fileSize, file.name.c_str());
+					if (!m_storage.WriteCompressed(res, TC("MemoryMap"), InvalidFileHandle, mem, fileSize, file.name.c_str()))
+						return false;
 				}
+				else
+				{
 
-				// Seems like best combo (for windows at least) is to use writes with overlap and max 16 at the same time.
-				// On one machine we get twice as fast without overlap if no bottleneck. On another machine (ntfs compression on) we get twice as slow without overlap
-				// Both machines behaves well with overlap AND bottleneck. Both machine are 128 logical core thread rippers.
-				bool useFileMapForWrite = fileSize > 0; // ::CreateFileMappingW does not work for zero-length files
-				bool useOverlap = false;//fileSize > 8 * 1024 * 1024;
+					// Seems like best combo (for windows at least) is to use writes with overlap and max 16 at the same time.
+					// On one machine we get twice as fast without overlap if no bottleneck. On another machine (ntfs compression on) we get twice as slow without overlap
+					// Both machines behaves well with overlap AND bottleneck. Both machine are 128 logical core thread rippers.
+					bool useFileMapForWrite = fileSize > 0; // ::CreateFileMappingW does not work for zero-length files
+					bool useOverlap = false;//fileSize > 8 * 1024 * 1024;
 
 
-				u32 attributes = DefaultAttributes();
-				if (useOverlap)
-					attributes |= FILE_FLAG_OVERLAPPED;
+					u32 attributes = DefaultAttributes();
+					if (useOverlap)
+						attributes |= FILE_FLAG_OVERLAPPED;
 
-				FileAccessor destinationFile(m_logger, file.name.c_str());
+					FileAccessor destinationFile(m_logger, file.name.c_str());
 
-				if (useFileMapForWrite)
-				{
-					if (!destinationFile.CreateMemoryWrite(false, attributes, fileSize, m_tempPath.data))
-						return false;
-					memcpy(destinationFile.GetData(), mem, fileSize);
-				}
-				else
-				{
-					if (!destinationFile.CreateWrite(false, attributes, fileSize, m_tempPath.data))
-						return false;
+					if (useFileMapForWrite)
+					{
+						if (!destinationFile.CreateMemoryWrite(false, attributes, fileSize, m_tempPath.data))
+							return false;
+						memcpy(destinationFile.GetData(), mem, fileSize);
+					}
+					else
+					{
+						if (!destinationFile.CreateWrite(false, attributes, fileSize, m_tempPath.data))
+							return false;
 
-					#if PLATFORM_WINDOWS
-					//shouldBottleneck = fileSize > 64 * 1024 * 1024;
-					//if (shouldBottleneck)
-					//	bottleneck.Enter();
-					#endif
+						#if PLATFORM_WINDOWS
+						//shouldBottleneck = fileSize > 64 * 1024 * 1024;
+						//if (shouldBottleneck)
+						//	bottleneck.Enter();
+						#endif
 
-					if (!destinationFile.Write(mem, fileSize))
+						if (!destinationFile.Write(mem, fileSize))
+							return false;
+					}
+
+					if (u64 time = file.lastWriteTime)
+						if (!SetFileLastWriteTime(destinationFile.GetHandle(), time))
+							return m_logger.Error(TC("Failed to set file time on filehandle for %s"), file.name.c_str());
+
+					if (!destinationFile.Close(file.lastWriteTime ? nullptr : &file.lastWriteTime))
 						return false;
 				}
-				if (u64 time = file.lastWriteTime)
-					if (!SetFileLastWriteTime(destinationFile.GetHandle(), time))
-						return m_logger.Error(TC("Failed to set file time on filehandle for %s"), file.name.c_str());
 
-				if (!destinationFile.Close())
-					return false;
+				Storage::CachedFileInfo cfi;
+				m_storage.VerifyAndGetCachedFileInfo(cfi, file.key, file.lastWriteTime, fileSize);
 			}
 			else
 			{
