[UBA]
* Added special handling for cmd.exe copy processes and prevent spawning external process. Saves quite a lot of time when there is lots of copying

diff --git a/Engine/Source/Programs/UnrealBuildAccelerator/Common/Private/UbaProcess.cpp b/Engine/Source/Programs/UnrealBuildAccelerator/Common/Private/UbaProcess.cpp
--- a/Engine/Source/Programs/UnrealBuildAccelerator/Common/Private/UbaProcess.cpp
+++ b/Engine/Source/Programs/UnrealBuildAccelerator/Common/Private/UbaProcess.cpp
@@ -186,9 +186,17 @@ namespace uba
 		m_session.ProcessAdded(*this, 0);
 
 		if (async)
-			m_messageThread.Start([this, runningRemote, environment]() { ThreadRun(runningRemote, environment); return 0; });
+			m_messageThread.Start([this, runningRemote, environment]()
+				{
+					ThreadRun(runningRemote, environment);
+					ThreadExit();
+					return 0;
+				});
 		else
+		{
 			ThreadRun(runningRemote, environment);
+			ThreadExit();
+		}
 	}
 
 	bool ProcessImpl::IsActive()
@@ -407,11 +415,13 @@ namespace uba
 
 	void ProcessImpl::ThreadRun(bool runningRemote, void* environment)
 	{
-		{
 		SystemStatsScope systemStatsScope(m_systemStats);
 		StorageStatsScope storageStatsScope(m_storageStats);
 		SessionStatsScope sessionStatsScope(m_sessionStats);
 
+		if (HandleSpecialApplication())
+			return;
+
 		u8* comMemory = m_comMemory.memory;
 		u64 comMemorySize = CommunicationMemSize;
 		#if !PLATFORM_WINDOWS
@@ -484,8 +494,6 @@ namespace uba
 			ClearTempFiles();
 		}
 
-		m_processStats.wallTime = GetTime() - m_startTime;
-
 		#if PLATFORM_WINDOWS
 		if (m_accountingJobObject)
 		{
@@ -500,7 +508,11 @@ namespace uba
 			m_exitCode = ProcessCancelExitCode;
 		else
 			m_exitCode = exitCode;
-		}
+	}
+
+	void ProcessImpl::ThreadExit()
+	{
+		m_processStats.wallTime = GetTime() - m_startTime;
 
 		SystemStats::GetGlobal().Add(m_systemStats);
 
@@ -545,6 +557,84 @@ namespace uba
 		m_session.ProcessExited(*this, m_processStats.wallTime);
 	}
 
+	bool ProcessImpl::HandleSpecialApplication()
+	{
+	#if PLATFORM_WINDOWS
+		// TODO: This is super hacky... but we don't want to spawn the cmd.exe just to copy a file since the overhead can be half a second
+		// "C:\WINDOWS\system32\cmd.exe" /C "copy /Y "E:\dev\fn\Engine\Source\Runtime\RenderCore\RenderCore.natvis" "E:\dev\fn\Engine\Intermediate\Build\Win64\x64\UnrealPak\Development\RenderCore\RenderCore.natvis" 1>nul"
+		if (!Contains(m_startInfo.application, TC("cmd.exe")))
+			return false;
+		//m_session.m_logger.Info(TC("GOT HERE: %s"), m_startInfo.arguments);
+		if (!StartsWith(m_startInfo.arguments, TC("/C \"copy /Y \"")))
+			return false;
+		const tchar* fromFileBegin = m_startInfo.arguments + 13;
+		const tchar* fromFileEnd = TStrchr(fromFileBegin, '\"');
+		if (!fromFileEnd)
+			return false;
+		const tchar* toFileBegin = TStrchr(fromFileEnd + 1, '\"');
+		if (!toFileBegin)
+			return false;
+		++toFileBegin;
+		const tchar* toFileEnd = TStrchr(toFileBegin, '\"');
+		if (!toFileEnd)
+			return false;
+
+		m_processStats.wallTime = GetTime() - m_startTime;
+
+		StringBuffer<> workDir(m_startInfo.workingDir);
+		workDir.EnsureEndsWithSlash();
+
+		StringBuffer<> fromName;
+		StringBuffer<> toName;
+
+		StringBuffer<> temp;
+		temp.Append(fromFileBegin, fromFileEnd - fromFileBegin);
+		FixPath(temp.data, workDir.data, workDir.count, fromName);
+		temp.Clear().Append(toFileBegin, toFileEnd - toFileBegin);
+		FixPath(temp.data, workDir.data, workDir.count, toName);
+
+		DWORD oldAttributes = GetFileAttributes(toName.data);
+		if (oldAttributes != INVALID_FILE_ATTRIBUTES && (oldAttributes & FILE_ATTRIBUTE_READONLY))
+			SetFileAttributes(toName.data, oldAttributes & (~FILE_ATTRIBUTE_READONLY));
+
+		if (!CopyFileW(fromName.data, toName.data, false))
+		{
+			m_exitCode = GetLastError();
+			temp.Clear().Appendf(TC("Failed to copy %s to %s (%s)"), fromName.data, toName.data, LastErrorToText(m_exitCode).data);
+			m_logLines.push_back({ TString(temp.data), LogEntryType_Error });
+			return true;
+		}
+		
+		SetFileAttributes(toName.data, DefaultAttributes());
+
+		StringKey toKey = ToStringKeyLower(toName);
+		m_session.RegisterCreateFileForWrite(toKey, toName.data, toName.count, true);
+
+		//m_writtenFiles.try_emplace(name);
+		//WrittenFile& writtenFile = m_writtenFiles[toName.data];
+		//writtenFile.key = toKey;
+		//writtenFile.name = toName.data;
+		//writtenFile.owner = this;
+		//writtenFile.attributes = DefaultAttributes();
+
+		StackBinaryWriter<1024> trackedInputs;
+		trackedInputs.WriteString(fromName);
+		m_trackedInputs.resize(trackedInputs.GetPosition());
+		memcpy(m_trackedInputs.data(), trackedInputs.GetData(), trackedInputs.GetPosition());
+
+		StackBinaryWriter<1024> trackedOutputs;
+		trackedOutputs.WriteString(toName);
+		m_trackedOutputs.resize(trackedOutputs.GetPosition());
+		memcpy(m_trackedOutputs.data(), trackedOutputs.GetData(), trackedOutputs.GetPosition());
+		
+		m_exitCode = 0;
+		return true;
+	#else
+		return false;
+	#endif
+	}
+
+
 	bool ProcessImpl::HandleMessage(BinaryReader& reader, BinaryWriter& writer)
 	{
 		MessageType messageType = (MessageType)reader.ReadByte();
diff --git a/Engine/Source/Programs/UnrealBuildAccelerator/Common/Public/UbaProcess.h b/Engine/Source/Programs/UnrealBuildAccelerator/Common/Public/UbaProcess.h
--- a/Engine/Source/Programs/UnrealBuildAccelerator/Common/Public/UbaProcess.h
+++ b/Engine/Source/Programs/UnrealBuildAccelerator/Common/Public/UbaProcess.h
@@ -68,6 +68,8 @@ namespace uba
 		void SetWritten();
 
 		void ThreadRun(bool runningRemote, void* environment);
+		void ThreadExit();
+		bool HandleSpecialApplication();
 		bool HandleMessage(BinaryReader& reader, BinaryWriter& writer);
 		bool HandleInit(BinaryReader& reader, BinaryWriter& writer);
 		bool HandleCreateFile(BinaryReader& reader, BinaryWriter& writer);
