[UBA]
* Lots of fixes related to cache and prepared code for integration with UBT (ubt will provide input and output files for cache entries)

diff --git a/Engine/Source/Programs/UnrealBuildAccelerator/CacheService/Private/UbaCacheService.cpp b/Engine/Source/Programs/UnrealBuildAccelerator/CacheService/Private/UbaCacheService.cpp
--- a/Engine/Source/Programs/UnrealBuildAccelerator/CacheService/Private/UbaCacheService.cpp
+++ b/Engine/Source/Programs/UnrealBuildAccelerator/CacheService/Private/UbaCacheService.cpp
@@ -49,26 +49,23 @@ namespace uba
 		return -1;
 	}
 
-	StorageServer* g_storageServer;
-	CacheServer* g_cacheServer;
+	bool g_exitRequested;
 	
 	void CtrlBreakPressed()
 	{
-		if (g_storageServer)
-		{
-			g_storageServer->SaveCasTable(true);
-			LoggerWithWriter(g_consoleLogWriter).Info(TC("CAS table saved..."));
-
-			g_cacheServer->Save();
-		}
-		abort();
+		if (!g_exitRequested)
+			LoggerWithWriter(g_consoleLogWriter, TC("")).Info(TC("Exiting..."));
+		g_exitRequested = true;
 	}
 
 	#if PLATFORM_WINDOWS
 	BOOL ConsoleHandler(DWORD signal)
 	{
 		if (signal == CTRL_C_EVENT)
+		{
 			CtrlBreakPressed();
+			return TRUE;
+		}
 		return FALSE;
 	}
 	#else
@@ -227,24 +224,18 @@ namespace uba
 			return -1;
 
 		{
-			g_storageServer = &storageServer;
-			g_cacheServer = &cacheServer;
-
 			auto stopListen = MakeGuard([&]() { networkBackend.StopListen(); });
 			auto stopServer = MakeGuard([&]() { networkServer.DisconnectClients(); });
 
 			if (!networkServer.StartListen(networkBackend, port, listenIp.data))
 				return -1;
 
-			while (!IsEscapePressed())
+			while (!g_exitRequested && !IsEscapePressed())
 			{
 				Sleep(1000);
 				if (!cacheServer.RunMaintenance(false))
 					break;
 			}
-
-			g_cacheServer = nullptr;
-			g_storageServer = nullptr;
 		}
 
 		cacheServer.Save();
@@ -255,9 +246,7 @@ namespace uba
 #if PLATFORM_WINDOWS
 int wmain(int argc, wchar_t* argv[])
 {
-	int res = uba::WrappedMain(argc, argv);
-	Sleep(2000);
-	return res;
+	return uba::WrappedMain(argc, argv);
 }
 #else
 int main(int argc, char* argv[])
diff --git a/Engine/Source/Programs/UnrealBuildAccelerator/Common/Private/UbaCacheClient.cpp b/Engine/Source/Programs/UnrealBuildAccelerator/Common/Private/UbaCacheClient.cpp
--- a/Engine/Source/Programs/UnrealBuildAccelerator/Common/Private/UbaCacheClient.cpp
+++ b/Engine/Source/Programs/UnrealBuildAccelerator/Common/Private/UbaCacheClient.cpp
@@ -4,12 +4,14 @@
 #include "UbaCompactTables.h"
 #include "UbaFileAccessor.h"
 #include "UbaNetworkMessage.h"
-#include "UbaProcess.h"
+#include "UbaProcessStartInfo.h"
 #include "UbaRootPaths.h"
+#include "UbaSession.h"
 #include "UbaStorage.h"
 #include "UbaStorageUtils.h"
 
-#define UBA_LOG_CACHE_INFO 0
+#define UBA_LOG_WRITE_CACHE_INFO 0 // 0 = Disabled, 1 = Normal, 2 = Detailed
+#define UBA_LOG_FETCH_CACHE_INFO 0
 
 namespace uba
 {
@@ -39,12 +41,13 @@ namespace uba
 		u32 casKeyTableSizeSent = 0;
 	};
 
-	CacheClient::CacheClient(LogWriter& writer, StorageImpl& storage, NetworkClient& client, Session& session)
-	:	m_logger(writer, TC("UbaCacheClient"))
-	,	m_storage(storage)
-	,	m_client(client)
-	,	m_session(session)
+	CacheClient::CacheClient(const CacheClientCreateInfo& info)
+	:	m_logger(info.writer, TC("UbaCacheClient"))
+	,	m_storage(info.storage)
+	,	m_client(info.client)
+	,	m_session(info.session)
 	{
+		m_reportMissReason = info.reportMissReason;
 		m_client.RegisterOnConnected([this]()
 			{
 				StackBinaryWriter<1024> writer;
@@ -67,24 +70,25 @@ namespace uba
 
 	CacheClient::~CacheClient() = default;
 
-	bool CacheClient::WriteToCache(const RootPaths& rootPaths, u32 bucketId, const ProcessHandle& process)
+	bool CacheClient::WriteToCache(const RootPaths& rootPaths, u32 bucketId, const ProcessStartInfo& info, const u8* inputs, u64 inputsSize, const u8* outputs, u64 outputsSize)
 	{
 		if (!m_connected)
 			return false;
 
-		auto& si = process.GetStartInfo();
-		if (!si.trackInputs)
+		if (!inputsSize)
 			return false;
 
-		CasKey cmdKey = GetCmdKey(rootPaths, si);
+		CasKey cmdKey = GetCmdKey(rootPaths, info);
 		if (cmdKey == CasKeyZero)
+		{
+			#if UBA_LOG_WRITE_CACHE_INFO
+			m_logger.Info(TC("WRITECACHE FAIL: %s"), info.description);
+			#endif
 			return false;
+		}
 
-		const Vector<u8>& inputs = process.GetTrackedInputs();
-		BinaryReader inputsReader(inputs.data(), 0, inputs.size());
-
-		const Vector<u8>& outputs = process.GetTrackedOutputs();
-		BinaryReader outputsReader(outputs.data(), 0, outputs.size());
+		BinaryReader inputsReader(inputs, 0, inputsSize);
+		BinaryReader outputsReader(outputs, 0, outputsSize);
 
 		Map<u32, u32> inputsStringToCasKey;
 		Map<u32, u32> outputsStringToCasKey;
@@ -124,9 +128,16 @@ namespace uba
 					return false;
 				}
 			}
-			else if (path.EndsWith(TC(".rsp"))) // Paths can be absolute in rsp files so we need to normalize those paths
+			else if (ShouldNormalize(path)) // Paths can be absolute in rsp files so we need to normalize those paths
 			{
-				casKey = AsCompressed(rootPaths.NormalizeAndHashFile(m_logger, path.data), true);
+				casKey = rootPaths.NormalizeAndHashFile(m_logger, path.data);
+				if (casKey == CasKeyZero)
+				{
+					success = false;
+					continue;
+				}
+
+				casKey = AsCompressed(casKey, true);
 			}
 			else if (path[path.count-1] == ':')
 			{
@@ -154,17 +165,33 @@ namespace uba
 			u32 pathOffset = bucket.sendPathTable.Add(qualifiedPath.c_str(), u32(qualifiedPath.size()), &requiredPathTableSize);
 
 			if (!isOutput) // Output files should be removed from input files.. For example when cl.exe compiles pch it reads previous pch file and we don't want it to be input
+			{
 				if (outputsStringToCasKey.find(pathOffset) != outputsStringToCasKey.end())
 					continue;
+				//m_logger.Info(TC("INPUT ENTRY: %s -> %u"), qualifiedPath.c_str(), pathOffset);
+			}
+			else
+			{
+				inputsStringToCasKey.erase(pathOffset);
+				//m_logger.Info(TC("OUT ENTRY: %s -> %u"), qualifiedPath.c_str(), pathOffset);
+			}
 
-			auto insres = (isOutput ? outputsStringToCasKey : inputsStringToCasKey).try_emplace(pathOffset);
+			auto& stringToCasKey = isOutput ? outputsStringToCasKey : inputsStringToCasKey;
+			auto insres = stringToCasKey.try_emplace(pathOffset);
 			
 			if (!insres.second)
+			{
+				//m_logger.Warning(TC("Input file %s exists multiple times"), qualifiedPath.c_str()); 
 				continue;
+			}
 
-			// .dep.json contains absolute paths, need to normalize file
-			if (isOutput && path.EndsWith(TC(".dep.json"))) // TODO: More data driven approach. Also, hash does not match content atm.
-				casKey = AsCompressed(rootPaths.NormalizeAndHashFile(m_logger, path.data), true);
+			// Files that needs to be normalized
+			if (isOutput && ShouldNormalize(path))
+			{
+				casKey = rootPaths.NormalizeAndHashFile(m_logger, path.data);
+				if (casKey != CasKeyZero)
+					casKey = AsCompressed(casKey, true);
+			}
 
 			// Get file caskey using storage
 			if (casKey == CasKeyZero)
@@ -173,7 +200,11 @@ namespace uba
 				if (!m_storage.StoreCasFile(casKey, path.data, CasKeyZero, deferCreation))
 					return false;
 				if (casKey == CasKeyZero) // If file is not found it was a temporary file that was deleted and is not really an output
+				{
+					//m_logger.Warning(TC("Can't find file %s"), path.data); 
+					stringToCasKey.erase(insres.first);
 					continue; // m_logger.Info(TC("This should never happen! (%s)"), path.data);
+				}
 			}
 
 			UBA_ASSERT(IsCompressed(casKey));
@@ -184,7 +215,7 @@ namespace uba
 			return false;
 
 		if (outputsStringToCasKey.empty())
-			m_logger.Warning(TC("NO OUTPUTS FROM process %s"), process.GetStartInfo().description); 
+			m_logger.Warning(TC("NO OUTPUTS FROM process %s"), info.description); 
 
 		// Make sure server has enough of the path table to be able to resolve offsets from cache entry
 		if (!SendPathTable(bucket, requiredPathTableSize))
@@ -199,8 +230,26 @@ namespace uba
 			return false;
 
 
-		#if UBA_LOG_CACHE_INFO
-		m_logger.Info(TC("WRITECACHE: %s -> %u %s"), si.description, bucketId, CasKeyString(cmdKey).str);
+		#if UBA_LOG_WRITE_CACHE_INFO
+		m_logger.BeginScope();
+		m_logger.Info(TC("WRITECACHE: %s -> %u %s"), info.description, bucketId, CasKeyString(cmdKey).str);
+		#if UBA_LOG_WRITE_CACHE_INFO == 2
+		for (auto& kv : inputsStringToCasKey)
+		{
+			StringBuffer<> path;
+			CasKey casKey;
+			bucket.sendCasKeyTable.GetPathAndKey(path, casKey, bucket.sendPathTable, kv.second);
+			m_logger.Info(TC("   IN: %s -> %s"), path.data, CasKeyString(casKey).str);
+		}
+		for (auto& kv : outputsStringToCasKey)
+		{
+			StringBuffer<> path;
+			CasKey casKey;
+			bucket.sendCasKeyTable.GetPathAndKey(path, casKey, bucket.sendPathTable, kv.second);
+			m_logger.Info(TC("   OUT: %s -> %s"), path.data, CasKeyString(casKey).str);
+		}
+		#endif // 2
+		m_logger.EndScope();
 		#endif
 
 		return true;
@@ -263,10 +312,13 @@ namespace uba
 		// Traverse entries and test inputs against local machine
 		u32 entryCount = reader.ReadU16();
 
-		#if UBA_LOG_CACHE_INFO
+		#if UBA_LOG_FETCH_CACHE_INFO
 		auto mg = MakeGuard([&]() { m_logger.Info(TC("FETCHCACHE %s: %s -> %u %s (%u)"), success ? TC("SUCC") : TC("FAIL"), info.description, bucketId, CasKeyString(cmdKey).str, entryCount); });
 		#endif
 
+		struct MissInfo { TString path; u32 entryIndex; CasKey cache; CasKey local; };
+		Vector<MissInfo> misses;
+
 		for (u32 i=0; i!=entryCount; ++i)
 		{
 			u64 outputSize = 0;
@@ -278,27 +330,27 @@ namespace uba
 				while (reader.GetPositionData() != inputEnd)
 				{
 					u32 casKeyOffset = u32(reader.Read7BitEncoded());
+					StringBuffer<MaxPath> path;
+
+					CasKey cacheCasKey;
+					CasKey localCasKey;
 
 					auto insres = offsetIsMatch.try_emplace(casKeyOffset);
 					if (insres.second)
 					{
 						if (casKeyOffset >= bucket.serverCasKeyTable.GetSize())
-						{
-							TimerScope ts2(cacheStats.fetchCasTable);
-							if (!FetchCasTable(bucket))
+							if (!FetchCasTable(bucket, cacheStats, casKeyOffset + sizeof(CasKey)))
 								return false;
-						}
 
-						StringBuffer<MaxPath> path;
-						CasKey cacheCasKey;
 						if (!GetLocalPathAndCasKey(bucket, rootPaths, path, cacheCasKey, bucket.serverCasKeyTable, bucket.serverPathTable, casKeyOffset))
 							return false;
-						UBA_ASSERT(IsCompressed(cacheCasKey));
+						UBA_ASSERTF(IsCompressed(cacheCasKey), TC("Cache entry for %s has uncompressed cache key for path %s (%s)"), info.description, path.data, CasKeyString(cacheCasKey).str);
 
-						CasKey localCasKey;
-						if (path.EndsWith(TC(".rsp")) || path.EndsWith(TC(".dep.json"))) // Need to normalize caskey for these files since they contain absolute paths
+						if (ShouldNormalize(path)) // Need to normalize caskey for these files since they contain absolute paths
 						{
-							localCasKey = AsCompressed(rootPaths.NormalizeAndHashFile(m_logger, path.data), true);
+							localCasKey = rootPaths.NormalizeAndHashFile(m_logger, path.data);
+							if (localCasKey != CasKeyZero)
+								localCasKey = AsCompressed(localCasKey, true);
 						}
 						else
 						{
@@ -314,6 +366,10 @@ namespace uba
 					{
 						reader.Skip(inputEnd -  reader.GetPositionData());
 						isMatch = false;
+
+						if (m_reportMissReason && path.count) // if empty this has already been reported
+							misses.push_back({TString(path.data), i, cacheCasKey, localCasKey });
+
 						break;
 					}
 				}
@@ -335,7 +391,7 @@ namespace uba
 			{
 				u32 casKeyOffset = u32(reader.Read7BitEncoded());
 				if (casKeyOffset >= bucket.serverCasKeyTable.GetSize())
-					if (!FetchCasTable(bucket))
+					if (!FetchCasTable(bucket, cacheStats, casKeyOffset + sizeof(CasKey)))
 						return false;
 
 				TimerScope ts(cacheStats.fetchOutput);
@@ -347,8 +403,9 @@ namespace uba
 				UBA_ASSERT(IsCompressed(casKey));
 
 				FileFetcher fetcher { m_storage.m_bufferSlots };
+				fetcher.m_errorOnFail = false;
 
-				if (path.EndsWith(TC(".dep.json")))
+				if (ShouldNormalize(path))
 				{
 					// Fetch into memory, file is in special format without absolute paths
 					MemoryBlock normalizedBlock(1*1024*1024);
@@ -416,6 +473,9 @@ namespace uba
 			return true;
 		}
 
+		for (auto& miss : misses)
+			m_logger.Info(TC("Cache miss on %s because of mismatch of %s (entry: %u, local: %s cache: %s)"), info.description, miss.path.data(), miss.entryIndex, CasKeyString(miss.local).str, CasKeyString(miss.cache).str);
+
 		return false;
 	}
 
@@ -556,7 +616,7 @@ namespace uba
 			}
 			else // If we don't have the cas key it should be one of the normalized files.... otherwise there is a bug
 			{
-				if (path.EndsWith(TC(".dep.json")))
+				if (ShouldNormalize(path))
 				{
 					FileAccessor file(m_logger, path.data);
 					if (!file.OpenMemoryRead())
@@ -614,14 +674,18 @@ namespace uba
 		return true;
 	}
 
-	bool CacheClient::FetchCasTable(Bucket& bucket)
+	bool CacheClient::FetchCasTable(Bucket& bucket, CacheStats& stats, u32 requiredCasTableSize)
 	{
-		SCOPED_WRITE_LOCK(bucket.casKeyTableNetworkLock, lock);
+		TimerScope ts2(stats.fetchCasTable);
+
+		SCOPED_WRITE_LOCK(bucket.casKeyTableNetworkLock, lock); // Use one lock over both queries
+
+		if (requiredCasTableSize <= bucket.serverCasKeyTable.GetSize())
+			return true;
 
 		StackBinaryReader<SendMaxSize> reader;
 		{
-			//SCOPED_WRITE_LOCK(bucket.pathTableNetworkLock, lock);
-			u32 targetSize = ~0u;
+			u32 targetSize = requiredCasTableSize;
 			while (bucket.serverCasKeyTable.GetSize() < targetSize)
 			{
 				StackBinaryWriter<16> writer;
@@ -640,7 +704,7 @@ namespace uba
 			}
 		}
 		{
-			u32 targetSize = ~0u;
+			u32 targetSize = ~0u; // For now, read all because we don't know how much we need (it would require parsing all path offsets in caskey table
 			while (bucket.serverPathTable.GetSize() < targetSize)
 			{
 				StackBinaryWriter<16> writer;
@@ -665,16 +729,27 @@ namespace uba
 	{
 		CasKeyHasher hasher;
 
-		// Add hash of application binary to key
-		CasKey applicationCasKey;
-		bool deferCreation = true;
-		if (!m_storage.StoreCasFile(applicationCasKey, info.application, CasKeyZero, deferCreation))
-			return CasKeyZero;
-		hasher.Update(&applicationCasKey, sizeof(CasKey));
+		
+		#if PLATFORM_WINDOWS
+		// cmd.exe is special.. we can't hash it because it might be different on different os versions but should do the same thing regardless of version
+		if (Contains(info.application, TC("cmd.exe")))
+		{
+			hasher.Update(TC("cmd.exe"), 7*sizeof(tchar));
+		}
+		else
+		#endif
+		{
+			// Add hash of application binary to key
+			CasKey applicationCasKey;
+			bool deferCreation = true;
+			if (!m_storage.StoreCasFile(applicationCasKey, info.application, CasKeyZero, deferCreation))
+				return CasKeyZero;
+			hasher.Update(&applicationCasKey, sizeof(CasKey));
+		}
 
 		// Add arguments list to key
 		auto hashString = [&](const tchar* str, u64 strLen, u32 rootPos) { hasher.Update(str, strLen*sizeof(tchar)); };
-		if (!rootPaths.NormalizeString(m_logger, info.arguments, TStrlen(info.arguments), hashString, TC("")))
+		if (!rootPaths.NormalizeString(m_logger, info.arguments, TStrlen(info.arguments), hashString, TC("CmdKey")))
 			return CasKeyZero;
 
 		// Add content of rsp file to key (This will cost a bit of perf since we need to normalize.. should this be part of key?)
@@ -685,10 +760,12 @@ namespace uba
 				rspStart += 2;
 				if (auto rspEnd = TStrchr(rspStart, '"'))
 				{
+					StringBuffer<MaxPath> workingDir(info.workingDir);
+					workingDir.EnsureEndsWithSlash();
 					StringBuffer<> rsp;
-					if (rspStart[1] != ':')
-						rsp.Append(info.workingDir).EnsureEndsWithSlash();
 					rsp.Append(rspStart, rspEnd - rspStart);
+					StringBuffer<> fullPath;
+					FixPath(rsp.data, workingDir.data, workingDir.count, fullPath);
 					CasKey rspCasKey = rootPaths.NormalizeAndHashFile(m_logger, rsp.data);
 					hasher.Update(&rspCasKey, sizeof(CasKey));
 				}
@@ -698,11 +775,21 @@ namespace uba
 		return ToCasKey(hasher, false);
 	}
 
-	bool CacheClient::GetLocalPathAndCasKey(Bucket& bucket, const RootPaths& rootPaths, StringBufferBase& outPath, CasKey& outKey, CompactCasKeyTable& casKeyTable, CompactPathTable& pathTable, u32 offset)
+	bool CacheClient::ShouldNormalize(const StringBufferBase& path)
 	{
-		if (!m_connected)
-			return false;
+		if (path.EndsWith(TC(".dep.json"))) // Contains absolute paths
+			return true;
+		if (path.EndsWith(TC(".tlh"))) // Contains absolute path in a comment
+			return true;
+		if (path.EndsWith(TC(".rsp"))) // Contains absolute paths in some cases
+			return true;
+		if (path.EndsWith(TC(".bat"))) // Contains absolute paths in some cases
+			return true;
+		return false;
+	}
 
+	bool CacheClient::GetLocalPathAndCasKey(Bucket& bucket, const RootPaths& rootPaths, StringBufferBase& outPath, CasKey& outKey, CompactCasKeyTable& casKeyTable, CompactPathTable& pathTable, u32 offset)
+	{
 		SCOPED_READ_LOCK(bucket.casKeyTableNetworkLock, lock); // TODO: Is this needed?
 
 		StringBuffer<MaxPath> normalizedPath;
diff --git a/Engine/Source/Programs/UnrealBuildAccelerator/Common/Private/UbaCacheServer.cpp b/Engine/Source/Programs/UnrealBuildAccelerator/Common/Private/UbaCacheServer.cpp
--- a/Engine/Source/Programs/UnrealBuildAccelerator/Common/Private/UbaCacheServer.cpp
+++ b/Engine/Source/Programs/UnrealBuildAccelerator/Common/Private/UbaCacheServer.cpp
@@ -6,6 +6,7 @@
 #include "UbaFileAccessor.h"
 #include "UbaNetworkServer.h"
 #include "UbaStorageServer.h"
+//#include <oodle2.h>
 
 namespace uba
 {
@@ -587,6 +588,22 @@ namespace uba
 						}
 					});
 
+				#if 0
+				u8* mem = bucket.m_pathTable.GetMemory();
+				u64 memLeft = bucket.m_pathTable.GetSize();
+				while (memLeft)
+				{
+					u8 buffer[256*1024];
+					auto compressor = OodleLZ_Compressor_Kraken;
+					auto compressionLevel = OodleLZ_CompressionLevel_SuperFast;
+					u64 toCompress = Min(memLeft, u64(256*1024 - 128));
+					auto compressedBlockSize = OodleLZ_Compress(compressor, mem, (OO_SINTa)toCompress, buffer, compressionLevel);
+					(void)compressedBlockSize;
+					memLeft -= toCompress;
+				}
+				#endif
+
+
 				m_logger.Detail(TC("    Bucket %u Updated cache entries with new tables (%s)"), bucketIndex, TimeToText(GetTime() - updateEntriesStart).str);
 			}
 
@@ -740,10 +757,15 @@ namespace uba
 			#endif
 
 			u32 casKeyOffset = bucket.m_casKeyTable.Add(casKey, pathOffset);
-			inputs.insert(casKeyOffset);
+			auto insres = inputs.insert(casKeyOffset);
+			if (!insres.second)
+			{
+				m_logger.Warning(TC("Input file %s exists more than once in cache entry"), path.data);
+				continue;
+			}
 			bytesForInput += Get7BitEncodedCount(casKeyOffset);
 
-			//m_logger.Info(TC("%s - %s%s"), path.data, CasKeyString(casKey).str, isInput ? TC("") : TC(" (output)"));
+			//m_logger.Info(TC("%s - %s"), path.data, CasKeyString(casKey).str);
 		}
 
 		Vector<u8> inputCasKeyOffsets;
@@ -1012,6 +1034,8 @@ namespace uba
 									bucket.m_casKeyTable.GetPathAndKey(path, casKey, bucket.m_pathTable, offset);
 									if (path.Contains(filterString.data))
 										return true;
+									if (Contains(CasKeyString(casKey).str, filterString.data))
+										return true;
 								}
 								return false;
 							};
diff --git a/Engine/Source/Programs/UnrealBuildAccelerator/Common/Private/UbaCompactTables.cpp b/Engine/Source/Programs/UnrealBuildAccelerator/Common/Private/UbaCompactTables.cpp
--- a/Engine/Source/Programs/UnrealBuildAccelerator/Common/Private/UbaCompactTables.cpp
+++ b/Engine/Source/Programs/UnrealBuildAccelerator/Common/Private/UbaCompactTables.cpp
@@ -114,6 +114,13 @@ namespace uba
 
 	void CompactPathTable::GetString(StringBufferBase& out, u64 offset) const
 	{
+		#if UBA_DEBUG
+		{
+			SCOPED_READ_LOCK(const_cast<CompactPathTable*>(this)->m_lock, lock)
+			UBA_ASSERTF(offset < m_mem.writtenSize, TC("Reading path key from offset %llu which is out of bounds (Max %llu)"), offset, m_mem.writtenSize);
+		}
+		#endif
+
 		u32 offsets[256];
 		offsets[0] = u32(offset);
 		u32 offsetCount = 0;
@@ -282,6 +289,13 @@ namespace uba
 
 	void CompactCasKeyTable::GetPathAndKey(StringBufferBase& outPath, CasKey& outKey, const CompactPathTable& pathTable, u64 offset) const
 	{
+		#if UBA_DEBUG
+		{
+			SCOPED_READ_LOCK(const_cast<CompactCasKeyTable*>(this)->m_lock, lock)
+			UBA_ASSERTF(offset + sizeof(CasKey) < m_mem.writtenSize, TC("Reading cas key from offset %llu which is out of bounds (Max %llu)"), offset + sizeof(CasKey), m_mem.writtenSize);
+		}
+		#endif
+
 		BinaryReader reader(m_mem.memory, offset, 1000);
 		u32 stringOffset = (u32)reader.Read7BitEncoded();
 		outKey = reader.ReadCasKey();
diff --git a/Engine/Source/Programs/UnrealBuildAccelerator/Common/Private/UbaScheduler.cpp b/Engine/Source/Programs/UnrealBuildAccelerator/Common/Private/UbaScheduler.cpp
--- a/Engine/Source/Programs/UnrealBuildAccelerator/Common/Private/UbaScheduler.cpp
+++ b/Engine/Source/Programs/UnrealBuildAccelerator/Common/Private/UbaScheduler.cpp
@@ -515,8 +515,11 @@ namespace uba
 		delete[] dependencies;
 		delete si;
 
-		if (m_writeToCache && exitCode == 0 && process.GetStartInfo().trackInputs)
-			m_cacheClient->WriteToCache(*m_rootPaths[0], 0, ph);
+		if (m_writeToCache && exitCode == 0)
+		{
+			// TODO: Read dep.json file
+			UBA_ASSERTF(false, TC("Not implemented"));
+		}
 
 		ph.m_process = nullptr;
 	}
diff --git a/Engine/Source/Programs/UnrealBuildAccelerator/Common/Public/UbaCacheClient.h b/Engine/Source/Programs/UnrealBuildAccelerator/Common/Public/UbaCacheClient.h
--- a/Engine/Source/Programs/UnrealBuildAccelerator/Common/Public/UbaCacheClient.h
+++ b/Engine/Source/Programs/UnrealBuildAccelerator/Common/Public/UbaCacheClient.h
@@ -14,16 +14,27 @@ namespace uba
 	class RootPaths;
 	class Session;
 	class StorageImpl;
+	struct CacheStats;
 	struct CasKey;
 	struct ProcessStartInfo;
 
+	struct CacheClientCreateInfo
+	{
+		CacheClientCreateInfo(LogWriter& w, StorageImpl& st, NetworkClient& c, Session& se) : writer(w), storage(st), client(c), session(se) {}
+		LogWriter& writer;
+		StorageImpl& storage;
+		NetworkClient& client;
+		Session& session;
+		bool reportMissReason = false;
+	};
+
 	class CacheClient
 	{
 	public:
-		CacheClient(LogWriter& writer, StorageImpl& storage, NetworkClient& client, Session& session);
+		CacheClient(const CacheClientCreateInfo& info);
 		~CacheClient();
 
-		bool WriteToCache(const RootPaths& rootPaths, u32 bucketId, const ProcessHandle& process);
+		bool WriteToCache(const RootPaths& rootPaths, u32 bucketId, const ProcessStartInfo& info, const u8* inputs, u64 inputsSize, const u8* outputs, u64 outputsSize);
 		bool FetchFromCache(const RootPaths& rootPaths, u32 bucketId, const ProcessStartInfo& info);
 
 		bool WriteCacheSummary(const tchar* destinationFile, const tchar* filterString = nullptr);
@@ -38,9 +49,10 @@ namespace uba
 		bool SendPathTable(Bucket& bucket, u32 requiredPathTableSize);
 		bool SendCasTable(Bucket& bucket, u32 requiredCasTableSize);
 		bool SendCacheEntry(Bucket& bucket, const RootPaths& rootPaths, const CasKey& cmdKey, const Map<u32, u32>& inputsStringToCasKey, const Map<u32, u32>& outputsStringToCasKey);
-		bool FetchCasTable(Bucket& bucket);
+		bool FetchCasTable(Bucket& bucket, CacheStats& stats, u32 requiredCasTableSize);
 
 		CasKey GetCmdKey(const RootPaths& rootPaths, const ProcessStartInfo& info);
+		bool ShouldNormalize(const StringBufferBase& path);
 
 		bool GetLocalPathAndCasKey(Bucket& bucket, const RootPaths& rootPaths, StringBufferBase& outPath, CasKey& outKey, CompactCasKeyTable& casKeyTable, CompactPathTable& pathTable, u32 offset);
 
@@ -48,6 +60,7 @@ namespace uba
 		StorageImpl& m_storage;
 		NetworkClient& m_client;
 		Session& m_session;
+		bool m_reportMissReason;
 
 		Atomic<bool> m_connected;
 
diff --git a/Engine/Source/Programs/UnrealBuildAccelerator/Common/Public/UbaRootPaths.h b/Engine/Source/Programs/UnrealBuildAccelerator/Common/Public/UbaRootPaths.h
--- a/Engine/Source/Programs/UnrealBuildAccelerator/Common/Public/UbaRootPaths.h
+++ b/Engine/Source/Programs/UnrealBuildAccelerator/Common/Public/UbaRootPaths.h
@@ -78,7 +78,7 @@ namespace uba
 			auto root = FindRoot(path);
 			if (!root)
 			{
-				logger.Info(TC("PATH WITHOUT ROOT: %s (inside file %s)"), path.data, hint);
+				logger.Info(TC("PATH WITHOUT ROOT: %s (inside file %s at offset %u)"), path.data, hint, destPos);
 				return false;
 			}
 
diff --git a/Engine/Source/Programs/UnrealBuildAccelerator/Host/Private/UbaExports.cpp b/Engine/Source/Programs/UnrealBuildAccelerator/Host/Private/UbaExports.cpp
--- a/Engine/Source/Programs/UnrealBuildAccelerator/Host/Private/UbaExports.cpp
+++ b/Engine/Source/Programs/UnrealBuildAccelerator/Host/Private/UbaExports.cpp
@@ -519,7 +519,9 @@ extern "C"
 		auto networkClient = new NetworkClientWithBackend(ctorSuccess, ncci, networkBackend);
 		if (!ctorSuccess)
 			return nullptr;
-		return new CacheClient(writer, storage, *networkClient, *session);
+		CacheClientCreateInfo info{writer, storage, *networkClient, *session};
+		info.reportMissReason = true;
+		return new CacheClient(info);
 	}
 
 	bool CacheClient_Connect(uba::CacheClient* cacheClient, const uba::tchar* host, int port)
@@ -533,13 +535,13 @@ extern "C"
 		return true;
 	}
 
-	bool CacheClient_WriteToCache(uba::CacheClient* cacheClient, uba::RootPaths* rootPaths, uba::u32 bucket, const uba::ProcessHandle* process)
+	bool CacheClient_WriteToCache(uba::CacheClient* cacheClient, uba::RootPaths* rootPaths, uba::u32 bucket, const uba::ProcessHandle* process, const uba::u8* inputs, uba::u32 inputsSize, const uba::u8* outputs, uba::u32 outputsSize)
 	{
 		using namespace uba;
 		auto& si = process->GetStartInfo();
 		if (!si.trackInputs)
 			return false;
-		return cacheClient->WriteToCache(*rootPaths, bucket, *process);
+		return cacheClient->WriteToCache(*rootPaths, bucket, process->GetStartInfo(), inputs, inputsSize, outputs, outputsSize);
 	}
 
 	bool CacheClient_FetchFromCache(uba::CacheClient* cacheClient, uba::RootPaths* rootPaths, uba::u32 bucket, const uba::ProcessStartInfo& info)
diff --git a/Engine/Source/Programs/UnrealBuildAccelerator/Host/Public/UbaExports.h b/Engine/Source/Programs/UnrealBuildAccelerator/Host/Public/UbaExports.h
--- a/Engine/Source/Programs/UnrealBuildAccelerator/Host/Public/UbaExports.h
+++ b/Engine/Source/Programs/UnrealBuildAccelerator/Host/Public/UbaExports.h
@@ -122,7 +122,7 @@ extern "C"
 	// Cache
 	UBA_API uba::CacheClient* CacheClient_Create(uba::SessionServer* session);
 	UBA_API bool CacheClient_Connect(uba::CacheClient* cacheClient, const uba::tchar* host, int port);
-	UBA_API bool CacheClient_WriteToCache(uba::CacheClient* cacheClient, uba::RootPaths* rootPaths, uba::u32 bucket, const uba::ProcessHandle* process);
+	UBA_API bool CacheClient_WriteToCache(uba::CacheClient* cacheClient, uba::RootPaths* rootPaths, uba::u32 bucket, const uba::ProcessHandle* process, const uba::u8* inputs, uba::u32 inputsSize, const uba::u8* outputs, uba::u32 outputsSize);
 	UBA_API bool CacheClient_FetchFromCache(uba::CacheClient* cacheClient, uba::RootPaths* rootPaths, uba::u32 bucket, const uba::ProcessStartInfo& info);
 	UBA_API void CacheClient_Destroy(uba::CacheClient* cacheClient);
 
diff --git a/Engine/Source/Programs/UnrealBuildAccelerator/Visualizer/Private/UbaVisualizer.cpp b/Engine/Source/Programs/UnrealBuildAccelerator/Visualizer/Private/UbaVisualizer.cpp
--- a/Engine/Source/Programs/UnrealBuildAccelerator/Visualizer/Private/UbaVisualizer.cpp
+++ b/Engine/Source/Programs/UnrealBuildAccelerator/Visualizer/Private/UbaVisualizer.cpp
@@ -1283,6 +1283,10 @@ namespace uba
 
 				if (process.cacheFetch)
 				{
+					if (process.returned)
+						logger.Info(L"  Cache:  Miss");
+					else
+						logger.Info(L"  Cache:  Hit");
 					cacheStats.Read(reader, m_traceView.version);
 					if (process.exitCode == 0)
 					{
