[UBA]
* Fixed compile error in ubacli

diff --git a/Engine/Source/Programs/UnrealBuildAccelerator/Cli/Private/UbaCli.cpp b/Engine/Source/Programs/UnrealBuildAccelerator/Cli/Private/UbaCli.cpp
--- a/Engine/Source/Programs/UnrealBuildAccelerator/Cli/Private/UbaCli.cpp
+++ b/Engine/Source/Programs/UnrealBuildAccelerator/Cli/Private/UbaCli.cpp
@@ -565,7 +565,7 @@ namespace uba
 		auto CreateCacheClient = [&]()
 			{
 				auto nc = new NetworkClient(ctorSuccess);
-				cacheClient = new CacheClient(logWriter, storageServer, *nc, sessionServer);
+				cacheClient = new CacheClient({logWriter, storageServer, *nc, sessionServer});
 			};
 
 		if (cacheHost.count)
@@ -675,6 +675,31 @@ namespace uba
 
 		auto RunScheduler = [&](const tchar* yamlFile)
 		{
+			/*
+			bool ctorSuccess;
+			NetworkBackendMemory nbm(logWriter);
+			
+			StringBuffer<> cacheRootDir(g_rootDir);
+			cacheRootDir.Append("CacheServer");
+			NetworkServer cacheNetworkServer(ctorSuccess);
+			StorageServerCreateInfo storageInfo2(cacheNetworkServer, cacheRootDir.data, logWriter);
+			storageInfo2.writeRecievedCasFilesToDisk = true;
+			StorageServer cacheStorageServer(storageInfo2);
+			CacheServer cacheServer(logWriter, cacheRootDir.data, cacheNetworkServer, cacheStorageServer);
+			auto csg = MakeGuard([&]() { cacheServer.Save(); });
+
+			NetworkClient cacheNetworkClient(ctorSuccess);
+			CacheClient cacheClient(logWriter, storageServer, cacheNetworkClient, sessionServer);
+
+			if (false)
+			{
+				cacheServer.Load();
+
+				cacheNetworkServer.StartListen(nbm);
+				cacheNetworkClient.Connect(nbm, TC("127.0.0.1"));
+			}
+			auto g = MakeGuard([&]() { cacheNetworkClient.Disconnect(); cacheNetworkServer.DisconnectClients(); });
+			*/
 			auto g = MakeGuard([&]() { if (cacheClient) cacheClient->GetClient().Disconnect(); });
 
 			SchedulerCreateInfo info(sessionServer);
