[UBA]
* Fixed so cache server works on linux (hosting linux clients)
* Added UbaCacheService to deploy-linux.bat

diff --git a/Engine/Source/Programs/UnrealBuildAccelerator/Common/Private/UbaFileAccessor.cpp b/Engine/Source/Programs/UnrealBuildAccelerator/Common/Private/UbaFileAccessor.cpp
--- a/Engine/Source/Programs/UnrealBuildAccelerator/Common/Private/UbaFileAccessor.cpp
+++ b/Engine/Source/Programs/UnrealBuildAccelerator/Common/Private/UbaFileAccessor.cpp
@@ -127,9 +127,7 @@ namespace uba
 
 	bool FileAccessor::CreateMemoryWrite(bool allowRead, u32 flagsAndAttributes, u64 size, const tchar* tempPath)
 	{
-		#if PLATFORM_WINDOWS
 		allowRead = true; // It is not possible to have write only access to file mappings it seems
-		#endif
 
 		m_size = size;
 
@@ -137,13 +135,21 @@ namespace uba
 		if (!CreateWrite(allowRead, flagsAndAttributes, size, tempPath))
 			return false;
 
-		m_mappingHandle = uba::CreateFileMappingW(m_fileHandle, PAGE_READWRITE, size, m_fileName);
+		const tchar* realFileName = m_fileName;
+
+		#if !PLATFORM_WINDOWS
+		StringBuffer<> tempFile;
+		if (m_tempPath)
+			realFileName = tempFile.Append(m_tempPath).Append("Temp_").AppendValue(m_tempFileIndex).data;
+		#endif
+
+		m_mappingHandle = uba::CreateFileMappingW(m_fileHandle, PAGE_READWRITE, size, realFileName);
 		if (!m_mappingHandle.IsValid())
-			return m_logger.Error(TC("Failed to create memory map %s (%s)"), m_fileName, LastErrorToText().data);
+			return m_logger.Error(TC("Failed to create memory map %s (%s)"), realFileName, LastErrorToText().data);
 
 		m_data = MapViewOfFile(m_mappingHandle, FILE_MAP_WRITE, 0, size);
 		if (!m_data)
-			return m_logger.Error(TC("Failed to map view of file %s with size %llu, for write (%s)"), m_fileName, size, LastErrorToText().data);
+			return m_logger.Error(TC("Failed to map view of file %s with size %llu, for write (%s)"), realFileName, size, LastErrorToText().data);
 
 		return true;
 	}
diff --git a/Engine/Source/Programs/UnrealBuildAccelerator/Common/Private/UbaRootPaths.cpp b/Engine/Source/Programs/UnrealBuildAccelerator/Common/Private/UbaRootPaths.cpp
--- a/Engine/Source/Programs/UnrealBuildAccelerator/Common/Private/UbaRootPaths.cpp
+++ b/Engine/Source/Programs/UnrealBuildAccelerator/Common/Private/UbaRootPaths.cpp
@@ -36,9 +36,11 @@ namespace uba
 			root.index = index;
 			root.path = rp;
 
-			ToLower(root.path.data());
+			if (CaseInsensitiveFs)
+				ToLower(root.path.data());
+
 			if (root.path[root.path.size()-1] != PathSeparator)
-				return logger.Error(TC("Root path must end with separator"));
+				return logger.Error(TC("Root path '%s' must end with separator"), rp);
 
 			root.includeInKey = includeInKey;
 
@@ -79,7 +81,7 @@ namespace uba
 		CoTaskMemFree(path);
 
 		#else
-		UBA_ASSERT(false);
+		// no system roots
 		#endif
 		return true;
 	}
@@ -90,7 +92,10 @@ namespace uba
 			return nullptr;
 
 		StringBuffer<MaxPath> shortPath;
-		shortPath.Append(path.data, m_shortestRoot).MakeLower();
+		shortPath.Append(path.data, m_shortestRoot);
+		if (CaseInsensitiveFs)
+			shortPath.MakeLower();
+
 		StringKey key = ToStringKeyNoCheck(shortPath.data, m_shortestRoot);
 		for (u32 i=0, e=u32(m_roots.size()); i!=e; ++i)
 		{
diff --git a/Engine/Source/Programs/UnrealBuildAccelerator/Deploy-Linux.bat b/Engine/Source/Programs/UnrealBuildAccelerator/Deploy-Linux.bat
--- a/Engine/Source/Programs/UnrealBuildAccelerator/Deploy-Linux.bat
+++ b/Engine/Source/Programs/UnrealBuildAccelerator/Deploy-Linux.bat
@@ -25,6 +25,7 @@ call Engine/Build/BatchFiles/RunUBT.bat -NoUba -NoUbaLocal -NoSNDBS -NoXGE ^
 	-Target="UbaHost Linux %config%" ^
 	-Target="UbaTest Linux %config% %options%" ^
 	-Target="UbaTestApp Linux %config%" ^
+	-Target="UbaCacheService Linux %config%" ^
 	-Target="UbaCoordinatorHorde Linux %config% %options%"
 
 if %ERRORLEVEL% NEQ 0 exit /b %ERRORLEVEL%
