[UBA/UBT]
* Fixed so linux UbaCacheService can host windows clients

diff --git a/Engine/Source/Programs/UnrealBuildAccelerator/Common/Private/UbaStorageServer.cpp b/Engine/Source/Programs/UnrealBuildAccelerator/Common/Private/UbaStorageServer.cpp
--- a/Engine/Source/Programs/UnrealBuildAccelerator/Common/Private/UbaStorageServer.cpp
+++ b/Engine/Source/Programs/UnrealBuildAccelerator/Common/Private/UbaStorageServer.cpp
@@ -156,7 +156,12 @@ namespace uba
 						m_trace->FileEndStore(clientId, store.casEntry->key);
 				}
 
-				delete store.fileAccessor;
+				if (auto fa = store.fileAccessor)
+				{
+					const tchar* filename = fa->GetFileName();
+					delete store.fileAccessor;
+					free((void*)filename);
+				}
 
 				m_casDataBuffer.UnmapView(store.mappedView, TC("OnDisconnected"));
 				it = m_activeStores.erase(it);
@@ -746,10 +751,13 @@ namespace uba
 				{
 					StringBuffer<> casKeyName;
 					GetCasFileName(casKeyName, casKey);
-					fileAccessor = new FileAccessor(m_logger, casKeyName.data);
+					fileAccessor = new FileAccessor(m_logger, TStrdup(casKeyName.data));
 					if (!fileAccessor->CreateMemoryWrite(false, DefaultAttributes(), fileSize, m_tempPath.data))
 					{
+						const tchar* filename = fileAccessor->GetFileName();
 						delete fileAccessor;
+						free((void*)filename);
+
 						m_logger.Error(TC("Failed to create cas file %s"), casKeyName.data);
 						casEntry.verified = false;
 						return false;
@@ -819,7 +827,9 @@ namespace uba
 					if (activeStore.fileAccessor)
 					{
 						bool success = activeStore.fileAccessor->Close();
+						const tchar* filename = activeStore.fileAccessor->GetFileName();
 						delete activeStore.fileAccessor;
+						free((void*)filename);
 						if (!success)
 							return m_logger.Error(TC("REVISIT THIS!"));
 					}
diff --git a/Engine/Source/Programs/UnrealBuildTool/Executors/UnrealBuildAccelerator/UBAExecutor.cs b/Engine/Source/Programs/UnrealBuildTool/Executors/UnrealBuildAccelerator/UBAExecutor.cs
--- a/Engine/Source/Programs/UnrealBuildTool/Executors/UnrealBuildAccelerator/UBAExecutor.cs
+++ b/Engine/Source/Programs/UnrealBuildTool/Executors/UnrealBuildAccelerator/UBAExecutor.cs
@@ -563,13 +563,13 @@ namespace UnrealBuildTool
 			IRootPaths rootPaths = IRootPaths.Create(_ubaLogger!);
 			foreach (DirectoryItem root in action.RootPaths)
 			{
-				rootPaths.RegisterRoot(root.FullName + "\\", true);
+				rootPaths.RegisterRoot(root.FullName + Path.DirectorySeparatorChar, true);
 			}
 
 			DirectoryReference? autoSdkDir;
 			if (UEBuildPlatformSDK.TryGetHostPlatformAutoSDKDir(out autoSdkDir))
 			{
-				rootPaths.RegisterRoot(autoSdkDir.FullName + "\\", true);
+				rootPaths.RegisterRoot(autoSdkDir.FullName + Path.DirectorySeparatorChar, true);
 			}
 
 			rootPaths.RegisterSystemRoots();
