[UBA]
* Fixed so one UbaCacheService can handle both case sensitive and insensitive paths

diff --git a/Engine/Source/Programs/UnrealBuildAccelerator/Common/Private/UbaCacheClient.cpp b/Engine/Source/Programs/UnrealBuildAccelerator/Common/Private/UbaCacheClient.cpp
--- a/Engine/Source/Programs/UnrealBuildAccelerator/Common/Private/UbaCacheClient.cpp
+++ b/Engine/Source/Programs/UnrealBuildAccelerator/Common/Private/UbaCacheClient.cpp
@@ -15,13 +15,15 @@
 
 namespace uba
 {
+	u64 MakeId(u32 bucketId) { return u64(bucketId) | (u64(!CaseInsensitiveFs) << 32); }
+
 	struct CacheClient::Bucket
 	{
 		Bucket(u32 id_)
 		:	id(id_)
-		,	serverPathTable(CachePathTableMaxSize, CompactPathTable::V1)
+		,	serverPathTable(CachePathTableMaxSize, CompactPathTable::V1, CaseInsensitiveFs)
 		,	serverCasKeyTable(CacheCasKeyTableMaxSize)
-		,	sendPathTable(CachePathTableMaxSize, CompactPathTable::V1)
+		,	sendPathTable(CachePathTableMaxSize, CompactPathTable::V1, CaseInsensitiveFs)
 		,	sendCasKeyTable(CacheCasKeyTableMaxSize)
 		{
 		}
@@ -300,7 +302,7 @@ namespace uba
 			// Fetch entries.. server will provide as many as fits. TODO: Should it be possible to ask for more entries?
 			StackBinaryWriter<32> writer;
 			NetworkMessage msg(m_client, CacheServiceId, CacheMessageType_FetchEntries, writer);
-			writer.Write7BitEncoded(bucket.id);
+			writer.Write7BitEncoded(MakeId(bucket.id));
 			writer.WriteCasKey(cmdKey);
 			if (!msg.Send(reader))
 				return false;
@@ -508,7 +510,7 @@ namespace uba
 		{
 			StackBinaryWriter<SendMaxSize> writer;
 			NetworkMessage msg(m_client, CacheServiceId, CacheMessageType_StorePathTable, writer);
-			writer.Write7BitEncoded(bucket.id);
+			writer.Write7BitEncoded(MakeId(bucket.id));
 			u32 toSend = Min(requiredPathTableSize - bucket.pathTableSizeSent, u32(m_client.GetMessageMaxSize() - 32));
 			left -= toSend;
 			writer.WriteBytes(bucket.sendPathTable.GetMemory() + bucket.pathTableSizeSent, toSend);
@@ -532,7 +534,7 @@ namespace uba
 		{
 			StackBinaryWriter<SendMaxSize> writer;
 			NetworkMessage msg(m_client, CacheServiceId, CacheMessageType_StoreCasTable, writer);
-			writer.Write7BitEncoded(bucket.id);
+			writer.Write7BitEncoded(MakeId(bucket.id));
 			u32 toSend = Min(requiredCasTableSize - bucket.casKeyTableSizeSent, u32(m_client.GetMessageMaxSize() - 32));
 			left -= toSend;
 			writer.WriteBytes(bucket.sendCasKeyTable.GetMemory() + bucket.casKeyTableSizeSent, toSend);
@@ -552,7 +554,7 @@ namespace uba
 			StackBinaryWriter<SendMaxSize> writer;
 
 			NetworkMessage msg(m_client, CacheServiceId, CacheMessageType_StoreEntry, writer);
-			writer.Write7BitEncoded(bucket.id);
+			writer.Write7BitEncoded(MakeId(bucket.id));
 			writer.WriteCasKey(cmdKey);
 
 			writer.Write7BitEncoded(outputsStringToCasKey.size());
@@ -666,7 +668,7 @@ namespace uba
 		// Send done.. confirm to server
 		StackBinaryWriter<SendMaxSize> writer;
 		NetworkMessage msg(m_client, CacheServiceId, CacheMessageType_StoreEntryDone, writer);
-		writer.Write7BitEncoded(bucket.id);
+		writer.Write7BitEncoded(MakeId(bucket.id));
 		writer.WriteCasKey(cmdKey);
 		if (!msg.Send(reader))
 			return false;
@@ -690,7 +692,7 @@ namespace uba
 			{
 				StackBinaryWriter<16> writer;
 				NetworkMessage msg(m_client, CacheServiceId, CacheMessageType_FetchCasTable, writer);
-				writer.Write7BitEncoded(bucket.id);
+				writer.Write7BitEncoded(MakeId(bucket.id));
 				writer.WriteU32(bucket.serverCasKeyTable.GetSize());
 
 				reader.Reset();
@@ -709,7 +711,7 @@ namespace uba
 			{
 				StackBinaryWriter<16> writer;
 				NetworkMessage msg(m_client, CacheServiceId, CacheMessageType_FetchPathTable, writer);
-				writer.Write7BitEncoded(bucket.id);
+				writer.Write7BitEncoded(MakeId(bucket.id));
 				writer.WriteU32(bucket.serverPathTable.GetSize());
 
 				reader.Reset();
diff --git a/Engine/Source/Programs/UnrealBuildAccelerator/Common/Private/UbaCacheServer.cpp b/Engine/Source/Programs/UnrealBuildAccelerator/Common/Private/UbaCacheServer.cpp
--- a/Engine/Source/Programs/UnrealBuildAccelerator/Common/Private/UbaCacheServer.cpp
+++ b/Engine/Source/Programs/UnrealBuildAccelerator/Common/Private/UbaCacheServer.cpp
@@ -10,7 +10,9 @@
 
 namespace uba
 {
-	static constexpr u32 CacheFileVersion = 1;
+	static constexpr u32 CacheFileVersion = 2;
+
+	bool IsCaseInsensitive(u64 id) { return (id & (1ull << 32)) == 0; }
 
 	struct CacheServer::CacheEntry
 	{
@@ -21,17 +23,18 @@ namespace uba
 
 	struct CacheServer::ConnectionBucket
 	{
-		ConnectionBucket() : pathTable(CachePathTableMaxSize, CompactPathTable::V1), casKeyTable(CacheCasKeyTableMaxSize) {}
+		ConnectionBucket(u64 i) : pathTable(CachePathTableMaxSize, CompactPathTable::V1, IsCaseInsensitive(i)), casKeyTable(CacheCasKeyTableMaxSize), id(i) {}
 		CompactPathTable pathTable;
 		CompactCasKeyTable casKeyTable;
 
 		ReaderWriterLock cacheEntryLookupLock;
 		UnorderedMap<CasKey, CacheEntry> cacheEntryLookup;
+		u64 id;
 	};
 
 	struct CacheServer::Connection
 	{
-		UnorderedMap<u32, ConnectionBucket> buckets;
+		UnorderedMap<u64, ConnectionBucket> buckets;
 	};
 
 	struct CacheServer::CacheEntries
@@ -42,7 +45,7 @@ namespace uba
 
 	struct CacheServer::Bucket
 	{
-		Bucket() : m_pathTable(CachePathTableMaxSize, CompactPathTable::V1), m_casKeyTable(CacheCasKeyTableMaxSize) {}
+		Bucket(u64 id) : m_pathTable(CachePathTableMaxSize, CompactPathTable::V1, IsCaseInsensitive(id)), m_casKeyTable(CacheCasKeyTableMaxSize) {}
 		ReaderWriterLock m_cacheEntryLookupLock;
 		UnorderedMap<CasKey, CacheEntries> m_cacheEntryLookup;
 
@@ -120,8 +123,8 @@ namespace uba
 		u32 bucketCount = reader.ReadU32();
 		while (bucketCount--)
 		{
-			u32 bucketId = reader.ReadU32();
-			Bucket& bucket = m_buckets[bucketId];
+			u64 id = reader.ReadU64();
+			Bucket& bucket = m_buckets.try_emplace(id, id).first->second;
 
 			u32 pathTableSize = reader.ReadU32();
 			if (pathTableSize)
@@ -631,6 +634,22 @@ namespace uba
 		lock.Leave();
 	}
 
+	CacheServer::ConnectionBucket& CacheServer::GetConnectionBucket(const ConnectionInfo& connectionInfo, BinaryReader& reader)
+	{
+		u64 id = reader.Read7BitEncoded();
+		SCOPED_WRITE_LOCK(m_connectionsLock, lock);
+		auto& connection = m_connections[connectionInfo.GetId()];
+		return connection.buckets.try_emplace(id, id).first->second;
+	}
+
+	CacheServer::Bucket& CacheServer::GetBucket(BinaryReader& reader)
+	{
+		u64 id = reader.Read7BitEncoded();
+		SCOPED_WRITE_LOCK(m_bucketsLock, bucketsLock);
+		return m_buckets.try_emplace(id, id).first->second;
+		
+	}
+
 	bool CacheServer::HandleMessage(const ConnectionInfo& connectionInfo, u8 messageType, BinaryReader& reader, BinaryWriter& writer)
 	{
 		switch (messageType)
@@ -642,53 +661,36 @@ namespace uba
 				return m_logger.Error(TC("Different network versions. Client: %u, Server: %u. Disconnecting"), clientVersion, CacheNetworkVersion);
 			SCOPED_READ_LOCK(m_maintenanceLock, lock);
 			SCOPED_WRITE_LOCK(m_connectionsLock, lock2);
-			m_connections[connectionInfo.GetId()];
+			m_connections.try_emplace(connectionInfo.GetId());
 			return true;
 		}
 		case CacheMessageType_StorePathTable:
 		{
-			u32 bucketId = u32(reader.Read7BitEncoded());
-			SCOPED_WRITE_LOCK(m_connectionsLock, lock);
-			auto& connection = m_connections[connectionInfo.GetId()];
-			auto& connectionBucket = connection.buckets[bucketId];
-			lock.Leave();
-			connectionBucket.pathTable.ReadMem(reader, false);
+			GetConnectionBucket(connectionInfo, reader).pathTable.ReadMem(reader, false);
 			return true;
 		}
 		case CacheMessageType_StoreCasTable:
 		{
-			u32 bucketId = u32(reader.Read7BitEncoded());
-			SCOPED_WRITE_LOCK(m_connectionsLock, lock);
-			auto& connection = m_connections[connectionInfo.GetId()];
-			auto& connectionBucket = connection.buckets[bucketId];
-			lock.Leave();
-			connectionBucket.casKeyTable.ReadMem(reader, false);
+			GetConnectionBucket(connectionInfo, reader).casKeyTable.ReadMem(reader, false);
 			return true;
 		}
 		case CacheMessageType_StoreEntry:
 		{
-			u32 bucketId = u32(reader.Read7BitEncoded());
-			SCOPED_WRITE_LOCK(m_connectionsLock, lock);
-			auto& connection = m_connections[connectionInfo.GetId()];
-			auto& connectionBucket = connection.buckets[bucketId];
-			lock.Leave();
-			return HandleStoreEntry(bucketId, connectionBucket, reader, writer);
+			auto& bucket = GetConnectionBucket(connectionInfo, reader);
+			return HandleStoreEntry(bucket, reader, writer);
 		}
 		case CacheMessageType_StoreEntryDone:
 		{
-			u32 bucketId = u32(reader.Read7BitEncoded());
-			SCOPED_WRITE_LOCK(m_connectionsLock, lock);
-			auto& connection = m_connections[connectionInfo.GetId()];
-			auto& connectionBucket = connection.buckets[bucketId];
-			lock.Leave();
+			auto& connectionBucket = GetConnectionBucket(connectionInfo, reader);
 			CasKey cmdKey = reader.ReadCasKey();
 
 			SCOPED_WRITE_LOCK(connectionBucket.cacheEntryLookupLock, lock2);
 			auto findIt = connectionBucket.cacheEntryLookup.find(cmdKey);
 			if (findIt != connectionBucket.cacheEntryLookup.end())
 			{
+				u64 id = connectionBucket.id;
 				SCOPED_WRITE_LOCK(m_bucketsLock, bucketsLock);
-				Bucket& bucket = m_buckets[bucketId];
+				Bucket& bucket = m_buckets.try_emplace(id, id).first->second;
 				bucketsLock.Leave();
 
 				SCOPED_WRITE_LOCK(bucket.m_cacheEntryLookupLock, lock3);
@@ -720,7 +722,7 @@ namespace uba
 		}
 	}
 
-	bool CacheServer::HandleStoreEntry(u32 bucketId, ConnectionBucket& connectionBucket, BinaryReader& reader, BinaryWriter& writer)
+	bool CacheServer::HandleStoreEntry(ConnectionBucket& connectionBucket, BinaryReader& reader, BinaryWriter& writer)
 	{
 		CasKey cmdKey = reader.ReadCasKey();
 
@@ -731,9 +733,9 @@ namespace uba
 		u64 bytesForInput = 0;
 
 		u64 outputStartOffset = reader.GetPosition();
-
+		u64 id = connectionBucket.id;
 		SCOPED_WRITE_LOCK(m_bucketsLock, bucketsLock);
-		Bucket& bucket = m_buckets[bucketId];
+		Bucket& bucket = m_buckets.try_emplace(id, id).first->second;
 		bucketsLock.Leave();
 
 		while (reader.GetLeft())
@@ -901,12 +903,7 @@ namespace uba
 
 	bool CacheServer::HandleFetchPathTable(BinaryReader& reader, BinaryWriter& writer)
 	{
-		u32 bucketId = u32(reader.Read7BitEncoded());
-
-		SCOPED_WRITE_LOCK(m_bucketsLock, bucketsLock);
-		Bucket& bucket = m_buckets[bucketId];
-		bucketsLock.Leave();
-
+		Bucket& bucket = GetBucket(reader);
 		u32 haveSize = reader.ReadU32();
 		u32 size = bucket.m_pathTable.GetSize();
 		writer.WriteU32(size);
@@ -917,12 +914,7 @@ namespace uba
 
 	bool CacheServer::HandleFetchCasTable(BinaryReader& reader, BinaryWriter& writer)
 	{
-		u32 bucketId = u32(reader.Read7BitEncoded());
-
-		SCOPED_WRITE_LOCK(m_bucketsLock, bucketsLock);
-		Bucket& bucket = m_buckets[bucketId];
-		bucketsLock.Leave();
-
+		Bucket& bucket = GetBucket(reader);
 		u32 haveSize = reader.ReadU32();
 		u32 size = bucket.m_casKeyTable.GetSize();
 		writer.WriteU32(size);
@@ -933,12 +925,7 @@ namespace uba
 
 	bool CacheServer::HandleFetchEntries(BinaryReader& reader, BinaryWriter& writer)
 	{
-		u32 bucketId = u32(reader.Read7BitEncoded());
-
-		SCOPED_WRITE_LOCK(m_bucketsLock, bucketsLock);
-		Bucket& bucket = m_buckets[bucketId];
-		bucketsLock.Leave();
-
+		Bucket& bucket = GetBucket(reader);
 		CasKey cmdKey = reader.ReadCasKey();
 
 		u16& entryCount = *(u16*)writer.AllocWrite(2);
diff --git a/Engine/Source/Programs/UnrealBuildAccelerator/Common/Private/UbaCompactTables.cpp b/Engine/Source/Programs/UnrealBuildAccelerator/Common/Private/UbaCompactTables.cpp
--- a/Engine/Source/Programs/UnrealBuildAccelerator/Common/Private/UbaCompactTables.cpp
+++ b/Engine/Source/Programs/UnrealBuildAccelerator/Common/Private/UbaCompactTables.cpp
@@ -5,14 +5,15 @@
 
 namespace uba
 {
-	CompactPathTable::CompactPathTable(u64 reserveSize, Version version, u64 reservePathCount, u64 reserveSegmentCount)
+	CompactPathTable::CompactPathTable(u64 reserveSize, Version version, bool caseInsensitive, u64 reservePathCount, u64 reserveSegmentCount)
+	:	m_reserveSize(reserveSize)
+	,	m_version(version)
+	,	m_caseInsensitive(caseInsensitive)
 	{
-		m_reserveSize = reserveSize;
 		if (reservePathCount)
 			m_offsets.reserve(reservePathCount);
 		if (reserveSegmentCount)
 			m_segmentOffsets.reserve(reserveSegmentCount);
-		m_version = version;
 	}
 
 	u32 CompactPathTable::Add(const tchar* str, u64 strLen, u32* outRequiredCasTableSize)
@@ -33,7 +34,7 @@ namespace uba
 
 		const tchar* stringKeyString = str;
 		StringBuffer<MaxPath> tempStringKeyStr;
-		if (CaseInsensitiveFs)
+		if (m_caseInsensitive)
 			stringKeyString = tempStringKeyStr.Append(str).MakeLower().data;
 
 		return InternalAdd(str, stringKeyString, strLen);
@@ -41,7 +42,7 @@ namespace uba
 
 	u32 CompactPathTable::InternalAdd(const tchar* str, const tchar* stringKeyString, u64 strLen)
 	{
-		StringKey key = ToStringKey(stringKeyString, strLen);
+		StringKey key = ToStringKeyNoCheck(stringKeyString, strLen);
 		auto insres = m_offsets.try_emplace(key);
 		if (!insres.second)
 			return insres.first->second;
@@ -206,7 +207,7 @@ namespace uba
 				reader2.SkipString();
 				StringBuffer<> str;
 				GetString(str, offset);
-				if (CaseInsensitiveFs)
+				if (m_caseInsensitive)
 					str.MakeLower();
 				m_offsets.try_emplace(ToStringKeyNoCheck(str.data, str.count), offset);
 			}
@@ -227,7 +228,7 @@ namespace uba
 				}
 				StringBuffer<> str;
 				GetString(str, offset);
-				if (CaseInsensitiveFs)
+				if (m_caseInsensitive)
 					str.MakeLower();
 				m_offsets.try_emplace(ToStringKeyNoCheck(str.data, str.count), offset);
 			}
diff --git a/Engine/Source/Programs/UnrealBuildAccelerator/Common/Public/UbaCacheServer.h b/Engine/Source/Programs/UnrealBuildAccelerator/Common/Public/UbaCacheServer.h
--- a/Engine/Source/Programs/UnrealBuildAccelerator/Common/Public/UbaCacheServer.h
+++ b/Engine/Source/Programs/UnrealBuildAccelerator/Common/Public/UbaCacheServer.h
@@ -30,8 +30,15 @@ namespace uba
 
 		struct Connection;
 		struct ConnectionBucket;
+		struct CacheEntry;
+		struct CacheEntries;
+		struct Bucket;
+
+		ConnectionBucket& GetConnectionBucket(const ConnectionInfo& connectionInfo, BinaryReader& reader);
+		Bucket& GetBucket(BinaryReader& reader);
+
 		bool HandleMessage(const ConnectionInfo& connectionInfo, u8 messageType, BinaryReader& reader, BinaryWriter& writer);
-		bool HandleStoreEntry(u32 bucketId, ConnectionBucket& bucket, BinaryReader& reader, BinaryWriter& writer);
+		bool HandleStoreEntry(ConnectionBucket& bucket, BinaryReader& reader, BinaryWriter& writer);
 		bool HandleFetchPathTable(BinaryReader& reader, BinaryWriter& writer);
 		bool HandleFetchCasTable(BinaryReader& reader, BinaryWriter& writer);
 		bool HandleFetchEntries(BinaryReader& reader, BinaryWriter& writer);
@@ -46,12 +53,8 @@ namespace uba
 		ReaderWriterLock m_maintenanceLock;
 		Atomic<u32> m_addsSinceMaintenance;
 
-		struct CacheEntry;
-		struct CacheEntries;
-		struct Bucket;
-
 		ReaderWriterLock m_bucketsLock;
-		Map<u32, Bucket> m_buckets;
+		Map<u64, Bucket> m_buckets;
 
 		ReaderWriterLock m_connectionsLock;
 		Map<u32, Connection> m_connections;
diff --git a/Engine/Source/Programs/UnrealBuildAccelerator/Common/Public/UbaCompactTables.h b/Engine/Source/Programs/UnrealBuildAccelerator/Common/Public/UbaCompactTables.h
--- a/Engine/Source/Programs/UnrealBuildAccelerator/Common/Public/UbaCompactTables.h
+++ b/Engine/Source/Programs/UnrealBuildAccelerator/Common/Public/UbaCompactTables.h
@@ -16,7 +16,7 @@ namespace uba
 	public:
 		enum Version : u8 { V0, V1 };
 
-		CompactPathTable(u64 reserveSize, Version version, u64 reservePathCount = 0, u64 reserveSegmentCount = 0);
+		CompactPathTable(u64 reserveSize, Version version, bool caseSensitive, u64 reservePathCount = 0, u64 reserveSegmentCount = 0);
 
 		u32 Add(const tchar* str, u64 strLen, u32* outRequiredCasTableSize = nullptr);
 		u32 AddNoLock(const tchar* str, u64 strLen);
@@ -40,6 +40,7 @@ namespace uba
 		UnorderedMap<StringKey, u32> m_segmentOffsets;
 		u64 m_reserveSize;
 		Version m_version;
+		bool m_caseInsensitive;
 	};
 
 	class CompactCasKeyTable
diff --git a/Engine/Source/Programs/UnrealBuildAccelerator/Common/Public/UbaNetwork.h b/Engine/Source/Programs/UnrealBuildAccelerator/Common/Public/UbaNetwork.h
--- a/Engine/Source/Programs/UnrealBuildAccelerator/Common/Public/UbaNetwork.h
+++ b/Engine/Source/Programs/UnrealBuildAccelerator/Common/Public/UbaNetwork.h
@@ -16,7 +16,7 @@ namespace uba
 	static constexpr u32 SystemNetworkVersion = 1339;
 	static constexpr u32 StorageNetworkVersion = 4;
 	static constexpr u32 SessionNetworkVersion = 32;
-	static constexpr u32 CacheNetworkVersion = 1;
+	static constexpr u32 CacheNetworkVersion = 2;
 
 	static constexpr u32 CachePathTableMaxSize = 32*1024*1024;
 	static constexpr u32 CacheCasKeyTableMaxSize = 32*1024*1024;
