[UBA]
* Added header to compressed obj files containing the cas key of the compressed file
* Fixed so compressed obj files work with cache system

diff --git a/Engine/Source/Programs/UnrealBuildAccelerator/Common/Private/UbaCacheClient.cpp b/Engine/Source/Programs/UnrealBuildAccelerator/Common/Private/UbaCacheClient.cpp
--- a/Engine/Source/Programs/UnrealBuildAccelerator/Common/Private/UbaCacheClient.cpp
+++ b/Engine/Source/Programs/UnrealBuildAccelerator/Common/Private/UbaCacheClient.cpp
@@ -199,7 +199,8 @@ namespace uba
 			if (casKey == CasKeyZero)
 			{
 				bool deferCreation = true;
-				if (!m_storage.StoreCasFile(casKey, path.data, CasKeyZero, deferCreation))
+				bool fileIsCompressed = m_session.ShouldStoreObjFilesCompressed() && path.EndsWith(TC(".obj"));
+				if (!m_storage.StoreCasFile(casKey, path.data, CasKeyZero, deferCreation, fileIsCompressed))
 					return false;
 				if (casKey == CasKeyZero) // If file is not found it was a temporary file that was deleted and is not really an output
 				{
@@ -357,7 +358,8 @@ namespace uba
 						else
 						{
 							bool deferCreation = true;
-							m_storage.StoreCasFile(localCasKey, path.data, CasKeyZero, deferCreation);
+							bool fileIsCompressed = m_session.ShouldStoreObjFilesCompressed() && path.EndsWith(TC(".obj"));
+							m_storage.StoreCasFile(localCasKey, path.data, CasKeyZero, deferCreation, fileIsCompressed);
 							UBA_ASSERT(localCasKey == CasKeyZero || IsCompressed(localCasKey));
 						}
 
@@ -411,7 +413,8 @@ namespace uba
 				{
 					// Fetch into memory, file is in special format without absolute paths
 					MemoryBlock normalizedBlock(1*1024*1024);
-					if (!fetcher.RetrieveFile(m_logger, m_client, casKey, path.data, &normalizedBlock))
+					bool destinationIsCompressed = false;
+					if (!fetcher.RetrieveFile(m_logger, m_client, casKey, path.data, destinationIsCompressed, &normalizedBlock))
 						return false;
 					u32 rootOffsets = *(u32*)normalizedBlock.memory;
 					char* fileStart = (char*)normalizedBlock.memory + sizeof(u32);
@@ -461,8 +464,8 @@ namespace uba
 				}
 				else
 				{
-					UBA_ASSERT(!m_session.ShouldStoreObjFilesCompressed()); // TODO: Implement Retrieve .obj that is not decompressing
-					if (!fetcher.RetrieveFile(m_logger, m_client, casKey, path.data))
+					bool destinationIsCompressed = m_session.ShouldStoreObjFilesCompressed() && path.EndsWith(TC(".obj"));
+					if (!fetcher.RetrieveFile(m_logger, m_client, casKey, path.data, destinationIsCompressed))
 						return false;
 				}
 				if (!m_storage.FakeCopy(casKey, path.data, fetcher.m_size, fetcher.m_lastWritten, false))
@@ -494,7 +497,8 @@ namespace uba
 			return false;
 
 		FileFetcher fetcher { m_storage.m_bufferSlots };
-		if (!fetcher.RetrieveFile(m_logger, m_client, statusFileCasKey, destinationFile))
+		bool destinationIsCompressed = false;
+		if (!fetcher.RetrieveFile(m_logger, m_client, statusFileCasKey, destinationFile, destinationIsCompressed))
 			return false;
 		return true;
 	}
@@ -744,7 +748,8 @@ namespace uba
 			// Add hash of application binary to key
 			CasKey applicationCasKey;
 			bool deferCreation = true;
-			if (!m_storage.StoreCasFile(applicationCasKey, info.application, CasKeyZero, deferCreation))
+			bool fileIsCompressed = false;
+			if (!m_storage.StoreCasFile(applicationCasKey, info.application, CasKeyZero, deferCreation, fileIsCompressed))
 				return CasKeyZero;
 			hasher.Update(&applicationCasKey, sizeof(CasKey));
 		}
diff --git a/Engine/Source/Programs/UnrealBuildAccelerator/Common/Private/UbaCacheServer.cpp b/Engine/Source/Programs/UnrealBuildAccelerator/Common/Private/UbaCacheServer.cpp
--- a/Engine/Source/Programs/UnrealBuildAccelerator/Common/Private/UbaCacheServer.cpp
+++ b/Engine/Source/Programs/UnrealBuildAccelerator/Common/Private/UbaCacheServer.cpp
@@ -1076,7 +1076,9 @@ namespace uba
 			return false;
 
 		CasKey key;
-		if (!m_storage.StoreCasFile(key, tempFile.data))
+		bool deferCreation = false;
+		bool fileIsCompressed = false;
+		if (!m_storage.StoreCasFile(key, tempFile.data, CasKeyZero, deferCreation, fileIsCompressed))
 			return false;
 
 		writer.WriteCasKey(key);
diff --git a/Engine/Source/Programs/UnrealBuildAccelerator/Common/Private/UbaSession.cpp b/Engine/Source/Programs/UnrealBuildAccelerator/Common/Private/UbaSession.cpp
--- a/Engine/Source/Programs/UnrealBuildAccelerator/Common/Private/UbaSession.cpp
+++ b/Engine/Source/Programs/UnrealBuildAccelerator/Common/Private/UbaSession.cpp
@@ -2,6 +2,7 @@
 
 #include "UbaSession.h"
 #include "UbaBottleneck.h"
+#include "UbaCompressedObjFileHeader.h"
 #include "UbaFileAccessor.h"
 #include "UbaProcess.h"
 #include "UbaStorage.h"
@@ -9,6 +10,7 @@
 #include "UbaApplicationRules.h"
 #include "UbaPathUtils.h"
 #include "UbaProtocol.h"
+#include "UbaStorageUtils.h"
 #include "UbaWorkManager.h"
 
 #if PLATFORM_WINDOWS
@@ -2046,6 +2048,8 @@ namespace uba
 
 		auto writeFile = [&](WrittenFile& file)
 		{
+			bool shouldEvictFromMemory = IsRarelyReadAfterWritten(process, file.name.c_str(), file.name.size()) || file.mappingWritten > m_keepOutputFileMemoryMapsThreshold;
+
 			if (ShouldWriteToDisk(file.name.c_str(), file.name.size()))
 			{
 				// This is to kill I/O when writing lots of pdb/dlls in parallel
@@ -2074,8 +2078,11 @@ namespace uba
 				if (m_storeObjFilesCompressed && EndsWith(file.name.c_str(), file.name.size(), TC(".obj")))
 				{
 					Storage::WriteResult res;
-					if (!m_storage.WriteCompressed(res, TC("MemoryMap"), InvalidFileHandle, mem, fileSize, file.name.c_str()))
+					CompressedObjFileHeader header { CalculateCasKey(mem, fileSize, true, m_workManager) };
+
+					if (!m_storage.WriteCompressed(res, TC("MemoryMap"), InvalidFileHandle, mem, fileSize, file.name.c_str(), &header, sizeof(header)))
 						return false;
+					shouldEvictFromMemory = true;
 				}
 				else
 				{
@@ -2131,7 +2138,7 @@ namespace uba
 				uba::DeleteFileW(file.name.c_str());
 			}
 
-			if (IsRarelyReadAfterWritten(process, file.name.c_str(), file.name.size()) || file.mappingWritten > m_keepOutputFileMemoryMapsThreshold)
+			if (shouldEvictFromMemory)
 			{
 				m_workManager->AddWork([mh = file.mappingHandle]()
 					{
diff --git a/Engine/Source/Programs/UnrealBuildAccelerator/Common/Private/UbaSessionServer.cpp b/Engine/Source/Programs/UnrealBuildAccelerator/Common/Private/UbaSessionServer.cpp
--- a/Engine/Source/Programs/UnrealBuildAccelerator/Common/Private/UbaSessionServer.cpp
+++ b/Engine/Source/Programs/UnrealBuildAccelerator/Common/Private/UbaSessionServer.cpp
@@ -270,7 +270,8 @@ namespace uba
 				// Make sure cas entry exists and caskey is calculated (cas content creation is deferred in case client already has it)
 				CasKey casKey;
 				bool deferCreation = true;
-				if (!m_storage.StoreCasFile(casKey, fileName.data, CasKeyZero, deferCreation) || casKey == CasKeyZero)
+				bool fileIsCompressed = false;
+				if (!m_storage.StoreCasFile(casKey, fileName.data, CasKeyZero, deferCreation, fileIsCompressed) || casKey == CasKeyZero)
 					continue;
 
 				auto& ki = keys[keysIndex++];
@@ -383,7 +384,8 @@ namespace uba
 
 			CasKey casKey;
 			bool deferCreation = true;
-			if (!m_storage.StoreCasFile(casKey, path, CasKeyZero, deferCreation))
+			bool fileIsCompressed = false;
+			if (!m_storage.StoreCasFile(casKey, path, CasKeyZero, deferCreation, fileIsCompressed))
 				return false;
 			UBA_ASSERTF(casKey != CasKeyZero, TC("Failed to store cas for %s when calculating key for tracked inputs on %s"), path, fileName);
 			hasher.Update(&casKey, sizeof(CasKey));
@@ -583,14 +585,16 @@ namespace uba
 						if (!GetDirectoryOfCurrentModule(m_logger, dir))
 							return false;
 						u64 dirCount = dir.count;
+						bool deferCreation = true;
+						bool fileIsCompressed = false;
 						if (binAsVersion && m_agentBinaryKey == CasKeyZero)
 						{
 							UBA_ASSERT(IsWindows);
 							dir.Append(PathSeparator).Append(UBA_AGENT_EXECUTABLE);
-							m_storage.StoreCasFile(m_agentBinaryKey, dir.data, CasKeyZero, true);
+							m_storage.StoreCasFile(m_agentBinaryKey, dir.data, CasKeyZero, deferCreation, fileIsCompressed);
 						}
 						dir.Resize(dirCount).Append(PathSeparator).Append(UBA_DETOURS_LIBRARY);
-						if (!m_storage.StoreCasFile(m_detoursBinaryKey, dir.data, CasKeyZero, true))
+						if (!m_storage.StoreCasFile(m_detoursBinaryKey, dir.data, CasKeyZero, deferCreation, fileIsCompressed))
 							return m_logger.Error(TC("Failed to create cas for %s"), dir.data);
 						UBA_ASSERT(m_detoursBinaryKey != CasKeyZero);
 					}
@@ -1396,7 +1400,8 @@ namespace uba
 			}
 		}
 
-		if (!m_storage.StoreCasFile(out, fileName, casKeyOverride, deferCreation)) // We can defer the creation of the cas file since client might already have it
+		bool fileIsCompressed = false;
+		if (!m_storage.StoreCasFile(out, fileName, casKeyOverride, deferCreation, fileIsCompressed)) // We can defer the creation of the cas file since client might already have it
 			return false;
 
 		if (out != CasKeyZero)
diff --git a/Engine/Source/Programs/UnrealBuildAccelerator/Common/Private/UbaStorage.cpp b/Engine/Source/Programs/UnrealBuildAccelerator/Common/Private/UbaStorage.cpp
--- a/Engine/Source/Programs/UnrealBuildAccelerator/Common/Private/UbaStorage.cpp
+++ b/Engine/Source/Programs/UnrealBuildAccelerator/Common/Private/UbaStorage.cpp
@@ -2,6 +2,7 @@
 
 #include "UbaStorage.h"
 #include "UbaBottleneck.h"
+#include "UbaCompressedObjFileHeader.h"
 #include "UbaFileAccessor.h"
 #include "UbaBinaryReaderWriter.h"
 #include "UbaDirectoryIterator.h"
@@ -234,10 +235,10 @@ namespace uba
 		if (!uba::GetFileSizeEx(fileSize, readHandle))
 			return m_logger.Error(TC("GetFileSize failed for %s (%s)"), from, LastErrorToText().data);
 
-		return WriteCompressed(out, from, readHandle, 0, fileSize, to);
+		return WriteCompressed(out, from, readHandle, 0, fileSize, to, nullptr, 0);
 	}
 
-	bool StorageImpl::WriteCompressed(WriteResult& out, const tchar* from, FileHandle readHandle, u8* readMem, u64 fileSize, const tchar* to)
+	bool StorageImpl::WriteCompressed(WriteResult& out, const tchar* from, FileHandle readHandle, u8* readMem, u64 fileSize, const tchar* to, const void* header, u64 headerSize)
 	{
 		StorageStats& stats = Stats();
 
@@ -251,6 +252,9 @@ namespace uba
 		FileAccessor destinationFile(m_logger, to);
 		if (!destinationFile.CreateWrite(false, DefaultAttributes(), 0, m_tempPath.data))
 			return false;
+		if (headerSize)
+			if (!destinationFile.Write(header, headerSize))
+				return false;
 		if (!destinationFile.Write(&fileSize, sizeof(u64))) // Store file size first in compressed file
 			return false;
 #else
@@ -494,15 +498,18 @@ namespace uba
 		return true;
 	}
 
-	bool StorageImpl::WriteCasFileNoCheck(WriteResult& out, const tchar* fileName, const tchar* casFile, bool storeCompressed)
+	bool StorageImpl::WriteCasFileNoCheck(WriteResult& out, const tchar* fileName, bool fileIsCompressed, const tchar* casFile, bool storeCompressed)
 	{
-		if (storeCompressed)
+		if (storeCompressed && !fileIsCompressed)
 		{
 			if (!WriteCompressed(out, fileName, casFile))
 				return false;
 		}
 		else
 		{
+			if (!storeCompressed)
+				return m_logger.Error(TC("Write compressed file to uncompressed store not implemented"));
+
 			FileHandle readHandle;
 			if (!OpenFileSequentialRead(m_logger, fileName, readHandle))
 				return m_logger.Error(TC("Failed to open file %s for read (%s)"), fileName, LastErrorToText().data);
@@ -515,9 +522,18 @@ namespace uba
 			if (!destinationFile.CreateWrite(false, DefaultAttributes(), 0, m_tempPath.data))
 				return false;
 
+			if (fileIsCompressed)
+			{
+				CompressedObjFileHeader header;
+				if (!ReadFile(m_logger, fileName, readHandle, &header, sizeof(header)))
+					return false;
+				fileSize -= sizeof(header);
+			}
+
 			u8* slot= m_bufferSlots.Pop();
 			auto _ = MakeGuard([&](){ m_bufferSlots.Push(slot); });
 			u64 left = fileSize;
+
 			while (left)
 			{
 				u32 toRead = u32(Min(left, BufferSlotSize));
@@ -538,7 +554,7 @@ namespace uba
 		return true;
 	}
 
-	bool StorageImpl::WriteCasFile(WriteResult& out, const tchar* fileName, const CasKey& casKey)
+	bool StorageImpl::WriteCasFile(WriteResult& out, const tchar* fileName, bool fileIsCompressed, const CasKey& casKey)
 	{
 		UBA_ASSERT(IsCompressed(casKey) == m_storeCompressed);
 		StringBuffer<> casFile;
@@ -551,7 +567,7 @@ namespace uba
 		//UBA_ASSERT(false);
 #endif
 
-		return WriteCasFileNoCheck(out, fileName, casFile.data, IsCompressed(casKey));
+		return WriteCasFileNoCheck(out, fileName, fileIsCompressed, casFile.data, IsCompressed(casKey));
 	}
 
 	void StorageImpl::CasEntryAccessed(const CasKey& casKey)
@@ -819,7 +835,7 @@ namespace uba
 	}
 
 
-	bool StorageImpl::AddCasFile(const tchar* fileName, const CasKey& casKey, bool deferCreation)
+	bool StorageImpl::AddCasFile(const tchar* fileName, const CasKey& casKey, bool deferCreation, bool fileIsCompressed)
 	{
 		UBA_ASSERTF(IsCompressed(casKey) == m_storeCompressed, TC("CasKey compress mode must match storage compress mode (%s)"), fileName);
 		SCOPED_WRITE_LOCK(m_casLookupLock, lookupLock);
@@ -853,8 +869,10 @@ namespace uba
 			auto res = m_deferredCasCreationLookup.try_emplace(casKey);
 			if (res.second)
 			{
-				res.first->second = fileName;
-				auto res2 = m_deferredCasCreationLookupByName.try_emplace(res.first->second.c_str(), casKey);
+				DeferedCasCreation& dcc = res.first->second;
+				dcc.fileName = fileName;
+				dcc.fileIsCompressed = fileIsCompressed;
+				auto res2 = m_deferredCasCreationLookupByName.try_emplace(dcc.fileName.c_str(), casKey);
 				UBA_ASSERT(res2.second); (void)res2;
 			}
 			return true;
@@ -865,7 +883,7 @@ namespace uba
 
 		WriteResult res;
 #if !UBA_USE_SPARSEFILE
-		if (!WriteCasFileNoCheck(res, fileName, casFile.data, IsCompressed(casKey)))
+		if (!WriteCasFileNoCheck(res, fileName, fileIsCompressed, casFile.data, IsCompressed(casKey)))
 			return false;
 #else
 
@@ -1654,7 +1672,7 @@ namespace uba
 		return true;
 	}
 
-	bool StorageImpl::StoreCasFile(CasKey& out, const tchar* fileName, const CasKey& casKeyOverride, bool deferCreation)
+	bool StorageImpl::StoreCasFile(CasKey& out, const tchar* fileName, const CasKey& casKeyOverride, bool deferCreation, bool fileIsCompressed)
 	{
 		StringBuffer<> forKey;
 		forKey.Append(fileName);
@@ -1674,7 +1692,7 @@ namespace uba
 			if (fileEntry.casKey != CasKeyZero)
 			{
 				UBA_ASSERT(casKeyOverride == CasKeyZero || casKeyOverride == fileEntry.casKey);
-				if (!AddCasFile(fileName, fileEntry.casKey, deferCreation))
+				if (!AddCasFile(fileName, fileEntry.casKey, deferCreation, fileIsCompressed))
 					return false;
 			}
 			out = fileEntry.casKey;
@@ -1706,14 +1724,14 @@ namespace uba
 			if (casKeyOverride != CasKeyZero && casKeyOverride != fileEntry.casKey)
 			{
 				fileEntry.casKey = casKeyOverride;
-				if (!AddCasFile(fileName, fileEntry.casKey, deferCreation))
+				if (!AddCasFile(fileName, fileEntry.casKey, deferCreation, fileIsCompressed))
 					return false;
 				out = fileEntry.casKey;
 				return true;
 			}
 			if (fileSize == fileEntry.size && lastWritten == fileEntry.lastWritten)
 			{
-				if (!AddCasFile(fileName, fileEntry.casKey, deferCreation))
+				if (!AddCasFile(fileName, fileEntry.casKey, deferCreation, fileIsCompressed))
 					return false;
 				out = fileEntry.casKey;
 				return true;
@@ -1723,14 +1741,26 @@ namespace uba
 		fileEntry.size = fileSize;
 		fileEntry.lastWritten = lastWritten;
 		if (casKeyOverride == CasKeyZero)
-			fileEntry.casKey = CalculateCasKey(fileName, fileHandle, fileSize, m_storeCompressed);
+		{
+			if (fileIsCompressed)
+			{
+				CompressedObjFileHeader header;
+				if (!ReadFile(m_logger, fileName, fileHandle, &header, sizeof(header)))
+					return m_logger.Error(TC("Failed to read header of compressed file %s (%s)"), fileName, LastErrorToText().data);
+				fileEntry.casKey = AsCompressed(header.casKey, m_storeCompressed);
+			}
+			else
+			{
+				fileEntry.casKey = CalculateCasKey(fileName, fileHandle, fileSize, m_storeCompressed);
+			}
+		}
 		else
 			fileEntry.casKey = casKeyOverride;
 
 		if (fileEntry.casKey == CasKeyZero)
 			return false;
 
-		if (!AddCasFile(fileName, fileEntry.casKey, deferCreation))
+		if (!AddCasFile(fileName, fileEntry.casKey, deferCreation, fileIsCompressed))
 			return false;
 
 		out = fileEntry.casKey;
@@ -1765,12 +1795,13 @@ namespace uba
 		if (findIt == m_deferredCasCreationLookup.end())
 			return false;
 		casEntry.verified = true;
-		StringBuffer<> deferredCreation(findIt->second);
+		StringBuffer<> deferredCreation(findIt->second.fileName);
+		bool fileIsCompressed = findIt->second.fileIsCompressed;
 		m_deferredCasCreationLookupByName.erase(deferredCreation.data);
 		m_deferredCasCreationLookup.erase(findIt);
 		deferredLock.Leave();
 		WriteResult res;
-		if (!WriteCasFile(res, deferredCreation.data, casKey))
+		if (!WriteCasFile(res, deferredCreation.data, fileIsCompressed, casKey))
 			return false;
 #if UBA_USE_SPARSEFILE
 		casEntry.mappingHandle = res.mappingHandle;
@@ -1818,7 +1849,8 @@ namespace uba
 		if (!fileName)
 			return false;
 		WriteResult res;
-		if (!WriteCasFile(res, fileName, casKey))
+		bool fileIsCompressed = false; // TODO: This needs to be supported
+		if (!WriteCasFile(res, fileName, fileIsCompressed, casKey))
 			return false;
 		casEntry.exists = true;
 		entryLock.Leave();
@@ -2039,10 +2071,13 @@ namespace uba
 
 					if (writeCompressed)
 					{
+						// TODO: This should be more general.. send in header?
 						FileAccessor destinationFile(m_logger, destination);
-						if (!destinationFile.CreateMemoryWrite(false, fileAttributes, mappedView.size, m_tempPath.data))
+						if (!destinationFile.CreateMemoryWrite(false, fileAttributes, mappedView.size + sizeof(CompressedObjFileHeader), m_tempPath.data))
 							return false;
-						memcpy(destinationFile.GetData(), mappedView.memory, mappedView.size);
+						u8* mem = destinationFile.GetData();
+						*(CompressedObjFileHeader*)mem = { casKey };
+						memcpy(mem + sizeof(CompressedObjFileHeader), mappedView.memory, mappedView.size);
 						return destinationFile.Close();
 					}
 
diff --git a/Engine/Source/Programs/UnrealBuildAccelerator/Common/Private/UbaStorageClient.cpp b/Engine/Source/Programs/UnrealBuildAccelerator/Common/Private/UbaStorageClient.cpp
--- a/Engine/Source/Programs/UnrealBuildAccelerator/Common/Private/UbaStorageClient.cpp
+++ b/Engine/Source/Programs/UnrealBuildAccelerator/Common/Private/UbaStorageClient.cpp
@@ -672,7 +672,7 @@ namespace uba
 		return true;
 	}
 
-	bool StorageClient::StoreCasFile(CasKey& out, const tchar* fileName, const CasKey& casKeyOverride, bool deferCreation)
+	bool StorageClient::StoreCasFile(CasKey& out, const tchar* fileName, const CasKey& casKeyOverride, bool deferCreation, bool fileIsCompressed)
 	{
 		UBA_ASSERTF(false, TC("This StoreCasFile function should not be used on the client side"));
 		return true;
diff --git a/Engine/Source/Programs/UnrealBuildAccelerator/Common/Private/UbaStorageServer.cpp b/Engine/Source/Programs/UnrealBuildAccelerator/Common/Private/UbaStorageServer.cpp
--- a/Engine/Source/Programs/UnrealBuildAccelerator/Common/Private/UbaStorageServer.cpp
+++ b/Engine/Source/Programs/UnrealBuildAccelerator/Common/Private/UbaStorageServer.cpp
@@ -212,7 +212,8 @@ namespace uba
 		m_externalFileMappings.try_emplace(fileNameKey, ExternalFileMapping{mappingHandle, mappingOffset, fileSize});
 		externalFileLock.Leave();
 
-		if (!AddCasFile(fileName, fileEntry.casKey, deferCreation))
+		bool fileIsCompressed = false;
+		if (!AddCasFile(fileName, fileEntry.casKey, deferCreation, fileIsCompressed))
 			return false;
 
 		return true;
@@ -239,7 +240,7 @@ namespace uba
 		u8* fileMem = MapViewOfFile(mapping.mappingHandle, FILE_MAP_READ, mapping.mappingOffset, mapping.fileSize);
 		UBA_ASSERT(fileMem);
 		auto memClose = MakeGuard([&](){ UnmapViewOfFile(fileMem, mapping.fileSize, from); });
-		return StorageImpl::WriteCompressed(out, from, InvalidFileHandle, fileMem, mapping.fileSize, toFile);
+		return StorageImpl::WriteCompressed(out, from, InvalidFileHandle, fileMem, mapping.fileSize, toFile, nullptr, 0);
 	}
 
 	bool StorageServer::IsDisallowedPath(const tchar* fileName)
@@ -434,7 +435,9 @@ namespace uba
 							return false;
 						}
 
-						if (!AddCasFile(hint.data, casKey, false))
+						bool deferCreation = false;
+						bool fileIsCompressed = false;
+						if (!AddCasFile(hint.data, casKey, deferCreation, fileIsCompressed))
 						{
 							m_logger.Error(TC("FetchBegin failed for cas file %s (%s). Can't add cas file to database"), CasKeyString(casKey).str, hint.data);
 							writer.WriteU16(0);
diff --git a/Engine/Source/Programs/UnrealBuildAccelerator/Common/Private/UbaStorageUtils.cpp b/Engine/Source/Programs/UnrealBuildAccelerator/Common/Private/UbaStorageUtils.cpp
--- a/Engine/Source/Programs/UnrealBuildAccelerator/Common/Private/UbaStorageUtils.cpp
+++ b/Engine/Source/Programs/UnrealBuildAccelerator/Common/Private/UbaStorageUtils.cpp
@@ -1,6 +1,7 @@
 // Copyright Epic Games, Inc. All Rights Reserved.
 
 #include "UbaStorageUtils.h"
+#include "UbaCompressedObjFileHeader.h"
 #include "UbaStorage.h"
 #include "UbaFileAccessor.h"
 #include "UbaStats.h"
@@ -401,7 +402,7 @@ namespace uba
 		return true;
 	}
 
-	bool FileFetcher::RetrieveFile(Logger& logger, NetworkClient& client, const CasKey& casKey, const tchar* destination, MemoryBlock* destinationMem)
+	bool FileFetcher::RetrieveFile(Logger& logger, NetworkClient& client, const CasKey& casKey, const tchar* destination, bool destinationIsCompressed, MemoryBlock* destinationMem)
 	{
 		u8* slot = m_bufferSlots.Pop();
 		auto sg = MakeGuard([&](){ m_bufferSlots.Push(slot); });
@@ -463,6 +464,22 @@ namespace uba
 			if (!destinationFile.CreateWrite(false, DefaultAttributes(), sizeOnDisk, m_tempPath.data))
 				return false;
 
+
+		auto WriteDestination = [&](const void* source, u64 sourceSize, u64 sourceOffset = 0)
+			{
+				if (!destinationMem)
+				{
+					if (!destinationFile.Write(source, sourceSize, sourceOffset))
+						return false;
+				}
+				else
+				{
+					void* mem = destinationMem->Allocate(sourceSize, 1, TC(""));
+					memcpy(mem, source, sourceSize);
+				}
+				return true;
+			};
+
 		//u8* writePos = writeMem;
 
 		// This is here just to prevent server from getting a million messages at the same time.
@@ -473,16 +490,40 @@ namespace uba
 			m_retrieveOneBatchAtTheTimeLock.EnterWrite();
 		auto oatg = MakeGuard([&]() { if (oneAtTheTime) m_retrieveOneBatchAtTheTimeLock.LeaveWrite(); });
 
+		u32 readIndex = 0;
+
+		if (destinationIsCompressed)
+		{
+			CompressedObjFileHeader header { casKey };
+			if (!WriteDestination(&header, sizeof(header)))
+				return false;
+
+			if (!WriteDestination(readBuffer, responseSize))
+				return false;
+
+			u64 leftCompressed = fileSize - responseSize;
+			while (leftCompressed)
+			{
+				if (fetchId == u16(~0))
+					return logger.Error(TC("Cas content error (2). Server believes %s was only one segment but client sees more. "));//UncompressedSize: %llu LeftUncompressed: %llu Size: %llu Left to read: %llu ResponseSize: %u. (%s)"), destination, actualSize, leftUncompressed, fileSize, left, responseSize, CasKeyString(casKey).str);
+				if (!SendBatchMessages(logger, client, fetchId, slot, BufferSlotSize, leftCompressed, sizeOfFirstMessage, readIndex, responseSize))
+					return logger.Error(TC("Failed to send batched messages to server (%s)"), CasKeyString(casKey).str);
 
-		bool sendSegmentMessage = responseSize == 0;
-		u64 leftUncompressed = actualSize;
-		readBuffer += sizeof(u64); // Size is stored first
-		u64 maxReadSize = BufferSlotHalfSize - sizeof(u64);
+				if (!WriteDestination(slot, responseSize))
+					return false;
 
-		if (actualSize)
+				leftCompressed -= responseSize;
+			}
+			actualSize = sizeof(header) + fileSize;
+		}
+		else if (actualSize)
 		{
+			bool sendSegmentMessage = responseSize == 0;
+			u64 leftUncompressed = actualSize;
+			readBuffer += sizeof(u64); // Size is stored first
+			u64 maxReadSize = BufferSlotHalfSize - sizeof(u64);
+
 			u64 leftCompressed = fileSize - responseSize;
-			u32 readIndex = 0;
 			do
 			{
 				// First read in a full decompressable block
@@ -561,16 +602,9 @@ namespace uba
 					if (decompLen != uncompressedSize)
 						return logger.Error(TC("Expected %u but got %i when decompressing %u bytes for file %s"), uncompressedSize, int(decompLen), compressedSize, destination);
 
-					if (!destinationMem)
-					{
-						if (!destinationFile.Write(decompressBuffer, uncompressedSize, actualSize - leftUncompressed))
-							return false;
-					}
-					else
-					{
-						void* mem = destinationMem->Allocate(uncompressedSize, 1, TC(""));
-						memcpy(mem, decompressBuffer, uncompressedSize);
-					}
+					if (!WriteDestination(decompressBuffer, uncompressedSize, actualSize - leftUncompressed))
+						return false;
+
 					leftUncompressed -= uncompressedSize;
 				}
 
diff --git a/Engine/Source/Programs/UnrealBuildAccelerator/Common/Public/UbaStorage.h b/Engine/Source/Programs/UnrealBuildAccelerator/Common/Public/UbaStorage.h
--- a/Engine/Source/Programs/UnrealBuildAccelerator/Common/Public/UbaStorage.h
+++ b/Engine/Source/Programs/UnrealBuildAccelerator/Common/Public/UbaStorage.h
@@ -44,7 +44,7 @@ namespace uba
 		struct CachedFileInfo { CasKey casKey; };
 		virtual bool VerifyAndGetCachedFileInfo(CachedFileInfo& out, StringKey fileNameKey, u64 verifiedLastWriteTime, u64 verifiedSize) = 0;
 
-		virtual bool StoreCasFile(CasKey& out, const tchar* fileName, const CasKey& casKeyOverride = CasKeyZero, bool deferCreation = false) = 0;
+		virtual bool StoreCasFile(CasKey& out, const tchar* fileName, const CasKey& casKeyOverride, bool deferCreation, bool fileIsCompressed) = 0;
 		virtual bool StoreCasFile(CasKey& out, StringKey fileNameKey, const tchar* fileName, FileMappingHandle mappingHandle, u64 mappingOffset, u64 fileSize, const tchar* hint, bool deferCreation = false, bool keepMappingInMemory = false) = 0;
 		virtual bool DropCasFile(const CasKey& casKey, bool forceDelete, const tchar* hint) = 0;
 		virtual bool CalculateCasKey(CasKey& out, const tchar* fileName) = 0;
@@ -67,7 +67,7 @@ namespace uba
 		virtual void Ping() {}
 
 		struct WriteResult { FileMappingHandle mappingHandle; u64 size = InvalidValue; u64 offset = InvalidValue; };
-		virtual bool WriteCompressed(WriteResult& out, const tchar* from, FileHandle readHandle, u8* readMem, u64 fileSize, const tchar* toFile) = 0;
+		virtual bool WriteCompressed(WriteResult& out, const tchar* from, FileHandle readHandle, u8* readMem, u64 fileSize, const tchar* toFile, const void* header, u64 headerSize) = 0;
 	};
 
 	struct StorageCreateInfo
@@ -120,7 +120,7 @@ namespace uba
 		virtual bool DeleteCasForFile(const tchar* file) override;
 		virtual bool RetrieveCasFile(RetrieveResult& out, const CasKey& casKey, const tchar* hint, FileMappingBuffer* mappingBuffer = nullptr, u64 memoryMapAlignment = 1, bool allowProxy = true) override;
 		virtual bool VerifyAndGetCachedFileInfo(CachedFileInfo& out, StringKey fileNameKey, u64 verifiedLastWriteTime, u64 verifiedSize) override;
-		virtual bool StoreCasFile(CasKey& out, const tchar* fileName, const CasKey& casKeyOverride = CasKeyZero, bool deferCreation = false) override;
+		virtual bool StoreCasFile(CasKey& out, const tchar* fileName, const CasKey& casKeyOverride, bool deferCreation, bool fileIsCompressed) override;
 		virtual bool StoreCasFile(CasKey& out, StringKey fileNameKey, const tchar* fileName, FileMappingHandle mappingHandle, u64 mappingOffset, u64 fileSize, const tchar* hint, bool deferCreation = false, bool keepMappingInMemory = false) override;
 		virtual bool DropCasFile(const CasKey& casKey, bool forceDelete, const tchar* hint) override;
 		virtual bool CalculateCasKey(CasKey& out, const tchar* fileName) override;
@@ -145,12 +145,12 @@ namespace uba
 		CasKey CalculateCasKey(u8* fileMem, u64 fileSize, bool storeCompressed);
 
 		virtual bool WriteCompressed(WriteResult& out, const tchar* from, const tchar* toFile);
-		virtual bool WriteCompressed(WriteResult& out, const tchar* from, FileHandle readHandle, u8* readMem, u64 fileSize, const tchar* toFile) final;
+		virtual bool WriteCompressed(WriteResult& out, const tchar* from, FileHandle readHandle, u8* readMem, u64 fileSize, const tchar* toFile, const void* header, u64 headerSize) final;
 		bool WriteMemToCompressedFile(FileAccessor& destination, u32 workCount, const u8* uncompressedData, u64 fileSize, u64 maxUncompressedBlock, u64& totalWritten);
-		bool WriteCasFileNoCheck(WriteResult& out, const tchar* fileName, const tchar* casFile, bool storeCompressed);
-		bool WriteCasFile(WriteResult& out, const tchar* fileName, const CasKey& casKey);
+		bool WriteCasFileNoCheck(WriteResult& out, const tchar* fileName, bool fileIsCompressed, const tchar* casFile, bool storeCompressed);
+		bool WriteCasFile(WriteResult& out, const tchar* fileName, bool fileIsCompressed, const CasKey& casKey);
 		bool VerifyExisting(bool& outReturnValue, ScopedWriteLock& entryLock, const CasKey& casKey, CasEntry& casEntry, StringBufferBase& casFile);
-		bool AddCasFile(const tchar* fileName, const CasKey& casKey, bool deferCreation);
+		bool AddCasFile(const tchar* fileName, const CasKey& casKey, bool deferCreation, bool fileIsCompressed);
 		void CasEntryAccessed(const CasKey& casKey);
 		virtual bool IsDisallowedPath(const tchar* fileName);
 		virtual bool DecompressMemoryToMemory(u8* compressedData, u8* writeData, u64 decompressedSize, const tchar* readHint) override;
@@ -224,7 +224,8 @@ namespace uba
 		FileMappingBuffer m_casDataBuffer;
 
 		ReaderWriterLock m_deferredCasCreationLookupLock;
-		UnorderedMap<CasKey, TString> m_deferredCasCreationLookup;
+		struct DeferedCasCreation { TString fileName; bool fileIsCompressed; };
+		UnorderedMap<CasKey, DeferedCasCreation> m_deferredCasCreationLookup;
 		UnorderedMap<const tchar*, CasKey, HashString, EqualString> m_deferredCasCreationLookupByName;
 
 		DirectoryCache m_dirCache;
diff --git a/Engine/Source/Programs/UnrealBuildAccelerator/Common/Public/UbaStorageClient.h b/Engine/Source/Programs/UnrealBuildAccelerator/Common/Public/UbaStorageClient.h
--- a/Engine/Source/Programs/UnrealBuildAccelerator/Common/Public/UbaStorageClient.h
+++ b/Engine/Source/Programs/UnrealBuildAccelerator/Common/Public/UbaStorageClient.h
@@ -49,7 +49,7 @@ namespace uba
 
 		virtual bool GetZone(StringBufferBase& out) override;
 		virtual bool RetrieveCasFile(RetrieveResult& out, const CasKey& casKey, const tchar* hint, FileMappingBuffer* mappingBuffer = nullptr, u64 memoryMapAlignment = 1, bool allowProxy = true) override;
-		virtual bool StoreCasFile(CasKey& out, const tchar* fileName, const CasKey& casKeyOverride = CasKeyZero, bool deferCreation = false) override;
+		virtual bool StoreCasFile(CasKey& out, const tchar* fileName, const CasKey& casKeyOverride, bool deferCreation, bool fileIsCompressed) override;
 		virtual bool StoreCasFile(CasKey& out, StringKey fileNameKey, const tchar* fileName, FileMappingHandle mappingFile, u64 mappingOffset, u64 fileSize, const tchar* hint, bool deferCreation = false, bool keepMappingInMemory = false) override;
 		virtual bool HasCasFile(const CasKey& casKey, CasEntry** out = nullptr) override;
 		virtual void Ping() override;
diff --git a/Engine/Source/Programs/UnrealBuildAccelerator/Common/Public/UbaStorageUtils.h b/Engine/Source/Programs/UnrealBuildAccelerator/Common/Public/UbaStorageUtils.h
--- a/Engine/Source/Programs/UnrealBuildAccelerator/Common/Public/UbaStorageUtils.h
+++ b/Engine/Source/Programs/UnrealBuildAccelerator/Common/Public/UbaStorageUtils.h
@@ -35,7 +35,7 @@ namespace uba
 
 	struct FileFetcher
 	{
-		bool RetrieveFile(Logger& logger, NetworkClient& client, const CasKey& casKey, const tchar* destination, MemoryBlock* destinationMem = nullptr);
+		bool RetrieveFile(Logger& logger, NetworkClient& client, const CasKey& casKey, const tchar* destination, bool destinationIsCompressed, MemoryBlock* destinationMem = nullptr);
 
 		BufferSlots& m_bufferSlots;
 		StringBuffer<> m_tempPath;
diff --git a/Engine/Source/Programs/UnrealBuildAccelerator/Core/Private/UbaApplicationRules.cpp b/Engine/Source/Programs/UnrealBuildAccelerator/Core/Private/UbaApplicationRules.cpp
--- a/Engine/Source/Programs/UnrealBuildAccelerator/Core/Private/UbaApplicationRules.cpp
+++ b/Engine/Source/Programs/UnrealBuildAccelerator/Core/Private/UbaApplicationRules.cpp
@@ -54,10 +54,7 @@ namespace uba
 	public:
 		virtual bool IsOutputFile(const tchar* file, u64 fileLen) const override
 		{
-			return EndsWith(file, fileLen, TC(".h.obj"))
-				|| EndsWith(file, fileLen, TC(".c.obj"))
-				|| EndsWith(file, fileLen, TC(".cc.obj"))
-				|| EndsWith(file, fileLen, TC(".cpp.obj"))
+			return EndsWith(file, fileLen, TC(".obj"))
 				|| EndsWith(file, fileLen, TC(".dep.json"))
 				|| EndsWith(file, fileLen, TC(".rc2.res")) // Not really an obj file.. 
 				;// || EndsWith(file, fileLen, TC(".h.pch") // Not tested enough
diff --git a/Engine/Source/Programs/UnrealBuildAccelerator/Core/Public/UbaCompressedObjFileHeader.h b/Engine/Source/Programs/UnrealBuildAccelerator/Core/Public/UbaCompressedObjFileHeader.h
new file mode 100644
--- /dev/null
+++ b/Engine/Source/Programs/UnrealBuildAccelerator/Core/Public/UbaCompressedObjFileHeader.h
@@ -0,0 +1,13 @@
+// Copyright Epic Games, Inc. All Rights Reserved.
+
+#pragma once
+
+#include "UbaHash.h"
+
+namespace uba
+{
+	struct CompressedObjFileHeader
+	{
+		CasKey casKey;
+	};
+}
diff --git a/Engine/Source/Programs/UnrealBuildAccelerator/Detours/Private/Windows/UbaDetoursObjFilesPreloader.h b/Engine/Source/Programs/UnrealBuildAccelerator/Detours/Private/Windows/UbaDetoursObjFilesPreloader.h
--- a/Engine/Source/Programs/UnrealBuildAccelerator/Detours/Private/Windows/UbaDetoursObjFilesPreloader.h
+++ b/Engine/Source/Programs/UnrealBuildAccelerator/Detours/Private/Windows/UbaDetoursObjFilesPreloader.h
@@ -2,6 +2,7 @@
 
 #pragma once
 
+#include "UbaCompressedObjFileHeader.h"
 #include "UbaDetoursUtilsWin.h"
 #include "UbaEvent.h"
 #include <oodle2.h>
@@ -175,8 +176,8 @@ namespace uba
 				if (!preload.objMem)
 					continue;
 				preload.objCompressedSize = objFileSize.QuadPart;
-				preload.objDecompressedSize = *(u64*)preload.objMem;
-				preload.objReadOffset = 8;
+				preload.objDecompressedSize = *(u64*)(preload.objMem + sizeof(CompressedObjFileHeader));
+				preload.objReadOffset = 8 + sizeof(CompressedObjFileHeader);
 				preload.objWriteOffset = 0;
 				preload.objLeft = preload.objDecompressedSize;
 				CloseHandle(objFileMappingHandle);
@@ -229,6 +230,7 @@ namespace uba
 				}
 
 				BinaryReader reader(preload.objMem, preload.objReadOffset, preload.objCompressedSize);
+
 				u32 compressedBlockSize = reader.ReadU32();
 				u32 decompressedBlockSize = reader.ReadU32();
 
@@ -241,7 +243,7 @@ namespace uba
 
 				OO_SINTa decompLen = OodleLZ_Decompress(reader.GetPositionData(), (OO_SINTa)compressedBlockSize, destMem, (OO_SINTa)decompressedBlockSize,
 					OodleLZ_FuzzSafe_Yes, OodleLZ_CheckCRC_No, OodleLZ_Verbosity_None, NULL, 0, NULL, NULL, decoderMem, decoredMemSize);
-				UBA_ASSERT(decompLen == decompressedBlockSize);(void)decompLen;
+				UBA_ASSERTF(decompLen == decompressedBlockSize, TC("Failed to decompress .obj file %s (%s)"), info.name, info.originalName);(void)decompLen;
 
 				bool isDone = preload.objLeft.fetch_sub(decompressedBlockSize) == decompressedBlockSize;
 				if (!isDone)
diff --git a/Engine/Source/Programs/UnrealBuildAccelerator/Host/Private/UbaExports.cpp b/Engine/Source/Programs/UnrealBuildAccelerator/Host/Private/UbaExports.cpp
--- a/Engine/Source/Programs/UnrealBuildAccelerator/Host/Private/UbaExports.cpp
+++ b/Engine/Source/Programs/UnrealBuildAccelerator/Host/Private/UbaExports.cpp
@@ -520,7 +520,7 @@ extern "C"
 		if (!ctorSuccess)
 			return nullptr;
 		CacheClientCreateInfo info{writer, storage, *networkClient, *session};
-		info.reportMissReason = true;
+		//info.reportMissReason = true;
 		return new CacheClient(info);
 	}
 
diff --git a/Engine/Source/Programs/UnrealBuildAccelerator/Test/Private/UbaTestBasics.h b/Engine/Source/Programs/UnrealBuildAccelerator/Test/Private/UbaTestBasics.h
--- a/Engine/Source/Programs/UnrealBuildAccelerator/Test/Private/UbaTestBasics.h
+++ b/Engine/Source/Programs/UnrealBuildAccelerator/Test/Private/UbaTestBasics.h
@@ -314,7 +314,7 @@ namespace uba
 	{
 		for (u32 i=0;i!=2; ++i)
 		{
-			CompactPathTable table(64*1024, CompactPathTable::Version(i));
+			CompactPathTable table(64*1024, CompactPathTable::Version(i), CaseInsensitiveFs);
 
 			StringBuffer<> str;
 			str.Append("foo").EnsureEndsWithSlash().Append("bar");
@@ -331,7 +331,7 @@ namespace uba
 			if (!str.Equals(str2.data))
 				return false;
 
-			CompactPathTable table2(64*1024, CompactPathTable::Version(i));
+			CompactPathTable table2(64*1024, CompactPathTable::Version(i), CaseInsensitiveFs);
 			BinaryReader reader(table.GetMemory(), 0, table.GetSize());
 			table2.ReadMem(reader, true);
 			u32 offset2 = table2.Add(str.data, str.count);
diff --git a/Engine/Source/Programs/UnrealBuildAccelerator/Test/Private/UbaTestStorage.h b/Engine/Source/Programs/UnrealBuildAccelerator/Test/Private/UbaTestStorage.h
--- a/Engine/Source/Programs/UnrealBuildAccelerator/Test/Private/UbaTestStorage.h
+++ b/Engine/Source/Programs/UnrealBuildAccelerator/Test/Private/UbaTestStorage.h
@@ -27,7 +27,9 @@ namespace uba
 		storage.LoadCasTable();
 
 		CasKey key;
-		if (!storage.StoreCasFile(key, detoursLib.data))
+		bool deferCreation = false;
+		bool fileIsCompressed = false;
+		if (!storage.StoreCasFile(key, detoursLib.data, CasKeyZero, deferCreation, fileIsCompressed))
 			return logger.Error(TC("Failed to store file %s"), detoursLib.data);
 		if (key == CasKeyZero)
 			return logger.Error(TC("Failed to find file %s"), detoursLib.data);
