[UBA]
* Fix crash in visualizer that can happen if box is selected when new build trigger a refresh of view

diff --git a/Engine/Source/Programs/UnrealBuildAccelerator/Visualizer/Private/UbaVisualizer.cpp b/Engine/Source/Programs/UnrealBuildAccelerator/Visualizer/Private/UbaVisualizer.cpp
--- a/Engine/Source/Programs/UnrealBuildAccelerator/Visualizer/Private/UbaVisualizer.cpp
+++ b/Engine/Source/Programs/UnrealBuildAccelerator/Visualizer/Private/UbaVisualizer.cpp
@@ -292,6 +292,22 @@ namespace uba
 		out.Append(L" - ");
 	}
 
+	bool Visualizer::Unselect()
+	{
+		if (m_processSelected || m_sessionSelectedIndex != ~0u || m_statsSelected || m_timelineSelected || m_fetchedFilesSelected != ~0u || m_workSelected)
+		{
+			m_processSelected = false;
+			m_sessionSelectedIndex = ~0u;
+			m_statsSelected = false;
+			m_buttonSelected = ~0u;
+			m_timelineSelected = 0;
+			m_fetchedFilesSelected = ~0u;
+			m_workSelected = false;
+			return true;
+		}
+		return false;
+	}
+
 	void Visualizer::Reset()
 	{
 		for (HBITMAP bm : m_textBitmaps)
@@ -312,6 +328,8 @@ namespace uba
 		//m_replay = 0;
 		m_startTime = GetTime();
 		m_pauseTime = 0;
+
+		Unselect();
 	}
 
 	void Visualizer::ThreadLoop()
@@ -2155,17 +2173,8 @@ namespace uba
 
 	void Visualizer::UnselectAndRedraw()
 	{
-		if (m_processSelected || m_sessionSelectedIndex != ~0u || m_statsSelected || m_timelineSelected || m_fetchedFilesSelected != ~0u || m_workSelected)
-		{
-			m_processSelected = false;
-			m_sessionSelectedIndex = ~0u;
-			m_statsSelected = false;
-			m_buttonSelected = ~0u;
-			m_timelineSelected = 0;
-			m_fetchedFilesSelected = ~0u;
-			m_workSelected = false;
+		if (Unselect())
 			RedrawWindow(m_hwnd, NULL, NULL, RDW_INVALIDATE);
-		}
 	}
 
 	bool Visualizer::UpdateAutoscroll()
diff --git a/Engine/Source/Programs/UnrealBuildAccelerator/Visualizer/Private/UbaVisualizer.h b/Engine/Source/Programs/UnrealBuildAccelerator/Visualizer/Private/UbaVisualizer.h
--- a/Engine/Source/Programs/UnrealBuildAccelerator/Visualizer/Private/UbaVisualizer.h
+++ b/Engine/Source/Programs/UnrealBuildAccelerator/Visualizer/Private/UbaVisualizer.h
@@ -25,6 +25,7 @@ namespace uba
 		HWND GetHwnd();
 
 	private:
+		bool Unselect();
 		void Reset();
 		void PaintClient(const Function<void(HDC hdc, HDC memDC, RECT& clientRect)>& paintFunc);
 		void PaintAll(HDC hdc, const RECT& clientRect);
