[UBA]
* Bumped reserved memory used by detoured process to track directory table, mapped file table, etc. This to solve unity stress test build error

diff --git a/Engine/Source/Programs/UnrealBuildAccelerator/Detours/Private/UbaDetoursShared.cpp b/Engine/Source/Programs/UnrealBuildAccelerator/Detours/Private/UbaDetoursShared.cpp
--- a/Engine/Source/Programs/UnrealBuildAccelerator/Detours/Private/UbaDetoursShared.cpp
+++ b/Engine/Source/Programs/UnrealBuildAccelerator/Detours/Private/UbaDetoursShared.cpp
@@ -37,7 +37,7 @@ namespace uba
 		g_systemRootMem.Create();
 		g_systemTempMem.Create();
 
-		u64 reserveSizeMb = IsWindows ? 160 : 1024; // The sync primitives on linux/macos is much bigger
+		u64 reserveSizeMb = IsWindows ? 192 : 1024; // The sync primitives on linux/macos is much bigger
 		g_memoryBlockMem.Create(reserveSizeMb * 1024 * 1024);
 		g_directoryTableMem.Create(&g_memoryBlock);
 		g_mappedFileTableMem.Create(g_memoryBlock);
