[UBA]
* Fixed bugs in detoured logic for NtCreateFile to make sure it behaves correctly when open for read/write etc..

diff --git a/Engine/Source/Programs/UnrealBuildAccelerator/Detours/Private/PosixOS/UbaDetoursFunctionsPosix.cpp b/Engine/Source/Programs/UnrealBuildAccelerator/Detours/Private/PosixOS/UbaDetoursFunctionsPosix.cpp
--- a/Engine/Source/Programs/UnrealBuildAccelerator/Detours/Private/PosixOS/UbaDetoursFunctionsPosix.cpp
+++ b/Engine/Source/Programs/UnrealBuildAccelerator/Detours/Private/PosixOS/UbaDetoursFunctionsPosix.cpp
@@ -232,7 +232,6 @@ namespace uba
 		return strerror(error);
 	}
 
-	enum : u8 { AccessFlag_Read = 1, AccessFlag_Write = 2 };
 	u8 GetFileAccessFlags(int flags)
 	{
 		u8 access = 0;
diff --git a/Engine/Source/Programs/UnrealBuildAccelerator/Detours/Private/UbaDetoursFileMappingTable.cpp b/Engine/Source/Programs/UnrealBuildAccelerator/Detours/Private/UbaDetoursFileMappingTable.cpp
--- a/Engine/Source/Programs/UnrealBuildAccelerator/Detours/Private/UbaDetoursFileMappingTable.cpp
+++ b/Engine/Source/Programs/UnrealBuildAccelerator/Detours/Private/UbaDetoursFileMappingTable.cpp
@@ -78,7 +78,7 @@ namespace uba
 		u32 mappedFileTableSize = reader.ReadU32();
 		u32 directoryTableSize = u32(reader.ReadU32());
 		pcs.Leave();
-		DEBUG_LOG_PIPE(L"CreateFile", L"%ls (%ls)", (access == 0 ? L"ATTRIB" : ((access & GENERIC_WRITE) ? L"WRITE" : L"READ")), fileName);
+		DEBUG_LOG_PIPE(L"CreateFile", L"%ls (%ls)", (access == 0 ? L"ATTRIB" : ((access & AccessFlag_Write) ? L"WRITE" : L"READ")), fileName);
 
 		if (lock)
 			g_mappedFileTable.Parse(mappedFileTableSize);
diff --git a/Engine/Source/Programs/UnrealBuildAccelerator/Detours/Private/UbaDetoursFileMappingTable.h b/Engine/Source/Programs/UnrealBuildAccelerator/Detours/Private/UbaDetoursFileMappingTable.h
--- a/Engine/Source/Programs/UnrealBuildAccelerator/Detours/Private/UbaDetoursFileMappingTable.h
+++ b/Engine/Source/Programs/UnrealBuildAccelerator/Detours/Private/UbaDetoursFileMappingTable.h
@@ -57,7 +57,9 @@ namespace uba
 		UnorderedMap<const void*, int> m_memLookup;
 	};
 
-	void Rpc_CreateFileW(const tchar* fileName, const StringKey& fileNameKey, u8 access, tchar* outNewName, u64 newNameCapacity, u64& outSize, u32& outCloseId, bool lock);
+	enum : u8 { AccessFlag_Read = 1, AccessFlag_Write = 2 };
+
+	void Rpc_CreateFileW(const tchar* fileName, const StringKey& fileNameKey, u8 accessFlags, tchar* outNewName, u64 newNameCapacity, u64& outSize, u32& outCloseId, bool lock);
 	u32  Rpc_UpdateDirectory(const StringKey& dirKey, const tchar* dirName, u64 dirNameLen, bool lockDirTable = true);
 	void Rpc_UpdateCloseHandle(const tchar* handleName, u32 closeId, bool deleteOnClose, const tchar* newName, u64 mappingHandle, u64 mappingWritten, bool success);
 	void Rpc_UpdateTables();
diff --git a/Engine/Source/Programs/UnrealBuildAccelerator/Detours/Private/Windows/UbaDetoursFunctionsKernelBase.inl b/Engine/Source/Programs/UnrealBuildAccelerator/Detours/Private/Windows/UbaDetoursFunctionsKernelBase.inl
--- a/Engine/Source/Programs/UnrealBuildAccelerator/Detours/Private/Windows/UbaDetoursFunctionsKernelBase.inl
+++ b/Engine/Source/Programs/UnrealBuildAccelerator/Detours/Private/Windows/UbaDetoursFunctionsKernelBase.inl
@@ -1333,7 +1333,7 @@ bool Shared_MoveFile(LPCWSTR lpExistingFileName, LPCWSTR lpNewFileName, DWORD dw
 		StringBuffer<> dest;
 		FixPath(dest, lpNewFileName);
 
-		if (IsOutputFile(dest.data, dest.count, GENERIC_WRITE))
+		if (IsOutputFile(dest.data, dest.count, true))
 		{
 			sourceInfo.deleted = true;
 			UBA_ASSERT(!sourceInfo.memoryFile->isLocalOnly);
diff --git a/Engine/Source/Programs/UnrealBuildAccelerator/Detours/Private/Windows/UbaDetoursFunctionsNtDll.inl b/Engine/Source/Programs/UnrealBuildAccelerator/Detours/Private/Windows/UbaDetoursFunctionsNtDll.inl
--- a/Engine/Source/Programs/UnrealBuildAccelerator/Detours/Private/Windows/UbaDetoursFunctionsNtDll.inl
+++ b/Engine/Source/Programs/UnrealBuildAccelerator/Detours/Private/Windows/UbaDetoursFunctionsNtDll.inl
@@ -1,5 +1,39 @@
 // Copyright Epic Games, Inc. All Rights Reserved.
 
+bool IsContentWrite(u32 desiredAccess, u32 createDisposition)
+{
+	if (desiredAccess & (FILE_WRITE_DATA | FILE_APPEND_DATA | GENERIC_WRITE))
+		return true;
+	if (createDisposition == FILE_CREATE || createDisposition == FILE_OVERWRITE || createDisposition == FILE_OVERWRITE_IF)
+		return true;
+	return false;
+}
+
+bool IsContentRead(u32 desiredAccess, u32 createDisposition)
+{
+	return (desiredAccess & (GENERIC_READ | FILE_READ_DATA)) != 0;
+}
+
+bool IsContentUse(u32 desiredAccess, u32 createDisposition)
+{
+	return IsContentRead(desiredAccess, createDisposition) || IsContentWrite(desiredAccess, createDisposition);
+}
+
+bool IsWrite(u32 desiredAccess, u32 createDisposition)
+{
+	return IsContentWrite(desiredAccess, createDisposition) || (desiredAccess & (FILE_WRITE_ATTRIBUTES | FILE_WRITE_EA)) != 0;
+}
+
+u8 GetFileAccessFlags(DWORD desiredAccess)
+{
+	u8 access = 0;
+	if (IsContentRead(desiredAccess, 0))
+		access |= AccessFlag_Read;
+	if (IsWrite(desiredAccess, 0))
+		access |= AccessFlag_Write;
+	return access;
+}
+
 const wchar_t* ToString(NTSTATUS s) { return NT_SUCCESS(s) ? L"Success" : L"Error"; }
 
 struct FILE_FS_DEVICE_INFORMATION
@@ -518,7 +552,7 @@ NTSTATUS NTAPI Local_NtCreateFile(bool IsCreateFunc, PHANDLE hFileHandle, ACCESS
 	ULONG ShareAccess, ULONG CreateDisposition, ULONG CreateOptions, PVOID EaBuffer, ULONG EaLength)
 {
 #if 0
-	if (DesiredAccess & (FILE_WRITE_DATA | FILE_APPEND_DATA) || (CreateDisposition & (FILE_CREATE | FILE_OVERWRITE)))
+	if (IsContentWrite(DesiredAccess, CreateDisposition))
 	{
 		StringBuffer<> b;
 		b.Append(ObjectAttributes->ObjectName->Buffer, ObjectAttributes->ObjectName->Length / 2);
@@ -645,17 +679,13 @@ NTSTATUS NTAPI Shared_NtCreateFile(bool IsCreateFunc, PHANDLE hFileHandle, ACCES
 
 	u32 dirTableOffset = ~u32(0);
 
+	//UBA_ASSERT(CreateDisposition != FILE_SUPERSEDE);
 	bool isDeleteOnClose = (CreateOptions & FILE_DELETE_ON_CLOSE) != 0; // clang is using CreateFile with DeleteOnClose to delete files after build errors
-	bool failIfNotExists = (CreateDisposition & (FILE_CREATE | FILE_OVERWRITE | FILE_OVERWRITE_IF)) == FILE_OVERWRITE;// || CreateDisposition == FILE_OPEN;
-
-	DWORD dwDesiredAccess = DesiredAccess & (GENERIC_WRITE | GENERIC_READ | GENERIC_EXECUTE);
-	if (DesiredAccess & FILE_READ_DATA)
-		dwDesiredAccess |= GENERIC_READ;
-	if (DesiredAccess & (FILE_WRITE_DATA | FILE_APPEND_DATA) || (CreateDisposition & (FILE_CREATE | FILE_OVERWRITE)))
-		dwDesiredAccess |= GENERIC_WRITE;
+	bool TODO_DELETE_failIfNotExists = false;//(CreateDisposition & (FILE_CREATE | FILE_OVERWRITE | FILE_OVERWRITE_IF)) == FILE_OVERWRITE;// || CreateDisposition == FILE_OPEN;
 
-	bool isWrite = (dwDesiredAccess & GENERIC_WRITE) != 0;
-	bool keepInMemory = (KeepInMemory(fileName.data, fileName.count) && dwDesiredAccess) || IsOutputFile(fileName.data, fileName.count, dwDesiredAccess, isDeleteOnClose);
+	bool useContent = IsContentUse(DesiredAccess, CreateDisposition);
+	bool isWrite = IsWrite(DesiredAccess, CreateDisposition);
+	bool keepInMemory = (KeepInMemory(fileName.data, fileName.count) && useContent) || IsOutputFile(fileName.data, fileName.count, isWrite, isDeleteOnClose);
 
 #if UBA_DEBUG_LOG_ENABLED
 	const wchar_t* isWriteStr = isWrite ? L" WRITE" : L""; (void)isWriteStr;
@@ -700,7 +730,7 @@ NTSTATUS NTAPI Shared_NtCreateFile(bool IsCreateFunc, PHANDLE hFileHandle, ACCES
 			// This is an optimization where we populate directory table and use that to figure out if file exists or not..
 			// .. in msvc's case it doesn't matter much because these tables are already up to date when msvc use CreateFile.
 			// .. clang otoh is using CreateFile with tooons of different paths trying to open files.. in remote worker case this becomes super expensive
-			if ((!isWrite || failIfNotExists) && !isSystemOrTempFile) // We need to skip SystemTemp.. lots of stuff going on there.
+			if ((!isWrite || TODO_DELETE_failIfNotExists) && !isSystemOrTempFile) // We need to skip SystemTemp.. lots of stuff going on there.
 			{
 				CHECK_PATH(fileNameLower.data);
 				fileNameKey = ToStringKey(fileNameLower);
@@ -747,11 +777,11 @@ NTSTATUS NTAPI Shared_NtCreateFile(bool IsCreateFunc, PHANDLE hFileHandle, ACCES
 
 				bool isWriteAttributes = (DesiredAccess & FILE_WRITE_ATTRIBUTES) != 0;
 
-				if (allowEarlyOut && dwDesiredAccess == 0 && !isWriteAttributes)
+				if (allowEarlyOut && !useContent && !isWriteAttributes)
 				{
 					auto dh = new DetouredHandle(HandleType_File);
 					dh->fileObject = new FileObject();
-					dh->fileObject->desiredAccess = dwDesiredAccess;
+					dh->fileObject->desiredAccess = DesiredAccess;
 					dh->dirTableOffset = dirTableOffset;
 
 					FileInfo* tempFileInfo = new FileInfo();
@@ -873,7 +903,7 @@ NTSTATUS NTAPI Shared_NtCreateFile(bool IsCreateFunc, PHANDLE hFileHandle, ACCES
 		info.name = info.originalName;
 		if (!keepInMemory && !isSystemOrTempFile)
 		{
-			u8 access = GetFileAccessFlags(dwDesiredAccess);
+			u8 access = GetFileAccessFlags(DesiredAccess);
 			wchar_t newFileName[512];
 			Rpc_CreateFileW(lpFileName, fileNameKey, access, newFileName, sizeof_array(newFileName), size, closeId, false);
 			info.name = g_memoryBlock.Strdup(newFileName);
@@ -882,7 +912,7 @@ NTSTATUS NTAPI Shared_NtCreateFile(bool IsCreateFunc, PHANDLE hFileHandle, ACCES
 
 		info.size = size;
 		info.fileNameKey = fileNameKey;
-		info.lastDesiredAccess = dwDesiredAccess;
+		info.lastDesiredAccess = DesiredAccess;
 	}
 	else
 	{
@@ -890,25 +920,27 @@ NTSTATUS NTAPI Shared_NtCreateFile(bool IsCreateFunc, PHANDLE hFileHandle, ACCES
 
 		if (!info.originalName)
 			info.originalName = g_memoryBlock.Strdup(fileName.data);
-		if (isWrite) //(info.lastDesiredAccess != dwDesiredAccess)
+		if (isWrite)
 		{
+			bool lastWasWrite = IsContentWrite(info.lastDesiredAccess, 0);
 			UBA_ASSERT(!info.isFileMap);
-			bool shouldReport = !(info.lastDesiredAccess & GENERIC_WRITE) || info.deleted;
+			bool shouldReport = !lastWasWrite || info.deleted;
 			shouldReport = shouldReport && !keepInMemory;
 			if (shouldReport)
 			{
 				u64 size = InvalidValue;
 				info.deleted = false;
 				wchar_t newFileName[1024];
-				u8 access = GetFileAccessFlags(dwDesiredAccess);
+				u8 access = GetFileAccessFlags(DesiredAccess);
 				Rpc_CreateFileW(lpFileName, fileNameKey, access, newFileName, sizeof_array(newFileName), size, closeId, false);
 				info.name = g_memoryBlock.Strdup(newFileName);
 				//info.size = size; // TODO: Should this be set?
 				lpFileName = info.name;
 			}
-			if (dwDesiredAccess == 0 || info.lastDesiredAccess == 0)
+			bool lastUseContent = IsContentUse(info.lastDesiredAccess, 0);
+			if (!useContent || !lastUseContent)
 				lpFileName = info.name;
-			info.lastDesiredAccess |= dwDesiredAccess;
+			info.lastDesiredAccess |= DesiredAccess;
 		}
 		else if (info.deleted)
 		{
@@ -931,7 +963,7 @@ NTSTATUS NTAPI Shared_NtCreateFile(bool IsCreateFunc, PHANDLE hFileHandle, ACCES
 
 	auto TrackFileInput = [&]()
 		{
-			if (!keepInMemory && dwDesiredAccess && (dwDesiredAccess & GENERIC_WRITE) == 0)
+			if (!keepInMemory && useContent && !isWrite)
 			{
 				if (!info.tracked)
 				{
@@ -952,7 +984,7 @@ NTSTATUS NTAPI Shared_NtCreateFile(bool IsCreateFunc, PHANDLE hFileHandle, ACCES
 		UBA_ASSERT(!lpFileName[2]);
 
 		bool isDir = lpFileName[1] == 'd';
-		if (isDir && dwDesiredAccess != 0)
+		if (isDir && useContent)
 			return STATUS_FILE_IS_A_DIRECTORY;
 
 		MemoryFile& mf = g_emptyMemoryFile;
@@ -961,7 +993,7 @@ NTSTATUS NTAPI Shared_NtCreateFile(bool IsCreateFunc, PHANDLE hFileHandle, ACCES
 		DetouredHandle* dh = new DetouredHandle(HandleType_File);
 		dh->dirTableOffset = dirTableOffset;
 		dh->fileObject = new FileObject();
-		dh->fileObject->desiredAccess = dwDesiredAccess;
+		dh->fileObject->desiredAccess = DesiredAccess;
 		dh->fileObject->closeId = closeId;
 		dh->fileObject->fileInfo = &info;
 		UBA_ASSERT(!isDeleteOnClose);
@@ -994,7 +1026,7 @@ NTSTATUS NTAPI Shared_NtCreateFile(bool IsCreateFunc, PHANDLE hFileHandle, ACCES
 		UBA_ASSERT(info.size != InvalidValue);
 		dh->dirTableOffset = dirTableOffset;
 		dh->fileObject = new FileObject();
-		dh->fileObject->desiredAccess = dwDesiredAccess;
+		dh->fileObject->desiredAccess = DesiredAccess;
 		dh->fileObject->closeId = closeId;
 		dh->fileObject->fileInfo = &info;
 		UBA_ASSERT(!isDeleteOnClose);
@@ -1060,7 +1092,7 @@ NTSTATUS NTAPI Shared_NtCreateFile(bool IsCreateFunc, PHANDLE hFileHandle, ACCES
 					return STATUS_OBJECT_NAME_NOT_FOUND;
 				}
 
-				bool isLocal = !IsOutputFile(fileName.data, fileName.count, dwDesiredAccess, isDeleteOnClose);
+				bool isLocal = !IsOutputFile(fileName.data, fileName.count, isWrite, isDeleteOnClose);
 				//UBA_ASSERTF(CreateDisposition != FILE_OPEN || Contains(fileName.data, L"vctip_"), TC("Unsupported disposition %u for file %s"), CreateDisposition, fileName.data);
 				info.memoryFile = new MemoryFile(isLocal, FileTypeMaxSize(fileName, isSystemOrTempFile));
 			}
@@ -1080,7 +1112,7 @@ NTSTATUS NTAPI Shared_NtCreateFile(bool IsCreateFunc, PHANDLE hFileHandle, ACCES
 		DetouredHandle* dh = new DetouredHandle(HandleType_File);
 		dh->fileObject = fileObject;
 		dh->dirTableOffset = dirTableOffset;
-		dh->fileObject->desiredAccess = dwDesiredAccess;
+		dh->fileObject->desiredAccess = DesiredAccess;
 		dh->fileObject->closeId = closeId;
 		dh->fileObject->fileInfo = &info;
 		dh->fileObject->deleteOnClose = isDeleteOnClose;
@@ -1098,7 +1130,7 @@ NTSTATUS NTAPI Shared_NtCreateFile(bool IsCreateFunc, PHANDLE hFileHandle, ACCES
 		DetouredHandle* dh = new DetouredHandle(HandleType_File);
 		dh->fileObject = fileObject;
 		dh->dirTableOffset = dirTableOffset;
-		dh->fileObject->desiredAccess = dwDesiredAccess;
+		dh->fileObject->desiredAccess = DesiredAccess;
 		dh->fileObject->closeId = closeId;
 		dh->fileObject->fileInfo = &info;
 		dh->fileObject->deleteOnClose = isDeleteOnClose;
@@ -1155,7 +1187,7 @@ NTSTATUS NTAPI Shared_NtCreateFile(bool IsCreateFunc, PHANDLE hFileHandle, ACCES
 	dh->trueHandle = *hFileHandle;
 	dh->dirTableOffset = dirTableOffset;
 	dh->fileObject = new FileObject();
-	dh->fileObject->desiredAccess = dwDesiredAccess;
+	dh->fileObject->desiredAccess = DesiredAccess;
 	dh->fileObject->closeId = closeId;
 	dh->fileObject->fileInfo = &info;
 	dh->fileObject->deleteOnClose = isDeleteOnClose;
@@ -1280,7 +1312,7 @@ NTSTATUS NTAPI Detoured_NtClose(HANDLE handle)
 	wchar_t temp[512];
 	if (fi.memoryFile)
 	{
-		if ((fo->desiredAccess & GENERIC_WRITE))
+		if (IsWrite(fo->desiredAccess, 0))
 		{
 			// TODO: There are race conditions in this code. There could be other file handles accessing the same piece of memory (allthough unlikely)
 			u64 alignedWritten = AlignUp(fi.memoryFile->writtenSize, 64 * 1024);
@@ -1311,7 +1343,7 @@ NTSTATUS NTAPI Detoured_NtClose(HANDLE handle)
 		mappingWritten = fi.memoryFile->writtenSize;
 
 		u32 orginalNameLen = TStrlen(fi.originalName);
-		if (IsOutputFile(fi.originalName, orginalNameLen, fo->desiredAccess, fo->deleteOnClose) && !g_rules->IsThrowAway(fi.originalName, orginalNameLen))
+		if (IsOutputFile(fi.originalName, orginalNameLen, IsWrite(fo->desiredAccess, 0), fo->deleteOnClose) && !g_rules->IsThrowAway(fi.originalName, orginalNameLen))
 		{
 			// Need to report this file to host so it can be tracked in directory table
 			if (!fi.memoryFile->isReported)
@@ -1350,7 +1382,7 @@ NTSTATUS NTAPI Detoured_NtClose(HANDLE handle)
 			}
 
 		}
-		else if (NeedsSharedMemory(fi.originalName) && (fo->desiredAccess & GENERIC_WRITE))
+		else if (NeedsSharedMemory(fi.originalName) && IsWrite(fo->desiredAccess, 0))
 		{
 			StringBuffer<> fixedName;
 			FixPath(fixedName, path);
diff --git a/Engine/Source/Programs/UnrealBuildAccelerator/Detours/Private/Windows/UbaDetoursFunctionsWin.cpp b/Engine/Source/Programs/UnrealBuildAccelerator/Detours/Private/Windows/UbaDetoursFunctionsWin.cpp
--- a/Engine/Source/Programs/UnrealBuildAccelerator/Detours/Private/Windows/UbaDetoursFunctionsWin.cpp
+++ b/Engine/Source/Programs/UnrealBuildAccelerator/Detours/Private/Windows/UbaDetoursFunctionsWin.cpp
@@ -364,7 +364,7 @@ const wchar_t* HandleToName(HANDLE handle)
 
 bool NeedsSharedMemory(const wchar_t* file) { return g_allowKeepFilesInMemory && g_rules->NeedsSharedMemory(file); }
 u64 FileTypeMaxSize(const StringBufferBase& file, bool isSystemOrTempFile) { return g_rules->FileTypeMaxSize(file, isSystemOrTempFile); }
-bool IsOutputFile(LPCWSTR fileName, u64 fileNameLen, DWORD desiredAccess, bool isDeleteOnClose = false) { return ((desiredAccess & GENERIC_WRITE) || isDeleteOnClose) && g_allowKeepFilesInMemory && g_rules->IsOutputFile(fileName, fileNameLen); }
+bool IsOutputFile(LPCWSTR fileName, u64 fileNameLen, bool isWrite, bool isDeleteOnClose = false) { return (isWrite || isDeleteOnClose) && g_allowKeepFilesInMemory && g_rules->IsOutputFile(fileName, fileNameLen); }
 
 
 bool EnsureMapped(DetouredHandle& handle, DWORD dwFileOffsetHigh = 0, DWORD dwFileOffsetLow = 0, SIZE_T numberOfBytesToMap = 0, void* baseAddress = nullptr)
@@ -408,17 +408,6 @@ bool EnsureMapped(DetouredHandle& handle, DWORD dwFileOffsetHigh = 0, DWORD dwFi
 	return true;
 }
 
-enum : u8 { AccessFlag_Read = 1, AccessFlag_Write = 2 };
-u8 GetFileAccessFlags(DWORD dwDesiredAccess)
-{
-	u8 access = 0;
-	if (dwDesiredAccess & GENERIC_READ)
-		access |= AccessFlag_Read;
-	if (dwDesiredAccess & GENERIC_WRITE)
-		access |= AccessFlag_Write;
-	return access;
-}
-
 ReaderWriterLock g_longPathNameCacheLock;
 GrowingUnorderedMap<const wchar_t*, const wchar_t*, HashString, EqualString> g_longPathNameCache(&g_memoryBlock);
 
