[UBA]
* Fixed a bug with getting caskeytable in cache client
* Added more logging (disabled behind define)
* Added asserts for reading out of bounds in BinaryReader
* Added comp/raw bytes fetched from cache server per cache hit'

diff --git a/Engine/Source/Programs/UnrealBuildAccelerator/Common/Private/UbaCacheClient.cpp b/Engine/Source/Programs/UnrealBuildAccelerator/Common/Private/UbaCacheClient.cpp
--- a/Engine/Source/Programs/UnrealBuildAccelerator/Common/Private/UbaCacheClient.cpp
+++ b/Engine/Source/Programs/UnrealBuildAccelerator/Common/Private/UbaCacheClient.cpp
@@ -11,7 +11,7 @@
 #include "UbaStorageUtils.h"
 
 #define UBA_LOG_WRITE_CACHE_INFO 0 // 0 = Disabled, 1 = Normal, 2 = Detailed
-#define UBA_LOG_FETCH_CACHE_INFO 0
+#define UBA_LOG_FETCH_CACHE_INFO 0 // 0 = Disabled, 1 = Misses, 2 = Both misses and hits
 
 namespace uba
 {
@@ -41,6 +41,8 @@ namespace uba
 
 		ReaderWriterLock casKeyTableNetworkLock;
 		u32 casKeyTableSizeSent = 0;
+
+		Atomic<u32> availableCasKeyTableSize;
 	};
 
 	CacheClient::CacheClient(const CacheClientCreateInfo& info)
@@ -50,6 +52,10 @@ namespace uba
 	,	m_session(info.session)
 	{
 		m_reportMissReason = info.reportMissReason;
+		#if UBA_LOG_FETCH_CACHE_INFO
+		m_reportMissReason = true;
+		#endif
+
 		m_client.RegisterOnConnected([this]()
 			{
 				StackBinaryWriter<1024> writer;
@@ -316,7 +322,11 @@ namespace uba
 		u32 entryCount = reader.ReadU16();
 
 		#if UBA_LOG_FETCH_CACHE_INFO
-		auto mg = MakeGuard([&]() { m_logger.Info(TC("FETCHCACHE %s: %s -> %u %s (%u)"), success ? TC("SUCC") : TC("FAIL"), info.description, bucketId, CasKeyString(cmdKey).str, entryCount); });
+		auto mg = MakeGuard([&]()
+			{
+				if (!success || UBA_LOG_FETCH_CACHE_INFO == 2)
+					m_logger.Info(TC("FETCHCACHE %s: %s -> %u %s (%u)"), success ? TC("SUCC") : TC("FAIL"), info.description, bucketId, CasKeyString(cmdKey).str, entryCount);
+			});
 		#endif
 
 		struct MissInfo { TString path; u32 entryIndex; CasKey cache; CasKey local; };
@@ -341,9 +351,8 @@ namespace uba
 					auto insres = offsetIsMatch.try_emplace(casKeyOffset);
 					if (insres.second)
 					{
-						if (casKeyOffset >= bucket.serverCasKeyTable.GetSize())
-							if (!FetchCasTable(bucket, cacheStats, casKeyOffset + sizeof(CasKey)))
-								return false;
+						if (!FetchCasTable(bucket, cacheStats, casKeyOffset))
+							return false;
 
 						if (!GetLocalPathAndCasKey(bucket, rootPaths, path, cacheCasKey, bucket.serverCasKeyTable, bucket.serverPathTable, casKeyOffset))
 							return false;
@@ -394,9 +403,8 @@ namespace uba
 			while (reader.GetPositionData() != outputEnd)
 			{
 				u32 casKeyOffset = u32(reader.Read7BitEncoded());
-				if (casKeyOffset >= bucket.serverCasKeyTable.GetSize())
-					if (!FetchCasTable(bucket, cacheStats, casKeyOffset + sizeof(CasKey)))
-						return false;
+				if (!FetchCasTable(bucket, cacheStats, casKeyOffset))
+					return false;
 
 				TimerScope ts(cacheStats.fetchOutput);
 
@@ -456,10 +464,10 @@ namespace uba
 						return false;
 					if (!destFile.Write(localBlock.memory, localBlock.writtenSize))
 						return false;
-					if (!destFile.Close(&fetcher.m_lastWritten))
+					if (!destFile.Close(&fetcher.lastWritten))
 						return false;
 
-					fetcher.m_size = fileSize;
+					fetcher.sizeOnDisk = fileSize;
 					casKey = CalculateCasKey(localBlock.memory, localBlock.writtenSize, false, nullptr);
 				}
 				else
@@ -468,7 +476,11 @@ namespace uba
 					if (!fetcher.RetrieveFile(m_logger, m_client, casKey, path.data, destinationIsCompressed))
 						return false;
 				}
-				if (!m_storage.FakeCopy(casKey, path.data, fetcher.m_size, fetcher.m_lastWritten, false))
+
+				cacheStats.fetchBytesRaw += fetcher.sizeOnDisk;
+				cacheStats.fetchBytesComp += fetcher.bytesReceived;
+
+				if (!m_storage.FakeCopy(casKey, path.data, fetcher.sizeOnDisk, fetcher.lastWritten, false))
 					return false;
 				if (!m_session.RegisterNewFile(path.data))
 					return false;
@@ -680,37 +692,54 @@ namespace uba
 		return true;
 	}
 
-	bool CacheClient::FetchCasTable(Bucket& bucket, CacheStats& stats, u32 requiredCasTableSize)
+	bool CacheClient::FetchCasTable(Bucket& bucket, CacheStats& stats, u32 requiredCasTableOffset)
 	{
-		TimerScope ts2(stats.fetchCasTable);
-
-		SCOPED_WRITE_LOCK(bucket.casKeyTableNetworkLock, lock); // Use one lock over both queries
+		auto hasEnoughData = [&bucket, requiredCasTableOffset](u32 tableSize)
+			{
+				u32 neededSize = requiredCasTableOffset + 4;
+				if (neededSize > tableSize)
+					return false;
+				BinaryReader r(bucket.serverCasKeyTable.GetMemory(), requiredCasTableOffset, tableSize);
+				u8 bytesNeeded = Get7BitEncodedCount(r.Read7BitEncoded());
+				neededSize = requiredCasTableOffset + bytesNeeded + sizeof(CasKey);
+				return neededSize <= tableSize;
+			};
 
-		if (requiredCasTableSize <= bucket.serverCasKeyTable.GetSize())
+		if (hasEnoughData(bucket.availableCasKeyTableSize))
 			return true;
 
+		TimerScope ts2(stats.fetchCasTable);
+
 		StackBinaryReader<SendMaxSize> reader;
+
+		SCOPED_WRITE_LOCK(bucket.casKeyTableNetworkLock, lock); // Use one lock over both queries
 		{
-			u32 targetSize = requiredCasTableSize;
-			while (bucket.serverCasKeyTable.GetSize() < targetSize)
+			bool messageSent = false;
+			while (true)
 			{
+				u32 tableSize = bucket.serverCasKeyTable.GetSize();
+				if (hasEnoughData(tableSize))
+				{
+					if (!messageSent)
+						return true;
+					break;
+				}
+
 				StackBinaryWriter<16> writer;
 				NetworkMessage msg(m_client, CacheServiceId, CacheMessageType_FetchCasTable, writer);
 				writer.Write7BitEncoded(MakeId(bucket.id));
-				writer.WriteU32(bucket.serverCasKeyTable.GetSize());
+				writer.WriteU32(tableSize);
 
 				reader.Reset();
 				if (!msg.Send(reader))
 					return false;
-				u32 size = reader.ReadU32();
-				if (targetSize == ~0u)
-					targetSize = size;
-
+				reader.ReadU32();
+				messageSent = true;
 				bucket.serverCasKeyTable.ReadMem(reader, false);
 			}
 		}
 		{
-			u32 targetSize = ~0u; // For now, read all because we don't know how much we need (it would require parsing all path offsets in caskey table
+			u32 targetSize = ~0u; // For now, read all because we don't know how much we need (it would require parsing all path offsets in caskey table)
 			while (bucket.serverPathTable.GetSize() < targetSize)
 			{
 				StackBinaryWriter<16> writer;
@@ -728,6 +757,8 @@ namespace uba
 				bucket.serverPathTable.ReadMem(reader, false);
 			}
 		}
+
+		bucket.availableCasKeyTableSize = bucket.serverCasKeyTable.GetSize();
 		return true;
 	}
 
diff --git a/Engine/Source/Programs/UnrealBuildAccelerator/Common/Private/UbaStats.cpp b/Engine/Source/Programs/UnrealBuildAccelerator/Common/Private/UbaStats.cpp
--- a/Engine/Source/Programs/UnrealBuildAccelerator/Common/Private/UbaStats.cpp
+++ b/Engine/Source/Programs/UnrealBuildAccelerator/Common/Private/UbaStats.cpp
@@ -277,11 +277,12 @@ namespace uba
 		#define UBA_CACHE_STAT(type, var, ver) LogStat(logger, #var, var, frequency);
 		UBA_CACHE_STATS
 		#undef UBA_CACHE_STAT
+		logger.Info(TC("   Bytes   Raw/Comp %9s %9s"), BytesToText(fetchBytesRaw).str, BytesToText(fetchBytesComp).str);
 	}
 
 	bool CacheStats::IsEmpty()
 	{
-		#define UBA_CACHE_STAT(type, var, ver) if (var.count) return false;
+		#define UBA_CACHE_STAT(type, var, ver) if (var != type()) return false;
 		UBA_CACHE_STATS
 		#undef UBA_CACHE_STAT
 		return true;
diff --git a/Engine/Source/Programs/UnrealBuildAccelerator/Common/Private/UbaStorageUtils.cpp b/Engine/Source/Programs/UnrealBuildAccelerator/Common/Private/UbaStorageUtils.cpp
--- a/Engine/Source/Programs/UnrealBuildAccelerator/Common/Private/UbaStorageUtils.cpp
+++ b/Engine/Source/Programs/UnrealBuildAccelerator/Common/Private/UbaStorageUtils.cpp
@@ -409,7 +409,6 @@ namespace uba
 
 		u64 fileSize = 0;
 		u64 actualSize = 0;
-		u64 sizeOnDisk = 0;
 
 		u8* readBuffer = nullptr;
 		u8* readPosition = nullptr;
@@ -457,7 +456,8 @@ namespace uba
 				actualSize = *(u64*)readBuffer;
 		}
 
-		sizeOnDisk = IsCompressed(casKey) ? fileSize : actualSize;
+		bytesReceived = fileSize;
+		sizeOnDisk = destinationIsCompressed ? fileSize : actualSize;
 
 		FileAccessor destinationFile(logger, destination);
 		if (!destinationMem)
@@ -630,15 +630,9 @@ namespace uba
 		}
 
 		if (!destinationMem)
-			if (!destinationFile.Close(&m_lastWritten))
+			if (!destinationFile.Close(&lastWritten))
 				return false;
 
-		m_size = actualSize;
-
-		//stats.recvCasBytesRaw += actualSize;
-		//stats.recvCasBytesComp += fileSize;
-		//out.size = actualSize;
-
 		return true;
 	}
 
diff --git a/Engine/Source/Programs/UnrealBuildAccelerator/Common/Public/UbaCacheClient.h b/Engine/Source/Programs/UnrealBuildAccelerator/Common/Public/UbaCacheClient.h
--- a/Engine/Source/Programs/UnrealBuildAccelerator/Common/Public/UbaCacheClient.h
+++ b/Engine/Source/Programs/UnrealBuildAccelerator/Common/Public/UbaCacheClient.h
@@ -49,7 +49,7 @@ namespace uba
 		bool SendPathTable(Bucket& bucket, u32 requiredPathTableSize);
 		bool SendCasTable(Bucket& bucket, u32 requiredCasTableSize);
 		bool SendCacheEntry(Bucket& bucket, const RootPaths& rootPaths, const CasKey& cmdKey, const Map<u32, u32>& inputsStringToCasKey, const Map<u32, u32>& outputsStringToCasKey);
-		bool FetchCasTable(Bucket& bucket, CacheStats& stats, u32 requiredCasTableSize);
+		bool FetchCasTable(Bucket& bucket, CacheStats& stats, u32 requiredCasTableOffset);
 
 		CasKey GetCmdKey(const RootPaths& rootPaths, const ProcessStartInfo& info);
 		bool ShouldNormalize(const StringBufferBase& path);
diff --git a/Engine/Source/Programs/UnrealBuildAccelerator/Common/Public/UbaStats.h b/Engine/Source/Programs/UnrealBuildAccelerator/Common/Public/UbaStats.h
--- a/Engine/Source/Programs/UnrealBuildAccelerator/Common/Public/UbaStats.h
+++ b/Engine/Source/Programs/UnrealBuildAccelerator/Common/Public/UbaStats.h
@@ -123,6 +123,8 @@ namespace uba
 		UBA_CACHE_STAT(Timer, fetchCasTable, 0) \
 		UBA_CACHE_STAT(Timer, testEntries, 0) \
 		UBA_CACHE_STAT(Timer, fetchOutput, 0) \
+		UBA_CACHE_STAT(AtomicU64, fetchBytesRaw, 26) \
+		UBA_CACHE_STAT(AtomicU64, fetchBytesComp, 26) \
 
     struct CacheStats
 	{
diff --git a/Engine/Source/Programs/UnrealBuildAccelerator/Common/Public/UbaStorageUtils.h b/Engine/Source/Programs/UnrealBuildAccelerator/Common/Public/UbaStorageUtils.h
--- a/Engine/Source/Programs/UnrealBuildAccelerator/Common/Public/UbaStorageUtils.h
+++ b/Engine/Source/Programs/UnrealBuildAccelerator/Common/Public/UbaStorageUtils.h
@@ -42,7 +42,8 @@ namespace uba
 		ReaderWriterLock m_retrieveOneBatchAtTheTimeLock;
 		bool m_errorOnFail = true;
 
-		u64 m_lastWritten = 0;
-		u64 m_size = 0;
+		u64 lastWritten = 0;
+		u64 sizeOnDisk = 0;
+		u64 bytesReceived = 0;
 	};
 }
\ No newline at end of file
diff --git a/Engine/Source/Programs/UnrealBuildAccelerator/Common/Public/UbaTrace.h b/Engine/Source/Programs/UnrealBuildAccelerator/Common/Public/UbaTrace.h
--- a/Engine/Source/Programs/UnrealBuildAccelerator/Common/Public/UbaTrace.h
+++ b/Engine/Source/Programs/UnrealBuildAccelerator/Common/Public/UbaTrace.h
@@ -59,7 +59,7 @@ namespace uba
 	using Color = u32;
 	inline Color toColor(u8 r, u8 g, u8 b) { return (r << 16) + (g << 8) + b; }
 
-	static constexpr u32 TraceVersion = 25;
+	static constexpr u32 TraceVersion = 26;
 	static constexpr u32 TraceReadCompatibilityVersion = 6;
 
 	class Trace : public WorkTracker
diff --git a/Engine/Source/Programs/UnrealBuildAccelerator/Core/Public/UbaBinaryReaderWriter.h b/Engine/Source/Programs/UnrealBuildAccelerator/Core/Public/UbaBinaryReaderWriter.h
--- a/Engine/Source/Programs/UnrealBuildAccelerator/Core/Public/UbaBinaryReaderWriter.h
+++ b/Engine/Source/Programs/UnrealBuildAccelerator/Core/Public/UbaBinaryReaderWriter.h
@@ -98,6 +98,7 @@ namespace uba
 	// Implementations
 
 	#define UBA_ASSERT_WRITE(size) UBA_ASSERTF(m_pos + size <= m_end, TC("BinaryWriter overflow. Written: %llu, Capacity: %llu, Trying to write: %llu"), u64(m_pos - m_begin), u64(m_end - m_begin), u64(size))
+	#define UBA_ASSERT_READ(size) UBA_ASSERTF(m_pos + size <= m_end, TC("BinaryReader overflow. Read: %llu, Size: %llu, Trying to read: %llu"), u64(m_pos - m_begin), u64(m_end - m_begin), u64(size))
 
 	void BinaryWriter::WriteByte(u8 value)
 	{
@@ -286,6 +287,7 @@ namespace uba
 
 	void BinaryReader::ReadBytes(void* data, u64 size)
 	{
+		UBA_ASSERT_READ(size);
 		memcpy(data, m_pos, size);
 		m_pos += size;
 	}
@@ -293,11 +295,13 @@ namespace uba
 
 	u8 BinaryReader::ReadByte()
 	{
+		UBA_ASSERT_READ(sizeof(u8));
 		return *m_pos++;
 	}
 
 	u16 BinaryReader::ReadU16()
 	{
+		UBA_ASSERT_READ(sizeof(u16));
 		u16 value = *(u16*)m_pos;
 		m_pos += sizeof(u16);
 		return value;
@@ -305,6 +309,7 @@ namespace uba
 
 	u32 BinaryReader::ReadU32()
 	{
+		UBA_ASSERT_READ(sizeof(u32));
 		u32 value = *(u32*)m_pos;
 		m_pos += sizeof(u32);
 		return value;
@@ -312,6 +317,7 @@ namespace uba
 
 	u64 BinaryReader::ReadU64()
 	{
+		UBA_ASSERT_READ(sizeof(u64));
 		u64 value = *(u64*)m_pos;
 		m_pos += sizeof(u64);
 		return value;
@@ -331,18 +337,21 @@ namespace uba
 		u64 left = charLen;
 		while (left--)
 		{
+			UBA_ASSERT_READ(1);
 			u8 a = *m_pos++;
 			if (a <= 127)
 			{
 				*it++ = a;
 				continue;
 			}
+			UBA_ASSERT_READ(1);
 			u8 b = *m_pos++;
 			if (a >= 192 && a <= 223)
 			{
 				*it++ = (a-192)*64 + (b-128);
 				continue;
 			}
+			UBA_ASSERT_READ(1);
 			u8 c = *m_pos++;
 			if (a >= 224 && a <= 239)
 			{
@@ -389,6 +398,7 @@ namespace uba
 
 	Guid BinaryReader::ReadGuid()
 	{
+		UBA_ASSERT_READ(sizeof(Guid));
 		u64 g[2];
 		g[0] = *(u64*)m_pos;
 		g[1] = ((u64*)m_pos)[1];
@@ -398,6 +408,7 @@ namespace uba
 
 	StringKey BinaryReader::ReadStringKey()
 	{
+		UBA_ASSERT_READ(sizeof(StringKey));
 		StringKey k;
 		k.a = *(u64*)m_pos;
 		k.b = ((u64*)m_pos)[1];
@@ -407,6 +418,7 @@ namespace uba
 
 	CasKey BinaryReader::ReadCasKey()
 	{
+		UBA_ASSERT_READ(sizeof(CasKey));
 		CasKey k;
 		k.a = *(u64*)m_pos;
 		k.b = ((u64*)m_pos)[1];
@@ -422,6 +434,7 @@ namespace uba
 		bool hasMoreBytes;
 		do
 		{
+			UBA_ASSERT_READ(1);
 			u8 value = *m_pos++;
 			hasMoreBytes = value & 0x80;
 			result |= u64(value & 0x7f) << (byteIndex * 7);
