[UBA]
* Minor fixes needed to be able to cache linux builds (nopch)

diff --git a/Engine/Source/Programs/UnrealBuildAccelerator/Common/Private/UbaCompactTables.cpp b/Engine/Source/Programs/UnrealBuildAccelerator/Common/Private/UbaCompactTables.cpp
--- a/Engine/Source/Programs/UnrealBuildAccelerator/Common/Private/UbaCompactTables.cpp
+++ b/Engine/Source/Programs/UnrealBuildAccelerator/Common/Private/UbaCompactTables.cpp
@@ -274,7 +274,7 @@ namespace uba
 		}
 		else if (outRequiredCasTableSize)
 		{
-			BinaryReader reader(m_mem.memory, insres.first->second, 1000);
+			BinaryReader reader(m_mem.memory, insres.first->second, ~0u);
 			reader.Read7BitEncoded();
 			*outRequiredCasTableSize = Max(*outRequiredCasTableSize, u32(reader.GetPosition() + sizeof(CasKey)));
 		}
@@ -283,7 +283,7 @@ namespace uba
 
 	void CompactCasKeyTable::GetKey(CasKey& outKey, u64 offset) const
 	{
-		BinaryReader reader(m_mem.memory, offset, 1000);
+		BinaryReader reader(m_mem.memory, offset, ~0u);
 		reader.Read7BitEncoded();
 		outKey = reader.ReadCasKey();
 	}
@@ -297,7 +297,7 @@ namespace uba
 		}
 		#endif
 
-		BinaryReader reader(m_mem.memory, offset, 1000);
+		BinaryReader reader(m_mem.memory, offset, ~0u);
 		u32 stringOffset = (u32)reader.Read7BitEncoded();
 		outKey = reader.ReadCasKey();
 		pathTable.GetString(outPath, stringOffset);
diff --git a/Engine/Source/Programs/UnrealBuildAccelerator/Common/Private/UbaRootPaths.cpp b/Engine/Source/Programs/UnrealBuildAccelerator/Common/Private/UbaRootPaths.cpp
--- a/Engine/Source/Programs/UnrealBuildAccelerator/Common/Private/UbaRootPaths.cpp
+++ b/Engine/Source/Programs/UnrealBuildAccelerator/Common/Private/UbaRootPaths.cpp
@@ -39,9 +39,6 @@ namespace uba
 			if (CaseInsensitiveFs)
 				ToLower(root.path.data());
 
-			if (root.path[root.path.size()-1] != PathSeparator)
-				return logger.Error(TC("Root path '%s' must end with separator"), rp);
-
 			root.includeInKey = includeInKey;
 
 			m_longestRoot = Max(u32(root.path.size()), m_longestRoot);
diff --git a/Engine/Source/Programs/UnrealBuildTool/Executors/UnrealBuildAccelerator/UBAExecutor.cs b/Engine/Source/Programs/UnrealBuildTool/Executors/UnrealBuildAccelerator/UBAExecutor.cs
--- a/Engine/Source/Programs/UnrealBuildTool/Executors/UnrealBuildAccelerator/UBAExecutor.cs
+++ b/Engine/Source/Programs/UnrealBuildTool/Executors/UnrealBuildAccelerator/UBAExecutor.cs
@@ -563,13 +563,13 @@ namespace UnrealBuildTool
 			IRootPaths rootPaths = IRootPaths.Create(_ubaLogger!);
 			foreach (DirectoryItem root in action.RootPaths)
 			{
-				rootPaths.RegisterRoot(root.FullName + Path.DirectorySeparatorChar, true);
+				rootPaths.RegisterRoot(root.FullName, true);
 			}
 
 			DirectoryReference? autoSdkDir;
 			if (UEBuildPlatformSDK.TryGetHostPlatformAutoSDKDir(out autoSdkDir))
 			{
-				rootPaths.RegisterRoot(autoSdkDir.FullName + Path.DirectorySeparatorChar, true);
+				rootPaths.RegisterRoot(autoSdkDir.FullName, true);
 			}
 
 			rootPaths.RegisterSystemRoots();
