[UBA]
* Fixed so .d files (clang dep files) are normalized in cache

diff --git a/Engine/Source/Programs/UnrealBuildAccelerator/Common/Private/UbaCacheClient.cpp b/Engine/Source/Programs/UnrealBuildAccelerator/Common/Private/UbaCacheClient.cpp
--- a/Engine/Source/Programs/UnrealBuildAccelerator/Common/Private/UbaCacheClient.cpp
+++ b/Engine/Source/Programs/UnrealBuildAccelerator/Common/Private/UbaCacheClient.cpp
@@ -815,7 +815,9 @@ namespace uba
 
 	bool CacheClient::ShouldNormalize(const StringBufferBase& path)
 	{
-		if (path.EndsWith(TC(".dep.json"))) // Contains absolute paths
+		if (path.EndsWith(TC(".dep.json"))) // Contains absolute paths (dep file for msvc)
+			return true;
+		if (path.EndsWith(TC(".d"))) // Contains absolute paths (dep file for clang)
 			return true;
 		if (path.EndsWith(TC(".tlh"))) // Contains absolute path in a comment
 			return true;
