[uba]
* Fixed bug where normalized file which didn't need to be normalized crashed the client

diff --git a/Engine/Source/Programs/UnrealBuildAccelerator/Common/Private/UbaCacheClient.cpp b/Engine/Source/Programs/UnrealBuildAccelerator/Common/Private/UbaCacheClient.cpp
--- a/Engine/Source/Programs/UnrealBuildAccelerator/Common/Private/UbaCacheClient.cpp
+++ b/Engine/Source/Programs/UnrealBuildAccelerator/Common/Private/UbaCacheClient.cpp
@@ -659,7 +659,7 @@ namespace uba
 					if (!rootPaths.NormalizeString<char>(m_logger, (const char*)file.GetData(), file.GetSize(), handleString, path.data))
 						return false;
 
-					if (!rootOffsets.empty())
+					if (rootOffsetsSize)
 					{
 						u8* mem = (u8*)block.Allocate(rootOffsetsSize, 1, TC(""));
 						rootOffsetsStart = u32(mem - block.memory);
@@ -667,6 +667,9 @@ namespace uba
 						for (u32 rootOffset : rootOffsets)
 							writer.Write7BitEncoded(rootOffset);
 					}
+					else
+						rootOffsetsStart = u32(block.writtenSize);
+
 
 					auto& s = m_storage;
 					FileSender sender { m_logger, m_client, s.m_bufferSlots, s.Stats(), m_sendOneAtTheTimeLock, s.m_casCompressor, s.m_casCompressionLevel };
diff --git a/Engine/Source/Programs/UnrealBuildAccelerator/Common/Public/UbaRootPaths.h b/Engine/Source/Programs/UnrealBuildAccelerator/Common/Public/UbaRootPaths.h
--- a/Engine/Source/Programs/UnrealBuildAccelerator/Common/Public/UbaRootPaths.h
+++ b/Engine/Source/Programs/UnrealBuildAccelerator/Common/Public/UbaRootPaths.h
@@ -52,7 +52,7 @@ namespace uba
 		{
 			auto absPathChars = searchPos;
 			CharType lastChar = 0;
-			while (absPathChars < strEnd && !(lastChar == ':' && *absPathChars == '\\'))
+			while (absPathChars < strEnd && !(lastChar == ':' && *absPathChars == PathSeparator))
 			{
 				lastChar = *absPathChars;
 				++absPathChars;
