[UBA]
* Added code to point headers included in pch to be redirected to empty files in order to handle loading of pch produced by other machines. This feature is disabled for now since it is not ready for prime time

diff --git a/Engine/Source/Programs/UnrealBuildAccelerator/Detours/Private/Windows/UbaDetoursFunctionsWin.cpp b/Engine/Source/Programs/UnrealBuildAccelerator/Detours/Private/Windows/UbaDetoursFunctionsWin.cpp
--- a/Engine/Source/Programs/UnrealBuildAccelerator/Detours/Private/Windows/UbaDetoursFunctionsWin.cpp
+++ b/Engine/Source/Programs/UnrealBuildAccelerator/Detours/Private/Windows/UbaDetoursFunctionsWin.cpp
@@ -43,6 +43,7 @@ DETOURED_FUNCTIONS_MEMORY
 #include "UbaApplicationRules.h"
 #include "UbaDetoursShared.h"
 #include "UbaDetoursObjFilesPreloader.h"
+#include "UbaDetoursPrepopulatePchFiles.h"
 
 #include "Shlwapi.h"
 #include <detours/detours.h>
@@ -267,7 +268,8 @@ void MemoryFile::EnsureCommited(DetouredHandle& handle, u64 size)
 			if (writtenSize == 0 && !isReported)
 			{
 				u64 newReserve = AlignUp(size, g_pageSize);
-				Rpc_WriteLogf(L"TODO: RE-RESERVING MemoryFile. Initial reserve: %llu, New reserve: %llu. Please fix application rules", reserveSize, newReserve);
+				if (reserveSize)
+					Rpc_WriteLogf(L"TODO: RE-RESERVING MemoryFile. Initial reserve: %llu, New reserve: %llu. Please fix application rules", reserveSize, newReserve);
 				Unreserve();
 				Reserve(newReserve);
 				shouldRemap = false;
@@ -870,6 +872,8 @@ void Init(const DetoursPayload& payload, u64 startTime)
 
 	if (payload.storeObjFilesCompressed && g_rulesIndex == 2) 	// Rules index 2 is link.exe.. compressed obj files is experimental so solution is a bit hacky
 		g_objFilesPreloader.Start(cmdLine);
+	else if (g_rulesIndex == 1 || g_rulesIndex == 7 || g_rulesIndex == 11 || g_rulesIndex == 14)
+		PrepopulatePchIncludedFiles(cmdLine, g_rulesIndex);
 
 	g_stats.attach.time += GetTime() - startTime;
 	g_stats.attach.count = 1;
diff --git a/Engine/Source/Programs/UnrealBuildAccelerator/Detours/Private/Windows/UbaDetoursPrepopulatePchFiles.h b/Engine/Source/Programs/UnrealBuildAccelerator/Detours/Private/Windows/UbaDetoursPrepopulatePchFiles.h
new file mode 100644
--- /dev/null
+++ b/Engine/Source/Programs/UnrealBuildAccelerator/Detours/Private/Windows/UbaDetoursPrepopulatePchFiles.h
@@ -0,0 +1,217 @@
+// Copyright Epic Games, Inc. All Rights Reserved.
+
+#pragma once
+
+#include "UbaDetoursUtilsWin.h"
+#include "UbaEvent.h"
+#include <oodle2.h>
+
+namespace uba
+{
+	extern MemoryFile& g_emptyMemoryFile;
+
+	template<typename LineFunc>
+	bool ReadLines(const tchar* file, const LineFunc& lineFunc)
+	{
+		HANDLE handle = CreateFileW(file, GENERIC_READ, FILE_SHARE_READ, NULL, OPEN_ALWAYS, 0, NULL);
+		if (handle == INVALID_HANDLE_VALUE)
+			return false;
+		auto hg = MakeGuard([&]() { CloseHandle(handle); });
+		LARGE_INTEGER fileSize;
+		if (!GetFileSizeEx(handle, &fileSize))
+			return false;
+
+		char buffer[512];
+		u64 left = fileSize.QuadPart;
+		StringBuffer<2048> line;
+		while (left)
+		{
+			u64 toRead = Min(left, u64(sizeof(buffer)));
+			left -= toRead;
+
+			DWORD read = 0;
+			if (!ReadFile(handle, buffer, u32(toRead), &read, NULL))
+				return false;
+
+			u64 start = 0;
+			for (u64 i=0;i!=read;++i)
+			{
+				if (buffer[i] != '\n')
+					continue;
+				u64 end = i;
+				if (i > 0 && buffer[i-1] == '\r')
+					--end;
+				line.Append(buffer + start, u32(end - start));
+				if (line.count)
+					if (!lineFunc(line))
+						return false;
+				line.Clear();
+				start = i + 1;
+			}
+			if (read && buffer[read - 1] == '\r')
+				--read;
+			line.Append(buffer + start, u32(read - start));
+		}
+		if (line.count)
+			if (!lineFunc(line))
+				return false;
+		return true;
+	}
+
+	void PrepopulatePchIncludedFiles(const tchar* commandline, u32 rulesIndex)
+	{
+		if (true)
+			return;
+
+		auto startOfRsp = TStrchr(commandline, '@');
+		if (!startOfRsp)
+			return;
+		bool hasQuotes = false;
+		if (*++startOfRsp == '\"')
+		{
+			++startOfRsp;
+			hasQuotes = true;
+		}
+
+		auto endOfRsp = TStrchr(startOfRsp, hasQuotes ? '\"' : ' ');
+		if (!endOfRsp)
+			endOfRsp = startOfRsp + TStrlen(startOfRsp);
+
+		StringBuffer<512> rsp;
+		rsp.Append(startOfRsp, endOfRsp - startOfRsp);
+
+		Vector<TString> includes;
+
+		if (rulesIndex == 1)
+		{
+			bool usesPch = false;
+			StringBuffer<512> pch;
+			ReadLines(rsp.data, [&](StringBufferBase& line)
+				{
+					if (line.StartsWith(L"/Yu"))
+						usesPch = true;
+					if (!line.Contains(L"/Fp\""))
+						return true;
+					auto pchStart = line.data + 4;
+					auto pchEnd = TStrchr(pchStart, '\"');
+					pch.Append(pchStart, u32(pchEnd - pchStart));
+					return usesPch != true;
+				});
+
+			if (!pch.count || !usesPch)
+				return;
+
+			StringBuffer<512> dep;
+			dep.Append(pch.data, pch.count - 4).Append(L".dep.json");
+
+			bool inIncludes = false;
+			ReadLines(dep.data, [&](StringBufferBase& line)
+				{
+					if (!inIncludes)
+					{
+						inIncludes = line.Contains(L"\"Includes\":");
+						return true;
+					}
+					auto includeStart = TStrchr(line.data, '\"');
+					if (!includeStart)
+						return false;
+					++includeStart;
+					tchar* readIt = includeStart;
+					tchar* writeIt = readIt;
+					tchar lastChar = 0;
+					while (*readIt != '\"')
+					{
+						if (lastChar != '\\' || *readIt != '\\')
+						{
+							*writeIt = *readIt;
+							++writeIt;
+						}
+						lastChar = *readIt;
+						++readIt;
+					}
+
+					includes.push_back(TString(includeStart, writeIt));
+					return true;
+				});
+		}
+		else if (rulesIndex == 7 || rulesIndex == 11 || rulesIndex == 14)
+		{
+			StringBuffer<512> pch;
+			ReadLines(rsp.data, [&](StringBufferBase& line)
+				{
+					if (!line.StartsWith(L"-include-pch"))
+						return true;
+					auto pchStart = line.data + 14;
+					auto pchEnd = TStrchr(pchStart, '\"');
+					pch.Append(pchStart, u32(pchEnd - pchStart));
+					return false;
+				});
+
+			if (!pch.count)
+				return;
+
+			StringBuffer<512> dep;
+			dep.Append(pch.data, pch.count - 4).Append(L".d");
+
+			bool inIncludes = false;
+			ReadLines(dep.data, [&](StringBufferBase& line)
+				{
+					if (!inIncludes)
+					{
+						inIncludes = true;
+						return true;
+					}
+
+					tchar* it = line.data + 2;
+					bool isFirst = true;
+					while (true)
+					{
+						tchar* end = TStrchr(it, ' ');
+						if (!end && !isFirst)
+							break;
+						if (end)
+							*end = 0;
+						isFirst = false;
+
+						StringBuffer<> fullPath;
+						FixPath(fullPath, it);
+						includes.push_back(TString(fullPath.data, fullPath.count));
+						if (!end)
+							break;
+						it = end + 1;
+					}
+					return true;
+				});
+		}
+
+		const tchar* canBeIncludedMultipleTimes[] = 
+		{
+			L".h.inl",
+			L"UnrealNames.inl",
+			L"ShowFlagsValues.inl",
+			L"AnimMTStats.h",
+			L"bits\\byteswap", // For linux
+		};
+
+		SCOPED_WRITE_LOCK(g_mappedFileTable.m_lookupLock, _);
+		for (auto& include : includes)
+		{
+			bool skip = false;
+			for (auto i : canBeIncludedMultipleTimes)
+				skip |= Contains(include.data(), i);
+			if (skip)
+				continue;
+
+			StringBuffer<> t(include.data());
+			auto fileNameKey = ToStringKeyLower(t);
+			auto insres = g_mappedFileTable.m_lookup.try_emplace(fileNameKey);
+			FileInfo& info = insres.first->second;
+			info.originalName = g_memoryBlock.Strdup(include.data());
+			info.name = info.originalName;
+			info.size = 0;
+			info.fileNameKey = fileNameKey;
+			info.lastDesiredAccess = GENERIC_READ;
+			info.memoryFile = &g_emptyMemoryFile;
+		}
+	}
+}
\ No newline at end of file
