[UBA]
* Added more logging information when farm gets a include not found error. This is temporary and will be deleted once we have figured out why this can happen now and then.

diff --git a/Engine/Source/Programs/UnrealBuildAccelerator/Common/Private/UbaProcess.cpp b/Engine/Source/Programs/UnrealBuildAccelerator/Common/Private/UbaProcess.cpp
--- a/Engine/Source/Programs/UnrealBuildAccelerator/Common/Private/UbaProcess.cpp
+++ b/Engine/Source/Programs/UnrealBuildAccelerator/Common/Private/UbaProcess.cpp
@@ -979,7 +979,12 @@ namespace uba
 		bool printInSession = reader.ReadBool();
 		bool isError = reader.ReadBool();
 		TString line = reader.ReadString();
-		LogLine(printInSession, std::move(line), isError ? LogEntryType_Error : LogEntryType_Info);
+		LogEntryType entryType = isError ? LogEntryType_Error : LogEntryType_Info;
+
+		if (!m_session.LogLine(*this, line.c_str(), entryType))
+			return false;
+
+		LogLine(printInSession, std::move(line), entryType);
 		return true;
 	}
 
diff --git a/Engine/Source/Programs/UnrealBuildAccelerator/Common/Private/UbaSession.cpp b/Engine/Source/Programs/UnrealBuildAccelerator/Common/Private/UbaSession.cpp
--- a/Engine/Source/Programs/UnrealBuildAccelerator/Common/Private/UbaSession.cpp
+++ b/Engine/Source/Programs/UnrealBuildAccelerator/Common/Private/UbaSession.cpp
@@ -2260,6 +2260,11 @@ namespace uba
 		return true;
 	}
 
+	bool Session::LogLine(ProcessImpl& process, const tchar* line, LogEntryType logType)
+	{
+		return true;
+	}
+
 	void Session::PrintSessionStats(Logger& logger)
 	{
 		u64 mappingBufferSize;
diff --git a/Engine/Source/Programs/UnrealBuildAccelerator/Common/Private/UbaSessionClient.cpp b/Engine/Source/Programs/UnrealBuildAccelerator/Common/Private/UbaSessionClient.cpp
--- a/Engine/Source/Programs/UnrealBuildAccelerator/Common/Private/UbaSessionClient.cpp
+++ b/Engine/Source/Programs/UnrealBuildAccelerator/Common/Private/UbaSessionClient.cpp
@@ -1911,6 +1911,94 @@ namespace uba
 		return SendUpdateDirectoryTable(reader);
 	}
 
+	bool SessionClient::LogLine(ProcessImpl& process, const tchar* line, LogEntryType logType)
+	{
+		// TODO: Remove this once we have figured out a bug that seems to exist for remote execution
+		#if PLATFORM_WINDOWS
+		auto rules = process.m_startInfo.rules;
+		if (!rules)
+			return true;
+
+		const tchar* errorPos = nullptr;
+		if (rules->index == 1)
+		{
+			if (!Contains(line, TC("C1083"), false, &errorPos))
+				return true;
+		}
+		else
+		{
+			if (!Contains(line, TC("' file not found")))
+				return true;
+			if (!Contains(line, TC("fatal error: '"), false, &errorPos))
+				return true;
+		}
+
+		const tchar* fileBegin = TStrchr(errorPos, '\'');
+		if (!fileBegin)
+			return true;
+		++fileBegin;
+		const tchar* fileEnd = TStrchr(fileBegin, '\'');
+		if (!fileEnd)
+			return true;
+
+		StringBuffer<> searchString;
+		searchString.Append(fileBegin, fileEnd - fileBegin).Replace('/', PathSeparator);
+
+		MemoryBlock memoryBlock;
+		DirectoryTable dirTable(&memoryBlock);
+		{
+			SCOPED_WRITE_LOCK(m_directoryTable.m_memoryLock, lock2);
+			m_directoryTable.m_memorySize = m_directoryTableMemPos;
+			dirTable.Init(m_directoryTable.m_memory, 0, m_directoryTable.m_memorySize);
+		}
+
+		u32 foundCount = 0;
+		dirTable.TraverseAllFilesNoLock([&](const DirectoryTable::EntryInformation& info, const StringBufferBase& path)
+			{
+				if (!path.Contains(searchString.data))
+					return;
+				++foundCount;
+				StringBuffer<> logStr;
+				logStr.Appendf(TC("File %s found in directory table while searching for matches for %s (size %llu attr %u)"), path.data, searchString.data, info.size, info.attributes);
+				process.LogLine(false, logStr.data, logType);
+
+				StringKey fileNameKey = ToStringKey(path);
+				SCOPED_READ_LOCK(m_fileMappingTableLookupLock, mlock);
+				auto findIt = m_fileMappingTableLookup.find(fileNameKey);
+				if (findIt != m_fileMappingTableLookup.end())
+				{
+					auto& entry = findIt->second;
+					SCOPED_READ_LOCK(entry.lock, entryCs);
+					logStr.Clear().Appendf(TC("File %s found in mapping table table."), path.data);
+					if (entry.handled && entry.mapping.IsValid())
+					{
+						StringBuffer<128> mappingName;
+						Storage::GetMappingString(mappingName, entry.mapping, entry.mappingOffset);
+						logStr.Appendf(TC(" Mapping name: %s"), mappingName.data);
+					}
+				}
+				else
+					logStr.Clear().Appendf(TC("File %s not found in mapping table table."), path.data);
+				process.LogLine(false, logStr.data, logType);
+
+				CasKey key;
+				if (GetCasKeyForFile(key, process.m_id, path, fileNameKey))
+					logStr.Clear().Appendf(TC("File %s caskey is %s"), path.data, CasKeyString(key).str);
+				else
+					logStr.Clear().Appendf(TC("File %s caskey not found"), path.data);
+				process.LogLine(false, logStr.data, logType);
+			});
+		if (!foundCount)
+		{
+			StringBuffer<> logStr;
+			logStr.Appendf(TC("No matching entry found in directory table while searching for matches for %s. DirTable size: %u"), searchString.data, GetDirectoryTableSize());
+			process.LogLine(false, logStr.data, logType);
+		}
+
+		#endif
+		return true;
+	}
+
 	void SessionClient::TraceSessionUpdate()
 	{
 		float cpuLoad = 0.0f;//UpdateCpuLoad();
diff --git a/Engine/Source/Programs/UnrealBuildAccelerator/Common/Public/UbaSession.h b/Engine/Source/Programs/UnrealBuildAccelerator/Common/Public/UbaSession.h
--- a/Engine/Source/Programs/UnrealBuildAccelerator/Common/Public/UbaSession.h
+++ b/Engine/Source/Programs/UnrealBuildAccelerator/Common/Public/UbaSession.h
@@ -134,6 +134,7 @@ namespace uba
 		virtual void FileEntryAdded(StringKey fileNameKey, u64 lastWritten, u64 size);
 		virtual bool FlushWrittenFiles(ProcessImpl& process);
 		virtual bool UpdateEnvironment(ProcessImpl& process, const tchar* reason, bool resetStats);
+		virtual bool LogLine(ProcessImpl& process, const tchar* line, LogEntryType logType);
 
 		static constexpr CasKey CasKeyIsDirectory = { ~u64(0), ~u64(0), ~u32(0) };
 
diff --git a/Engine/Source/Programs/UnrealBuildAccelerator/Common/Public/UbaSessionClient.h b/Engine/Source/Programs/UnrealBuildAccelerator/Common/Public/UbaSessionClient.h
--- a/Engine/Source/Programs/UnrealBuildAccelerator/Common/Public/UbaSessionClient.h
+++ b/Engine/Source/Programs/UnrealBuildAccelerator/Common/Public/UbaSessionClient.h
@@ -64,6 +64,7 @@ namespace uba
 		virtual bool CustomMessage(Process& process, BinaryReader& reader, BinaryWriter& writer) override;
 		virtual bool FlushWrittenFiles(ProcessImpl& process) override;
 		virtual bool UpdateEnvironment(ProcessImpl& process, const tchar* reason, bool resetStats) override;
+		virtual bool LogLine(ProcessImpl& process, const tchar* line, LogEntryType logType) override;
 		virtual void TraceSessionUpdate() override;
 
 		struct InternalProcessStartInfo;
diff --git a/Engine/Source/Programs/UnrealBuildAccelerator/Core/Private/UbaStringBuffer.cpp b/Engine/Source/Programs/UnrealBuildAccelerator/Core/Private/UbaStringBuffer.cpp
--- a/Engine/Source/Programs/UnrealBuildAccelerator/Core/Private/UbaStringBuffer.cpp
+++ b/Engine/Source/Programs/UnrealBuildAccelerator/Core/Private/UbaStringBuffer.cpp
@@ -45,7 +45,12 @@ namespace uba
 	bool Contains(const tchar* str, const tchar* sub, bool ignoreCase, const tchar** pos)
 	{
 		if (!ignoreCase)
-			return TStrstr(str, sub) != 0;
+		{
+			auto res = TStrstr(str, sub);
+			if (pos)
+				*pos = res;
+			return res != nullptr;
+		}
 		for (const tchar* a = str; *a; ++a)
 		{
 			bool contains = true;
diff --git a/Engine/Source/Programs/UnrealBuildAccelerator/Core/Public/UbaDirectoryTable.h b/Engine/Source/Programs/UnrealBuildAccelerator/Core/Public/UbaDirectoryTable.h
--- a/Engine/Source/Programs/UnrealBuildAccelerator/Core/Public/UbaDirectoryTable.h
+++ b/Engine/Source/Programs/UnrealBuildAccelerator/Core/Public/UbaDirectoryTable.h
@@ -311,6 +311,42 @@ namespace uba
 			}
 		}
 
+#if PLATFORM_WINDOWS
+		template<typename Func>
+		void TraverseFilesRecursiveNoLock(const StringBufferBase& path, const Func& func)
+		{
+			auto findIt = m_lookup.find(ToStringKey(path));
+			if (findIt == m_lookup.end())
+				return;
+			StringKeyHasher hasher;
+			hasher.Update(path.data, path.count);
+			PopulateDirectory(hasher, findIt->second);
+			for (auto& fileKv : findIt->second.files)
+			{
+				DirectoryTable::EntryInformation info;
+				StringBuffer<> fileName(path);
+				fileName.Append(PathSeparator);
+				GetEntryInformation(info, fileKv.second, fileName.data + fileName.count, fileName.capacity - fileName.count);
+				fileName.count = TStrlen(fileName.data);
+				if (CaseInsensitiveFs)
+					fileName.MakeLower();
+				func(info, fileName);
+				TraverseFilesRecursiveNoLock(fileName, func);
+			}
+		}
+
+		template<typename Func>
+		void TraverseAllFilesNoLock(const Func& func)
+		{
+			for (tchar l='a';l!='z'; ++l)
+			{
+				StringBuffer<4> drive;
+				drive.Append(l).Append(':');
+				TraverseFilesRecursiveNoLock(drive, func);
+			}
+		}
+#endif
+
 		DirectoryTable(MemoryBlock* block) : m_memoryBlock(block), m_lookup(block) {}
 
 		MemoryBlock* m_memoryBlock;
