[UBA]
* Added more info to logging trying to catch bug where header files are not "found" remotely

diff --git a/Engine/Source/Programs/UnrealBuildAccelerator/Common/Private/UbaSessionClient.cpp b/Engine/Source/Programs/UnrealBuildAccelerator/Common/Private/UbaSessionClient.cpp
--- a/Engine/Source/Programs/UnrealBuildAccelerator/Common/Private/UbaSessionClient.cpp
+++ b/Engine/Source/Programs/UnrealBuildAccelerator/Common/Private/UbaSessionClient.cpp
@@ -1957,6 +1957,8 @@ namespace uba
 			{
 				if (!path.Contains(searchString.data))
 					return;
+				auto ToString = [](bool b) { return b ? TC("true") : TC("false"); };
+
 				++foundCount;
 				StringBuffer<> logStr;
 				logStr.Appendf(TC("File %s found in directory table while searching for matches for %s (size %llu attr %u)"), path.data, searchString.data, info.size, info.attributes);
@@ -1970,11 +1972,18 @@ namespace uba
 					auto& entry = findIt->second;
 					SCOPED_READ_LOCK(entry.lock, entryCs);
 					logStr.Clear().Appendf(TC("File %s found in mapping table table."), path.data);
-					if (entry.handled && entry.mapping.IsValid())
+					if (entry.handled)
 					{
 						StringBuffer<128> mappingName;
-						Storage::GetMappingString(mappingName, entry.mapping, entry.mappingOffset);
-						logStr.Appendf(TC(" Mapping name: %s"), mappingName.data);
+						if (entry.mapping.IsValid())
+							Storage::GetMappingString(mappingName, entry.mapping, entry.mappingOffset);
+						else
+							mappingName.Append(TC("Not valid"));
+						logStr.Appendf(TC(" Success: %s Size: %u IsDir: %s Mapping name: %s Mapping offset: %u"), ToString(entry.success), entry.size, ToString(entry.isDir), mappingName.data, entry.mappingOffset);
+					}
+					else
+					{
+						logStr.Appendf(TC(" Entry not handled"));
 					}
 				}
 				else
@@ -1983,7 +1992,37 @@ namespace uba
 
 				CasKey key;
 				if (GetCasKeyForFile(key, process.m_id, path, fileNameKey))
-					logStr.Clear().Appendf(TC("File %s caskey is %s"), path.data, CasKeyString(key).str);
+				{
+					logStr.Clear().Appendf(TC("File %s caskey is %s."), path.data, CasKeyString(key).str);
+
+					StringBuffer<512> casKeyFile;
+					if (m_storage.GetCasFileName(casKeyFile, key))
+					{
+						logStr.Appendf(TC(" CasKeyFile: %s"), casKeyFile.data);
+						u64 size = 0;
+						u32 attributes = 0;
+						bool exists = FileExists(m_logger, casKeyFile.data, &size, &attributes);
+						logStr.Appendf(TC(" Exists: %s"), ToString(exists));
+						if (exists)
+						{
+							logStr.Appendf(TC(" Size: %llu Attr: %u"), size, attributes);
+
+							FileHandle fileHandle = uba::CreateFileW(casKeyFile.data, GENERIC_READ, FILE_SHARE_READ, OPEN_EXISTING, DefaultAttributes());
+							if (fileHandle == InvalidFileHandle)
+							{
+								logStr.Appendf(TC(" Failed to open file %s (%s)"), casKeyFile.data, LastErrorToText().data);
+							}
+							else
+							{
+								logStr.Appendf(TC(" CreateFile for read successful"));
+								uba::CloseFile(casKeyFile.data, fileHandle);
+							}
+						}
+					}
+					else
+						logStr.Appendf(TC(" Failed to get cas filename for cas key"));
+
+				}
 				else
 					logStr.Clear().Appendf(TC("File %s caskey not found"), path.data);
 				process.LogLine(false, logStr.data, logType);
