[UBA]
* Added retry code for when failing to decompress cas file to memory on client.. it could be that cas file is corrupt so we delete the file and re-download it.
* Made CheckCasFiles parallel
* Fixed so CheckCasFiles check all files that was created after .isRunning was created to see that their content look valid
* Added time to log when correcting/validation cas store

diff --git a/Engine/Source/Programs/UnrealBuildAccelerator/Cli/Private/UbaClient.h b/Engine/Source/Programs/UnrealBuildAccelerator/Cli/Private/UbaClient.h
--- a/Engine/Source/Programs/UnrealBuildAccelerator/Cli/Private/UbaClient.h
+++ b/Engine/Source/Programs/UnrealBuildAccelerator/Cli/Private/UbaClient.h
@@ -42,6 +42,7 @@ namespace uba
 			storageClientInfo.startProxyUserData = this;
 			storageClient = new StorageClient(storageClientInfo);
 
+			storageClient->LoadCasTable(false);
 			SessionClientCreateInfo sessionClientInfo(*storageClient, *networkClient, info.logWriter);
 			sessionClientInfo.maxProcessCount = info.maxProcessorCount;
 			sessionClientInfo.rootDir = clientRootDir.data;
diff --git a/Engine/Source/Programs/UnrealBuildAccelerator/Common/Private/UbaSessionClient.cpp b/Engine/Source/Programs/UnrealBuildAccelerator/Common/Private/UbaSessionClient.cpp
--- a/Engine/Source/Programs/UnrealBuildAccelerator/Common/Private/UbaSessionClient.cpp
+++ b/Engine/Source/Programs/UnrealBuildAccelerator/Common/Private/UbaSessionClient.cpp
@@ -414,110 +414,138 @@ namespace uba
 			return true;
 		}
 
-		StringBuffer<> newName;
-		bool isDir = casKey == CasKeyIsDirectory;
-		u64 fileSize = InvalidValue;
-		CasKey newCasKey;
+		// Code for doing retry if failing to decompress casfile. We've seen cases of corrupt cas files on clients
+		bool shouldRetry = true;
+		FileMappingEntry* retryEntry = nullptr;
+		auto retryEntryGuard = MakeGuard([&]() { if (retryEntry) retryEntry->lock.LeaveWrite(); });
 
-		u32 memoryMapAlignment = 0;
-		if (m_allowMemoryMaps)
+		while (true)
 		{
-			memoryMapAlignment = GetMemoryMapAlignment(fileName.data, fileName.count);
-			if (!memoryMapAlignment && !m_useStorage)
-				memoryMapAlignment = 64 * 1024;
-		}
+			StringBuffer<> newName;
+			bool isDir = casKey == CasKeyIsDirectory;
+			u64 fileSize = InvalidValue;
+			CasKey newCasKey;
 
-		if (isDir)
-		{
-			newName.Append(TC("$d"));
-		}
-		else if (casKey != CasKeyZero)
-		{
-			if (m_useStorage || memoryMapAlignment == 0)
+			u32 memoryMapAlignment = 0;
+			if (m_allowMemoryMaps)
 			{
-				bool storeUncompressed = memoryMapAlignment == 0;
-				bool allowProxy = msg.process.m_startInfo.rules->AllowStorageProxy(fileName);
-				if (!RetrieveCasFile(newCasKey, fileSize, casKey, fileName.data, storeUncompressed, allowProxy))
-					return m_logger.Error(TC("Error retrieving cas entry %s (%s)"), CasKeyString(casKey).str, fileName.data);
+				memoryMapAlignment = GetMemoryMapAlignment(fileName.data, fileName.count);
+				if (!memoryMapAlignment && !m_useStorage)
+					memoryMapAlignment = 64 * 1024;
+			}
 
-				#if !UBA_USE_SPARSEFILE
-				if (!m_storage.GetCasFileName(newName, newCasKey))
-					return false;
-				#else
-				if (!memoryMapAlignment)
-					memoryMapAlignment = 4096;
-				MemoryMap map;
-				if (!CreateMemoryMapFromView(map, fileNameKey, fileName.data, newCasKey, memoryMapAlignment))
-					return false;
-				newName.Append(map.name);
-				fileSize = map.size;
-				#endif
+			if (isDir)
+			{
+				newName.Append(TC("$d"));
 			}
-			else
+			else if (casKey != CasKeyZero)
 			{
-				StorageStats& stats = m_storage.Stats();
-				TimerScope ts(stats.ensureCas);
-
-				SCOPED_WRITE_LOCK(m_fileMappingTableLookupLock, lookupLock);
-				auto insres = m_fileMappingTableLookup.try_emplace(fileNameKey);
-				FileMappingEntry& entry = insres.first->second;
-				lookupLock.Leave();
-
-				SCOPED_WRITE_LOCK(entry.lock, entryCs);
-				ts.Leave();
-
-				if (entry.handled)
+				if (m_useStorage || memoryMapAlignment == 0)
 				{
-					if (!entry.success)
+					bool storeUncompressed = memoryMapAlignment == 0;
+					bool allowProxy = msg.process.m_startInfo.rules->AllowStorageProxy(fileName);
+					if (!RetrieveCasFile(newCasKey, fileSize, casKey, fileName.data, storeUncompressed, allowProxy))
+						return m_logger.Error(TC("Error retrieving cas entry %s (%s)"), CasKeyString(casKey).str, fileName.data);
+
+					#if !UBA_USE_SPARSEFILE
+					if (!m_storage.GetCasFileName(newName, newCasKey))
 						return false;
+					#else
+					if (!memoryMapAlignment)
+						memoryMapAlignment = 4096;
+					MemoryMap map;
+					if (!CreateMemoryMapFromView(map, fileNameKey, fileName.data, newCasKey, memoryMapAlignment))
+						return false;
+					newName.Append(map.name);
+					fileSize = map.size;
+					#endif
 				}
 				else
 				{
-					TimerScope s(m_stats.storageRetrieve);
-					casKey = AsCompressed(casKey, false);
-					entry.handled = true;
-					Storage::RetrieveResult result;
-					bool allowProxy = msg.process.m_startInfo.rules->AllowStorageProxy(fileName);
-					if (!m_storage.RetrieveCasFile(result, casKey, fileName.data, &m_fileMappingBuffer, memoryMapAlignment, allowProxy))
-						return m_logger.Error(TC("Error retrieving cas entry %s (%s)"), CasKeyString(casKey).str, fileName.data);
-					entry.success = true;
-					entry.size = result.size;
-					entry.mapping = result.view.handle;
-					entry.mappingOffset = result.view.offset;
-				}
+					StorageStats& stats = m_storage.Stats();
+					TimerScope ts(stats.ensureCas);
 
-				fileSize = entry.size;
-				if (entry.mapping.IsValid())
-					Storage::GetMappingString(newName, entry.mapping, entry.mappingOffset);
-				else
-					newName.Append(entry.isDir ? TC("$d") : TC("$f"));
-			}
-		}
+					SCOPED_WRITE_LOCK(m_fileMappingTableLookupLock, lookupLock);
+					auto insres = m_fileMappingTableLookup.try_emplace(fileNameKey);
+					FileMappingEntry& entry = insres.first->second;
+					lookupLock.Leave();
 
-		UBA_ASSERTF(!newName.IsEmpty(), TC("No casfile available for %s using %s"), fileName.data, CasKeyString(casKey).str);
+					SCOPED_WRITE_LOCK(entry.lock, entryCs);
+					ts.Leave();
 
-		if (newName[0] != '^')
-		{
-			if (!isDir && memoryMapAlignment)
-			{
-				MemoryMap map;
-				if (!CreateMemoryMapFromFile(map, fileNameKey, newName.data, IsCompressed(newCasKey), memoryMapAlignment))
-					return false;
+					if (entry.handled)
+					{
+						if (!entry.success)
+							return false;
+					}
+					else
+					{
+						TimerScope s(m_stats.storageRetrieve);
+						casKey = AsCompressed(casKey, false);
+						entry.handled = true;
+						Storage::RetrieveResult result;
+						bool allowProxy = msg.process.m_startInfo.rules->AllowStorageProxy(fileName);
+						if (!m_storage.RetrieveCasFile(result, casKey, fileName.data, &m_fileMappingBuffer, memoryMapAlignment, allowProxy))
+							return m_logger.Error(TC("Error retrieving cas entry %s (%s)"), CasKeyString(casKey).str, fileName.data);
+						entry.success = true;
+						entry.size = result.size;
+						entry.mapping = result.view.handle;
+						entry.mappingOffset = result.view.offset;
+					}
 
-				fileSize = map.size;
-				newName.Clear().Append(map.name);
+					fileSize = entry.size;
+					if (entry.mapping.IsValid())
+						Storage::GetMappingString(newName, entry.mapping, entry.mappingOffset);
+					else
+						newName.Append(entry.isDir ? TC("$d") : TC("$f"));
+				}
 			}
-			else if (!IsRarelyRead(msg.process, fileName))
+
+			UBA_ASSERTF(!newName.IsEmpty(), TC("No casfile available for %s using %s"), fileName.data, CasKeyString(casKey).str);
+
+			if (newName[0] != '^')
 			{
-				AddFileMapping(fileNameKey, fileName.data, newName.data, fileSize);
+				if (!isDir && memoryMapAlignment)
+				{
+					if (retryEntry)
+						retryEntryGuard.Execute();
+
+					MemoryMap map;
+					if (!CreateMemoryMapFromFile(map, fileNameKey, newName.data, IsCompressed(newCasKey), memoryMapAlignment))
+					{
+						if (!shouldRetry)
+							return false;
+						shouldRetry = false;
+
+						// We need to take a lock around the file map entry since there might be another thread also wanting to map this
+						{
+							SCOPED_WRITE_LOCK(m_fileMappingTableLookupLock, lookupLock);
+							retryEntry = &m_fileMappingTableLookup.try_emplace(fileNameKey).first->second;
+							lookupLock.Leave();
+							retryEntry->lock.EnterWrite();
+							retryEntry->handled = false;
+						}
+
+						if (!m_storage.ReportBadCasFile(newCasKey))
+							return false;
+
+						continue;
+					}
+					fileSize = map.size;
+					newName.Clear().Append(map.name);
+				}
+				else if (!IsRarelyRead(msg.process, fileName))
+				{
+					AddFileMapping(fileNameKey, fileName.data, newName.data, fileSize);
+				}
 			}
-		}
 
-		out.directoryTableSize = GetDirectoryTableSize();
-		out.mappedFileTableSize = GetFileMappingSize();
-		out.fileName.Append(newName);
-		out.size = fileSize;
-		return true;
+			out.directoryTableSize = GetDirectoryTableSize();
+			out.mappedFileTableSize = GetFileMappingSize();
+			out.fileName.Append(newName);
+			out.size = fileSize;
+			return true;
+		}
 	}
 
 	bool SessionClient::SendFiles(ProcessImpl& process, Timer& sendFiles)
diff --git a/Engine/Source/Programs/UnrealBuildAccelerator/Common/Private/UbaStorage.cpp b/Engine/Source/Programs/UnrealBuildAccelerator/Common/Private/UbaStorage.cpp
--- a/Engine/Source/Programs/UnrealBuildAccelerator/Common/Private/UbaStorage.cpp
+++ b/Engine/Source/Programs/UnrealBuildAccelerator/Common/Private/UbaStorage.cpp
@@ -12,7 +12,7 @@
 
 namespace uba
 {
-	constexpr u32 CasTableVersion = IsWindows ? 32 : 34;
+	constexpr u32 CasTableVersion = IsWindows ? 33 : 35;
 	constexpr u32 MaxWorkItemsPerAction = 128; // Cap this to not starve other things
 
 #if UBA_USE_SPARSEFILE
@@ -817,7 +817,7 @@ namespace uba
 			{
 				outReturnValue = false;
 				m_logger.Error(TC("Failed to delete %s. Clean cas folder and restart"), casFile.data);
-				return true;
+				return false;
 			}
 		}
 		else
@@ -896,37 +896,60 @@ namespace uba
 		return true;
 	}
 
-	void StorageImpl::TraverseAllCasFiles(const tchar* dir, u32 recursion, const Function<void(const StringBufferBase& fullPath, const DirectoryEntry& e)>& func)
+	void StorageImpl::TraverseAllCasFiles(const tchar* dir, const Function<void(const StringBufferBase& fullPath, const DirectoryEntry& e)>& func, bool allowParallel)
 	{
+		Atomic<u32> workLeft;
+
+		auto traverseCasFileDir = [&func, &workLeft, this](const tchar* casFileDir)
+			{
+				TraverseDir(m_logger, casFileDir,
+					[&](const DirectoryEntry& e)
+					{
+						UBA_ASSERT(!IsDirectory(e.attributes));
+						StringBuffer<> fullPath(casFileDir);
+						fullPath.EnsureEndsWithSlash().Append(e.name);
+						func(fullPath, e);
+					});
+				--workLeft;
+			};
+
 		TraverseDir(m_logger, dir,
 			[&](const DirectoryEntry& e)
 			{
+				if (!IsDirectory(e.attributes))
+					return;
 				StringBuffer<> fullPath(dir);
 				fullPath.EnsureEndsWithSlash().Append(e.name);
-				if (IsDirectory(e.attributes))
+				++workLeft;
+
+				if (allowParallel && m_workManager)
 				{
-					TraverseAllCasFiles(fullPath.data, recursion + 1, func);
+					m_workManager->AddWork([&traverseCasFileDir, p = TString(fullPath.data)]() { traverseCasFileDir(p.data()); }, 1, TC("TraverseCasFiles"));
 				}
-				else if (recursion != 0)
+				else
 				{
-					func(fullPath, e);
+					traverseCasFileDir(fullPath.data);
 				}
 			});
+
+		while (workLeft)
+			m_workManager->DoWork();
 	}
 
 	void StorageImpl::TraverseAllCasFiles(const Function<void(const CasKey& key, u64 size)>& func)
 	{
 		StringBuffer<> casRoot;
 		casRoot.Append(m_rootDir.data, m_rootDir.count - 1);
-		TraverseAllCasFiles(casRoot.data, 0, [&](const StringBufferBase& fullPath, const DirectoryEntry& e)
+		TraverseAllCasFiles(casRoot.data, [&](const StringBufferBase& fullPath, const DirectoryEntry& e)
 			{
 				func(CasKeyFromString(e.name), e.size);
 			});
 	}
 
-	void StorageImpl::CheckAllCasFiles()
+	bool StorageImpl::CheckAllCasFiles(u64 checkContentOfFilesNewerThanTime)
 	{
 #if !UBA_USE_SPARSEFILE
+		u64 startTime = GetTime();
 		u64 before = m_casTotalBytes;
 		m_casTotalBytes = 0;
 		// Need to scan all files to see so there are no orphans in the folders
@@ -935,12 +958,71 @@ namespace uba
 
 		m_logger.Info(TC("Previous run was not gracefully shutdown. Reparsing cas directory %s to check for added/missing files"), casRoot.data);
 
-		TraverseAllCasFiles(casRoot.data, 0, [this](const StringBufferBase& fullPath, const DirectoryEntry& e)
+		Atomic<bool> success = true;
+		TraverseAllCasFiles(casRoot.data, [&](const StringBufferBase& fullPath, const DirectoryEntry& e)
 			{
 				CasKey casKey = CasKeyFromString(e.name);
+				bool deleteFile = false;
+				u64 size = e.size;
+
+				// Do quick check on new files so the have good content (ignore files with zero size, they are tested further down)
+				if (size && e.lastWritten >= checkContentOfFilesNewerThanTime)
+				{
+					FileAccessor fa(m_logger, fullPath.data);
+					if (fa.OpenMemoryRead())
+					{
+						if (IsCompressed(casKey))
+						{
+							BinaryReader reader(fa.GetData(), 0, fa.GetSize());
+							if (reader.GetLeft() < 12)
+							{
+								m_logger.Detail(TC("Corrupt cas. Is %llu, must be at least 12 bytes (%s)"), size, e.name);
+								deleteFile = true;
+							}
+							else
+							{
+								reader.ReadU64(); // Decompressed size
+								while (true)
+								{
+									if (reader.GetLeft() <= 8)
+									{
+										m_logger.Detail(TC("Corrupt cas. Missing beginning of block (%s)"), size, e.name);
+										deleteFile = true;
+										break;
+									}
+
+									u32 compressedBlockSize = reader.ReadU32();
+									reader.ReadU32();
+
+									if (!compressedBlockSize || compressedBlockSize > reader.GetLeft())
+									{
+										m_logger.Detail(TC("Corrupt cas. Bad block (%s)"), size, e.name);
+										deleteFile = true;
+										break;
+									}
+
+									reader.Skip(compressedBlockSize);
+
+									if (!reader.GetLeft())
+										break;
+								}
+							}
+						}
+						else
+						{
+							// TODO: Some simple validation of uncompressed file.. maybe even rehash and check? (since it is not many uncompressed files.. zero on server and a few on client)
+						}
+					}
+					else
+					{
+						deleteFile = true;
+					}
+				}
+
+
+				SCOPED_WRITE_LOCK(m_casLookupLock, lookupLock);
 				auto insres = m_casLookup.try_emplace(casKey);
 				CasEntry& entry = insres.first->second;
-				u64 size = e.size;
 				entry.verified = true;
 				entry.exists = true;
 
@@ -955,23 +1037,35 @@ namespace uba
 				{
 					UBA_ASSERT(entry.key == casKey);
 					// We should probably delete this one.. something is wrong
-					if (entry.size != 0 && entry.size != size)
+					if (entry.size != 0 && entry.size != size && !deleteFile)
 						m_logger.Detail(TC("Found cas entry which has a different size than what the table thought! Was %llu, is %llu (%s)"), entry.size, size, e.name);
 					entry.size = size;
 				}
 
-				if (!size)
+				if (!size && casKey != ToCasKey(CasKeyHasher(), IsCompressed(casKey)))
 				{
-					CasKeyHasher hasher;
-					if (casKey != ToCasKey(hasher, IsCompressed(casKey)))
-					{
-						m_logger.Detail(TC("Found file that has size 0 but does not have correct caskey (%s)"), e.name);
-						DeleteFileW(fullPath.data);
-						DetachEntry(entry);
-						m_casLookup.erase(insres.first);
-					}
+					m_logger.Detail(TC("Found file that has size 0 but does not have correct caskey (%s)"), e.name);
+					deleteFile = true;
 				}
-			});
+
+				if (!deleteFile)
+					return;
+
+				DetachEntry(entry);
+				m_casLookup.erase(insres.first);
+				m_casTotalBytes -= size;
+				lookupLock.Leave();
+
+				if (DeleteFileW(fullPath.data))
+					return;
+
+				m_logger.Error(TC("Failed to delete file (%s)"), fullPath.data);
+				success = false;
+
+			}, true);
+
+		if (!success)
+			return false;
 
 		u32 didNotExistCount = 0;
 		// All files we saw is tagged as "handled") so let's see if there are cas entries we didn't find
@@ -993,13 +1087,18 @@ namespace uba
 		if (didNotExistCount)
 			m_logger.Info(TC("Found %u cas entries that didn't have a file"), didNotExistCount);
 
+		u64 duration = GetTime() - startTime;
+
 		u64 after = m_casTotalBytes;
 		if (before != after)
-			m_logger.Info(TC("Corrected storage size from %s to %s"), BytesToText(before).str, BytesToText(after).str);
+			m_logger.Info(TC("Corrected storage size from %s to %s in %s"), BytesToText(before).str, BytesToText(after).str, TimeToText(duration).str);
+		else
+			m_logger.Info(TC("Validated storage (size %s) in %s"), BytesToText(after).str, TimeToText(duration).str);
 		m_casMaxBytes = m_casTotalBytes;
 #else
 		UBA_ASSERT(false); // not implemented
 #endif
+		return true;
 	}
 
 	void StorageImpl::HandleOverflow(Set<CasKey>* outDeletedFiles)
@@ -1310,8 +1409,18 @@ namespace uba
 		bool resave = false;
 		if (wasTerminated)
 		{
-			CheckAllCasFiles();
+			u64 fileTime = 0;
+			FileHandle fh = uba::CreateFileW(isRunningName, 0, 0x00000007, 0x00000003, FILE_FLAG_BACKUP_SEMANTICS);
+			auto fhg = MakeGuard([&]() { uba::CloseFile(isRunningName, fh); });
+			if (fh != InvalidFileHandle)
+				GetFileLastWriteTime(fileTime, fh);
+
+			if (!CheckAllCasFiles(fileTime))
+				return false;
 			resave = true;
+
+			if (fh != InvalidFileHandle)
+				SetFileLastWriteTime(fh, GetSystemTimeAsFileTime());
 		}
 
 		if (!m_manuallyHandleOverflow)
@@ -1513,7 +1622,7 @@ namespace uba
 		u32 errorCount = 0;
 		u64 newestWrittenError = 0;
 		ReaderWriterLock lock;
-		TraverseAllCasFiles(casRoot.data, 0, [&](const StringBufferBase& fullPath, const DirectoryEntry& e)
+		TraverseAllCasFiles(casRoot.data, [&](const StringBufferBase& fullPath, const DirectoryEntry& e)
 			{
 				++entryCount;
 				workManager.AddWork([&, filePath = TString(fullPath.data), name = TString(e.name), lastWritten = e.lastWritten]()
@@ -1971,6 +2080,12 @@ namespace uba
 		return true;
 	}
 
+	bool StorageImpl::ReportBadCasFile(const CasKey& casKey)
+	{
+		DropCasFile(casKey, true, TC("BadCasFile"));
+		return true;
+	}
+
 	bool StorageImpl::CalculateCasKey(CasKey& out, const tchar* fileName)
 	{
 		FileHandle fileHandle;
@@ -2099,7 +2214,7 @@ namespace uba
 						return m_logger.Error(TC("Failed to open file %s for read (%s)"), casFile.data, LastErrorToText().data);
 
 					if (!ReadFile(m_logger, casFile.data, readHandle, &decompressedSize, sizeof(u64)))
-						return m_logger.Error(TC("Failed to read first bytes from file %s (%s)"), casFile.data, LastErrorToText().data);
+						return m_logger.Error(TC("Failed to read first 8 bytes from compressed file %s (%s)"), casFile.data, LastErrorToText().data);
 				}
 
 				bool writeDirectlyToFile = casEntry->mappingHandle.IsValid() && false; // Experiment to try to fix bottlenecks on cloud
@@ -2464,6 +2579,8 @@ namespace uba
 			u8* slot = m_bufferSlots.Pop();
 			auto _ = MakeGuard([&]() { m_bufferSlots.Push(slot); });
 
+			u64 bytesRead = 8; // We know size has already been read
+
 			u8* readBuffer = slot;
 			u8* writePos = dest;
 			u64 left = decompressedSize;
@@ -2471,15 +2588,34 @@ namespace uba
 			{
 				u32 sizes[2];
 				if (!ReadFile(m_logger, fileName, fileHandle, sizes, sizeof(u32) * 2))
+				{
+					u64 compressedSize;
+					if (!uba::GetFileSizeEx(compressedSize, fileHandle))
+						return m_logger.Error(TC("GetFileSize failed for %s (%s)"), fileName, LastErrorToText().data);
+					if (bytesRead + 8 > compressedSize)
+						return m_logger.Error(TC("File %s corrupt. Tried to read 8 bytes. File is smaller than expected (Read: %llu, Size: %llu)"), fileName, bytesRead, compressedSize);
 					return false;
+				}
 				u32 compressedBlockSize = sizes[0];
 				u32 decompressedBlockSize = sizes[1];
 
+				bytesRead += 8;
+
 				if (!ReadFile(m_logger, fileName, fileHandle, readBuffer, compressedBlockSize))
+				{
+					u64 compressedSize;
+					if (!uba::GetFileSizeEx(compressedSize, fileHandle))
+						return m_logger.Error(TC("GetFileSize failed for %s (%s)"), fileName, LastErrorToText().data);
+					if (bytesRead + compressedBlockSize > compressedSize)
+						return m_logger.Error(TC("File %s corrupt. Compressed block size (%u) is larger than what is left of file (%llu)"), fileName, compressedBlockSize, compressedSize - bytesRead);
 					return false;
+				}
+				bytesRead += compressedBlockSize;
+
 				TimerScope ts(stats.decompressToMem);
 				OO_SINTa decompLen = OodleLZ_Decompress(readBuffer, (OO_SINTa)compressedBlockSize, writePos, (OO_SINTa)decompressedBlockSize);
-				UBA_ASSERT(decompLen == decompressedBlockSize); (void)decompLen;
+				if (decompLen != decompressedBlockSize)
+					return m_logger.Error(TC("Failed to decompress data from file %s at pos %llu"), fileName, decompressedSize - left);
 				writePos += decompressedBlockSize;
 				left -= decompressedBlockSize;
 			}
diff --git a/Engine/Source/Programs/UnrealBuildAccelerator/Common/Private/UbaStorageClient.cpp b/Engine/Source/Programs/UnrealBuildAccelerator/Common/Private/UbaStorageClient.cpp
--- a/Engine/Source/Programs/UnrealBuildAccelerator/Common/Private/UbaStorageClient.cpp
+++ b/Engine/Source/Programs/UnrealBuildAccelerator/Common/Private/UbaStorageClient.cpp
@@ -205,6 +205,7 @@ namespace uba
 			if (casEntry->verified && casEntry->exists)
 				return true;
 
+			casEntry->dropped = false;  // In case this comes from a retry where previous cas was dropped
 			casEntry->verified = true;
 		}
 
diff --git a/Engine/Source/Programs/UnrealBuildAccelerator/Common/Public/UbaStorage.h b/Engine/Source/Programs/UnrealBuildAccelerator/Common/Public/UbaStorage.h
--- a/Engine/Source/Programs/UnrealBuildAccelerator/Common/Public/UbaStorage.h
+++ b/Engine/Source/Programs/UnrealBuildAccelerator/Common/Public/UbaStorage.h
@@ -47,6 +47,7 @@ namespace uba
 		virtual bool StoreCasFile(CasKey& out, const tchar* fileName, const CasKey& casKeyOverride, bool deferCreation, bool fileIsCompressed) = 0;
 		virtual bool StoreCasFile(CasKey& out, StringKey fileNameKey, const tchar* fileName, FileMappingHandle mappingHandle, u64 mappingOffset, u64 fileSize, const tchar* hint, bool deferCreation = false, bool keepMappingInMemory = false) = 0;
 		virtual bool DropCasFile(const CasKey& casKey, bool forceDelete, const tchar* hint) = 0;
+		virtual bool ReportBadCasFile(const CasKey& casKey) = 0;
 		virtual bool CalculateCasKey(CasKey& out, const tchar* fileName) = 0;
 		virtual bool CopyOrLink(const CasKey& casKey, const tchar* destination, u32 fileAttributes, bool writeCompressed = false) = 0;
 		virtual bool FakeCopy(const CasKey& casKey, const tchar* destination, u64 size = 0, u64 lastWritten = 0, bool deleteExisting = true) = 0;
@@ -123,6 +124,7 @@ namespace uba
 		virtual bool StoreCasFile(CasKey& out, const tchar* fileName, const CasKey& casKeyOverride, bool deferCreation, bool fileIsCompressed) override;
 		virtual bool StoreCasFile(CasKey& out, StringKey fileNameKey, const tchar* fileName, FileMappingHandle mappingHandle, u64 mappingOffset, u64 fileSize, const tchar* hint, bool deferCreation = false, bool keepMappingInMemory = false) override;
 		virtual bool DropCasFile(const CasKey& casKey, bool forceDelete, const tchar* hint) override;
+		virtual bool ReportBadCasFile(const CasKey& casKey) override;
 		virtual bool CalculateCasKey(CasKey& out, const tchar* fileName) override;
 		virtual bool CopyOrLink(const CasKey& casKey, const tchar* destination, u32 fileAttributes, bool writeCompressed = false) override;
 		virtual bool FakeCopy(const CasKey& casKey, const tchar* destination, u64 size = 0, u64 lastWritten = 0, bool deleteExisting = true) override;
@@ -161,9 +163,9 @@ namespace uba
 		void CasEntryDeleted(CasEntry& entry, u64 size);
 		void AttachEntry(CasEntry& entry);
 		void DetachEntry(CasEntry& entry);
-		void TraverseAllCasFiles(const tchar* dir, u32 recursion, const Function<void(const StringBufferBase& fullPath, const DirectoryEntry& e)>& func);
+		void TraverseAllCasFiles(const tchar* dir, const Function<void(const StringBufferBase& fullPath, const DirectoryEntry& e)>& func, bool allowParallel = false);
 		void TraverseAllCasFiles(const Function<void(const CasKey& key, u64 size)>& func);
-		void CheckAllCasFiles();
+		bool CheckAllCasFiles(u64 checkContentOfFilesNewerThanTime = ~u64(0));
 		void HandleOverflow(Set<CasKey>* outDeletedFiles);
 		bool OpenCasDataFile(u32 index, u64 size);
 		bool CreateCasDataFiles();
diff --git a/Engine/Source/Programs/UnrealBuildAccelerator/Test/Private/UbaTestBasics.h b/Engine/Source/Programs/UnrealBuildAccelerator/Test/Private/UbaTestBasics.h
--- a/Engine/Source/Programs/UnrealBuildAccelerator/Test/Private/UbaTestBasics.h
+++ b/Engine/Source/Programs/UnrealBuildAccelerator/Test/Private/UbaTestBasics.h
@@ -119,7 +119,10 @@ namespace uba
 
 	bool TestFiles(Logger& logger, const StringBufferBase& rootDir)
 	{
-		FileAccessor fileHandle(logger, TC("UbaTestFile"));
+		StringBuffer<> testFileName(rootDir);
+		testFileName.Append(TC("UbaTestFile"));
+
+		FileAccessor fileHandle(logger, testFileName.data);
 		if (!fileHandle.CreateWrite())
 			return logger.Error(TC("Failed to create file for write"));
 
@@ -131,13 +134,23 @@ namespace uba
 			return false;
 
 		FileHandle fileHandle2;
-		if (!OpenFileSequentialRead(logger, TC("UbaTestFile"), fileHandle2))
+		if (!OpenFileSequentialRead(logger, testFileName.data, fileHandle2))
 			return logger.Error(TC("Failed to create file for read"));
 
 		u64 writeTime = 0;
 		if (!GetFileLastWriteTime(writeTime, fileHandle2))
 			return logger.Error(TC("Failed to get last written time"));
 
+		u64 writeTime2 = 0;
+		TraverseDir(logger, rootDir.data, [&](const DirectoryEntry& de)
+			{
+				if (Equals(de.name, TC("UbaTestFile")))
+					writeTime2 = de.lastWritten;
+			});
+
+		if (writeTime != writeTime2)
+			return logger.Error(TC("GetFileLastWriteTime and TraverseDir are returning different last write time for same file"));
+
 		u64 systemTime = GetSystemTimeAsFileTime();
 		if (systemTime < writeTime)
 			return logger.Error(TC("System time is lower than last written time"));
@@ -146,10 +159,10 @@ namespace uba
 
 
 		u8 byte2 = 0;
-		if (!ReadFile(logger, TC("UbaTestFile"), fileHandle2, &byte2, 1))
+		if (!ReadFile(logger, testFileName.data, fileHandle2, &byte2, 1))
 			return false;
 
-		if (!CloseFile(TC("UbaTestFile"), fileHandle2))
+		if (!CloseFile(testFileName.data, fileHandle2))
 			return false;
 
 		FileHandle fileHandle3;
@@ -174,28 +187,31 @@ namespace uba
 			return logger.Error(TC("Fail to remove TestDir"));
 
 		u64 size = 0;
-		if (!FileExists(logger, TC("UbaTestFile"), &size) || size != 1)
+		if (!FileExists(logger, testFileName.data, &size) || size != 1)
 			return logger.Error(TC("UbaTestFile not found"));
 
-		DeleteFileW(TC("UbaTestFile2"));
+		StringBuffer<> testFileName2(rootDir);
+		testFileName2.Append(TC("UbaTestFile2"));
+
+		DeleteFileW(testFileName2.data);
 
-		if (DeleteFileW(TC("UbaTestFile2")))
+		if (DeleteFileW(testFileName2.data))
 			return logger.Error(TC("Did not fail to delete non-existing UbaTestFile2 (or were things not cleaned before test)"));
 		else if (GetLastError() != ERROR_FILE_NOT_FOUND)
 			return logger.Error(TC("GetLastError did not return correct error failing to delete non-existing file UbaTestFile2"));
 
-		if (!CreateHardLinkW(TC("UbaTestFile2"), TC("UbaTestFile")))
+		if (!CreateHardLinkW(testFileName2.data, testFileName.data))
 			return logger.Error(TC("Failed to create hardlink from UbaTestFile to UbaTestFile2"));
 
-		if (!DeleteFileW(TC("UbaTestFile")))
+		if (!DeleteFileW(testFileName.data))
 			return logger.Error(TC("Failed to delete UbaTestFile"));
 
-		if (FileExists(logger, TC("UbaTestFile")))
+		if (FileExists(logger, testFileName.data))
 			return logger.Error(TC("Found non-existing file UbaTestFile"));
 
 		// CreateHardLinkW is a symbolic link on non-windows.. need to revisit
 		#if PLATFORM_WINDOWS
-		if (!FileExists(logger, TC("UbaTestFile2")))
+		if (!FileExists(logger, testFileName2.data))
 			return logger.Error(TC("Failed to find file UbaTestFile2"));
 
 		StringBuffer<> currentDir;
@@ -203,13 +219,13 @@ namespace uba
 			return logger.Error(TC("GetCurrentDirectoryW failed"));
 
 		bool foundFile = false;
-		if (!TraverseDir(logger, currentDir.data, [&](const DirectoryEntry& de) { foundFile |= TStrcmp(de.name, TC("UbaTestFile2")) == 0; }, true))
+		if (!TraverseDir(logger, rootDir.data, [&](const DirectoryEntry& de) { foundFile |= TStrcmp(de.name, TC("UbaTestFile2")) == 0; }, true))
 			return logger.Error(TC("Failed to TraverseDir '.'"));
 
 		if (!foundFile)
 			return logger.Error(TC("Did not find UbaTestFile2 with TraverseDir"));
 
-		if (!DeleteFileW(TC("UbaTestFile2")))
+		if (!DeleteFileW(testFileName2.data))
 			return false;
 		#endif
 
