[Uba]
* Fixed so fetch bytes in cache detail info is not showing if no output was fetched

diff --git a/Engine/Source/Programs/UnrealBuildAccelerator/Common/Private/UbaStats.cpp b/Engine/Source/Programs/UnrealBuildAccelerator/Common/Private/UbaStats.cpp
--- a/Engine/Source/Programs/UnrealBuildAccelerator/Common/Private/UbaStats.cpp
+++ b/Engine/Source/Programs/UnrealBuildAccelerator/Common/Private/UbaStats.cpp
@@ -277,7 +277,8 @@ namespace uba
 		#define UBA_CACHE_STAT(type, var, ver) LogStat(logger, #var, var, frequency);
 		UBA_CACHE_STATS
 		#undef UBA_CACHE_STAT
-		logger.Info(TC("   Bytes   Raw/Comp %9s %9s"), BytesToText(fetchBytesRaw).str, BytesToText(fetchBytesComp).str);
+		if (fetchBytesComp)
+			logger.Info(TC("   Bytes   Raw/Comp %9s %9s"), BytesToText(fetchBytesRaw).str, BytesToText(fetchBytesComp).str);
 	}
 
 	bool CacheStats::IsEmpty()
