[UBA]
* Reverted CasTableVersion bump to prevent helpers shared between branches from jumping back and forth

diff --git a/Engine/Source/Programs/UnrealBuildAccelerator/Common/Private/UbaStorage.cpp b/Engine/Source/Programs/UnrealBuildAccelerator/Common/Private/UbaStorage.cpp
--- a/Engine/Source/Programs/UnrealBuildAccelerator/Common/Private/UbaStorage.cpp
+++ b/Engine/Source/Programs/UnrealBuildAccelerator/Common/Private/UbaStorage.cpp
@@ -12,7 +12,7 @@
 
 namespace uba
 {
-	constexpr u32 CasTableVersion = IsWindows ? 33 : 35;
+	constexpr u32 CasTableVersion = IsWindows ? 32 : 34;
 	constexpr u32 MaxWorkItemsPerAction = 128; // Cap this to not starve other things
 
 #if UBA_USE_SPARSEFILE
