[UBA]
* Added better shutdown handling in UbaAgent and UbaCacheService when requested from outside (ctrl-c etc)

diff --git a/Engine/Source/Programs/UnrealBuildAccelerator/Agent/Private/UbaAgent.cpp b/Engine/Source/Programs/UnrealBuildAccelerator/Agent/Private/UbaAgent.cpp
--- a/Engine/Source/Programs/UnrealBuildAccelerator/Agent/Private/UbaAgent.cpp
+++ b/Engine/Source/Programs/UnrealBuildAccelerator/Agent/Private/UbaAgent.cpp
@@ -113,20 +113,26 @@ namespace uba
 		return -1;
 	}
 
-	StorageClient* g_storageClient;
-	NetworkClient* g_client;
+	ReaderWriterLock* g_exitLock = new ReaderWriterLock();
+	LoggerWithWriter* g_logger;
+	SessionClient* g_sessionClient;
+	Atomic<bool> g_shouldExit;
 
-	void CtrlBreakPressed()
+	bool ShouldExit()
 	{
-		if (g_storageClient)
-		{
-			g_storageClient->SaveCasTable(true);
-			LoggerWithWriter(g_consoleLogWriter).Info(TC("CAS table saved..."));
-		}
+		return g_shouldExit || IsEscapePressed();
+	}
 
-		abort();
-		//if (g_client)
-		//	g_client->Disconnect();
+	void CtrlBreakPressed()
+	{
+		g_shouldExit = true;
+
+		g_exitLock->EnterWrite();
+		if (g_logger)
+			g_logger->Info(TC("  Exiting..."));
+		if (g_sessionClient)
+			g_sessionClient->Stop();
+		g_exitLock->LeaveWrite();
 	}
 
 	#if PLATFORM_WINDOWS
@@ -140,9 +146,8 @@ namespace uba
 	}
 	BOOL ConsoleHandler(DWORD signal)
 	{
-		if (signal == CTRL_C_EVENT)
-			CtrlBreakPressed();
-		return FALSE;
+		CtrlBreakPressed();
+		return TRUE;
 	}
 	#else
 	void ConsoleHandler(int sig)
@@ -181,20 +186,20 @@ namespace uba
 #if UBA_AUTO_UPDATE
 	const tchar* g_ubaAgentBinaries[] = { UBA_AGENT_EXECUTABLE, UBA_DETOURS_LIBRARY };
 
-	bool DownloadBinaries(CasKey* keys)
+	bool DownloadBinaries(StorageClient& storageClient, CasKey* keys)
 	{
 		StringBuffer<256> binDir(g_rootDir);
 		binDir.Append(TC("\\binaries\\"));
-		g_storageClient->CreateDirectory(binDir.data);
+		storageClient.CreateDirectory(binDir.data);
 		u32 index = 0;
 		for (auto file : g_ubaAgentBinaries)
 		{
 			Storage::RetrieveResult result;
-			if (!g_storageClient->RetrieveCasFile(result, keys[index++], file))
+			if (!storageClient.RetrieveCasFile(result, keys[index++], file))
 				return false;
 			StringBuffer<256> fullFile(binDir);
 			fullFile.Append(file);
-			if (!g_storageClient->CopyOrLink(result.casKey, fullFile.data, DefaultAttributes()))
+			if (!storageClient.CopyOrLink(result.casKey, fullFile.data, DefaultAttributes()))
 				return false;
 		}
 		return true;
@@ -695,6 +700,11 @@ namespace uba
 		FilteredLogWriter logWriter(g_consoleLogWriter, verbose ? LogEntryType_Debug : LogEntryType_Detail);
 		LoggerWithWriter logger(logWriter, TC(""));
 
+		g_exitLock->EnterWrite();
+		g_logger = &logger;
+		g_exitLock->LeaveWrite();
+		auto glg = MakeGuard([]() { g_exitLock->EnterWrite(); g_logger = nullptr; g_exitLock->LeaveWrite(); });
+
 #if UBA_AUTO_UPDATE
 		if (waitProcessId != ~0u)
 			if (!WaitForProcess(waitProcessId))
@@ -943,6 +953,7 @@ namespace uba
 		SetConsoleCtrlHandler(ConsoleHandler, TRUE);
 		#else
 		signal(SIGINT, ConsoleHandler);
+		signal(SIGTERM, ConsoleHandler);
 		#endif
 
 		bool relaunch = false;
@@ -981,24 +992,18 @@ namespace uba
 				ncci.cryptoKey128 = crypto;
 			bool ctorSuccess = true;
 			NetworkClient* client = new NetworkClient(ctorSuccess, ncci);
-			g_client = client;
-			auto csg = MakeGuard([&]() { g_client = nullptr; client->Disconnect(); delete client; });
+			auto csg = MakeGuard([&]() { client->Disconnect(); delete client; });
 			if (!ctorSuccess)
 				return -1;
 
-			bool exit = false;
-
 			if (useListen)
 			{
 				client->StartListen(*networkBackend, port);
 				u64 startTime = GetTime();
 				while (!client->IsOrWasConnected(200))
 				{
-					if (IsEscapePressed())
-					{
-						exit = true;
-						break;
-					}
+					if (ShouldExit())
+						return 0;
 
 					u64 waitTime = GetTime() - startTime;
 					if (!poll && TimeToMs(waitTime) > listenTimeoutSec*1000)
@@ -1016,11 +1021,9 @@ namespace uba
 				bool timedOut = false;
 				while (!client->Connect(*networkBackend, host.data, port, &timedOut))
 				{
-					if (IsEscapePressed())
-					{
-						exit = true;
-						break;
-					}
+					if (ShouldExit())
+						return 0;
+
 					if (!timedOut)
 						return -1;
 
@@ -1032,9 +1035,6 @@ namespace uba
 				}
 			}
 
-			if (exit)
-				return 0;
-
 			if (!command.empty())
 			{
 				StackBinaryWriter<128> writer;
@@ -1140,7 +1140,7 @@ namespace uba
 			storageInfo.proxyPort = proxyPort;
 
 			auto storageClient = new StorageClient(storageInfo);
-			auto bscsg = MakeGuard([&]() { g_storageClient = nullptr; delete storageClient; });
+			auto bscsg = MakeGuard([&]() { delete storageClient; });
 
 			if (!storageClient->LoadCasTable(true))
 				return -1;
@@ -1205,8 +1205,14 @@ namespace uba
 					return 0;
 				});
 
+			g_sessionClient = sessionClient;
+
 			auto disconnectAndStopLoggingThread = MakeGuard([&]()
 			{
+				g_exitLock->EnterWrite();
+				g_sessionClient = nullptr;
+				g_exitLock->LeaveWrite();
+
 				networkBackend->StopListen();
 				storageClient->StopProxy();
 				auto proxyServer = proxy.server.load();
@@ -1215,6 +1221,7 @@ namespace uba
 				sessionClient->Stop();
 				sessionClient->SendSummary([&](Logger& logger) { if (proxyServer) proxyServer->PrintSummary(logger); });
 				client->Disconnect();
+
 				loopLogging = false;
 				logLinesAvailable.Set();
 				loggingThread.Wait();
@@ -1225,7 +1232,7 @@ namespace uba
 			{
 #if UBA_AUTO_UPDATE
 				logger.Info(TC("Downloading new binaries..."));
-				if (!DownloadBinaries(keys))
+				if (!DownloadBinaries(*storageClient, keys))
 					return -1;
 				relaunch = true;
 				break;
@@ -1252,7 +1259,7 @@ namespace uba
 			storageClient->Start();
 			sessionClient->Start();
 
-			while (true)
+			while (!ShouldExit())
 			{
 				if (useListen)
 				{
@@ -1349,7 +1356,7 @@ namespace uba
 			PrintContentionSummary(contLogger);
 			#endif
 		}
-		while (poll && !isTerminating);
+		while (poll && !isTerminating && !ShouldExit());
 
 #if UBA_AUTO_UPDATE
 		if (relaunch)
diff --git a/Engine/Source/Programs/UnrealBuildAccelerator/CacheService/Private/UbaCacheService.cpp b/Engine/Source/Programs/UnrealBuildAccelerator/CacheService/Private/UbaCacheService.cpp
--- a/Engine/Source/Programs/UnrealBuildAccelerator/CacheService/Private/UbaCacheService.cpp
+++ b/Engine/Source/Programs/UnrealBuildAccelerator/CacheService/Private/UbaCacheService.cpp
@@ -49,24 +49,30 @@ namespace uba
 		return -1;
 	}
 
-	bool g_exitRequested;
+	ReaderWriterLock* g_exitLock = new ReaderWriterLock();
+	LoggerWithWriter* g_logger;
+	Atomic<bool> g_shouldExit;
+
+	bool ShouldExit()
+	{
+		return g_shouldExit || IsEscapePressed();
+	}
 	
 	void CtrlBreakPressed()
 	{
-		if (!g_exitRequested)
-			LoggerWithWriter(g_consoleLogWriter, TC("")).Info(TC("Exiting..."));
-		g_exitRequested = true;
+		g_shouldExit = true;
+
+		g_exitLock->EnterWrite();
+		if (g_logger)
+			g_logger->Info(TC("  Exiting..."));
+		g_exitLock->LeaveWrite();
 	}
 
 	#if PLATFORM_WINDOWS
 	BOOL ConsoleHandler(DWORD signal)
 	{
-		if (signal == CTRL_C_EVENT)
-		{
-			CtrlBreakPressed();
-			return TRUE;
-		}
-		return FALSE;
+		CtrlBreakPressed();
+		return TRUE;
 	}
 	#else
 	void ConsoleHandler(int sig)
@@ -167,6 +173,11 @@ namespace uba
 		FilteredLogWriter logWriter(g_consoleLogWriter, quiet ? LogEntryType_Info : LogEntryType_Detail);
 		LoggerWithWriter logger(logWriter, TC(""));
 
+		g_exitLock->EnterWrite();
+		g_logger = &logger;
+		g_exitLock->LeaveWrite();
+		auto glg = MakeGuard([]() { g_exitLock->EnterWrite(); g_logger = nullptr; g_exitLock->LeaveWrite(); });
+
 		const tchar* dbgStr = TC("");
 		#if UBA_DEBUG
 		dbgStr = TC(" (DEBUG)");
@@ -220,7 +231,7 @@ namespace uba
 		if (!cacheServer.Load())
 			return -1;
 
-		if (!cacheServer.RunMaintenance(true))
+		if (!cacheServer.RunMaintenance(true, ShouldExit))
 			return -1;
 
 		{
@@ -230,10 +241,10 @@ namespace uba
 			if (!networkServer.StartListen(networkBackend, port, listenIp.data))
 				return -1;
 
-			while (!g_exitRequested && !IsEscapePressed())
+			while (!ShouldExit())
 			{
 				Sleep(1000);
-				if (!cacheServer.RunMaintenance(false))
+				if (!cacheServer.RunMaintenance(false, ShouldExit))
 					break;
 			}
 		}
diff --git a/Engine/Source/Programs/UnrealBuildAccelerator/Common/Private/UbaCacheServer.cpp b/Engine/Source/Programs/UnrealBuildAccelerator/Common/Private/UbaCacheServer.cpp
--- a/Engine/Source/Programs/UnrealBuildAccelerator/Common/Private/UbaCacheServer.cpp
+++ b/Engine/Source/Programs/UnrealBuildAccelerator/Common/Private/UbaCacheServer.cpp
@@ -286,7 +286,7 @@ namespace uba
 		return true;
 	}
 
-	bool CacheServer::RunMaintenance(bool force)
+	bool CacheServer::RunMaintenance(bool force, const Function<bool()>& shouldExit)
 	{
 		SCOPED_WRITE_LOCK(m_maintenanceLock, lock);
 		SCOPED_READ_LOCK(m_connectionsLock, lock2);
@@ -320,6 +320,9 @@ namespace uba
 		}
 		u64 totalCasCount = existingCas.size() + deletedCasCount;
 
+		if (shouldExit())
+			return true;
+
 		u64 now = GetSystemTimeAsFileTime();
 		u64 oldest = now;
 
@@ -459,6 +462,9 @@ namespace uba
 
 		m_logger.Detail(TC("  Deleted %llu cas files and %llu cache entries (%s)"), deletedCasCount, deleteEntryCount.load(), TimeToText(GetTime() - deleteCacheEntriesStartTime).str);
 
+		if (shouldExit())
+			return true;
+
 		Atomic<u32> bucketCounter;
 		m_server.ParallelFor(16, m_buckets, [&](auto& it)
 		{
diff --git a/Engine/Source/Programs/UnrealBuildAccelerator/Common/Public/UbaCacheServer.h b/Engine/Source/Programs/UnrealBuildAccelerator/Common/Public/UbaCacheServer.h
--- a/Engine/Source/Programs/UnrealBuildAccelerator/Common/Public/UbaCacheServer.h
+++ b/Engine/Source/Programs/UnrealBuildAccelerator/Common/Public/UbaCacheServer.h
@@ -22,7 +22,7 @@ namespace uba
 		bool Load();
 		bool Save();
 
-		bool RunMaintenance(bool force);
+		bool RunMaintenance(bool force, const Function<bool()>& shouldExit);
 
 	private:
 		bool SaveNoLock();
