[UBA]
* Fixed so trace system reports a warning and frees up trace memory when running out of reserved space in trace memory.

diff --git a/Engine/Source/Programs/UnrealBuildAccelerator/Common/Private/UbaTrace.cpp b/Engine/Source/Programs/UnrealBuildAccelerator/Common/Private/UbaTrace.cpp
--- a/Engine/Source/Programs/UnrealBuildAccelerator/Common/Private/UbaTrace.cpp
+++ b/Engine/Source/Programs/UnrealBuildAccelerator/Common/Private/UbaTrace.cpp
@@ -16,41 +16,38 @@ namespace uba
 
 	Trace::~Trace()
 	{
-		if (m_memoryBegin)
-			UnmapViewOfFile(m_memoryBegin, m_memoryCapacity, TC("Trace"));
-		if (m_memoryHandle.IsValid())
-			CloseFileMapping(m_memoryHandle);
+		FreeMemory();
 	}
 
 	struct Trace::WriterScope : ScopedWriteLock, BinaryWriter
 	{
 		WriterScope(Trace& trace) : ScopedWriteLock(trace.m_memoryLock), BinaryWriter(trace.m_memoryBegin, trace.m_memoryPos, trace.m_memoryCapacity), m_trace(trace)
 		{
-			isValid = EnsureMemory(TraceMessageMaxSize);
+			EnsureMemory(TraceMessageMaxSize);
 		}
 
 		~WriterScope()
 		{
+			if (!m_isValid)
+				return;
 			m_trace.m_memoryPos = GetPosition();
 			*(u32*)m_trace.m_memoryBegin = u32(m_trace.m_memoryPos);
 		}
 
+		bool IsValid() { return m_isValid; }
+
 		WriterScope(const WriterScope&) = delete;
 		void operator=(const WriterScope&) = delete;
 
 		bool EnsureMemory(u64 size)
 		{
-			u64 committedMemoryNeeded = AlignUp(m_trace.m_memoryPos + size, size);
-			if (m_trace.m_memoryCommitted >= committedMemoryNeeded)
-				return true;
-			if (!MapViewCommit(m_trace.m_memoryBegin + m_trace.m_memoryCommitted, committedMemoryNeeded - m_trace.m_memoryCommitted))
-				return m_trace.m_logger.Error(TC("Failed to commit memory for trace (Pos: %llu Capacity: %llu, Already Committed: %llu, Needed: %llu): %s"), m_trace.m_memoryPos, m_trace.m_memoryCapacity, m_trace.m_memoryCommitted, committedMemoryNeeded, LastErrorToText().data);
-			m_trace.m_memoryCommitted = committedMemoryNeeded;
-			return true;
+			if (m_isValid)
+				m_isValid = m_trace.EnsureMemory(size);
+			return m_isValid;
 		}
 
 		Trace& m_trace;
-		bool isValid;
+		bool m_isValid = true;
 	};
 
 	bool Trace::StartWrite(const tchar* namedTrace, u64 traceMemCapacity)
@@ -75,7 +72,7 @@ namespace uba
 
 		{
 			WriterScope writer(*this);
-			if (!writer.isValid)
+			if (!writer.IsValid())
 				return false;
 			writer.AllocWrite(4);
 			writer.WriteU32(TraceVersion);
@@ -96,16 +93,16 @@ namespace uba
 	{
 		if (!m_memoryBegin)
 			return true;
+		auto g = MakeGuard([this]() { FreeMemory(); });
 
 		{
 			WriterScope writer(*this);
-			if (!writer.isValid)
+			if (!writer.IsValid())
 				return false;
 			writer.WriteByte(TraceType_Summary);
 			writer.Write7BitEncoded(GetTime() - m_startTime);
 		}
 
-
 		if (!writeFileName || !*writeFileName)
 			return true;
 		FileAccessor traceFile(m_logger, writeFileName);
@@ -117,9 +114,6 @@ namespace uba
 		if (!traceFile.Close())
 			return false;
 		m_logger.Info(TC("Trace file written to %s with size %s"), writeFileName, BytesToText(fileSize).str);
-		
-		UnmapViewOfFile(m_memoryBegin, m_memoryCapacity, TC("Trace"));
-		m_memoryBegin = nullptr;
 		return true;
 	}
 
@@ -135,6 +129,41 @@ namespace uba
 		EndWork(id);
 	}
 
+	void Trace::FreeMemory()
+	{
+		if (m_memoryBegin)
+		{
+			UnmapViewOfFile(m_memoryBegin, m_memoryCapacity, TC("Trace"));
+			m_memoryBegin = nullptr;
+		}
+		if (m_memoryHandle.IsValid())
+		{
+			CloseFileMapping(m_memoryHandle);
+			m_memoryHandle = {};
+		}
+	}
+
+	bool Trace::EnsureMemory(u64 size)
+	{
+		if (!m_memoryBegin)
+			return false;
+
+		u64 committedMemoryNeeded = AlignUp(m_memoryPos + size, 64*1024);
+		if (m_memoryCommitted >= committedMemoryNeeded)
+			return true;
+
+		if (MapViewCommit(m_memoryBegin + m_memoryCommitted, committedMemoryNeeded - m_memoryCommitted))
+		{
+			m_memoryCommitted = committedMemoryNeeded;
+			return true;
+		}
+
+		FreeMemory();
+		m_logger.Warning(TC("Failed to commit memory for trace (Pos: %llu Capacity: %llu, Already Committed: %llu, Needed: %llu): %s"), m_memoryPos, m_memoryCapacity, m_memoryCommitted, committedMemoryNeeded, LastErrorToText().data);
+		return false;
+	}
+
+
 	u32 Trace::AddString(const tchar* string)
 	{
 		if (!m_memoryBegin)
@@ -148,7 +177,7 @@ namespace uba
 		{
 			insres.first->second = u32(m_strings.size() - 1);
 			WriterScope writer(*this);
-			if (!writer.isValid)
+			if (!writer.IsValid())
 				return 0;
 			writer.WriteByte(TraceType_String);
 			writer.WriteString(string, stringLen);
@@ -160,7 +189,7 @@ namespace uba
 		if (!m_memoryBegin) \
 			return; \
 		WriterScope writer(*this); \
-		if (!writer.isValid) \
+		if (!writer.IsValid()) \
 			return; \
 		writer.WriteByte(x); \
 		writer.Write7BitEncoded(GetTime() - m_startTime);
@@ -235,7 +264,7 @@ namespace uba
 			if (lineCounter++ == 100) // We don't want to write the entire error in the trace stream to blow the entire buffer
 				break;
 			if (!writer.EnsureMemory(1 + (line.text.size()+1)*sizeof(tchar)))
-				break;
+				return;
 			writer.WriteByte(line.type);
 			writer.WriteString(line.text);
 		}
diff --git a/Engine/Source/Programs/UnrealBuildAccelerator/Common/Public/UbaTrace.h b/Engine/Source/Programs/UnrealBuildAccelerator/Common/Public/UbaTrace.h
--- a/Engine/Source/Programs/UnrealBuildAccelerator/Common/Public/UbaTrace.h
+++ b/Engine/Source/Programs/UnrealBuildAccelerator/Common/Public/UbaTrace.h
@@ -100,6 +100,8 @@ namespace uba
 
 	private:
 		struct WriterScope;
+		void FreeMemory();
+		bool EnsureMemory(u64 size);
 		u32 AddString(const tchar* string);
 
 		LoggerWithWriter m_logger;
