[UBA]
* Added placeholder icon to UbaVisualizer

diff --git a/Engine/Source/Programs/UnrealBuildAccelerator/UbaVisualizer.Target.cs b/Engine/Source/Programs/UnrealBuildAccelerator/UbaVisualizer.Target.cs
--- a/Engine/Source/Programs/UnrealBuildAccelerator/UbaVisualizer.Target.cs
+++ b/Engine/Source/Programs/UnrealBuildAccelerator/UbaVisualizer.Target.cs
@@ -10,5 +10,6 @@ public class UbaVisualizerTarget : TargetRules
 		LaunchModuleName = "UbaVisualizer";
 		UbaAgentTarget.CommonUbaSettings(this, Target);
 		bIsBuildingConsoleApplication = false;
+		WindowsPlatform.ApplicationIcon = "Programs/UnrealBuildAccelerator/Visualizer/Resources/Windows/UbaVisualizer.ico";
 	}
 }
diff --git a/Engine/Source/Programs/UnrealBuildAccelerator/Visualizer/Private/UbaVisualizer.cpp b/Engine/Source/Programs/UnrealBuildAccelerator/Visualizer/Private/UbaVisualizer.cpp
--- a/Engine/Source/Programs/UnrealBuildAccelerator/Visualizer/Private/UbaVisualizer.cpp
+++ b/Engine/Source/Programs/UnrealBuildAccelerator/Visualizer/Private/UbaVisualizer.cpp
@@ -437,7 +437,7 @@ namespace uba
 		wndClassEx.cbSize = sizeof(wndClassEx);
 		wndClassEx.style = CS_HREDRAW | CS_VREDRAW;
 		wndClassEx.lpfnWndProc = &StaticWinProc;
-		wndClassEx.hIcon = NULL;//Icon;
+		wndClassEx.hIcon = LoadIcon(hInstance, MAKEINTRESOURCE(123));
 		wndClassEx.hCursor = LoadCursor(NULL, IDC_ARROW);
 		wndClassEx.hInstance = hInstance;
 		wndClassEx.hbrBackground = NULL;
