[Uba]
* Changed so env variables are only explicitly populated if process is detoured

diff --git a/Engine/Source/Programs/UnrealBuildAccelerator/Common/Private/UbaProcess.cpp b/Engine/Source/Programs/UnrealBuildAccelerator/Common/Private/UbaProcess.cpp
--- a/Engine/Source/Programs/UnrealBuildAccelerator/Common/Private/UbaProcess.cpp
+++ b/Engine/Source/Programs/UnrealBuildAccelerator/Common/Private/UbaProcess.cpp
@@ -1561,13 +1561,7 @@ namespace uba
 
 			Vector<const char*> envvars;
 
-			const char* it = (const char*)environment;
-			while (*it)
-			{
-				const char* s = it;
-				envvars.push_back(s);
-				it += TStrlen(s) + 1;
-			}
+			char** envvarsPtr = nullptr;
 
 			int outPipe[2] = { -1, -1 };
 			int errPipe[2] = { -1, -1 };
@@ -1576,6 +1570,14 @@ namespace uba
 
 			if (m_detourEnabled)
 			{
+				const char* it = (const char*)environment;
+				while (*it)
+				{
+					const char* s = it;
+					envvars.push_back(s);
+					it += TStrlen(s) + 1;
+				}
+
 				const char* detoursLib = m_session.m_detoursLibrary.c_str();
 				if (*detoursLib)
 				{
@@ -1628,9 +1630,10 @@ namespace uba
 					envvars.push_back("UBA_REMOTE=1");
 				if (!logFile.IsEmpty())
 					envvars.push_back(logFile.data);
-			}
 
-			envvars.push_back(nullptr);
+				envvars.push_back(nullptr);
+				envvarsPtr = (char**)envvars.data();
+			}
 
 			if (true)
 			{
@@ -1663,7 +1666,7 @@ namespace uba
 			pid_t processID;
 			while (true)
 			{
-				res = posix_spawnp(&processID, m_realApplication.c_str(), &fileActions, &attr, (char**)argsArray, (char**)envvars.data());
+				res = posix_spawnp(&processID, m_realApplication.c_str(), &fileActions, &attr, (char**)argsArray, envvarsPtr);
 				if (res == 0)
 					break;
 
