[UBA]
* Added UbaCli.exe -checkfiletable that can be used to check so files on disk match what uba thinks they are

diff --git a/Engine/Source/Programs/UnrealBuildAccelerator/Cli/Private/UbaCli.cpp b/Engine/Source/Programs/UnrealBuildAccelerator/Cli/Private/UbaCli.cpp
--- a/Engine/Source/Programs/UnrealBuildAccelerator/Cli/Private/UbaCli.cpp
+++ b/Engine/Source/Programs/UnrealBuildAccelerator/Cli/Private/UbaCli.cpp
@@ -8,6 +8,7 @@
 #include "UbaNetworkBackendTcp.h"
 #include "UbaPlatform.h"
 #include "UbaProtocol.h"
+#include "UbaRootPaths.h"
 #include "UbaScheduler.h"
 #include "UbaSessionClient.h"
 #include "UbaSessionServer.h"
@@ -68,6 +69,7 @@ namespace uba
 		logger.Info(TC("   -loop=<count>           Loop the commandline <count> number of times. Will exit when/if it fails"));
 		logger.Info(TC("   -workdir=<dir>          Working directory"));
 		logger.Info(TC("   -checkcas               Check so all cas entries are correct"));
+		logger.Info(TC("   -checkfiletable         Check so file table has correct cas stored"));
 		logger.Info(TC("   -checkaws               Check if we are inside aws and output information about aws"));
 		logger.Info(TC("   -deletecas              Deletes the casdb"));
 		logger.Info(TC("   -getcas                 Will print hash of application"));
@@ -167,6 +169,7 @@ namespace uba
 		bool printSummary = false;
 		bool populateCache = false;
 		bool writeCacheSummary = false;
+		TString checkFileTable;
 		TString cacheFilterString;
 
 		u32 loopCount = 1;
@@ -308,6 +311,15 @@ namespace uba
 			{
 				checkCas = true;
 			}
+			else if (name.Equals(TC("-checkfiletable")))
+			{
+				if (value.IsEmpty())
+					return PrintHelp(TC("-checkfiletable needs a value"));
+				StringBuffer<> temp;
+				if ((temp.count = GetFullPathNameW(value.Replace('/', PathSeparator).data, temp.capacity, temp.data, nullptr)) == 0)
+					return PrintHelp(StringBuffer<>().Appendf(TC("-checkfiletable has invalid path %s"), temp.data).data);
+				checkFileTable = temp.data;
+			}
 			else if (name.Equals(TC("-checkcas2")))
 			{
 				checkCas2 = true;
@@ -390,6 +402,18 @@ namespace uba
 			return success ? 0 : -1;
 		}
 
+		if (!checkFileTable.empty())
+		{
+			StorageCreateInfo storageInfo(g_rootDir.data, logWriter);
+			storageInfo.casCapacityBytes = 0;
+			storageInfo.storeCompressed = storeCompressed;
+			StorageImpl storage(storageInfo);
+			if (!storage.LoadCasTable())
+				return -1;
+			bool success = storage.CheckFileTable(checkFileTable.data(), DefaultProcessorCount);
+			return success ? 0 : -1;
+		}
+
 		if (checkCas2) // Creates a storage server and storage client and transfer _all_ cas files over network
 		{
 			NetworkBackendTcp networkBackend(logWriter);
@@ -594,6 +618,14 @@ namespace uba
 			}
 		}
 
+		// Remove empty spaces and line feeds etc at the end.. just to solve annoying copy paste command lines and accidentally getting line feed
+		while (!arguments.empty())
+		{
+			tchar lastChar = arguments[arguments.size()-1];
+			if (lastChar != '\n' && lastChar != '\r' && lastChar != '\t' && lastChar != ' ')
+				break;
+			arguments.resize(arguments.size() - 1);
+		}
 
 
 		if (isRemote)
@@ -626,6 +658,13 @@ namespace uba
 				return logger.Error(TC("Error exit code: %u"), process.GetExitCode());
 			u64 time = GetTime() - start;
 			logger.Info(TC("%s run took %s"), (enableDetour ? TC("Boxed") : TC("Native")), TimeToText(time).str);
+
+			if (populateCache)
+			{
+				logger.Error(TC("Populating cache not implemented... todo"));
+				RootPaths rootPaths;
+				cacheClient->WriteToCache(rootPaths, 0, pinfo, nullptr, 1, nullptr, 0);
+			}
 			return true;
 		};
 
diff --git a/Engine/Source/Programs/UnrealBuildAccelerator/Common/Private/UbaStorage.cpp b/Engine/Source/Programs/UnrealBuildAccelerator/Common/Private/UbaStorage.cpp
--- a/Engine/Source/Programs/UnrealBuildAccelerator/Common/Private/UbaStorage.cpp
+++ b/Engine/Source/Programs/UnrealBuildAccelerator/Common/Private/UbaStorage.cpp
@@ -1689,6 +1689,77 @@ namespace uba
 		return true;
 	}
 
+	bool StorageImpl::CheckFileTable(const tchar* searchPath, u32 workerCount)
+	{
+		m_logger.Info(TC("Searching %s to check files against file table..."), searchPath);
+		WorkManagerImpl workManager(workerCount);
+
+		List<TString> directories;
+		directories.push_back(searchPath);
+
+		Atomic<u32> foundFiles;
+		Atomic<u32> trackedFiles;
+		Atomic<u32> testedFiles;
+		Atomic<u32> errorCount;
+
+		u64 startTime = GetTime();
+		
+		while (!directories.empty())
+		{
+			TString dir = directories.front();
+			directories.pop_front();
+			TraverseDir(m_logger, dir.data(), [&](const DirectoryEntry& e)
+				{
+					StringBuffer<> path(dir);
+					path.EnsureEndsWithSlash().Append(e.name);
+					if (CaseInsensitiveFs)
+						path.MakeLower();
+
+					if (IsDirectory(e.attributes))
+					{
+						if (Equals(e.name, TC("Content")))
+							return;
+						directories.push_back(path.data);
+						return;
+					}
+					u64 lastWritten = e.lastWritten;
+					u64 size = e.size;
+
+					workManager.AddWork([&, p = TString(path.data), lastWritten, size]()
+						{
+							++foundFiles;
+							StringKey key = ToStringKey(p.data(), p.size());
+							
+							auto findIt = m_fileTableLookup.find(key);
+							if (findIt == m_fileTableLookup.end())
+								return;
+							++trackedFiles;
+
+							FileEntry& fe = findIt->second;
+							if (fe.lastWritten != lastWritten || fe.size != size)
+								return;
+							++testedFiles;
+
+							CasKey casKey;
+							if (!CalculateCasKey(casKey, p.data()))
+								m_logger.Warning(TC("Failed to calculate cas key for %s"), p.data());
+
+							if (casKey != fe.casKey)
+							{
+								++errorCount;
+								m_logger.Error(TC("CasKey mismatch for %s even though size and lastwritten were the same. Corrupt path table!"), p.data());
+							}
+
+						}, 1, TC(""));
+				});
+		}
+		workManager.FlushWork();
+
+		m_logger.Info(TC("Done. %u errors found. Searched %u files where %u was tracked and %u matched table."), errorCount.load(), foundFiles.load(), trackedFiles.load(), testedFiles.load(), TimeToText(GetTime() - startTime).str);
+
+		return errorCount == 0;
+	}
+
 	const tchar* StorageImpl::GetTempPath()
 	{
 		return m_tempPath.data;
@@ -1736,7 +1807,7 @@ namespace uba
 			Atomic<u32> atomicDeleteCount;
 			TraverseDir(m_logger, m_rootDir.data, [&](const DirectoryEntry& e)
 				{
-					if (!IsDirectory(e.attributes))
+					if (IsDirectory(e.attributes))
 						return;
 					workManager.AddWork([&, name = TString(e.name)]()
 						{
diff --git a/Engine/Source/Programs/UnrealBuildAccelerator/Common/Public/UbaStorage.h b/Engine/Source/Programs/UnrealBuildAccelerator/Common/Public/UbaStorage.h
--- a/Engine/Source/Programs/UnrealBuildAccelerator/Common/Public/UbaStorage.h
+++ b/Engine/Source/Programs/UnrealBuildAccelerator/Common/Public/UbaStorage.h
@@ -105,6 +105,7 @@ namespace uba
 
 		bool LoadCasTable(bool logStats = true);
 		bool CheckCasContent(u32 workerCount);
+		bool CheckFileTable(const tchar* searchPath, u32 workerCount);
 		const tchar* GetTempPath();
 
 		virtual bool SaveCasTable(bool deleteIsRunningfile, bool deleteDropped = true) override;
