[UBA]
* Added support for providing custom environment variables when creating UbaSession
* Reverted code using real environment variables when spawning processes...

diff --git a/Engine/Source/Programs/UnrealBuildAccelerator/Common/Private/UbaProcess.cpp b/Engine/Source/Programs/UnrealBuildAccelerator/Common/Private/UbaProcess.cpp
--- a/Engine/Source/Programs/UnrealBuildAccelerator/Common/Private/UbaProcess.cpp
+++ b/Engine/Source/Programs/UnrealBuildAccelerator/Common/Private/UbaProcess.cpp
@@ -1561,7 +1561,13 @@ namespace uba
 
 			Vector<const char*> envvars;
 
-			char** envvarsPtr = nullptr;
+			const char* it = (const char*)environment;
+			while (*it)
+			{
+				const char* s = it;
+				envvars.push_back(s);
+				it += TStrlen(s) + 1;
+			}
 
 			int outPipe[2] = { -1, -1 };
 			int errPipe[2] = { -1, -1 };
@@ -1570,14 +1576,6 @@ namespace uba
 
 			if (m_detourEnabled)
 			{
-				const char* it = (const char*)environment;
-				while (*it)
-				{
-					const char* s = it;
-					envvars.push_back(s);
-					it += TStrlen(s) + 1;
-				}
-
 				const char* detoursLib = m_session.m_detoursLibrary.c_str();
 				if (*detoursLib)
 				{
@@ -1630,11 +1628,10 @@ namespace uba
 					envvars.push_back("UBA_REMOTE=1");
 				if (!logFile.IsEmpty())
 					envvars.push_back(logFile.data);
-
-				envvars.push_back(nullptr);
-				envvarsPtr = (char**)envvars.data();
 			}
 
+			envvars.push_back(nullptr);
+
 			if (true)
 			{
 				if (pipe(outPipe) || pipe(errPipe))
@@ -1666,7 +1663,7 @@ namespace uba
 			pid_t processID;
 			while (true)
 			{
-				res = posix_spawnp(&processID, m_realApplication.c_str(), &fileActions, &attr, (char**)argsArray, envvarsPtr);
+				res = posix_spawnp(&processID, m_realApplication.c_str(), &fileActions, &attr, (char**)argsArray, (char**)envvars.data());
 				if (res == 0)
 					break;
 
diff --git a/Engine/Source/Programs/UnrealBuildAccelerator/Common/Private/UbaSession.cpp b/Engine/Source/Programs/UnrealBuildAccelerator/Common/Private/UbaSession.cpp
--- a/Engine/Source/Programs/UnrealBuildAccelerator/Common/Private/UbaSession.cpp
+++ b/Engine/Source/Programs/UnrealBuildAccelerator/Common/Private/UbaSession.cpp
@@ -1646,47 +1646,56 @@ namespace uba
 			return m_environmentVariables.data();
 
 #if PLATFORM_WINDOWS
-		auto strs = GetEnvironmentStringsW();
-		for (auto it = strs; *it; it += TStrlen(it) + 1)
+		auto HandleEnvironmentVar = [&](const tchar* env)
 		{
 			StringBuffer<> varName;
-			varName.Append(it, TStrchr(it, '=') - it);
-			const tchar* varValue = it + varName.count + 1;
+			varName.Append(env, TStrchr(env, '=') - env);
+			const tchar* varValue = env + varName.count + 1;
 
 			if (m_runningRemote && varName.Equals(TC("PATH")))
 			{
 				AddEnvironmentVariableNoLock(TC("PATH"), TC("c:\\noenvironment"));
-				continue;
+				return;
 			}
 			if (varName.Equals(TC("TEMP")) || varName.Equals(TC("TMP")))
 			{
 				AddEnvironmentVariableNoLock(varName.data, m_tempPath.data);
-				continue;
+				return;
 			}
 			if (varName.Equals(TC("_CL_")) || varName.Equals(TC("CL")))
 			{
-				continue;
+				return;
 			}
 
 			AddEnvironmentVariableNoLock(varName.data, varValue);
-		}
-
-		FreeEnvironmentStrings(strs);
+		};
 
+		if (m_environmentMemory.empty())
+		{
+			auto strs = GetEnvironmentStringsW();
+			for (auto env = strs; *env; env += TStrlen(env) + 1)
+				HandleEnvironmentVar(env);
+			FreeEnvironmentStrings(strs);
+		}
+		else
+		{
+			BinaryReader reader(m_environmentMemory.data(), 0, m_environmentMemory.size());
+			while (reader.GetLeft())
+				HandleEnvironmentVar(reader.ReadString().c_str());
+		}
 		AddEnvironmentVariableNoLock(TC("MSBUILDDISABLENODEREUSE"), TC("1")); // msbuild will reuse existing helper nodes but since those are not detoured we can't let that happen
 		AddEnvironmentVariableNoLock(TC("DOTNET_CLI_USE_MSBUILD_SERVER"), TC("0")); // Disable msbuild server
 		AddEnvironmentVariableNoLock(TC("DOTNET_CLI_TELEMETRY_OPTOUT"), TC("1")); // Stop talking to telemetry service
 #else
-		int i = 0;
-		while (char* env = environ[i++])
+		auto HandleEnvironmentVar = [&](const tchar* env)
 		{
 			if (StartsWith(env, "TMPDIR="))
-				continue;
+				return;
 
 			if (!StartsWith(env, "PATH="))
 			{
 				m_environmentVariables.insert(m_environmentVariables.end(), env, env + TStrlen(env) + 1);
-				continue;
+				return;
 			}
 
 			TString paths;
@@ -1717,6 +1726,19 @@ namespace uba
 				paths.append(s, e);
 			}
 			AddEnvironmentVariableNoLock("PATH", paths.c_str());
+		};
+
+		if (m_environmentMemory.empty())
+		{
+			int i = 0;
+			while (char* env = environ[i++])
+				HandleEnvironmentVar(env);
+		}
+		else
+		{
+			BinaryReader reader(m_environmentMemory.data(), 0, m_environmentMemory.size());
+			while (reader.GetLeft())
+				HandleEnvironmentVar(reader.ReadString().c_str());
 		}
 		AddEnvironmentVariableNoLock("TMPDIR", m_tempPath.data);
 #endif
diff --git a/Engine/Source/Programs/UnrealBuildAccelerator/Common/Private/UbaSessionServer.cpp b/Engine/Source/Programs/UnrealBuildAccelerator/Common/Private/UbaSessionServer.cpp
--- a/Engine/Source/Programs/UnrealBuildAccelerator/Common/Private/UbaSessionServer.cpp
+++ b/Engine/Source/Programs/UnrealBuildAccelerator/Common/Private/UbaSessionServer.cpp
@@ -90,7 +90,7 @@ namespace uba
 	};
 
 
-	SessionServer::SessionServer(const SessionServerCreateInfo& info)
+	SessionServer::SessionServer(const SessionServerCreateInfo& info, const u8* environment, u32 environmentSize)
 	:	Session(info, TC("UbaSessionServer"), false, &info.server)
 	,	m_server(info.server)
 	,	m_maxRemoteProcessCount(~0u)
@@ -115,6 +115,12 @@ namespace uba
 			}
 		);
 
+		if (environmentSize)
+		{
+			m_environmentMemory.resize(environmentSize);
+			memcpy(m_environmentMemory.data(), environment, environmentSize);
+		}
+
 		m_uiLanguage = GetUserDefaultUILanguage();
 		m_resetCas = info.resetCas;
 		m_remoteExecutionEnabled = info.remoteExecutionEnabled;
diff --git a/Engine/Source/Programs/UnrealBuildAccelerator/Common/Public/UbaSession.h b/Engine/Source/Programs/UnrealBuildAccelerator/Common/Public/UbaSession.h
--- a/Engine/Source/Programs/UnrealBuildAccelerator/Common/Public/UbaSession.h
+++ b/Engine/Source/Programs/UnrealBuildAccelerator/Common/Public/UbaSession.h
@@ -237,6 +237,8 @@ namespace uba
 		HANDLE m_processJobObject = NULL;
 		#endif
 
+		Vector<u8> m_environmentMemory;
+
 		ReaderWriterLock m_environmentVariablesLock;
 		Vector<tchar> m_environmentVariables;
 		UnorderedSet<const tchar*, HashStringNoCase, EqualStringNoCase> m_localEnvironmentVariables;
diff --git a/Engine/Source/Programs/UnrealBuildAccelerator/Common/Public/UbaSessionServer.h b/Engine/Source/Programs/UnrealBuildAccelerator/Common/Public/UbaSessionServer.h
--- a/Engine/Source/Programs/UnrealBuildAccelerator/Common/Public/UbaSessionServer.h
+++ b/Engine/Source/Programs/UnrealBuildAccelerator/Common/Public/UbaSessionServer.h
@@ -14,8 +14,8 @@ namespace uba
 	class SessionServer final : public Session
 	{
 	public:
-		// Ctor/dtor
-		SessionServer(const SessionServerCreateInfo& info);
+		// Ctor/dtor. "environmnent" should be an array of utf8-written strings. This is really only needed for posix platforms where process is a c# process
+		SessionServer(const SessionServerCreateInfo& info, const u8* environment = nullptr, u32 environmentSize = 0);
 		~SessionServer();
 
 		// Run process remotely.
diff --git a/Engine/Source/Programs/UnrealBuildAccelerator/Host/Private/UbaExports.cpp b/Engine/Source/Programs/UnrealBuildAccelerator/Host/Private/UbaExports.cpp
--- a/Engine/Source/Programs/UnrealBuildAccelerator/Host/Private/UbaExports.cpp
+++ b/Engine/Source/Programs/UnrealBuildAccelerator/Host/Private/UbaExports.cpp
@@ -313,9 +313,9 @@ extern "C"
 		delete info;
 	}
 
-	uba::SessionServer* SessionServer_Create(const uba::SessionServerCreateInfo& info)
+	uba::SessionServer* SessionServer_Create(const uba::SessionServerCreateInfo& info, const uba::u8* environment, uba::u32 environmentSize)
 	{
-		return new uba::SessionServer(info);
+		return new uba::SessionServer(info, environment, environmentSize);
 	}
 	void SessionServer_SetRemoteProcessAvailable(uba::SessionServer* server, SessionServer_RemoteProcessAvailableCallback* available, void* userData)
 	{
diff --git a/Engine/Source/Programs/UnrealBuildAccelerator/Host/Public/UbaExports.h b/Engine/Source/Programs/UnrealBuildAccelerator/Host/Public/UbaExports.h
--- a/Engine/Source/Programs/UnrealBuildAccelerator/Host/Public/UbaExports.h
+++ b/Engine/Source/Programs/UnrealBuildAccelerator/Host/Public/UbaExports.h
@@ -89,7 +89,7 @@ extern "C"
 	UBA_API uba::SessionServerCreateInfo* SessionServerCreateInfo_Create(uba::StorageServer& storage, uba::NetworkServer& client, uba::LogWriter& writer, const uba::tchar* rootDir, const uba::tchar* traceOutputFile,
 		bool disableCustomAllocator, bool launchVisualizer, bool resetCas, bool writeToDisk, bool detailedTrace, bool allowWaitOnMem = false, bool allowKillOnMem = false, bool storeObjFilesCompressed = false);
 	UBA_API void SessionServerCreateInfo_Destroy(uba::SessionServerCreateInfo* info);
-	UBA_API uba::SessionServer* SessionServer_Create(const uba::SessionServerCreateInfo& info);
+	UBA_API uba::SessionServer* SessionServer_Create(const uba::SessionServerCreateInfo& info, const uba::u8* environment = nullptr, uba::u32 environmentSize = 0);
 	UBA_API void SessionServer_SetRemoteProcessAvailable(uba::SessionServer* server, SessionServer_RemoteProcessAvailableCallback* available, void* userData);
 	UBA_API void SessionServer_SetRemoteProcessReturned(uba::SessionServer* server, SessionServer_RemoteProcessReturnedCallback* returned, void* userData);
 	UBA_API void SessionServer_RefreshDirectory(uba::SessionServer* server, const uba::tchar* directory);
