[UBA]
* Fixed so unit tests are creating the UbaTest dir first to make sure it exists when tests start

diff --git a/Engine/Source/Programs/UnrealBuildAccelerator/Test/Private/UbaTestAll.h b/Engine/Source/Programs/UnrealBuildAccelerator/Test/Private/UbaTestAll.h
--- a/Engine/Source/Programs/UnrealBuildAccelerator/Test/Private/UbaTestAll.h
+++ b/Engine/Source/Programs/UnrealBuildAccelerator/Test/Private/UbaTestAll.h
@@ -114,6 +114,8 @@ namespace uba
 		#else
 		testRootDir.count = GetFullPathNameW("~/UbaTest", testRootDir.capacity, testRootDir.data, nullptr);
 		#endif
+		DeleteAllFiles(logger, testRootDir.data, false);
+		CreateDirectoryW(testRootDir.data);
 		testRootDir.EnsureEndsWithSlash();
 
 		logger.Info(TC("Running tests (Test rootdir: %s)"), testRootDir.data);
