[UBA]
* Removed code in SessionServer::WaitOnAllTasks that terminates processes on windows. that code was never supposed to be there
* Removed printing of error when trying to traverse directory on volume that has bitlock enabled and not allowed to access

diff --git a/Engine/Source/Programs/UnrealBuildAccelerator/Common/Private/UbaSessionServer.cpp b/Engine/Source/Programs/UnrealBuildAccelerator/Common/Private/UbaSessionServer.cpp
--- a/Engine/Source/Programs/UnrealBuildAccelerator/Common/Private/UbaSessionServer.cpp
+++ b/Engine/Source/Programs/UnrealBuildAccelerator/Common/Private/UbaSessionServer.cpp
@@ -436,15 +436,6 @@ namespace uba
 					processes.push_back(pair.second);
 			}
 
-			#if PLATFORM_WINDOWS
-			if (m_processJobObject != NULL)
-			{
-				SCOPED_WRITE_LOCK(m_processJobObjectLock, lock);
-				CloseHandle(m_processJobObject);
-				m_processJobObject = NULL;
-			}
-			#endif
-
 			for (auto& process : processes)
 				process.WaitForExit(100000);
 		}
diff --git a/Engine/Source/Programs/UnrealBuildAccelerator/Common/Public/UbaDirectoryIterator.h b/Engine/Source/Programs/UnrealBuildAccelerator/Common/Public/UbaDirectoryIterator.h
--- a/Engine/Source/Programs/UnrealBuildAccelerator/Common/Public/UbaDirectoryIterator.h
+++ b/Engine/Source/Programs/UnrealBuildAccelerator/Common/Public/UbaDirectoryIterator.h
@@ -92,7 +92,7 @@ namespace uba
 		NTSTATUS res = NtCreateFile(&handle, FILE_READ_ATTRIBUTES | FILE_LIST_DIRECTORY | FILE_ATTRIBUTE_UNPINNED, &ObjectAttributes, &IoStatusBlock, 0, FILE_ATTRIBUTE_NORMAL, ShareAccess, FILE_OPEN, FILE_DIRECTORY_FILE | FILE_SYNCHRONOUS_IO_ALERT, NULL, 0);
 		if (res == STATUS_OBJECT_NAME_NOT_FOUND || res == STATUS_OBJECT_PATH_NOT_FOUND)
 			return !errorOnNotFound;
-		if (res == STATUS_NOT_A_DIRECTORY || res == STATUS_ACCESS_DENIED || res == STATUS_NO_MEDIA_IN_DEVICE)
+		if (res == STATUS_NOT_A_DIRECTORY || res == STATUS_ACCESS_DENIED || res == STATUS_NO_MEDIA_IN_DEVICE || res == STATUS_FVE_LOCKED_VOLUME)
 			return false;
 		if (res != STATUS_SUCCESS)
 			return logger.Error(L"NtCreateFile for TraverseDir on %s failed with error code %x", dirPath, (DWORD)res);
