[UBA]
* Fixed code related to process reuse. It was not resetting variables properly
* Added support for detoured SHGetKnownFolderPath
* Fixed so logs for remote child processes are networked over to server when remoteLogs are enabled
* Improved some assert messages
* Fixed some binaryreaderwriter mismatches (they where harmless, just caused asserts)
* Fixed so local system temp folder can always be found when recursing from root to system temp. This is a horrible hack because it could cause bad conflicts between host file system and helper file system.. but solves current problems
* Moved removal of \\?\ into FixPath

diff --git a/Engine/Source/Programs/UnrealBuildAccelerator/Common/Private/UbaProcess.cpp b/Engine/Source/Programs/UnrealBuildAccelerator/Common/Private/UbaProcess.cpp
--- a/Engine/Source/Programs/UnrealBuildAccelerator/Common/Private/UbaProcess.cpp
+++ b/Engine/Source/Programs/UnrealBuildAccelerator/Common/Private/UbaProcess.cpp
@@ -1088,12 +1088,28 @@ namespace uba
 		if (!newProcess)
 			return true;
 
-		m_startTime = GetTime();
+		m_arguments = nextProcess.arguments;
+		m_description = nextProcess.description;
+		m_logFile = nextProcess.logFile;
+
+		m_startInfo.arguments = m_arguments.c_str();
+		m_startInfo.description = m_description.c_str();
+		m_startInfo.logFile = m_logFile.c_str();
+
+		m_childProcesses.clear();
+		m_logLines.clear();
+		m_trackedInputs.clear();
+		m_trackedOutputs.clear();
+		
+		ClearTempFiles();
+		
 		m_processStats = {};
 		m_sessionStats = {};
 		m_storageStats = {};
 		m_systemStats = {};
 
+		m_startTime = GetTime();
+
 		writer.WriteString(nextProcess.arguments);
 		writer.WriteString(nextProcess.workingDir);
 		writer.WriteString(nextProcess.description);
@@ -1107,6 +1123,12 @@ namespace uba
 		return true;
 	}
 
+	bool ProcessImpl::HandleSHGetKnownFolderPath(BinaryReader& reader, BinaryWriter& writer)
+	{
+		m_session.SHGetKnownFolderPath(*this, reader, writer);
+		return true;
+	}
+
 	bool ProcessImpl::CreateTempFile(BinaryReader& reader, ProcHandle nativeProcessHandle, const tchar* application)
 	{
 		StringKey key = reader.ReadStringKey();
diff --git a/Engine/Source/Programs/UnrealBuildAccelerator/Common/Private/UbaSession.cpp b/Engine/Source/Programs/UnrealBuildAccelerator/Common/Private/UbaSession.cpp
--- a/Engine/Source/Programs/UnrealBuildAccelerator/Common/Private/UbaSession.cpp
+++ b/Engine/Source/Programs/UnrealBuildAccelerator/Common/Private/UbaSession.cpp
@@ -2256,6 +2256,12 @@ namespace uba
 		return true;
 	}
 
+	bool Session::SHGetKnownFolderPath(Process& process, BinaryReader& reader, BinaryWriter& writer)
+	{
+		UBA_ASSERT(false); // Should only be called on UbaSessionClient
+		return false;
+	}
+
 	void Session::FileEntryAdded(StringKey fileNameKey, u64 lastWritten, u64 size)
 	{
 	}
diff --git a/Engine/Source/Programs/UnrealBuildAccelerator/Common/Private/UbaSessionClient.cpp b/Engine/Source/Programs/UnrealBuildAccelerator/Common/Private/UbaSessionClient.cpp
--- a/Engine/Source/Programs/UnrealBuildAccelerator/Common/Private/UbaSessionClient.cpp
+++ b/Engine/Source/Programs/UnrealBuildAccelerator/Common/Private/UbaSessionClient.cpp
@@ -9,6 +9,10 @@
 #include "UbaProtocol.h"
 #include "UbaStorage.h"
 
+#if PLATFORM_WINDOWS
+#include <winerror.h>
+#endif
+
 namespace uba
 {
 	SessionClient::SessionClient(const SessionClientCreateInfo& info)
@@ -166,7 +170,8 @@ namespace uba
 				return false;
 
 			rec.key = reader.ReadCasKey();
-			rec.serverTime = reader.ReadU64();
+			if (rec.key != CasKeyZero)
+				rec.serverTime = reader.ReadU64();
 		}
 		out = rec.key;
 		return true;
@@ -713,7 +718,7 @@ namespace uba
 			if (findIt != writtenFiles.end())
 			{
 				auto insres = writtenFiles.try_emplace(toName);
-				UBA_ASSERTF(insres.second, TC("Moving written file to other written file."));
+				UBA_ASSERTF(insres.second, TC("Moving written file %s to other written file %s."), fromName, toName);
 				insres.first->second = findIt->second;
 				insres.first->second.owner = &msg.process;
 				writtenFiles.erase(findIt);
@@ -727,7 +732,7 @@ namespace uba
 			if (findIt != m_outputFiles.end())
 			{
 				auto insres = m_outputFiles.try_emplace(toName);
-				UBA_ASSERTF(insres.second, TC("Failed to add move destination file %s as output file because it is already added."), toName);
+				UBA_ASSERTF(insres.second, TC("Failed to add move destination file %s as output file because it is already added. (Moved from %s)"), toName, fromName);
 				insres.first->second = findIt->second;
 				m_outputFiles.erase(findIt);
 				sendMove = false;
@@ -1396,6 +1401,8 @@ namespace uba
 		dest.Append(TC("<log>")).Append(logFile);
 		f.key = ToStringKeyLower(dest);
 		SendFile(f, dest.data, pi.GetId(), false);
+		for (auto& child : pi.m_childProcesses)
+			SendLogFileToServer(*(ProcessImpl*)child.m_process);
 	}
 
 	void SessionClient::GetLogFileName(StringBufferBase& out, const tchar* logFile, const tchar* arguments)
@@ -1847,32 +1854,16 @@ namespace uba
 			if (m_processFinished)
 				m_processFinished(&process);
 
-			pi.m_exitCode = ~0u;
-			pi.m_processStats = {};
-			pi.m_sessionStats = {};
-			pi.m_storageStats = {};
-			pi.m_systemStats = {};
-
 			outNextProcess.arguments = reader.ReadString();
 			outNextProcess.workingDir = reader.ReadString();
 			outNextProcess.description = reader.ReadString();
 			outNextProcess.logFile = reader.ReadString();
-
 			if (m_logToFile)
 			{
 				StringBuffer<512> logFile;
 				GetLogFileName(logFile, outNextProcess.logFile.c_str(), outNextProcess.arguments.c_str());
 				outNextProcess.logFile = logFile.data;
 			}
-
-			// TODO: Probably need to fill up with more stuff.. this is fine for current usecase
-			pi.m_arguments = outNextProcess.arguments;
-			pi.m_description = outNextProcess.description;
-			pi.m_logFile = outNextProcess.logFile;
-
-			pi.m_startInfo.arguments = pi.m_arguments.c_str();
-			pi.m_startInfo.description = pi.m_description.c_str();
-			pi.m_startInfo.logFile = pi.m_logFile.c_str();
 		}
 
 		reader.Reset();
@@ -1898,18 +1889,35 @@ namespace uba
 		return true;
 	}
 
+	bool SessionClient::SHGetKnownFolderPath(Process& process, BinaryReader& reader, BinaryWriter& writer)
+	{
+#if PLATFORM_WINDOWS
+		StackBinaryWriter<SendMaxSize> msgWriter;
+		NetworkMessage msg(m_client, ServiceId, SessionMessageType_SHGetKnownFolderPath, msgWriter);
+		msgWriter.WriteBytes(reader.GetPositionData(), reader.GetLeft());
+		BinaryReader msgReader(writer.GetData(), 0);
+		if (!msg.Send(msgReader, m_stats.customMsg))
+		{
+			writer.WriteU32(u32(E_FAIL));
+			return false;
+		}
+		writer.AllocWrite(msgReader.GetPosition());
+#endif
+		return true;
+	}
+
 	bool SessionClient::FlushWrittenFiles(ProcessImpl& process)
 	{
 		SCOPED_WRITE_LOCK(process.m_writtenFilesLock, lock);
-		if (!SendFiles(process, process.m_processStats.sendFiles))
-			return false;
+		bool success = SendFiles(process, process.m_processStats.sendFiles);
 		{
 			SCOPED_WRITE_LOCK(m_outputFilesLock, lock2);
 			for (auto& kv : process.m_writtenFiles)
 				m_outputFiles.erase(kv.first);
 		}
 		process.m_writtenFiles.clear();
-		return true;
+
+		return success;
 	}
 
 	bool SessionClient::UpdateEnvironment(ProcessImpl& process, const tchar* reason, bool resetStats)
diff --git a/Engine/Source/Programs/UnrealBuildAccelerator/Common/Private/UbaSessionServer.cpp b/Engine/Source/Programs/UnrealBuildAccelerator/Common/Private/UbaSessionServer.cpp
--- a/Engine/Source/Programs/UnrealBuildAccelerator/Common/Private/UbaSessionServer.cpp
+++ b/Engine/Source/Programs/UnrealBuildAccelerator/Common/Private/UbaSessionServer.cpp
@@ -1369,6 +1369,27 @@ namespace uba
 				writer.WriteByte(255);
 				return true;
 			}
+			case SessionMessageType_SHGetKnownFolderPath:
+			{
+#if PLATFORM_WINDOWS
+				GUID kfid;
+				reader.ReadBytes(&kfid, sizeof(GUID));
+				u32 flags = reader.ReadU32();
+				PWSTR str;
+
+				static HMODULE moduleHandle = LoadLibrary(L"Shell32.dll");
+				using SHGetKnownFolderPathFunc = HRESULT(const GUID& rfid, DWORD dwFlags, HANDLE hToken, PWSTR* ppszPath);
+				static SHGetKnownFolderPathFunc* SHGetKnownFolderPath = (SHGetKnownFolderPathFunc*)GetProcAddress(moduleHandle, "SHGetKnownFolderPath");;
+				HRESULT res = SHGetKnownFolderPath(kfid, flags, NULL, &str);
+				writer.WriteU32(res);
+				if (res == S_OK)
+				{
+					writer.WriteString(str);
+					CoTaskMemFree(str);
+				}
+#endif
+				return true;
+			}
 		}
 
 		UBA_ASSERT(false);
diff --git a/Engine/Source/Programs/UnrealBuildAccelerator/Common/Private/UbaStorageServer.cpp b/Engine/Source/Programs/UnrealBuildAccelerator/Common/Private/UbaStorageServer.cpp
--- a/Engine/Source/Programs/UnrealBuildAccelerator/Common/Private/UbaStorageServer.cpp
+++ b/Engine/Source/Programs/UnrealBuildAccelerator/Common/Private/UbaStorageServer.cpp
@@ -737,6 +737,7 @@ namespace uba
 						entryLock.Leave();
 						CasEntryAccessed(casEntry);
 						writer.WriteU16(u16(~0));
+						writer.WriteBool(m_traceStore);
 						return true;
 					}
 				}
diff --git a/Engine/Source/Programs/UnrealBuildAccelerator/Common/Public/UbaNetwork.h b/Engine/Source/Programs/UnrealBuildAccelerator/Common/Public/UbaNetwork.h
--- a/Engine/Source/Programs/UnrealBuildAccelerator/Common/Public/UbaNetwork.h
+++ b/Engine/Source/Programs/UnrealBuildAccelerator/Common/Public/UbaNetwork.h
@@ -84,6 +84,7 @@ namespace uba
 		UBA_SESSION_MESSAGE(UpdateEnvironment) \
 		UBA_SESSION_MESSAGE(Summary) \
 		UBA_SESSION_MESSAGE(Command) \
+		UBA_SESSION_MESSAGE(SHGetKnownFolderPath) \
 
 	enum SessionMessageType : u8 
 	{
diff --git a/Engine/Source/Programs/UnrealBuildAccelerator/Common/Public/UbaProcess.h b/Engine/Source/Programs/UnrealBuildAccelerator/Common/Public/UbaProcess.h
--- a/Engine/Source/Programs/UnrealBuildAccelerator/Common/Public/UbaProcess.h
+++ b/Engine/Source/Programs/UnrealBuildAccelerator/Common/Public/UbaProcess.h
@@ -96,6 +96,7 @@ namespace uba
 		bool HandleUpdateEnvironment(BinaryReader& reader, BinaryWriter& writer);
 		bool HandleGetNextProcess(BinaryReader& reader, BinaryWriter& writer);
 		bool HandleCustom(BinaryReader& reader, BinaryWriter& writer);
+		bool HandleSHGetKnownFolderPath(BinaryReader& reader, BinaryWriter& writer);
 
 		void LogLine(bool printInSession, TString&& line, LogEntryType logType);
 		bool CreateTempFile(BinaryReader& reader, ProcHandle nativeProcessHandle, const tchar* application);
diff --git a/Engine/Source/Programs/UnrealBuildAccelerator/Common/Public/UbaSession.h b/Engine/Source/Programs/UnrealBuildAccelerator/Common/Public/UbaSession.h
--- a/Engine/Source/Programs/UnrealBuildAccelerator/Common/Public/UbaSession.h
+++ b/Engine/Source/Programs/UnrealBuildAccelerator/Common/Public/UbaSession.h
@@ -131,6 +131,7 @@ namespace uba
 		virtual bool AllocFailed(Process& process, const tchar* allocType, u32 error);
 		virtual bool GetNextProcess(Process& process, bool& outNewProcess, NextProcessInfo& outNextProcess, u32 prevExitCode, BinaryReader& statsReader);
 		virtual bool CustomMessage(Process& process, BinaryReader& reader, BinaryWriter& writer);
+		virtual bool SHGetKnownFolderPath(Process& process, BinaryReader& reader, BinaryWriter& writer);
 		virtual void FileEntryAdded(StringKey fileNameKey, u64 lastWritten, u64 size);
 		virtual bool FlushWrittenFiles(ProcessImpl& process);
 		virtual bool UpdateEnvironment(ProcessImpl& process, const tchar* reason, bool resetStats);
diff --git a/Engine/Source/Programs/UnrealBuildAccelerator/Common/Public/UbaSessionClient.h b/Engine/Source/Programs/UnrealBuildAccelerator/Common/Public/UbaSessionClient.h
--- a/Engine/Source/Programs/UnrealBuildAccelerator/Common/Public/UbaSessionClient.h
+++ b/Engine/Source/Programs/UnrealBuildAccelerator/Common/Public/UbaSessionClient.h
@@ -62,6 +62,7 @@ namespace uba
 		virtual void PrintSessionStats(Logger& logger) override;
 		virtual bool GetNextProcess(Process& process, bool& outNewProcess, NextProcessInfo& outNextProcess, u32 prevExitCode, BinaryReader& statsReader) override;
 		virtual bool CustomMessage(Process& process, BinaryReader& reader, BinaryWriter& writer) override;
+		virtual bool SHGetKnownFolderPath(Process& process, BinaryReader& reader, BinaryWriter& writer) override;
 		virtual bool FlushWrittenFiles(ProcessImpl& process) override;
 		virtual bool UpdateEnvironment(ProcessImpl& process, const tchar* reason, bool resetStats) override;
 		virtual bool LogLine(ProcessImpl& process, const tchar* line, LogEntryType logType) override;
diff --git a/Engine/Source/Programs/UnrealBuildAccelerator/Core/Public/UbaPathUtils.h b/Engine/Source/Programs/UnrealBuildAccelerator/Core/Public/UbaPathUtils.h
--- a/Engine/Source/Programs/UnrealBuildAccelerator/Core/Public/UbaPathUtils.h
+++ b/Engine/Source/Programs/UnrealBuildAccelerator/Core/Public/UbaPathUtils.h
@@ -60,7 +60,7 @@ namespace uba
 				hasDotDot = true;
 			if (c == '\\')
 			{
-				if (lastChar == '?' && lastLastChar == '?') // We want to get rid of \\??\  .
+				if (lastChar == '?' && (lastLastChar == '?' || lastLastChar == '\\')) // We want to get rid of \\??\ and \\?\ .
 				{
 					write = buffer;
 					continue;
diff --git a/Engine/Source/Programs/UnrealBuildAccelerator/Core/Public/UbaProtocol.h b/Engine/Source/Programs/UnrealBuildAccelerator/Core/Public/UbaProtocol.h
--- a/Engine/Source/Programs/UnrealBuildAccelerator/Core/Public/UbaProtocol.h
+++ b/Engine/Source/Programs/UnrealBuildAccelerator/Core/Public/UbaProtocol.h
@@ -32,6 +32,7 @@ namespace uba
 		UBA_PROCESS_MESSAGE(UpdateEnvironment) \
 		UBA_PROCESS_MESSAGE(GetNextProcess) \
 		UBA_PROCESS_MESSAGE(Custom) \
+		UBA_PROCESS_MESSAGE(SHGetKnownFolderPath) \
 
 	enum MessageType : u8
 	{
diff --git a/Engine/Source/Programs/UnrealBuildAccelerator/Detours/Private/Windows/UbaDetoursFunctionsKernelBase.inl b/Engine/Source/Programs/UnrealBuildAccelerator/Detours/Private/Windows/UbaDetoursFunctionsKernelBase.inl
--- a/Engine/Source/Programs/UnrealBuildAccelerator/Detours/Private/Windows/UbaDetoursFunctionsKernelBase.inl
+++ b/Engine/Source/Programs/UnrealBuildAccelerator/Detours/Private/Windows/UbaDetoursFunctionsKernelBase.inl
@@ -46,7 +46,7 @@ DWORD Detoured_GetCurrentDirectoryW(DWORD nBufferLength, LPWSTR lpBuffer)
 	SetLastError(ERROR_SUCCESS);
 	if (lpBuffer == nullptr || nBufferLength < length + 1)
 	{
-		DEBUG_LOG_DETOURED(L"GetCurrentDirectoryW", L"(buffer too small: %u)", nBufferLength);
+		DEBUG_LOG_DETOURED(L"GetCurrentDirectoryW", L"(buffer too small: %u) -> %llu", nBufferLength, length + 1);
 		return DWORD(length + 1);
 	}
 	memcpy(lpBuffer, g_virtualWorkingDir.data, length * 2);
@@ -218,6 +218,7 @@ BOOL Detoured_RemoveDirectoryW(LPCWSTR lpPathName)
 	else
 	{
 		UBA_ASSERTF(!g_runningRemote, L"RemoveDirectory is not implemented for remote (removing %s)", lpPathName);
+		SetLastError(ERROR_PATH_NOT_FOUND);
 		res = false;
 	}
 	DEBUG_LOG_TRUE(L"RemoveDirectoryW", L"%ls -> %ls", lpPathName, ToString(res));
@@ -1284,6 +1285,12 @@ BOOL Detoured_DeleteFileW(LPCWSTR lpFileName)
 		return TRUE;
 	}
 
+	if (fixedName.StartsWith(g_systemTemp.data))
+	{
+		DEBUG_LOG_TRUE(L"DeleteFileW", L"(%ls)", original);
+		return True_DeleteFileW(original);
+	}
+
 	StringKey fileNameKey = ToStringKey(fixedNameLower);
 
 	u32 directoryTableSize;
@@ -1317,9 +1324,13 @@ BOOL Detoured_DeleteFileW(LPCWSTR lpFileName)
 bool Shared_MoveFile(LPCWSTR lpExistingFileName, LPCWSTR lpNewFileName, DWORD dwFlags)
 {
 	DETOURED_CALL(MoveFileExW);
+
 	StringBuffer<> source;
 	FixPath(source, lpExistingFileName);
 
+	StringBuffer<> dest;
+	FixPath(dest, lpNewFileName);
+
 	StringKey sourceKey = ToStringKeyLower(source);
 
 	if (KeepInMemory(source.data, source.count))
@@ -1330,9 +1341,6 @@ bool Shared_MoveFile(LPCWSTR lpExistingFileName, LPCWSTR lpNewFileName, DWORD dw
 		FileInfo& sourceInfo = it->second;
 		lock.Leave();
 
-		StringBuffer<> dest;
-		FixPath(dest, lpNewFileName);
-
 		if (IsOutputFile(dest.data, dest.count, true))
 		{
 			sourceInfo.deleted = true;
@@ -1370,8 +1378,6 @@ bool Shared_MoveFile(LPCWSTR lpExistingFileName, LPCWSTR lpNewFileName, DWORD dw
 		return written == toWrite;
 	}
 
-	StringBuffer<> dest;
-	FixPath(dest, lpNewFileName);
 	StringKey destKey = ToStringKeyLower(dest);
 
 	u32 directoryTableSize;
@@ -1493,16 +1499,14 @@ __forceinline HANDLE Shared_FindFirstFileExW(LPCWSTR lpFileName, FINDEX_INFO_LEV
 
 	StringBuffer<> lowerName;
 	FixPath(lowerName, lpFileName);
-	lowerName.MakeLower();
-	const wchar_t* buf = lowerName.data;
 
-	if (wcsncmp(buf, L"\\\\?\\", 4) == 0)
-		buf += 4;
-
-	if (wcsncmp(buf, g_systemTemp.data, g_systemTemp.count) == 0 || wcsncmp(buf, g_systemRoot.data, g_systemRoot.count) == 0)
+	if (lowerName.StartsWith(g_systemTemp.data) || lowerName.StartsWith(g_systemRoot.data))
 		return Local_FindFirstFileExW(lpFileName, fInfoLevelId, lpFindFileData, fSearchOp, lpSearchFilter, dwAdditionalFlags, funcName);
 
-	wchar_t* fileName = const_cast<wchar_t*>(buf); // Not beautiful but We know this is a different buffer.
+	lowerName.MakeLower();
+	wchar_t* buf = lowerName.data;
+
+	wchar_t* fileName = lowerName.data;
 	wchar_t* lastBackslash = wcsrchr(fileName, '\\');
 	if (lastBackslash)
 		fileName = lastBackslash + 1;
@@ -1558,6 +1562,9 @@ __forceinline HANDLE Shared_FindFirstFileExW(LPCWSTR lpFileName, FINDEX_INFO_LEV
 
 	if (!exists)
 	{
+		if (g_systemTemp.StartsWith(lowerName.data)) // TODO: This is a big hack. We should make sure the uba system temp folder is virtualized and is always some root path that never can collide with the host file system
+			return Local_FindFirstFileExW(lpFileName, fInfoLevelId, lpFindFileData, fSearchOp, lpSearchFilter, dwAdditionalFlags, funcName);
+
 		DEBUG_LOG_DETOURED(funcName, L"(%ls) -> NotFound", lpFileName);
 		SetLastError(ERROR_FILE_NOT_FOUND);
 		return INVALID_HANDLE_VALUE;
@@ -1598,8 +1605,10 @@ __forceinline HANDLE Shared_FindFirstFileExW(LPCWSTR lpFileName, FINDEX_INFO_LEV
 	{
 		if (!Shared_GetNextFile(data, *listHandle))
 		{
-			DEBUG_LOG_DETOURED(funcName, L"(%ls) -> NotFound", lpFileName);
 			delete listHandle;
+			if (g_systemTemp.StartsWith(lowerName.data)) // TODO: This is a big hack. We should make sure the uba system temp folder is virtualized and is always some root path that never can collide with the host file system
+				return Local_FindFirstFileExW(lpFileName, fInfoLevelId, lpFindFileData, fSearchOp, lpSearchFilter, dwAdditionalFlags, funcName);
+			DEBUG_LOG_DETOURED(funcName, L"(%ls) -> NotFound(2)", lpFileName);
 			return INVALID_HANDLE_VALUE;
 		}
 		if (listHandle->wildcard.empty() || PathMatchSpecW(data.cFileName, listHandle->wildcard.c_str()))
@@ -2389,8 +2398,7 @@ HMODULE Recursive_LoadLibraryExW(LPCWSTR lpLibFileName, LPCWSTR originalName, DW
 			SCOPED_WRITE_LOCK(g_loadedModulesLock, lock);
 			g_loadedModules[res] = originalName;
 		}
-		if (g_isRunningWine)
-			OnModuleLoaded(res, lpLibFileName);
+		OnModuleLoaded(res, lpLibFileName);
 	}
 	return res;
 }
diff --git a/Engine/Source/Programs/UnrealBuildAccelerator/Detours/Private/Windows/UbaDetoursFunctionsNtDll.inl b/Engine/Source/Programs/UnrealBuildAccelerator/Detours/Private/Windows/UbaDetoursFunctionsNtDll.inl
--- a/Engine/Source/Programs/UnrealBuildAccelerator/Detours/Private/Windows/UbaDetoursFunctionsNtDll.inl
+++ b/Engine/Source/Programs/UnrealBuildAccelerator/Detours/Private/Windows/UbaDetoursFunctionsNtDll.inl
@@ -600,8 +600,6 @@ NTSTATUS NTAPI Shared_NtCreateFile(bool IsCreateFunc, PHANDLE hFileHandle, ACCES
 		}
 		else if (createFileName)
 		{
-			if (StartsWith(createFileName, L"\\\\?\\"))
-				createFileName += 4;
 			if (!FixPath(fileName, createFileName))
 				UBA_ASSERTF(false, L"FixPath failed for string '%ls'", createFileName);
 			if (fileName.StartsWith(L"\\\\.\\pipe"))
@@ -708,10 +706,12 @@ NTSTATUS NTAPI Shared_NtCreateFile(bool IsCreateFunc, PHANDLE hFileHandle, ACCES
 	{
 		ObjectAttributes->RootDirectory = rootDir;
 		NTSTATUS res = Local_NtCreateFile(IsCreateFunc, hFileHandle, DesiredAccess, ObjectAttributes, IoStatusBlock, AllocationSize, FileAttributes, ShareAccess, CreateDisposition, CreateOptions, EaBuffer, EaLength);
-		DEBUG_LOG_TRUE(funcName, L"(NODETOUR)%ls %llu (%.*ls) -> %ls", isWriteStr, uintptr_t(*hFileHandle), ObjectAttributes->ObjectName->Length / 2, ObjectAttributes->ObjectName->Buffer, ToString(res));
-
+		if (NT_ERROR(res))
+			*hFileHandle = INVALID_HANDLE_VALUE;
+		DEBUG_LOG_TRUE(funcName, L"(NODETOUR)%ls %llu (%.*ls) -> %ls (%u)", isWriteStr, uintptr_t(*hFileHandle), ObjectAttributes->ObjectName->Length / 2, ObjectAttributes->ObjectName->Buffer, ToString(res), res);
 		if (NT_ERROR(res))
 			return res;
+
 		if (!isSystemOrTempFile && !isWrite && !t_disallowDetour && fileName[fileName.count-1] != ':')
 			TrackInput(fileName.data);
 		else
diff --git a/Engine/Source/Programs/UnrealBuildAccelerator/Detours/Private/Windows/UbaDetoursFunctionsShell32.inl b/Engine/Source/Programs/UnrealBuildAccelerator/Detours/Private/Windows/UbaDetoursFunctionsShell32.inl
new file mode 100644
--- /dev/null
+++ b/Engine/Source/Programs/UnrealBuildAccelerator/Detours/Private/Windows/UbaDetoursFunctionsShell32.inl
@@ -0,0 +1,40 @@
+// Copyright Epic Games, Inc. All Rights Reserved.
+
+
+//typedef struct _MODLOAD_DATA {} *PMODLOAD_DATA;
+using SHGetKnownFolderPathFunc = HRESULT(REFKNOWNFOLDERID rfid, DWORD dwFlags, HANDLE hToken, PWSTR* ppszPath);
+SHGetKnownFolderPathFunc* True_SHGetKnownFolderPath;
+
+BOOL Detoured_SHGetKnownFolderPath(REFKNOWNFOLDERID rfid, DWORD dwFlags, HANDLE hToken, PWSTR* ppszPath)
+{
+	if (g_runningRemote)
+	{
+		UBA_ASSERT(hToken == NULL);
+		TimerScope ts(g_stats.getFullFileName);
+		SCOPED_WRITE_LOCK(g_communicationLock, pcs);
+		BinaryWriter writer;
+		writer.WriteByte(MessageType_SHGetKnownFolderPath);
+		writer.WriteBytes(&rfid, sizeof(KNOWNFOLDERID));
+		writer.WriteU32(dwFlags);
+		writer.Flush();
+		BinaryReader reader;
+		HRESULT res = reader.ReadU32();
+		*ppszPath = NULL;
+		if (res == S_OK)
+		{
+			StringBuffer<> path;
+			reader.ReadString(path);
+			u32 memSize = (path.count+1)*2;
+			void* mem = CoTaskMemAlloc(memSize);
+			memcpy(mem, path.data, memSize);
+			*ppszPath = (PWSTR)mem;
+		}
+		DEBUG_LOG_DETOURED(L"SHGetKnownFolderPath", L"(%ls) -> %ls", *ppszPath, ToString(res == S_OK));
+		return res;
+	}
+
+	SuppressDetourScope _;
+	HRESULT res = True_SHGetKnownFolderPath(rfid, dwFlags, hToken, ppszPath);
+	DEBUG_LOG_TRUE(L"SHGetKnownFolderPath", L"(%ls) -> %ls", *ppszPath, ToString(res == S_OK));
+	return res;
+}
diff --git a/Engine/Source/Programs/UnrealBuildAccelerator/Detours/Private/Windows/UbaDetoursFunctionsWin.cpp b/Engine/Source/Programs/UnrealBuildAccelerator/Detours/Private/Windows/UbaDetoursFunctionsWin.cpp
--- a/Engine/Source/Programs/UnrealBuildAccelerator/Detours/Private/Windows/UbaDetoursFunctionsWin.cpp
+++ b/Engine/Source/Programs/UnrealBuildAccelerator/Detours/Private/Windows/UbaDetoursFunctionsWin.cpp
@@ -470,6 +470,7 @@ thread_local const wchar_t* t_createFileFileName;
 #include "UbaDetoursFunctionsUcrtBase.inl"
 #include "UbaDetoursFunctionsImagehlp.inl"
 #include "UbaDetoursFunctionsDbgHelp.inl"
+#include "UbaDetoursFunctionsShell32.inl"
 
 extern u32 g_consoleStringIndex;
 
@@ -596,10 +597,8 @@ int DetourAttachFunctions(bool runningRemote)
 
 void OnModuleLoaded(HMODULE moduleHandle, const wchar_t* name)
 {
-	UBA_ASSERT(g_isRunningWine);
-
 	// SymLoadModuleExW do something bad that cause remote wine to fail everything after this call.. TODO: Revisit
-	if (!True_SymLoadModuleExW && Contains(name, L"dbghelp.dll"))
+	if (g_isRunningWine && !True_SymLoadModuleExW && Contains(name, L"dbghelp.dll"))
 	{
 		True_SymLoadModuleExW = (SymLoadModuleExWFunc*)GetProcAddress(moduleHandle, "SymLoadModuleExW");
 		UBA_ASSERT(True_SymLoadModuleExW);
@@ -611,7 +610,7 @@ void OnModuleLoaded(HMODULE moduleHandle, const wchar_t* name)
 	}
 
 	// ImageGetDigestStream is buggy in wine so we have to detour it for ShaderCompileWorker
-	if (!True_ImageGetDigestStream && Contains(name, L"imagehlp.dll"))
+	if (g_isRunningWine && !True_ImageGetDigestStream && Contains(name, L"imagehlp.dll"))
 	{
 		True_ImageGetDigestStream = (ImageGetDigestStreamFunc*)GetProcAddress(moduleHandle, "ImageGetDigestStream");
 		UBA_ASSERT(True_ImageGetDigestStream);
@@ -621,6 +620,18 @@ void OnModuleLoaded(HMODULE moduleHandle, const wchar_t* name)
 		LONG error = DetourTransactionCommit(); (void)error;
 		UBA_ASSERT(!error);
 	}
+
+	// SHGetKnownFolderPath is used by Metal.exe and must always execute on host
+	if (!True_SHGetKnownFolderPath && Contains(name, L"shell32.dll"))
+	{
+		True_SHGetKnownFolderPath = (SHGetKnownFolderPathFunc*)GetProcAddress(moduleHandle, "SHGetKnownFolderPath");
+		UBA_ASSERT(True_SHGetKnownFolderPath);
+		DetourTransactionBegin();
+		DetourUpdateThread(GetCurrentThread());
+		DetourAttachFunction((PVOID*)&True_SHGetKnownFolderPath, Detoured_SHGetKnownFolderPath, "SHGetKnownFolderPath");
+		LONG error = DetourTransactionCommit(); (void)error;
+		UBA_ASSERT(!error);
+	}
 }
 
 int DetourDetachFunctions()
