[UBA]
* Added handler for SIGTERM to UbaCacheService and UbaCli

diff --git a/Engine/Source/Programs/UnrealBuildAccelerator/CacheService/Private/UbaCacheService.cpp b/Engine/Source/Programs/UnrealBuildAccelerator/CacheService/Private/UbaCacheService.cpp
--- a/Engine/Source/Programs/UnrealBuildAccelerator/CacheService/Private/UbaCacheService.cpp
+++ b/Engine/Source/Programs/UnrealBuildAccelerator/CacheService/Private/UbaCacheService.cpp
@@ -205,6 +205,7 @@ namespace uba
 		SetConsoleCtrlHandler(ConsoleHandler, TRUE);
 		#else
 		signal(SIGINT, ConsoleHandler);
+		signal(SIGTERM, ConsoleHandler);
 		#endif
 
 		NetworkBackendTcp networkBackend(logWriter);
diff --git a/Engine/Source/Programs/UnrealBuildAccelerator/Cli/Private/UbaCli.cpp b/Engine/Source/Programs/UnrealBuildAccelerator/Cli/Private/UbaCli.cpp
--- a/Engine/Source/Programs/UnrealBuildAccelerator/Cli/Private/UbaCli.cpp
+++ b/Engine/Source/Programs/UnrealBuildAccelerator/Cli/Private/UbaCli.cpp
@@ -544,6 +544,7 @@ namespace uba
 		SetConsoleCtrlHandler(ConsoleHandler, TRUE);
 		#else
 		signal(SIGINT, ConsoleHandler);
+		signal(SIGTERM, ConsoleHandler);
 		#endif
 
 		NetworkBackendTcp networkBackend(logWriter);
