[UBA]
* Added option to enable cache miss reason logging. (Useful when looking for determinism/portability issues)

diff --git a/Engine/Source/Programs/UnrealBuildAccelerator/Host/Private/UbaExports.cpp b/Engine/Source/Programs/UnrealBuildAccelerator/Host/Private/UbaExports.cpp
--- a/Engine/Source/Programs/UnrealBuildAccelerator/Host/Private/UbaExports.cpp
+++ b/Engine/Source/Programs/UnrealBuildAccelerator/Host/Private/UbaExports.cpp
@@ -505,7 +505,7 @@ extern "C"
 		scheduler->GetStats(outQueued, outActiveLocal, outActiveRemote, outFinished);
 	}
 
-	uba::CacheClient* CacheClient_Create(uba::SessionServer* session)
+	uba::CacheClient* CacheClient_Create(uba::SessionServer* session, bool reportMissReason)
 	{
 		using namespace uba;
 		LogWriter& writer = session->GetLogWriter();
@@ -520,7 +520,7 @@ extern "C"
 		if (!ctorSuccess)
 			return nullptr;
 		CacheClientCreateInfo info{writer, storage, *networkClient, *session};
-		//info.reportMissReason = true;
+		info.reportMissReason = reportMissReason;
 		return new CacheClient(info);
 	}
 
diff --git a/Engine/Source/Programs/UnrealBuildAccelerator/Host/Public/UbaExports.h b/Engine/Source/Programs/UnrealBuildAccelerator/Host/Public/UbaExports.h
--- a/Engine/Source/Programs/UnrealBuildAccelerator/Host/Public/UbaExports.h
+++ b/Engine/Source/Programs/UnrealBuildAccelerator/Host/Public/UbaExports.h
@@ -120,7 +120,7 @@ extern "C"
 	UBA_API void Scheduler_GetStats(uba::Scheduler* scheduler, uba::u32& outQueued, uba::u32& outActiveLocal, uba::u32& outActiveRemote, uba::u32& outFinished);
 
 	// Cache
-	UBA_API uba::CacheClient* CacheClient_Create(uba::SessionServer* session);
+	UBA_API uba::CacheClient* CacheClient_Create(uba::SessionServer* session, bool reportMissReason = false);
 	UBA_API bool CacheClient_Connect(uba::CacheClient* cacheClient, const uba::tchar* host, int port);
 	UBA_API bool CacheClient_WriteToCache(uba::CacheClient* cacheClient, uba::RootPaths* rootPaths, uba::u32 bucket, const uba::ProcessHandle* process, const uba::u8* inputs, uba::u32 inputsSize, const uba::u8* outputs, uba::u32 outputsSize);
 	UBA_API bool CacheClient_FetchFromCache(uba::CacheClient* cacheClient, uba::RootPaths* rootPaths, uba::u32 bucket, const uba::ProcessStartInfo& info);
