[UBA]
* Added support for cache client requesting server shutdown

diff --git a/Engine/Source/Programs/UnrealBuildAccelerator/CacheService/Private/UbaCacheService.cpp b/Engine/Source/Programs/UnrealBuildAccelerator/CacheService/Private/UbaCacheService.cpp
--- a/Engine/Source/Programs/UnrealBuildAccelerator/CacheService/Private/UbaCacheService.cpp
+++ b/Engine/Source/Programs/UnrealBuildAccelerator/CacheService/Private/UbaCacheService.cpp
@@ -242,7 +242,7 @@ namespace uba
 			if (!networkServer.StartListen(networkBackend, port, listenIp.data))
 				return -1;
 
-			while (!ShouldExit())
+			while (!ShouldExit() && !cacheServer.ShouldShutdown())
 			{
 				Sleep(1000);
 				if (!cacheServer.RunMaintenance(false, ShouldExit))
diff --git a/Engine/Source/Programs/UnrealBuildAccelerator/Common/Private/UbaCacheClient.cpp b/Engine/Source/Programs/UnrealBuildAccelerator/Common/Private/UbaCacheClient.cpp
--- a/Engine/Source/Programs/UnrealBuildAccelerator/Common/Private/UbaCacheClient.cpp
+++ b/Engine/Source/Programs/UnrealBuildAccelerator/Common/Private/UbaCacheClient.cpp
@@ -496,6 +496,17 @@ namespace uba
 		return false;
 	}
 
+	bool CacheClient::RequestServerShutdown(const tchar* reason)
+	{
+		StackBinaryWriter<1024> writer;
+		NetworkMessage msg(m_client, CacheServiceId, CacheMessageType_RequestShutdown, writer);
+		writer.WriteString(reason);
+		StackBinaryReader<512> reader;
+		if (!msg.Send(reader))
+			return false;
+		return reader.ReadBool();
+	}
+
 	bool CacheClient::WriteCacheSummary(const tchar* destinationFile, const tchar* filterString)
 	{
 		StackBinaryWriter<1024> writer;
diff --git a/Engine/Source/Programs/UnrealBuildAccelerator/Common/Private/UbaCacheServer.cpp b/Engine/Source/Programs/UnrealBuildAccelerator/Common/Private/UbaCacheServer.cpp
--- a/Engine/Source/Programs/UnrealBuildAccelerator/Common/Private/UbaCacheServer.cpp
+++ b/Engine/Source/Programs/UnrealBuildAccelerator/Common/Private/UbaCacheServer.cpp
@@ -633,6 +633,17 @@ namespace uba
 		return true;
 	}
 
+	bool CacheServer::ShouldShutdown()
+	{
+		if (!m_shutdownRequested)
+			return false;
+		SCOPED_WRITE_LOCK(m_maintenanceLock, lock);
+		SCOPED_READ_LOCK(m_connectionsLock, lock2);
+		if (!m_connections.empty() || m_addsSinceMaintenance)
+			return false;
+		return true;
+	}
+
 	void CacheServer::OnDisconnected(u32 clientId)
 	{
 		SCOPED_WRITE_LOCK(m_connectionsLock, lock);
@@ -723,6 +734,15 @@ namespace uba
 		case CacheMessageType_CreateStatusFile:
 			return HandleCreateStatusFile(reader, writer);
 
+		case CacheMessageType_RequestShutdown:
+		{
+			TString reason = reader.ReadString();
+			m_logger.Info(TC("Shutdown requested. Reason: %s"), reason.empty() ? TC("Unknown") : reason.c_str());
+			m_shutdownRequested = true;
+			writer.WriteBool(true);
+			return true;
+		}
+
 		default:
 			return false;
 		}
diff --git a/Engine/Source/Programs/UnrealBuildAccelerator/Common/Public/UbaCacheClient.h b/Engine/Source/Programs/UnrealBuildAccelerator/Common/Public/UbaCacheClient.h
--- a/Engine/Source/Programs/UnrealBuildAccelerator/Common/Public/UbaCacheClient.h
+++ b/Engine/Source/Programs/UnrealBuildAccelerator/Common/Public/UbaCacheClient.h
@@ -36,6 +36,7 @@ namespace uba
 
 		bool WriteToCache(const RootPaths& rootPaths, u32 bucketId, const ProcessStartInfo& info, const u8* inputs, u64 inputsSize, const u8* outputs, u64 outputsSize);
 		bool FetchFromCache(const RootPaths& rootPaths, u32 bucketId, const ProcessStartInfo& info);
+		bool RequestServerShutdown(const tchar* reason);
 
 		bool WriteCacheSummary(const tchar* destinationFile, const tchar* filterString = nullptr);
 
diff --git a/Engine/Source/Programs/UnrealBuildAccelerator/Common/Public/UbaCacheServer.h b/Engine/Source/Programs/UnrealBuildAccelerator/Common/Public/UbaCacheServer.h
--- a/Engine/Source/Programs/UnrealBuildAccelerator/Common/Public/UbaCacheServer.h
+++ b/Engine/Source/Programs/UnrealBuildAccelerator/Common/Public/UbaCacheServer.h
@@ -24,6 +24,8 @@ namespace uba
 
 		bool RunMaintenance(bool force, const Function<bool()>& shouldExit);
 
+		bool ShouldShutdown();
+
 	private:
 		bool SaveNoLock();
 		void OnDisconnected(u32 clientId);
@@ -59,6 +61,8 @@ namespace uba
 		ReaderWriterLock m_connectionsLock;
 		Map<u32, Connection> m_connections;
 
+		Atomic<bool> m_shutdownRequested = false;
+
 		CacheServer(const CacheServer&) = delete;
 		CacheServer& operator=(const CacheServer&) = delete;
 	};
diff --git a/Engine/Source/Programs/UnrealBuildAccelerator/Common/Public/UbaNetwork.h b/Engine/Source/Programs/UnrealBuildAccelerator/Common/Public/UbaNetwork.h
--- a/Engine/Source/Programs/UnrealBuildAccelerator/Common/Public/UbaNetwork.h
+++ b/Engine/Source/Programs/UnrealBuildAccelerator/Common/Public/UbaNetwork.h
@@ -111,6 +111,7 @@ namespace uba
 		UBA_CACHE_MESSAGE(FetchCasTable) \
 		UBA_CACHE_MESSAGE(FetchEntries) \
 		UBA_CACHE_MESSAGE(CreateStatusFile) \
+		UBA_CACHE_MESSAGE(RequestShutdown) \
 
 	enum CacheMessageType : u8
 	{
diff --git a/Engine/Source/Programs/UnrealBuildAccelerator/Host/Private/UbaExports.cpp b/Engine/Source/Programs/UnrealBuildAccelerator/Host/Private/UbaExports.cpp
--- a/Engine/Source/Programs/UnrealBuildAccelerator/Host/Private/UbaExports.cpp
+++ b/Engine/Source/Programs/UnrealBuildAccelerator/Host/Private/UbaExports.cpp
@@ -550,6 +550,11 @@ extern "C"
 		return cacheClient->FetchFromCache(*rootPaths, bucket, info);
 	}
 
+	void CacheClient_RequestServerShutdown(uba::CacheClient* cacheClient, const uba::tchar* reason)
+	{
+		cacheClient->RequestServerShutdown(reason);
+	}
+
 	void CacheClient_Destroy(uba::CacheClient* cacheClient)
 	{
 		using namespace uba;
diff --git a/Engine/Source/Programs/UnrealBuildAccelerator/Host/Public/UbaExports.h b/Engine/Source/Programs/UnrealBuildAccelerator/Host/Public/UbaExports.h
--- a/Engine/Source/Programs/UnrealBuildAccelerator/Host/Public/UbaExports.h
+++ b/Engine/Source/Programs/UnrealBuildAccelerator/Host/Public/UbaExports.h
@@ -124,6 +124,7 @@ extern "C"
 	UBA_API bool CacheClient_Connect(uba::CacheClient* cacheClient, const uba::tchar* host, int port);
 	UBA_API bool CacheClient_WriteToCache(uba::CacheClient* cacheClient, uba::RootPaths* rootPaths, uba::u32 bucket, const uba::ProcessHandle* process, const uba::u8* inputs, uba::u32 inputsSize, const uba::u8* outputs, uba::u32 outputsSize);
 	UBA_API bool CacheClient_FetchFromCache(uba::CacheClient* cacheClient, uba::RootPaths* rootPaths, uba::u32 bucket, const uba::ProcessStartInfo& info);
+	UBA_API void CacheClient_RequestServerShutdown(uba::CacheClient* cacheClient, const uba::tchar* reason);
 	UBA_API void CacheClient_Destroy(uba::CacheClient* cacheClient);
 
 	// Misc
