[UBA]
* Fixed so additional error logging for missing files is not including files that does not match name

diff --git a/Engine/Source/Programs/UnrealBuildAccelerator/Common/Private/UbaSessionClient.cpp b/Engine/Source/Programs/UnrealBuildAccelerator/Common/Private/UbaSessionClient.cpp
--- a/Engine/Source/Programs/UnrealBuildAccelerator/Common/Private/UbaSessionClient.cpp
+++ b/Engine/Source/Programs/UnrealBuildAccelerator/Common/Private/UbaSessionClient.cpp
@@ -1991,8 +1991,11 @@ namespace uba
 		u32 foundCount = 0;
 		dirTable.TraverseAllFilesNoLock([&](const DirectoryTable::EntryInformation& info, const StringBufferBase& path)
 			{
-				if (!path.Contains(searchString.data))
+				if (!path.EndsWith(searchString.data))
 					return;
+				if (path[path.count - searchString.count - 1] != PathSeparator)
+					return;
+
 				auto ToString = [](bool b) { return b ? TC("true") : TC("false"); };
 
 				++foundCount;
