[UBA]
* Moved ReadLines to UbaFile

diff --git a/Engine/Source/Programs/UnrealBuildAccelerator/Common/Private/UbaScheduler.cpp b/Engine/Source/Programs/UnrealBuildAccelerator/Common/Private/UbaScheduler.cpp
--- a/Engine/Source/Programs/UnrealBuildAccelerator/Common/Private/UbaScheduler.cpp
+++ b/Engine/Source/Programs/UnrealBuildAccelerator/Common/Private/UbaScheduler.cpp
@@ -555,51 +555,6 @@ namespace uba
 		m_session.UpdateStatus(3, 1, TC("Finished"), 4, StringBuffer<32>().Appendf(TC("%u"), m_finishedProcesses.load()).data, LogEntryType_Info);
 	}
 
-	template<typename LineFunc>
-	bool ReadLines(Logger& logger, const tchar* file, const LineFunc& lineFunc)
-	{
-		FileHandle handle;
-		if (!OpenFileSequentialRead(logger, file, handle))
-			return logger.Error(TC("Failed to open file %s"), file);
-		auto fg = MakeGuard([&]() { CloseFile(file, handle); });
-		u64 fileSize = 0;
-		if (!GetFileSizeEx(fileSize, handle))
-			return logger.Error(TC("Failed to get size of file %s"), file);
-		char buffer[512];
-		u64 left = fileSize;
-		std::string line;
-		while (left)
-		{
-			u64 toRead = Min(left, u64(sizeof(buffer)));
-			left -= toRead;
-			if (!ReadFile(logger, file, handle, buffer, toRead))
-				return false;
-
-			u64 start = 0;
-			for (u64 i=0;i!=toRead;++i)
-			{
-				if (buffer[i] != '\n')
-					continue;
-				u64 end = i;
-				if (i > 0 && buffer[i-1] == '\r')
-					--end;
-				line.append(buffer + start, end - start);
-				if (!line.empty())
-					if (!lineFunc(TString(line.begin(), line.end())))
-						return false;
-				line.clear();
-				start = i + 1;
-			}
-			if (toRead && buffer[toRead - 1] == '\r')
-				--toRead;
-			line.append(buffer + start, toRead - start);
-		}
-		if (!line.empty())
-			if (!lineFunc(TString(line.begin(), line.end())))
-				return false;
-		return true;
-	}
-
 	bool Scheduler::EnqueueFromFile(const tchar* yamlFilename)
 	{
 		auto& logger = m_session.GetLogger();
diff --git a/Engine/Source/Programs/UnrealBuildAccelerator/Common/Public/UbaFile.h b/Engine/Source/Programs/UnrealBuildAccelerator/Common/Public/UbaFile.h
--- a/Engine/Source/Programs/UnrealBuildAccelerator/Common/Public/UbaFile.h
+++ b/Engine/Source/Programs/UnrealBuildAccelerator/Common/Public/UbaFile.h
@@ -2,13 +2,12 @@
 
 #pragma once
 
+#include "UbaLogger.h"
 #include "UbaMemory.h"
 #include "UbaPlatform.h"
 
 namespace uba
 {
-	class Logger;
-
 	#if !PLATFORM_WINDOWS
 	inline constexpr u32 ERROR_FILE_NOT_FOUND = ENOENT;
 	inline constexpr u32 ERROR_PATH_NOT_FOUND = ENOENT;
@@ -90,4 +89,50 @@ namespace uba
 		struct CreatedDir { ReaderWriterLock lock; bool handled = false; };
 		UnorderedMap<TString, CreatedDir> m_createdDirs;
 	};
+
+
+	template<typename LineFunc>
+	bool ReadLines(Logger& logger, const tchar* file, const LineFunc& lineFunc)
+	{
+		FileHandle handle;
+		if (!OpenFileSequentialRead(logger, file, handle))
+			return logger.Error(TC("Failed to open file %s"), file);
+		auto fg = MakeGuard([&]() { CloseFile(file, handle); });
+		u64 fileSize = 0;
+		if (!GetFileSizeEx(fileSize, handle))
+			return logger.Error(TC("Failed to get size of file %s"), file);
+		char buffer[512];
+		u64 left = fileSize;
+		std::string line;
+		while (left)
+		{
+			u64 toRead = Min(left, u64(sizeof(buffer)));
+			left -= toRead;
+			if (!ReadFile(logger, file, handle, buffer, toRead))
+				return false;
+
+			u64 start = 0;
+			for (u64 i=0;i!=toRead;++i)
+			{
+				if (buffer[i] != '\n')
+					continue;
+				u64 end = i;
+				if (i > 0 && buffer[i-1] == '\r')
+					--end;
+				line.append(buffer + start, end - start);
+				if (!line.empty())
+					if (!lineFunc(TString(line.begin(), line.end())))
+						return false;
+				line.clear();
+				start = i + 1;
+			}
+			if (toRead && buffer[toRead - 1] == '\r')
+				--toRead;
+			line.append(buffer + start, toRead - start);
+		}
+		if (!line.empty())
+			if (!lineFunc(TString(line.begin(), line.end())))
+				return false;
+		return true;
+	}
 }
