[UBA]
* Changed NT_SUCCESS to BCRYPT_SUCCESS for bcrypt code

diff --git a/Engine/Source/Programs/UnrealBuildAccelerator/Common/Private/UbaCrypto.cpp b/Engine/Source/Programs/UnrealBuildAccelerator/Common/Private/UbaCrypto.cpp
--- a/Engine/Source/Programs/UnrealBuildAccelerator/Common/Private/UbaCrypto.cpp
+++ b/Engine/Source/Programs/UnrealBuildAccelerator/Common/Private/UbaCrypto.cpp
@@ -26,7 +26,7 @@ namespace uba
 #if UBA_CRYPTO_TYPE == 1
 		BCRYPT_ALG_HANDLE providerHandle = NULL;
 		NTSTATUS res = BCryptOpenAlgorithmProvider(&providerHandle, BCRYPT_AES_ALGORITHM, NULL, 0);
-		if (!NT_SUCCESS(res))
+		if (!BCRYPT_SUCCESS(res))
 		{
 			logger.Error(L"ERROR: BCryptOpenAlgorithmProvider - Failed to open aes algorithm (0x%x)", res);
 			return InvalidCryptoKey;
@@ -41,7 +41,7 @@ namespace uba
 		//ULONG size = 0;
 		//u32 len = 0;
 		//NTSTATUS ret = BCryptGetProperty(providerHandle, BCRYPT_OBJECT_LENGTH, (UCHAR*)&len, sizeof(len), &size, 0);
-		//UBA_ASSERT(NT_SUCCESS(ret));
+		//UBA_ASSERT(BCRYPT_SUCCESS(ret));
 		//void* buf = calloc(1, len);
 		u32 objectBufferLen = 0;
 		u8* objectBuffer = nullptr;
@@ -49,7 +49,7 @@ namespace uba
 
 		BCRYPT_KEY_HANDLE keyHandle = NULL;
 		res = BCryptGenerateSymmetricKey(providerHandle, &keyHandle, objectBuffer, objectBufferLen, (u8*)key128, kAesBytes128, 0);
-		if (!NT_SUCCESS(res))
+		if (!BCRYPT_SUCCESS(res))
 		{
 			logger.Error(L"ERROR: BCryptGenerateSymmetricKey - Failed to generate symmetric key (0x%x)", res);
 			return InvalidCryptoKey;
@@ -69,7 +69,7 @@ namespace uba
 		u32 objectBufferLen = 0;
 		u8* objectBuffer = nullptr;
 		NTSTATUS res = BCryptDuplicateKey((BCRYPT_KEY_HANDLE)original, &newKey, objectBuffer, objectBufferLen, 0);
-		if (NT_SUCCESS(res))
+		if (BCRYPT_SUCCESS(res))
 			return (CryptoKey)(u64)newKey;
 		logger.Error(L"ERROR: BCryptDuplicateKey failed (0x%x)", res);
 		return InvalidCryptoKey;
@@ -94,7 +94,7 @@ namespace uba
 		u32 alignedSize = (size / kAesBytes128) * kAesBytes128;
 		BCRYPT_KEY_HANDLE newKey;
 		NTSTATUS res = BCryptDuplicateKey((BCRYPT_KEY_HANDLE)key, &newKey, objectBuffer, objectBufferLen, 0);
-		if (!NT_SUCCESS(res))
+		if (!BCRYPT_SUCCESS(res))
 		{
 			logger.Error(L"ERROR: BCryptDuplicateKey failed (0x%x)", res);
 			return false;
@@ -108,7 +108,7 @@ namespace uba
 			BCryptEncrypt(newKey, data, alignedSize, NULL, initVector, initVectorSize, data, alignedSize, &cipherTextLength, 0) :
 			BCryptDecrypt(newKey, data, alignedSize, NULL, initVector, initVectorSize, data, alignedSize, &cipherTextLength, 0);
 
-		if (!NT_SUCCESS(res))
+		if (!BCRYPT_SUCCESS(res))
 		{
 			logger.Error(L"ERROR: %s failed (0x%x)", (encrypt ? L"BCryptEncrypt" : L"BCryptDecrypt"), res);
 			return false;
