[UBA]
* Object file utilities - strategical submit just to be able to be able to roll back to something that works

diff --git a/Engine/Source/Programs/UnrealBuildAccelerator/Common/Private/UbaObjectFile.cpp b/Engine/Source/Programs/UnrealBuildAccelerator/Common/Private/UbaObjectFile.cpp
new file mode 100644
--- /dev/null
+++ b/Engine/Source/Programs/UnrealBuildAccelerator/Common/Private/UbaObjectFile.cpp
@@ -0,0 +1,497 @@
+// Copyright Epic Games, Inc. All Rights Reserved.
+
+#include "UbaObjectFile.h"
+#include "UbaFileAccessor.h"
+
+namespace uba
+{
+	constexpr u16 ImageFileMachineUnknown = 0;
+	constexpr u8 ImageSizeofShortName = 8;
+	constexpr u16 ImageSymClassExternal = 0x0002; // IMAGE_SYM_CLASS_EXTERNAL
+	constexpr u16 ImageSymUndefined = 0; // IMAGE_SYM_UNDEFINED
+
+#pragma pack(push)
+#pragma pack(1)
+
+	struct ImageFileHeader // IMAGE_FILE_HEADER
+	{
+		u16 Machine;
+		u16 NumberOfSections;
+		u32 TimeDateStamp;
+		u32 PointerToSymbolTable;
+		u32 NumberOfSymbols;
+		u16 SizeOfOptionalHeader;
+		u16 Characteristics;
+	};
+
+	struct AnonObjectHeaderBigobj // ANON_OBJECT_HEADER_BIGOBJ
+	{
+		u16 Sig1;            // Must be IMAGE_FILE_MACHINE_UNKNOWN
+		u16 Sig2;            // Must be 0xffff
+		u16 Version;         // >= 2 (implies the Flags field is present)
+		u16 Machine;         // Actual machine - IMAGE_FILE_MACHINE_xxx
+		u32 TimeDateStamp;
+		Guid ClassID;         // {D1BAA1C7-BAEE-4ba9-AF20-FAF66AA4DCB8}
+		u32 SizeOfData;      // Size of data that follows the header
+		u32 Flags;           // 0x1 -> contains metadata
+		u32 MetaDataSize;    // Size of CLR metadata
+		u32 MetaDataOffset;  // Offset of CLR metadata
+
+		// bigobj specifics
+		u32 NumberOfSections; // extended from WORD
+		u32 PointerToSymbolTable;
+		u32 NumberOfSymbols;
+	};
+
+	struct ImageSectionHeader // IMAGE_SECTION_HEADER
+	{
+		u8 Name[ImageSizeofShortName];
+		union
+		{
+			u32 PhysicalAddress;
+			u32 VirtualSize;
+		} Misc;
+		u32 VirtualAddress;
+		u32 SizeOfRawData;
+		u32 PointerToRawData;
+		u32 PointerToRelocations;
+		u32 PointerToLinenumbers;
+		u16 NumberOfRelocations;
+		u16 NumberOfLinenumbers;
+		u32 Characteristics;
+	};
+	static_assert(sizeof(ImageSectionHeader) == 40);
+
+	struct ImageRelocation // IMAGE_RELOCATION
+	{
+		union
+		{
+			u32 VirtualAddress;
+			u32 RelocCount;             // Set to the real count when IMAGE_SCN_LNK_NRELOC_OVFL is set
+		} DUMMYUNIONNAME;
+		u32 SymbolTableIndex;
+		u16 Type;
+	};
+
+
+	struct ImageSymbolEx // IMAGE_SYMBOL_EX
+	{
+		union
+		{
+			char ShortName[8];
+			struct
+			{
+				u32 Short; // if 0, use LongName
+				u32 Long; // offset into string table
+			} Name;
+			u32 LongName[2]; // PBYTE  [2]
+		} N;
+		u32   Value;
+		u32    SectionNumber;
+		u16    Type;
+		u8    StorageClass;
+		u8    NumberOfAuxSymbols;
+	};
+
+	struct ImageSymbol // IMAGE_SYMBOL
+	{
+		union
+		{
+			u8 ShortName[8];
+			struct
+			{
+				u32   Short;     // if 0, use LongName
+				u32   Long;      // offset into string table
+			} Name;
+			u32   LongName[2];    // PBYTE [2]
+		} N;
+		u32 Value;
+		u16 SectionNumber;
+		u16 Type;
+		u8 StorageClass;
+		u8 NumberOfAuxSymbols;
+	};
+
+#pragma pack(pop)
+
+	bool IsBigObj(u8* data, u64 size)
+	{
+		if (size < sizeof(AnonObjectHeaderBigobj))
+			return false;
+		auto& header = *(AnonObjectHeaderBigobj *)data;
+		if (header.Sig1 != ImageFileMachineUnknown)
+			return false;
+		if (header.Sig2 != 0xffff)
+			return false;
+		if (header.Version < 2)
+			return false;
+		constexpr u8 bigObjClassId[16] = { 0xc7, 0xa1, 0xba, 0xd1, 0xee, 0xba, 0xa9, 0x4b, 0xaf, 0x20, 0xfa, 0xf6, 0x6a, 0xa4, 0xdc, 0xb8 };
+		if (header.ClassID != *(const Guid*)bigObjClassId)
+			return false;
+		return true;
+	}
+
+	struct StringView
+	{
+		const char* strBegin;
+		const char* strEnd;
+
+		u32 Length() const
+		{
+			return u32(strEnd - strBegin);
+		}
+
+		bool StartsWith(const char* str, u32 strLen) const
+		{
+			if (strLen > Length())
+				return false;
+			return memcmp(strBegin, str, strLen) == 0;
+		}
+
+		std::string ToString() const
+		{
+			return std::string(strBegin, strEnd);
+		}
+	};
+
+	template<typename SymbolType>
+	StringView GetSymbolName(SymbolType& symbol, u8* stringTableMem)
+	{
+		if (symbol.N.Name.Short == 0)
+		{
+			const char* name = (char*)stringTableMem + symbol.N.Name.Long;
+			return { name, name + strlen(name) };
+		}
+		auto shortName = (char*)symbol.N.ShortName;
+		return { shortName, shortName + strnlen(shortName, ImageSizeofShortName) };
+	}
+
+	ObjectFile::~ObjectFile()
+	{
+		delete m_file;
+	}
+
+	bool ObjectFile::Parse(Logger& logger, const tchar* filename)
+	{
+		m_file = new FileAccessor(logger, filename);
+		if (!m_file->OpenMemoryRead())
+			return false;
+
+		m_data = m_file->GetData();
+		m_dataSize = m_file->GetSize();
+		m_isBigObj = IsBigObj(m_data, m_dataSize);
+
+		if (m_isBigObj)
+		{
+			auto& header = *(AnonObjectHeaderBigobj*)m_data;
+			m_symbolsMem = m_data + header.PointerToSymbolTable;
+			m_symbolCount = header.NumberOfSymbols;
+			m_stringTableMem = m_data + header.PointerToSymbolTable + header.NumberOfSymbols * sizeof(ImageSymbolEx);
+			m_sectionsMem = m_data + sizeof(AnonObjectHeaderBigobj);
+			m_sectionCount = header.NumberOfSections;
+		}
+		else
+		{
+			auto& header = *(ImageFileHeader*)m_data;
+			m_symbolsMem = m_data + header.PointerToSymbolTable;
+			m_symbolCount = header.NumberOfSymbols;
+			m_stringTableMem = m_data + header.PointerToSymbolTable + header.NumberOfSymbols * sizeof(ImageSymbol);
+			m_sectionsMem = m_data + sizeof(ImageFileHeader);
+			m_sectionCount = header.NumberOfSections;
+		}
+
+		auto sections = (ImageSectionHeader*)m_sectionsMem;
+		for (u32 i=0; i!=m_sectionCount; ++i)
+		{
+			if (strncmp((char*)sections[i].Name, ".drectve", 8) != 0)
+				continue;
+			m_directiveSectionMem = (u8*)(sections + i);
+			break;
+		}
+
+		ParseExports();
+
+		if (m_isBigObj)
+			ParseImports<ImageSymbolEx>();
+		else
+			ParseImports<ImageSymbol>();
+		return true;
+	}
+
+	bool ObjectFile::CreateStripped(Logger& logger, const tchar* newFilename, const UnorderedSymbols& allNeededImports, const UnorderedSymbols& allSharedExports)
+	{
+		if (!m_file)
+			return false;
+		FileAccessor newFile(logger, newFilename);
+		if (!newFile.CreateMemoryWrite(false, DefaultAttributes(), m_dataSize))
+			return false;
+
+		u8* newData = newFile.GetData();
+
+		memcpy(newData, m_data, m_dataSize);
+
+		WriteExports(logger, newData, allNeededImports);
+
+		if (m_isBigObj)
+			WriteImports<ImageSymbolEx>(logger, newData, allSharedExports);
+		else
+			WriteImports<ImageSymbol>(logger, newData, allSharedExports);
+
+		if (!newFile.Close())
+			return false;
+		return true;
+	}
+
+	const tchar* ObjectFile::GetFileName() const
+	{
+		return m_file->GetFileName();
+	}
+
+	const UnorderedSymbols& ObjectFile::GetImports() const
+	{
+		return m_imports;
+	}
+
+	const UnorderedSymbols& ObjectFile::GetExports() const
+	{
+		return m_exports;
+	}
+
+	void ObjectFile::ParseExports()
+	{
+		if (!m_directiveSectionMem)
+			return;
+		auto directiveSection = (ImageSectionHeader*)m_directiveSectionMem;
+		u8* directiveData = m_data + directiveSection->PointerToRawData;
+
+		static constexpr u8 utf8Bom[3] = { 0xef, 0xbb, 0xbf };
+		UBA_ASSERT(memcmp(directiveData, utf8Bom, 3) != 0);
+
+		auto str = (char*)directiveData;
+		while (str)
+		{
+			char* exportStr = strstr(str, "/EXPORT:");
+			if (!exportStr)
+				break;
+			exportStr += 8;
+			char* exportEnd = strchr(exportStr, ' ');
+			str = exportEnd;
+			if (!exportEnd)
+				exportEnd = exportStr + strlen(exportStr);
+			else
+				++str;
+			if (strncmp(exportEnd-5, ",DATA", 5) == 0)
+				exportEnd -= 5;
+					
+			m_exports.emplace(std::string(exportStr, exportEnd - exportStr));
+		}
+	}
+
+	template<typename SymbolType>
+	void ObjectFile::ParseImports()
+	{
+		auto symbols = (SymbolType*)m_symbolsMem;
+		for (u32 i=0; i!=m_symbolCount; ++i)
+		{
+			auto& symbol = symbols[i];
+			if (symbol.StorageClass != ImageSymClassExternal)
+				continue;
+			if (symbol.SectionNumber != ImageSymUndefined)
+				continue;
+			StringView symbolName = GetSymbolName(symbol, m_stringTableMem);
+			if (symbolName.StartsWith("__imp_", 6))
+				symbolName.strBegin += 6;
+			m_imports.emplace(std::string(symbolName.strBegin, symbolName.Length()));
+		}
+	}
+
+	void ObjectFile::WriteExports(Logger& logger, u8* newData, const UnorderedSymbols& allNeededImports)
+	{
+		if (!m_directiveSectionMem)
+			return;
+		auto directiveSection = (ImageSectionHeader*)m_directiveSectionMem;
+		if (directiveSection->SizeOfRawData < 10)
+			return;
+
+
+		const u8* directiveData = m_data + directiveSection->PointerToRawData;
+
+		u8* newDirectiveData = newData + (directiveData - m_data);
+
+		char* writePos = (char*)newDirectiveData;
+		const char* lastCopyPos = (const char*)directiveData;
+
+		auto readPos = lastCopyPos;
+		auto readEnd = readPos + directiveSection->SizeOfRawData;
+		auto readLastPossiblePos = readEnd - 9;
+
+		while (true)
+		{
+			const char* exportStr = nullptr;
+			const char* it = readPos;
+
+			while (it < readLastPossiblePos)
+			{
+				if (memcmp(it, "/EXPORT:", 8) != 0)
+				{
+					++it;
+					continue;
+				}
+
+				exportStr = it;
+				break;
+			}
+
+			if (!exportStr)
+			{
+				readPos = readEnd - 1;
+				break;
+			}
+
+			const char* startPos = exportStr;
+			exportStr += 8;
+			const char* exportEnd = strchr(exportStr, ' ');
+			if (!exportEnd)
+				readPos = exportEnd = exportStr + strlen(exportStr);
+			else
+				readPos = exportEnd;
+
+			if (strncmp(exportEnd-5, ",DATA", 5) == 0)
+				exportEnd -= 5;
+
+			if (allNeededImports.find(std::string(exportStr, exportEnd - exportStr)) != allNeededImports.end())
+				continue;
+			
+			// TODO: Use exclusion list provided from command line
+			if (strstr(exportStr, "Initialize") == exportStr && strstr(exportStr, "Module"))
+				continue;
+
+			u64 toCopy = startPos - lastCopyPos - 1;
+			memcpy(writePos, lastCopyPos, toCopy);
+			writePos += toCopy;
+			lastCopyPos = readPos;
+			if (!*readPos)
+				break;
+		}
+
+		u64 toCopy = readPos - lastCopyPos;
+		memcpy(writePos, lastCopyPos, toCopy);
+		writePos += toCopy;
+		*writePos = 0;
+
+		auto newDirectiveSection = (ImageSectionHeader*)(newData + (m_directiveSectionMem - m_data));
+		newDirectiveSection->SizeOfRawData = u32((u8*)writePos - newDirectiveData);
+		UBA_ASSERT(newDirectiveSection->SizeOfRawData <= directiveSection->SizeOfRawData);
+
+		memset(writePos, 0, directiveSection->SizeOfRawData - newDirectiveSection->SizeOfRawData);
+	}
+
+	template<typename SymbolType>
+	void ObjectFile::WriteImports(Logger& logger, u8* newData, const UnorderedSymbols& allSharedExports)
+	{
+		/*
+		auto sections = (ImageSectionHeader*)m_sectionsMem;
+		for (u32 i=0; i!=m_sectionCount; ++i)
+		{
+			auto& section = sections[i];
+			auto relocations = (ImageRelocation*)(newData + section.PointerToRelocations);
+			u32 relocationCount = section.NumberOfRelocations;
+
+			if (section.Characteristics & IMAGE_SCN_LNK_NRELOC_OVFL)
+			{
+				UBA_ASSERT(section.NumberOfRelocations == 0xffff);
+				relocationCount = relocations[0].RelocCount;
+			}
+
+			if (relocationCount)
+				printf("RELOCATION SECTION: %s\n", section.Name);
+
+			for (u32 j=0; j!=relocationCount; ++j)
+			{
+				auto& relocation = relocations[j];
+				(void)relocation;
+
+			}
+		}
+		*/
+
+		/*
+		u8* newSymbolsMem = newData + (m_symbolsMem - m_data);
+
+		auto symbols = (SymbolType*)newSymbolsMem;
+
+		Vector<bool> symbolsToFix;
+		symbolsToFix.resize(m_symbolCount);
+
+		for (u32 i=0; i!=m_symbolCount; ++i)
+		{
+			auto& symbol = symbols[i];
+
+			if (symbol.StorageClass != ImageSymClassExternal)
+				continue;
+
+			if (symbol.SectionNumber != ImageSymUndefined)
+				continue;
+
+			StringView symbolName = GetSymbolName(symbol, m_stringTableMem);
+			if (!symbolName.StartsWith("__imp_", 6))
+				continue;
+
+			symbolName.strBegin += 6;
+			if (allSharedExports.find(std::string(symbolName.strBegin, symbolName.Length())) == allSharedExports.end())
+				continue;
+
+			UBA_ASSERT(symbol.N.Name.Short == 0);
+			symbol.N.Name.Long += 6; // Let's just move the offset in string table 6 bytes forward to skip __imp_ :)
+
+			symbolsToFix[i] = true;
+		}
+
+		auto sections = (ImageSectionHeader*)m_sectionsMem;
+		for (u32 i=0; i!=m_sectionCount; ++i)
+		{
+			auto& section = sections[i];
+			auto relocations = (ImageRelocation*)(newData + section.PointerToRelocations);
+			u32 relocationCount = section.NumberOfRelocations;
+
+			if (section.Characteristics & IMAGE_SCN_LNK_NRELOC_OVFL)
+			{
+				UBA_ASSERT(section.NumberOfRelocations == 0xffff);
+				relocationCount = relocations[0].RelocCount;
+			}
+
+			for (u32 j=0; j!=relocationCount; ++j)
+			{
+				auto& relocation = relocations[j];
+				if (!symbolsToFix[relocation.SymbolTableIndex])
+					continue;
+
+				auto& symbol = symbols[relocation.SymbolTableIndex];
+
+				u8* instruction = newData + section.PointerToRawData + relocation.VirtualAddress - 2;
+
+				static constexpr u8 indirectCallInstruction[] = { 0xFF, 0x15, 0x00, 0x00, 0x00, 0x00 };
+				static constexpr u8 indirectCallInstruction2[] = { 0xFF, 0x25, 0x00, 0x00, 0x00, 0x00 };
+
+				if (memcmp(instruction, indirectCallInstruction, 6) == 0 || memcmp(instruction, indirectCallInstruction2, 6) == 0)
+				{
+					static constexpr u8 relativeCallInstruction[] = { 0xE8, 0x00, 0x00, 0x00, 0x00, 0x90 };
+					memcpy(instruction, relativeCallInstruction, 6);
+					relocation.VirtualAddress -= 1;
+				}
+				else if (*instruction == 0x8b) // mov
+				{
+					UBA_ASSERT(symbol.Type != IMAGE_SYM_DTYPE_FUNCTION);
+					*instruction = 0x8d;
+				}
+				else
+				{
+					logger.Info(TC("Unknown instruction 0x%x for symbol __imp_%S"), instruction[-1], GetSymbolName(symbol, m_stringTableMem).ToString().c_str());
+					//UBA_ASSERTF(false, TC("Unknown instruction 0x%x for symbol __imp_%S"), *instruction, GetSymbolName(symbol, m_stringTableMem).ToString().c_str());
+
+				}
+			}
+		}
+		*/
+	}
+}
diff --git a/Engine/Source/Programs/UnrealBuildAccelerator/Common/Public/UbaObjectFile.h b/Engine/Source/Programs/UnrealBuildAccelerator/Common/Public/UbaObjectFile.h
new file mode 100644
--- /dev/null
+++ b/Engine/Source/Programs/UnrealBuildAccelerator/Common/Public/UbaObjectFile.h
@@ -0,0 +1,50 @@
+// Copyright Epic Games, Inc. All Rights Reserved.
+
+#pragma once
+
+#include "UbaLogger.h"
+#include "UbaStringBuffer.h"
+
+namespace uba
+{
+	class FileAccessor;
+	
+	using UnorderedSymbols = UnorderedSet<std::string>;
+
+	class ObjectFile
+	{
+	public:
+		bool Parse(Logger& logger, const tchar* filename);
+
+		bool CreateStripped(Logger& logger, const tchar* newFilename, const UnorderedSymbols& allNeededImports, const UnorderedSymbols& allExports);
+
+		const tchar* GetFileName() const;
+		const UnorderedSymbols& GetImports() const;
+		const UnorderedSymbols& GetExports() const;
+
+		~ObjectFile();
+
+	private:
+		FileAccessor* m_file = nullptr;
+
+		u8* m_data = nullptr;
+		u64 m_dataSize = 0;
+
+		u8* m_sectionsMem = nullptr;
+		u32 m_sectionCount = 0;
+		u8* m_directiveSectionMem = nullptr;
+		u8* m_stringTableMem = nullptr;
+		u8* m_symbolsMem = nullptr;
+		u32 m_symbolCount = 0;
+		bool m_isBigObj = false;
+		
+		UnorderedSymbols m_imports;
+		UnorderedSymbols m_exports;
+
+		void ParseExports();
+		template<typename SymbolType> void ParseImports();
+
+		void WriteExports(Logger& logger, u8* newData, const UnorderedSymbols& allNeededImports);
+		template<typename SymbolType> void WriteImports(Logger& logger, u8* newData, const UnorderedSymbols& allSharedExports);
+	};
+}
diff --git a/Engine/Source/Programs/UnrealBuildAccelerator/Deploy-Windows.bat b/Engine/Source/Programs/UnrealBuildAccelerator/Deploy-Windows.bat
--- a/Engine/Source/Programs/UnrealBuildAccelerator/Deploy-Windows.bat
+++ b/Engine/Source/Programs/UnrealBuildAccelerator/Deploy-Windows.bat
@@ -21,6 +21,7 @@ call Engine/Build/BatchFiles/RunUBT.bat -NoUba -NoUbaLocal -NoSNDBS -NoXGE -Arch
 	-Target="UbaHost Win64 %config%" ^
 	-Target="UbaStorageProxy Win64 %config%" ^
 	-Target="UbaVisualizer Win64 %config%" ^
+	-Target="UbaObjTool Win64 %config%" ^
 	-Target="UbaCacheService Win64 %config%"
 
 if %ERRORLEVEL% NEQ 0 exit /b %ERRORLEVEL%
diff --git a/Engine/Source/Programs/UnrealBuildAccelerator/ObjTool/Private/UbaObjTool.cpp b/Engine/Source/Programs/UnrealBuildAccelerator/ObjTool/Private/UbaObjTool.cpp
new file mode 100644
--- /dev/null
+++ b/Engine/Source/Programs/UnrealBuildAccelerator/ObjTool/Private/UbaObjTool.cpp
@@ -0,0 +1,218 @@
+// Copyright Epic Games, Inc. All Rights Reserved.
+
+#include "UbaObjectFile.h"
+#include "UbaDirectoryIterator.h"
+#include "UbaVersion.h"
+#include "UbaWorkManager.h"
+
+namespace uba
+{
+	const tchar* Version = GetVersionString();
+	u32	DefaultProcessorCount = []() { return GetLogicalProcessorCount(); }();
+
+	int PrintHelp(const tchar* message)
+	{
+		LoggerWithWriter logger(g_consoleLogWriter, TC(""));
+		if (*message)
+		{
+			logger.Info(TC(""));
+			logger.Error(TC("%s"), message);
+		}
+		const tchar* dbgStr = TC("");
+		#if UBA_DEBUG
+		dbgStr = TC(" (DEBUG)");
+		#endif
+
+		logger.Info(TC(""));
+		logger.Info(TC("-------------------------------------------"));
+		logger.Info(TC("   UbaObjTool v%s%s"), Version, dbgStr);
+		logger.Info(TC("-------------------------------------------"));
+		logger.Info(TC(""));
+		logger.Info(TC("  UbaObjTool.exe [options...] <objfile>"));
+		logger.Info(TC(""));
+		logger.Info(TC(""));
+		logger.Info(TC("  -printsymbols                    Print the symbols found in obj file."));
+		logger.Info(TC(""));
+		return -1;
+	}
+
+	int WrappedMain(int argc, tchar* argv[])
+	{
+		using namespace uba;
+
+		TString objFile;
+		bool printSymbols = false;
+
+		Vector<TString> objFilesToStrip;
+		Vector<TString> objFilesDependencies;
+
+		auto parseArg = [&](const tchar* arg)
+			{
+				StringBuffer<> name;
+				StringBuffer<> value;
+
+				if (const tchar* equals = TStrchr(arg,'='))
+				{
+					name.Append(arg, equals - arg);
+					value.Append(equals+1);
+				}
+				else
+				{
+					name.Append(arg);
+				}
+
+				if (name.StartsWith(TC("/D:")))
+				{
+					objFilesDependencies.push_back(name.data + 3);
+				}
+				else if (name.StartsWith(TC("/S:")))
+				{
+					objFilesToStrip.push_back(name.data + 3);
+				}
+				else if (name.Equals(TC("-printsymbols")))
+				{
+					printSymbols = true;
+				}
+				else if (name.Equals(TC("-?")))
+				{
+					return PrintHelp(TC(""));
+				}
+				else if (objFile.empty() && name[0] != '-' && name[0] != '/')
+				{
+					objFile = name.data;
+					return 0;
+				}
+				else
+				{
+					StringBuffer<> msg;
+					msg.Appendf(TC("Unknown argument '%s'"), name.data);
+					return PrintHelp(msg.data);
+				}
+				return 0;
+			};
+
+		for (int i=1; i!=argc; ++i)
+		{
+			const tchar* arg = argv[i];
+			if (*arg == '@')
+			{
+				++arg;
+				StringBuffer<> temp;
+				if (*arg == '\"')
+				{
+					temp.Append(arg + 1);
+					temp.Resize(temp.count - 1);
+					arg = temp.data;
+				}
+				int res = 0;
+				LoggerWithWriter logger(g_consoleLogWriter, TC(""));
+				if (!ReadLines(logger, arg, [&](const TString& line)
+					{
+						res = parseArg(line.c_str());
+						return res == 0;
+					}))
+					return -1;
+				if (res != 0)
+					return res;
+				continue;
+			}
+			int res = parseArg(arg);
+			if (res != 0)
+				return res;
+		}
+
+		FilteredLogWriter logWriter(g_consoleLogWriter, LogEntryType_Info);
+		LoggerWithWriter logger(logWriter, TC(""));
+
+		if (!objFilesToStrip.empty())
+		{
+			CriticalSection cs;
+			bool success = true;
+			UnorderedSymbols allNeededImports;
+
+			u32 workerCount = 1;//DefaultProcessorCount;
+			WorkManagerImpl workManager(workerCount);
+			workManager.ParallelFor(workerCount, objFilesDependencies, [&](auto& it)
+				{
+					const TString& objFile = *it;
+					ObjectFile objectFile;
+					bool res = objectFile.Parse(logger, objFile.c_str());
+
+					ScopedCriticalSection _(cs);
+					success &= res;
+					if (res)
+						allNeededImports.insert(objectFile.GetImports().begin(), objectFile.GetImports().end());
+				});
+
+			if (!success)
+				return -1;
+
+			UnorderedSymbols allSharedExports;
+			List<ObjectFile> objectFiles;
+
+			workManager.ParallelFor(workerCount, objFilesToStrip, [&](auto& it)
+				{
+					cs.Enter();
+					ObjectFile& objectFile = objectFiles.emplace_back();
+					cs.Leave();
+
+					const TString& objFile = *it;
+					bool res = objectFile.Parse(logger, objFile.c_str());
+
+					ScopedCriticalSection _(cs);
+					success &= res;
+					if (res)
+						allSharedExports.insert(objectFile.GetExports().begin(), objectFile.GetExports().end());
+				});
+
+			if (!success)
+				return -1;
+
+			workManager.ParallelFor(workerCount, objectFiles, [&](auto& it)
+				{
+					auto& objectFile = *it;
+					const tchar* fileName = objectFile.GetFileName();
+					const tchar* lastDot = TStrrchr(fileName, '.');
+					UBA_ASSERT(lastDot);
+					StringBuffer<> newFilename;
+					newFilename.Append(fileName, lastDot - fileName).Append(TC(".strip")).Append(lastDot);
+					objectFile.CreateStripped(logger, newFilename.data, allNeededImports, allSharedExports);
+				});
+
+			//logger.Info(TC("Stripped %llu symbols from %llu obj files"), strippedSymbolCount.load(), objFilesToStrip.size());
+			if (!success)
+				return -1;
+		}
+		else
+		{
+			if (objFile.empty())
+				return PrintHelp(TC("No obj file provided"));
+
+			ObjectFile objectFile;
+			if (!objectFile.Parse(logger, objFile.c_str()))
+				return -1;
+
+			if (printSymbols)
+			{
+				for (auto& symbol : objectFile.GetImports())
+					logger.Info(TC("I %S"), symbol.c_str());
+
+				for (auto& symbol : objectFile.GetExports())
+					logger.Info(TC("E %S"), symbol.c_str());
+			}
+		}
+		return 0;
+	}
+}
+
+#if PLATFORM_WINDOWS
+int wmain(int argc, wchar_t* argv[])
+{
+	return uba::WrappedMain(argc, argv);
+}
+#else
+int main(int argc, char* argv[])
+{
+	return uba::WrappedMain(argc, argv);
+}
+#endif
diff --git a/Engine/Source/Programs/UnrealBuildAccelerator/ObjTool/UbaObjTool.Build.cs b/Engine/Source/Programs/UnrealBuildAccelerator/ObjTool/UbaObjTool.Build.cs
new file mode 100644
--- /dev/null
+++ b/Engine/Source/Programs/UnrealBuildAccelerator/ObjTool/UbaObjTool.Build.cs
@@ -0,0 +1,23 @@
+// Copyright Epic Games, Inc. All Rights Reserved.
+
+using UnrealBuildTool;
+
+[SupportedPlatforms(UnrealPlatformClass.Desktop)]
+public class UbaObjTool : ModuleRules
+{
+	public UbaObjTool(ReadOnlyTargetRules Target) : base(Target)
+	{
+		PrivatePCHHeaderFile = "../Core/Public/UbaCorePch.h";
+
+		UnsafeTypeCastWarningLevel = WarningLevel.Error;
+		StaticAnalyzerDisabledCheckers.Clear();
+
+		PrivateDependencyModuleNames.AddRange(new string[] {
+			"UbaCommon",
+		});
+
+		PrivateDefinitions.AddRange(new string[] {
+			"_CONSOLE",
+		});
+	}
+}
diff --git a/Engine/Source/Programs/UnrealBuildAccelerator/UbaObjTool.Target.cs b/Engine/Source/Programs/UnrealBuildAccelerator/UbaObjTool.Target.cs
new file mode 100644
--- /dev/null
+++ b/Engine/Source/Programs/UnrealBuildAccelerator/UbaObjTool.Target.cs
@@ -0,0 +1,13 @@
+// Copyright Epic Games, Inc. All Rights Reserved.
+
+using UnrealBuildTool;
+
+[SupportedPlatforms(UnrealPlatformClass.Desktop)]
+public class UbaObjToolTarget : TargetRules
+{
+	public UbaObjToolTarget(TargetInfo Target) : base(Target)
+	{
+		LaunchModuleName = "UbaObjTool";
+		UbaAgentTarget.CommonUbaSettings(this, Target);
+	}
+}
