[UBA]
* Improved error messages related to files sent to server. Messages will now show process they are related to

diff --git a/Engine/Source/Programs/UnrealBuildAccelerator/Common/Private/UbaSessionServer.cpp b/Engine/Source/Programs/UnrealBuildAccelerator/Common/Private/UbaSessionServer.cpp
--- a/Engine/Source/Programs/UnrealBuildAccelerator/Common/Private/UbaSessionServer.cpp
+++ b/Engine/Source/Programs/UnrealBuildAccelerator/Common/Private/UbaSessionServer.cpp
@@ -841,13 +841,13 @@ namespace uba
 						bool writeCompressed = m_storeObjFilesCompressed && destination.EndsWith(TC(".obj"));
 						success = m_storage.CopyOrLink(casKey, destination.data, attributes, writeCompressed);
 						if (!success)
-							m_logger.Error(TC("Failed to copy cas from %s to %s"), CasKeyString(casKey).str, destination.data);
+							m_logger.Error(TC("Failed to copy cas from %s to %s (%s)"), CasKeyString(casKey).str, destination.data, GetProcessDescription(processId).c_str());
 					}
 					else
 					{
 						success = m_storage.FakeCopy(casKey, destination.data);
 						if (!success)
-							m_logger.Error(TC("Failed to fake copy cas from %s to %s"), CasKeyString(casKey).str, destination.data);
+							m_logger.Error(TC("Failed to fake copy cas from %s to %s (%s)"), CasKeyString(casKey).str, destination.data, GetProcessDescription(processId).c_str());
 						SCOPED_WRITE_LOCK(m_receivedFilesLock, lock);
 						m_receivedFiles.try_emplace(destinationKey, casKey);
 					}
@@ -1688,6 +1688,16 @@ namespace uba
 		return h;
 	}
 
+	TString SessionServer::GetProcessDescription(u32 processId)
+	{
+		StringBuffer<512> str;
+		SCOPED_READ_LOCK(m_processesLock, lock);
+		auto findIt = m_processes.find(processId);
+		if (findIt == m_processes.end())
+			return str.Appendf(TC("<Process with id %u not found>"), processId).data;
+		return str.Appendf(TC("%s"), findIt->second.GetStartInfo().description).data;
+	}
+
 	bool SessionServer::PrepareProcess(const ProcessStartInfo& startInfo, bool isChild, StringBufferBase& outRealApplication, const tchar*& outRealWorkingDir)
 	{
 		if (!m_memTotal || !m_allowWaitOnMem || isChild)
diff --git a/Engine/Source/Programs/UnrealBuildAccelerator/Common/Public/UbaSessionServer.h b/Engine/Source/Programs/UnrealBuildAccelerator/Common/Public/UbaSessionServer.h
--- a/Engine/Source/Programs/UnrealBuildAccelerator/Common/Public/UbaSessionServer.h
+++ b/Engine/Source/Programs/UnrealBuildAccelerator/Common/Public/UbaSessionServer.h
@@ -80,6 +80,8 @@ namespace uba
 		void OnCancelled(RemoteProcess* process);
 		ProcessHandle ProcessRemoved(u32 processId);
 
+		TString GetProcessDescription(u32 processId);
+
 		virtual bool PrepareProcess(const ProcessStartInfo& startInfo, bool isChild, StringBufferBase& outRealApplication, const tchar*& outRealWorkingDir) override final;
 		virtual bool CreateFile(CreateFileResponse& out, const CreateFileMessage& msg) override final;
 		virtual void FileEntryAdded(StringKey fileNameKey, u64 lastWritten, u64 size) override final;
