[UbaObjTool]
* Strategic submit to get p4 history

diff --git a/Engine/Source/Programs/UnrealBuildAccelerator/Common/Private/UbaObjectFile.cpp b/Engine/Source/Programs/UnrealBuildAccelerator/Common/Private/UbaObjectFile.cpp
--- a/Engine/Source/Programs/UnrealBuildAccelerator/Common/Private/UbaObjectFile.cpp
+++ b/Engine/Source/Programs/UnrealBuildAccelerator/Common/Private/UbaObjectFile.cpp
@@ -67,7 +67,7 @@ namespace uba
 		union
 		{
 			u32 VirtualAddress;
-			u32 RelocCount;             // Set to the real count when IMAGE_SCN_LNK_NRELOC_OVFL is set
+			u32 RelocCount; // Set to the real count when IMAGE_SCN_LNK_NRELOC_OVFL is set
 		} DUMMYUNIONNAME;
 		u32 SymbolTableIndex;
 		u16 Type;
@@ -86,11 +86,11 @@ namespace uba
 			} Name;
 			u32 LongName[2]; // PBYTE  [2]
 		} N;
-		u32   Value;
-		u32    SectionNumber;
-		u16    Type;
-		u8    StorageClass;
-		u8    NumberOfAuxSymbols;
+		u32 Value;
+		u32 SectionNumber;
+		u16 Type;
+		u8 StorageClass;
+		u8 NumberOfAuxSymbols;
 	};
 
 	struct ImageSymbol // IMAGE_SYMBOL
@@ -100,10 +100,10 @@ namespace uba
 			u8 ShortName[8];
 			struct
 			{
-				u32   Short;     // if 0, use LongName
-				u32   Long;      // offset into string table
+				u32   Short; // if 0, use LongName
+				u32   Long; // offset into string table
 			} Name;
-			u32   LongName[2];    // PBYTE [2]
+			u32   LongName[2]; // PBYTE [2]
 		} N;
 		u32 Value;
 		u16 SectionNumber;
@@ -184,28 +184,28 @@ namespace uba
 		if (m_isBigObj)
 		{
 			auto& header = *(AnonObjectHeaderBigobj*)m_data;
-			m_symbolsMem = m_data + header.PointerToSymbolTable;
+			m_symbolsMemPos = header.PointerToSymbolTable;
 			m_symbolCount = header.NumberOfSymbols;
 			m_stringTableMem = m_data + header.PointerToSymbolTable + header.NumberOfSymbols * sizeof(ImageSymbolEx);
-			m_sectionsMem = m_data + sizeof(AnonObjectHeaderBigobj);
+			m_sectionsMemOffset = sizeof(AnonObjectHeaderBigobj);
 			m_sectionCount = header.NumberOfSections;
 		}
 		else
 		{
 			auto& header = *(ImageFileHeader*)m_data;
-			m_symbolsMem = m_data + header.PointerToSymbolTable;
+			m_symbolsMemPos = header.PointerToSymbolTable;
 			m_symbolCount = header.NumberOfSymbols;
 			m_stringTableMem = m_data + header.PointerToSymbolTable + header.NumberOfSymbols * sizeof(ImageSymbol);
-			m_sectionsMem = m_data + sizeof(ImageFileHeader);
+			m_sectionsMemOffset = sizeof(ImageFileHeader);
 			m_sectionCount = header.NumberOfSections;
 		}
 
-		auto sections = (ImageSectionHeader*)m_sectionsMem;
+		auto sections = (ImageSectionHeader*)(m_data + m_sectionsMemOffset);
 		for (u32 i=0; i!=m_sectionCount; ++i)
 		{
 			if (strncmp((char*)sections[i].Name, ".drectve", 8) != 0)
 				continue;
-			m_directiveSectionMem = (u8*)(sections + i);
+			m_directiveSectionMemOffset = u64((u8*)(sections + i) - m_data);
 			break;
 		}
 
@@ -218,25 +218,48 @@ namespace uba
 		return true;
 	}
 
-	bool ObjectFile::CreateStripped(Logger& logger, const tchar* newFilename, const UnorderedSymbols& allNeededImports, const UnorderedSymbols& allSharedExports)
+	bool ObjectFile::CreateStripped(Logger& logger, const tchar* newFilename, const UnorderedSymbols& allNeededImports, const UnorderedSymbols& loopbacksToAdd)
 	{
 		if (!m_file)
 			return false;
+
+		Vector<u32> symbolsToAdd;
+		if (m_isBigObj)
+			CalculateImports<ImageSymbolEx>(logger, loopbacksToAdd, symbolsToAdd);
+		else
+			CalculateImports<ImageSymbol>(logger, loopbacksToAdd, symbolsToAdd);
+
+		u32 importsToFixCount = u32(symbolsToAdd.size());
+		u32 bytesToAdd = 0;
+
+		if (importsToFixCount)
+		{
+			bytesToAdd = sizeof(ImageSectionHeader);
+			bytesToAdd += (m_isBigObj ? sizeof(ImageSymbolEx) : sizeof(ImageSymbol)) * importsToFixCount * 2; // Symbols
+			bytesToAdd += (sizeof(ImageRelocation) + 8) * importsToFixCount; // relocations plus memory
+		}
+
+
 		FileAccessor newFile(logger, newFilename);
-		if (!newFile.CreateMemoryWrite(false, DefaultAttributes(), m_dataSize))
+		if (!newFile.CreateMemoryWrite(false, DefaultAttributes(), m_dataSize + bytesToAdd))
 			return false;
 
 		u8* newData = newFile.GetData();
 
-		memcpy(newData, m_data, m_dataSize);
+		if (importsToFixCount)
+		{
+			if (m_isBigObj)
+				WriteImports<ImageSymbolEx>(logger, newData, symbolsToAdd);
+			else
+				WriteImports<ImageSymbol>(logger, newData, symbolsToAdd);
+		}
+		else
+		{
+			memcpy(newData, m_data, m_dataSize);
+		}
 
 		WriteExports(logger, newData, allNeededImports);
 
-		if (m_isBigObj)
-			WriteImports<ImageSymbolEx>(logger, newData, allSharedExports);
-		else
-			WriteImports<ImageSymbol>(logger, newData, allSharedExports);
-
 		if (!newFile.Close())
 			return false;
 		return true;
@@ -259,9 +282,9 @@ namespace uba
 
 	void ObjectFile::ParseExports()
 	{
-		if (!m_directiveSectionMem)
+		if (!m_directiveSectionMemOffset)
 			return;
-		auto directiveSection = (ImageSectionHeader*)m_directiveSectionMem;
+		auto directiveSection = (ImageSectionHeader*)(m_data + m_directiveSectionMemOffset);
 		u8* directiveData = m_data + directiveSection->PointerToRawData;
 
 		static constexpr u8 utf8Bom[3] = { 0xef, 0xbb, 0xbf };
@@ -290,7 +313,7 @@ namespace uba
 	template<typename SymbolType>
 	void ObjectFile::ParseImports()
 	{
-		auto symbols = (SymbolType*)m_symbolsMem;
+		auto symbols = (SymbolType*)(m_data + m_symbolsMemPos);
 		for (u32 i=0; i!=m_symbolCount; ++i)
 		{
 			auto& symbol = symbols[i];
@@ -299,24 +322,30 @@ namespace uba
 			if (symbol.SectionNumber != ImageSymUndefined)
 				continue;
 			StringView symbolName = GetSymbolName(symbol, m_stringTableMem);
-			if (symbolName.StartsWith("__imp_", 6))
-				symbolName.strBegin += 6;
+
+			//if (strstr(symbolName.ToString().c_str(), "?SerializeRecord@FOutputDevice@@UEAAXAEBVFLogRecord@UE@@@Z"))
+			//	printf("");
+
+			// We must add all imports regardless if they have __imp_ or not
+
+			//if (symbolName.StartsWith("__imp_", 6))
+			//	symbolName.strBegin += 6;
 			m_imports.emplace(std::string(symbolName.strBegin, symbolName.Length()));
 		}
 	}
 
 	void ObjectFile::WriteExports(Logger& logger, u8* newData, const UnorderedSymbols& allNeededImports)
 	{
-		if (!m_directiveSectionMem)
+		if (!m_directiveSectionMemOffset)
 			return;
-		auto directiveSection = (ImageSectionHeader*)m_directiveSectionMem;
+
+		auto directiveSection = (const ImageSectionHeader*)(m_data + m_directiveSectionMemOffset);
 		if (directiveSection->SizeOfRawData < 10)
 			return;
-
-
 		const u8* directiveData = m_data + directiveSection->PointerToRawData;
 
-		u8* newDirectiveData = newData + (directiveData - m_data);
+		auto newDirectiveSection = (ImageSectionHeader*)(newData + m_directiveSectionMemOffset);
+		u8* newDirectiveData = newData + newDirectiveSection->PointerToRawData;
 
 		char* writePos = (char*)newDirectiveData;
 		const char* lastCopyPos = (const char*)directiveData;
@@ -359,7 +388,11 @@ namespace uba
 			if (strncmp(exportEnd-5, ",DATA", 5) == 0)
 				exportEnd -= 5;
 
-			if (allNeededImports.find(std::string(exportStr, exportEnd - exportStr)) != allNeededImports.end())
+			std::string temp(exportStr, exportEnd - exportStr);
+			if (allNeededImports.find(temp) != allNeededImports.end())
+				continue;
+			temp = "__imp_" + temp;
+			if (allNeededImports.find(temp) != allNeededImports.end())
 				continue;
 			
 			// TODO: Use exclusion list provided from command line
@@ -379,119 +412,166 @@ namespace uba
 		writePos += toCopy;
 		*writePos = 0;
 
-		auto newDirectiveSection = (ImageSectionHeader*)(newData + (m_directiveSectionMem - m_data));
 		newDirectiveSection->SizeOfRawData = u32((u8*)writePos - newDirectiveData);
 		UBA_ASSERT(newDirectiveSection->SizeOfRawData <= directiveSection->SizeOfRawData);
 
 		memset(writePos, 0, directiveSection->SizeOfRawData - newDirectiveSection->SizeOfRawData);
 	}
 
-	template<typename SymbolType>
-	void ObjectFile::WriteImports(Logger& logger, u8* newData, const UnorderedSymbols& allSharedExports)
+	template<typename SymbolType> void ObjectFile::CalculateImports(Logger& logger, const UnorderedSymbols& loopbacksToAdd, Vector<u32>& outImports)
 	{
-		/*
-		auto sections = (ImageSectionHeader*)m_sectionsMem;
-		for (u32 i=0; i!=m_sectionCount; ++i)
-		{
-			auto& section = sections[i];
-			auto relocations = (ImageRelocation*)(newData + section.PointerToRelocations);
-			u32 relocationCount = section.NumberOfRelocations;
-
-			if (section.Characteristics & IMAGE_SCN_LNK_NRELOC_OVFL)
-			{
-				UBA_ASSERT(section.NumberOfRelocations == 0xffff);
-				relocationCount = relocations[0].RelocCount;
-			}
-
-			if (relocationCount)
-				printf("RELOCATION SECTION: %s\n", section.Name);
-
-			for (u32 j=0; j!=relocationCount; ++j)
-			{
-				auto& relocation = relocations[j];
-				(void)relocation;
-
-			}
-		}
-		*/
-
-		/*
-		u8* newSymbolsMem = newData + (m_symbolsMem - m_data);
-
-		auto symbols = (SymbolType*)newSymbolsMem;
-
-		Vector<bool> symbolsToFix;
-		symbolsToFix.resize(m_symbolCount);
+		// Calculate how much more memory we need by counting imports
+		auto symbols = (SymbolType*)(m_data + m_symbolsMemPos);
 
 		for (u32 i=0; i!=m_symbolCount; ++i)
 		{
 			auto& symbol = symbols[i];
-
 			if (symbol.StorageClass != ImageSymClassExternal)
 				continue;
-
 			if (symbol.SectionNumber != ImageSymUndefined)
 				continue;
-
 			StringView symbolName = GetSymbolName(symbol, m_stringTableMem);
 			if (!symbolName.StartsWith("__imp_", 6))
 				continue;
-
 			symbolName.strBegin += 6;
-			if (allSharedExports.find(std::string(symbolName.strBegin, symbolName.Length())) == allSharedExports.end())
+
+			auto findIt = loopbacksToAdd.find(std::string(symbolName.strBegin, symbolName.Length()));
+			if (findIt == loopbacksToAdd.end())
 				continue;
+			const_cast<UnorderedSymbols&>(loopbacksToAdd).erase(findIt);
+
+			u32 offsetToStringTable = symbol.N.Name.Long + 6;
+			outImports.push_back(offsetToStringTable);
+		}
+	}
 
-			UBA_ASSERT(symbol.N.Name.Short == 0);
-			symbol.N.Name.Long += 6; // Let's just move the offset in string table 6 bytes forward to skip __imp_ :)
+	template<typename SymbolType>
+	void ObjectFile::WriteImports(Logger& logger, u8* newData, const Vector<u32>& symbolsToAdd)
+	{
+		u32 importsToFixCount = u32(symbolsToAdd.size());
+
+		// Copy header and sections
+		u32 offsetToAfterLastSection = m_sectionsMemOffset + m_sectionCount*sizeof(ImageSectionHeader);
+		memcpy(newData, m_data, offsetToAfterLastSection);
 
-			symbolsToFix[i] = true;
+		// Get some information
+		u32 symbolCount;
+		if (m_isBigObj)
+		{
+			auto& header = *(AnonObjectHeaderBigobj*)newData;
+			symbolCount = header.NumberOfSymbols;
+		}
+		else
+		{
+			auto& header = *(ImageFileHeader*)newData;
+			symbolCount = header.NumberOfSymbols;
 		}
 
-		auto sections = (ImageSectionHeader*)m_sectionsMem;
-		for (u32 i=0; i!=m_sectionCount; ++i)
+		// Add another section
+		u32 newSectionIndex = m_sectionCount;
+		auto& newSection = *(ImageSectionHeader*)(newData + offsetToAfterLastSection);
+		memset(&newSection, 0, sizeof(ImageSectionHeader));
+		memcpy(newSection.Name, ".text$mn", 8);
+
+		// Add dummy memory for relocations
+		u32 newRelocationVirtualMemPos = offsetToAfterLastSection + sizeof(ImageSectionHeader);
+		u32 newRelocationVirtualMemSize = 8 * importsToFixCount;
+		memset(newData + newRelocationVirtualMemPos, 0, newRelocationVirtualMemSize);
+
+		u32 newRelocationsPos = newRelocationVirtualMemPos + newRelocationVirtualMemSize;
+		u32 newRelocationsSize = sizeof(ImageRelocation) * importsToFixCount;
+		newSection.PointerToRelocations = newRelocationsPos;
+		newSection.NumberOfRelocations = u16(importsToFixCount);
+		newSection.PointerToRawData = newRelocationVirtualMemPos;
+		newSection.SizeOfRawData = newRelocationVirtualMemSize;
+		newSection.Characteristics = IMAGE_SCN_CNT_CODE | IMAGE_SCN_MEM_EXECUTE | IMAGE_SCN_MEM_READ;
+		UBA_ASSERT(importsToFixCount <= 65535);
+
+		// Add the new relocations and raw memory
+		memset(newData + newRelocationsPos, 0, sizeof(ImageRelocation) *  importsToFixCount);
+		for (u32 i=0; i!=importsToFixCount; ++i)
 		{
-			auto& section = sections[i];
-			auto relocations = (ImageRelocation*)(newData + section.PointerToRelocations);
-			u32 relocationCount = section.NumberOfRelocations;
+			auto& relocation = ((ImageRelocation*)(newData + newRelocationsPos))[i];
+			relocation.VirtualAddress = 8 * i;
+			relocation.SymbolTableIndex = symbolCount + i;
+			relocation.Type = IMAGE_REL_AMD64_ADDR64;
+		}
 
-			if (section.Characteristics & IMAGE_SCN_LNK_NRELOC_OVFL)
-			{
-				UBA_ASSERT(section.NumberOfRelocations == 0xffff);
-				relocationCount = relocations[0].RelocCount;
-			}
+		// Offset of everything after new section and relocations
+		u32 memoryOffset = newRelocationsPos + newRelocationsSize - offsetToAfterLastSection;
 
-			for (u32 j=0; j!=relocationCount; ++j)
-			{
-				auto& relocation = relocations[j];
-				if (!symbolsToFix[relocation.SymbolTableIndex])
-					continue;
+		u32 symbolTablePos;
+		u32 symbolTableSize;
+		if (m_isBigObj)
+		{
+			auto& header = *(AnonObjectHeaderBigobj*)newData;
+			++header.NumberOfSections;
+			header.PointerToSymbolTable += memoryOffset;
+			symbolTablePos = header.PointerToSymbolTable;
+			symbolTableSize = header.NumberOfSymbols * sizeof(SymbolType);
+			header.NumberOfSymbols += importsToFixCount*2;
+		}
+		else
+		{
+			auto& header = *(ImageFileHeader*)newData;
+			++header.NumberOfSections;
+			header.PointerToSymbolTable += memoryOffset;
+			symbolTablePos = header.PointerToSymbolTable;
+			symbolTableSize = header.NumberOfSymbols * sizeof(SymbolType);
+			header.NumberOfSymbols += importsToFixCount*2;
+		}
 
-				auto& symbol = symbols[relocation.SymbolTableIndex];
+		u32 offsetToAfterSymbolTable = symbolTablePos + symbolTableSize;
 
-				u8* instruction = newData + section.PointerToRawData + relocation.VirtualAddress - 2;
+		// Copy everything after sections til end of symbol table
+		u64 nextToCopySize = m_symbolsMemPos + symbolTableSize - offsetToAfterLastSection;
+		memcpy(newData + offsetToAfterLastSection + memoryOffset, m_data + offsetToAfterLastSection, nextToCopySize);
 
-				static constexpr u8 indirectCallInstruction[] = { 0xFF, 0x15, 0x00, 0x00, 0x00, 0x00 };
-				static constexpr u8 indirectCallInstruction2[] = { 0xFF, 0x25, 0x00, 0x00, 0x00, 0x00 };
+		UBA_ASSERT(m_isBigObj);
+		u32 newSymbolsPos = offsetToAfterSymbolTable;
+		u32 newSymbolsSize = sizeof(SymbolType)*importsToFixCount*2;
+		auto newSymbols = (SymbolType*)(newData + newSymbolsPos);
+		memset(newSymbols, 0, newSymbolsSize);
+		for (u32 i=0; i!=importsToFixCount; ++i)
+		{
+			auto& symbol = newSymbols[i];
+			symbol.N.Name.Long = symbolsToAdd[i];
+			symbol.SectionNumber = ImageSymUndefined;
+			symbol.StorageClass = ImageSymClassExternal;
+		}
+		for (u32 i=0; i!=importsToFixCount; ++i)
+		{
+			auto& symbol = newSymbols[i + importsToFixCount];
+			symbol.N.Name.Long = symbolsToAdd[i] - 6;
+			symbol.SectionNumber = (decltype(SymbolType::SectionNumber))newSectionIndex;
+			symbol.StorageClass = ImageSymClassExternal;
+			symbol.Value = i * 8;
+		}
 
-				if (memcmp(instruction, indirectCallInstruction, 6) == 0 || memcmp(instruction, indirectCallInstruction2, 6) == 0)
-				{
-					static constexpr u8 relativeCallInstruction[] = { 0xE8, 0x00, 0x00, 0x00, 0x00, 0x90 };
-					memcpy(instruction, relativeCallInstruction, 6);
-					relocation.VirtualAddress -= 1;
-				}
-				else if (*instruction == 0x8b) // mov
-				{
-					UBA_ASSERT(symbol.Type != IMAGE_SYM_DTYPE_FUNCTION);
-					*instruction = 0x8d;
-				}
-				else
-				{
-					logger.Info(TC("Unknown instruction 0x%x for symbol __imp_%S"), instruction[-1], GetSymbolName(symbol, m_stringTableMem).ToString().c_str());
-					//UBA_ASSERTF(false, TC("Unknown instruction 0x%x for symbol __imp_%S"), *instruction, GetSymbolName(symbol, m_stringTableMem).ToString().c_str());
+		u64 lastToCopySize = m_dataSize - (m_stringTableMem - m_data);
+		memcpy(newData + newSymbolsPos + newSymbolsSize, m_stringTableMem, lastToCopySize);
 
-				}
+
+		// Traverse all sections and update their PointerToRawData
+		auto sections = (ImageSectionHeader*)(newData + m_sectionsMemOffset);
+		for (u32 i=0; i!=m_sectionCount; ++i)
+		{
+			auto& section = sections[i];
+			if (section.PointerToRawData)
+			{
+				UBA_ASSERT(section.PointerToRawData < symbolTablePos);
+				section.PointerToRawData += memoryOffset;
+			}
+			if (section.PointerToRelocations)
+			{
+				UBA_ASSERT(section.PointerToRelocations < symbolTablePos);
+				section.PointerToRelocations += memoryOffset;
+			}
+			if (section.PointerToLinenumbers)
+			{
+				UBA_ASSERT(section.PointerToLinenumbers < symbolTablePos);
+				section.PointerToLinenumbers += memoryOffset;
 			}
 		}
-		*/
 	}
 }
diff --git a/Engine/Source/Programs/UnrealBuildAccelerator/Common/Public/UbaObjectFile.h b/Engine/Source/Programs/UnrealBuildAccelerator/Common/Public/UbaObjectFile.h
--- a/Engine/Source/Programs/UnrealBuildAccelerator/Common/Public/UbaObjectFile.h
+++ b/Engine/Source/Programs/UnrealBuildAccelerator/Common/Public/UbaObjectFile.h
@@ -16,11 +16,12 @@ namespace uba
 	public:
 		bool Parse(Logger& logger, const tchar* filename);
 
-		bool CreateStripped(Logger& logger, const tchar* newFilename, const UnorderedSymbols& allNeededImports, const UnorderedSymbols& allExports);
+		bool CreateStripped(Logger& logger, const tchar* newFilename, const UnorderedSymbols& allNeededImports, const UnorderedSymbols& loopbacksToAdd);
 
 		const tchar* GetFileName() const;
 		const UnorderedSymbols& GetImports() const;
 		const UnorderedSymbols& GetExports() const;
+		const UnorderedSymbols& GetSymbols() const;
 
 		~ObjectFile();
 
@@ -30,21 +31,24 @@ namespace uba
 		u8* m_data = nullptr;
 		u64 m_dataSize = 0;
 
-		u8* m_sectionsMem = nullptr;
+		u32 m_sectionsMemOffset = 0;
 		u32 m_sectionCount = 0;
-		u8* m_directiveSectionMem = nullptr;
+		u64 m_directiveSectionMemOffset = 0;
 		u8* m_stringTableMem = nullptr;
-		u8* m_symbolsMem = nullptr;
+		u32 m_symbolsMemPos = 0;
 		u32 m_symbolCount = 0;
 		bool m_isBigObj = false;
 		
 		UnorderedSymbols m_imports;
 		UnorderedSymbols m_exports;
+		UnorderedSymbols m_symbols;
 
 		void ParseExports();
 		template<typename SymbolType> void ParseImports();
 
+
 		void WriteExports(Logger& logger, u8* newData, const UnorderedSymbols& allNeededImports);
-		template<typename SymbolType> void WriteImports(Logger& logger, u8* newData, const UnorderedSymbols& allSharedExports);
+		template<typename SymbolType> void CalculateImports(Logger& logger, const UnorderedSymbols& loopbacksToAdd, Vector<u32>& outImports);
+		template<typename SymbolType> void WriteImports(Logger& logger, u8* newData, const Vector<u32>& symbolsToAdd);
 	};
 }
