[UbaObjTool]
* Fixed off-by-one issue with segment index in generated symbols

diff --git a/Engine/Source/Programs/UnrealBuildAccelerator/Common/Private/UbaObjectFile.cpp b/Engine/Source/Programs/UnrealBuildAccelerator/Common/Private/UbaObjectFile.cpp
--- a/Engine/Source/Programs/UnrealBuildAccelerator/Common/Private/UbaObjectFile.cpp
+++ b/Engine/Source/Programs/UnrealBuildAccelerator/Common/Private/UbaObjectFile.cpp
@@ -543,7 +543,7 @@ namespace uba
 		{
 			auto& symbol = newSymbols[i + importsToFixCount];
 			symbol.N.Name.Long = symbolsToAdd[i] - 6;
-			symbol.SectionNumber = (decltype(SymbolType::SectionNumber))newSectionIndex;
+			symbol.SectionNumber = (decltype(SymbolType::SectionNumber))(newSectionIndex + 1);
 			symbol.StorageClass = ImageSymClassExternal;
 			symbol.Value = i * 8;
 		}
diff --git a/Engine/Source/Programs/UnrealBuildAccelerator/ObjTool/Private/UbaObjTool.cpp b/Engine/Source/Programs/UnrealBuildAccelerator/ObjTool/Private/UbaObjTool.cpp
--- a/Engine/Source/Programs/UnrealBuildAccelerator/ObjTool/Private/UbaObjTool.cpp
+++ b/Engine/Source/Programs/UnrealBuildAccelerator/ObjTool/Private/UbaObjTool.cpp
@@ -130,7 +130,7 @@ namespace uba
 			bool success = true;
 			UnorderedSymbols allNeededImports;
 
-			u32 workerCount = 1;//DefaultProcessorCount;
+			u32 workerCount = DefaultProcessorCount;
 			WorkManagerImpl workManager(workerCount);
 			workManager.ParallelFor(workerCount, objFilesDependencies, [&](auto& it)
 				{
@@ -148,16 +148,19 @@ namespace uba
 				return -1;
 
 			UnorderedSymbols allSharedExports;
-			List<ObjectFile> objectFiles;
+
+			struct ObjectFileRec { ObjectFile file; UnorderedSet<std::string> loopbacksToAdd; };
+			Map<TString, ObjectFileRec> objectFileRecs;
 
 			workManager.ParallelFor(workerCount, objFilesToStrip, [&](auto& it)
 				{
+					const TString& objFileName = *it;
+
 					cs.Enter();
-					ObjectFile& objectFile = objectFiles.emplace_back();
+					ObjectFile& objectFile = objectFileRecs.try_emplace(objFileName).first->second.file;
 					cs.Leave();
 
-					const TString& objFile = *it;
-					bool res = objectFile.Parse(logger, objFile.c_str());
+					bool res = objectFile.Parse(logger, objFileName.c_str());
 
 					ScopedCriticalSection _(cs);
 					success &= res;
@@ -168,15 +171,32 @@ namespace uba
 			if (!success)
 				return -1;
 
-			workManager.ParallelFor(workerCount, objectFiles, [&](auto& it)
+			// Figure out which loopback symbols that should be added to which obj file
+			for (auto& objFileName : objFilesToStrip)
+			{
+				ObjectFileRec& rec = objectFileRecs[objFileName];//kv.second;
+				for (auto& importSymbol : rec.file.GetImports())
+				{
+					if (strncmp(importSymbol.c_str(), "__imp_", 6) != 0)
+						continue;
+					std::string tmp = importSymbol.substr(6);
+					auto findIt = allSharedExports.find(tmp);
+					if (findIt == allSharedExports.end())
+						continue;
+					rec.loopbacksToAdd.emplace(tmp);
+					allSharedExports.erase(findIt);
+				}
+			}
+
+			workManager.ParallelFor(workerCount, objectFileRecs, [&](auto& it)
 				{
-					auto& objectFile = *it;
-					const tchar* fileName = objectFile.GetFileName();
+					ObjectFileRec& rec = it->second;
+					const tchar* fileName = rec.file.GetFileName();
 					const tchar* lastDot = TStrrchr(fileName, '.');
 					UBA_ASSERT(lastDot);
 					StringBuffer<> newFilename;
 					newFilename.Append(fileName, lastDot - fileName).Append(TC(".strip")).Append(lastDot);
-					objectFile.CreateStripped(logger, newFilename.data, allNeededImports, allSharedExports);
+					rec.file.CreateStripped(logger, newFilename.data, allNeededImports, rec.loopbacksToAdd);
 				});
 
 			//logger.Info(TC("Stripped %llu symbols from %llu obj files"), strippedSymbolCount.load(), objFilesToStrip.size());
