[UBA]
* Improved error message on session client when failing to upload file to server

diff --git a/Engine/Source/Programs/UnrealBuildAccelerator/Common/Private/UbaSessionClient.cpp b/Engine/Source/Programs/UnrealBuildAccelerator/Common/Private/UbaSessionClient.cpp
--- a/Engine/Source/Programs/UnrealBuildAccelerator/Common/Private/UbaSessionClient.cpp
+++ b/Engine/Source/Programs/UnrealBuildAccelerator/Common/Private/UbaSessionClient.cpp
@@ -597,7 +597,7 @@ namespace uba
 				return m_logger.Error(TC("Failed to send file %s to server"), source.name.c_str());
 		}
 		if (!reader.ReadBool())
-			return m_logger.Error(TC("Server failed to receive file %s"), source.name.c_str());
+			return m_logger.Error(TC("Server failed to receive file %s (%s)"), source.name.c_str(), destination);
 		return true;
 	}
 
