UbaObjTool - Fixed so operator new/delete etc are removed from obj files to avoid duplicated symbols

diff --git a/Engine/Source/Programs/UnrealBuildAccelerator/Common/Private/UbaObjectFile.cpp b/Engine/Source/Programs/UnrealBuildAccelerator/Common/Private/UbaObjectFile.cpp
--- a/Engine/Source/Programs/UnrealBuildAccelerator/Common/Private/UbaObjectFile.cpp
+++ b/Engine/Source/Programs/UnrealBuildAccelerator/Common/Private/UbaObjectFile.cpp
@@ -6,12 +6,20 @@
 namespace uba
 {
 	constexpr u16 ImageFileMachineUnknown = 0;
-	constexpr u8 ImageSizeofShortName = 8;
+	constexpr u8 ImageSizeofShortName = 8; // IMAGE_SIZEOF_SHORT_NAME
 	constexpr u16 ImageSymClassExternal = 0x0002; // IMAGE_SYM_CLASS_EXTERNAL
 	constexpr u16 ImageSymUndefined = 0; // IMAGE_SYM_UNDEFINED
 
-#pragma pack(push)
-#pragma pack(1)
+	constexpr u32 ImageScnCntCode = 0x00000020; // IMAGE_SCN_CNT_CODE
+	constexpr u32 ImageScnMemExecute = 0x20000000; // IMAGE_SCN_MEM_EXECUTE
+	constexpr u32 ImageScnMemRead = 0x40000000; // IMAGE_SCN_MEM_READ
+	constexpr u32 ImageScnLnkComdat = 0x00001000; // IMAGE_SCN_LNK_COMDAT
+
+	const u16 ImageRelAmd64Addr64 = 0x0001; // IMAGE_REL_AMD64_ADDR64
+
+
+	#pragma pack(push)
+	#pragma pack(1)
 
 	struct ImageFileHeader // IMAGE_FILE_HEADER
 	{
@@ -68,7 +76,7 @@ namespace uba
 		{
 			u32 VirtualAddress;
 			u32 RelocCount; // Set to the real count when IMAGE_SCN_LNK_NRELOC_OVFL is set
-		} DUMMYUNIONNAME;
+		};
 		u32 SymbolTableIndex;
 		u16 Type;
 	};
@@ -112,7 +120,48 @@ namespace uba
 		u8 NumberOfAuxSymbols;
 	};
 
-#pragma pack(pop)
+	#pragma pack(pop)
+
+	// These are symbols that are added to all dlls through some macros.
+	// When merging dlls we need to remove duplicates of these
+	static UnorderedSymbols PotentiallyDuplicatedSymbols = 
+	{
+		// operator new/delete 
+		"??2@YAPEAX_K@Z",
+		"??2@YAPEAX_KAEBUnothrow_t@std@@@Z",
+		"??2@YAPEAX_KW4align_val_t@std@@@Z",
+		"??2@YAPEAX_KW4align_val_t@std@@AEBUnothrow_t@1@@Z",
+		"??3@YAXPEAX@Z",
+		"??3@YAXPEAXAEBUnothrow_t@std@@@Z",
+		"??3@YAXPEAXW4align_val_t@std@@@Z",
+		"??3@YAXPEAXW4align_val_t@std@@AEBUnothrow_t@1@@Z",
+		"??3@YAXPEAX_K@Z",
+		"??3@YAXPEAX_KAEBUnothrow_t@std@@@Z",
+		"??3@YAXPEAX_KW4align_val_t@std@@@Z",
+		"??3@YAXPEAX_KW4align_val_t@std@@AEBUnothrow_t@1@@Z",
+		"??_U@YAPEAX_K@Z",
+		"??_U@YAPEAX_KAEBUnothrow_t@std@@@Z",
+		"??_U@YAPEAX_KW4align_val_t@std@@@Z",
+		"??_U@YAPEAX_KW4align_val_t@std@@AEBUnothrow_t@1@@Z",
+		"??_V@YAXPEAX@Z",
+		"??_V@YAXPEAXAEBUnothrow_t@std@@@Z",
+		"??_V@YAXPEAXW4align_val_t@std@@@Z",
+		"??_V@YAXPEAXW4align_val_t@std@@AEBUnothrow_t@1@@Z",
+		"??_V@YAXPEAX_K@Z",
+		"??_V@YAXPEAX_KAEBUnothrow_t@std@@@Z",
+		"??_V@YAXPEAX_KW4align_val_t@std@@@Z",
+		"??_V@YAXPEAX_KW4align_val_t@std@@AEBUnothrow_t@1@@Z",
+
+		// Memory
+		"?FMemory_Free@@YAXPEAX@Z",
+		"?FMemory_Malloc@@YAPEAX_K0@Z",
+		"?FMemory_Realloc@@YAPEAXPEAX_K1@Z",
+
+		// Have no idea why these are duplicated
+		"??_GIFunction_OwnedObject@Function@Private@Core@UE@@UEAAPEAXI@Z",
+		"??_GIDelegateInstance@@UEAAPEAXI@Z",
+	};
+
 
 	bool IsBigObj(u8* data, u64 size)
 	{
@@ -148,6 +197,26 @@ namespace uba
 			return memcmp(strBegin, str, strLen) == 0;
 		}
 
+		bool Contains(const char* str, u32 strLen) const
+		{
+			const char* it = strBegin;
+			const char* itEnd = strEnd - strLen + 1;
+			while (it < itEnd)
+			{
+				if (memcmp(it, str, strLen) == 0)
+					return true;
+				++it;
+			}
+			return false;
+		}
+
+		bool Equals(const char* str, u32 strLen) const
+		{
+			if (strLen != Length())
+				return false;
+			return memcmp(strBegin, str, strLen) == 0;
+		}
+
 		std::string ToString() const
 		{
 			return std::string(strBegin, strEnd);
@@ -173,6 +242,8 @@ namespace uba
 
 	bool ObjectFile::Parse(Logger& logger, const tchar* filename)
 	{
+		UBA_ASSERT(!m_file);
+
 		m_file = new FileAccessor(logger, filename);
 		if (!m_file->OpenMemoryRead())
 			return false;
@@ -184,28 +255,28 @@ namespace uba
 		if (m_isBigObj)
 		{
 			auto& header = *(AnonObjectHeaderBigobj*)m_data;
-			m_symbolsMemPos = header.PointerToSymbolTable;
-			m_symbolCount = header.NumberOfSymbols;
-			m_stringTableMem = m_data + header.PointerToSymbolTable + header.NumberOfSymbols * sizeof(ImageSymbolEx);
-			m_sectionsMemOffset = sizeof(AnonObjectHeaderBigobj);
-			m_sectionCount = header.NumberOfSections;
+			m_info.symbolsMemPos = header.PointerToSymbolTable;
+			m_info.symbolCount = header.NumberOfSymbols;
+			m_info.stringTableMem = m_data + header.PointerToSymbolTable + header.NumberOfSymbols * sizeof(ImageSymbolEx);
+			m_info.sectionsMemOffset = sizeof(AnonObjectHeaderBigobj);
+			m_info.sectionCount = header.NumberOfSections;
 		}
 		else
 		{
 			auto& header = *(ImageFileHeader*)m_data;
-			m_symbolsMemPos = header.PointerToSymbolTable;
-			m_symbolCount = header.NumberOfSymbols;
-			m_stringTableMem = m_data + header.PointerToSymbolTable + header.NumberOfSymbols * sizeof(ImageSymbol);
-			m_sectionsMemOffset = sizeof(ImageFileHeader);
-			m_sectionCount = header.NumberOfSections;
+			m_info.symbolsMemPos = header.PointerToSymbolTable;
+			m_info.symbolCount = header.NumberOfSymbols;
+			m_info.stringTableMem = m_data + header.PointerToSymbolTable + header.NumberOfSymbols * sizeof(ImageSymbol);
+			m_info.sectionsMemOffset = sizeof(ImageFileHeader);
+			m_info.sectionCount = header.NumberOfSections;
 		}
 
-		auto sections = (ImageSectionHeader*)(m_data + m_sectionsMemOffset);
-		for (u32 i=0; i!=m_sectionCount; ++i)
+		auto sections = (ImageSectionHeader*)(m_data + m_info.sectionsMemOffset);
+		for (u32 i=0; i!=m_info.sectionCount; ++i)
 		{
 			if (strncmp((char*)sections[i].Name, ".drectve", 8) != 0)
 				continue;
-			m_directiveSectionMemOffset = u64((u8*)(sections + i) - m_data);
+			m_info.directiveSectionMemOffset = u64((u8*)(sections + i) - m_data);
 			break;
 		}
 
@@ -215,13 +286,13 @@ namespace uba
 			ParseImports<ImageSymbolEx>();
 		else
 			ParseImports<ImageSymbol>();
+
 		return true;
 	}
 
-	bool ObjectFile::CreateStripped(Logger& logger, const tchar* newFilename, const UnorderedSymbols& allNeededImports, const UnorderedSymbols& loopbacksToAdd)
+	bool ObjectFile::CreateStripped(Logger& logger, const tchar* newFilename, const UnorderedSymbols& allNeededImports, UnorderedSymbols& loopbacksToAdd, UnorderedSymbols& toRemove)
 	{
-		if (!m_file)
-			return false;
+		UBA_ASSERT(m_file);
 
 		Vector<u32> symbolsToAdd;
 		if (m_isBigObj)
@@ -239,30 +310,33 @@ namespace uba
 			bytesToAdd += (sizeof(ImageRelocation) + 8) * importsToFixCount; // relocations plus memory
 		}
 
-
 		FileAccessor newFile(logger, newFilename);
 		if (!newFile.CreateMemoryWrite(false, DefaultAttributes(), m_dataSize + bytesToAdd))
 			return false;
 
 		u8* newData = newFile.GetData();
+		Info newInfo = m_info;
 
 		if (importsToFixCount)
 		{
 			if (m_isBigObj)
-				WriteImports<ImageSymbolEx>(logger, newData, symbolsToAdd);
+				WriteImports<ImageSymbolEx>(logger, newData, newInfo, symbolsToAdd);
 			else
-				WriteImports<ImageSymbol>(logger, newData, symbolsToAdd);
+				WriteImports<ImageSymbol>(logger, newData, newInfo, symbolsToAdd);
 		}
 		else
 		{
 			memcpy(newData, m_data, m_dataSize);
 		}
 
+		if (m_isBigObj)
+			RemoveSymbols<ImageSymbolEx>(logger, newData, newInfo, toRemove);
+		else
+			RemoveSymbols<ImageSymbol>(logger, newData, newInfo, toRemove);
+
 		WriteExports(logger, newData, allNeededImports);
 
-		if (!newFile.Close())
-			return false;
-		return true;
+		return newFile.Close();
 	}
 
 	const tchar* ObjectFile::GetFileName() const
@@ -280,11 +354,16 @@ namespace uba
 		return m_exports;
 	}
 
+	const UnorderedSymbols& ObjectFile::GetPotentialDuplicates() const
+	{
+		return m_potentialDuplicates;
+	}
+
 	void ObjectFile::ParseExports()
 	{
-		if (!m_directiveSectionMemOffset)
+		if (!m_info.directiveSectionMemOffset)
 			return;
-		auto directiveSection = (ImageSectionHeader*)(m_data + m_directiveSectionMemOffset);
+		auto directiveSection = (ImageSectionHeader*)(m_data + m_info.directiveSectionMemOffset);
 		u8* directiveData = m_data + directiveSection->PointerToRawData;
 
 		static constexpr u8 utf8Bom[3] = { 0xef, 0xbb, 0xbf };
@@ -306,45 +385,50 @@ namespace uba
 			if (strncmp(exportEnd-5, ",DATA", 5) == 0)
 				exportEnd -= 5;
 					
-			m_exports.emplace(std::string(exportStr, exportEnd - exportStr));
+			m_exports.emplace(std::string(exportStr, exportEnd));
 		}
 	}
 
 	template<typename SymbolType>
 	void ObjectFile::ParseImports()
 	{
-		auto symbols = (SymbolType*)(m_data + m_symbolsMemPos);
-		for (u32 i=0; i!=m_symbolCount; ++i)
+		auto symbols = (SymbolType*)(m_data + m_info.symbolsMemPos);
+		for (u32 i=0; i!=m_info.symbolCount; ++i)
 		{
 			auto& symbol = symbols[i];
+
 			if (symbol.StorageClass != ImageSymClassExternal)
 				continue;
-			if (symbol.SectionNumber != ImageSymUndefined)
-				continue;
-			StringView symbolName = GetSymbolName(symbol, m_stringTableMem);
-
-			//if (strstr(symbolName.ToString().c_str(), "?SerializeRecord@FOutputDevice@@UEAAXAEBVFLogRecord@UE@@@Z"))
-			//	printf("");
 
-			// We must add all imports regardless if they have __imp_ or not
+			StringView symbolName = GetSymbolName(symbol, m_info.stringTableMem);
+			std::string symbolString = symbolName.ToString();
 
-			//if (symbolName.StartsWith("__imp_", 6))
-			//	symbolName.strBegin += 6;
-			m_imports.emplace(std::string(symbolName.strBegin, symbolName.Length()));
+			if (symbol.SectionNumber != ImageSymUndefined)
+			{
+				if (PotentiallyDuplicatedSymbols.find(symbolString) != PotentiallyDuplicatedSymbols.end())
+				{
+					m_potentialDuplicates.emplace(symbolString);
+				}
+			}
+			else
+			{
+				m_imports.emplace(symbolString);
+			}
 		}
 	}
 
 	void ObjectFile::WriteExports(Logger& logger, u8* newData, const UnorderedSymbols& allNeededImports)
 	{
-		if (!m_directiveSectionMemOffset)
+		if (!m_info.directiveSectionMemOffset)
 			return;
 
-		auto directiveSection = (const ImageSectionHeader*)(m_data + m_directiveSectionMemOffset);
+		auto directiveSection = (const ImageSectionHeader*)(m_data + m_info.directiveSectionMemOffset);
 		if (directiveSection->SizeOfRawData < 10)
 			return;
+
 		const u8* directiveData = m_data + directiveSection->PointerToRawData;
 
-		auto newDirectiveSection = (ImageSectionHeader*)(newData + m_directiveSectionMemOffset);
+		auto newDirectiveSection = (ImageSectionHeader*)(newData + m_info.directiveSectionMemOffset);
 		u8* newDirectiveData = newData + newDirectiveSection->PointerToRawData;
 
 		char* writePos = (char*)newDirectiveData;
@@ -418,57 +502,41 @@ namespace uba
 		memset(writePos, 0, directiveSection->SizeOfRawData - newDirectiveSection->SizeOfRawData);
 	}
 
-	template<typename SymbolType> void ObjectFile::CalculateImports(Logger& logger, const UnorderedSymbols& loopbacksToAdd, Vector<u32>& outImports)
+	template<typename SymbolType> void ObjectFile::CalculateImports(Logger& logger, UnorderedSymbols& loopbacksToAdd, Vector<u32>& outImports)
 	{
 		// Calculate how much more memory we need by counting imports
-		auto symbols = (SymbolType*)(m_data + m_symbolsMemPos);
+		auto symbols = (SymbolType*)(m_data + m_info.symbolsMemPos);
 
-		for (u32 i=0; i!=m_symbolCount; ++i)
+		for (u32 i=0; i!=m_info.symbolCount; ++i)
 		{
 			auto& symbol = symbols[i];
 			if (symbol.StorageClass != ImageSymClassExternal)
 				continue;
 			if (symbol.SectionNumber != ImageSymUndefined)
 				continue;
-			StringView symbolName = GetSymbolName(symbol, m_stringTableMem);
+			StringView symbolName = GetSymbolName(symbol, m_info.stringTableMem);
 			if (!symbolName.StartsWith("__imp_", 6))
 				continue;
 			symbolName.strBegin += 6;
-
-			auto findIt = loopbacksToAdd.find(std::string(symbolName.strBegin, symbolName.Length()));
+			auto findIt = loopbacksToAdd.find(symbolName.ToString());
 			if (findIt == loopbacksToAdd.end())
 				continue;
-			const_cast<UnorderedSymbols&>(loopbacksToAdd).erase(findIt);
-
-			u32 offsetToStringTable = symbol.N.Name.Long + 6;
-			outImports.push_back(offsetToStringTable);
+			loopbacksToAdd.erase(findIt);
+			outImports.push_back(symbol.N.Name.Long);
 		}
 	}
 
 	template<typename SymbolType>
-	void ObjectFile::WriteImports(Logger& logger, u8* newData, const Vector<u32>& symbolsToAdd)
+	void ObjectFile::WriteImports(Logger& logger, u8* newData, Info& newInfo, const Vector<u32>& symbolsToAdd)
 	{
 		u32 importsToFixCount = u32(symbolsToAdd.size());
 
 		// Copy header and sections
-		u32 offsetToAfterLastSection = m_sectionsMemOffset + m_sectionCount*sizeof(ImageSectionHeader);
+		u32 offsetToAfterLastSection = m_info.sectionsMemOffset + m_info.sectionCount*sizeof(ImageSectionHeader);
 		memcpy(newData, m_data, offsetToAfterLastSection);
 
-		// Get some information
-		u32 symbolCount;
-		if (m_isBigObj)
-		{
-			auto& header = *(AnonObjectHeaderBigobj*)newData;
-			symbolCount = header.NumberOfSymbols;
-		}
-		else
-		{
-			auto& header = *(ImageFileHeader*)newData;
-			symbolCount = header.NumberOfSymbols;
-		}
-
 		// Add another section
-		u32 newSectionIndex = m_sectionCount;
+		u32 newSectionIndex = m_info.sectionCount;
 		auto& newSection = *(ImageSectionHeader*)(newData + offsetToAfterLastSection);
 		memset(&newSection, 0, sizeof(ImageSectionHeader));
 		memcpy(newSection.Name, ".text$mn", 8);
@@ -484,17 +552,18 @@ namespace uba
 		newSection.NumberOfRelocations = u16(importsToFixCount);
 		newSection.PointerToRawData = newRelocationVirtualMemPos;
 		newSection.SizeOfRawData = newRelocationVirtualMemSize;
-		newSection.Characteristics = IMAGE_SCN_CNT_CODE | IMAGE_SCN_MEM_EXECUTE | IMAGE_SCN_MEM_READ;
+		newSection.Characteristics = ImageScnCntCode | ImageScnMemExecute | ImageScnMemRead;
 		UBA_ASSERT(importsToFixCount <= 65535);
 
 		// Add the new relocations and raw memory
+		u32 newSymbolIndex = m_info.symbolCount;
 		memset(newData + newRelocationsPos, 0, sizeof(ImageRelocation) *  importsToFixCount);
 		for (u32 i=0; i!=importsToFixCount; ++i)
 		{
 			auto& relocation = ((ImageRelocation*)(newData + newRelocationsPos))[i];
 			relocation.VirtualAddress = 8 * i;
-			relocation.SymbolTableIndex = symbolCount + i;
-			relocation.Type = IMAGE_REL_AMD64_ADDR64;
+			relocation.SymbolTableIndex = newSymbolIndex + i;
+			relocation.Type = ImageRelAmd64Addr64;
 		}
 
 		// Offset of everything after new section and relocations
@@ -524,10 +593,11 @@ namespace uba
 		u32 offsetToAfterSymbolTable = symbolTablePos + symbolTableSize;
 
 		// Copy everything after sections til end of symbol table
-		u64 nextToCopySize = m_symbolsMemPos + symbolTableSize - offsetToAfterLastSection;
+		u64 nextToCopySize = m_info.symbolsMemPos + symbolTableSize - offsetToAfterLastSection;
 		memcpy(newData + offsetToAfterLastSection + memoryOffset, m_data + offsetToAfterLastSection, nextToCopySize);
 
-		UBA_ASSERT(m_isBigObj);
+		newInfo.symbolsMemPos += memoryOffset;
+
 		u32 newSymbolsPos = offsetToAfterSymbolTable;
 		u32 newSymbolsSize = sizeof(SymbolType)*importsToFixCount*2;
 		auto newSymbols = (SymbolType*)(newData + newSymbolsPos);
@@ -535,26 +605,26 @@ namespace uba
 		for (u32 i=0; i!=importsToFixCount; ++i)
 		{
 			auto& symbol = newSymbols[i];
-			symbol.N.Name.Long = symbolsToAdd[i];
+			symbol.N.Name.Long = symbolsToAdd[i] + 6; // Remove __imp_
 			symbol.SectionNumber = ImageSymUndefined;
 			symbol.StorageClass = ImageSymClassExternal;
 		}
 		for (u32 i=0; i!=importsToFixCount; ++i)
 		{
 			auto& symbol = newSymbols[i + importsToFixCount];
-			symbol.N.Name.Long = symbolsToAdd[i] - 6;
+			symbol.N.Name.Long = symbolsToAdd[i];
 			symbol.SectionNumber = (decltype(SymbolType::SectionNumber))(newSectionIndex + 1);
 			symbol.StorageClass = ImageSymClassExternal;
 			symbol.Value = i * 8;
 		}
 
-		u64 lastToCopySize = m_dataSize - (m_stringTableMem - m_data);
-		memcpy(newData + newSymbolsPos + newSymbolsSize, m_stringTableMem, lastToCopySize);
+		u64 lastToCopySize = m_dataSize - (m_info.stringTableMem - m_data);
+		memcpy(newData + newSymbolsPos + newSymbolsSize, m_info.stringTableMem, lastToCopySize);
 
 
 		// Traverse all sections and update their PointerToRawData
-		auto sections = (ImageSectionHeader*)(newData + m_sectionsMemOffset);
-		for (u32 i=0; i!=m_sectionCount; ++i)
+		auto sections = (ImageSectionHeader*)(newData + m_info.sectionsMemOffset);
+		for (u32 i=0; i!=m_info.sectionCount; ++i)
 		{
 			auto& section = sections[i];
 			if (section.PointerToRawData)
@@ -574,4 +644,34 @@ namespace uba
 			}
 		}
 	}
+
+	template<typename SymbolType> void ObjectFile::RemoveSymbols(Logger& logger, u8* newData, Info& newInfo, UnorderedSymbols& toRemove)
+	{
+		auto symbols = (SymbolType*)(newData + newInfo.symbolsMemPos);
+		auto sections = (ImageSectionHeader*)(newData + newInfo.sectionsMemOffset);
+
+		for (u32 i=0; i!=m_info.symbolCount; ++i)
+		{
+			auto& symbol = symbols[i];
+
+			if (symbol.StorageClass != ImageSymClassExternal)
+				continue;
+			if (!symbol.SectionNumber)
+				continue;
+
+			auto& section = sections[symbol.SectionNumber-1];
+			if (!(section.Characteristics & ImageScnLnkComdat))
+				continue;
+
+			StringView symbolName = GetSymbolName(symbol, m_info.stringTableMem);
+
+			auto findIt = toRemove.find(symbolName.ToString());
+			if (findIt == toRemove.end())
+				continue;
+
+			memset(&section, 0, sizeof(ImageSectionHeader));
+			symbol.SectionNumber = ImageSymUndefined;
+		}
+	}
+
 }
diff --git a/Engine/Source/Programs/UnrealBuildAccelerator/Common/Public/UbaObjectFile.h b/Engine/Source/Programs/UnrealBuildAccelerator/Common/Public/UbaObjectFile.h
--- a/Engine/Source/Programs/UnrealBuildAccelerator/Common/Public/UbaObjectFile.h
+++ b/Engine/Source/Programs/UnrealBuildAccelerator/Common/Public/UbaObjectFile.h
@@ -16,39 +16,46 @@ namespace uba
 	public:
 		bool Parse(Logger& logger, const tchar* filename);
 
-		bool CreateStripped(Logger& logger, const tchar* newFilename, const UnorderedSymbols& allNeededImports, const UnorderedSymbols& loopbacksToAdd);
+		bool CreateStripped(Logger& logger, const tchar* newFilename, const UnorderedSymbols& allNeededImports, UnorderedSymbols& loopbacksToAdd, UnorderedSymbols& toRemove);
 
 		const tchar* GetFileName() const;
 		const UnorderedSymbols& GetImports() const;
 		const UnorderedSymbols& GetExports() const;
-		const UnorderedSymbols& GetSymbols() const;
+		const UnorderedSymbols& GetPotentialDuplicates() const;
 
 		~ObjectFile();
 
 	private:
+		struct Info;
+
+		void ParseExports();
+		template<typename SymbolType> void ParseImports();
+
+		void WriteExports(Logger& logger, u8* newData, const UnorderedSymbols& allNeededImports);
+		template<typename SymbolType> void CalculateImports(Logger& logger, UnorderedSymbols& loopbacksToAdd, Vector<u32>& outImports);
+		template<typename SymbolType> void WriteImports(Logger& logger, u8* newData, Info& newInfo, const Vector<u32>& symbolsToAdd);
+		template<typename SymbolType> void RemoveSymbols(Logger& logger, u8* newData, Info& newInfo, UnorderedSymbols& toRemove);
+
+		struct Info
+		{
+			u32 sectionsMemOffset = 0;
+			u32 sectionCount = 0;
+			u64 directiveSectionMemOffset = 0;
+			u8* stringTableMem = nullptr;
+			u32 symbolsMemPos = 0;
+			u32 symbolCount = 0;
+		};
+
 		FileAccessor* m_file = nullptr;
 
 		u8* m_data = nullptr;
 		u64 m_dataSize = 0;
 
-		u32 m_sectionsMemOffset = 0;
-		u32 m_sectionCount = 0;
-		u64 m_directiveSectionMemOffset = 0;
-		u8* m_stringTableMem = nullptr;
-		u32 m_symbolsMemPos = 0;
-		u32 m_symbolCount = 0;
 		bool m_isBigObj = false;
-		
+		Info m_info;
+
 		UnorderedSymbols m_imports;
 		UnorderedSymbols m_exports;
-		UnorderedSymbols m_symbols;
-
-		void ParseExports();
-		template<typename SymbolType> void ParseImports();
-
-
-		void WriteExports(Logger& logger, u8* newData, const UnorderedSymbols& allNeededImports);
-		template<typename SymbolType> void CalculateImports(Logger& logger, const UnorderedSymbols& loopbacksToAdd, Vector<u32>& outImports);
-		template<typename SymbolType> void WriteImports(Logger& logger, u8* newData, const Vector<u32>& symbolsToAdd);
+		UnorderedSymbols m_potentialDuplicates;
 	};
 }
diff --git a/Engine/Source/Programs/UnrealBuildAccelerator/ObjTool/Private/UbaObjTool.cpp b/Engine/Source/Programs/UnrealBuildAccelerator/ObjTool/Private/UbaObjTool.cpp
--- a/Engine/Source/Programs/UnrealBuildAccelerator/ObjTool/Private/UbaObjTool.cpp
+++ b/Engine/Source/Programs/UnrealBuildAccelerator/ObjTool/Private/UbaObjTool.cpp
@@ -128,7 +128,7 @@ namespace uba
 		{
 			CriticalSection cs;
 			bool success = true;
-			UnorderedSymbols allNeededImports;
+			UnorderedSymbols allNeededImports; // Imports needed from the outside of the stripped obj files
 
 			u32 workerCount = DefaultProcessorCount;
 			WorkManagerImpl workManager(workerCount);
@@ -147,18 +147,19 @@ namespace uba
 			if (!success)
 				return -1;
 
-			UnorderedSymbols allSharedExports;
+			UnorderedSymbols allSharedExports; // Exports from all the obj files about to be stripped
 
-			struct ObjectFileRec { ObjectFile file; UnorderedSet<std::string> loopbacksToAdd; };
+			struct ObjectFileRec { ObjectFile file; UnorderedSymbols loopbacksToAdd;  UnorderedSymbols toRemove; };
 			Map<TString, ObjectFileRec> objectFileRecs;
+			CriticalSection cs2;
 
 			workManager.ParallelFor(workerCount, objFilesToStrip, [&](auto& it)
 				{
 					const TString& objFileName = *it;
 
-					cs.Enter();
+					cs2.Enter();
 					ObjectFile& objectFile = objectFileRecs.try_emplace(objFileName).first->second.file;
-					cs.Leave();
+					cs2.Leave();
 
 					bool res = objectFile.Parse(logger, objFileName.c_str());
 
@@ -171,21 +172,27 @@ namespace uba
 			if (!success)
 				return -1;
 
+			UnorderedSymbols duplicates;
+
 			// Figure out which loopback symbols that should be added to which obj file
 			for (auto& objFileName : objFilesToStrip)
 			{
-				ObjectFileRec& rec = objectFileRecs[objFileName];//kv.second;
+				ObjectFileRec& rec = objectFileRecs[objFileName];
 				for (auto& importSymbol : rec.file.GetImports())
 				{
 					if (strncmp(importSymbol.c_str(), "__imp_", 6) != 0)
 						continue;
-					std::string tmp = importSymbol.substr(6);
-					auto findIt = allSharedExports.find(tmp);
+					std::string importString = importSymbol.substr(6);
+					auto findIt = allSharedExports.find(importString);
 					if (findIt == allSharedExports.end())
 						continue;
-					rec.loopbacksToAdd.emplace(tmp);
+					rec.loopbacksToAdd.emplace(importString);
 					allSharedExports.erase(findIt);
 				}
+
+				for (auto& dupSymbol : rec.file.GetPotentialDuplicates())
+					if (!duplicates.insert(dupSymbol).second)
+						rec.toRemove.insert(dupSymbol);
 			}
 
 			workManager.ParallelFor(workerCount, objectFileRecs, [&](auto& it)
@@ -196,7 +203,7 @@ namespace uba
 					UBA_ASSERT(lastDot);
 					StringBuffer<> newFilename;
 					newFilename.Append(fileName, lastDot - fileName).Append(TC(".strip")).Append(lastDot);
-					rec.file.CreateStripped(logger, newFilename.data, allNeededImports, rec.loopbacksToAdd);
+					rec.file.CreateStripped(logger, newFilename.data, allNeededImports, rec.loopbacksToAdd, rec.toRemove);
 				});
 
 			//logger.Info(TC("Stripped %llu symbols from %llu obj files"), strippedSymbolCount.load(), objFilesToStrip.size());
