[UbaObjectTool]
* Added symbols for UE4_VISUALIZERS_HELPERS to not have to change ModuleBoilerplate.h
* Refactored code and added stubs for elf object files

diff --git a/Engine/Source/Programs/UnrealBuildAccelerator/Common/Private/UbaObjectFile.cpp b/Engine/Source/Programs/UnrealBuildAccelerator/Common/Private/UbaObjectFile.cpp
--- a/Engine/Source/Programs/UnrealBuildAccelerator/Common/Private/UbaObjectFile.cpp
+++ b/Engine/Source/Programs/UnrealBuildAccelerator/Common/Private/UbaObjectFile.cpp
@@ -2,343 +2,51 @@
 
 #include "UbaObjectFile.h"
 #include "UbaFileAccessor.h"
+#include "UbaObjectFileCoff.h"
+#include "UbaObjectFileElf.h"
 
 namespace uba
 {
-	constexpr u16 ImageFileMachineUnknown = 0;
-	constexpr u8 ImageSizeofShortName = 8; // IMAGE_SIZEOF_SHORT_NAME
-	constexpr u16 ImageSymClassExternal = 0x0002; // IMAGE_SYM_CLASS_EXTERNAL
-	constexpr u16 ImageSymUndefined = 0; // IMAGE_SYM_UNDEFINED
-
-	constexpr u32 ImageScnCntCode = 0x00000020; // IMAGE_SCN_CNT_CODE
-	constexpr u32 ImageScnMemExecute = 0x20000000; // IMAGE_SCN_MEM_EXECUTE
-	constexpr u32 ImageScnMemRead = 0x40000000; // IMAGE_SCN_MEM_READ
-	constexpr u32 ImageScnLnkComdat = 0x00001000; // IMAGE_SCN_LNK_COMDAT
-
-	const u16 ImageRelAmd64Addr64 = 0x0001; // IMAGE_REL_AMD64_ADDR64
-
-
-	#pragma pack(push)
-	#pragma pack(1)
-
-	struct ImageFileHeader // IMAGE_FILE_HEADER
+	ObjectFile* ObjectFile::CreateAndParse(Logger& logger, const tchar* filename)
 	{
-		u16 Machine;
-		u16 NumberOfSections;
-		u32 TimeDateStamp;
-		u32 PointerToSymbolTable;
-		u32 NumberOfSymbols;
-		u16 SizeOfOptionalHeader;
-		u16 Characteristics;
-	};
+		auto file = new FileAccessor(logger, filename);
+		auto fileGuard = MakeGuard([&]() { delete file; });
 
-	struct AnonObjectHeaderBigobj // ANON_OBJECT_HEADER_BIGOBJ
-	{
-		u16 Sig1;            // Must be IMAGE_FILE_MACHINE_UNKNOWN
-		u16 Sig2;            // Must be 0xffff
-		u16 Version;         // >= 2 (implies the Flags field is present)
-		u16 Machine;         // Actual machine - IMAGE_FILE_MACHINE_xxx
-		u32 TimeDateStamp;
-		Guid ClassID;         // {D1BAA1C7-BAEE-4ba9-AF20-FAF66AA4DCB8}
-		u32 SizeOfData;      // Size of data that follows the header
-		u32 Flags;           // 0x1 -> contains metadata
-		u32 MetaDataSize;    // Size of CLR metadata
-		u32 MetaDataOffset;  // Offset of CLR metadata
+		if (!file->OpenMemoryRead())
+			return nullptr;
 
-		// bigobj specifics
-		u32 NumberOfSections; // extended from WORD
-		u32 PointerToSymbolTable;
-		u32 NumberOfSymbols;
-	};
+		u8* data = file->GetData();
+		u64 dataSize = file->GetSize();
 
-	struct ImageSectionHeader // IMAGE_SECTION_HEADER
-	{
-		u8 Name[ImageSizeofShortName];
-		union
-		{
-			u32 PhysicalAddress;
-			u32 VirtualSize;
-		} Misc;
-		u32 VirtualAddress;
-		u32 SizeOfRawData;
-		u32 PointerToRawData;
-		u32 PointerToRelocations;
-		u32 PointerToLinenumbers;
-		u16 NumberOfRelocations;
-		u16 NumberOfLinenumbers;
-		u32 Characteristics;
-	};
-	static_assert(sizeof(ImageSectionHeader) == 40);
+		if (dataSize < 4)
+			return nullptr;
 
-	struct ImageRelocation // IMAGE_RELOCATION
-	{
-		union
-		{
-			u32 VirtualAddress;
-			u32 RelocCount; // Set to the real count when IMAGE_SCN_LNK_NRELOC_OVFL is set
-		};
-		u32 SymbolTableIndex;
-		u16 Type;
-	};
+		ObjectFile* objectFile = nullptr;
 
+		constexpr u8 elfMagic[] = { 0x7f, 'E', 'L', 'F' };
+		if (memcmp(data, elfMagic, sizeof(elfMagic)) == 0)
+			objectFile = new ObjectFileElf();
+		else
+			objectFile = new ObjectFileCoff();
 
-	struct ImageSymbolEx // IMAGE_SYMBOL_EX
-	{
-		union
-		{
-			char ShortName[8];
-			struct
-			{
-				u32 Short; // if 0, use LongName
-				u32 Long; // offset into string table
-			} Name;
-			u32 LongName[2]; // PBYTE  [2]
-		} N;
-		u32 Value;
-		u32 SectionNumber;
-		u16 Type;
-		u8 StorageClass;
-		u8 NumberOfAuxSymbols;
-	};
-
-	struct ImageSymbol // IMAGE_SYMBOL
-	{
-		union
-		{
-			u8 ShortName[8];
-			struct
-			{
-				u32   Short; // if 0, use LongName
-				u32   Long; // offset into string table
-			} Name;
-			u32   LongName[2]; // PBYTE [2]
-		} N;
-		u32 Value;
-		u16 SectionNumber;
-		u16 Type;
-		u8 StorageClass;
-		u8 NumberOfAuxSymbols;
-	};
-
-	#pragma pack(pop)
-
-	// These are symbols that are added to all dlls through some macros.
-	// When merging dlls we need to remove duplicates of these
-	static UnorderedSymbols PotentiallyDuplicatedSymbols = 
-	{
-		// operator new/delete 
-		"??2@YAPEAX_K@Z",
-		"??2@YAPEAX_KAEBUnothrow_t@std@@@Z",
-		"??2@YAPEAX_KW4align_val_t@std@@@Z",
-		"??2@YAPEAX_KW4align_val_t@std@@AEBUnothrow_t@1@@Z",
-		"??3@YAXPEAX@Z",
-		"??3@YAXPEAXAEBUnothrow_t@std@@@Z",
-		"??3@YAXPEAXW4align_val_t@std@@@Z",
-		"??3@YAXPEAXW4align_val_t@std@@AEBUnothrow_t@1@@Z",
-		"??3@YAXPEAX_K@Z",
-		"??3@YAXPEAX_KAEBUnothrow_t@std@@@Z",
-		"??3@YAXPEAX_KW4align_val_t@std@@@Z",
-		"??3@YAXPEAX_KW4align_val_t@std@@AEBUnothrow_t@1@@Z",
-		"??_U@YAPEAX_K@Z",
-		"??_U@YAPEAX_KAEBUnothrow_t@std@@@Z",
-		"??_U@YAPEAX_KW4align_val_t@std@@@Z",
-		"??_U@YAPEAX_KW4align_val_t@std@@AEBUnothrow_t@1@@Z",
-		"??_V@YAXPEAX@Z",
-		"??_V@YAXPEAXAEBUnothrow_t@std@@@Z",
-		"??_V@YAXPEAXW4align_val_t@std@@@Z",
-		"??_V@YAXPEAXW4align_val_t@std@@AEBUnothrow_t@1@@Z",
-		"??_V@YAXPEAX_K@Z",
-		"??_V@YAXPEAX_KAEBUnothrow_t@std@@@Z",
-		"??_V@YAXPEAX_KW4align_val_t@std@@@Z",
-		"??_V@YAXPEAX_KW4align_val_t@std@@AEBUnothrow_t@1@@Z",
-
-		// Memory
-		"?FMemory_Free@@YAXPEAX@Z",
-		"?FMemory_Malloc@@YAPEAX_K0@Z",
-		"?FMemory_Realloc@@YAPEAXPEAX_K1@Z",
-
-		// Have no idea why these are duplicated
-		"??_GIFunction_OwnedObject@Function@Private@Core@UE@@UEAAPEAXI@Z",
-		"??_GIDelegateInstance@@UEAAPEAXI@Z",
-	};
+		fileGuard.Cancel();
+		objectFile->m_file = file;
+		objectFile->m_data = data;
+		objectFile->m_dataSize = dataSize;
 
+		if (objectFile->Parse(logger, filename))
+			return objectFile;
 
-	bool IsBigObj(u8* data, u64 size)
-	{
-		if (size < sizeof(AnonObjectHeaderBigobj))
-			return false;
-		auto& header = *(AnonObjectHeaderBigobj *)data;
-		if (header.Sig1 != ImageFileMachineUnknown)
-			return false;
-		if (header.Sig2 != 0xffff)
-			return false;
-		if (header.Version < 2)
-			return false;
-		constexpr u8 bigObjClassId[16] = { 0xc7, 0xa1, 0xba, 0xd1, 0xee, 0xba, 0xa9, 0x4b, 0xaf, 0x20, 0xfa, 0xf6, 0x6a, 0xa4, 0xdc, 0xb8 };
-		if (header.ClassID != *(const Guid*)bigObjClassId)
-			return false;
-		return true;
+		delete objectFile;
+		return nullptr;
 	}
 
-	struct StringView
-	{
-		const char* strBegin;
-		const char* strEnd;
-
-		u32 Length() const
-		{
-			return u32(strEnd - strBegin);
-		}
-
-		bool StartsWith(const char* str, u32 strLen) const
-		{
-			if (strLen > Length())
-				return false;
-			return memcmp(strBegin, str, strLen) == 0;
-		}
-
-		bool Contains(const char* str, u32 strLen) const
-		{
-			const char* it = strBegin;
-			const char* itEnd = strEnd - strLen + 1;
-			while (it < itEnd)
-			{
-				if (memcmp(it, str, strLen) == 0)
-					return true;
-				++it;
-			}
-			return false;
-		}
-
-		bool Equals(const char* str, u32 strLen) const
-		{
-			if (strLen != Length())
-				return false;
-			return memcmp(strBegin, str, strLen) == 0;
-		}
-
-		std::string ToString() const
-		{
-			return std::string(strBegin, strEnd);
-		}
-	};
-
-	template<typename SymbolType>
-	StringView GetSymbolName(SymbolType& symbol, u8* stringTableMem)
-	{
-		if (symbol.N.Name.Short == 0)
-		{
-			const char* name = (char*)stringTableMem + symbol.N.Name.Long;
-			return { name, name + strlen(name) };
-		}
-		auto shortName = (char*)symbol.N.ShortName;
-		return { shortName, shortName + strnlen(shortName, ImageSizeofShortName) };
-	}
 
 	ObjectFile::~ObjectFile()
 	{
 		delete m_file;
 	}
 
-	bool ObjectFile::Parse(Logger& logger, const tchar* filename)
-	{
-		UBA_ASSERT(!m_file);
-
-		m_file = new FileAccessor(logger, filename);
-		if (!m_file->OpenMemoryRead())
-			return false;
-
-		m_data = m_file->GetData();
-		m_dataSize = m_file->GetSize();
-		m_isBigObj = IsBigObj(m_data, m_dataSize);
-
-		if (m_isBigObj)
-		{
-			auto& header = *(AnonObjectHeaderBigobj*)m_data;
-			m_info.symbolsMemPos = header.PointerToSymbolTable;
-			m_info.symbolCount = header.NumberOfSymbols;
-			m_info.stringTableMem = m_data + header.PointerToSymbolTable + header.NumberOfSymbols * sizeof(ImageSymbolEx);
-			m_info.sectionsMemOffset = sizeof(AnonObjectHeaderBigobj);
-			m_info.sectionCount = header.NumberOfSections;
-		}
-		else
-		{
-			auto& header = *(ImageFileHeader*)m_data;
-			m_info.symbolsMemPos = header.PointerToSymbolTable;
-			m_info.symbolCount = header.NumberOfSymbols;
-			m_info.stringTableMem = m_data + header.PointerToSymbolTable + header.NumberOfSymbols * sizeof(ImageSymbol);
-			m_info.sectionsMemOffset = sizeof(ImageFileHeader);
-			m_info.sectionCount = header.NumberOfSections;
-		}
-
-		auto sections = (ImageSectionHeader*)(m_data + m_info.sectionsMemOffset);
-		for (u32 i=0; i!=m_info.sectionCount; ++i)
-		{
-			if (strncmp((char*)sections[i].Name, ".drectve", 8) != 0)
-				continue;
-			m_info.directiveSectionMemOffset = u64((u8*)(sections + i) - m_data);
-			break;
-		}
-
-		ParseExports();
-
-		if (m_isBigObj)
-			ParseImports<ImageSymbolEx>();
-		else
-			ParseImports<ImageSymbol>();
-
-		return true;
-	}
-
-	bool ObjectFile::CreateStripped(Logger& logger, const tchar* newFilename, const UnorderedSymbols& allNeededImports, UnorderedSymbols& loopbacksToAdd, UnorderedSymbols& toRemove)
-	{
-		UBA_ASSERT(m_file);
-
-		Vector<u32> symbolsToAdd;
-		if (m_isBigObj)
-			CalculateImports<ImageSymbolEx>(logger, loopbacksToAdd, symbolsToAdd);
-		else
-			CalculateImports<ImageSymbol>(logger, loopbacksToAdd, symbolsToAdd);
-
-		u32 importsToFixCount = u32(symbolsToAdd.size());
-		u32 bytesToAdd = 0;
-
-		if (importsToFixCount)
-		{
-			bytesToAdd = sizeof(ImageSectionHeader);
-			bytesToAdd += (m_isBigObj ? sizeof(ImageSymbolEx) : sizeof(ImageSymbol)) * importsToFixCount * 2; // Symbols
-			bytesToAdd += (sizeof(ImageRelocation) + 8) * importsToFixCount; // relocations plus memory
-		}
-
-		FileAccessor newFile(logger, newFilename);
-		if (!newFile.CreateMemoryWrite(false, DefaultAttributes(), m_dataSize + bytesToAdd))
-			return false;
-
-		u8* newData = newFile.GetData();
-		Info newInfo = m_info;
-
-		if (importsToFixCount)
-		{
-			if (m_isBigObj)
-				WriteImports<ImageSymbolEx>(logger, newData, newInfo, symbolsToAdd);
-			else
-				WriteImports<ImageSymbol>(logger, newData, newInfo, symbolsToAdd);
-		}
-		else
-		{
-			memcpy(newData, m_data, m_dataSize);
-		}
-
-		if (m_isBigObj)
-			RemoveSymbols<ImageSymbolEx>(logger, newData, newInfo, toRemove);
-		else
-			RemoveSymbols<ImageSymbol>(logger, newData, newInfo, toRemove);
-
-		WriteExports(logger, newData, allNeededImports);
-
-		return newFile.Close();
-	}
-
 	const tchar* ObjectFile::GetFileName() const
 	{
 		return m_file->GetFileName();
@@ -358,320 +66,4 @@ namespace uba
 	{
 		return m_potentialDuplicates;
 	}
-
-	void ObjectFile::ParseExports()
-	{
-		if (!m_info.directiveSectionMemOffset)
-			return;
-		auto directiveSection = (ImageSectionHeader*)(m_data + m_info.directiveSectionMemOffset);
-		u8* directiveData = m_data + directiveSection->PointerToRawData;
-
-		static constexpr u8 utf8Bom[3] = { 0xef, 0xbb, 0xbf };
-		UBA_ASSERT(memcmp(directiveData, utf8Bom, 3) != 0);
-
-		auto str = (char*)directiveData;
-		while (str)
-		{
-			char* exportStr = strstr(str, "/EXPORT:");
-			if (!exportStr)
-				break;
-			exportStr += 8;
-			char* exportEnd = strchr(exportStr, ' ');
-			str = exportEnd;
-			if (!exportEnd)
-				exportEnd = exportStr + strlen(exportStr);
-			else
-				++str;
-			if (strncmp(exportEnd-5, ",DATA", 5) == 0)
-				exportEnd -= 5;
-					
-			m_exports.emplace(std::string(exportStr, exportEnd));
-		}
-	}
-
-	template<typename SymbolType>
-	void ObjectFile::ParseImports()
-	{
-		auto symbols = (SymbolType*)(m_data + m_info.symbolsMemPos);
-		for (u32 i=0; i!=m_info.symbolCount; ++i)
-		{
-			auto& symbol = symbols[i];
-
-			if (symbol.StorageClass != ImageSymClassExternal)
-				continue;
-
-			StringView symbolName = GetSymbolName(symbol, m_info.stringTableMem);
-			std::string symbolString = symbolName.ToString();
-
-			if (symbol.SectionNumber != ImageSymUndefined)
-			{
-				if (PotentiallyDuplicatedSymbols.find(symbolString) != PotentiallyDuplicatedSymbols.end())
-				{
-					m_potentialDuplicates.emplace(symbolString);
-				}
-			}
-			else
-			{
-				m_imports.emplace(symbolString);
-			}
-		}
-	}
-
-	void ObjectFile::WriteExports(Logger& logger, u8* newData, const UnorderedSymbols& allNeededImports)
-	{
-		if (!m_info.directiveSectionMemOffset)
-			return;
-
-		auto directiveSection = (const ImageSectionHeader*)(m_data + m_info.directiveSectionMemOffset);
-		if (directiveSection->SizeOfRawData < 10)
-			return;
-
-		const u8* directiveData = m_data + directiveSection->PointerToRawData;
-
-		auto newDirectiveSection = (ImageSectionHeader*)(newData + m_info.directiveSectionMemOffset);
-		u8* newDirectiveData = newData + newDirectiveSection->PointerToRawData;
-
-		char* writePos = (char*)newDirectiveData;
-		const char* lastCopyPos = (const char*)directiveData;
-
-		auto readPos = lastCopyPos;
-		auto readEnd = readPos + directiveSection->SizeOfRawData;
-		auto readLastPossiblePos = readEnd - 9;
-
-		while (true)
-		{
-			const char* exportStr = nullptr;
-			const char* it = readPos;
-
-			while (it < readLastPossiblePos)
-			{
-				if (memcmp(it, "/EXPORT:", 8) != 0)
-				{
-					++it;
-					continue;
-				}
-
-				exportStr = it;
-				break;
-			}
-
-			if (!exportStr)
-			{
-				readPos = readEnd - 1;
-				break;
-			}
-
-			const char* startPos = exportStr;
-			exportStr += 8;
-			const char* exportEnd = strchr(exportStr, ' ');
-			if (!exportEnd)
-				readPos = exportEnd = exportStr + strlen(exportStr);
-			else
-				readPos = exportEnd;
-
-			if (strncmp(exportEnd-5, ",DATA", 5) == 0)
-				exportEnd -= 5;
-
-			std::string temp(exportStr, exportEnd - exportStr);
-			if (allNeededImports.find(temp) != allNeededImports.end())
-				continue;
-			temp = "__imp_" + temp;
-			if (allNeededImports.find(temp) != allNeededImports.end())
-				continue;
-			
-			// TODO: Use exclusion list provided from command line
-			if (strstr(exportStr, "Initialize") == exportStr && strstr(exportStr, "Module"))
-				continue;
-
-			u64 toCopy = startPos - lastCopyPos - 1;
-			memcpy(writePos, lastCopyPos, toCopy);
-			writePos += toCopy;
-			lastCopyPos = readPos;
-			if (!*readPos)
-				break;
-		}
-
-		u64 toCopy = readPos - lastCopyPos;
-		memcpy(writePos, lastCopyPos, toCopy);
-		writePos += toCopy;
-		*writePos = 0;
-
-		newDirectiveSection->SizeOfRawData = u32((u8*)writePos - newDirectiveData);
-		UBA_ASSERT(newDirectiveSection->SizeOfRawData <= directiveSection->SizeOfRawData);
-
-		memset(writePos, 0, directiveSection->SizeOfRawData - newDirectiveSection->SizeOfRawData);
-	}
-
-	template<typename SymbolType> void ObjectFile::CalculateImports(Logger& logger, UnorderedSymbols& loopbacksToAdd, Vector<u32>& outImports)
-	{
-		// Calculate how much more memory we need by counting imports
-		auto symbols = (SymbolType*)(m_data + m_info.symbolsMemPos);
-
-		for (u32 i=0; i!=m_info.symbolCount; ++i)
-		{
-			auto& symbol = symbols[i];
-			if (symbol.StorageClass != ImageSymClassExternal)
-				continue;
-			if (symbol.SectionNumber != ImageSymUndefined)
-				continue;
-			StringView symbolName = GetSymbolName(symbol, m_info.stringTableMem);
-			if (!symbolName.StartsWith("__imp_", 6))
-				continue;
-			symbolName.strBegin += 6;
-			auto findIt = loopbacksToAdd.find(symbolName.ToString());
-			if (findIt == loopbacksToAdd.end())
-				continue;
-			loopbacksToAdd.erase(findIt);
-			outImports.push_back(symbol.N.Name.Long);
-		}
-	}
-
-	template<typename SymbolType>
-	void ObjectFile::WriteImports(Logger& logger, u8* newData, Info& newInfo, const Vector<u32>& symbolsToAdd)
-	{
-		u32 importsToFixCount = u32(symbolsToAdd.size());
-
-		// Copy header and sections
-		u32 offsetToAfterLastSection = m_info.sectionsMemOffset + m_info.sectionCount*sizeof(ImageSectionHeader);
-		memcpy(newData, m_data, offsetToAfterLastSection);
-
-		// Add another section
-		u32 newSectionIndex = m_info.sectionCount;
-		auto& newSection = *(ImageSectionHeader*)(newData + offsetToAfterLastSection);
-		memset(&newSection, 0, sizeof(ImageSectionHeader));
-		memcpy(newSection.Name, ".text$mn", 8);
-
-		// Add dummy memory for relocations
-		u32 newRelocationVirtualMemPos = offsetToAfterLastSection + sizeof(ImageSectionHeader);
-		u32 newRelocationVirtualMemSize = 8 * importsToFixCount;
-		memset(newData + newRelocationVirtualMemPos, 0, newRelocationVirtualMemSize);
-
-		u32 newRelocationsPos = newRelocationVirtualMemPos + newRelocationVirtualMemSize;
-		u32 newRelocationsSize = sizeof(ImageRelocation) * importsToFixCount;
-		newSection.PointerToRelocations = newRelocationsPos;
-		newSection.NumberOfRelocations = u16(importsToFixCount);
-		newSection.PointerToRawData = newRelocationVirtualMemPos;
-		newSection.SizeOfRawData = newRelocationVirtualMemSize;
-		newSection.Characteristics = ImageScnCntCode | ImageScnMemExecute | ImageScnMemRead;
-		UBA_ASSERT(importsToFixCount <= 65535);
-
-		// Add the new relocations and raw memory
-		u32 newSymbolIndex = m_info.symbolCount;
-		memset(newData + newRelocationsPos, 0, sizeof(ImageRelocation) *  importsToFixCount);
-		for (u32 i=0; i!=importsToFixCount; ++i)
-		{
-			auto& relocation = ((ImageRelocation*)(newData + newRelocationsPos))[i];
-			relocation.VirtualAddress = 8 * i;
-			relocation.SymbolTableIndex = newSymbolIndex + i;
-			relocation.Type = ImageRelAmd64Addr64;
-		}
-
-		// Offset of everything after new section and relocations
-		u32 memoryOffset = newRelocationsPos + newRelocationsSize - offsetToAfterLastSection;
-
-		u32 symbolTablePos;
-		u32 symbolTableSize;
-		if (m_isBigObj)
-		{
-			auto& header = *(AnonObjectHeaderBigobj*)newData;
-			++header.NumberOfSections;
-			header.PointerToSymbolTable += memoryOffset;
-			symbolTablePos = header.PointerToSymbolTable;
-			symbolTableSize = header.NumberOfSymbols * sizeof(SymbolType);
-			header.NumberOfSymbols += importsToFixCount*2;
-		}
-		else
-		{
-			auto& header = *(ImageFileHeader*)newData;
-			++header.NumberOfSections;
-			header.PointerToSymbolTable += memoryOffset;
-			symbolTablePos = header.PointerToSymbolTable;
-			symbolTableSize = header.NumberOfSymbols * sizeof(SymbolType);
-			header.NumberOfSymbols += importsToFixCount*2;
-		}
-
-		u32 offsetToAfterSymbolTable = symbolTablePos + symbolTableSize;
-
-		// Copy everything after sections til end of symbol table
-		u64 nextToCopySize = m_info.symbolsMemPos + symbolTableSize - offsetToAfterLastSection;
-		memcpy(newData + offsetToAfterLastSection + memoryOffset, m_data + offsetToAfterLastSection, nextToCopySize);
-
-		newInfo.symbolsMemPos += memoryOffset;
-
-		u32 newSymbolsPos = offsetToAfterSymbolTable;
-		u32 newSymbolsSize = sizeof(SymbolType)*importsToFixCount*2;
-		auto newSymbols = (SymbolType*)(newData + newSymbolsPos);
-		memset(newSymbols, 0, newSymbolsSize);
-		for (u32 i=0; i!=importsToFixCount; ++i)
-		{
-			auto& symbol = newSymbols[i];
-			symbol.N.Name.Long = symbolsToAdd[i] + 6; // Remove __imp_
-			symbol.SectionNumber = ImageSymUndefined;
-			symbol.StorageClass = ImageSymClassExternal;
-		}
-		for (u32 i=0; i!=importsToFixCount; ++i)
-		{
-			auto& symbol = newSymbols[i + importsToFixCount];
-			symbol.N.Name.Long = symbolsToAdd[i];
-			symbol.SectionNumber = (decltype(SymbolType::SectionNumber))(newSectionIndex + 1);
-			symbol.StorageClass = ImageSymClassExternal;
-			symbol.Value = i * 8;
-		}
-
-		u64 lastToCopySize = m_dataSize - (m_info.stringTableMem - m_data);
-		memcpy(newData + newSymbolsPos + newSymbolsSize, m_info.stringTableMem, lastToCopySize);
-
-
-		// Traverse all sections and update their PointerToRawData
-		auto sections = (ImageSectionHeader*)(newData + m_info.sectionsMemOffset);
-		for (u32 i=0; i!=m_info.sectionCount; ++i)
-		{
-			auto& section = sections[i];
-			if (section.PointerToRawData)
-			{
-				UBA_ASSERT(section.PointerToRawData < symbolTablePos);
-				section.PointerToRawData += memoryOffset;
-			}
-			if (section.PointerToRelocations)
-			{
-				UBA_ASSERT(section.PointerToRelocations < symbolTablePos);
-				section.PointerToRelocations += memoryOffset;
-			}
-			if (section.PointerToLinenumbers)
-			{
-				UBA_ASSERT(section.PointerToLinenumbers < symbolTablePos);
-				section.PointerToLinenumbers += memoryOffset;
-			}
-		}
-	}
-
-	template<typename SymbolType> void ObjectFile::RemoveSymbols(Logger& logger, u8* newData, Info& newInfo, UnorderedSymbols& toRemove)
-	{
-		auto symbols = (SymbolType*)(newData + newInfo.symbolsMemPos);
-		auto sections = (ImageSectionHeader*)(newData + newInfo.sectionsMemOffset);
-
-		for (u32 i=0; i!=m_info.symbolCount; ++i)
-		{
-			auto& symbol = symbols[i];
-
-			if (symbol.StorageClass != ImageSymClassExternal)
-				continue;
-			if (!symbol.SectionNumber)
-				continue;
-
-			auto& section = sections[symbol.SectionNumber-1];
-			if (!(section.Characteristics & ImageScnLnkComdat))
-				continue;
-
-			StringView symbolName = GetSymbolName(symbol, m_info.stringTableMem);
-
-			auto findIt = toRemove.find(symbolName.ToString());
-			if (findIt == toRemove.end())
-				continue;
-
-			memset(&section, 0, sizeof(ImageSectionHeader));
-			symbol.SectionNumber = ImageSymUndefined;
-		}
-	}
-
 }
