[UbaObjectTool]
* Added symbols for UE4_VISUALIZERS_HELPERS to not have to change ModuleBoilerplate.h
* Refactored code and added stubs for elf object files

diff --git a/Engine/Source/Programs/UnrealBuildAccelerator/Common/Private/UbaObjectFileCoff.cpp b/Engine/Source/Programs/UnrealBuildAccelerator/Common/Private/UbaObjectFileCoff.cpp
new file mode 100644
--- /dev/null
+++ b/Engine/Source/Programs/UnrealBuildAccelerator/Common/Private/UbaObjectFileCoff.cpp
@@ -0,0 +1,649 @@
+// Copyright Epic Games, Inc. All Rights Reserved.
+
+#include "UbaObjectFileCoff.h"
+#include "UbaFileAccessor.h"
+
+namespace uba
+{
+	constexpr u16 ImageFileMachineUnknown = 0;
+	constexpr u8 ImageSizeofShortName = 8; // IMAGE_SIZEOF_SHORT_NAME
+	constexpr u16 ImageSymClassExternal = 0x0002; // IMAGE_SYM_CLASS_EXTERNAL
+	constexpr u16 ImageSymUndefined = 0; // IMAGE_SYM_UNDEFINED
+
+	constexpr u32 ImageScnCntCode = 0x00000020; // IMAGE_SCN_CNT_CODE
+	constexpr u32 ImageScnMemExecute = 0x20000000; // IMAGE_SCN_MEM_EXECUTE
+	constexpr u32 ImageScnMemRead = 0x40000000; // IMAGE_SCN_MEM_READ
+	constexpr u32 ImageScnLnkComdat = 0x00001000; // IMAGE_SCN_LNK_COMDAT
+
+	const u16 ImageRelAmd64Addr64 = 0x0001; // IMAGE_REL_AMD64_ADDR64
+
+
+	#pragma pack(push)
+	#pragma pack(1)
+
+	struct Elf64Header
+	{
+	  unsigned char	e_ident[16];	/* Magic number and other info */
+	  u16 e_type;			/* Object file type */
+	  u16 e_machine;		/* Architecture */
+	  u32 e_version;		/* Object file version */
+	  u64 e_entry;		/* Entry point virtual address */
+	  u64 e_phoff;		/* Program header table file offset */
+	  u64 e_shoff;		/* Section header table file offset */
+	  u32 e_flags;		/* Processor-specific flags */
+	  u16 e_ehsize;		/* ELF header size in bytes */
+	  u16 e_phentsize;		/* Program header table entry size */
+	  u16 e_phnum;		/* Program header table entry count */
+	  u16 e_shentsize;		/* Section header table entry size */
+	  u16 e_shnum;		/* Section header table entry count */
+	  u16 e_shstrndx;		/* Section header string table index */
+	};
+
+
+	struct ImageFileHeader // IMAGE_FILE_HEADER
+	{
+		u16 Machine;
+		u16 NumberOfSections;
+		u32 TimeDateStamp;
+		u32 PointerToSymbolTable;
+		u32 NumberOfSymbols;
+		u16 SizeOfOptionalHeader;
+		u16 Characteristics;
+	};
+
+	struct AnonObjectHeaderBigobj // ANON_OBJECT_HEADER_BIGOBJ
+	{
+		u16 Sig1;            // Must be IMAGE_FILE_MACHINE_UNKNOWN
+		u16 Sig2;            // Must be 0xffff
+		u16 Version;         // >= 2 (implies the Flags field is present)
+		u16 Machine;         // Actual machine - IMAGE_FILE_MACHINE_xxx
+		u32 TimeDateStamp;
+		Guid ClassID;         // {D1BAA1C7-BAEE-4ba9-AF20-FAF66AA4DCB8}
+		u32 SizeOfData;      // Size of data that follows the header
+		u32 Flags;           // 0x1 -> contains metadata
+		u32 MetaDataSize;    // Size of CLR metadata
+		u32 MetaDataOffset;  // Offset of CLR metadata
+
+		// bigobj specifics
+		u32 NumberOfSections; // extended from WORD
+		u32 PointerToSymbolTable;
+		u32 NumberOfSymbols;
+	};
+
+	struct ImageSectionHeader // IMAGE_SECTION_HEADER
+	{
+		u8 Name[ImageSizeofShortName];
+		union
+		{
+			u32 PhysicalAddress;
+			u32 VirtualSize;
+		} Misc;
+		u32 VirtualAddress;
+		u32 SizeOfRawData;
+		u32 PointerToRawData;
+		u32 PointerToRelocations;
+		u32 PointerToLinenumbers;
+		u16 NumberOfRelocations;
+		u16 NumberOfLinenumbers;
+		u32 Characteristics;
+	};
+	static_assert(sizeof(ImageSectionHeader) == 40);
+
+	struct ImageRelocation // IMAGE_RELOCATION
+	{
+		union
+		{
+			u32 VirtualAddress;
+			u32 RelocCount; // Set to the real count when IMAGE_SCN_LNK_NRELOC_OVFL is set
+		};
+		u32 SymbolTableIndex;
+		u16 Type;
+	};
+
+
+	struct ImageSymbolEx // IMAGE_SYMBOL_EX
+	{
+		union
+		{
+			char ShortName[8];
+			struct
+			{
+				u32 Short; // if 0, use LongName
+				u32 Long; // offset into string table
+			} Name;
+			u32 LongName[2]; // PBYTE  [2]
+		} N;
+		u32 Value;
+		u32 SectionNumber;
+		u16 Type;
+		u8 StorageClass;
+		u8 NumberOfAuxSymbols;
+	};
+
+	struct ImageSymbol // IMAGE_SYMBOL
+	{
+		union
+		{
+			u8 ShortName[8];
+			struct
+			{
+				u32   Short; // if 0, use LongName
+				u32   Long; // offset into string table
+			} Name;
+			u32   LongName[2]; // PBYTE [2]
+		} N;
+		u32 Value;
+		u16 SectionNumber;
+		u16 Type;
+		u8 StorageClass;
+		u8 NumberOfAuxSymbols;
+	};
+
+	#pragma pack(pop)
+
+	// These are symbols that are added to all dlls through some macros.
+	// When merging dlls we need to remove duplicates of these
+	static UnorderedSymbols PotentiallyDuplicatedSymbols = 
+	{
+		// REPLACEMENT_OPERATOR_NEW_AND_DELETE
+		"??2@YAPEAX_K@Z",
+		"??2@YAPEAX_KAEBUnothrow_t@std@@@Z",
+		"??2@YAPEAX_KW4align_val_t@std@@@Z",
+		"??2@YAPEAX_KW4align_val_t@std@@AEBUnothrow_t@1@@Z",
+		"??3@YAXPEAX@Z",
+		"??3@YAXPEAXAEBUnothrow_t@std@@@Z",
+		"??3@YAXPEAXW4align_val_t@std@@@Z",
+		"??3@YAXPEAXW4align_val_t@std@@AEBUnothrow_t@1@@Z",
+		"??3@YAXPEAX_K@Z",
+		"??3@YAXPEAX_KAEBUnothrow_t@std@@@Z",
+		"??3@YAXPEAX_KW4align_val_t@std@@@Z",
+		"??3@YAXPEAX_KW4align_val_t@std@@AEBUnothrow_t@1@@Z",
+		"??_U@YAPEAX_K@Z",
+		"??_U@YAPEAX_KAEBUnothrow_t@std@@@Z",
+		"??_U@YAPEAX_KW4align_val_t@std@@@Z",
+		"??_U@YAPEAX_KW4align_val_t@std@@AEBUnothrow_t@1@@Z",
+		"??_V@YAXPEAX@Z",
+		"??_V@YAXPEAXAEBUnothrow_t@std@@@Z",
+		"??_V@YAXPEAXW4align_val_t@std@@@Z",
+		"??_V@YAXPEAXW4align_val_t@std@@AEBUnothrow_t@1@@Z",
+		"??_V@YAXPEAX_K@Z",
+		"??_V@YAXPEAX_KAEBUnothrow_t@std@@@Z",
+		"??_V@YAXPEAX_KW4align_val_t@std@@@Z",
+		"??_V@YAXPEAX_KW4align_val_t@std@@AEBUnothrow_t@1@@Z",
+
+		// UE_DEFINE_FMEMORY_WRAPPERS
+		"?FMemory_Free@@YAXPEAX@Z",
+		"?FMemory_Malloc@@YAPEAX_K0@Z",
+		"?FMemory_Realloc@@YAPEAXPEAX_K1@Z",
+
+		// UE4_VISUALIZERS_HELPERS
+		"?GNameBlocksDebug@@3PEAPEAEEA",
+		"?GObjectIndexToPackedObjectRefDebug@@3AEAPEB_KEA",
+		"?GObjectArrayForDebugVisualizers@@3AEAPEAVFChunkedFixedUObjectArray@@EA",
+		"?GComplexObjectPathDebug@@3AEAPEAUFStoredObjectPathDebug@Private@CoreUObject@UE@@EA",
+		"?GObjectHandlePackageDebug@@3AEAPEAUFObjectHandlePackageDebugData@Private@CoreUObject@UE@@EA",
+
+		// Have no idea why these are duplicated
+		"??_GIFunction_OwnedObject@Function@Private@Core@UE@@UEAAPEAXI@Z",
+		"??_GIDelegateInstance@@UEAAPEAXI@Z",
+	};
+
+	template<typename SymbolType>
+	StringView GetSymbolName(SymbolType& symbol, u8* stringTableMem)
+	{
+		if (symbol.N.Name.Short == 0)
+		{
+			const char* name = (char*)stringTableMem + symbol.N.Name.Long;
+			return { name, name + strlen(name) };
+		}
+		auto shortName = (char*)symbol.N.ShortName;
+		return { shortName, shortName + strnlen(shortName, ImageSizeofShortName) };
+	}
+
+	bool IsBigObj(u8* data, u64 size)
+	{
+		if (size < sizeof(AnonObjectHeaderBigobj))
+			return false;
+		auto& header = *(AnonObjectHeaderBigobj *)data;
+		if (header.Sig1 != ImageFileMachineUnknown)
+			return false;
+		if (header.Sig2 != 0xffff)
+			return false;
+		if (header.Version < 2)
+			return false;
+		constexpr u8 bigObjClassId[16] = { 0xc7, 0xa1, 0xba, 0xd1, 0xee, 0xba, 0xa9, 0x4b, 0xaf, 0x20, 0xfa, 0xf6, 0x6a, 0xa4, 0xdc, 0xb8 };
+		if (header.ClassID != *(const Guid*)bigObjClassId)
+			return false;
+		return true;
+	}
+
+	bool ObjectFileCoff::Parse(Logger& logger, const tchar* filename)
+	{
+		m_isBigObj = IsBigObj(m_data, m_dataSize);
+		
+		if (m_isBigObj)
+		{
+			auto& header = *(AnonObjectHeaderBigobj*)m_data;
+			m_info.symbolsMemPos = header.PointerToSymbolTable;
+			m_info.symbolCount = header.NumberOfSymbols;
+			m_info.stringTableMem = m_data + header.PointerToSymbolTable + header.NumberOfSymbols * sizeof(ImageSymbolEx);
+			m_info.sectionsMemOffset = sizeof(AnonObjectHeaderBigobj);
+			m_info.sectionCount = header.NumberOfSections;
+		}
+		else
+		{
+			auto& header = *(ImageFileHeader*)m_data;
+			m_info.symbolsMemPos = header.PointerToSymbolTable;
+			m_info.symbolCount = header.NumberOfSymbols;
+			m_info.stringTableMem = m_data + header.PointerToSymbolTable + header.NumberOfSymbols * sizeof(ImageSymbol);
+			m_info.sectionsMemOffset = sizeof(ImageFileHeader);
+			m_info.sectionCount = header.NumberOfSections;
+		}
+
+		if (!ParseExports())
+			return false;
+
+		if (m_isBigObj)
+			ParseImports<ImageSymbolEx>();
+		else
+			ParseImports<ImageSymbol>();
+
+		return true;
+	}
+
+	bool ObjectFileCoff::CreateStripped(Logger& logger, const tchar* newFilename, const UnorderedSymbols& allNeededImports)
+	{
+		UBA_ASSERT(m_file);
+
+		Vector<u32> symbolsToAdd;
+		if (m_isBigObj)
+			CalculateImports<ImageSymbolEx>(logger, symbolsToAdd);
+		else
+			CalculateImports<ImageSymbol>(logger, symbolsToAdd);
+
+		u32 importsToFixCount = u32(symbolsToAdd.size());
+		u32 bytesToAdd = 0;
+
+		if (importsToFixCount)
+		{
+			bytesToAdd = sizeof(ImageSectionHeader);
+			bytesToAdd += (m_isBigObj ? sizeof(ImageSymbolEx) : sizeof(ImageSymbol)) * importsToFixCount * 2; // Symbols
+			bytesToAdd += (sizeof(ImageRelocation) + 8) * importsToFixCount; // relocations plus memory
+		}
+
+		FileAccessor newFile(logger, newFilename);
+		if (!newFile.CreateMemoryWrite(false, DefaultAttributes(), m_dataSize + bytesToAdd))
+			return false;
+
+		u8* newData = newFile.GetData();
+		Info newInfo = m_info;
+
+		if (importsToFixCount)
+		{
+			if (m_isBigObj)
+				WriteImports<ImageSymbolEx>(logger, newData, newInfo, symbolsToAdd);
+			else
+				WriteImports<ImageSymbol>(logger, newData, newInfo, symbolsToAdd);
+		}
+		else
+		{
+			memcpy(newData, m_data, m_dataSize);
+		}
+
+		if (m_isBigObj)
+			RemoveSymbols<ImageSymbolEx>(logger, newData, newInfo);
+		else
+			RemoveSymbols<ImageSymbol>(logger, newData, newInfo);
+
+		WriteExports(logger, newData, allNeededImports);
+
+		return newFile.Close();
+	}
+
+	bool ObjectFileCoff::ComputeLoopbacksAndDuplicates(UnorderedSymbols& allSharedExports, UnorderedSymbols& duplicates)
+	{
+		for (auto& importSymbol : m_imports)
+		{
+			if (strncmp(importSymbol.c_str(), "__imp_", 6) != 0)
+				continue;
+			std::string importString = importSymbol.substr(6);
+			auto findIt = allSharedExports.find(importString);
+			if (findIt == allSharedExports.end())
+				continue;
+			m_loopbacksToAdd.emplace(importString);
+			allSharedExports.erase(findIt);
+		}
+
+		for (auto& dupSymbol : m_potentialDuplicates)
+			if (!duplicates.insert(dupSymbol).second)
+				m_toRemove.insert(dupSymbol);
+
+		return true;
+	}
+
+	bool ObjectFileCoff::ParseExports()
+	{
+		auto sections = (ImageSectionHeader*)(m_data + m_info.sectionsMemOffset);
+		for (u32 i=0; i!=m_info.sectionCount; ++i)
+		{
+			if (strncmp((char*)sections[i].Name, ".drectve", 8) != 0)
+				continue;
+			m_info.directiveSectionMemOffset = u64((u8*)(sections + i) - m_data);
+			break;
+		}
+		if (!m_info.directiveSectionMemOffset)
+			return true;
+
+		auto directiveSection = (ImageSectionHeader*)(m_data + m_info.directiveSectionMemOffset);
+		u8* directiveData = m_data + directiveSection->PointerToRawData;
+
+		static constexpr u8 utf8Bom[3] = { 0xef, 0xbb, 0xbf };
+		UBA_ASSERT(memcmp(directiveData, utf8Bom, 3) != 0);
+
+		auto str = (char*)directiveData;
+		while (str)
+		{
+			char* exportStr = strstr(str, "/EXPORT:");
+			if (!exportStr)
+				break;
+			exportStr += 8;
+			char* exportEnd = strchr(exportStr, ' ');
+			str = exportEnd;
+			if (!exportEnd)
+				exportEnd = exportStr + strlen(exportStr);
+			else
+				++str;
+			if (strncmp(exportEnd-5, ",DATA", 5) == 0)
+				exportEnd -= 5;
+					
+			m_exports.emplace(std::string(exportStr, exportEnd));
+		}
+		return true;
+	}
+
+	template<typename SymbolType>
+	void ObjectFileCoff::ParseImports()
+	{
+		auto symbols = (SymbolType*)(m_data + m_info.symbolsMemPos);
+		for (u32 i=0; i!=m_info.symbolCount; ++i)
+		{
+			auto& symbol = symbols[i];
+
+			if (symbol.StorageClass != ImageSymClassExternal)
+				continue;
+
+			StringView symbolName = GetSymbolName(symbol, m_info.stringTableMem);
+			std::string symbolString = symbolName.ToString();
+
+			if (symbol.SectionNumber != ImageSymUndefined)
+			{
+				if (PotentiallyDuplicatedSymbols.find(symbolString) != PotentiallyDuplicatedSymbols.end())
+				{
+					m_potentialDuplicates.emplace(symbolString);
+				}
+			}
+			else
+			{
+				m_imports.emplace(symbolString);
+			}
+		}
+	}
+
+	void ObjectFileCoff::WriteExports(Logger& logger, u8* newData, const UnorderedSymbols& allNeededImports)
+	{
+		if (!m_info.directiveSectionMemOffset)
+			return;
+
+		auto directiveSection = (const ImageSectionHeader*)(m_data + m_info.directiveSectionMemOffset);
+		if (directiveSection->SizeOfRawData < 10)
+			return;
+
+		const u8* directiveData = m_data + directiveSection->PointerToRawData;
+
+		auto newDirectiveSection = (ImageSectionHeader*)(newData + m_info.directiveSectionMemOffset);
+		u8* newDirectiveData = newData + newDirectiveSection->PointerToRawData;
+
+		char* writePos = (char*)newDirectiveData;
+		const char* lastCopyPos = (const char*)directiveData;
+
+		auto readPos = lastCopyPos;
+		auto readEnd = readPos + directiveSection->SizeOfRawData;
+		auto readLastPossiblePos = readEnd - 9;
+
+		while (true)
+		{
+			const char* exportStr = nullptr;
+			const char* it = readPos;
+
+			while (it < readLastPossiblePos)
+			{
+				if (memcmp(it, "/EXPORT:", 8) != 0)
+				{
+					++it;
+					continue;
+				}
+
+				exportStr = it;
+				break;
+			}
+
+			if (!exportStr)
+			{
+				readPos = readEnd - 1;
+				break;
+			}
+
+			const char* startPos = exportStr;
+			exportStr += 8;
+			const char* exportEnd = strchr(exportStr, ' ');
+			if (!exportEnd)
+				readPos = exportEnd = exportStr + strlen(exportStr);
+			else
+				readPos = exportEnd;
+
+			if (strncmp(exportEnd-5, ",DATA", 5) == 0)
+				exportEnd -= 5;
+
+			std::string temp(exportStr, exportEnd - exportStr);
+			if (allNeededImports.find(temp) != allNeededImports.end())
+				continue;
+			temp = "__imp_" + temp;
+			if (allNeededImports.find(temp) != allNeededImports.end())
+				continue;
+			
+			// TODO: Use exclusion list provided from command line
+			if (strstr(exportStr, "Initialize") == exportStr && strstr(exportStr, "Module"))
+				continue;
+
+			u64 toCopy = startPos - lastCopyPos - 1;
+			memcpy(writePos, lastCopyPos, toCopy);
+			writePos += toCopy;
+			lastCopyPos = readPos;
+			if (!*readPos)
+				break;
+		}
+
+		u64 toCopy = readPos - lastCopyPos;
+		memcpy(writePos, lastCopyPos, toCopy);
+		writePos += toCopy;
+		*writePos = 0;
+
+		newDirectiveSection->SizeOfRawData = u32((u8*)writePos - newDirectiveData);
+		UBA_ASSERT(newDirectiveSection->SizeOfRawData <= directiveSection->SizeOfRawData);
+
+		memset(writePos, 0, directiveSection->SizeOfRawData - newDirectiveSection->SizeOfRawData);
+	}
+
+	template<typename SymbolType> void ObjectFileCoff::CalculateImports(Logger& logger, Vector<u32>& outImports)
+	{
+		// Calculate how much more memory we need by counting imports
+		auto symbols = (SymbolType*)(m_data + m_info.symbolsMemPos);
+
+		for (u32 i=0; i!=m_info.symbolCount; ++i)
+		{
+			auto& symbol = symbols[i];
+			if (symbol.StorageClass != ImageSymClassExternal)
+				continue;
+			if (symbol.SectionNumber != ImageSymUndefined)
+				continue;
+			StringView symbolName = GetSymbolName(symbol, m_info.stringTableMem);
+			if (!symbolName.StartsWith("__imp_", 6))
+				continue;
+			symbolName.strBegin += 6;
+			auto findIt = m_loopbacksToAdd.find(symbolName.ToString());
+			if (findIt == m_loopbacksToAdd.end())
+				continue;
+			m_loopbacksToAdd.erase(findIt);
+			outImports.push_back(symbol.N.Name.Long);
+		}
+	}
+
+	template<typename SymbolType>
+	void ObjectFileCoff::WriteImports(Logger& logger, u8* newData, Info& newInfo, const Vector<u32>& symbolsToAdd)
+	{
+		u32 importsToFixCount = u32(symbolsToAdd.size());
+
+		// Copy header and sections
+		u32 offsetToAfterLastSection = m_info.sectionsMemOffset + m_info.sectionCount*sizeof(ImageSectionHeader);
+		memcpy(newData, m_data, offsetToAfterLastSection);
+
+		// Add another section
+		u32 newSectionIndex = m_info.sectionCount;
+		auto& newSection = *(ImageSectionHeader*)(newData + offsetToAfterLastSection);
+		memset(&newSection, 0, sizeof(ImageSectionHeader));
+		memcpy(newSection.Name, ".text$mn", 8);
+
+		// Add dummy memory for relocations
+		u32 newRelocationVirtualMemPos = offsetToAfterLastSection + sizeof(ImageSectionHeader);
+		u32 newRelocationVirtualMemSize = 8 * importsToFixCount;
+		memset(newData + newRelocationVirtualMemPos, 0, newRelocationVirtualMemSize);
+
+		u32 newRelocationsPos = newRelocationVirtualMemPos + newRelocationVirtualMemSize;
+		u32 newRelocationsSize = sizeof(ImageRelocation) * importsToFixCount;
+		newSection.PointerToRelocations = newRelocationsPos;
+		newSection.NumberOfRelocations = u16(importsToFixCount);
+		newSection.PointerToRawData = newRelocationVirtualMemPos;
+		newSection.SizeOfRawData = newRelocationVirtualMemSize;
+		newSection.Characteristics = ImageScnCntCode | ImageScnMemExecute | ImageScnMemRead;
+		UBA_ASSERT(importsToFixCount <= 65535);
+
+		// Add the new relocations and raw memory
+		u32 newSymbolIndex = m_info.symbolCount;
+		memset(newData + newRelocationsPos, 0, sizeof(ImageRelocation) *  importsToFixCount);
+		for (u32 i=0; i!=importsToFixCount; ++i)
+		{
+			auto& relocation = ((ImageRelocation*)(newData + newRelocationsPos))[i];
+			relocation.VirtualAddress = 8 * i;
+			relocation.SymbolTableIndex = newSymbolIndex + i;
+			relocation.Type = ImageRelAmd64Addr64;
+		}
+
+		// Offset of everything after new section and relocations
+		u32 memoryOffset = newRelocationsPos + newRelocationsSize - offsetToAfterLastSection;
+
+		u32 symbolTablePos;
+		u32 symbolTableSize;
+		if (m_isBigObj)
+		{
+			auto& header = *(AnonObjectHeaderBigobj*)newData;
+			++header.NumberOfSections;
+			header.PointerToSymbolTable += memoryOffset;
+			symbolTablePos = header.PointerToSymbolTable;
+			symbolTableSize = header.NumberOfSymbols * sizeof(SymbolType);
+			header.NumberOfSymbols += importsToFixCount*2;
+		}
+		else
+		{
+			auto& header = *(ImageFileHeader*)newData;
+			++header.NumberOfSections;
+			header.PointerToSymbolTable += memoryOffset;
+			symbolTablePos = header.PointerToSymbolTable;
+			symbolTableSize = header.NumberOfSymbols * sizeof(SymbolType);
+			header.NumberOfSymbols += importsToFixCount*2;
+		}
+
+		u32 offsetToAfterSymbolTable = symbolTablePos + symbolTableSize;
+
+		// Copy everything after sections til end of symbol table
+		u64 nextToCopySize = m_info.symbolsMemPos + symbolTableSize - offsetToAfterLastSection;
+		memcpy(newData + offsetToAfterLastSection + memoryOffset, m_data + offsetToAfterLastSection, nextToCopySize);
+
+		newInfo.symbolsMemPos += memoryOffset;
+
+		u32 newSymbolsPos = offsetToAfterSymbolTable;
+		u32 newSymbolsSize = sizeof(SymbolType)*importsToFixCount*2;
+		auto newSymbols = (SymbolType*)(newData + newSymbolsPos);
+		memset(newSymbols, 0, newSymbolsSize);
+		for (u32 i=0; i!=importsToFixCount; ++i)
+		{
+			auto& symbol = newSymbols[i];
+			symbol.N.Name.Long = symbolsToAdd[i] + 6; // Remove __imp_
+			symbol.SectionNumber = ImageSymUndefined;
+			symbol.StorageClass = ImageSymClassExternal;
+		}
+		for (u32 i=0; i!=importsToFixCount; ++i)
+		{
+			auto& symbol = newSymbols[i + importsToFixCount];
+			symbol.N.Name.Long = symbolsToAdd[i];
+			symbol.SectionNumber = (decltype(SymbolType::SectionNumber))(newSectionIndex + 1);
+			symbol.StorageClass = ImageSymClassExternal;
+			symbol.Value = i * 8;
+		}
+
+		u64 lastToCopySize = m_dataSize - (m_info.stringTableMem - m_data);
+		memcpy(newData + newSymbolsPos + newSymbolsSize, m_info.stringTableMem, lastToCopySize);
+
+
+		// Traverse all sections and update their PointerToRawData
+		auto sections = (ImageSectionHeader*)(newData + m_info.sectionsMemOffset);
+		for (u32 i=0; i!=m_info.sectionCount; ++i)
+		{
+			auto& section = sections[i];
+			if (section.PointerToRawData)
+			{
+				UBA_ASSERT(section.PointerToRawData < symbolTablePos);
+				section.PointerToRawData += memoryOffset;
+			}
+			if (section.PointerToRelocations)
+			{
+				UBA_ASSERT(section.PointerToRelocations < symbolTablePos);
+				section.PointerToRelocations += memoryOffset;
+			}
+			if (section.PointerToLinenumbers)
+			{
+				UBA_ASSERT(section.PointerToLinenumbers < symbolTablePos);
+				section.PointerToLinenumbers += memoryOffset;
+			}
+		}
+	}
+
+	template<typename SymbolType> void ObjectFileCoff::RemoveSymbols(Logger& logger, u8* newData, Info& newInfo)
+	{
+		auto symbols = (SymbolType*)(newData + newInfo.symbolsMemPos);
+		auto sections = (ImageSectionHeader*)(newData + newInfo.sectionsMemOffset);
+
+		for (u32 i=0; i!=m_info.symbolCount; ++i)
+		{
+			auto& symbol = symbols[i];
+
+			if (symbol.StorageClass != ImageSymClassExternal)
+				continue;
+
+			if (!symbol.SectionNumber)
+				continue;
+
+			StringView symbolName = GetSymbolName(symbol, m_info.stringTableMem);
+			auto findIt = m_toRemove.find(symbolName.ToString());
+			if (findIt == m_toRemove.end())
+				continue;
+
+			auto& section = sections[symbol.SectionNumber-1];
+
+			symbol.SectionNumber = ImageSymUndefined;
+
+			if (section.Characteristics & ImageScnLnkComdat)
+				memset(&section, 0, sizeof(ImageSectionHeader));
+		}
+	}
+
+}
diff --git a/Engine/Source/Programs/UnrealBuildAccelerator/Common/Private/UbaObjectFileElf.cpp b/Engine/Source/Programs/UnrealBuildAccelerator/Common/Private/UbaObjectFileElf.cpp
new file mode 100644
--- /dev/null
+++ b/Engine/Source/Programs/UnrealBuildAccelerator/Common/Private/UbaObjectFileElf.cpp
@@ -0,0 +1,30 @@
+// Copyright Epic Games, Inc. All Rights Reserved.
+
+#include "UbaObjectFileElf.h"
+#include "UbaFileAccessor.h"
+
+namespace uba
+{
+
+	bool ObjectFileElf::Parse(Logger& logger, const tchar* filename)
+	{
+		return true;
+	}
+
+	bool ObjectFileElf::ComputeLoopbacksAndDuplicates(UnorderedSymbols& allSharedExports, UnorderedSymbols& duplicates)
+	{
+		return true;
+	}
+
+	bool ObjectFileElf::CreateStripped(Logger& logger, const tchar* newFilename, const UnorderedSymbols& allNeededImports)
+	{
+		FileAccessor newFile(logger, newFilename);
+		if (!newFile.CreateMemoryWrite(false, DefaultAttributes(), m_dataSize))
+			return false;
+
+		u8* newData = newFile.GetData();
+		memcpy(newData, m_data, m_dataSize);
+
+		return newFile.Close();
+	}
+}
diff --git a/Engine/Source/Programs/UnrealBuildAccelerator/Common/Public/UbaObjectFile.h b/Engine/Source/Programs/UnrealBuildAccelerator/Common/Public/UbaObjectFile.h
--- a/Engine/Source/Programs/UnrealBuildAccelerator/Common/Public/UbaObjectFile.h
+++ b/Engine/Source/Programs/UnrealBuildAccelerator/Common/Public/UbaObjectFile.h
@@ -14,48 +14,71 @@ namespace uba
 	class ObjectFile
 	{
 	public:
-		bool Parse(Logger& logger, const tchar* filename);
+		static ObjectFile* CreateAndParse(Logger& logger, const tchar* filename);
 
-		bool CreateStripped(Logger& logger, const tchar* newFilename, const UnorderedSymbols& allNeededImports, UnorderedSymbols& loopbacksToAdd, UnorderedSymbols& toRemove);
+		virtual bool ComputeLoopbacksAndDuplicates(UnorderedSymbols& allSharedExports, UnorderedSymbols& duplicates) = 0;
+		virtual bool CreateStripped(Logger& logger, const tchar* newFilename, const UnorderedSymbols& allNeededImports) = 0;
 
 		const tchar* GetFileName() const;
 		const UnorderedSymbols& GetImports() const;
 		const UnorderedSymbols& GetExports() const;
 		const UnorderedSymbols& GetPotentialDuplicates() const;
 
-		~ObjectFile();
+		virtual ~ObjectFile();
 
-	private:
-		struct Info;
-
-		void ParseExports();
-		template<typename SymbolType> void ParseImports();
-
-		void WriteExports(Logger& logger, u8* newData, const UnorderedSymbols& allNeededImports);
-		template<typename SymbolType> void CalculateImports(Logger& logger, UnorderedSymbols& loopbacksToAdd, Vector<u32>& outImports);
-		template<typename SymbolType> void WriteImports(Logger& logger, u8* newData, Info& newInfo, const Vector<u32>& symbolsToAdd);
-		template<typename SymbolType> void RemoveSymbols(Logger& logger, u8* newData, Info& newInfo, UnorderedSymbols& toRemove);
-
-		struct Info
-		{
-			u32 sectionsMemOffset = 0;
-			u32 sectionCount = 0;
-			u64 directiveSectionMemOffset = 0;
-			u8* stringTableMem = nullptr;
-			u32 symbolsMemPos = 0;
-			u32 symbolCount = 0;
-		};
+	protected:
+		virtual bool Parse(Logger& logger, const tchar* filename) = 0;
 
 		FileAccessor* m_file = nullptr;
-
 		u8* m_data = nullptr;
 		u64 m_dataSize = 0;
 
-		bool m_isBigObj = false;
-		Info m_info;
-
 		UnorderedSymbols m_imports;
 		UnorderedSymbols m_exports;
 		UnorderedSymbols m_potentialDuplicates;
 	};
+
+
+	struct StringView
+	{
+		const char* strBegin;
+		const char* strEnd;
+
+		u32 Length() const
+		{
+			return u32(strEnd - strBegin);
+		}
+
+		bool StartsWith(const char* str, u32 strLen) const
+		{
+			if (strLen > Length())
+				return false;
+			return memcmp(strBegin, str, strLen) == 0;
+		}
+
+		bool Contains(const char* str, u32 strLen) const
+		{
+			const char* it = strBegin;
+			const char* itEnd = strEnd - strLen + 1;
+			while (it < itEnd)
+			{
+				if (memcmp(it, str, strLen) == 0)
+					return true;
+				++it;
+			}
+			return false;
+		}
+
+		bool Equals(const char* str, u32 strLen) const
+		{
+			if (strLen != Length())
+				return false;
+			return memcmp(strBegin, str, strLen) == 0;
+		}
+
+		std::string ToString() const
+		{
+			return std::string(strBegin, strEnd);
+		}
+	};
 }
diff --git a/Engine/Source/Programs/UnrealBuildAccelerator/Common/Public/UbaObjectFileCoff.h b/Engine/Source/Programs/UnrealBuildAccelerator/Common/Public/UbaObjectFileCoff.h
new file mode 100644
--- /dev/null
+++ b/Engine/Source/Programs/UnrealBuildAccelerator/Common/Public/UbaObjectFileCoff.h
@@ -0,0 +1,45 @@
+// Copyright Epic Games, Inc. All Rights Reserved.
+
+#pragma once
+
+#include "UbaObjectFile.h"
+
+namespace uba
+{
+	class ObjectFileCoff : public ObjectFile
+	{
+	public:
+		virtual bool Parse(Logger& logger, const tchar* filename) override;
+		virtual bool ComputeLoopbacksAndDuplicates(UnorderedSymbols& allSharedExports, UnorderedSymbols& duplicates) override;
+		virtual bool CreateStripped(Logger& logger, const tchar* newFilename, const UnorderedSymbols& allNeededImports) override;
+
+	private:
+		struct Info;
+
+		bool ParseExports();
+		template<typename SymbolType> void ParseImports();
+
+		void WriteExports(Logger& logger, u8* newData, const UnorderedSymbols& allNeededImports);
+		template<typename SymbolType> void CalculateImports(Logger& logger, Vector<u32>& outImports);
+		template<typename SymbolType> void WriteImports(Logger& logger, u8* newData, Info& newInfo, const Vector<u32>& symbolsToAdd);
+		template<typename SymbolType> void RemoveSymbols(Logger& logger, u8* newData, Info& newInfo);
+
+		struct Info
+		{
+			u32 sectionsMemOffset = 0;
+			u32 sectionCount = 0;
+			u64 directiveSectionMemOffset = 0;
+			u8* stringTableMem = nullptr;
+			u32 symbolsMemPos = 0;
+			u32 symbolCount = 0;
+		};
+
+		bool m_isBigObj = false;
+		Info m_info;
+
+
+
+		UnorderedSymbols m_loopbacksToAdd;
+		UnorderedSymbols m_toRemove;
+	};
+}
diff --git a/Engine/Source/Programs/UnrealBuildAccelerator/Common/Public/UbaObjectFileElf.h b/Engine/Source/Programs/UnrealBuildAccelerator/Common/Public/UbaObjectFileElf.h
new file mode 100644
--- /dev/null
+++ b/Engine/Source/Programs/UnrealBuildAccelerator/Common/Public/UbaObjectFileElf.h
@@ -0,0 +1,16 @@
+// Copyright Epic Games, Inc. All Rights Reserved.
+
+#pragma once
+
+#include "UbaObjectFile.h"
+
+namespace uba
+{
+	class ObjectFileElf : public ObjectFile
+	{
+	public:
+		virtual bool Parse(Logger& logger, const tchar* filename) override;
+		virtual bool ComputeLoopbacksAndDuplicates(UnorderedSymbols& allSharedExports, UnorderedSymbols& duplicates) override;
+		virtual bool CreateStripped(Logger& logger, const tchar* newFilename, const UnorderedSymbols& allNeededImports) override;
+	};
+}
diff --git a/Engine/Source/Programs/UnrealBuildAccelerator/ObjTool/Private/UbaObjTool.cpp b/Engine/Source/Programs/UnrealBuildAccelerator/ObjTool/Private/UbaObjTool.cpp
--- a/Engine/Source/Programs/UnrealBuildAccelerator/ObjTool/Private/UbaObjTool.cpp
+++ b/Engine/Source/Programs/UnrealBuildAccelerator/ObjTool/Private/UbaObjTool.cpp
@@ -127,7 +127,7 @@ namespace uba
 		if (!objFilesToStrip.empty())
 		{
 			CriticalSection cs;
-			bool success = true;
+			Atomic<bool> success = true;
 			UnorderedSymbols allNeededImports; // Imports needed from the outside of the stripped obj files
 
 			u32 workerCount = DefaultProcessorCount;
@@ -135,98 +135,83 @@ namespace uba
 			workManager.ParallelFor(workerCount, objFilesDependencies, [&](auto& it)
 				{
 					const TString& objFile = *it;
-					ObjectFile objectFile;
-					bool res = objectFile.Parse(logger, objFile.c_str());
+					ObjectFile* objectFile = ObjectFile::CreateAndParse(logger, objFile.c_str());
+					if (!objectFile)
+					{
+						success = false;
+						return;
+					}
+					auto g = MakeGuard([&]() { delete objectFile; });
 
 					ScopedCriticalSection _(cs);
-					success &= res;
-					if (res)
-						allNeededImports.insert(objectFile.GetImports().begin(), objectFile.GetImports().end());
+					allNeededImports.insert(objectFile->GetImports().begin(), objectFile->GetImports().end());
 				});
-
 			if (!success)
 				return -1;
 
 			UnorderedSymbols allSharedExports; // Exports from all the obj files about to be stripped
 
-			struct ObjectFileRec { ObjectFile file; UnorderedSymbols loopbacksToAdd;  UnorderedSymbols toRemove; };
-			Map<TString, ObjectFileRec> objectFileRecs;
-			CriticalSection cs2;
+			Map<TString, ObjectFile*> objectFiles;
+			auto g = MakeGuard([&]() { for (auto& kv : objectFiles) delete kv.second; });
 
 			workManager.ParallelFor(workerCount, objFilesToStrip, [&](auto& it)
 				{
 					const TString& objFileName = *it;
-
-					cs2.Enter();
-					ObjectFile& objectFile = objectFileRecs.try_emplace(objFileName).first->second.file;
-					cs2.Leave();
-
-					bool res = objectFile.Parse(logger, objFileName.c_str());
+					ObjectFile* objectFile = ObjectFile::CreateAndParse(logger, objFileName.c_str());
+					if (!objectFile)
+					{
+						success = false;
+						return;
+					}
 
 					ScopedCriticalSection _(cs);
-					success &= res;
-					if (res)
-						allSharedExports.insert(objectFile.GetExports().begin(), objectFile.GetExports().end());
+					objectFiles.try_emplace(objFileName, objectFile);
+					allSharedExports.insert(objectFile->GetExports().begin(), objectFile->GetExports().end());
 				});
-
 			if (!success)
 				return -1;
 
-			UnorderedSymbols duplicates;
-
 			// Figure out which loopback symbols that should be added to which obj file
+			UnorderedSymbols duplicates;
 			for (auto& objFileName : objFilesToStrip)
-			{
-				ObjectFileRec& rec = objectFileRecs[objFileName];
-				for (auto& importSymbol : rec.file.GetImports())
-				{
-					if (strncmp(importSymbol.c_str(), "__imp_", 6) != 0)
-						continue;
-					std::string importString = importSymbol.substr(6);
-					auto findIt = allSharedExports.find(importString);
-					if (findIt == allSharedExports.end())
-						continue;
-					rec.loopbacksToAdd.emplace(importString);
-					allSharedExports.erase(findIt);
-				}
-
-				for (auto& dupSymbol : rec.file.GetPotentialDuplicates())
-					if (!duplicates.insert(dupSymbol).second)
-						rec.toRemove.insert(dupSymbol);
-			}
+				if (!objectFiles[objFileName]->ComputeLoopbacksAndDuplicates(allSharedExports, duplicates))
+					return -1;
 
-			workManager.ParallelFor(workerCount, objectFileRecs, [&](auto& it)
+			workManager.ParallelFor(workerCount, objectFiles, [&](auto& it)
 				{
-					ObjectFileRec& rec = it->second;
-					const tchar* fileName = rec.file.GetFileName();
+					ObjectFile& file = *it->second;
+					const tchar* fileName = file.GetFileName();
 					const tchar* lastDot = TStrrchr(fileName, '.');
 					UBA_ASSERT(lastDot);
 					StringBuffer<> newFilename;
 					newFilename.Append(fileName, lastDot - fileName).Append(TC(".strip")).Append(lastDot);
-					rec.file.CreateStripped(logger, newFilename.data, allNeededImports, rec.loopbacksToAdd, rec.toRemove);
+					if (!file.CreateStripped(logger, newFilename.data, allNeededImports))
+						success = false;
 				});
-
-			//logger.Info(TC("Stripped %llu symbols from %llu obj files"), strippedSymbolCount.load(), objFilesToStrip.size());
 			if (!success)
 				return -1;
+
+			//logger.Info(TC("Stripped %llu symbols from %llu obj files"), strippedSymbolCount.load(), objFilesToStrip.size());
 		}
 		else
 		{
 			if (objFile.empty())
 				return PrintHelp(TC("No obj file provided"));
 
-			ObjectFile objectFile;
-			if (!objectFile.Parse(logger, objFile.c_str()))
+			ObjectFile* objectFile = ObjectFile::CreateAndParse(logger, objFile.c_str());
+			if (!objectFile)
 				return -1;
 
 			if (printSymbols)
 			{
-				for (auto& symbol : objectFile.GetImports())
+				for (auto& symbol : objectFile->GetImports())
 					logger.Info(TC("I %S"), symbol.c_str());
 
-				for (auto& symbol : objectFile.GetExports())
+				for (auto& symbol : objectFile->GetExports())
 					logger.Info(TC("E %S"), symbol.c_str());
 			}
+
+			delete objectFile;
 		}
 		return 0;
 	}
