[UBA]
* Fixed bug in storage proxy where files at very specific sizes could end up missing data so cas on the client side could have zeroed out bytes at the end of the cas file

diff --git a/Engine/Source/Programs/UnrealBuildAccelerator/Common/Private/UbaStorage.cpp b/Engine/Source/Programs/UnrealBuildAccelerator/Common/Private/UbaStorage.cpp
--- a/Engine/Source/Programs/UnrealBuildAccelerator/Common/Private/UbaStorage.cpp
+++ b/Engine/Source/Programs/UnrealBuildAccelerator/Common/Private/UbaStorage.cpp
@@ -664,7 +664,7 @@ namespace uba
 						{
 							bool f = false;
 							if (rec->error.compare_exchange_strong(f, true))
-								rec->logger->Error(TC("Expecting to be able to decompress to %u bytes but got %llu (%s)"), decompressedBlockSize, decompLen, rec->hint);
+								rec->logger->Error(TC("Expecting to be able to decompress %u bytes to %u bytes but got %llu (%s)"), compressedBlockSize, decompressedBlockSize, decompLen, rec->hint);
 							if (!--rec->refCount)
 								delete rec;
 							rec->done.Set();
diff --git a/Engine/Source/Programs/UnrealBuildAccelerator/Common/Private/UbaStorageProxy.cpp b/Engine/Source/Programs/UnrealBuildAccelerator/Common/Private/UbaStorageProxy.cpp
--- a/Engine/Source/Programs/UnrealBuildAccelerator/Common/Private/UbaStorageProxy.cpp
+++ b/Engine/Source/Programs/UnrealBuildAccelerator/Common/Private/UbaStorageProxy.cpp
@@ -234,8 +234,9 @@ namespace uba
 
 					if (file.received < file.size)
 					{
+						u64 left = file.size - file.received;
 						u64 segmentSize = m_client.GetMessageMaxSize() - 5; // This is server response size - header.. TODO: Should be taken from server
-						u32 segmentCount = u32(file.size / segmentSize);
+						u32 segmentCount = u32((left + segmentSize - 1) / segmentSize);
 						file.messagesInFlight.resize(segmentCount);
 						for (u32 i=0; i!=segmentCount; ++i)
 						{
