[UBA]
* UbaCli - fixed so checkcas2 use storage proxy for transfer of cas

diff --git a/Engine/Source/Programs/UnrealBuildAccelerator/Cli/Private/UbaCli.cpp b/Engine/Source/Programs/UnrealBuildAccelerator/Cli/Private/UbaCli.cpp
--- a/Engine/Source/Programs/UnrealBuildAccelerator/Cli/Private/UbaCli.cpp
+++ b/Engine/Source/Programs/UnrealBuildAccelerator/Cli/Private/UbaCli.cpp
@@ -424,17 +424,18 @@ namespace uba
 			storageInfo.casCapacityBytes = 0;
 			storageInfo.storeCompressed = storeCompressed;
 			StorageServer storageServer(storageInfo);
-			NetworkClientCreateInfo ncci;
-			NetworkClient client(ctorSuccess, ncci);
+
 			StringBuffer<> rootDir2(g_rootDir.data);
 			rootDir2.Append("_CHECKCAS2");
-			StorageClientCreateInfo scci(client, rootDir2.data);
-			StorageClient storageClient(scci);
-			storageClient.Start();
+			Client client;
+
 			auto g = MakeGuard([&]() { server.DisconnectClients(); });
 			if (!server.StartListen(networkBackend, 1347, TC("127.0.0.1")))
 				return false;
-			if (!client.Connect(networkBackend, TC("127.0.0.1"), 1347))
+			ClientInitInfo cii { logWriter, networkBackend, rootDir2.data, TC("127.0.0.1"), 1347, TC("foo") };
+			cii.createSession = false;
+			cii.addDirSuffix = false;
+			if (!client.Init(cii))
 				return false;
 			bool success = true;
 			WorkManagerImpl workManager(DefaultProcessorCount);
@@ -444,8 +445,8 @@ namespace uba
 						{
 							Storage::RetrieveResult res;
 							storageServer.EnsureCasFile(casKey, TC("Dummy"));
-							CasKey casKey2 = AsCompressed(casKey, false);
-							if (!storageClient.RetrieveCasFile(res, casKey2, TC("")))
+							CasKey casKey2 = AsCompressed(casKey, true);
+							if (!client.storageClient->RetrieveCasFile(res, casKey2, TC("")))
 								success = false;
 						}, 1, TC(""));
 				});
diff --git a/Engine/Source/Programs/UnrealBuildAccelerator/Cli/Private/UbaClient.h b/Engine/Source/Programs/UnrealBuildAccelerator/Cli/Private/UbaClient.h
--- a/Engine/Source/Programs/UnrealBuildAccelerator/Cli/Private/UbaClient.h
+++ b/Engine/Source/Programs/UnrealBuildAccelerator/Cli/Private/UbaClient.h
@@ -18,6 +18,8 @@ namespace uba
 		const tchar* zone = nullptr;
 		u32 maxProcessorCount = 0;
 		u32 index = 0;
+		bool createSession = true;
+		bool addDirSuffix = true;
 	};
 
 	class Client
@@ -33,7 +35,10 @@ namespace uba
 				return false;
 
 			StringBuffer<> clientRootDir;
-			clientRootDir.Append(info.rootDir).Append("Agent").AppendValue(info.index);
+			clientRootDir.Append(info.rootDir);
+			if (info.addDirSuffix)
+				clientRootDir.Append("Agent").AppendValue(info.index);
+
 			StorageClientCreateInfo storageClientInfo(*networkClient, clientRootDir.data);
 			storageClientInfo.zone = info.zone;
 			storageClientInfo.getProxyBackendCallback = [](void* ud, const tchar* h) -> NetworkBackend& { return ((Client*)ud)->GetProxyBackend(h); };
@@ -43,15 +48,19 @@ namespace uba
 			storageClient = new StorageClient(storageClientInfo);
 
 			storageClient->LoadCasTable(false);
-			SessionClientCreateInfo sessionClientInfo(*storageClient, *networkClient, info.logWriter);
-			sessionClientInfo.maxProcessCount = info.maxProcessorCount;
-			sessionClientInfo.rootDir = clientRootDir.data;
-			sessionClientInfo.deleteSessionsOlderThanSeconds = 1;
-
-			sessionClient = new SessionClient(sessionClientInfo);
 
 			storageClient->Start();
-			sessionClient->Start();
+
+			if (info.createSession)
+			{
+				SessionClientCreateInfo sessionClientInfo(*storageClient, *networkClient, info.logWriter);
+				sessionClientInfo.maxProcessCount = info.maxProcessorCount;
+				sessionClientInfo.rootDir = clientRootDir.data;
+				sessionClientInfo.deleteSessionsOlderThanSeconds = 1;
+
+				sessionClient = new SessionClient(sessionClientInfo);
+				sessionClient->Start();
+			}
 			
 			return networkClient->Connect(*networkBackend, info.host, info.port);
 		}
