[UBA]
* Added so UbaObjTool handle duplicated symbols for elf object files

diff --git a/Engine/Source/Programs/UnrealBuildAccelerator/Common/Private/UbaObjectFileCoff.cpp b/Engine/Source/Programs/UnrealBuildAccelerator/Common/Private/UbaObjectFileCoff.cpp
--- a/Engine/Source/Programs/UnrealBuildAccelerator/Common/Private/UbaObjectFileCoff.cpp
+++ b/Engine/Source/Programs/UnrealBuildAccelerator/Common/Private/UbaObjectFileCoff.cpp
@@ -21,25 +21,6 @@ namespace uba
 	#pragma pack(push)
 	#pragma pack(1)
 
-	struct Elf64Header
-	{
-	  unsigned char	e_ident[16];	/* Magic number and other info */
-	  u16 e_type;			/* Object file type */
-	  u16 e_machine;		/* Architecture */
-	  u32 e_version;		/* Object file version */
-	  u64 e_entry;		/* Entry point virtual address */
-	  u64 e_phoff;		/* Program header table file offset */
-	  u64 e_shoff;		/* Section header table file offset */
-	  u32 e_flags;		/* Processor-specific flags */
-	  u16 e_ehsize;		/* ELF header size in bytes */
-	  u16 e_phentsize;		/* Program header table entry size */
-	  u16 e_phnum;		/* Program header table entry count */
-	  u16 e_shentsize;		/* Section header table entry size */
-	  u16 e_shnum;		/* Section header table entry count */
-	  u16 e_shstrndx;		/* Section header string table index */
-	};
-
-
 	struct ImageFileHeader // IMAGE_FILE_HEADER
 	{
 		u16 Machine;
diff --git a/Engine/Source/Programs/UnrealBuildAccelerator/Common/Private/UbaObjectFileElf.cpp b/Engine/Source/Programs/UnrealBuildAccelerator/Common/Private/UbaObjectFileElf.cpp
--- a/Engine/Source/Programs/UnrealBuildAccelerator/Common/Private/UbaObjectFileElf.cpp
+++ b/Engine/Source/Programs/UnrealBuildAccelerator/Common/Private/UbaObjectFileElf.cpp
@@ -5,14 +5,178 @@
 
 namespace uba
 {
+	struct Elf64Header
+	{
+	  unsigned char	e_ident[16];	/* Magic number and other info */
+	  u16	e_type;			/* Object file type */
+	  u16	e_machine;		/* Architecture */
+	  u32	e_version;		/* Object file version */
+	  u64	e_entry;		/* Entry point virtual address */
+	  u64	e_phoff;		/* Program header table file offset */
+	  u64	e_shoff;		/* Section header table file offset */
+	  u32	e_flags;		/* Processor-specific flags */
+	  u16	e_ehsize;		/* ELF header size in bytes */
+	  u16	e_phentsize;	/* Program header table entry size */
+	  u16	e_phnum;		/* Program header table entry count */
+	  u16	e_shentsize;	/* Section header table entry size */
+	  u16	e_shnum;		/* Section header table entry count */
+	  u16	e_shstrndx;		/* Section header string table index */
+	};
+
+	struct Elf64SectionHeader
+	{
+	  u32	sh_name;		/* Section name (string tbl index) */
+	  u32	sh_type;		/* Section type */
+	  u64	sh_flags;		/* Section flags */
+	  u64	sh_addr;		/* Section virtual addr at execution */
+	  u64	sh_offset;		/* Section file offset */
+	  u64	sh_size;		/* Section size in bytes */
+	  u32	sh_link;		/* Link to another section */
+	  u32	sh_info;		/* Additional section information */
+	  u64	sh_addralign;	/* Section alignment */
+	  u64	sh_entsize;		/* Entry size if section holds table */
+	};
+
+	struct Elf64Sym
+	{
+	  u32	st_name;		/* Symbol name (string tbl index) */
+	  u8	st_info;		/* Symbol type and binding */
+	  u8	st_other;		/* Symbol visibility */
+	  u16	st_shndx;		/* Section index */
+	  u64	st_value;		/* Symbol value */
+	  u64	st_size;		/* Symbol size */
+	};
+
+
+	#define SHT_SYMTAB	  2		/* Symbol table */
+	#define SHT_DYNSYM	  11		/* Dynamic linker symbol table */
+
+	#define STT_OBJECT	1		/* Symbol is a data object */
+	#define STT_FUNC	2		/* Symbol is a code object */
+	
+	#define STB_WEAK	2		/* Weak symbol */
+
+
+	#define ELF32_ST_BIND(val)		(((unsigned char) (val)) >> 4)
+	#define ELF32_ST_TYPE(val)		((val) & 0xf)
+	#define ELF32_ST_INFO(bind, type)	(((bind) << 4) + ((type) & 0xf))
+
+	/* Both Elf32_Sym and Elf64_Sym use the same one-byte st_info field.  */
+	#define ELF64_ST_BIND(val)		ELF32_ST_BIND (val)
+	#define ELF64_ST_TYPE(val)		ELF32_ST_TYPE (val)
+	#define ELF64_ST_INFO(bind, type)	ELF32_ST_INFO ((bind), (type))
+
+	// These are symbols that are added to all dlls through some macros.
+	// When merging dlls we need to remove duplicates of these
+	static UnorderedSymbols PotentiallyDuplicatedSymbols = 
+	{
+		// REPLACEMENT_OPERATOR_NEW_AND_DELETE
+		"_Znwm",
+		"_Znam",
+		"_ZnwmRKSt9nothrow_t",
+		"_ZnamRKSt9nothrow_t",
+		"_ZnwmSt11align_val_t",
+		"_ZnamSt11align_val_t",
+		"_ZnwmSt11align_val_tRKSt9nothrow_t",
+		"_ZnamSt11align_val_tRKSt9nothrow_t",
+		"_ZdlPv",
+		"_ZdaPv",
+		"_ZdlPvRKSt9nothrow_t",
+		"_ZdaPvRKSt9nothrow_t",
+		"_ZdlPvm",
+		"_ZdaPvm",
+		"_ZdlPvmRKSt9nothrow_t",
+		"_ZdaPvmRKSt9nothrow_t",
+		"_ZdlPvSt11align_val_t",
+		"_ZdaPvSt11align_val_t",
+		"_ZdlPvSt11align_val_tRKSt9nothrow_t",
+		"_ZdaPvSt11align_val_tRKSt9nothrow_t",
+		"_ZdlPvmSt11align_val_t",
+		"_ZdaPvmSt11align_val_t",
+		"_ZdlPvmSt11align_val_tRKSt9nothrow_t",
+		"_ZdaPvmSt11align_val_tRKSt9nothrow_t",
+
+		// UE_DEFINE_FMEMORY_WRAPPERS
+		"_Z14FMemory_Mallocmm",
+		"_Z15FMemory_ReallocPvmm",
+		"_Z12FMemory_FreePv",
+
+		// UE4_VISUALIZERS_HELPERS
+		"GNameBlocksDebug",
+		"GObjectIndexToPackedObjectRefDebug",
+		"GObjectArrayForDebugVisualizers",
+		"GComplexObjectPathDebug",
+		"GObjectHandlePackageDebug",
+	};
+
 
 	bool ObjectFileElf::Parse(Logger& logger, const tchar* filename)
 	{
+		auto& header = *(Elf64Header*)m_data;
+		if (header.e_ident[4] != 2) // 64-bit
+			return false;
+		if (header.e_ident[5] != 1) // Little endian
+			return false;
+		if (header.e_ident[6] != 1) // Version
+			return false;
+		if (header.e_type != 1) // Relocatable file
+			return false;
+
+		UBA_ASSERT(sizeof(Elf64SectionHeader) == header.e_shentsize);
+
+		auto sections = (Elf64SectionHeader*)(m_data + header.e_shoff);
+
+		auto& namesSection = sections[header.e_shstrndx];
+		u8* stringTableData = m_data + namesSection.sh_offset;
+
+		#if 0
+		u64 stringTableSize = namesSection.sh_size;
+		u32 counter = 0;
+		char* it = (char*)stringTableData;
+		char* end = it + stringTableSize;
+		while (it != end)
+		{
+			printf("%s\n", it);
+			u64 len = strlen(it);
+			it += len + 1;
+			counter++;
+		}
+		#endif
+
+		for (u16 i=0; i!=header.e_shnum; ++i)
+		{
+			auto& section = sections[i];
+			//auto sectionName = (char*)stringTableData + section.sh_name;
+
+			if (section.sh_type == SHT_SYMTAB || section.sh_type == SHT_DYNSYM)
+			{
+				UBA_ASSERT(section.sh_link == header.e_shstrndx);
+				UBA_ASSERT(section.sh_entsize == sizeof(Elf64Sym));
+				auto symbols = (Elf64Sym*)(m_data + section.sh_offset);
+				u64 symbolCount = section.sh_size / sizeof(Elf64Sym);
+				for (u64 j=0; j!=symbolCount; ++j)
+				{
+					auto& symbol = symbols[j];
+					u8 symbolType = ELF64_ST_TYPE(symbol.st_info);
+					if (symbolType != STT_FUNC && symbolType != STT_OBJECT)
+						continue;
+
+					auto name = (char*)stringTableData + symbol.st_name;
+					if (PotentiallyDuplicatedSymbols.find(name) == PotentiallyDuplicatedSymbols.end())
+						continue;
+
+					m_potentialDuplicates.emplace(name);
+				}
+			}
+		}
 		return true;
 	}
 
 	bool ObjectFileElf::ComputeLoopbacksAndDuplicates(UnorderedSymbols& allSharedExports, UnorderedSymbols& duplicates)
 	{
+		for (auto& dupSymbol : m_potentialDuplicates)
+			if (!duplicates.insert(dupSymbol).second)
+				m_toRemove.insert(dupSymbol);
 		return true;
 	}
 
@@ -25,6 +189,37 @@ namespace uba
 		u8* newData = newFile.GetData();
 		memcpy(newData, m_data, m_dataSize);
 
+		auto& header = *(Elf64Header*)newData;
+		auto sections = (Elf64SectionHeader*)(newData + header.e_shoff);
+
+		auto& namesSection = sections[header.e_shstrndx];
+		u8* stringTableData = newData + namesSection.sh_offset;
+
+		for (u16 i=0; i!=header.e_shnum; ++i)
+		{
+			auto& section = sections[i];
+
+			if (section.sh_type == SHT_SYMTAB || section.sh_type == SHT_DYNSYM)
+			{
+				UBA_ASSERT(section.sh_link == header.e_shstrndx);
+				UBA_ASSERT(section.sh_entsize == sizeof(Elf64Sym));
+				auto symbols = (Elf64Sym*)(newData + section.sh_offset);
+				u64 symbolCount = section.sh_size / sizeof(Elf64Sym);
+				for (u64 j=0; j!=symbolCount; ++j)
+				{
+					auto& symbol = symbols[j];
+					u8 symbolType = ELF64_ST_TYPE(symbol.st_info);
+					if (symbolType != STT_FUNC && symbolType != STT_OBJECT)
+						continue;
+					auto name = (char*)stringTableData + symbol.st_name;
+
+					if (m_toRemove.find(name) == m_toRemove.end())
+						continue;
+					symbol.st_info = ELF64_ST_INFO(STB_WEAK, symbolType);
+				}
+			}
+		}
+
 		return newFile.Close();
 	}
 }
diff --git a/Engine/Source/Programs/UnrealBuildAccelerator/Common/Public/UbaObjectFileElf.h b/Engine/Source/Programs/UnrealBuildAccelerator/Common/Public/UbaObjectFileElf.h
--- a/Engine/Source/Programs/UnrealBuildAccelerator/Common/Public/UbaObjectFileElf.h
+++ b/Engine/Source/Programs/UnrealBuildAccelerator/Common/Public/UbaObjectFileElf.h
@@ -12,5 +12,7 @@ namespace uba
 		virtual bool Parse(Logger& logger, const tchar* filename) override;
 		virtual bool ComputeLoopbacksAndDuplicates(UnorderedSymbols& allSharedExports, UnorderedSymbols& duplicates) override;
 		virtual bool CreateStripped(Logger& logger, const tchar* newFilename, const UnorderedSymbols& allNeededImports) override;
+
+		UnorderedSymbols m_toRemove;
 	};
 }
