[UBA]
* UbaObjTool - Fixed unity compile error

diff --git a/Engine/Source/Programs/UnrealBuildAccelerator/Common/Private/UbaObjectFileCoff.cpp b/Engine/Source/Programs/UnrealBuildAccelerator/Common/Private/UbaObjectFileCoff.cpp
--- a/Engine/Source/Programs/UnrealBuildAccelerator/Common/Private/UbaObjectFileCoff.cpp
+++ b/Engine/Source/Programs/UnrealBuildAccelerator/Common/Private/UbaObjectFileCoff.cpp
@@ -124,7 +124,7 @@ namespace uba
 
 	// These are symbols that are added to all dlls through some macros.
 	// When merging dlls we need to remove duplicates of these
-	static UnorderedSymbols PotentiallyDuplicatedSymbols = 
+	UnorderedSymbols ObjectFileCoff::PotentiallyDuplicatedSymbols = 
 	{
 		// REPLACEMENT_OPERATOR_NEW_AND_DELETE
 		"??2@YAPEAX_K@Z",
diff --git a/Engine/Source/Programs/UnrealBuildAccelerator/Common/Private/UbaObjectFileElf.cpp b/Engine/Source/Programs/UnrealBuildAccelerator/Common/Private/UbaObjectFileElf.cpp
--- a/Engine/Source/Programs/UnrealBuildAccelerator/Common/Private/UbaObjectFileElf.cpp
+++ b/Engine/Source/Programs/UnrealBuildAccelerator/Common/Private/UbaObjectFileElf.cpp
@@ -68,7 +68,7 @@ namespace uba
 
 	// These are symbols that are added to all dlls through some macros.
 	// When merging dlls we need to remove duplicates of these
-	static UnorderedSymbols PotentiallyDuplicatedSymbols = 
+	UnorderedSymbols ObjectFileElf::PotentiallyDuplicatedSymbols = 
 	{
 		// REPLACEMENT_OPERATOR_NEW_AND_DELETE
 		"_Znwm",
diff --git a/Engine/Source/Programs/UnrealBuildAccelerator/Common/Public/UbaObjectFileCoff.h b/Engine/Source/Programs/UnrealBuildAccelerator/Common/Public/UbaObjectFileCoff.h
--- a/Engine/Source/Programs/UnrealBuildAccelerator/Common/Public/UbaObjectFileCoff.h
+++ b/Engine/Source/Programs/UnrealBuildAccelerator/Common/Public/UbaObjectFileCoff.h
@@ -39,5 +39,7 @@ namespace uba
 
 		UnorderedSymbols m_loopbacksToAdd;
 		UnorderedSymbols m_toRemove;
+
+		static UnorderedSymbols PotentiallyDuplicatedSymbols;
 	};
 }
diff --git a/Engine/Source/Programs/UnrealBuildAccelerator/Common/Public/UbaObjectFileElf.h b/Engine/Source/Programs/UnrealBuildAccelerator/Common/Public/UbaObjectFileElf.h
--- a/Engine/Source/Programs/UnrealBuildAccelerator/Common/Public/UbaObjectFileElf.h
+++ b/Engine/Source/Programs/UnrealBuildAccelerator/Common/Public/UbaObjectFileElf.h
@@ -14,5 +14,7 @@ namespace uba
 		virtual bool CreateStripped(Logger& logger, const tchar* newFilename, const UnorderedSymbols& allNeededImports) override;
 
 		UnorderedSymbols m_toRemove;
+
+		static UnorderedSymbols PotentiallyDuplicatedSymbols;
 	};
 }
