[UBA]
* UbaObjTool - Added code to be able to see how many symbols that are kept after stripping

diff --git a/Engine/Source/Programs/UnrealBuildAccelerator/Common/Private/UbaObjectFileCoff.cpp b/Engine/Source/Programs/UnrealBuildAccelerator/Common/Private/UbaObjectFileCoff.cpp
--- a/Engine/Source/Programs/UnrealBuildAccelerator/Common/Private/UbaObjectFileCoff.cpp
+++ b/Engine/Source/Programs/UnrealBuildAccelerator/Common/Private/UbaObjectFileCoff.cpp
@@ -126,6 +126,7 @@ namespace uba
 	// When merging dlls we need to remove duplicates of these
 	UnorderedSymbols ObjectFileCoff::PotentiallyDuplicatedSymbols = 
 	{
+		/*
 		// REPLACEMENT_OPERATOR_NEW_AND_DELETE
 		"??2@YAPEAX_K@Z",
 		"??2@YAPEAX_KAEBUnothrow_t@std@@@Z",
@@ -167,6 +168,7 @@ namespace uba
 		// Have no idea why these are duplicated
 		"??_GIFunction_OwnedObject@Function@Private@Core@UE@@UEAAPEAXI@Z",
 		"??_GIDelegateInstance@@UEAAPEAXI@Z",
+		*/
 	};
 
 	template<typename SymbolType>
@@ -232,7 +234,7 @@ namespace uba
 		return true;
 	}
 
-	bool ObjectFileCoff::CreateStripped(Logger& logger, const tchar* newFilename, const UnorderedSymbols& allNeededImports)
+	bool ObjectFileCoff::CreateStripped(Logger& logger, const tchar* newFilename, const UnorderedSymbols& allNeededImports, u32& outKeptExportCount)
 	{
 		UBA_ASSERT(m_file);
 
@@ -276,7 +278,7 @@ namespace uba
 		else
 			RemoveSymbols<ImageSymbol>(logger, newData, newInfo);
 
-		WriteExports(logger, newData, allNeededImports);
+		WriteExports(logger, newData, allNeededImports, outKeptExportCount);
 
 		return newFile.Close();
 	}
@@ -370,7 +372,7 @@ namespace uba
 		}
 	}
 
-	void ObjectFileCoff::WriteExports(Logger& logger, u8* newData, const UnorderedSymbols& allNeededImports)
+	void ObjectFileCoff::WriteExports(Logger& logger, u8* newData, const UnorderedSymbols& allNeededImports, u32& outKeptExportCount)
 	{
 		if (!m_info.directiveSectionMemOffset)
 			return;
@@ -391,6 +393,7 @@ namespace uba
 		auto readEnd = readPos + directiveSection->SizeOfRawData;
 		auto readLastPossiblePos = readEnd - 9;
 
+		u32 exportCount = 0;
 		while (true)
 		{
 			const char* exportStr = nullptr;
@@ -414,6 +417,8 @@ namespace uba
 				break;
 			}
 
+			++exportCount;
+
 			const char* startPos = exportStr;
 			exportStr += 8;
 			const char* exportEnd = strchr(exportStr, ' ');
@@ -436,6 +441,8 @@ namespace uba
 			if (strstr(exportStr, "Initialize") == exportStr && strstr(exportStr, "Module"))
 				continue;
 
+			--exportCount;
+
 			u64 toCopy = startPos - lastCopyPos - 1;
 			memcpy(writePos, lastCopyPos, toCopy);
 			writePos += toCopy;
@@ -453,6 +460,8 @@ namespace uba
 		UBA_ASSERT(newDirectiveSection->SizeOfRawData <= directiveSection->SizeOfRawData);
 
 		memset(writePos, 0, directiveSection->SizeOfRawData - newDirectiveSection->SizeOfRawData);
+
+		outKeptExportCount = exportCount;
 	}
 
 	template<typename SymbolType> void ObjectFileCoff::CalculateImports(Logger& logger, Vector<u32>& outImports)
diff --git a/Engine/Source/Programs/UnrealBuildAccelerator/Common/Private/UbaObjectFileElf.cpp b/Engine/Source/Programs/UnrealBuildAccelerator/Common/Private/UbaObjectFileElf.cpp
--- a/Engine/Source/Programs/UnrealBuildAccelerator/Common/Private/UbaObjectFileElf.cpp
+++ b/Engine/Source/Programs/UnrealBuildAccelerator/Common/Private/UbaObjectFileElf.cpp
@@ -180,7 +180,7 @@ namespace uba
 		return true;
 	}
 
-	bool ObjectFileElf::CreateStripped(Logger& logger, const tchar* newFilename, const UnorderedSymbols& allNeededImports)
+	bool ObjectFileElf::CreateStripped(Logger& logger, const tchar* newFilename, const UnorderedSymbols& allNeededImports, u32& outKeptExportCount)
 	{
 		FileAccessor newFile(logger, newFilename);
 		if (!newFile.CreateMemoryWrite(false, DefaultAttributes(), m_dataSize))
diff --git a/Engine/Source/Programs/UnrealBuildAccelerator/Common/Public/UbaObjectFile.h b/Engine/Source/Programs/UnrealBuildAccelerator/Common/Public/UbaObjectFile.h
--- a/Engine/Source/Programs/UnrealBuildAccelerator/Common/Public/UbaObjectFile.h
+++ b/Engine/Source/Programs/UnrealBuildAccelerator/Common/Public/UbaObjectFile.h
@@ -17,7 +17,7 @@ namespace uba
 		static ObjectFile* CreateAndParse(Logger& logger, const tchar* filename);
 
 		virtual bool ComputeLoopbacksAndDuplicates(UnorderedSymbols& allSharedExports, UnorderedSymbols& duplicates) = 0;
-		virtual bool CreateStripped(Logger& logger, const tchar* newFilename, const UnorderedSymbols& allNeededImports) = 0;
+		virtual bool CreateStripped(Logger& logger, const tchar* newFilename, const UnorderedSymbols& allNeededImports, u32& outKeptExportCount) = 0;
 
 		const tchar* GetFileName() const;
 		const UnorderedSymbols& GetImports() const;
diff --git a/Engine/Source/Programs/UnrealBuildAccelerator/Common/Public/UbaObjectFileCoff.h b/Engine/Source/Programs/UnrealBuildAccelerator/Common/Public/UbaObjectFileCoff.h
--- a/Engine/Source/Programs/UnrealBuildAccelerator/Common/Public/UbaObjectFileCoff.h
+++ b/Engine/Source/Programs/UnrealBuildAccelerator/Common/Public/UbaObjectFileCoff.h
@@ -11,7 +11,7 @@ namespace uba
 	public:
 		virtual bool Parse(Logger& logger, const tchar* filename) override;
 		virtual bool ComputeLoopbacksAndDuplicates(UnorderedSymbols& allSharedExports, UnorderedSymbols& duplicates) override;
-		virtual bool CreateStripped(Logger& logger, const tchar* newFilename, const UnorderedSymbols& allNeededImports) override;
+		virtual bool CreateStripped(Logger& logger, const tchar* newFilename, const UnorderedSymbols& allNeededImports, u32& outKeptExportCount) override;
 
 	private:
 		struct Info;
@@ -19,7 +19,7 @@ namespace uba
 		bool ParseExports();
 		template<typename SymbolType> void ParseImports();
 
-		void WriteExports(Logger& logger, u8* newData, const UnorderedSymbols& allNeededImports);
+		void WriteExports(Logger& logger, u8* newData, const UnorderedSymbols& allNeededImports, u32& outKeptExportCount);
 		template<typename SymbolType> void CalculateImports(Logger& logger, Vector<u32>& outImports);
 		template<typename SymbolType> void WriteImports(Logger& logger, u8* newData, Info& newInfo, const Vector<u32>& symbolsToAdd);
 		template<typename SymbolType> void RemoveSymbols(Logger& logger, u8* newData, Info& newInfo);
diff --git a/Engine/Source/Programs/UnrealBuildAccelerator/Common/Public/UbaObjectFileElf.h b/Engine/Source/Programs/UnrealBuildAccelerator/Common/Public/UbaObjectFileElf.h
--- a/Engine/Source/Programs/UnrealBuildAccelerator/Common/Public/UbaObjectFileElf.h
+++ b/Engine/Source/Programs/UnrealBuildAccelerator/Common/Public/UbaObjectFileElf.h
@@ -11,7 +11,7 @@ namespace uba
 	public:
 		virtual bool Parse(Logger& logger, const tchar* filename) override;
 		virtual bool ComputeLoopbacksAndDuplicates(UnorderedSymbols& allSharedExports, UnorderedSymbols& duplicates) override;
-		virtual bool CreateStripped(Logger& logger, const tchar* newFilename, const UnorderedSymbols& allNeededImports) override;
+		virtual bool CreateStripped(Logger& logger, const tchar* newFilename, const UnorderedSymbols& allNeededImports, u32& outKeptExportCount) override;
 
 		UnorderedSymbols m_toRemove;
 
diff --git a/Engine/Source/Programs/UnrealBuildAccelerator/ObjTool/Private/UbaObjTool.cpp b/Engine/Source/Programs/UnrealBuildAccelerator/ObjTool/Private/UbaObjTool.cpp
--- a/Engine/Source/Programs/UnrealBuildAccelerator/ObjTool/Private/UbaObjTool.cpp
+++ b/Engine/Source/Programs/UnrealBuildAccelerator/ObjTool/Private/UbaObjTool.cpp
@@ -177,6 +177,8 @@ namespace uba
 				if (!objectFiles[objFileName]->ComputeLoopbacksAndDuplicates(allSharedExports, duplicates))
 					return -1;
 
+			Atomic<u32> totalExportCount;
+			Atomic<u32> totalKeptExportCount;
 			workManager.ParallelFor(workerCount, objectFiles, [&](auto& it)
 				{
 					ObjectFile& file = *it->second;
@@ -185,13 +187,17 @@ namespace uba
 					UBA_ASSERT(lastDot);
 					StringBuffer<> newFilename;
 					newFilename.Append(fileName, lastDot - fileName).Append(TC(".strip")).Append(lastDot);
-					if (!file.CreateStripped(logger, newFilename.data, allNeededImports))
+					u32 keptExportCount = 0;
+					if (!file.CreateStripped(logger, newFilename.data, allNeededImports, keptExportCount))
 						success = false;
+
+					totalExportCount += u32(file.GetExports().size());
+					totalKeptExportCount += keptExportCount;
 				});
 			if (!success)
 				return -1;
 
-			//logger.Info(TC("Stripped %llu symbols from %llu obj files"), strippedSymbolCount.load(), objFilesToStrip.size());
+			//logger.Info(TC("Reduced export count from %llu to %llu"), totalExportCount.load(), totalKeptExportCount.size());
 		}
 		else
 		{
