[UBA]
* Added rsaenh.dll to known files list

diff --git a/Engine/Source/Programs/UnrealBuildAccelerator/Common/Private/UbaStorageServer.cpp b/Engine/Source/Programs/UnrealBuildAccelerator/Common/Private/UbaStorageServer.cpp
--- a/Engine/Source/Programs/UnrealBuildAccelerator/Common/Private/UbaStorageServer.cpp
+++ b/Engine/Source/Programs/UnrealBuildAccelerator/Common/Private/UbaStorageServer.cpp
@@ -430,7 +430,7 @@ namespace uba
 
 						if (checkedCasKey != casKey)
 						{
-							m_logger.Error(TC("FetchBegin failed for cas file %s (%s). Server has a source file"), CasKeyString(casKey).str, hint.data);
+							m_logger.Error(TC("FetchBegin failed for cas file %s (%s). File on disk has different cas %s"), CasKeyString(casKey).str, hint.data, CasKeyString(checkedCasKey).str);
 							writer.WriteU16(0);
 							return false;
 						}
diff --git a/Engine/Source/Programs/UnrealBuildAccelerator/Core/Public/UbaWinBinDependencyParser.h b/Engine/Source/Programs/UnrealBuildAccelerator/Core/Public/UbaWinBinDependencyParser.h
--- a/Engine/Source/Programs/UnrealBuildAccelerator/Core/Public/UbaWinBinDependencyParser.h
+++ b/Engine/Source/Programs/UnrealBuildAccelerator/Core/Public/UbaWinBinDependencyParser.h
@@ -53,6 +53,7 @@ namespace uba
 		L"psapi.dll",
 		L"rasadhlp.dll",
 		L"rpcrt4.dll",
+		L"rsaenh.dll",
 		L"rstrtmgr.dll",
 		L"sechost.dll",
 		L"setupapi.dll",
