[UBA]
* Added support for extracting import/export symbols from obj file and strip it from exports at the same time

diff --git a/Engine/Source/Programs/UnrealBuildAccelerator/Common/Private/UbaFile.cpp b/Engine/Source/Programs/UnrealBuildAccelerator/Common/Private/UbaFile.cpp
--- a/Engine/Source/Programs/UnrealBuildAccelerator/Common/Private/UbaFile.cpp
+++ b/Engine/Source/Programs/UnrealBuildAccelerator/Common/Private/UbaFile.cpp
@@ -796,6 +796,20 @@ namespace uba
 #endif
 	}
 
+	u64 GetNowFileTime()
+	{
+		#if PLATFORM_WINDOWS
+		FILETIME ft;
+		SYSTEMTIME st;
+		GetSystemTime(&st);
+		SystemTimeToFileTime(&st, &ft);
+		return (u64&)ft;
+		#else
+		// TODO
+		return 0ull;
+		#endif
+	}
+
 	bool GetCurrentDirectoryW(StringBufferBase& out)
 	{
 #if PLATFORM_WINDOWS
diff --git a/Engine/Source/Programs/UnrealBuildAccelerator/Common/Private/UbaObjectFile.cpp b/Engine/Source/Programs/UnrealBuildAccelerator/Common/Private/UbaObjectFile.cpp
--- a/Engine/Source/Programs/UnrealBuildAccelerator/Common/Private/UbaObjectFile.cpp
+++ b/Engine/Source/Programs/UnrealBuildAccelerator/Common/Private/UbaObjectFile.cpp
@@ -1,13 +1,14 @@
 // Copyright Epic Games, Inc. All Rights Reserved.
 
 #include "UbaObjectFile.h"
+#include "UbaBinaryReaderWriter.h"
 #include "UbaFileAccessor.h"
 #include "UbaObjectFileCoff.h"
 #include "UbaObjectFileElf.h"
 
 namespace uba
 {
-	ObjectFile* ObjectFile::CreateAndParse(Logger& logger, const tchar* filename)
+	ObjectFile* ObjectFile::OpenAndParse(Logger& logger, const tchar* filename)
 	{
 		auto file = new FileAccessor(logger, filename);
 		auto fileGuard = MakeGuard([&]() { delete file; });
@@ -15,9 +16,17 @@ namespace uba
 		if (!file->OpenMemoryRead())
 			return nullptr;
 
-		u8* data = file->GetData();
-		u64 dataSize = file->GetSize();
+		ObjectFile* objectFile = Parse(logger, file->GetData(), file->GetSize(), filename);
+		if (!objectFile)
+			return nullptr;
 
+		fileGuard.Cancel();
+		objectFile->m_file = file;
+		return objectFile;
+	}
+
+	ObjectFile* ObjectFile::Parse(Logger& logger, u8* data, u64 dataSize, const tchar* hint)
+	{
 		if (dataSize < 4)
 			return nullptr;
 
@@ -29,21 +38,94 @@ namespace uba
 		else
 			objectFile = new ObjectFileCoff();
 
-		fileGuard.Cancel();
-		objectFile->m_file = file;
 		objectFile->m_data = data;
 		objectFile->m_dataSize = dataSize;
 
-		if (objectFile->Parse(logger, filename))
+		if (objectFile->Parse(logger, hint))
 			return objectFile;
 
 		delete objectFile;
 		return nullptr;
 	}
 
+	bool ObjectFile::CopyMemoryAndClose()
+	{
+		u8* data = (u8*)malloc(m_dataSize);
+		memcpy(data, m_data, m_dataSize);
+		m_data = data;
+		m_ownsData = true;
+		delete m_file;
+		m_file = nullptr;
+		return true;
+	}
+
+	bool ObjectFile::StripExports(Logger& logger)
+	{
+		u32 keptExportCount = 0;
+		return StripExports(logger, m_data, {}, keptExportCount);
+		return true;
+	}
+
+	bool ObjectFile::WriteSymbols(Logger& logger, MemoryBlock& memoryBlock)
+	{
+		auto write = [&](const void* data, u64 dataSize) { memcpy(memoryBlock.Allocate(dataSize, 1, TC("")), data, dataSize); };
+
+		// Write all imports
+		for (auto& symbol : m_imports)
+		{
+			write(symbol.c_str(), symbol.size());
+			write("", 1);
+		}
+		write("", 1);
+
+		// Write all exports
+		for (auto& kv : m_exports)
+		{
+			write(kv.first.c_str(), kv.first.size());
+			write(kv.second.c_str(), kv.second.size());
+			write("", 1);
+		}
+		write("", 1);
+		return true;
+	}
+
+	bool ObjectFile::WriteSymbols(Logger& logger, const tchar* exportsFilename)
+	{
+		FileAccessor exportsFile(logger, exportsFilename);
+		if (!exportsFile.CreateWrite())
+			return false;
+
+		char buffer[256*1024];
+		u64 bufferPos = 0;
+		auto flush = [&]() { exportsFile.Write(buffer, bufferPos); bufferPos = 0; };
+		auto write = [&](const void* data, u64 dataSize) { if (bufferPos + dataSize > sizeof(buffer)) flush(); memcpy(buffer + bufferPos, data, dataSize); bufferPos += dataSize; };
+
+		// Write all imports
+		for (auto& symbol : m_imports)
+		{
+			write(symbol.c_str(), symbol.size());
+			write("", 1);
+		}
+		write("", 1);
+
+		// Write all exports
+		for (auto& kv : m_exports)
+		{
+			write(kv.first.c_str(), kv.first.size());
+			write(kv.second.c_str(), kv.second.size());
+			write("", 1);
+		}
+		write("", 1);
+
+		flush();
+
+		return exportsFile.Close();
+	}
 
 	ObjectFile::~ObjectFile()
 	{
+		if (m_ownsData)
+			free(m_data);
 		delete m_file;
 	}
 
@@ -57,7 +139,7 @@ namespace uba
 		return m_imports;
 	}
 
-	const UnorderedSymbols& ObjectFile::GetExports() const
+	const UnorderedExports& ObjectFile::GetExports() const
 	{
 		return m_exports;
 	}
@@ -66,4 +148,22 @@ namespace uba
 	{
 		return m_potentialDuplicates;
 	}
+
+	bool ObjectFile::CreateExtraFile(Logger& logger, const tchar* extraObjFilename, const UnorderedSymbols& allNeededImports, const UnorderedSymbols& allSharedImports, const UnorderedExports& allSharedExports)
+	{
+		ObjectFileCoff objectFile;
+		MemoryBlock memoryBlock(16*1024*1024);
+
+		if (!((ObjectFile&)objectFile).CreateExtraFile(logger, memoryBlock, allNeededImports, allSharedImports, allSharedExports))
+			return false;
+
+		FileAccessor extraFile(logger, extraObjFilename);
+		if (!extraFile.CreateWrite())
+			return false;
+
+		if (!extraFile.Write(memoryBlock.memory, memoryBlock.writtenSize))
+			return false;
+
+		return extraFile.Close();
+	}
 }
diff --git a/Engine/Source/Programs/UnrealBuildAccelerator/Common/Private/UbaObjectFileCoff.cpp b/Engine/Source/Programs/UnrealBuildAccelerator/Common/Private/UbaObjectFileCoff.cpp
--- a/Engine/Source/Programs/UnrealBuildAccelerator/Common/Private/UbaObjectFileCoff.cpp
+++ b/Engine/Source/Programs/UnrealBuildAccelerator/Common/Private/UbaObjectFileCoff.cpp
@@ -11,9 +11,12 @@ namespace uba
 	constexpr u16 ImageSymUndefined = 0; // IMAGE_SYM_UNDEFINED
 
 	constexpr u32 ImageScnCntCode = 0x00000020; // IMAGE_SCN_CNT_CODE
+	constexpr u32 ImageScnLnkInfo = 0x00000200; // IMAGE_SCN_LNK_INFO
+	constexpr u32 ImageScnLnkRemove = 0x00000800; // IMAGE_SCN_LNK_REMOVE
+	constexpr u32 ImageScnLnkComdat = 0x00001000; // IMAGE_SCN_LNK_COMDAT
+	constexpr u32 ImageScnAlign1Bytes = 0x00100000; // IMAGE_SCN_ALIGN_1BYTES
 	constexpr u32 ImageScnMemExecute = 0x20000000; // IMAGE_SCN_MEM_EXECUTE
 	constexpr u32 ImageScnMemRead = 0x40000000; // IMAGE_SCN_MEM_READ
-	constexpr u32 ImageScnLnkComdat = 0x00001000; // IMAGE_SCN_LNK_COMDAT
 
 	const u16 ImageRelAmd64Addr64 = 0x0001; // IMAGE_REL_AMD64_ADDR64
 
@@ -126,7 +129,7 @@ namespace uba
 	// When merging dlls we need to remove duplicates of these
 	UnorderedSymbols ObjectFileCoff::PotentiallyDuplicatedSymbols = 
 	{
-		/*
+#if 0 // This is not needed anymore
 		// REPLACEMENT_OPERATOR_NEW_AND_DELETE
 		"??2@YAPEAX_K@Z",
 		"??2@YAPEAX_KAEBUnothrow_t@std@@@Z",
@@ -168,15 +171,15 @@ namespace uba
 		// Have no idea why these are duplicated
 		"??_GIFunction_OwnedObject@Function@Private@Core@UE@@UEAAPEAXI@Z",
 		"??_GIDelegateInstance@@UEAAPEAXI@Z",
-		*/
+#endif
 	};
 
 	template<typename SymbolType>
-	StringView GetSymbolName(SymbolType& symbol, u8* stringTableMem)
+	StringView GetSymbolName(SymbolType& symbol, const u8* data, u32 stringTableMemPos)
 	{
 		if (symbol.N.Name.Short == 0)
 		{
-			const char* name = (char*)stringTableMem + symbol.N.Name.Long;
+			auto name = (const char*)(data + stringTableMemPos + symbol.N.Name.Long);
 			return { name, name + strlen(name) };
 		}
 		auto shortName = (char*)symbol.N.ShortName;
@@ -209,7 +212,7 @@ namespace uba
 			auto& header = *(AnonObjectHeaderBigobj*)m_data;
 			m_info.symbolsMemPos = header.PointerToSymbolTable;
 			m_info.symbolCount = header.NumberOfSymbols;
-			m_info.stringTableMem = m_data + header.PointerToSymbolTable + header.NumberOfSymbols * sizeof(ImageSymbolEx);
+			m_info.stringTableMemPos = header.PointerToSymbolTable + header.NumberOfSymbols * sizeof(ImageSymbolEx);
 			m_info.sectionsMemOffset = sizeof(AnonObjectHeaderBigobj);
 			m_info.sectionCount = header.NumberOfSections;
 		}
@@ -218,7 +221,7 @@ namespace uba
 			auto& header = *(ImageFileHeader*)m_data;
 			m_info.symbolsMemPos = header.PointerToSymbolTable;
 			m_info.symbolCount = header.NumberOfSymbols;
-			m_info.stringTableMem = m_data + header.PointerToSymbolTable + header.NumberOfSymbols * sizeof(ImageSymbol);
+			m_info.stringTableMemPos = header.PointerToSymbolTable + header.NumberOfSymbols * sizeof(ImageSymbol);
 			m_info.sectionsMemOffset = sizeof(ImageFileHeader);
 			m_info.sectionCount = header.NumberOfSections;
 		}
@@ -236,7 +239,7 @@ namespace uba
 
 	bool ObjectFileCoff::CreateStripped(Logger& logger, const tchar* newFilename, const UnorderedSymbols& allNeededImports, u32& outKeptExportCount)
 	{
-		UBA_ASSERT(m_file);
+		UBA_ASSERT(m_data);
 
 		Vector<u32> symbolsToAdd;
 		if (m_isBigObj)
@@ -278,12 +281,12 @@ namespace uba
 		else
 			RemoveSymbols<ImageSymbol>(logger, newData, newInfo);
 
-		WriteExports(logger, newData, allNeededImports, outKeptExportCount);
+		StripExports(logger, newData, allNeededImports, outKeptExportCount);
 
 		return newFile.Close();
 	}
 
-	bool ObjectFileCoff::ComputeLoopbacksAndDuplicates(UnorderedSymbols& allSharedExports, UnorderedSymbols& duplicates)
+	bool ObjectFileCoff::ComputeLoopbacksAndDuplicates(UnorderedExports& allSharedExports, UnorderedSymbols& duplicates)
 	{
 		for (auto& importSymbol : m_imports)
 		{
@@ -336,10 +339,16 @@ namespace uba
 				exportEnd = exportStr + strlen(exportStr);
 			else
 				++str;
-			if (strncmp(exportEnd-5, ",DATA", 5) == 0)
-				exportEnd -= 5;
+
+			std::string tmp(exportStr, exportEnd);
+			std::string extra;
+			if (const char* comma = strchr(tmp.c_str(), ','))
+			{
+				extra = comma;
+				tmp = tmp.substr(0, comma - tmp.data());
+			}
 					
-			m_exports.emplace(std::string(exportStr, exportEnd));
+			m_exports.emplace(tmp, extra);
 		}
 		return true;
 	}
@@ -355,7 +364,7 @@ namespace uba
 			if (symbol.StorageClass != ImageSymClassExternal)
 				continue;
 
-			StringView symbolName = GetSymbolName(symbol, m_info.stringTableMem);
+			StringView symbolName = GetSymbolName(symbol, m_data, m_info.stringTableMemPos);
 			std::string symbolString = symbolName.ToString();
 
 			if (symbol.SectionNumber != ImageSymUndefined)
@@ -372,14 +381,14 @@ namespace uba
 		}
 	}
 
-	void ObjectFileCoff::WriteExports(Logger& logger, u8* newData, const UnorderedSymbols& allNeededImports, u32& outKeptExportCount)
+	bool ObjectFileCoff::StripExports(Logger& logger, u8* newData, const UnorderedSymbols& allNeededImports, u32& outKeptExportCount)
 	{
 		if (!m_info.directiveSectionMemOffset)
-			return;
+			return true;
 
 		auto directiveSection = (const ImageSectionHeader*)(m_data + m_info.directiveSectionMemOffset);
 		if (directiveSection->SizeOfRawData < 10)
-			return;
+			return true;
 
 		const u8* directiveData = m_data + directiveSection->PointerToRawData;
 
@@ -437,10 +446,6 @@ namespace uba
 			if (allNeededImports.find(temp) != allNeededImports.end())
 				continue;
 			
-			// TODO: Use exclusion list provided from command line
-			if (strstr(exportStr, "Initialize") == exportStr && strstr(exportStr, "Module"))
-				continue;
-
 			--exportCount;
 
 			u64 toCopy = startPos - lastCopyPos - 1;
@@ -456,14 +461,133 @@ namespace uba
 		writePos += toCopy;
 		*writePos = 0;
 
+		u32 sizeOfRawData = directiveSection->SizeOfRawData;
 		newDirectiveSection->SizeOfRawData = u32((u8*)writePos - newDirectiveData);
-		UBA_ASSERT(newDirectiveSection->SizeOfRawData <= directiveSection->SizeOfRawData);
+		UBA_ASSERT(newDirectiveSection->SizeOfRawData <= sizeOfRawData);
 
-		memset(writePos, 0, directiveSection->SizeOfRawData - newDirectiveSection->SizeOfRawData);
+		memset(writePos, 0, sizeOfRawData - newDirectiveSection->SizeOfRawData);
 
 		outKeptExportCount = exportCount;
+
+		return true;
+	}
+
+	bool ObjectFileCoff::CreateExtraFile(Logger& logger, MemoryBlock& memoryBlock, const UnorderedSymbols& allNeededImports, const UnorderedSymbols& allSharedImports, const UnorderedExports& allSharedExports)
+	{
+		std::string tmp;
+
+		u32 totalStringSize = 0;
+		UnorderedSymbols neededLoopbacks;
+		for (auto& symbol : allSharedImports)
+		{
+			if (strncmp(symbol.data(), "__imp_", 6) != 0)
+				continue;
+			tmp = symbol.substr(6);
+			if (allSharedExports.find(tmp) == allSharedExports.end())
+				continue;
+			if (neededLoopbacks.insert(symbol).second)
+				totalStringSize += u32(symbol.size()) + 1;
+		}
+		u32 loopbackCount = u32(neededLoopbacks.size());
+		UBA_ASSERT(loopbackCount < 65536);
+
+
+		auto allocate = [&](u64 size) { return memoryBlock.Allocate(size, 1, TC("")); };
+		auto write = [&](const void* data, u64 size) { memcpy(allocate(size), data, size); };
+
+
+		// Header
+		auto& header = *(ImageFileHeader*)allocate(sizeof(ImageFileHeader));
+		header.NumberOfSections = 2;
+
+		// Directive section
+		auto& directiveSection = *(ImageSectionHeader*)allocate(sizeof(ImageSectionHeader));
+		memcpy(directiveSection.Name, ".drectve", 8);
+		directiveSection.Characteristics = ImageScnAlign1Bytes|ImageScnLnkInfo|ImageScnLnkRemove;
+
+		// Session for loopbacks
+		auto& textSection = *(ImageSectionHeader*)allocate(sizeof(ImageSectionHeader));
+		memcpy(textSection.Name, ".text$mn", 8);
+		textSection.Characteristics = ImageScnCntCode | ImageScnMemExecute | ImageScnMemRead;
+
+
+		// Directive raw data
+		u32 directiveRawDataStart = u32(memoryBlock.writtenSize);
+		directiveSection.PointerToRawData = directiveRawDataStart;
+		char slashExport[] = "/EXPORT:";
+		for (auto& kv : allSharedExports)
+		{
+			auto& symbol = kv.first;
+
+			if (allNeededImports.find(symbol) == allNeededImports.end())
+			{
+				std::string temp = "__imp_" + symbol;
+				if (allNeededImports.find(temp) == allNeededImports.end())
+					continue;
+			}
+
+			write(slashExport, sizeof(slashExport) - 1);
+			write(symbol.data(), symbol.size());
+			write(kv.second.data(), kv.second.size());
+			write(" ", 1);
+		}
+		write("", 1);
+		directiveSection.SizeOfRawData = u32(memoryBlock.writtenSize) - directiveRawDataStart;
+
+		// Memory for relocations and write relocations
+		u32 relocationsRawDataPos = u32(memoryBlock.writtenSize);
+		allocate(loopbackCount * 8);
+		u32 relocationsMemPos = u32(memoryBlock.writtenSize);
+		auto relocations = (ImageRelocation*)allocate(loopbackCount * sizeof(ImageRelocation));
+		textSection.PointerToRelocations = relocationsMemPos;
+		textSection.NumberOfRelocations = u16(loopbackCount);
+		textSection.PointerToRawData = relocationsRawDataPos;
+		textSection.SizeOfRawData = loopbackCount * 8;
+		for (u32 i=0; i!=loopbackCount; ++i)
+		{
+			auto& relocation = relocations[i];
+			relocation.VirtualAddress = 8 * i;
+			relocation.SymbolTableIndex = i;
+			relocation.Type = ImageRelAmd64Addr64;
+		}
+
+		// Memory for symbols
+		header.PointerToSymbolTable = u32(memoryBlock.writtenSize);
+		header.NumberOfSymbols = loopbackCount * 2;
+		auto symbols = (ImageSymbol*)allocate(header.NumberOfSymbols * sizeof(ImageSymbol));
+
+		// Write string table
+		u64 stringStart = memoryBlock.writtenSize;
+		auto& stringTableSize = *(u32*)allocate(4);
+		Vector<u32> symbolsToAdd;
+		for (auto& str : neededLoopbacks)
+		{
+			symbolsToAdd.push_back(u32(memoryBlock.writtenSize - stringStart));
+			write(str.data(), str.size()+1);
+		}
+		stringTableSize = u32(memoryBlock.writtenSize - stringStart);
+		
+		// Write symbols
+		for (u32 i=0; i!=loopbackCount; ++i)
+		{
+			auto& symbol = symbols[i];
+			symbol.N.Name.Long = symbolsToAdd[i] + 6; // Remove __imp_
+			symbol.SectionNumber = ImageSymUndefined;
+			symbol.StorageClass = ImageSymClassExternal;
+		}
+		for (u32 i=0; i!=loopbackCount; ++i)
+		{
+			auto& symbol = symbols[i + loopbackCount];
+			symbol.N.Name.Long = symbolsToAdd[i];
+			symbol.SectionNumber = 1 + 1;
+			symbol.StorageClass = ImageSymClassExternal;
+			symbol.Value = i * 8;
+		}
+
+		return true;
 	}
 
+
 	template<typename SymbolType> void ObjectFileCoff::CalculateImports(Logger& logger, Vector<u32>& outImports)
 	{
 		// Calculate how much more memory we need by counting imports
@@ -476,7 +600,7 @@ namespace uba
 				continue;
 			if (symbol.SectionNumber != ImageSymUndefined)
 				continue;
-			StringView symbolName = GetSymbolName(symbol, m_info.stringTableMem);
+			StringView symbolName = GetSymbolName(symbol, m_data, m_info.stringTableMemPos);
 			if (!symbolName.StartsWith("__imp_", 6))
 				continue;
 			symbolName.strBegin += 6;
@@ -580,8 +704,8 @@ namespace uba
 			symbol.Value = i * 8;
 		}
 
-		u64 lastToCopySize = m_dataSize - (m_info.stringTableMem - m_data);
-		memcpy(newData + newSymbolsPos + newSymbolsSize, m_info.stringTableMem, lastToCopySize);
+		u64 lastToCopySize = m_dataSize - m_info.stringTableMemPos;
+		memcpy(newData + newSymbolsPos + newSymbolsSize, m_data + m_info.stringTableMemPos, lastToCopySize);
 
 
 		// Traverse all sections and update their PointerToRawData
@@ -622,7 +746,7 @@ namespace uba
 			if (!symbol.SectionNumber)
 				continue;
 
-			StringView symbolName = GetSymbolName(symbol, m_info.stringTableMem);
+			StringView symbolName = GetSymbolName(symbol, m_data, m_info.stringTableMemPos);
 			auto findIt = m_toRemove.find(symbolName.ToString());
 			if (findIt == m_toRemove.end())
 				continue;
diff --git a/Engine/Source/Programs/UnrealBuildAccelerator/Common/Private/UbaObjectFileElf.cpp b/Engine/Source/Programs/UnrealBuildAccelerator/Common/Private/UbaObjectFileElf.cpp
--- a/Engine/Source/Programs/UnrealBuildAccelerator/Common/Private/UbaObjectFileElf.cpp
+++ b/Engine/Source/Programs/UnrealBuildAccelerator/Common/Private/UbaObjectFileElf.cpp
@@ -70,6 +70,7 @@ namespace uba
 	// When merging dlls we need to remove duplicates of these
 	UnorderedSymbols ObjectFileElf::PotentiallyDuplicatedSymbols = 
 	{
+#if 0 // This is not needed anymore
 		// REPLACEMENT_OPERATOR_NEW_AND_DELETE
 		"_Znwm",
 		"_Znam",
@@ -107,6 +108,7 @@ namespace uba
 		"GObjectArrayForDebugVisualizers",
 		"GComplexObjectPathDebug",
 		"GObjectHandlePackageDebug",
+#endif
 	};
 
 
@@ -172,7 +174,7 @@ namespace uba
 		return true;
 	}
 
-	bool ObjectFileElf::ComputeLoopbacksAndDuplicates(UnorderedSymbols& allSharedExports, UnorderedSymbols& duplicates)
+	bool ObjectFileElf::ComputeLoopbacksAndDuplicates(UnorderedExports& allSharedExports, UnorderedSymbols& duplicates)
 	{
 		for (auto& dupSymbol : m_potentialDuplicates)
 			if (!duplicates.insert(dupSymbol).second)
@@ -222,4 +224,14 @@ namespace uba
 
 		return newFile.Close();
 	}
+
+	bool ObjectFileElf::StripExports(Logger& logger, u8* newData, const UnorderedSymbols& allNeededImports, u32& outKeptExportCount)
+	{
+		return true;
+	}
+
+	bool ObjectFileElf::CreateExtraFile(Logger& logger, MemoryBlock& memoryBlock, const UnorderedSymbols& allNeededImports, const UnorderedSymbols& allSharedImports, const UnorderedExports& allSharedExports)
+	{
+		return true;
+	}
 }
diff --git a/Engine/Source/Programs/UnrealBuildAccelerator/Common/Private/UbaSession.cpp b/Engine/Source/Programs/UnrealBuildAccelerator/Common/Private/UbaSession.cpp
--- a/Engine/Source/Programs/UnrealBuildAccelerator/Common/Private/UbaSession.cpp
+++ b/Engine/Source/Programs/UnrealBuildAccelerator/Common/Private/UbaSession.cpp
@@ -4,6 +4,7 @@
 #include "UbaBottleneck.h"
 #include "UbaCompressedObjFileHeader.h"
 #include "UbaFileAccessor.h"
+#include "UbaObjectFile.h"
 #include "UbaProcess.h"
 #include "UbaStorage.h"
 #include "UbaDirectoryIterator.h"
@@ -968,6 +969,7 @@ namespace uba
 		UBA_ASSERTF(m_shouldWriteToDisk || m_allowMemoryMaps, TC("Can't disable both should write to disk and allow memory maps"));
 
 		m_storeObjFilesCompressed = info.storeObjFilesCompressed;
+		m_extractObjFilesSymbols = info.extractObjFilesSymbols;
 
 		m_detailedTrace = info.detailedTrace;
 		m_logToFile = info.logToFile;
@@ -1897,19 +1899,6 @@ namespace uba
 			writtenFile.owner = &msg.process;
 			writtenFile.attributes = msg.attributes;
 
-			auto GetNowFileTime = []()
-				{
-					#if PLATFORM_WINDOWS
-					FILETIME ft;
-					SYSTEMTIME st;
-					GetSystemTime(&st);
-					SystemTimeToFileTime(&st, &ft);
-					return (u64&)ft;
-					#else
-					return 0ull;
-					#endif
-				};
-
 			bool addMapping = true;
 			if (!insres.second)
 			{
@@ -1962,6 +1951,10 @@ namespace uba
 				fileSize = writtenFile.mappingWritten;
 				lastWriteTime = writtenFile.lastWriteTime;
 			}
+
+			if (m_extractObjFilesSymbols && EndsWith(file.name.c_str(), file.name.size(), TC(".obj")) && !EndsWith(file.name.c_str(), file.name.size(), TC(".extra.obj")))
+				if (!ExtractSymbolsFromObjectFile(msg, name, fileSize))
+					return false;
 		}
 
 		if (!msg.newName.IsEmpty())
@@ -2096,8 +2089,7 @@ namespace uba
 
 				auto memClose = MakeGuard([&](){ UnmapViewOfFile(mem, fileSize, file.name.c_str()); });
 
-
-				if (m_storeObjFilesCompressed && EndsWith(file.name.c_str(), file.name.size(), TC(".obj")))
+				if (m_storeObjFilesCompressed && process.m_startInfo.rules->StoreFileCompressed(file.name.c_str(), file.name.size()))
 				{
 					Storage::WriteResult res;
 					CompressedObjFileHeader header { CalculateCasKey(mem, fileSize, true, m_workManager) };
@@ -2627,6 +2619,69 @@ namespace uba
 		return m_cpuLoad;
 	}
 
+	bool Session::ExtractSymbolsFromObjectFile(const CloseFileMessage& msg, const tchar* fileName, u64 fileSize)
+	{
+		FileMappingHandle source;
+		source.FromU64(msg.mappingHandle);
+		FileMappingHandle objectFileMappingHandle;
+
+		if (!DuplicateFileMapping(msg.process.m_nativeProcessHandle, source, GetCurrentProcessHandle(), &objectFileMappingHandle, FILE_MAP_ALL_ACCESS, false, 0))
+			return m_logger.Error(TC("Failed to duplicate file mapping handle for %s"), fileName);
+		auto ofmh = MakeGuard([&]() { CloseFileMapping(objectFileMappingHandle); });
+
+		u8* mem = MapViewOfFile(objectFileMappingHandle, FILE_MAP_ALL_ACCESS, 0, fileSize);
+		if (!mem)
+			return m_logger.Error(TC("Failed to map view of filehandle for read %s (%s)"), fileName, LastErrorToText().data);
+		auto memClose = MakeGuard([&](){ UnmapViewOfFile(mem, fileSize, fileName); });
+
+		ObjectFile* objectFile = ObjectFile::Parse(m_logger, mem, fileSize, fileName);
+		if (!objectFile)
+			return false;
+		auto ofg = MakeGuard([&]() { delete objectFile; });
+
+		if (!objectFile->StripExports(m_logger))
+			return false;
+
+		const tchar* lastDot = TStrrchr(fileName, '.');
+		UBA_ASSERT(lastDot);
+		StringBuffer<> exportsFile;
+		exportsFile.Append(fileName, lastDot - fileName).Append(TC(".sym"));
+
+		MemoryBlock memoryBlock(1*1024*1024);
+		if (!objectFile->WriteSymbols(m_logger, memoryBlock))
+			return false;
+
+		FileMappingHandle symHandle = CreateMemoryMappingW(m_logger, PAGE_READWRITE, memoryBlock.writtenSize);
+		if (!symHandle.IsValid())
+			return false;
+
+		u8* mem2 = MapViewOfFile(symHandle, FILE_MAP_ALL_ACCESS, 0, memoryBlock.writtenSize);
+		if (!mem2)
+			return false;
+
+		memcpy(mem2, memoryBlock.memory, memoryBlock.writtenSize);
+		UnmapViewOfFile(mem2, memoryBlock.writtenSize, TC(""));
+
+		StringKey symFileKey = CaseInsensitiveFs ? ToStringKeyLower(exportsFile) : ToStringKey(exportsFile);
+		u64 lastWriteTime = GetNowFileTime();
+
+		if (!RegisterCreateFileForWrite(symFileKey, exportsFile.data, exportsFile.count, false, memoryBlock.writtenSize, lastWriteTime))
+			return false;
+
+		auto insres = msg.process.m_writtenFiles.try_emplace(exportsFile.data);
+		WrittenFile& writtenFile = insres.first->second;
+
+		UBA_ASSERT(writtenFile.owner == nullptr || writtenFile.owner == &msg.process);
+		writtenFile.key = symFileKey;
+		writtenFile.owner = &msg.process;
+		writtenFile.attributes = msg.attributes;
+		writtenFile.mappingHandle = symHandle;
+		writtenFile.mappingWritten = memoryBlock.writtenSize;
+		writtenFile.lastWriteTime = lastWriteTime;
+		writtenFile.name = insres.first->first;
+
+		return true;
+	}
 
 	void Session::ThreadTraceLoop()
 	{
diff --git a/Engine/Source/Programs/UnrealBuildAccelerator/Common/Public/UbaFile.h b/Engine/Source/Programs/UnrealBuildAccelerator/Common/Public/UbaFile.h
--- a/Engine/Source/Programs/UnrealBuildAccelerator/Common/Public/UbaFile.h
+++ b/Engine/Source/Programs/UnrealBuildAccelerator/Common/Public/UbaFile.h
@@ -75,6 +75,7 @@ namespace uba
 	bool SearchPathW(const tchar* a, const tchar* b, const tchar* c, u32 d, tchar* e, tchar** f);
 	u64 GetSystemTimeAsFileTime();
 	u64 GetFileTimeAsSeconds(u64 fileTime);
+	u64 GetNowFileTime();
 	bool GetCurrentDirectoryW(StringBufferBase& out);
 
 
diff --git a/Engine/Source/Programs/UnrealBuildAccelerator/Common/Public/UbaObjectFile.h b/Engine/Source/Programs/UnrealBuildAccelerator/Common/Public/UbaObjectFile.h
--- a/Engine/Source/Programs/UnrealBuildAccelerator/Common/Public/UbaObjectFile.h
+++ b/Engine/Source/Programs/UnrealBuildAccelerator/Common/Public/UbaObjectFile.h
@@ -10,31 +10,42 @@ namespace uba
 	class FileAccessor;
 	
 	using UnorderedSymbols = UnorderedSet<std::string>;
+	using UnorderedExports = UnorderedMap<std::string, std::string>;
 
 	class ObjectFile
 	{
 	public:
-		static ObjectFile* CreateAndParse(Logger& logger, const tchar* filename);
+		static ObjectFile* OpenAndParse(Logger& logger, const tchar* filename);
+		static ObjectFile* Parse(Logger& logger, u8* data, u64 dataSize, const tchar* hint);
 
-		virtual bool ComputeLoopbacksAndDuplicates(UnorderedSymbols& allSharedExports, UnorderedSymbols& duplicates) = 0;
+		virtual bool CopyMemoryAndClose();
+		virtual bool StripExports(Logger& logger);
+		virtual bool ComputeLoopbacksAndDuplicates(UnorderedExports& allSharedExports, UnorderedSymbols& duplicates) = 0;
 		virtual bool CreateStripped(Logger& logger, const tchar* newFilename, const UnorderedSymbols& allNeededImports, u32& outKeptExportCount) = 0;
+		virtual bool WriteSymbols(Logger& logger, MemoryBlock& memoryBlock);
+		virtual bool WriteSymbols(Logger& logger, const tchar* exportsFilename);
 
 		const tchar* GetFileName() const;
 		const UnorderedSymbols& GetImports() const;
-		const UnorderedSymbols& GetExports() const;
+		const UnorderedExports& GetExports() const;
 		const UnorderedSymbols& GetPotentialDuplicates() const;
 
+		static bool CreateExtraFile(Logger& logger, const tchar* extraObjFilename, const UnorderedSymbols& allNeededImports, const UnorderedSymbols& allSharedImports, const UnorderedExports& allSharedExports);
+
 		virtual ~ObjectFile();
 
 	protected:
 		virtual bool Parse(Logger& logger, const tchar* filename) = 0;
+		virtual bool StripExports(Logger& logger, u8* newData, const UnorderedSymbols& allNeededImports, u32& outKeptExportCount) = 0;
+		virtual bool CreateExtraFile(Logger& logger, MemoryBlock& memoryBlock, const UnorderedSymbols& allNeededImports, const UnorderedSymbols& allSharedImports, const UnorderedExports& allSharedExports) = 0;
 
 		FileAccessor* m_file = nullptr;
 		u8* m_data = nullptr;
 		u64 m_dataSize = 0;
+		bool m_ownsData = false;
 
 		UnorderedSymbols m_imports;
-		UnorderedSymbols m_exports;
+		UnorderedExports m_exports;
 		UnorderedSymbols m_potentialDuplicates;
 	};
 
diff --git a/Engine/Source/Programs/UnrealBuildAccelerator/Common/Public/UbaObjectFileCoff.h b/Engine/Source/Programs/UnrealBuildAccelerator/Common/Public/UbaObjectFileCoff.h
--- a/Engine/Source/Programs/UnrealBuildAccelerator/Common/Public/UbaObjectFileCoff.h
+++ b/Engine/Source/Programs/UnrealBuildAccelerator/Common/Public/UbaObjectFileCoff.h
@@ -6,11 +6,12 @@
 
 namespace uba
 {
+
 	class ObjectFileCoff : public ObjectFile
 	{
 	public:
 		virtual bool Parse(Logger& logger, const tchar* filename) override;
-		virtual bool ComputeLoopbacksAndDuplicates(UnorderedSymbols& allSharedExports, UnorderedSymbols& duplicates) override;
+		virtual bool ComputeLoopbacksAndDuplicates(UnorderedExports& allSharedExports, UnorderedSymbols& duplicates) override;
 		virtual bool CreateStripped(Logger& logger, const tchar* newFilename, const UnorderedSymbols& allNeededImports, u32& outKeptExportCount) override;
 
 	private:
@@ -19,7 +20,9 @@ namespace uba
 		bool ParseExports();
 		template<typename SymbolType> void ParseImports();
 
-		void WriteExports(Logger& logger, u8* newData, const UnorderedSymbols& allNeededImports, u32& outKeptExportCount);
+		virtual bool StripExports(Logger& logger, u8* newData, const UnorderedSymbols& allNeededImports, u32& outKeptExportCount) override;
+		virtual bool CreateExtraFile(Logger& logger, MemoryBlock& memoryBlock, const UnorderedSymbols& allNeededImports, const UnorderedSymbols& allSharedImports, const UnorderedExports& allSharedExports) override;
+
 		template<typename SymbolType> void CalculateImports(Logger& logger, Vector<u32>& outImports);
 		template<typename SymbolType> void WriteImports(Logger& logger, u8* newData, Info& newInfo, const Vector<u32>& symbolsToAdd);
 		template<typename SymbolType> void RemoveSymbols(Logger& logger, u8* newData, Info& newInfo);
@@ -29,7 +32,7 @@ namespace uba
 			u32 sectionsMemOffset = 0;
 			u32 sectionCount = 0;
 			u64 directiveSectionMemOffset = 0;
-			u8* stringTableMem = nullptr;
+			u32 stringTableMemPos = 0;
 			u32 symbolsMemPos = 0;
 			u32 symbolCount = 0;
 		};
diff --git a/Engine/Source/Programs/UnrealBuildAccelerator/Common/Public/UbaObjectFileElf.h b/Engine/Source/Programs/UnrealBuildAccelerator/Common/Public/UbaObjectFileElf.h
--- a/Engine/Source/Programs/UnrealBuildAccelerator/Common/Public/UbaObjectFileElf.h
+++ b/Engine/Source/Programs/UnrealBuildAccelerator/Common/Public/UbaObjectFileElf.h
@@ -10,9 +10,13 @@ namespace uba
 	{
 	public:
 		virtual bool Parse(Logger& logger, const tchar* filename) override;
-		virtual bool ComputeLoopbacksAndDuplicates(UnorderedSymbols& allSharedExports, UnorderedSymbols& duplicates) override;
+		virtual bool ComputeLoopbacksAndDuplicates(UnorderedExports& allSharedExports, UnorderedSymbols& duplicates) override;
 		virtual bool CreateStripped(Logger& logger, const tchar* newFilename, const UnorderedSymbols& allNeededImports, u32& outKeptExportCount) override;
 
+	private:
+		virtual bool StripExports(Logger& logger, u8* newData, const UnorderedSymbols& allNeededImports, u32& outKeptExportCount) override;
+		virtual bool CreateExtraFile(Logger& logger, MemoryBlock& memoryBlock, const UnorderedSymbols& allNeededImports, const UnorderedSymbols& allSharedImports, const UnorderedExports& allSharedExports) override;
+
 		UnorderedSymbols m_toRemove;
 
 		static UnorderedSymbols PotentiallyDuplicatedSymbols;
diff --git a/Engine/Source/Programs/UnrealBuildAccelerator/Common/Public/UbaSession.h b/Engine/Source/Programs/UnrealBuildAccelerator/Common/Public/UbaSession.h
--- a/Engine/Source/Programs/UnrealBuildAccelerator/Common/Public/UbaSession.h
+++ b/Engine/Source/Programs/UnrealBuildAccelerator/Common/Public/UbaSession.h
@@ -149,6 +149,8 @@ namespace uba
 		void WriteSummary(BinaryWriter& writer, const Function<void(Logger& logger)>& summaryFunc);
 		float UpdateCpuLoad();
 
+		bool ExtractSymbolsFromObjectFile(const CloseFileMessage& msg, const tchar* fileName, u64 fileSize);
+
 		void ThreadTraceLoop();
 		virtual void TraceSessionUpdate();
 
@@ -174,6 +176,7 @@ namespace uba
 		bool m_detailedTrace;
 		bool m_logToFile;
 		bool m_storeObjFilesCompressed;
+		bool m_extractObjFilesSymbols;
 
 		u64 m_keepOutputFileMemoryMapsThreshold;
 
diff --git a/Engine/Source/Programs/UnrealBuildAccelerator/Common/Public/UbaSessionCreateInfo.h b/Engine/Source/Programs/UnrealBuildAccelerator/Common/Public/UbaSessionCreateInfo.h
--- a/Engine/Source/Programs/UnrealBuildAccelerator/Common/Public/UbaSessionCreateInfo.h
+++ b/Engine/Source/Programs/UnrealBuildAccelerator/Common/Public/UbaSessionCreateInfo.h
@@ -27,6 +27,7 @@ namespace uba
 		bool traceEnabled = false;				// Set to true to always create in-memory trace data. Is not needed if traceName, traceOutputFile or launchVisualizer is set
 		bool detailedTrace = false;				// Enable detailed trace to include jobs, individual file I/O etc in trace dump
 		bool storeObjFilesCompressed = false;	// Compiler will write .obj files compressed to disk and linker will decompress when reading
+		bool extractObjFilesSymbols = false;	// Will extract import/export symbols and clear export table from obj file.
 		u64 deleteSessionsOlderThanSeconds = 12 * 60 * 60; // Delete session folders older than 12 hours by default . Set to 0 to not delete or 1 to delete all
 		u64 keepOutputFileMemoryMapsThreshold = 256 * 1024; // If allowMemoryMaps is true, output files will be kept in memory if smaller than this size
 	};
diff --git a/Engine/Source/Programs/UnrealBuildAccelerator/Core/Private/UbaApplicationRules.cpp b/Engine/Source/Programs/UnrealBuildAccelerator/Core/Private/UbaApplicationRules.cpp
--- a/Engine/Source/Programs/UnrealBuildAccelerator/Core/Private/UbaApplicationRules.cpp
+++ b/Engine/Source/Programs/UnrealBuildAccelerator/Core/Private/UbaApplicationRules.cpp
@@ -82,6 +82,11 @@ namespace uba
 		{
 			return true;
 		}
+
+		virtual bool StoreFileCompressed(const tchar* fileName, u64 fileNameLen) const
+		{
+			return EndsWith(fileName, fileNameLen, TC(".obj"));
+		}
 	};
 
 	class ApplicationRulesLinkExe : public ApplicationRulesVC
@@ -197,6 +202,11 @@ namespace uba
 		{
 			return true;
 		}
+
+		virtual bool StoreFileCompressed(const tchar* fileName, u64 fileNameLen) const
+		{
+			return EndsWith(fileName, fileNameLen, TC(".obj"));
+		}
 	};
 
 	class ApplicationRulesLdLLdExe : public ApplicationRulesClang
diff --git a/Engine/Source/Programs/UnrealBuildAccelerator/Core/Public/UbaApplicationRules.h b/Engine/Source/Programs/UnrealBuildAccelerator/Core/Public/UbaApplicationRules.h
--- a/Engine/Source/Programs/UnrealBuildAccelerator/Core/Public/UbaApplicationRules.h
+++ b/Engine/Source/Programs/UnrealBuildAccelerator/Core/Public/UbaApplicationRules.h
@@ -129,6 +129,11 @@ namespace uba
 			return false;
 		}
 
+		virtual bool StoreFileCompressed(const tchar* fileName, u64 fileNameLen) const
+		{
+			return false;
+		}
+
 		u32 index = ~0u;
 	};
 
diff --git a/Engine/Source/Programs/UnrealBuildAccelerator/ObjTool/Private/UbaObjTool.cpp b/Engine/Source/Programs/UnrealBuildAccelerator/ObjTool/Private/UbaObjTool.cpp
--- a/Engine/Source/Programs/UnrealBuildAccelerator/ObjTool/Private/UbaObjTool.cpp
+++ b/Engine/Source/Programs/UnrealBuildAccelerator/ObjTool/Private/UbaObjTool.cpp
@@ -2,6 +2,7 @@
 
 #include "UbaObjectFile.h"
 #include "UbaDirectoryIterator.h"
+#include "UbaFileAccessor.h"
 #include "UbaVersion.h"
 #include "UbaWorkManager.h"
 
@@ -30,8 +31,18 @@ namespace uba
 		logger.Info(TC(""));
 		logger.Info(TC("  UbaObjTool.exe [options...] <objfile>"));
 		logger.Info(TC(""));
+		logger.Info(TC("   Options:"));
+		logger.Info(TC("    -printsymbols            Print the symbols found in obj file"));
+		logger.Info(TC("    -stripexports            Will strip exports and write them out in a .exp file"));
 		logger.Info(TC(""));
-		logger.Info(TC("  -printsymbols                    Print the symbols found in obj file."));
+		logger.Info(TC("  --- OR ---"));
+		logger.Info(TC(""));
+		logger.Info(TC("  UbaObjTool.exe @<rspfile>"));
+		logger.Info(TC(""));
+		logger.Info(TC("   Response file options:"));
+		logger.Info(TC("    /S:<objfile>             Obj file to strip. Will produce a .strip.obj file. Multiple allowed"));
+		logger.Info(TC("    /D:<objfile>             Obj file depending on obj files to strip. Multiple allowed"));
+		logger.Info(TC("    /O:<objfile>             Obj file to output containing exports and loopbacks"));
 		logger.Info(TC(""));
 		return -1;
 	}
@@ -42,9 +53,11 @@ namespace uba
 
 		TString objFile;
 		bool printSymbols = false;
+		bool stripExports = false;
 
 		Vector<TString> objFilesToStrip;
 		Vector<TString> objFilesDependencies;
+		TString extraObjFile;
 
 		auto parseArg = [&](const tchar* arg)
 			{
@@ -69,10 +82,18 @@ namespace uba
 				{
 					objFilesToStrip.push_back(name.data + 3);
 				}
+				else if (name.StartsWith(TC("/O:")))
+				{
+					extraObjFile = name.data + 3;
+				}
 				else if (name.Equals(TC("-printsymbols")))
 				{
 					printSymbols = true;
 				}
+				else if (name.Equals(TC("-stripexports")))
+				{
+					stripExports = true;
+				}
 				else if (name.Equals(TC("-?")))
 				{
 					return PrintHelp(TC(""));
@@ -121,6 +142,42 @@ namespace uba
 				return res;
 		}
 
+		struct SymbolFile
+		{
+			UnorderedSymbols imports;
+			UnorderedExports exports;
+
+			bool ParseFile(Logger& logger, const tchar* filename)
+			{
+				FileAccessor symFile(logger, filename);
+				if (!symFile.OpenMemoryRead())
+					return false;
+				auto readPos = (const char*)symFile.GetData();
+
+				while (*readPos)
+				{
+					auto strEnd = strlen(readPos);
+					imports.insert(std::string(readPos, readPos + strEnd));
+					readPos = readPos + strEnd + 1;
+				}
+				++readPos;
+
+				while (*readPos)
+				{
+					auto strEnd = strlen(readPos);
+					std::string extra;
+					if (const char* comma = strchr(readPos, ','))
+					{
+						strEnd = comma - readPos;
+						extra = comma;
+					}
+					exports.emplace(std::string(readPos, readPos + strEnd), extra);
+					readPos = readPos + strEnd + 1;
+				}
+				return true;
+			}
+		};
+
 		FilteredLogWriter logWriter(g_consoleLogWriter, LogEntryType_Info);
 		LoggerWithWriter logger(logWriter, TC(""));
 
@@ -134,22 +191,38 @@ namespace uba
 			WorkManagerImpl workManager(workerCount);
 			workManager.ParallelFor(workerCount, objFilesDependencies, [&](auto& it)
 				{
-					const TString& objFile = *it;
-					ObjectFile* objectFile = ObjectFile::CreateAndParse(logger, objFile.c_str());
-					if (!objectFile)
+					const TString& objFileName = *it;
+
+					if (EndsWith(objFileName.c_str(), objFileName.size(), TC(".sym")))
 					{
-						success = false;
-						return;
+						SymbolFile symbolFile;
+						if (!symbolFile.ParseFile(logger, objFileName.c_str()))
+						{
+							success = false;
+							return;
+						}
+						ScopedCriticalSection _(cs);
+						allNeededImports.insert(symbolFile.imports.begin(), symbolFile.imports.end());
 					}
-					auto g = MakeGuard([&]() { delete objectFile; });
+					else
+					{
+						ObjectFile* objectFile = ObjectFile::OpenAndParse(logger, objFileName.c_str());
+						if (!objectFile)
+						{
+							success = false;
+							return;
+						}
+						auto g = MakeGuard([&]() { delete objectFile; });
 
-					ScopedCriticalSection _(cs);
-					allNeededImports.insert(objectFile->GetImports().begin(), objectFile->GetImports().end());
+						ScopedCriticalSection _(cs);
+						allNeededImports.insert(objectFile->GetImports().begin(), objectFile->GetImports().end());
+					}
 				});
 			if (!success)
 				return -1;
 
-			UnorderedSymbols allSharedExports; // Exports from all the obj files about to be stripped
+			UnorderedSymbols allSharedImports; // Imports from all the obj files about to be stripped
+			UnorderedExports allSharedExports; // Exports from all the obj files about to be stripped
 
 			Map<TString, ObjectFile*> objectFiles;
 			auto g = MakeGuard([&]() { for (auto& kv : objectFiles) delete kv.second; });
@@ -157,67 +230,112 @@ namespace uba
 			workManager.ParallelFor(workerCount, objFilesToStrip, [&](auto& it)
 				{
 					const TString& objFileName = *it;
-					ObjectFile* objectFile = ObjectFile::CreateAndParse(logger, objFileName.c_str());
-					if (!objectFile)
+
+					if (EndsWith(objFileName.c_str(), objFileName.size(), TC(".sym")))
 					{
-						success = false;
-						return;
+						SymbolFile symbolFile;
+						if (!symbolFile.ParseFile(logger, objFileName.c_str()))
+						{
+							success = false;
+							return;
+						}
+						ScopedCriticalSection _(cs);
+						allSharedImports.insert(symbolFile.imports.begin(), symbolFile.imports.end());
+						allSharedExports.insert(symbolFile.exports.begin(), symbolFile.exports.end());
 					}
+					else
+					{
+						ObjectFile* objectFile = ObjectFile::OpenAndParse(logger, objFileName.c_str());
+						if (!objectFile)
+						{
+							success = false;
+							return;
+						}
 
-					ScopedCriticalSection _(cs);
-					objectFiles.try_emplace(objFileName, objectFile);
-					allSharedExports.insert(objectFile->GetExports().begin(), objectFile->GetExports().end());
+						ScopedCriticalSection _(cs);
+						objectFiles.try_emplace(objFileName, objectFile);
+						allSharedExports.insert(objectFile->GetExports().begin(), objectFile->GetExports().end());
+					}
 				});
 			if (!success)
 				return -1;
 
 			// Figure out which loopback symbols that should be added to which obj file
-			UnorderedSymbols duplicates;
-			for (auto& objFileName : objFilesToStrip)
-				if (!objectFiles[objFileName]->ComputeLoopbacksAndDuplicates(allSharedExports, duplicates))
-					return -1;
+			if (!objectFiles.empty())
+			{
+				UnorderedSymbols duplicates;
+				for (auto& objFileName : objFilesToStrip)
+					if (!objectFiles[objFileName]->ComputeLoopbacksAndDuplicates(allSharedExports, duplicates))
+						return -1;
 
-			Atomic<u32> totalExportCount;
-			Atomic<u32> totalKeptExportCount;
-			workManager.ParallelFor(workerCount, objectFiles, [&](auto& it)
-				{
-					ObjectFile& file = *it->second;
-					const tchar* fileName = file.GetFileName();
-					const tchar* lastDot = TStrrchr(fileName, '.');
-					UBA_ASSERT(lastDot);
-					StringBuffer<> newFilename;
-					newFilename.Append(fileName, lastDot - fileName).Append(TC(".strip")).Append(lastDot);
-					u32 keptExportCount = 0;
-					if (!file.CreateStripped(logger, newFilename.data, allNeededImports, keptExportCount))
-						success = false;
-
-					totalExportCount += u32(file.GetExports().size());
-					totalKeptExportCount += keptExportCount;
-				});
-			if (!success)
-				return -1;
+				Atomic<u32> totalExportCount;
+				Atomic<u32> totalKeptExportCount;
+				workManager.ParallelFor(workerCount, objectFiles, [&](auto& it)
+					{
+						ObjectFile& file = *it->second;
+						const tchar* fileName = file.GetFileName();
+						const tchar* lastDot = TStrrchr(fileName, '.');
+						UBA_ASSERT(lastDot);
+						StringBuffer<> newFilename;
+						newFilename.Append(fileName, lastDot - fileName).Append(TC(".strip")).Append(lastDot);
+						u32 keptExportCount = 0;
+						if (!file.CreateStripped(logger, newFilename.data, allNeededImports, keptExportCount))
+							success = false;
+
+						totalExportCount += u32(file.GetExports().size());
+						totalKeptExportCount += keptExportCount;
+					});
+				if (!success)
+					return -1;
+			}
+			else if (!extraObjFile.empty())
+			{
+				if (!ObjectFile::CreateExtraFile(logger, extraObjFile.c_str(), allNeededImports, allSharedImports, allSharedExports))
+					return -1;
+			}
 
 			//logger.Info(TC("Reduced export count from %llu to %llu"), totalExportCount.load(), totalKeptExportCount.size());
 		}
 		else
 		{
 			if (objFile.empty())
-				return PrintHelp(TC("No obj file provided"));
+				return PrintHelp(TC("No obj or rsp file provided"));
 
-			ObjectFile* objectFile = ObjectFile::CreateAndParse(logger, objFile.c_str());
+			ObjectFile* objectFile = ObjectFile::OpenAndParse(logger, objFile.c_str());
 			if (!objectFile)
 				return -1;
+			auto g = MakeGuard([&](){ delete objectFile; });
 
 			if (printSymbols)
 			{
 				for (auto& symbol : objectFile->GetImports())
 					logger.Info(TC("I %S"), symbol.c_str());
 
-				for (auto& symbol : objectFile->GetExports())
-					logger.Info(TC("E %S"), symbol.c_str());
+				for (auto& kv : objectFile->GetExports())
+					logger.Info(TC("E %S%S"), kv.first.c_str(), kv.second.c_str());
 			}
 
-			delete objectFile;
+			if (stripExports)
+			{
+				if (!objectFile->CopyMemoryAndClose())
+					return false;
+
+				const tchar* fileName = objFile.c_str();
+				const tchar* lastDot = TStrrchr(fileName, '.');
+				UBA_ASSERT(lastDot);
+				StringBuffer<> exportsFile;
+				exportsFile.Append(fileName, lastDot - fileName).Append(TC(".sym"));
+
+				if (!objectFile->WriteSymbols(logger, exportsFile.data))
+					return false;
+
+				u32 keptExportCount = 0;
+
+				StringBuffer<> newFilename;
+				newFilename.Append(fileName, lastDot - fileName).Append(TC(".TEST")).Append(lastDot);
+				if (!objectFile->CreateStripped(logger, newFilename.data, {}, keptExportCount))
+					return false;
+			}
 		}
 		return 0;
 	}
