[UBA]
* Changed so uba consumes command line option "/extractsymbols" and will use that to decide if symbols should be extracted from obj file
* UbaObjTool - Some cleanups and optimizations. Also added code just to get it in the version control (will remove again)
* Renamed symbol files from .sym to .exi (exportimport) to not collide with already existing .sym extension

diff --git a/Engine/Source/Programs/UnrealBuildAccelerator/Common/Private/UbaObjectFile.cpp b/Engine/Source/Programs/UnrealBuildAccelerator/Common/Private/UbaObjectFile.cpp
--- a/Engine/Source/Programs/UnrealBuildAccelerator/Common/Private/UbaObjectFile.cpp
+++ b/Engine/Source/Programs/UnrealBuildAccelerator/Common/Private/UbaObjectFile.cpp
@@ -149,12 +149,12 @@ namespace uba
 		return m_potentialDuplicates;
 	}
 
-	bool ObjectFile::CreateExtraFile(Logger& logger, const tchar* extraObjFilename, const UnorderedSymbols& allNeededImports, const UnorderedSymbols& allSharedImports, const UnorderedExports& allSharedExports)
+	bool ObjectFile::CreateExtraFile(Logger& logger, const tchar* extraObjFilename, const UnorderedSymbols& allNeededImports, const UnorderedSymbols& allSharedImports, const UnorderedExports& allSharedExports, bool includeExportsInFile)
 	{
 		ObjectFileCoff objectFile;
 		MemoryBlock memoryBlock(16*1024*1024);
 
-		if (!((ObjectFile&)objectFile).CreateExtraFile(logger, memoryBlock, allNeededImports, allSharedImports, allSharedExports))
+		if (!((ObjectFile&)objectFile).CreateExtraFile(logger, memoryBlock, allNeededImports, allSharedImports, allSharedExports, includeExportsInFile))
 			return false;
 
 		FileAccessor extraFile(logger, extraObjFilename);
@@ -166,4 +166,22 @@ namespace uba
 
 		return extraFile.Close();
 	}
+
+	bool ObjectFile::CreateDefFile(Logger& logger, const tchar* defFilename, const UnorderedSymbols& allNeededImports, const UnorderedExports& allSharedExports)
+	{
+		ObjectFileCoff objectFile;
+		MemoryBlock memoryBlock(16*1024*1024);
+
+		if (!((ObjectFile&)objectFile).CreateDefFile(logger, memoryBlock, allNeededImports, allSharedExports))
+			return false;
+
+		FileAccessor defFile(logger, defFilename);
+		if (!defFile.CreateWrite())
+			return false;
+
+		if (!defFile.Write(memoryBlock.memory, memoryBlock.writtenSize))
+			return false;
+
+		return defFile.Close();
+	}
 }
diff --git a/Engine/Source/Programs/UnrealBuildAccelerator/Common/Private/UbaObjectFileCoff.cpp b/Engine/Source/Programs/UnrealBuildAccelerator/Common/Private/UbaObjectFileCoff.cpp
--- a/Engine/Source/Programs/UnrealBuildAccelerator/Common/Private/UbaObjectFileCoff.cpp
+++ b/Engine/Source/Programs/UnrealBuildAccelerator/Common/Private/UbaObjectFileCoff.cpp
@@ -288,11 +288,13 @@ namespace uba
 
 	bool ObjectFileCoff::ComputeLoopbacksAndDuplicates(UnorderedExports& allSharedExports, UnorderedSymbols& duplicates)
 	{
+		std::string importString;
+
 		for (auto& importSymbol : m_imports)
 		{
 			if (strncmp(importSymbol.c_str(), "__imp_", 6) != 0)
 				continue;
-			std::string importString = importSymbol.substr(6);
+			importString = importSymbol.c_str() + 6;
 			auto findIt = allSharedExports.find(importString);
 			if (findIt == allSharedExports.end())
 				continue;
@@ -326,6 +328,9 @@ namespace uba
 		static constexpr u8 utf8Bom[3] = { 0xef, 0xbb, 0xbf };
 		UBA_ASSERT(memcmp(directiveData, utf8Bom, 3) != 0);
 
+		std::string tmp;
+		std::string extra;
+
 		auto str = (char*)directiveData;
 		while (str)
 		{
@@ -340,8 +345,8 @@ namespace uba
 			else
 				++str;
 
-			std::string tmp(exportStr, exportEnd);
-			std::string extra;
+			tmp.assign(exportStr, exportEnd);
+			extra.clear();
 			if (const char* comma = strchr(tmp.c_str(), ','))
 			{
 				extra = comma;
@@ -356,6 +361,7 @@ namespace uba
 	template<typename SymbolType>
 	void ObjectFileCoff::ParseImports()
 	{
+		std::string symbolString;
 		auto symbols = (SymbolType*)(m_data + m_info.symbolsMemPos);
 		for (u32 i=0; i!=m_info.symbolCount; ++i)
 		{
@@ -365,7 +371,7 @@ namespace uba
 				continue;
 
 			StringView symbolName = GetSymbolName(symbol, m_data, m_info.stringTableMemPos);
-			std::string symbolString = symbolName.ToString();
+			symbolName.ToString(symbolString);
 
 			if (symbol.SectionNumber != ImageSymUndefined)
 			{
@@ -395,6 +401,8 @@ namespace uba
 		auto newDirectiveSection = (ImageSectionHeader*)(newData + m_info.directiveSectionMemOffset);
 		u8* newDirectiveData = newData + newDirectiveSection->PointerToRawData;
 
+		std::string tmp;
+
 		char* writePos = (char*)newDirectiveData;
 		const char* lastCopyPos = (const char*)directiveData;
 
@@ -439,11 +447,11 @@ namespace uba
 			if (strncmp(exportEnd-5, ",DATA", 5) == 0)
 				exportEnd -= 5;
 
-			std::string temp(exportStr, exportEnd - exportStr);
-			if (allNeededImports.find(temp) != allNeededImports.end())
+			tmp.assign(exportStr, exportEnd - exportStr);
+			if (allNeededImports.find(tmp) != allNeededImports.end())
 				continue;
-			temp = "__imp_" + temp;
-			if (allNeededImports.find(temp) != allNeededImports.end())
+			tmp.assign("__imp_").append(exportStr, exportEnd - exportStr);
+			if (allNeededImports.find(tmp) != allNeededImports.end())
 				continue;
 			
 			--exportCount;
@@ -472,7 +480,7 @@ namespace uba
 		return true;
 	}
 
-	bool ObjectFileCoff::CreateExtraFile(Logger& logger, MemoryBlock& memoryBlock, const UnorderedSymbols& allNeededImports, const UnorderedSymbols& allSharedImports, const UnorderedExports& allSharedExports)
+	bool ObjectFileCoff::CreateExtraFile(Logger& logger, MemoryBlock& memoryBlock, const UnorderedSymbols& allNeededImports, const UnorderedSymbols& allSharedImports, const UnorderedExports& allSharedExports, bool includeExportsInFile)
 	{
 		std::string tmp;
 
@@ -498,41 +506,46 @@ namespace uba
 
 		// Header
 		auto& header = *(ImageFileHeader*)allocate(sizeof(ImageFileHeader));
-		header.NumberOfSections = 2;
-
-		// Directive section
-		auto& directiveSection = *(ImageSectionHeader*)allocate(sizeof(ImageSectionHeader));
-		memcpy(directiveSection.Name, ".drectve", 8);
-		directiveSection.Characteristics = ImageScnAlign1Bytes|ImageScnLnkInfo|ImageScnLnkRemove;
 
 		// Session for loopbacks
 		auto& textSection = *(ImageSectionHeader*)allocate(sizeof(ImageSectionHeader));
 		memcpy(textSection.Name, ".text$mn", 8);
 		textSection.Characteristics = ImageScnCntCode | ImageScnMemExecute | ImageScnMemRead;
+		u16 textSessionIndex = header.NumberOfSections;
+		++header.NumberOfSections;
 
-
-		// Directive raw data
-		u32 directiveRawDataStart = u32(memoryBlock.writtenSize);
-		directiveSection.PointerToRawData = directiveRawDataStart;
-		char slashExport[] = "/EXPORT:";
-		for (auto& kv : allSharedExports)
+		// Directive section
+		if (includeExportsInFile)
 		{
-			auto& symbol = kv.first;
+			auto& directiveSection = *(ImageSectionHeader*)allocate(sizeof(ImageSectionHeader));
+			memcpy(directiveSection.Name, ".drectve", 8);
+			directiveSection.Characteristics = ImageScnAlign1Bytes|ImageScnLnkInfo|ImageScnLnkRemove;
+			++header.NumberOfSections;
 
-			if (allNeededImports.find(symbol) == allNeededImports.end())
+			// Directive raw data
+			u32 directiveRawDataStart = u32(memoryBlock.writtenSize);
+			directiveSection.PointerToRawData = directiveRawDataStart;
+			char slashExport[] = "/EXPORT:";
+			for (auto& kv : allSharedExports)
 			{
-				std::string temp = "__imp_" + symbol;
-				if (allNeededImports.find(temp) == allNeededImports.end())
-					continue;
-			}
+				auto& symbol = kv.first;
 
-			write(slashExport, sizeof(slashExport) - 1);
-			write(symbol.data(), symbol.size());
-			write(kv.second.data(), kv.second.size());
-			write(" ", 1);
+				if (allNeededImports.find(symbol) == allNeededImports.end())
+				{
+					tmp.assign("__imp_").append(symbol);
+					if (allNeededImports.find(tmp) == allNeededImports.end())
+						continue;
+				}
+
+				write(slashExport, sizeof(slashExport) - 1);
+				write(symbol.data(), symbol.size());
+				write(kv.second.data(), kv.second.size());
+				write(",NONAME ", 8);
+			}
+			write("", 1);
+			directiveSection.SizeOfRawData = u32(memoryBlock.writtenSize) - directiveRawDataStart;
 		}
-		write("", 1);
-		directiveSection.SizeOfRawData = u32(memoryBlock.writtenSize) - directiveRawDataStart;
+
 
 		// Memory for relocations and write relocations
 		u32 relocationsRawDataPos = u32(memoryBlock.writtenSize);
@@ -579,7 +592,7 @@ namespace uba
 		{
 			auto& symbol = symbols[i + loopbackCount];
 			symbol.N.Name.Long = symbolsToAdd[i];
-			symbol.SectionNumber = 1 + 1;
+			symbol.SectionNumber = textSessionIndex + 1;
 			symbol.StorageClass = ImageSymClassExternal;
 			symbol.Value = i * 8;
 		}
@@ -587,11 +600,52 @@ namespace uba
 		return true;
 	}
 
+	bool ObjectFileCoff::CreateDefFile(Logger& logger, MemoryBlock& memoryBlock, const UnorderedSymbols& allNeededImports, const UnorderedExports& allSharedExports)
+	{
+		auto allocate = [&](u64 size) { return memoryBlock.Allocate(size, 1, TC("")); };
+		auto write = [&](const void* data, u64 size) { memcpy(allocate(size), data, size); };
+		
+		//write("EXPORTS\n", 8);
+
+		std::string tmp;
+		u32 counter = 0;
+
+		for (auto& kv : allSharedExports)
+		{
+			auto& symbol = kv.first;
+			if (allNeededImports.find(symbol) == allNeededImports.end())
+				if (allNeededImports.find(tmp.assign("__imp_").append(symbol)) == allNeededImports.end())
+					continue;
+
+			char buf[512];
+			u64 len = snprintf(buf, 256, "/EXPORT:_%u,@%s,NONAME\n", ++counter, symbol.data());
+			write(buf, len);
+
+			//write("/EXPORT:@", 9);
+			//write(symbol.data(), symbol.size());
+			//write(kv.second.data(), kv.second.size());
+
+			//char extra[64] = { 0 };
+			//if (!kv.second.empty())
+			//{
+			//	extra[0] = ' ';
+			//	strcpy_s(extra+1, 63, kv.second.c_str() + 1);
+			//}
+
+			//char buf[256];
+			//u64 len = snprintf(buf, 256, " @%u%s\n", ++counter, extra);
+			//write(buf, len);
+			//write("\n", 1);
+		}
+		write("", 0);
+		return true;
+	}
 
 	template<typename SymbolType> void ObjectFileCoff::CalculateImports(Logger& logger, Vector<u32>& outImports)
 	{
 		// Calculate how much more memory we need by counting imports
 		auto symbols = (SymbolType*)(m_data + m_info.symbolsMemPos);
+		std::string tmp;
 
 		for (u32 i=0; i!=m_info.symbolCount; ++i)
 		{
@@ -604,7 +658,7 @@ namespace uba
 			if (!symbolName.StartsWith("__imp_", 6))
 				continue;
 			symbolName.strBegin += 6;
-			auto findIt = m_loopbacksToAdd.find(symbolName.ToString());
+			auto findIt = m_loopbacksToAdd.find(symbolName.ToString(tmp));
 			if (findIt == m_loopbacksToAdd.end())
 				continue;
 			m_loopbacksToAdd.erase(findIt);
@@ -736,6 +790,8 @@ namespace uba
 		auto symbols = (SymbolType*)(newData + newInfo.symbolsMemPos);
 		auto sections = (ImageSectionHeader*)(newData + newInfo.sectionsMemOffset);
 
+		std::string tmp;
+
 		for (u32 i=0; i!=m_info.symbolCount; ++i)
 		{
 			auto& symbol = symbols[i];
@@ -747,7 +803,7 @@ namespace uba
 				continue;
 
 			StringView symbolName = GetSymbolName(symbol, m_data, m_info.stringTableMemPos);
-			auto findIt = m_toRemove.find(symbolName.ToString());
+			auto findIt = m_toRemove.find(symbolName.ToString(tmp));
 			if (findIt == m_toRemove.end())
 				continue;
 
diff --git a/Engine/Source/Programs/UnrealBuildAccelerator/Common/Private/UbaObjectFileElf.cpp b/Engine/Source/Programs/UnrealBuildAccelerator/Common/Private/UbaObjectFileElf.cpp
--- a/Engine/Source/Programs/UnrealBuildAccelerator/Common/Private/UbaObjectFileElf.cpp
+++ b/Engine/Source/Programs/UnrealBuildAccelerator/Common/Private/UbaObjectFileElf.cpp
@@ -230,7 +230,12 @@ namespace uba
 		return true;
 	}
 
-	bool ObjectFileElf::CreateExtraFile(Logger& logger, MemoryBlock& memoryBlock, const UnorderedSymbols& allNeededImports, const UnorderedSymbols& allSharedImports, const UnorderedExports& allSharedExports)
+	bool ObjectFileElf::CreateExtraFile(Logger& logger, MemoryBlock& memoryBlock, const UnorderedSymbols& allNeededImports, const UnorderedSymbols& allSharedImports, const UnorderedExports& allSharedExports, bool includeExportsInFile)
+	{
+		return true;
+	}
+
+	bool ObjectFileElf::CreateDefFile(Logger& logger, MemoryBlock& memoryBlock, const UnorderedSymbols& allNeededImports, const UnorderedExports& allSharedExports)
 	{
 		return true;
 	}
diff --git a/Engine/Source/Programs/UnrealBuildAccelerator/Common/Private/UbaProcess.cpp b/Engine/Source/Programs/UnrealBuildAccelerator/Common/Private/UbaProcess.cpp
--- a/Engine/Source/Programs/UnrealBuildAccelerator/Common/Private/UbaProcess.cpp
+++ b/Engine/Source/Programs/UnrealBuildAccelerator/Common/Private/UbaProcess.cpp
@@ -159,6 +159,15 @@ namespace uba
 		m_startInfo.application = m_virtualApplication.c_str();
 
 		m_arguments = startInfo.arguments;
+
+		tchar extractExports[] = TC(" /extractexports");
+		const tchar* pos = nullptr;
+		if (Contains(m_arguments.c_str(), extractExports, true, &pos))
+		{
+			m_arguments.erase(pos - m_arguments.c_str(), sizeof_array(extractExports));
+			m_extractExports = true;
+		}
+
 		m_startInfo.arguments = m_arguments.c_str();
 
 		m_virtualWorkingDir = startInfo.workingDir;
diff --git a/Engine/Source/Programs/UnrealBuildAccelerator/Common/Private/UbaSession.cpp b/Engine/Source/Programs/UnrealBuildAccelerator/Common/Private/UbaSession.cpp
--- a/Engine/Source/Programs/UnrealBuildAccelerator/Common/Private/UbaSession.cpp
+++ b/Engine/Source/Programs/UnrealBuildAccelerator/Common/Private/UbaSession.cpp
@@ -969,7 +969,6 @@ namespace uba
 		UBA_ASSERTF(m_shouldWriteToDisk || m_allowMemoryMaps, TC("Can't disable both should write to disk and allow memory maps"));
 
 		m_storeObjFilesCompressed = info.storeObjFilesCompressed;
-		m_extractObjFilesSymbols = info.extractObjFilesSymbols;
 
 		m_detailedTrace = info.detailedTrace;
 		m_logToFile = info.logToFile;
@@ -1952,7 +1951,7 @@ namespace uba
 				lastWriteTime = writtenFile.lastWriteTime;
 			}
 
-			if (m_extractObjFilesSymbols && EndsWith(file.name.c_str(), file.name.size(), TC(".obj")) && !EndsWith(file.name.c_str(), file.name.size(), TC(".extra.obj")))
+			if (msg.process.m_extractExports && msg.process.m_startInfo.rules->ShouldExtractSymbols(file.name.c_str(), file.name.size()))
 				if (!ExtractSymbolsFromObjectFile(msg, name, fileSize))
 					return false;
 		}
@@ -2645,7 +2644,7 @@ namespace uba
 		const tchar* lastDot = TStrrchr(fileName, '.');
 		UBA_ASSERT(lastDot);
 		StringBuffer<> exportsFile;
-		exportsFile.Append(fileName, lastDot - fileName).Append(TC(".sym"));
+		exportsFile.Append(fileName, lastDot - fileName).Append(TC(".exi"));
 
 		MemoryBlock memoryBlock(1*1024*1024);
 		if (!objectFile->WriteSymbols(m_logger, memoryBlock))
diff --git a/Engine/Source/Programs/UnrealBuildAccelerator/Common/Public/UbaObjectFile.h b/Engine/Source/Programs/UnrealBuildAccelerator/Common/Public/UbaObjectFile.h
--- a/Engine/Source/Programs/UnrealBuildAccelerator/Common/Public/UbaObjectFile.h
+++ b/Engine/Source/Programs/UnrealBuildAccelerator/Common/Public/UbaObjectFile.h
@@ -30,14 +30,16 @@ namespace uba
 		const UnorderedExports& GetExports() const;
 		const UnorderedSymbols& GetPotentialDuplicates() const;
 
-		static bool CreateExtraFile(Logger& logger, const tchar* extraObjFilename, const UnorderedSymbols& allNeededImports, const UnorderedSymbols& allSharedImports, const UnorderedExports& allSharedExports);
+		static bool CreateExtraFile(Logger& logger, const tchar* extraObjFilename, const UnorderedSymbols& allNeededImports, const UnorderedSymbols& allSharedImports, const UnorderedExports& allSharedExports, bool includeExportsInFile);
+		static bool CreateDefFile(Logger& logger, const tchar* defFilename, const UnorderedSymbols& allNeededImports, const UnorderedExports& allSharedExports);
 
 		virtual ~ObjectFile();
 
 	protected:
 		virtual bool Parse(Logger& logger, const tchar* filename) = 0;
 		virtual bool StripExports(Logger& logger, u8* newData, const UnorderedSymbols& allNeededImports, u32& outKeptExportCount) = 0;
-		virtual bool CreateExtraFile(Logger& logger, MemoryBlock& memoryBlock, const UnorderedSymbols& allNeededImports, const UnorderedSymbols& allSharedImports, const UnorderedExports& allSharedExports) = 0;
+		virtual bool CreateExtraFile(Logger& logger, MemoryBlock& memoryBlock, const UnorderedSymbols& allNeededImports, const UnorderedSymbols& allSharedImports, const UnorderedExports& allSharedExports, bool includeExportsInFile) = 0;
+		virtual bool CreateDefFile(Logger& logger, MemoryBlock& memoryBlock, const UnorderedSymbols& allNeededImports, const UnorderedExports& allSharedExports) = 0;
 
 		FileAccessor* m_file = nullptr;
 		u8* m_data = nullptr;
@@ -91,5 +93,11 @@ namespace uba
 		{
 			return std::string(strBegin, strEnd);
 		}
+
+		std::string& ToString(std::string& out) const
+		{
+			out.assign(strBegin, strEnd);
+			return out;
+		}
 	};
 }
diff --git a/Engine/Source/Programs/UnrealBuildAccelerator/Common/Public/UbaObjectFileCoff.h b/Engine/Source/Programs/UnrealBuildAccelerator/Common/Public/UbaObjectFileCoff.h
--- a/Engine/Source/Programs/UnrealBuildAccelerator/Common/Public/UbaObjectFileCoff.h
+++ b/Engine/Source/Programs/UnrealBuildAccelerator/Common/Public/UbaObjectFileCoff.h
@@ -21,7 +21,8 @@ namespace uba
 		template<typename SymbolType> void ParseImports();
 
 		virtual bool StripExports(Logger& logger, u8* newData, const UnorderedSymbols& allNeededImports, u32& outKeptExportCount) override;
-		virtual bool CreateExtraFile(Logger& logger, MemoryBlock& memoryBlock, const UnorderedSymbols& allNeededImports, const UnorderedSymbols& allSharedImports, const UnorderedExports& allSharedExports) override;
+		virtual bool CreateExtraFile(Logger& logger, MemoryBlock& memoryBlock, const UnorderedSymbols& allNeededImports, const UnorderedSymbols& allSharedImports, const UnorderedExports& allSharedExports, bool includeExportsInFile) override;
+		virtual bool CreateDefFile(Logger& logger, MemoryBlock& memoryBlock, const UnorderedSymbols& allNeededImports, const UnorderedExports& allSharedExports) override;
 
 		template<typename SymbolType> void CalculateImports(Logger& logger, Vector<u32>& outImports);
 		template<typename SymbolType> void WriteImports(Logger& logger, u8* newData, Info& newInfo, const Vector<u32>& symbolsToAdd);
diff --git a/Engine/Source/Programs/UnrealBuildAccelerator/Common/Public/UbaObjectFileElf.h b/Engine/Source/Programs/UnrealBuildAccelerator/Common/Public/UbaObjectFileElf.h
--- a/Engine/Source/Programs/UnrealBuildAccelerator/Common/Public/UbaObjectFileElf.h
+++ b/Engine/Source/Programs/UnrealBuildAccelerator/Common/Public/UbaObjectFileElf.h
@@ -15,7 +15,8 @@ namespace uba
 
 	private:
 		virtual bool StripExports(Logger& logger, u8* newData, const UnorderedSymbols& allNeededImports, u32& outKeptExportCount) override;
-		virtual bool CreateExtraFile(Logger& logger, MemoryBlock& memoryBlock, const UnorderedSymbols& allNeededImports, const UnorderedSymbols& allSharedImports, const UnorderedExports& allSharedExports) override;
+		virtual bool CreateExtraFile(Logger& logger, MemoryBlock& memoryBlock, const UnorderedSymbols& allNeededImports, const UnorderedSymbols& allSharedImports, const UnorderedExports& allSharedExports, bool includeExportsInFile) override;
+		virtual bool CreateDefFile(Logger& logger, MemoryBlock& memoryBlock, const UnorderedSymbols& allNeededImports, const UnorderedExports& allSharedExports) override;
 
 		UnorderedSymbols m_toRemove;
 
diff --git a/Engine/Source/Programs/UnrealBuildAccelerator/Common/Public/UbaProcess.h b/Engine/Source/Programs/UnrealBuildAccelerator/Common/Public/UbaProcess.h
--- a/Engine/Source/Programs/UnrealBuildAccelerator/Common/Public/UbaProcess.h
+++ b/Engine/Source/Programs/UnrealBuildAccelerator/Common/Public/UbaProcess.h
@@ -180,6 +180,8 @@ namespace uba
 
 		Thread m_messageThread;
 
+		bool m_extractExports = false;
+
 		ProcessImpl(const ProcessImpl&) = delete;
 		void operator=(const ProcessImpl&) = delete;
 	};
diff --git a/Engine/Source/Programs/UnrealBuildAccelerator/Common/Public/UbaSession.h b/Engine/Source/Programs/UnrealBuildAccelerator/Common/Public/UbaSession.h
--- a/Engine/Source/Programs/UnrealBuildAccelerator/Common/Public/UbaSession.h
+++ b/Engine/Source/Programs/UnrealBuildAccelerator/Common/Public/UbaSession.h
@@ -176,7 +176,6 @@ namespace uba
 		bool m_detailedTrace;
 		bool m_logToFile;
 		bool m_storeObjFilesCompressed;
-		bool m_extractObjFilesSymbols;
 
 		u64 m_keepOutputFileMemoryMapsThreshold;
 
diff --git a/Engine/Source/Programs/UnrealBuildAccelerator/Core/Private/UbaApplicationRules.cpp b/Engine/Source/Programs/UnrealBuildAccelerator/Core/Private/UbaApplicationRules.cpp
--- a/Engine/Source/Programs/UnrealBuildAccelerator/Core/Private/UbaApplicationRules.cpp
+++ b/Engine/Source/Programs/UnrealBuildAccelerator/Core/Private/UbaApplicationRules.cpp
@@ -87,6 +87,11 @@ namespace uba
 		{
 			return EndsWith(fileName, fileNameLen, TC(".obj"));
 		}
+
+		virtual bool ShouldExtractSymbols(const tchar* fileName, u64 fileNameLen) const
+		{
+			return EndsWith(fileName, fileNameLen, TC(".obj"));
+		}
 	};
 
 	class ApplicationRulesLinkExe : public ApplicationRulesVC
@@ -207,6 +212,11 @@ namespace uba
 		{
 			return EndsWith(fileName, fileNameLen, TC(".obj"));
 		}
+
+		virtual bool ShouldExtractSymbols(const tchar* fileName, u64 fileNameLen) const
+		{
+			return EndsWith(fileName, fileNameLen, TC(".obj"));
+		}
 	};
 
 	class ApplicationRulesLdLLdExe : public ApplicationRulesClang
diff --git a/Engine/Source/Programs/UnrealBuildAccelerator/Core/Public/UbaApplicationRules.h b/Engine/Source/Programs/UnrealBuildAccelerator/Core/Public/UbaApplicationRules.h
--- a/Engine/Source/Programs/UnrealBuildAccelerator/Core/Public/UbaApplicationRules.h
+++ b/Engine/Source/Programs/UnrealBuildAccelerator/Core/Public/UbaApplicationRules.h
@@ -134,6 +134,11 @@ namespace uba
 			return false;
 		}
 
+		virtual bool ShouldExtractSymbols(const tchar* fileName, u64 fileNameLen) const
+		{
+			return false;
+		}
+
 		u32 index = ~0u;
 	};
 
diff --git a/Engine/Source/Programs/UnrealBuildAccelerator/ObjTool/Private/UbaObjTool.cpp b/Engine/Source/Programs/UnrealBuildAccelerator/ObjTool/Private/UbaObjTool.cpp
--- a/Engine/Source/Programs/UnrealBuildAccelerator/ObjTool/Private/UbaObjTool.cpp
+++ b/Engine/Source/Programs/UnrealBuildAccelerator/ObjTool/Private/UbaObjTool.cpp
@@ -43,6 +43,7 @@ namespace uba
 		logger.Info(TC("    /S:<objfile>             Obj file to strip. Will produce a .strip.obj file. Multiple allowed"));
 		logger.Info(TC("    /D:<objfile>             Obj file depending on obj files to strip. Multiple allowed"));
 		logger.Info(TC("    /O:<objfile>             Obj file to output containing exports and loopbacks"));
+		logger.Info(TC("    /COMPRESS                Write '/O' file compressed"));
 		logger.Info(TC(""));
 		return -1;
 	}
@@ -58,6 +59,7 @@ namespace uba
 		Vector<TString> objFilesToStrip;
 		Vector<TString> objFilesDependencies;
 		TString extraObjFile;
+		TString defFile;
 
 		auto parseArg = [&](const tchar* arg)
 			{
@@ -86,6 +88,10 @@ namespace uba
 				{
 					extraObjFile = name.data + 3;
 				}
+				else if (name.StartsWith(TC("/E:")))
+				{
+					defFile = name.data + 3;
+				}
 				else if (name.Equals(TC("-printsymbols")))
 				{
 					printSymbols = true;
@@ -288,9 +294,17 @@ namespace uba
 				if (!success)
 					return -1;
 			}
-			else if (!extraObjFile.empty())
+			
+			if (!extraObjFile.empty())
+			{
+				bool includeExportsInFile = defFile.empty();
+				if (!ObjectFile::CreateExtraFile(logger, extraObjFile.c_str(), allNeededImports, allSharedImports, allSharedExports, includeExportsInFile))
+					return -1;
+			}
+
+			if (!defFile.empty())
 			{
-				if (!ObjectFile::CreateExtraFile(logger, extraObjFile.c_str(), allNeededImports, allSharedImports, allSharedExports))
+				if (!ObjectFile::CreateDefFile(logger, defFile.c_str(), allNeededImports, allSharedExports))
 					return -1;
 			}
 
