[UBA]
* Missed rename change in previous changelists

diff --git a/Engine/Source/Programs/UnrealBuildAccelerator/ObjTool/Private/UbaObjTool.cpp b/Engine/Source/Programs/UnrealBuildAccelerator/ObjTool/Private/UbaObjTool.cpp
--- a/Engine/Source/Programs/UnrealBuildAccelerator/ObjTool/Private/UbaObjTool.cpp
+++ b/Engine/Source/Programs/UnrealBuildAccelerator/ObjTool/Private/UbaObjTool.cpp
@@ -199,7 +199,7 @@ namespace uba
 				{
 					const TString& objFileName = *it;
 
-					if (EndsWith(objFileName.c_str(), objFileName.size(), TC(".sym")))
+					if (EndsWith(objFileName.c_str(), objFileName.size(), TC(".exi")))
 					{
 						SymbolFile symbolFile;
 						if (!symbolFile.ParseFile(logger, objFileName.c_str()))
@@ -237,7 +237,7 @@ namespace uba
 				{
 					const TString& objFileName = *it;
 
-					if (EndsWith(objFileName.c_str(), objFileName.size(), TC(".sym")))
+					if (EndsWith(objFileName.c_str(), objFileName.size(), TC(".exi")))
 					{
 						SymbolFile symbolFile;
 						if (!symbolFile.ParseFile(logger, objFileName.c_str()))
@@ -338,7 +338,7 @@ namespace uba
 				const tchar* lastDot = TStrrchr(fileName, '.');
 				UBA_ASSERT(lastDot);
 				StringBuffer<> exportsFile;
-				exportsFile.Append(fileName, lastDot - fileName).Append(TC(".sym"));
+				exportsFile.Append(fileName, lastDot - fileName).Append(TC(".exi"));
 
 				if (!objectFile->WriteSymbols(logger, exportsFile.data))
 					return false;
