UbaObjTool - Removed ,NONAME symbol export flag because it fails the build

diff --git a/Engine/Source/Programs/UnrealBuildAccelerator/Common/Private/UbaObjectFileCoff.cpp b/Engine/Source/Programs/UnrealBuildAccelerator/Common/Private/UbaObjectFileCoff.cpp
--- a/Engine/Source/Programs/UnrealBuildAccelerator/Common/Private/UbaObjectFileCoff.cpp
+++ b/Engine/Source/Programs/UnrealBuildAccelerator/Common/Private/UbaObjectFileCoff.cpp
@@ -540,7 +540,7 @@ namespace uba
 				write(slashExport, sizeof(slashExport) - 1);
 				write(symbol.data(), symbol.size());
 				write(kv.second.data(), kv.second.size());
-				write(",NONAME ", 8);
+				write(" ", 1);
 			}
 			write("", 1);
 			directiveSection.SizeOfRawData = u32(memoryBlock.writtenSize) - directiveRawDataStart;
