[UBA]
* Fixed so ubaobjtool works with compressed obj files
* Removed unused code

diff --git a/Engine/Source/Programs/UnrealBuildAccelerator/Common/Private/UbaObjectFile.cpp b/Engine/Source/Programs/UnrealBuildAccelerator/Common/Private/UbaObjectFile.cpp
--- a/Engine/Source/Programs/UnrealBuildAccelerator/Common/Private/UbaObjectFile.cpp
+++ b/Engine/Source/Programs/UnrealBuildAccelerator/Common/Private/UbaObjectFile.cpp
@@ -166,22 +166,4 @@ namespace uba
 
 		return extraFile.Close();
 	}
-
-	bool ObjectFile::CreateDefFile(Logger& logger, const tchar* defFilename, const UnorderedSymbols& allNeededImports, const UnorderedExports& allSharedExports)
-	{
-		ObjectFileCoff objectFile;
-		MemoryBlock memoryBlock(16*1024*1024);
-
-		if (!((ObjectFile&)objectFile).CreateDefFile(logger, memoryBlock, allNeededImports, allSharedExports))
-			return false;
-
-		FileAccessor defFile(logger, defFilename);
-		if (!defFile.CreateWrite())
-			return false;
-
-		if (!defFile.Write(memoryBlock.memory, memoryBlock.writtenSize))
-			return false;
-
-		return defFile.Close();
-	}
 }
diff --git a/Engine/Source/Programs/UnrealBuildAccelerator/Common/Private/UbaObjectFileCoff.cpp b/Engine/Source/Programs/UnrealBuildAccelerator/Common/Private/UbaObjectFileCoff.cpp
--- a/Engine/Source/Programs/UnrealBuildAccelerator/Common/Private/UbaObjectFileCoff.cpp
+++ b/Engine/Source/Programs/UnrealBuildAccelerator/Common/Private/UbaObjectFileCoff.cpp
@@ -237,78 +237,6 @@ namespace uba
 		return true;
 	}
 
-	bool ObjectFileCoff::CreateStripped(Logger& logger, const tchar* newFilename, const UnorderedSymbols& allNeededImports, u32& outKeptExportCount)
-	{
-		UBA_ASSERT(m_data);
-
-		Vector<u32> symbolsToAdd;
-		if (m_isBigObj)
-			CalculateImports<ImageSymbolEx>(logger, symbolsToAdd);
-		else
-			CalculateImports<ImageSymbol>(logger, symbolsToAdd);
-
-		u32 importsToFixCount = u32(symbolsToAdd.size());
-		u32 bytesToAdd = 0;
-
-		if (importsToFixCount)
-		{
-			bytesToAdd = sizeof(ImageSectionHeader);
-			bytesToAdd += (m_isBigObj ? sizeof(ImageSymbolEx) : sizeof(ImageSymbol)) * importsToFixCount * 2; // Symbols
-			bytesToAdd += (sizeof(ImageRelocation) + 8) * importsToFixCount; // relocations plus memory
-		}
-
-		FileAccessor newFile(logger, newFilename);
-		if (!newFile.CreateMemoryWrite(false, DefaultAttributes(), m_dataSize + bytesToAdd))
-			return false;
-
-		u8* newData = newFile.GetData();
-		Info newInfo = m_info;
-
-		if (importsToFixCount)
-		{
-			if (m_isBigObj)
-				WriteImports<ImageSymbolEx>(logger, newData, newInfo, symbolsToAdd);
-			else
-				WriteImports<ImageSymbol>(logger, newData, newInfo, symbolsToAdd);
-		}
-		else
-		{
-			memcpy(newData, m_data, m_dataSize);
-		}
-
-		if (m_isBigObj)
-			RemoveSymbols<ImageSymbolEx>(logger, newData, newInfo);
-		else
-			RemoveSymbols<ImageSymbol>(logger, newData, newInfo);
-
-		StripExports(logger, newData, allNeededImports, outKeptExportCount);
-
-		return newFile.Close();
-	}
-
-	bool ObjectFileCoff::ComputeLoopbacksAndDuplicates(UnorderedExports& allSharedExports, UnorderedSymbols& duplicates)
-	{
-		std::string importString;
-
-		for (auto& importSymbol : m_imports)
-		{
-			if (strncmp(importSymbol.c_str(), "__imp_", 6) != 0)
-				continue;
-			importString = importSymbol.c_str() + 6;
-			auto findIt = allSharedExports.find(importString);
-			if (findIt == allSharedExports.end())
-				continue;
-			m_loopbacksToAdd.emplace(importString);
-			allSharedExports.erase(findIt);
-		}
-
-		for (auto& dupSymbol : m_potentialDuplicates)
-			if (!duplicates.insert(dupSymbol).second)
-				m_toRemove.insert(dupSymbol);
-
-		return true;
-	}
-
 	bool ObjectFileCoff::ParseExports()
 	{
 		auto sections = (ImageSectionHeader*)(m_data + m_info.sectionsMemOffset);
@@ -600,47 +528,6 @@ namespace uba
 		return true;
 	}
 
-	bool ObjectFileCoff::CreateDefFile(Logger& logger, MemoryBlock& memoryBlock, const UnorderedSymbols& allNeededImports, const UnorderedExports& allSharedExports)
-	{
-		auto allocate = [&](u64 size) { return memoryBlock.Allocate(size, 1, TC("")); };
-		auto write = [&](const void* data, u64 size) { memcpy(allocate(size), data, size); };
-		
-		//write("EXPORTS\n", 8);
-
-		std::string tmp;
-		u32 counter = 0;
-
-		for (auto& kv : allSharedExports)
-		{
-			auto& symbol = kv.first;
-			if (allNeededImports.find(symbol) == allNeededImports.end())
-				if (allNeededImports.find(tmp.assign("__imp_").append(symbol)) == allNeededImports.end())
-					continue;
-
-			char buf[512];
-			u64 len = snprintf(buf, 256, "/EXPORT:_%u,@%s,NONAME\n", ++counter, symbol.data());
-			write(buf, len);
-
-			//write("/EXPORT:@", 9);
-			//write(symbol.data(), symbol.size());
-			//write(kv.second.data(), kv.second.size());
-
-			//char extra[64] = { 0 };
-			//if (!kv.second.empty())
-			//{
-			//	extra[0] = ' ';
-			//	strcpy_s(extra+1, 63, kv.second.c_str() + 1);
-			//}
-
-			//char buf[256];
-			//u64 len = snprintf(buf, 256, " @%u%s\n", ++counter, extra);
-			//write(buf, len);
-			//write("\n", 1);
-		}
-		write("", 0);
-		return true;
-	}
-
 	template<typename SymbolType> void ObjectFileCoff::CalculateImports(Logger& logger, Vector<u32>& outImports)
 	{
 		// Calculate how much more memory we need by counting imports
diff --git a/Engine/Source/Programs/UnrealBuildAccelerator/Common/Private/UbaObjectFileElf.cpp b/Engine/Source/Programs/UnrealBuildAccelerator/Common/Private/UbaObjectFileElf.cpp
--- a/Engine/Source/Programs/UnrealBuildAccelerator/Common/Private/UbaObjectFileElf.cpp
+++ b/Engine/Source/Programs/UnrealBuildAccelerator/Common/Private/UbaObjectFileElf.cpp
@@ -173,15 +173,7 @@ namespace uba
 		}
 		return true;
 	}
-
-	bool ObjectFileElf::ComputeLoopbacksAndDuplicates(UnorderedExports& allSharedExports, UnorderedSymbols& duplicates)
-	{
-		for (auto& dupSymbol : m_potentialDuplicates)
-			if (!duplicates.insert(dupSymbol).second)
-				m_toRemove.insert(dupSymbol);
-		return true;
-	}
-
+	/*
 	bool ObjectFileElf::CreateStripped(Logger& logger, const tchar* newFilename, const UnorderedSymbols& allNeededImports, u32& outKeptExportCount)
 	{
 		FileAccessor newFile(logger, newFilename);
@@ -224,7 +216,7 @@ namespace uba
 
 		return newFile.Close();
 	}
-
+	*/
 	bool ObjectFileElf::StripExports(Logger& logger, u8* newData, const UnorderedSymbols& allNeededImports, u32& outKeptExportCount)
 	{
 		return true;
@@ -234,9 +226,4 @@ namespace uba
 	{
 		return true;
 	}
-
-	bool ObjectFileElf::CreateDefFile(Logger& logger, MemoryBlock& memoryBlock, const UnorderedSymbols& allNeededImports, const UnorderedExports& allSharedExports)
-	{
-		return true;
-	}
 }
diff --git a/Engine/Source/Programs/UnrealBuildAccelerator/Common/Public/UbaObjectFile.h b/Engine/Source/Programs/UnrealBuildAccelerator/Common/Public/UbaObjectFile.h
--- a/Engine/Source/Programs/UnrealBuildAccelerator/Common/Public/UbaObjectFile.h
+++ b/Engine/Source/Programs/UnrealBuildAccelerator/Common/Public/UbaObjectFile.h
@@ -20,8 +20,6 @@ namespace uba
 
 		virtual bool CopyMemoryAndClose();
 		virtual bool StripExports(Logger& logger);
-		virtual bool ComputeLoopbacksAndDuplicates(UnorderedExports& allSharedExports, UnorderedSymbols& duplicates) = 0;
-		virtual bool CreateStripped(Logger& logger, const tchar* newFilename, const UnorderedSymbols& allNeededImports, u32& outKeptExportCount) = 0;
 		virtual bool WriteSymbols(Logger& logger, MemoryBlock& memoryBlock);
 		virtual bool WriteSymbols(Logger& logger, const tchar* exportsFilename);
 
@@ -31,7 +29,6 @@ namespace uba
 		const UnorderedSymbols& GetPotentialDuplicates() const;
 
 		static bool CreateExtraFile(Logger& logger, const tchar* extraObjFilename, const UnorderedSymbols& allNeededImports, const UnorderedSymbols& allSharedImports, const UnorderedExports& allSharedExports, bool includeExportsInFile);
-		static bool CreateDefFile(Logger& logger, const tchar* defFilename, const UnorderedSymbols& allNeededImports, const UnorderedExports& allSharedExports);
 
 		virtual ~ObjectFile();
 
@@ -39,7 +36,6 @@ namespace uba
 		virtual bool Parse(Logger& logger, const tchar* filename) = 0;
 		virtual bool StripExports(Logger& logger, u8* newData, const UnorderedSymbols& allNeededImports, u32& outKeptExportCount) = 0;
 		virtual bool CreateExtraFile(Logger& logger, MemoryBlock& memoryBlock, const UnorderedSymbols& allNeededImports, const UnorderedSymbols& allSharedImports, const UnorderedExports& allSharedExports, bool includeExportsInFile) = 0;
-		virtual bool CreateDefFile(Logger& logger, MemoryBlock& memoryBlock, const UnorderedSymbols& allNeededImports, const UnorderedExports& allSharedExports) = 0;
 
 		FileAccessor* m_file = nullptr;
 		u8* m_data = nullptr;
diff --git a/Engine/Source/Programs/UnrealBuildAccelerator/Common/Public/UbaObjectFileCoff.h b/Engine/Source/Programs/UnrealBuildAccelerator/Common/Public/UbaObjectFileCoff.h
--- a/Engine/Source/Programs/UnrealBuildAccelerator/Common/Public/UbaObjectFileCoff.h
+++ b/Engine/Source/Programs/UnrealBuildAccelerator/Common/Public/UbaObjectFileCoff.h
@@ -11,8 +11,6 @@ namespace uba
 	{
 	public:
 		virtual bool Parse(Logger& logger, const tchar* filename) override;
-		virtual bool ComputeLoopbacksAndDuplicates(UnorderedExports& allSharedExports, UnorderedSymbols& duplicates) override;
-		virtual bool CreateStripped(Logger& logger, const tchar* newFilename, const UnorderedSymbols& allNeededImports, u32& outKeptExportCount) override;
 
 	private:
 		struct Info;
@@ -22,7 +20,6 @@ namespace uba
 
 		virtual bool StripExports(Logger& logger, u8* newData, const UnorderedSymbols& allNeededImports, u32& outKeptExportCount) override;
 		virtual bool CreateExtraFile(Logger& logger, MemoryBlock& memoryBlock, const UnorderedSymbols& allNeededImports, const UnorderedSymbols& allSharedImports, const UnorderedExports& allSharedExports, bool includeExportsInFile) override;
-		virtual bool CreateDefFile(Logger& logger, MemoryBlock& memoryBlock, const UnorderedSymbols& allNeededImports, const UnorderedExports& allSharedExports) override;
 
 		template<typename SymbolType> void CalculateImports(Logger& logger, Vector<u32>& outImports);
 		template<typename SymbolType> void WriteImports(Logger& logger, u8* newData, Info& newInfo, const Vector<u32>& symbolsToAdd);
diff --git a/Engine/Source/Programs/UnrealBuildAccelerator/Common/Public/UbaObjectFileElf.h b/Engine/Source/Programs/UnrealBuildAccelerator/Common/Public/UbaObjectFileElf.h
--- a/Engine/Source/Programs/UnrealBuildAccelerator/Common/Public/UbaObjectFileElf.h
+++ b/Engine/Source/Programs/UnrealBuildAccelerator/Common/Public/UbaObjectFileElf.h
@@ -10,13 +10,10 @@ namespace uba
 	{
 	public:
 		virtual bool Parse(Logger& logger, const tchar* filename) override;
-		virtual bool ComputeLoopbacksAndDuplicates(UnorderedExports& allSharedExports, UnorderedSymbols& duplicates) override;
-		virtual bool CreateStripped(Logger& logger, const tchar* newFilename, const UnorderedSymbols& allNeededImports, u32& outKeptExportCount) override;
 
 	private:
 		virtual bool StripExports(Logger& logger, u8* newData, const UnorderedSymbols& allNeededImports, u32& outKeptExportCount) override;
 		virtual bool CreateExtraFile(Logger& logger, MemoryBlock& memoryBlock, const UnorderedSymbols& allNeededImports, const UnorderedSymbols& allSharedImports, const UnorderedExports& allSharedExports, bool includeExportsInFile) override;
-		virtual bool CreateDefFile(Logger& logger, MemoryBlock& memoryBlock, const UnorderedSymbols& allNeededImports, const UnorderedExports& allSharedExports) override;
 
 		UnorderedSymbols m_toRemove;
 
diff --git a/Engine/Source/Programs/UnrealBuildAccelerator/Core/Private/UbaApplicationRules.cpp b/Engine/Source/Programs/UnrealBuildAccelerator/Core/Private/UbaApplicationRules.cpp
--- a/Engine/Source/Programs/UnrealBuildAccelerator/Core/Private/UbaApplicationRules.cpp
+++ b/Engine/Source/Programs/UnrealBuildAccelerator/Core/Private/UbaApplicationRules.cpp
@@ -412,6 +412,19 @@ namespace uba
 		}
 	};
 
+	class ApplicationRulesUbaObjTool : public ApplicationRules
+	{
+		virtual bool IsOutputFile(const tchar* file, u64 fileLen) const override
+		{
+			return EndsWith(file, fileLen, TC(".obj"));
+		}
+
+		virtual bool StoreFileCompressed(const tchar* fileName, u64 fileNameLen) const
+		{
+			return EndsWith(fileName, fileNameLen, TC(".obj"));
+		}
+	};
+
 	const RulesRec* GetApplicationRules()
 	{
 		// TODO: Add support for data driven rules.
@@ -441,6 +454,7 @@ namespace uba
 			{ TC("UnrealBuildTool.dll"),		new ApplicationRulesUBTDll() },
 			{ TC("PVS-Studio.exe"),				new ApplicationRulesPVSStudio() },
 			{ TC("ShaderCompileWorker.exe"),	new ApplicationRulesShaderCompileWorker() },
+			{ TC("UbaObjTool.exe"),				new ApplicationRulesUbaObjTool() },
 			//{ L"MSBuild.dll"),				new ApplicationRules() },
 			//{ L"BreakpadSymbolEncoder.exe"),	new ApplicationRulesClang() },
 			//{ L"cmd.exe"),		new ApplicationRules() },
diff --git a/Engine/Source/Programs/UnrealBuildAccelerator/Detours/Private/Windows/UbaDetoursFunctionsUcrtBase.inl b/Engine/Source/Programs/UnrealBuildAccelerator/Detours/Private/Windows/UbaDetoursFunctionsUcrtBase.inl
--- a/Engine/Source/Programs/UnrealBuildAccelerator/Detours/Private/Windows/UbaDetoursFunctionsUcrtBase.inl
+++ b/Engine/Source/Programs/UnrealBuildAccelerator/Detours/Private/Windows/UbaDetoursFunctionsUcrtBase.inl
@@ -341,7 +341,7 @@ errno_t Detoured__wsplitpath_s(const wchar_t* path, wchar_t* drive, size_t drive
 
 	g_rules->RepairMalformedLibPath(path);
 	auto res = True__wsplitpath_s(path, drive, driveNumberOfElements, dir, dirNumberOfElements, fname, nameNumberOfElements, ext, extNumberOfElements);
-	DEBUG_LOG_TRUE(L"_wsplitpath_s", L"%ls %ls %ls %ls %ls", path, drive, dir, fname, ext);
+	//DEBUG_LOG_TRUE(L"_wsplitpath_s", L"%ls %ls %ls %ls %ls", path, drive, dir, fname, ext);
 	return res;
 }
 
diff --git a/Engine/Source/Programs/UnrealBuildAccelerator/ObjTool/Private/UbaObjTool.cpp b/Engine/Source/Programs/UnrealBuildAccelerator/ObjTool/Private/UbaObjTool.cpp
--- a/Engine/Source/Programs/UnrealBuildAccelerator/ObjTool/Private/UbaObjTool.cpp
+++ b/Engine/Source/Programs/UnrealBuildAccelerator/ObjTool/Private/UbaObjTool.cpp
@@ -59,7 +59,6 @@ namespace uba
 		Vector<TString> objFilesToStrip;
 		Vector<TString> objFilesDependencies;
 		TString extraObjFile;
-		TString defFile;
 
 		auto parseArg = [&](const tchar* arg)
 			{
@@ -88,10 +87,6 @@ namespace uba
 				{
 					extraObjFile = name.data + 3;
 				}
-				else if (name.StartsWith(TC("/E:")))
-				{
-					defFile = name.data + 3;
-				}
 				else if (name.Equals(TC("-printsymbols")))
 				{
 					printSymbols = true;
@@ -197,32 +192,16 @@ namespace uba
 			WorkManagerImpl workManager(workerCount);
 			workManager.ParallelFor(workerCount, objFilesDependencies, [&](auto& it)
 				{
-					const TString& objFileName = *it;
+					const TString& exiFilename = *it;
 
-					if (EndsWith(objFileName.c_str(), objFileName.size(), TC(".exi")))
+					SymbolFile symbolFile;
+					if (!symbolFile.ParseFile(logger, exiFilename.c_str()))
 					{
-						SymbolFile symbolFile;
-						if (!symbolFile.ParseFile(logger, objFileName.c_str()))
-						{
-							success = false;
-							return;
-						}
-						ScopedCriticalSection _(cs);
-						allNeededImports.insert(symbolFile.imports.begin(), symbolFile.imports.end());
-					}
-					else
-					{
-						ObjectFile* objectFile = ObjectFile::OpenAndParse(logger, objFileName.c_str());
-						if (!objectFile)
-						{
-							success = false;
-							return;
-						}
-						auto g = MakeGuard([&]() { delete objectFile; });
-
-						ScopedCriticalSection _(cs);
-						allNeededImports.insert(objectFile->GetImports().begin(), objectFile->GetImports().end());
+						success = false;
+						return;
 					}
+					ScopedCriticalSection _(cs);
+					allNeededImports.insert(symbolFile.imports.begin(), symbolFile.imports.end());
 				});
 			if (!success)
 				return -1;
@@ -235,78 +214,23 @@ namespace uba
 
 			workManager.ParallelFor(workerCount, objFilesToStrip, [&](auto& it)
 				{
-					const TString& objFileName = *it;
-
-					if (EndsWith(objFileName.c_str(), objFileName.size(), TC(".exi")))
+					const TString& exiFilename = *it;
+					SymbolFile symbolFile;
+					if (!symbolFile.ParseFile(logger, exiFilename.c_str()))
 					{
-						SymbolFile symbolFile;
-						if (!symbolFile.ParseFile(logger, objFileName.c_str()))
-						{
-							success = false;
-							return;
-						}
-						ScopedCriticalSection _(cs);
-						allSharedImports.insert(symbolFile.imports.begin(), symbolFile.imports.end());
-						allSharedExports.insert(symbolFile.exports.begin(), symbolFile.exports.end());
-					}
-					else
-					{
-						ObjectFile* objectFile = ObjectFile::OpenAndParse(logger, objFileName.c_str());
-						if (!objectFile)
-						{
-							success = false;
-							return;
-						}
-
-						ScopedCriticalSection _(cs);
-						objectFiles.try_emplace(objFileName, objectFile);
-						allSharedExports.insert(objectFile->GetExports().begin(), objectFile->GetExports().end());
+						success = false;
+						return;
 					}
+					ScopedCriticalSection _(cs);
+					allSharedImports.insert(symbolFile.imports.begin(), symbolFile.imports.end());
+					allSharedExports.insert(symbolFile.exports.begin(), symbolFile.exports.end());
 				});
 			if (!success)
 				return -1;
 
-			// Figure out which loopback symbols that should be added to which obj file
-			if (!objectFiles.empty())
-			{
-				UnorderedSymbols duplicates;
-				for (auto& objFileName : objFilesToStrip)
-					if (!objectFiles[objFileName]->ComputeLoopbacksAndDuplicates(allSharedExports, duplicates))
-						return -1;
-
-				Atomic<u32> totalExportCount;
-				Atomic<u32> totalKeptExportCount;
-				workManager.ParallelFor(workerCount, objectFiles, [&](auto& it)
-					{
-						ObjectFile& file = *it->second;
-						const tchar* fileName = file.GetFileName();
-						const tchar* lastDot = TStrrchr(fileName, '.');
-						UBA_ASSERT(lastDot);
-						StringBuffer<> newFilename;
-						newFilename.Append(fileName, lastDot - fileName).Append(TC(".strip")).Append(lastDot);
-						u32 keptExportCount = 0;
-						if (!file.CreateStripped(logger, newFilename.data, allNeededImports, keptExportCount))
-							success = false;
-
-						totalExportCount += u32(file.GetExports().size());
-						totalKeptExportCount += keptExportCount;
-					});
-				if (!success)
-					return -1;
-			}
-			
 			if (!extraObjFile.empty())
-			{
-				bool includeExportsInFile = defFile.empty();
-				if (!ObjectFile::CreateExtraFile(logger, extraObjFile.c_str(), allNeededImports, allSharedImports, allSharedExports, includeExportsInFile))
+				if (!ObjectFile::CreateExtraFile(logger, extraObjFile.c_str(), allNeededImports, allSharedImports, allSharedExports, true))
 					return -1;
-			}
-
-			if (!defFile.empty())
-			{
-				if (!ObjectFile::CreateDefFile(logger, defFile.c_str(), allNeededImports, allSharedExports))
-					return -1;
-			}
 
 			//logger.Info(TC("Reduced export count from %llu to %llu"), totalExportCount.load(), totalKeptExportCount.size());
 		}
@@ -343,12 +267,11 @@ namespace uba
 				if (!objectFile->WriteSymbols(logger, exportsFile.data))
 					return false;
 
-				u32 keptExportCount = 0;
-
-				StringBuffer<> newFilename;
-				newFilename.Append(fileName, lastDot - fileName).Append(TC(".TEST")).Append(lastDot);
-				if (!objectFile->CreateStripped(logger, newFilename.data, {}, keptExportCount))
-					return false;
+				//u32 keptExportCount = 0;
+				//StringBuffer<> newFilename;
+				//newFilename.Append(fileName, lastDot - fileName).Append(TC(".TEST")).Append(lastDot);
+				//if (!objectFile->CreateStripped(logger, newFilename.data, {}, keptExportCount))
+				//	return false;
 			}
 		}
 		return 0;
