[UBA]
* Added so ispc compiles also creates .exi files when compiling with merged modules.

diff --git a/Engine/Source/Programs/UnrealBuildAccelerator/Core/Private/UbaApplicationRules.cpp b/Engine/Source/Programs/UnrealBuildAccelerator/Core/Private/UbaApplicationRules.cpp
--- a/Engine/Source/Programs/UnrealBuildAccelerator/Core/Private/UbaApplicationRules.cpp
+++ b/Engine/Source/Programs/UnrealBuildAccelerator/Core/Private/UbaApplicationRules.cpp
@@ -359,6 +359,16 @@ namespace uba
 		{
 			return true;
 		}
+
+		virtual bool StoreFileCompressed(const tchar* fileName, u64 fileNameLen) const
+		{
+			return EndsWith(fileName, fileNameLen, TC(".obj"));
+		}
+
+		virtual bool ShouldExtractSymbols(const tchar* fileName, u64 fileNameLen) const
+		{
+			return EndsWith(fileName, fileNameLen, TC(".obj"));
+		}
 	};
 
 	class ApplicationRulesUBTDll : public ApplicationRules
