[UBA] Add vertical cursor line to UbaVisualizer and zoom in/out at cursor position as anchor point.

- Zoom in/out with the horizontal cursor position as anchor point; Previously the anchor point was always the left corner of the window.
- Add a vertical line with the same color as the info text to visualize where the cursor is and where zooming in/out will take place (can be disabled).

#rnx

diff --git a/Engine/Source/Programs/UnrealBuildAccelerator/Visualizer/Private/UbaVisualizer.cpp b/Engine/Source/Programs/UnrealBuildAccelerator/Visualizer/Private/UbaVisualizer.cpp
--- a/Engine/Source/Programs/UnrealBuildAccelerator/Visualizer/Private/UbaVisualizer.cpp
+++ b/Engine/Source/Programs/UnrealBuildAccelerator/Visualizer/Private/UbaVisualizer.cpp
@@ -1185,6 +1185,8 @@ namespace uba
 
 		if (m_visibleComponents[ComponentType_Timeline])
 			PaintTimeline(hdc, clientRect);
+		if (m_visibleComponents[ComponentType_Cursor])
+			PaintCursor(hdc, clientRect);
 
 		{
 			int top = 5;
@@ -1458,7 +1460,9 @@ namespace uba
 				L"timeline",
 				L"detailed data (use -BoxDetailedTrace for even more)",
 				L"workers (threads on host taking care of requests from helpers)",
+				L"cursor (vertical line)",
 			};
+			static_assert(sizeof(tooltip)/sizeof(tooltip[0]) == ComponentType_Count);
 
 			bool vis = m_visibleComponents[m_buttonSelected];
 			SelectObject(hdc, m_popupFont);
@@ -1724,6 +1728,20 @@ namespace uba
 		*/
 	}
 
+	void Visualizer::PaintCursor(HDC hdc, const RECT& clientRect)
+	{
+		// Get cursor position and transform into window coordinates
+		POINT cursorPos = {};
+		GetCursorPos(&cursorPos);
+		ScreenToClient(m_hwnd, &cursorPos);
+
+		// Draw cursor as vertical line
+		SelectObject(hdc, m_textPen);
+		const int cursorPosHorizontal = cursorPos.x;
+		MoveToEx(hdc, cursorPosHorizontal, clientRect.top, NULL);
+		LineTo(hdc, cursorPosHorizontal, clientRect.bottom);
+	}
+
 	void Visualizer::PaintDetailedStats(int& posY, const RECT& progressRect, TraceView::Session& session, bool isRemote, u64 playTime, const DrawTextFunc& drawTextFunc)
 	{
 		int stepY = FontHeight;
@@ -2379,19 +2397,25 @@ namespace uba
 			RECT r;
 			GetClientRect(hWnd, &r);
 
+			// Use mouse cursor as scroll anchor point
+			POINT cursorPos = {};
+			GetCursorPos(&cursorPos);
+			ScreenToClient(m_hwnd, &cursorPos);
+			const float scrollAnchorOffsetX = float(cursorPos.x);
+
 			SHORT controlState = GetAsyncKeyState(VK_CONTROL);
 			if (controlState & (1<<15))
 			{
 				float newValue = Max(m_zoomValue + float(delta)*0.0005f, 0.05f);
 				m_scrollPosY = Min(0.0f, float(m_scrollPosY)*newValue/m_zoomValue);//LOWORD(lParam);
-				m_scrollPosX = Min(0.0f, float(m_scrollPosX)*newValue/m_zoomValue);//LOWORD(lParam);
+				m_scrollPosX = Min(0.0f, float(m_scrollPosX - scrollAnchorOffsetX)*newValue/m_zoomValue + scrollAnchorOffsetX);//LOWORD(lParam);
 
 				m_zoomValue = newValue;
 			}
 			else
 			{
 				float newValue = m_horizontalScaleValue + m_horizontalScaleValue*float(delta)*0.0006f;
-				m_scrollPosX = Min(0.0f, float(m_scrollPosX)*newValue/m_horizontalScaleValue);//LOWORD(lParam);
+				m_scrollPosX = Min(0.0f, float(m_scrollPosX - scrollAnchorOffsetX)*newValue/m_horizontalScaleValue + scrollAnchorOffsetX);//LOWORD(lParam);
 				m_horizontalScaleValue = newValue;
 			}
 
@@ -2445,7 +2469,7 @@ namespace uba
 			}
 			else
 			{
-				if (UpdateSelection())
+				if (UpdateSelection() || m_visibleComponents[ComponentType_Cursor])
 					RedrawWindow(hWnd, NULL, NULL, RDW_INVALIDATE);
 				/*
 				else
diff --git a/Engine/Source/Programs/UnrealBuildAccelerator/Visualizer/Private/UbaVisualizer.h b/Engine/Source/Programs/UnrealBuildAccelerator/Visualizer/Private/UbaVisualizer.h
--- a/Engine/Source/Programs/UnrealBuildAccelerator/Visualizer/Private/UbaVisualizer.h
+++ b/Engine/Source/Programs/UnrealBuildAccelerator/Visualizer/Private/UbaVisualizer.h
@@ -31,6 +31,7 @@ namespace uba
 		void PaintAll(HDC hdc, const RECT& clientRect);
 		void PaintProcessRect(TraceView::Process& process, HDC hdc, RECT rect, const RECT& progressRect, bool selected, bool writingBitmap);
 		void PaintTimeline(HDC hdc, const RECT& clientRect);
+		void PaintCursor(HDC hdc, const RECT& clientRect);
 		using DrawTextFunc = Function<void(const StringBufferBase& text, RECT& rect)>;
 		void PaintDetailedStats(int& posY, const RECT& progressRect, TraceView::Session& session, bool isRemote, u64 playTime, const DrawTextFunc& drawTextFunc);
 
@@ -165,6 +166,7 @@ namespace uba
 			ComponentType_Timeline,
 			ComponentType_DetailedData,
 			ComponentType_Workers,
+			ComponentType_Cursor,
 			ComponentType_Count
 		};
 		bool m_visibleComponents[ComponentType_Count];
