[UBA]
* Made detoured GetLongPathNameW more resilient to bad input

diff --git a/Engine/Source/Programs/UnrealBuildAccelerator/Detours/Private/Windows/UbaDetoursFunctionsKernelBase.inl b/Engine/Source/Programs/UnrealBuildAccelerator/Detours/Private/Windows/UbaDetoursFunctionsKernelBase.inl
--- a/Engine/Source/Programs/UnrealBuildAccelerator/Detours/Private/Windows/UbaDetoursFunctionsKernelBase.inl
+++ b/Engine/Source/Programs/UnrealBuildAccelerator/Detours/Private/Windows/UbaDetoursFunctionsKernelBase.inl
@@ -985,18 +985,26 @@ DWORD Detoured_GetLongPathNameW(LPCWSTR lpszShortPath, LPWSTR lpszLongPath, DWOR
 {
 	DETOURED_CALL(GetLongPathNameW);
 
-	if (wcsncmp(lpszShortPath, L"\\\\?\\", 4) == 0)
-		lpszShortPath += 4;
+	if (!lpszShortPath)
+		return Local_GetLongPathNameW(lpszShortPath, lpszLongPath, cchBuffer);
+
+	const wchar_t* path = lpszShortPath;
+	if (wcsncmp(path, L"\\\\?\\", 4) == 0)
+		path += 4;
+
+	bool foundQuestionMark = false;
+	for (const wchar_t* i = path, *e = i + 4; *i && i!=e; ++i)
+		foundQuestionMark |= *i == '?';
 
 	// TODO: Add support for ~ and "\\?\"
-	if (!wcschr(lpszShortPath, '?'))
+	if (!foundQuestionMark)
 	{
 		StringBuffer<> fixedName;
-		FixPath(fixedName, lpszShortPath);
+		FixPath(fixedName, path);
 
-		DEBUG_LOG_DETOURED(L"GetLongPathNameW", L"(%ls)", lpszShortPath);
+		DEBUG_LOG_DETOURED(L"GetLongPathNameW", L"(%ls)", path);
 		WIN32_FILE_ATTRIBUTE_DATA data;
-		bool success = Shared_GetFileAttributesExW(fixedName.data, GetFileExInfoStandard, &data, lpszShortPath);
+		bool success = Shared_GetFileAttributesExW(fixedName.data, GetFileExInfoStandard, &data, path);
 
 		DWORD res = 0;
 		if (success)
@@ -1008,12 +1016,12 @@ DWORD Detoured_GetLongPathNameW(LPCWSTR lpszShortPath, LPWSTR lpszLongPath, DWOR
 #if UBA_DEBUG_VALIDATE
 		if (g_validateFileAccess)
 		{
-			if (!wcschr(lpszShortPath, '~') && !wcschr(lpszShortPath, '?'))
+			if (!wcschr(path, '~') && !wcschr(path, '?'))
 			{
 				wchar_t temp[MaxPath];
 				UBA_ASSERT(cchBuffer <= sizeof_array(temp));
 				SuppressDetourScope _;
-				DWORD res2 = True_GetLongPathNameW(lpszShortPath, temp, cchBuffer); (void)res2;
+				DWORD res2 = True_GetLongPathNameW(path, temp, cchBuffer); (void)res2;
 				UBA_ASSERT(res == res2);
 			}
 		}
