[UBA]
* Added exclusions for exports in UbObjTool and added D3D12SDKVersion and friends to exclusion list
* Fixed so UbaObjTool works for linux (or builds at least)
* Added version to importexport files

diff --git a/Engine/Source/Programs/UnrealBuildAccelerator/Common/Private/UbaObjectFile.cpp b/Engine/Source/Programs/UnrealBuildAccelerator/Common/Private/UbaObjectFile.cpp
--- a/Engine/Source/Programs/UnrealBuildAccelerator/Common/Private/UbaObjectFile.cpp
+++ b/Engine/Source/Programs/UnrealBuildAccelerator/Common/Private/UbaObjectFile.cpp
@@ -8,6 +8,8 @@
 
 namespace uba
 {
+	u8 SymbolFileVersion = 1;
+
 	ObjectFile* ObjectFile::OpenAndParse(Logger& logger, const tchar* filename)
 	{
 		auto file = new FileAccessor(logger, filename);
@@ -66,10 +68,13 @@ namespace uba
 		return true;
 	}
 
-	bool ObjectFile::WriteSymbols(Logger& logger, MemoryBlock& memoryBlock)
+	bool ObjectFile::WriteImportsAndExports(Logger& logger, MemoryBlock& memoryBlock)
 	{
 		auto write = [&](const void* data, u64 dataSize) { memcpy(memoryBlock.Allocate(dataSize, 1, TC("")), data, dataSize); };
 
+		write(&SymbolFileVersion, 1);
+		write(&m_type, 1);
+
 		// Write all imports
 		for (auto& symbol : m_imports)
 		{
@@ -89,7 +94,7 @@ namespace uba
 		return true;
 	}
 
-	bool ObjectFile::WriteSymbols(Logger& logger, const tchar* exportsFilename)
+	bool ObjectFile::WriteImportsAndExports(Logger& logger, const tchar* exportsFilename)
 	{
 		FileAccessor exportsFile(logger, exportsFilename);
 		if (!exportsFile.CreateWrite())
@@ -149,12 +154,17 @@ namespace uba
 		return m_potentialDuplicates;
 	}
 
-	bool ObjectFile::CreateExtraFile(Logger& logger, const tchar* extraObjFilename, const UnorderedSymbols& allNeededImports, const UnorderedSymbols& allSharedImports, const UnorderedExports& allSharedExports, bool includeExportsInFile)
+	bool ObjectFile::CreateExtraFile(Logger& logger, const tchar* extraObjFilename, ObjectFileType type, const UnorderedSymbols& allNeededImports, const UnorderedSymbols& allSharedImports, const UnorderedExports& allSharedExports, bool includeExportsInFile)
 	{
-		ObjectFileCoff objectFile;
+		ObjectFileCoff objectFileCoff;
+		ObjectFileElf objectFileElf;
+		
+		ObjectFile& objectFile = type == ObjectFileType_Coff ? (ObjectFile&)objectFileCoff : (ObjectFile&)objectFileElf;
+
+
 		MemoryBlock memoryBlock(16*1024*1024);
 
-		if (!((ObjectFile&)objectFile).CreateExtraFile(logger, memoryBlock, allNeededImports, allSharedImports, allSharedExports, includeExportsInFile))
+		if (!objectFile.CreateExtraFile(logger, memoryBlock, allNeededImports, allSharedImports, allSharedExports, includeExportsInFile))
 			return false;
 
 		FileAccessor extraFile(logger, extraObjFilename);
@@ -166,4 +176,40 @@ namespace uba
 
 		return extraFile.Close();
 	}
+
+	bool SymbolFile::ParseFile(Logger& logger, const tchar* filename)
+	{
+		FileAccessor symFile(logger, filename);
+		if (!symFile.OpenMemoryRead())
+			return false;
+		auto readPos = (const char*)symFile.GetData();
+
+		u8 version = *(u8*)readPos++;
+		if (SymbolFileVersion != version)
+			return logger.Error(TC("%s - Import/export file version mismatch"), filename);
+
+		type = *(const ObjectFileType*)readPos++;
+
+		while (*readPos)
+		{
+			auto strEnd = strlen(readPos);
+			imports.insert(std::string(readPos, readPos + strEnd));
+			readPos = readPos + strEnd + 1;
+		}
+		++readPos;
+
+		while (*readPos)
+		{
+			auto strEnd = strlen(readPos);
+			std::string extra;
+			if (const char* comma = strchr(readPos, ','))
+			{
+				strEnd = comma - readPos;
+				extra = comma;
+			}
+			exports.emplace(std::string(readPos, readPos + strEnd), extra);
+			readPos = readPos + strEnd + 1;
+		}
+		return true;
+	}
 }
diff --git a/Engine/Source/Programs/UnrealBuildAccelerator/Common/Private/UbaObjectFileCoff.cpp b/Engine/Source/Programs/UnrealBuildAccelerator/Common/Private/UbaObjectFileCoff.cpp
--- a/Engine/Source/Programs/UnrealBuildAccelerator/Common/Private/UbaObjectFileCoff.cpp
+++ b/Engine/Source/Programs/UnrealBuildAccelerator/Common/Private/UbaObjectFileCoff.cpp
@@ -203,6 +203,11 @@ namespace uba
 		return true;
 	}
 
+	ObjectFileCoff::ObjectFileCoff()
+	{
+		m_type = ObjectFileType_Coff;
+	}
+
 	bool ObjectFileCoff::Parse(Logger& logger, const tchar* filename)
 	{
 		m_isBigObj = IsBigObj(m_data, m_dataSize);
@@ -702,5 +707,4 @@ namespace uba
 				memset(&section, 0, sizeof(ImageSectionHeader));
 		}
 	}
-
 }
diff --git a/Engine/Source/Programs/UnrealBuildAccelerator/Common/Private/UbaObjectFileElf.cpp b/Engine/Source/Programs/UnrealBuildAccelerator/Common/Private/UbaObjectFileElf.cpp
--- a/Engine/Source/Programs/UnrealBuildAccelerator/Common/Private/UbaObjectFileElf.cpp
+++ b/Engine/Source/Programs/UnrealBuildAccelerator/Common/Private/UbaObjectFileElf.cpp
@@ -47,6 +47,7 @@ namespace uba
 	  u64	st_size;		/* Symbol size */
 	};
 
+	#define EM_X86_64	62	/* AMD x86-64 architecture */
 
 	#define SHT_SYMTAB	  2		/* Symbol table */
 	#define SHT_DYNSYM	  11		/* Dynamic linker symbol table */
@@ -111,6 +112,11 @@ namespace uba
 #endif
 	};
 
+	ObjectFileElf::ObjectFileElf()
+	{
+		m_type = ObjectFileType_Elf;
+	}
+
 
 	bool ObjectFileElf::Parse(Logger& logger, const tchar* filename)
 	{
@@ -224,6 +230,18 @@ namespace uba
 
 	bool ObjectFileElf::CreateExtraFile(Logger& logger, MemoryBlock& memoryBlock, const UnorderedSymbols& allNeededImports, const UnorderedSymbols& allSharedImports, const UnorderedExports& allSharedExports, bool includeExportsInFile)
 	{
+		auto& header = *(Elf64Header*)memoryBlock.Allocate(sizeof(Elf64Header), 1, TC(""));
+
+		header.e_ident[0] = 0x7f;
+		header.e_ident[1] = 'E';
+		header.e_ident[2] = 'L';
+		header.e_ident[3] = 'F';
+		header.e_ident[4] = 2;
+		header.e_ident[5] = 1;
+		header.e_ident[6] = 1;
+		header.e_type = 1;
+		header.e_machine = EM_X86_64;
+		header.e_ehsize = sizeof(Elf64Header);
 		return true;
 	}
 }
diff --git a/Engine/Source/Programs/UnrealBuildAccelerator/Common/Private/UbaSession.cpp b/Engine/Source/Programs/UnrealBuildAccelerator/Common/Private/UbaSession.cpp
--- a/Engine/Source/Programs/UnrealBuildAccelerator/Common/Private/UbaSession.cpp
+++ b/Engine/Source/Programs/UnrealBuildAccelerator/Common/Private/UbaSession.cpp
@@ -2647,7 +2647,7 @@ namespace uba
 		exportsFile.Append(fileName, lastDot - fileName).Append(TC(".exi"));
 
 		MemoryBlock memoryBlock(1*1024*1024);
-		if (!objectFile->WriteSymbols(m_logger, memoryBlock))
+		if (!objectFile->WriteImportsAndExports(m_logger, memoryBlock))
 			return false;
 
 		FileMappingHandle symHandle = CreateMemoryMappingW(m_logger, PAGE_READWRITE, memoryBlock.writtenSize);
diff --git a/Engine/Source/Programs/UnrealBuildAccelerator/Common/Public/UbaObjectFile.h b/Engine/Source/Programs/UnrealBuildAccelerator/Common/Public/UbaObjectFile.h
--- a/Engine/Source/Programs/UnrealBuildAccelerator/Common/Public/UbaObjectFile.h
+++ b/Engine/Source/Programs/UnrealBuildAccelerator/Common/Public/UbaObjectFile.h
@@ -12,6 +12,13 @@ namespace uba
 	using UnorderedSymbols = UnorderedSet<std::string>;
 	using UnorderedExports = UnorderedMap<std::string, std::string>;
 
+	enum ObjectFileType : u8
+	{
+		ObjectFileType_Unknown,
+		ObjectFileType_Coff,
+		ObjectFileType_Elf,
+	};
+
 	class ObjectFile
 	{
 	public:
@@ -20,15 +27,15 @@ namespace uba
 
 		virtual bool CopyMemoryAndClose();
 		virtual bool StripExports(Logger& logger);
-		virtual bool WriteSymbols(Logger& logger, MemoryBlock& memoryBlock);
-		virtual bool WriteSymbols(Logger& logger, const tchar* exportsFilename);
+		virtual bool WriteImportsAndExports(Logger& logger, MemoryBlock& memoryBlock);
+		virtual bool WriteImportsAndExports(Logger& logger, const tchar* exportsFilename);
 
 		const tchar* GetFileName() const;
 		const UnorderedSymbols& GetImports() const;
 		const UnorderedExports& GetExports() const;
 		const UnorderedSymbols& GetPotentialDuplicates() const;
 
-		static bool CreateExtraFile(Logger& logger, const tchar* extraObjFilename, const UnorderedSymbols& allNeededImports, const UnorderedSymbols& allSharedImports, const UnorderedExports& allSharedExports, bool includeExportsInFile);
+		static bool CreateExtraFile(Logger& logger, const tchar* extraObjFilename, ObjectFileType type, const UnorderedSymbols& allNeededImports, const UnorderedSymbols& allSharedImports, const UnorderedExports& allSharedExports, bool includeExportsInFile);
 
 		virtual ~ObjectFile();
 
@@ -42,6 +49,7 @@ namespace uba
 		u64 m_dataSize = 0;
 		bool m_ownsData = false;
 
+		ObjectFileType m_type;
 		UnorderedSymbols m_imports;
 		UnorderedExports m_exports;
 		UnorderedSymbols m_potentialDuplicates;
@@ -96,4 +104,13 @@ namespace uba
 			return out;
 		}
 	};
+
+	struct SymbolFile
+	{
+		UnorderedSymbols imports;
+		UnorderedExports exports;
+		ObjectFileType type;
+
+		bool ParseFile(Logger& logger, const tchar* filename);
+	};
 }
diff --git a/Engine/Source/Programs/UnrealBuildAccelerator/Common/Public/UbaObjectFileCoff.h b/Engine/Source/Programs/UnrealBuildAccelerator/Common/Public/UbaObjectFileCoff.h
--- a/Engine/Source/Programs/UnrealBuildAccelerator/Common/Public/UbaObjectFileCoff.h
+++ b/Engine/Source/Programs/UnrealBuildAccelerator/Common/Public/UbaObjectFileCoff.h
@@ -10,6 +10,7 @@ namespace uba
 	class ObjectFileCoff : public ObjectFile
 	{
 	public:
+		ObjectFileCoff();
 		virtual bool Parse(Logger& logger, const tchar* filename) override;
 
 	private:
@@ -42,5 +43,6 @@ namespace uba
 		UnorderedSymbols m_toRemove;
 
 		static UnorderedSymbols PotentiallyDuplicatedSymbols;
+		static UnorderedSymbols ExportsToKeep;
 	};
 }
diff --git a/Engine/Source/Programs/UnrealBuildAccelerator/Common/Public/UbaObjectFileElf.h b/Engine/Source/Programs/UnrealBuildAccelerator/Common/Public/UbaObjectFileElf.h
--- a/Engine/Source/Programs/UnrealBuildAccelerator/Common/Public/UbaObjectFileElf.h
+++ b/Engine/Source/Programs/UnrealBuildAccelerator/Common/Public/UbaObjectFileElf.h
@@ -9,6 +9,7 @@ namespace uba
 	class ObjectFileElf : public ObjectFile
 	{
 	public:
+		ObjectFileElf();
 		virtual bool Parse(Logger& logger, const tchar* filename) override;
 
 	private:
diff --git a/Engine/Source/Programs/UnrealBuildAccelerator/Core/Private/UbaApplicationRules.cpp b/Engine/Source/Programs/UnrealBuildAccelerator/Core/Private/UbaApplicationRules.cpp
--- a/Engine/Source/Programs/UnrealBuildAccelerator/Core/Private/UbaApplicationRules.cpp
+++ b/Engine/Source/Programs/UnrealBuildAccelerator/Core/Private/UbaApplicationRules.cpp
@@ -215,7 +215,7 @@ namespace uba
 
 		virtual bool ShouldExtractSymbols(const tchar* fileName, u64 fileNameLen) const
 		{
-			return EndsWith(fileName, fileNameLen, TC(".obj"));
+			return EndsWith(fileName, fileNameLen, TC(".obj")) || EndsWith(fileName, fileNameLen, TC(".o"));
 		}
 	};
 
@@ -367,7 +367,7 @@ namespace uba
 
 		virtual bool ShouldExtractSymbols(const tchar* fileName, u64 fileNameLen) const
 		{
-			return EndsWith(fileName, fileNameLen, TC(".obj"));
+			return EndsWith(fileName, fileNameLen, TC(".obj")) || EndsWith(fileName, fileNameLen, TC(".o"));
 		}
 	};
 
diff --git a/Engine/Source/Programs/UnrealBuildAccelerator/ObjTool/Private/UbaObjTool.cpp b/Engine/Source/Programs/UnrealBuildAccelerator/ObjTool/Private/UbaObjTool.cpp
--- a/Engine/Source/Programs/UnrealBuildAccelerator/ObjTool/Private/UbaObjTool.cpp
+++ b/Engine/Source/Programs/UnrealBuildAccelerator/ObjTool/Private/UbaObjTool.cpp
@@ -48,6 +48,15 @@ namespace uba
 		return -1;
 	}
 
+	// TODO: Add to rsp file instead
+	const char* NeededImports[] = 
+	{
+		"NvOptimusEnablement",
+		"AmdPowerXpressRequestHighPerformance",
+		"D3D12SDKVersion",
+		"D3D12SDKPath",
+	};
+
 	int WrappedMain(int argc, tchar* argv[])
 	{
 		using namespace uba;
@@ -143,51 +152,20 @@ namespace uba
 				return res;
 		}
 
-		struct SymbolFile
-		{
-			UnorderedSymbols imports;
-			UnorderedExports exports;
-
-			bool ParseFile(Logger& logger, const tchar* filename)
-			{
-				FileAccessor symFile(logger, filename);
-				if (!symFile.OpenMemoryRead())
-					return false;
-				auto readPos = (const char*)symFile.GetData();
-
-				while (*readPos)
-				{
-					auto strEnd = strlen(readPos);
-					imports.insert(std::string(readPos, readPos + strEnd));
-					readPos = readPos + strEnd + 1;
-				}
-				++readPos;
-
-				while (*readPos)
-				{
-					auto strEnd = strlen(readPos);
-					std::string extra;
-					if (const char* comma = strchr(readPos, ','))
-					{
-						strEnd = comma - readPos;
-						extra = comma;
-					}
-					exports.emplace(std::string(readPos, readPos + strEnd), extra);
-					readPos = readPos + strEnd + 1;
-				}
-				return true;
-			}
-		};
-
 		FilteredLogWriter logWriter(g_consoleLogWriter, LogEntryType_Info);
 		LoggerWithWriter logger(logWriter, TC(""));
 
 		if (!objFilesToStrip.empty())
 		{
+			ObjectFileType type = ObjectFileType_Unknown;
+
 			CriticalSection cs;
 			Atomic<bool> success = true;
 			UnorderedSymbols allNeededImports; // Imports needed from the outside of the stripped obj files
 
+			for (auto imp : NeededImports)
+				allNeededImports.insert(imp);
+
 			u32 workerCount = DefaultProcessorCount;
 			WorkManagerImpl workManager(workerCount);
 			workManager.ParallelFor(workerCount, objFilesDependencies, [&](auto& it)
@@ -201,6 +179,9 @@ namespace uba
 						return;
 					}
 					ScopedCriticalSection _(cs);
+					UBA_ASSERT(type == ObjectFileType_Unknown || type == symbolFile.type);
+					if (type == ObjectFileType_Unknown)
+						type = symbolFile.type;
 					allNeededImports.insert(symbolFile.imports.begin(), symbolFile.imports.end());
 				});
 			if (!success)
@@ -222,6 +203,9 @@ namespace uba
 						return;
 					}
 					ScopedCriticalSection _(cs);
+					UBA_ASSERT(type == ObjectFileType_Unknown || type == symbolFile.type);
+					if (type == ObjectFileType_Unknown)
+						type = symbolFile.type;
 					allSharedImports.insert(symbolFile.imports.begin(), symbolFile.imports.end());
 					allSharedExports.insert(symbolFile.exports.begin(), symbolFile.exports.end());
 				});
@@ -229,7 +213,7 @@ namespace uba
 				return -1;
 
 			if (!extraObjFile.empty())
-				if (!ObjectFile::CreateExtraFile(logger, extraObjFile.c_str(), allNeededImports, allSharedImports, allSharedExports, true))
+				if (!ObjectFile::CreateExtraFile(logger, extraObjFile.c_str(), type, allNeededImports, allSharedImports, allSharedExports, true))
 					return -1;
 
 			//logger.Info(TC("Reduced export count from %llu to %llu"), totalExportCount.load(), totalKeptExportCount.size());
@@ -264,7 +248,7 @@ namespace uba
 				StringBuffer<> exportsFile;
 				exportsFile.Append(fileName, lastDot - fileName).Append(TC(".exi"));
 
-				if (!objectFile->WriteSymbols(logger, exportsFile.data))
+				if (!objectFile->WriteImportsAndExports(logger, exportsFile.data))
 					return false;
 
 				//u32 keptExportCount = 0;
