[UBA] Fix scrolling UbaVisualizer with larger scroll ranges by using 32-bit information instead of 16-bit.

Use GetScrollInfo() instead of HIWORD(wParam) when tracking scroll bars in UbaVisualizer as recommended by the WinAPI documentation due to its bit size.

#rnx

diff --git a/Engine/Source/Programs/UnrealBuildAccelerator/Visualizer/Private/UbaVisualizer.cpp b/Engine/Source/Programs/UnrealBuildAccelerator/Visualizer/Private/UbaVisualizer.cpp
--- a/Engine/Source/Programs/UnrealBuildAccelerator/Visualizer/Private/UbaVisualizer.cpp
+++ b/Engine/Source/Programs/UnrealBuildAccelerator/Visualizer/Private/UbaVisualizer.cpp
@@ -2746,11 +2746,18 @@ namespace uba
 				RECT r;
 				GetClientRect(hWnd, &r);
 				float oldScrollY = m_scrollPosY;
+
+				// HIWORD(wParam) only carries 16-bits, so use GetScrollInfo for larger scroll areas
+				SCROLLINFO scrollInfo = {};
+				scrollInfo.cbSize = sizeof(scrollInfo);
+				scrollInfo.fMask = SIF_TRACKPOS;
+				GetScrollInfo(m_hwnd, SB_VERT, &scrollInfo);
+
 				switch (LOWORD(wParam))
 				{
 				case SB_THUMBTRACK:
 				case SB_THUMBPOSITION:
-					m_scrollPosY = -float(HIWORD(wParam));
+					m_scrollPosY = -float(scrollInfo.nTrackPos);
 					break;
 				case SB_PAGEDOWN:
 					m_scrollPosY = m_scrollPosY - r.bottom;
@@ -2781,17 +2788,23 @@ namespace uba
 				float oldScrollX = m_scrollPosX;
 				bool autoScroll = false;
 
+				// HIWORD(wParam) only carries 16-bits, so use GetScrollInfo for larger scroll areas
+				SCROLLINFO scrollInfo = {};
+				scrollInfo.cbSize = sizeof(scrollInfo);
+				scrollInfo.fMask = SIF_TRACKPOS;
+				GetScrollInfo(m_hwnd, SB_HORZ, &scrollInfo);
+
 				switch (LOWORD(wParam))
 				{
 				case SB_THUMBTRACK:
-					m_scrollPosX = -float(HIWORD(wParam));
+					m_scrollPosX = -float(scrollInfo.nTrackPos);
 					if (m_contentWidthWhenThumbTrack == 0)
 						m_contentWidthWhenThumbTrack = m_contentWidth;
 					break;
 				case SB_THUMBPOSITION:
 					autoScroll = m_contentWidthWhenThumbTrack - r.right <= HIWORD(wParam) + 10;
 					m_contentWidthWhenThumbTrack = 0;
-					m_scrollPosX = -float(HIWORD(wParam));
+					m_scrollPosX = -float(scrollInfo.nTrackPos);
 					break;
 				case SB_PAGEDOWN:
 					m_scrollPosX = m_scrollPosX - r.right;
