[UBA] Also assign left mouse button for drag-to-scroll mechanism in UbaVisualizer.

Allow both left and middle mouse buttons to be used for dragging/scrolling the view in UbaVisualizer.
This requires the internal state to be changed from a boolean to a reference-counter so that the set of input events don't cancel each other out.

#rnx

diff --git a/Engine/Source/Programs/UnrealBuildAccelerator/Visualizer/Private/UbaVisualizer.cpp b/Engine/Source/Programs/UnrealBuildAccelerator/Visualizer/Private/UbaVisualizer.cpp
--- a/Engine/Source/Programs/UnrealBuildAccelerator/Visualizer/Private/UbaVisualizer.cpp
+++ b/Engine/Source/Programs/UnrealBuildAccelerator/Visualizer/Private/UbaVisualizer.cpp
@@ -542,6 +542,39 @@ namespace uba
 		}
 	}
 
+	void Visualizer::StartDragToScroll(const POINT& anchor)
+	{
+		// Uses reference-counter method since multiple input events (left and middle mouse button) can trigger the drag-to-scroll mechansim
+		if (m_dragToScrollCounter == 0)
+		{
+			m_processSelected = false;
+			m_sessionSelectedIndex = ~0u;
+			m_statsSelected = false;
+			m_buttonSelected = ~0u;
+			m_timelineSelected = 0;
+			m_fetchedFilesSelected = ~0u;
+			m_workSelected = false;
+			m_autoScroll = false;
+			m_mouseAnchor = { anchor.x, anchor.y };
+			m_scrollAtAnchorX = m_scrollPosX;
+			m_scrollAtAnchorY = m_scrollPosY;
+			SetCapture(m_hwnd);
+			RedrawWindow(m_hwnd, NULL, NULL, RDW_INVALIDATE);
+		}
+		++m_dragToScrollCounter;
+	}
+
+	void Visualizer::StopDragToScroll()
+	{
+		--m_dragToScrollCounter;
+		if (m_dragToScrollCounter == 0)
+		{
+			ReleaseCapture();
+			if (UpdateSelection())
+				RedrawWindow(m_hwnd, NULL, NULL, RDW_INVALIDATE);
+		}
+	}
+
 	void Visualizer::PaintClient(const Function<void(HDC hdc, HDC memDC, RECT& clientRect)>& paintFunc)
 	{
 		HDC hdc = GetDC(m_hwnd);
@@ -2214,7 +2247,7 @@ namespace uba
 
 	bool Visualizer::UpdateSelection()
 	{
-		if (!m_mouseOverWindow || m_middleMouseDown)
+		if (!m_mouseOverWindow || m_dragToScrollCounter > 0)
 			return false;
 		POINT pos;
 		GetCursorPos(&pos);
@@ -2389,7 +2422,7 @@ namespace uba
 		}
 		case WM_MOUSEWHEEL:
 		{
-			if (m_middleMouseDown)
+			if (m_dragToScrollCounter > 0)
 				break;
 
 			int delta = GET_WHEEL_DELTA_WPARAM(wParam);// / WHEEL_DELTA;
@@ -2445,7 +2478,7 @@ namespace uba
 		{
 			POINTS p = MAKEPOINTS(lParam);
 			POINT pos{ p.x, p.y };
-			if (m_middleMouseDown)
+			if (m_dragToScrollCounter > 0)
 			{
 				RECT r;
 				GetClientRect(hWnd, &r);
@@ -2512,21 +2545,8 @@ namespace uba
 
 		case WM_MBUTTONDOWN:
 		{
-			m_processSelected = false;
-			m_sessionSelectedIndex = ~0u;
-			m_statsSelected = false;
-			m_buttonSelected = ~0u;
-			m_timelineSelected = 0;
-			m_fetchedFilesSelected = ~0u;
-			m_workSelected = false;
-			m_autoScroll = false;
 			POINTS p = MAKEPOINTS(lParam);
-			m_mouseAnchor = {p.x, p.y};
-			m_scrollAtAnchorX = m_scrollPosX;
-			m_scrollAtAnchorY = m_scrollPosY;
-			m_middleMouseDown = true;
-			SetCapture(hWnd);
-			RedrawWindow(hWnd, NULL, NULL, RDW_INVALIDATE);
+			StartDragToScroll(POINT{ p.x, p.y });
 			break;
 		}
 
@@ -2575,6 +2595,20 @@ namespace uba
 					RedrawWindow(hWnd, NULL, NULL, RDW_INVALIDATE | RDW_UPDATENOW);
 				}
 			}
+			else
+			{
+				POINTS p = MAKEPOINTS(lParam);
+				StartDragToScroll(POINT{ p.x, p.y });
+			}
+			break;
+		}
+
+		case WM_LBUTTONUP:
+		{
+			if (!(m_buttonSelected != ~0u || m_timelineSelected))
+			{
+				StopDragToScroll();
+			}
 			break;
 		}
 
@@ -2728,10 +2762,7 @@ namespace uba
 
 		case WM_MBUTTONUP:
 		{
-			ReleaseCapture();
-			m_middleMouseDown = false;
-			if (UpdateSelection())
-				RedrawWindow(hWnd, NULL, NULL, RDW_INVALIDATE);
+			StopDragToScroll();
 			//m_processSelected = false;
 			break;
 		}
diff --git a/Engine/Source/Programs/UnrealBuildAccelerator/Visualizer/Private/UbaVisualizer.h b/Engine/Source/Programs/UnrealBuildAccelerator/Visualizer/Private/UbaVisualizer.h
--- a/Engine/Source/Programs/UnrealBuildAccelerator/Visualizer/Private/UbaVisualizer.h
+++ b/Engine/Source/Programs/UnrealBuildAccelerator/Visualizer/Private/UbaVisualizer.h
@@ -69,6 +69,8 @@ namespace uba
 		void GetTitlePrefix(StringBufferBase& out);
 		void ThreadLoop();
 		void Pause(bool pause);
+		void StartDragToScroll(const POINT& anchor);
+		void StopDragToScroll();
 
 		StringBuffer<256> m_namedTrace;
 		StringBuffer<256> m_fileName;
@@ -179,7 +181,7 @@ namespace uba
 		POINT m_mouseAnchor = {};
 		float m_scrollAtAnchorX = 0;
 		float m_scrollAtAnchorY = 0;
-		bool m_middleMouseDown = false;
+		int m_dragToScrollCounter = 0;
 
 		Thread m_thread;
 
