[UBA]
* Bumped memoryblock size from 1mb to 4mb for cached files that contains paths that needs to be normalized. Previous size was 1mb but that seems to be too small for the farm (longer paths)

diff --git a/Engine/Source/Programs/UnrealBuildAccelerator/Common/Private/UbaCacheClient.cpp b/Engine/Source/Programs/UnrealBuildAccelerator/Common/Private/UbaCacheClient.cpp
--- a/Engine/Source/Programs/UnrealBuildAccelerator/Common/Private/UbaCacheClient.cpp
+++ b/Engine/Source/Programs/UnrealBuildAccelerator/Common/Private/UbaCacheClient.cpp
@@ -420,7 +420,7 @@ namespace uba
 				if (ShouldNormalize(path))
 				{
 					// Fetch into memory, file is in special format without absolute paths
-					MemoryBlock normalizedBlock(1*1024*1024);
+					MemoryBlock normalizedBlock(4*1024*1024);
 					bool destinationIsCompressed = false;
 					if (!fetcher.RetrieveFile(m_logger, m_client, casKey, path.data, destinationIsCompressed, &normalizedBlock))
 						return false;
@@ -428,7 +428,7 @@ namespace uba
 					char* fileStart = (char*)normalizedBlock.memory + sizeof(u32);
 
 					// "denormalize" fetched file into another memory block that will be written to disk
-					MemoryBlock localBlock(2*1024*1024);
+					MemoryBlock localBlock(4*1024*1024);
 
 					u64 lastWritten = 0;
 					BinaryReader reader2(normalizedBlock.memory, rootOffsets, normalizedBlock.writtenSize);
