[UBA]
* Fixed bug in cache client where it could create a caskey for a file that should always go through ShouldNormalize that matches caskey of non-normalized file
* Added timers around traverse directory calls to be able to see how much time that is spent in traversing directories

diff --git a/Engine/Source/Programs/UnrealBuildAccelerator/Common/Private/UbaCacheClient.cpp b/Engine/Source/Programs/UnrealBuildAccelerator/Common/Private/UbaCacheClient.cpp
--- a/Engine/Source/Programs/UnrealBuildAccelerator/Common/Private/UbaCacheClient.cpp
+++ b/Engine/Source/Programs/UnrealBuildAccelerator/Common/Private/UbaCacheClient.cpp
@@ -1,6 +1,7 @@
 // Copyright Epic Games, Inc. All Rights Reserved.
 
 #include "UbaCacheClient.h"
+#include "UbaApplicationRules.h"
 #include "UbaCompactTables.h"
 #include "UbaFileAccessor.h"
 #include "UbaNetworkMessage.h"
@@ -205,7 +206,7 @@ namespace uba
 			if (casKey == CasKeyZero)
 			{
 				bool deferCreation = true;
-				bool fileIsCompressed = m_session.ShouldStoreObjFilesCompressed() && path.EndsWith(TC(".obj"));
+				bool fileIsCompressed = IsFileCompressed(info, path.data, path.count);
 				if (!m_storage.StoreCasFile(casKey, path.data, CasKeyZero, deferCreation, fileIsCompressed))
 					return false;
 				if (casKey == CasKeyZero) // If file is not found it was a temporary file that was deleted and is not really an output
@@ -367,7 +368,7 @@ namespace uba
 						else
 						{
 							bool deferCreation = true;
-							bool fileIsCompressed = m_session.ShouldStoreObjFilesCompressed() && path.EndsWith(TC(".obj"));
+							bool fileIsCompressed = IsFileCompressed(info, path.data, path.count);
 							m_storage.StoreCasFile(localCasKey, path.data, CasKeyZero, deferCreation, fileIsCompressed);
 							UBA_ASSERT(localCasKey == CasKeyZero || IsCompressed(localCasKey));
 						}
@@ -426,6 +427,7 @@ namespace uba
 						return false;
 					u32 rootOffsets = *(u32*)normalizedBlock.memory;
 					char* fileStart = (char*)normalizedBlock.memory + sizeof(u32);
+					UBA_ASSERT(rootOffsets <= normalizedBlock.writtenSize);
 
 					// "denormalize" fetched file into another memory block that will be written to disk
 					MemoryBlock localBlock(4*1024*1024);
@@ -857,4 +859,14 @@ namespace uba
 		outPath.Append(root.path).Append(normalizedPath.data + 1);
 		return true;
 	}
+
+	bool CacheClient::IsFileCompressed(const ProcessStartInfo& info, const tchar* filename, u64 filenameSize)
+	{
+		if (!m_session.ShouldStoreObjFilesCompressed())
+			return false;
+		auto rules = info.rules;
+		if (!rules)
+			rules = m_session.GetRules(info);
+		return rules->StoreFileCompressed(filename, filenameSize);
+	}
 }
diff --git a/Engine/Source/Programs/UnrealBuildAccelerator/Common/Private/UbaLogger.cpp b/Engine/Source/Programs/UnrealBuildAccelerator/Common/Private/UbaLogger.cpp
--- a/Engine/Source/Programs/UnrealBuildAccelerator/Common/Private/UbaLogger.cpp
+++ b/Engine/Source/Programs/UnrealBuildAccelerator/Common/Private/UbaLogger.cpp
@@ -51,6 +51,9 @@ namespace uba
 		SetFocus(GetConsoleWindow());
 		SetActiveWindow(GetConsoleWindow());
 #endif
+
+		// *(int*)nullptr = 42; // Use this to force a crashdump instead of exiting
+
 		ExitProcess(terminateCode);
 #else
 		exit(-1);
diff --git a/Engine/Source/Programs/UnrealBuildAccelerator/Common/Private/UbaProcess.cpp b/Engine/Source/Programs/UnrealBuildAccelerator/Common/Private/UbaProcess.cpp
--- a/Engine/Source/Programs/UnrealBuildAccelerator/Common/Private/UbaProcess.cpp
+++ b/Engine/Source/Programs/UnrealBuildAccelerator/Common/Private/UbaProcess.cpp
@@ -160,13 +160,7 @@ namespace uba
 
 		m_arguments = startInfo.arguments;
 
-		tchar extractExports[] = TC(" /extractexports");
-		const tchar* pos = nullptr;
-		if (Contains(m_arguments.c_str(), extractExports, true, &pos))
-		{
-			m_arguments.erase(pos - m_arguments.c_str(), sizeof_array(extractExports));
-			m_extractExports = true;
-		}
+		m_extractExports = Contains(m_arguments.c_str(), TC("/extractexports"), true);
 
 		m_startInfo.arguments = m_arguments.c_str();
 
@@ -1266,6 +1260,15 @@ namespace uba
 				detoursLib = UBA_DETOURS_LIBRARY_ANSI;
 
 			TString commandLine = TC("\"") + m_realApplication + TC("\" ") + m_startInfo.arguments;
+
+			if (m_extractExports)
+			{
+				tchar extractExports[] = TC(" /extractexports");
+				const tchar* pos = nullptr;
+				Contains(commandLine.c_str(), extractExports, true, &pos);
+				commandLine.erase(pos - commandLine.c_str(), sizeof_array(extractExports));
+			}
+
 			LPCSTR dlls[] = { detoursLib };
 
 			STARTUPINFOEX siex;
diff --git a/Engine/Source/Programs/UnrealBuildAccelerator/Common/Private/UbaRootPaths.cpp b/Engine/Source/Programs/UnrealBuildAccelerator/Common/Private/UbaRootPaths.cpp
--- a/Engine/Source/Programs/UnrealBuildAccelerator/Common/Private/UbaRootPaths.cpp
+++ b/Engine/Source/Programs/UnrealBuildAccelerator/Common/Private/UbaRootPaths.cpp
@@ -118,6 +118,8 @@ namespace uba
 			return CasKeyZero;
 
 		CasKeyHasher hasher;
+		u64 magic = 0xFAFAFAFAFAFAFAFAull; // We add some magic to make sure this does not match cas key of non-normalized files (we need to always take normalized path)
+		hasher.Update(&magic, 8);
 		auto hashString = [&](const char* str, u64 strLen, u32 rootPos) { hasher.Update(str, strLen); };
 		if (!NormalizeString<char>(logger, (const char*)file.GetData(), file.GetSize(), hashString, filename))
 			return CasKeyZero;
diff --git a/Engine/Source/Programs/UnrealBuildAccelerator/Common/Private/UbaStats.cpp b/Engine/Source/Programs/UnrealBuildAccelerator/Common/Private/UbaStats.cpp
--- a/Engine/Source/Programs/UnrealBuildAccelerator/Common/Private/UbaStats.cpp
+++ b/Engine/Source/Programs/UnrealBuildAccelerator/Common/Private/UbaStats.cpp
@@ -42,7 +42,7 @@ namespace uba
 		struct Stat { const char* name; u64 nameLen; const Timer& timer; };
 		Stat stats[] =
 		{
-			#define UBA_SYSTEM_STAT(T) { #T, sizeof(#T), T },
+			#define UBA_SYSTEM_STAT(T, ver) { #T, sizeof(#T), T },
 			UBA_SYSTEM_STATS
 			#undef UBA_SYSTEM_STAT
 		};
@@ -58,7 +58,7 @@ namespace uba
 
 	bool SystemStats::IsEmpty()
 	{
-		#define UBA_SYSTEM_STAT(T) if (T.count) return false;
+		#define UBA_SYSTEM_STAT(T, ver) if (T.count) return false;
 		UBA_SYSTEM_STATS
 		#undef UBA_SYSTEM_STAT
 		return true;
@@ -66,21 +66,21 @@ namespace uba
 
 	void SystemStats::Add(const SystemStats& other)
 	{
-		#define UBA_SYSTEM_STAT(var) var += other.var;
+		#define UBA_SYSTEM_STAT(var, ver) var += other.var;
 		UBA_SYSTEM_STATS
 		#undef UBA_SYSTEM_STAT
 	}
 
 	void SystemStats::Write(BinaryWriter& writer)
 	{
-		#define UBA_SYSTEM_STAT(var) uba::Write(writer, var);
+		#define UBA_SYSTEM_STAT(var, ver) uba::Write(writer, var);
 		UBA_SYSTEM_STATS
 		#undef UBA_SYSTEM_STAT
 	}
 
-	void SystemStats::Read(BinaryReader& reader)
+	void SystemStats::Read(BinaryReader& reader, u32 version)
 	{
-		#define UBA_SYSTEM_STAT(var) uba::Read(reader, var);
+		#define UBA_SYSTEM_STAT(var, ver) if (ver <= version) uba::Read(reader, var);
 		UBA_SYSTEM_STATS
 		#undef UBA_SYSTEM_STAT
 	}
diff --git a/Engine/Source/Programs/UnrealBuildAccelerator/Common/Public/UbaCacheClient.h b/Engine/Source/Programs/UnrealBuildAccelerator/Common/Public/UbaCacheClient.h
--- a/Engine/Source/Programs/UnrealBuildAccelerator/Common/Public/UbaCacheClient.h
+++ b/Engine/Source/Programs/UnrealBuildAccelerator/Common/Public/UbaCacheClient.h
@@ -56,6 +56,7 @@ namespace uba
 		bool ShouldNormalize(const StringBufferBase& path);
 
 		bool GetLocalPathAndCasKey(Bucket& bucket, const RootPaths& rootPaths, StringBufferBase& outPath, CasKey& outKey, CompactCasKeyTable& casKeyTable, CompactPathTable& pathTable, u32 offset);
+		bool IsFileCompressed(const ProcessStartInfo& info, const tchar* filename, u64 filenameSize);
 
 		MutableLogger m_logger;
 		StorageImpl& m_storage;
diff --git a/Engine/Source/Programs/UnrealBuildAccelerator/Common/Public/UbaDirectoryIterator.h b/Engine/Source/Programs/UnrealBuildAccelerator/Common/Public/UbaDirectoryIterator.h
--- a/Engine/Source/Programs/UnrealBuildAccelerator/Common/Public/UbaDirectoryIterator.h
+++ b/Engine/Source/Programs/UnrealBuildAccelerator/Common/Public/UbaDirectoryIterator.h
@@ -4,6 +4,7 @@
 
 #include "UbaFile.h"
 #include "UbaLogger.h"
+#include "UbaProcessStats.h"
 #include "UbaStringBuffer.h"
 
 namespace uba
@@ -82,6 +83,9 @@ namespace uba
 		UNICODE_STRING uniName;
 		RtlInitUnicodeString(&uniName, str.data);
 
+		Timer timer(0, 1);
+		auto mg = MakeGuard([&]() { SystemStats::GetCurrent().traverseDir.Add(timer); });
+
 		HANDLE handle;
 
 		OBJECT_ATTRIBUTES ObjectAttributes;
@@ -89,7 +93,11 @@ namespace uba
 		memset(&IoStatusBlock, 0, sizeof(IoStatusBlock));
 		InitializeObjectAttributes(&ObjectAttributes, &uniName, OBJ_CASE_INSENSITIVE, NULL, NULL);
 		ULONG ShareAccess = FILE_SHARE_READ | FILE_SHARE_WRITE | FILE_SHARE_DELETE;
+
+		u64 startTime = GetTime();
 		NTSTATUS res = NtCreateFile(&handle, FILE_READ_ATTRIBUTES | FILE_LIST_DIRECTORY | FILE_ATTRIBUTE_UNPINNED, &ObjectAttributes, &IoStatusBlock, 0, FILE_ATTRIBUTE_NORMAL, ShareAccess, FILE_OPEN, FILE_DIRECTORY_FILE | FILE_SYNCHRONOUS_IO_ALERT, NULL, 0);
+		timer.time += startTime - GetTime();
+
 		if (res == STATUS_OBJECT_NAME_NOT_FOUND || res == STATUS_OBJECT_PATH_NOT_FOUND)
 			return !errorOnNotFound;
 		if (res == STATUS_NOT_A_DIRECTORY || res == STATUS_ACCESS_DENIED || res == STATUS_NO_MEDIA_IN_DEVICE || res == STATUS_FVE_LOCKED_VOLUME)
@@ -124,7 +132,10 @@ namespace uba
 		u8 buff[64*1024];
 		while (true)
 		{
+			startTime = GetTime();
 			res = NtQueryDirectoryFile(handle, 0, NULL, NULL, &IoStatusBlock, buff, sizeof(buff) - 2, (FILE_INFORMATION_CLASS)FileIdBothDirectoryInformation, FALSE, NULL, FALSE);
+			timer.time += startTime - GetTime();
+
 			if (res != STATUS_SUCCESS)
 			{
 				if (res != STATUS_NO_MORE_FILES)
diff --git a/Engine/Source/Programs/UnrealBuildAccelerator/Common/Public/UbaTrace.h b/Engine/Source/Programs/UnrealBuildAccelerator/Common/Public/UbaTrace.h
--- a/Engine/Source/Programs/UnrealBuildAccelerator/Common/Public/UbaTrace.h
+++ b/Engine/Source/Programs/UnrealBuildAccelerator/Common/Public/UbaTrace.h
@@ -59,7 +59,7 @@ namespace uba
 	using Color = u32;
 	inline Color toColor(u8 r, u8 g, u8 b) { return (r << 16) + (g << 8) + b; }
 
-	static constexpr u32 TraceVersion = 26;
+	static constexpr u32 TraceVersion = 27;
 	static constexpr u32 TraceReadCompatibilityVersion = 6;
 
 	class Trace : public WorkTracker
diff --git a/Engine/Source/Programs/UnrealBuildAccelerator/Core/Public/UbaProcessStats.h b/Engine/Source/Programs/UnrealBuildAccelerator/Core/Public/UbaProcessStats.h
--- a/Engine/Source/Programs/UnrealBuildAccelerator/Core/Public/UbaProcessStats.h
+++ b/Engine/Source/Programs/UnrealBuildAccelerator/Core/Public/UbaProcessStats.h
@@ -128,25 +128,26 @@ namespace uba
 	};
 
 	#define UBA_SYSTEM_STATS \
-		UBA_SYSTEM_STAT(createFile) \
-		UBA_SYSTEM_STAT(closeFile) \
-		UBA_SYSTEM_STAT(writeFile) \
-		UBA_SYSTEM_STAT(readFile) \
-		UBA_SYSTEM_STAT(setFileInfo) \
-		UBA_SYSTEM_STAT(createFileMapping) \
-		UBA_SYSTEM_STAT(mapViewOfFile) \
-		UBA_SYSTEM_STAT(unmapViewOfFile) \
-		UBA_SYSTEM_STAT(getFileTime) \
-		UBA_SYSTEM_STAT(closeHandle) \
+		UBA_SYSTEM_STAT(createFile, 0) \
+		UBA_SYSTEM_STAT(closeFile, 0) \
+		UBA_SYSTEM_STAT(writeFile, 0) \
+		UBA_SYSTEM_STAT(readFile, 0) \
+		UBA_SYSTEM_STAT(setFileInfo, 0) \
+		UBA_SYSTEM_STAT(createFileMapping, 0) \
+		UBA_SYSTEM_STAT(mapViewOfFile, 0) \
+		UBA_SYSTEM_STAT(unmapViewOfFile, 0) \
+		UBA_SYSTEM_STAT(getFileTime, 0) \
+		UBA_SYSTEM_STAT(closeHandle, 0) \
+		UBA_SYSTEM_STAT(traverseDir, 27) \
 
 
 	struct SystemStats
 	{
-		#define UBA_SYSTEM_STAT(T) ExtendedTimer T;
+		#define UBA_SYSTEM_STAT(T, ver) ExtendedTimer T;
 		UBA_SYSTEM_STATS
 		#undef UBA_SYSTEM_STAT
 
-		void Read(BinaryReader& reader);
+		void Read(BinaryReader& reader, u32 version);
 		void Write(BinaryWriter& writer);
 		void Print(Logger& logger, bool writeHeader, u64 frequency = GetFrequency());
 		bool IsEmpty();
diff --git a/Engine/Source/Programs/UnrealBuildAccelerator/Visualizer/Private/UbaTraceReader.cpp b/Engine/Source/Programs/UnrealBuildAccelerator/Visualizer/Private/UbaTraceReader.cpp
--- a/Engine/Source/Programs/UnrealBuildAccelerator/Visualizer/Private/UbaTraceReader.cpp
+++ b/Engine/Source/Programs/UnrealBuildAccelerator/Visualizer/Private/UbaTraceReader.cpp
@@ -508,7 +508,7 @@ namespace uba
 			{
 				sessionStats.Read(reader, out.version);
 				storageStats.Read(reader);
-				systemStats.Read(reader);
+				systemStats.Read(reader, out.version);
 			}
 			const u8* dataEnd = reader.GetPositionData();
 			process.stats.resize(dataEnd - dataStart);
@@ -567,7 +567,7 @@ namespace uba
 			processStats.Read(reader, out.version);
 			sessionStats.Read(reader, out.version);
 			storageStats.Read(reader);
-			systemStats.Read(reader);
+			systemStats.Read(reader, out.version);
 
 			const u8* dataEnd = reader.GetPositionData();
 			process.stats.resize(dataEnd - dataStart);
@@ -915,7 +915,7 @@ namespace uba
 			if (success)
 			{
 				storageStats.Read(reader);
-				systemStats.Read(reader);
+				systemStats.Read(reader, out.version);
 			}
 			const u8* dataEnd = reader.GetPositionData();
 			process.stats.resize(dataEnd - dataStart);
diff --git a/Engine/Source/Programs/UnrealBuildAccelerator/Visualizer/Private/UbaVisualizer.cpp b/Engine/Source/Programs/UnrealBuildAccelerator/Visualizer/Private/UbaVisualizer.cpp
--- a/Engine/Source/Programs/UnrealBuildAccelerator/Visualizer/Private/UbaVisualizer.cpp
+++ b/Engine/Source/Programs/UnrealBuildAccelerator/Visualizer/Private/UbaVisualizer.cpp
@@ -1344,7 +1344,7 @@ namespace uba
 					if (process.exitCode == 0)
 					{
 						storageStats.Read(reader);
-						systemStats.Read(reader);
+						systemStats.Read(reader, m_traceView.version);
 					}
 				}
 				else
@@ -1355,7 +1355,7 @@ namespace uba
 					{
 						sessionStats.Read(reader, m_traceView.version);
 						storageStats.Read(reader);
-						systemStats.Read(reader);
+						systemStats.Read(reader, m_traceView.version);
 					}
 				}
 
@@ -2185,7 +2185,7 @@ namespace uba
 			{
 				sessionStats.Read(reader, m_traceView.version);
 				storageStats.Read(reader);
-				systemStats.Read(reader);
+				systemStats.Read(reader, m_traceView.version);
 			}
 
 			out.Info(L"  ----------- Process stats -----------");
