[UBA]
* Fixed so normalized file has a header that can be used to validate if it was normalized or not

diff --git a/Engine/Source/Programs/UnrealBuildAccelerator/Common/Private/UbaCacheClient.cpp b/Engine/Source/Programs/UnrealBuildAccelerator/Common/Private/UbaCacheClient.cpp
--- a/Engine/Source/Programs/UnrealBuildAccelerator/Common/Private/UbaCacheClient.cpp
+++ b/Engine/Source/Programs/UnrealBuildAccelerator/Common/Private/UbaCacheClient.cpp
@@ -425,56 +425,66 @@ namespace uba
 					bool destinationIsCompressed = false;
 					if (!fetcher.RetrieveFile(m_logger, m_client, casKey, path.data, destinationIsCompressed, &normalizedBlock))
 						return false;
-					u32 rootOffsets = *(u32*)normalizedBlock.memory;
-					char* fileStart = (char*)normalizedBlock.memory + sizeof(u32);
-					UBA_ASSERT(rootOffsets <= normalizedBlock.writtenSize);
 
-					// "denormalize" fetched file into another memory block that will be written to disk
+					u8* memoryToWrite = normalizedBlock.memory;
+					u64 memoryToWriteSize = normalizedBlock.writtenSize;
 					MemoryBlock localBlock(4*1024*1024);
 
-					u64 lastWritten = 0;
-					BinaryReader reader2(normalizedBlock.memory, rootOffsets, normalizedBlock.writtenSize);
-					while (reader2.GetLeft())
+					// It could be that file was actually not normalized.. because it had no absolute paths, in that case we don't have the header and can just write it to disk as is
+					if (normalizedBlock.writtenSize > 8 && *(u64*)normalizedBlock.memory == 0xFAFAFAFAFAFAFAFAull)
 					{
-						u64 rootOffset = reader2.Read7BitEncoded();
-						if (u64 toWrite = rootOffset - lastWritten)
+						u32 rootOffsets = *(u32*)(normalizedBlock.memory + sizeof(u64));
+						char* fileStart = (char*)(normalizedBlock.memory + sizeof(u64) + sizeof(u32));
+						UBA_ASSERT(rootOffsets <= normalizedBlock.writtenSize);
+
+						// "denormalize" fetched file into another memory block that will be written to disk
+						u64 lastWritten = 0;
+						BinaryReader reader2(normalizedBlock.memory, rootOffsets, normalizedBlock.writtenSize);
+						while (reader2.GetLeft())
+						{
+							u64 rootOffset = reader2.Read7BitEncoded();
+							if (u64 toWrite = rootOffset - lastWritten)
+								memcpy(localBlock.Allocate(toWrite, 1, TC("")), fileStart + lastWritten, toWrite);
+							u8 rootIndex = fileStart[rootOffset] - RootPaths::RootStartByte;
+							auto& root = rootPaths.GetRoot(rootIndex);
+
+							#if PLATFORM_WINDOWS
+							StringBuffer<> pathTemp;
+							pathTemp.Append(root.path);
+							char rootPath[512];
+							u32 rootPathLen = pathTemp.Parse(rootPath, sizeof_array(rootPath));
+							#else
+							const char* rootPath = root.path.data();
+							u32 rootPathLen = root.path.size();
+							#endif
+
+							if (u32 toWrite = rootPathLen - 1)
+								memcpy(localBlock.Allocate(toWrite, 1, TC("")), rootPath, toWrite);
+							lastWritten = rootOffset + 1;
+						}
+
+						u64 fileSize = rootOffsets - (sizeof(u64) + sizeof(u32));
+						if (u64 toWrite = fileSize - lastWritten)
 							memcpy(localBlock.Allocate(toWrite, 1, TC("")), fileStart + lastWritten, toWrite);
-						u8 rootIndex = fileStart[rootOffset] - RootPaths::RootStartByte;
-						auto& root = rootPaths.GetRoot(rootIndex);
-
-						#if PLATFORM_WINDOWS
-						StringBuffer<> pathTemp;
-						pathTemp.Append(root.path);
-						char rootPath[512];
-						u32 rootPathLen = pathTemp.Parse(rootPath, sizeof_array(rootPath));
-						#else
-						const char* rootPath = root.path.data();
-						u32 rootPathLen = root.path.size();
-						#endif
-
-						if (u32 toWrite = rootPathLen - 1)
-							memcpy(localBlock.Allocate(toWrite, 1, TC("")), rootPath, toWrite);
-						lastWritten = rootOffset + 1;
-					}
 
-					u64 fileSize = rootOffsets - sizeof(u32);
-					if (u64 toWrite = fileSize - lastWritten)
-						memcpy(localBlock.Allocate(toWrite, 1, TC("")), fileStart + lastWritten, toWrite);
+						memoryToWrite = localBlock.memory;
+						memoryToWriteSize = localBlock.writtenSize;
+					}
 
 					FileAccessor destFile(m_logger, path.data);
 					if (!destFile.CreateWrite())
 						return false;
-					if (!destFile.Write(localBlock.memory, localBlock.writtenSize))
+					if (!destFile.Write(memoryToWrite, memoryToWriteSize))
 						return false;
 					if (!destFile.Close(&fetcher.lastWritten))
 						return false;
 
-					fetcher.sizeOnDisk = fileSize;
-					casKey = CalculateCasKey(localBlock.memory, localBlock.writtenSize, false, nullptr);
+					fetcher.sizeOnDisk = memoryToWriteSize;
+					casKey = CalculateCasKey(memoryToWrite, memoryToWriteSize, false, nullptr);
 				}
 				else
 				{
-					bool destinationIsCompressed = m_session.ShouldStoreObjFilesCompressed() && path.EndsWith(TC(".obj"));
+					bool destinationIsCompressed = IsFileCompressed(info, path.data, path.count);
 					if (!fetcher.RetrieveFile(m_logger, m_client, casKey, path.data, destinationIsCompressed))
 						return false;
 				}
@@ -652,18 +662,22 @@ namespace uba
 					FileAccessor file(m_logger, path.data);
 					if (!file.OpenMemoryRead())
 						return false;
-					MemoryBlock block(AlignUp(file.GetSize(), 64*1024));
+					MemoryBlock block(AlignUp(file.GetSize() + 16, 64*1024));
+					*(u64*)block.Allocate(sizeof(u64), 1, TC("")) = 0xFAFAFAFAFAFAFAFAull; // Magic to be able to know if file was normalized or not
+
 					u32& rootOffsetsStart = *(u32*)block.Allocate(sizeof(u32), 1, TC(""));
 					rootOffsetsStart = 0;
 					Vector<u32> rootOffsets;
 					u32 rootOffsetsSize = 0;
 
+					bool wasNormalized = false;
 					auto handleString = [&](const char* str, u64 strLen, u32 rootPos)
 						{
 							void* mem = block.Allocate(strLen, 1, TC(""));
 							memcpy(mem, str, strLen);
 							if (rootPos != ~0u)
 							{
+								wasNormalized = true;
 								rootOffsets.push_back(rootPos);
 								rootOffsetsSize += Get7BitEncodedCount(rootPos);
 							}
@@ -686,7 +700,16 @@ namespace uba
 
 					auto& s = m_storage;
 					FileSender sender { m_logger, m_client, s.m_bufferSlots, s.Stats(), m_sendOneAtTheTimeLock, s.m_casCompressor, s.m_casCompressionLevel };
-					if (!sender.SendFileCompressed(casKey, path.data, block.memory, block.writtenSize, TC("SendCacheEntry")))
+
+					u8* dataToSend = block.memory;
+					u64 sizeToSend = block.writtenSize;
+					if (!wasNormalized)
+					{
+						dataToSend += 9;
+						sizeToSend -= 9;
+					}
+
+					if (!sender.SendFileCompressed(casKey, path.data, dataToSend, sizeToSend, TC("SendCacheEntry")))
 						return m_logger.Error(TC("Failed to send cas content for file %s"), path.data);
 				}
 				else
@@ -803,7 +826,7 @@ namespace uba
 
 		// Add arguments list to key
 		auto hashString = [&](const tchar* str, u64 strLen, u32 rootPos) { hasher.Update(str, strLen*sizeof(tchar)); };
-		if (!rootPaths.NormalizeString(m_logger, info.arguments, TStrlen(info.arguments), hashString, TC("CmdKey")))
+		if (!rootPaths.NormalizeString(m_logger, info.arguments, TStrlen(info.arguments), hashString, TC("CmdKey "), info.description))
 			return CasKeyZero;
 
 		// Add content of rsp file to key (This will cost a bit of perf since we need to normalize.. should this be part of key?)
diff --git a/Engine/Source/Programs/UnrealBuildAccelerator/Common/Private/UbaRootPaths.cpp b/Engine/Source/Programs/UnrealBuildAccelerator/Common/Private/UbaRootPaths.cpp
--- a/Engine/Source/Programs/UnrealBuildAccelerator/Common/Private/UbaRootPaths.cpp
+++ b/Engine/Source/Programs/UnrealBuildAccelerator/Common/Private/UbaRootPaths.cpp
@@ -117,13 +117,22 @@ namespace uba
 		if (!file.OpenMemoryRead())
 			return CasKeyZero;
 
+		bool wasNormalized = false;
 		CasKeyHasher hasher;
-		u64 magic = 0xFAFAFAFAFAFAFAFAull; // We add some magic to make sure this does not match cas key of non-normalized files (we need to always take normalized path)
-		hasher.Update(&magic, 8);
-		auto hashString = [&](const char* str, u64 strLen, u32 rootPos) { hasher.Update(str, strLen); };
+		auto hashString = [&](const char* str, u64 strLen, u32 rootPos)
+			{
+				wasNormalized |= rootPos != ~0u;
+				hasher.Update(str, strLen);
+			};
 		if (!NormalizeString<char>(logger, (const char*)file.GetData(), file.GetSize(), hashString, filename))
 			return CasKeyZero;
 
+		if (wasNormalized)
+		{
+			u64 magic = 0xFAFAFAFAFAFAFAFAull; // We add some magic to make sure this does not match cas key of non-normalized files
+			hasher.Update(&magic, 8);
+		}
+
 		return ToCasKey(hasher, false);
 	}
 
diff --git a/Engine/Source/Programs/UnrealBuildAccelerator/Common/Public/UbaRootPaths.h b/Engine/Source/Programs/UnrealBuildAccelerator/Common/Public/UbaRootPaths.h
--- a/Engine/Source/Programs/UnrealBuildAccelerator/Common/Public/UbaRootPaths.h
+++ b/Engine/Source/Programs/UnrealBuildAccelerator/Common/Public/UbaRootPaths.h
@@ -26,7 +26,7 @@ namespace uba
 		const Root& GetRoot(u32 index) const;
 
 		template<typename CharType, typename Func>
-		bool NormalizeString(Logger& logger, const CharType* str, u64 strLen, const Func& func, const tchar* hint) const;
+		bool NormalizeString(Logger& logger, const CharType* str, u64 strLen, const Func& func, const tchar* hint, const tchar* hint2 = TC("")) const;
 
 		CasKey NormalizeAndHashFile(Logger& logger, const tchar* filename) const;
 
@@ -41,7 +41,7 @@ namespace uba
 
 
 	template<typename CharType, typename Func>
-	bool RootPaths::NormalizeString(Logger& logger, const CharType* str, u64 strLen, const Func& func, const tchar* hint) const
+	bool RootPaths::NormalizeString(Logger& logger, const CharType* str, u64 strLen, const Func& func, const tchar* hint, const tchar* hint2) const
 	{
 		auto strEnd = str + strLen;
 		auto searchPos = str;
@@ -78,7 +78,7 @@ namespace uba
 			auto root = FindRoot(path);
 			if (!root)
 			{
-				logger.Info(TC("PATH WITHOUT ROOT: %s (inside file %s at offset %u)"), path.data, hint, destPos);
+				logger.Info(TC("PATH WITHOUT ROOT: %s (inside %s%s at offset %u)"), path.data, hint, hint2, destPos);
 				return false;
 			}
 
