[UBA]
* Changed normalized file magic size from 8 bytes to 2..
* Fixed bug with wrong offset calculations

diff --git a/Engine/Source/Programs/UnrealBuildAccelerator/Common/Private/UbaCacheClient.cpp b/Engine/Source/Programs/UnrealBuildAccelerator/Common/Private/UbaCacheClient.cpp
--- a/Engine/Source/Programs/UnrealBuildAccelerator/Common/Private/UbaCacheClient.cpp
+++ b/Engine/Source/Programs/UnrealBuildAccelerator/Common/Private/UbaCacheClient.cpp
@@ -431,10 +431,10 @@ namespace uba
 					MemoryBlock localBlock(4*1024*1024);
 
 					// It could be that file was actually not normalized.. because it had no absolute paths, in that case we don't have the header and can just write it to disk as is
-					if (normalizedBlock.writtenSize > 8 && *(u64*)normalizedBlock.memory == 0xFAFAFAFAFAFAFAFAull)
+					if (normalizedBlock.writtenSize > 6 && *(u16*)normalizedBlock.memory == 0xFFFF)
 					{
-						u32 rootOffsets = *(u32*)(normalizedBlock.memory + sizeof(u64));
-						char* fileStart = (char*)(normalizedBlock.memory + sizeof(u64) + sizeof(u32));
+						u32 rootOffsets = *(u32*)(normalizedBlock.memory + sizeof(u16));
+						char* fileStart = (char*)(normalizedBlock.memory + sizeof(u16) + sizeof(u32));
 						UBA_ASSERT(rootOffsets <= normalizedBlock.writtenSize);
 
 						// "denormalize" fetched file into another memory block that will be written to disk
@@ -463,7 +463,7 @@ namespace uba
 							lastWritten = rootOffset + 1;
 						}
 
-						u64 fileSize = rootOffsets - (sizeof(u64) + sizeof(u32));
+						u64 fileSize = rootOffsets - (sizeof(u16) + sizeof(u32));
 						if (u64 toWrite = fileSize - lastWritten)
 							memcpy(localBlock.Allocate(toWrite, 1, TC("")), fileStart + lastWritten, toWrite);
 
@@ -663,7 +663,7 @@ namespace uba
 					if (!file.OpenMemoryRead())
 						return false;
 					MemoryBlock block(AlignUp(file.GetSize() + 16, 64*1024));
-					*(u64*)block.Allocate(sizeof(u64), 1, TC("")) = 0xFAFAFAFAFAFAFAFAull; // Magic to be able to know if file was normalized or not
+					*(u16*)block.Allocate(sizeof(u16), 1, TC("")) = 0xFFFF; // Magic to be able to know if file was normalized or not
 
 					u32& rootOffsetsStart = *(u32*)block.Allocate(sizeof(u32), 1, TC(""));
 					rootOffsetsStart = 0;
@@ -705,8 +705,9 @@ namespace uba
 					u64 sizeToSend = block.writtenSize;
 					if (!wasNormalized)
 					{
-						dataToSend += 9;
-						sizeToSend -= 9;
+						// magic + rootoffsetssize
+						dataToSend += 6;
+						sizeToSend -= 6;
 					}
 
 					if (!sender.SendFileCompressed(casKey, path.data, dataToSend, sizeToSend, TC("SendCacheEntry")))
diff --git a/Engine/Source/Programs/UnrealBuildAccelerator/Common/Private/UbaRootPaths.cpp b/Engine/Source/Programs/UnrealBuildAccelerator/Common/Private/UbaRootPaths.cpp
--- a/Engine/Source/Programs/UnrealBuildAccelerator/Common/Private/UbaRootPaths.cpp
+++ b/Engine/Source/Programs/UnrealBuildAccelerator/Common/Private/UbaRootPaths.cpp
@@ -129,8 +129,8 @@ namespace uba
 
 		if (wasNormalized)
 		{
-			u64 magic = 0xFAFAFAFAFAFAFAFAull; // We add some magic to make sure this does not match cas key of non-normalized files
-			hasher.Update(&magic, 8);
+			u16 magic = 0xFFFF; // We add some magic to make sure this does not match cas key of non-normalized files
+			hasher.Update(&magic, sizeof(u16));
 		}
 
 		return ToCasKey(hasher, false);
