[UBA]
* Fixed deadlock that could happen when lots of clients are connecting while maintenance is running
* Added UbaCli -cachecommand=<cmd> that can be used to query status of cache service
* Fixed code setting compressed byte on caskey to only use a bit
* Changed so a bit on the caskey can tell if a file is normalized or not
* Added stats for writing to cache to trace file
*

diff --git a/Engine/Source/Programs/UnrealBuildAccelerator/Cli/Private/UbaCli.cpp b/Engine/Source/Programs/UnrealBuildAccelerator/Cli/Private/UbaCli.cpp
--- a/Engine/Source/Programs/UnrealBuildAccelerator/Cli/Private/UbaCli.cpp
+++ b/Engine/Source/Programs/UnrealBuildAccelerator/Cli/Private/UbaCli.cpp
@@ -66,6 +66,7 @@ namespace uba
 		logger.Info(TC("   -dir=<rootdir>          The directory used to store data. Defaults to \"%s\""), DefaultRootDir);
 		logger.Info(TC("   -port=[<host>:]<port>   The ip/name and port (default: %u) of the machine we want to help"), DefaultPort);
 		logger.Info(TC("   -log                    Log all processes detouring information to file (only works with debug builds)"));
+		logger.Info(TC("   -quiet                  Does not output any logging in console except errors"));
 		logger.Info(TC("   -loop=<count>           Loop the commandline <count> number of times. Will exit when/if it fails"));
 		logger.Info(TC("   -workdir=<dir>          Working directory"));
 		logger.Info(TC("   -checkcas               Check so all cas entries are correct"));
@@ -82,6 +83,7 @@ namespace uba
 		logger.Info(TC("   -coordinator=<name>     Load a UbaCoordinator<name>.dll to instantiate a coordinator to get helpers"));
 		logger.Info(TC("   -cache=<host>[:<port>]  Connect to cache server. Will fetch from cache unless -populatecache is set"));
 		logger.Info(TC("   -populatecache          Populate cache server if connected to one"));
+		logger.Info(TC("   -cachecommand=<cmd>     Send command to cache server. Will output result in log"));
 		logger.Info(TC("   -writecachesummary      Write cache summary file about connected cache server"));
 		logger.Info(TC(""));
 		logger.Info(TC("  CoordinatorOptions (if coordinator set):"));
@@ -171,6 +173,7 @@ namespace uba
 		bool writeCacheSummary = false;
 		TString checkFileTable;
 		TString cacheFilterString;
+		TString cacheCommand;
 
 		u32 loopCount = 1;
 
@@ -180,7 +183,8 @@ namespace uba
 			CommandType_Local,
 			CommandType_Remote,
 			CommandType_Native,
-			CommandType_Agent
+			CommandType_Agent,
+			CommandType_None,
 		};
 
 		CommandType commandType = CommandType_NotSet;
@@ -350,10 +354,19 @@ namespace uba
 			{
 				populateCache = true;
 			}
+			else if (name.Equals(TC("-cachecommand")))
+			{
+				if (value.IsEmpty())
+					return PrintHelp(TC("-cachecommand needs a value"));
+				cacheCommand = value.data;
+				commandType = CommandType_None;
+				quiet = true;
+			}
 			else if (name.Equals(TC("-writecachesummary")))
 			{
 				writeCacheSummary = true;
 				cacheFilterString = value.data;
+				commandType = CommandType_None;
 			}
 			else if (name.Equals(TC("-storeraw")))
 			{
@@ -378,7 +391,7 @@ namespace uba
 			}
 		}
 
-		FilteredLogWriter logWriter(g_consoleLogWriter, quiet ? LogEntryType_Info : LogEntryType_Detail);
+		FilteredLogWriter logWriter(g_consoleLogWriter, quiet ? LogEntryType_Warning : LogEntryType_Detail);
 		LoggerWithWriter logger(logWriter, TC(""));
 
 		if (deleteCas)
@@ -481,44 +494,47 @@ namespace uba
 			return PrintHelp(errorMsg);
 		}
 
-		if (application.empty())
-			return PrintHelp(TC("No executable provided"));
-
 		StringBuffer<512> currentDir;
 		GetCurrentDirectoryW(currentDir);
 
-		bool isAbsolute = IsWindows ? application[1] == ':' : application[0] == '/';
-		if (!isAbsolute)
-		{
-			StringBuffer<> fullApplicationName;
-			if (!SearchPathForFile(logger, fullApplicationName, application.c_str(), currentDir.data))
-				return logger.Error(TC("Failed to find full path to %s"), application.c_str());
-			application = fullApplicationName.data;
-		}
-
-		if (getCas)
+		if (commandType != CommandType_None)
 		{
-			FileAccessor fa(logger, application.c_str());
-			if (!fa.OpenMemoryRead())
-				return logger.Error(TC("Failed to open file %s"), application.c_str());
-			u64 fileSize = fa.GetSize();
-			u8* data = fa.GetData();
-			bool is64Bit = true;
+			if (application.empty())
+				return PrintHelp(TC("No executable provided"));
 
-			CasKey key = CalculateCasKey(data, fileSize, false, nullptr);
+			bool isAbsolute = IsWindows ? application[1] == ':' : application[0] == '/';
+			if (!isAbsolute)
+			{
+				StringBuffer<> fullApplicationName;
+				if (!SearchPathForFile(logger, fullApplicationName, application.c_str(), currentDir.data))
+					return logger.Error(TC("Failed to find full path to %s"), application.c_str());
+				application = fullApplicationName.data;
+			}
 
-			if (data[0] != 'M' || data[1] != 'Z')
-				is64Bit = false;
-			else
+			if (getCas)
 			{
-				u32 offset = *(u32*)(data + 0x3c);
-				is64Bit = *(u32*)(data + offset) == 0x00004550;
+				FileAccessor fa(logger, application.c_str());
+				if (!fa.OpenMemoryRead())
+					return logger.Error(TC("Failed to open file %s"), application.c_str());
+				u64 fileSize = fa.GetSize();
+				u8* data = fa.GetData();
+				bool is64Bit = true;
+
+				CasKey key = CalculateCasKey(data, fileSize, false, nullptr);
+
+				if (data[0] != 'M' || data[1] != 'Z')
+					is64Bit = false;
+				else
+				{
+					u32 offset = *(u32*)(data + 0x3c);
+					is64Bit = *(u32*)(data + offset) == 0x00004550;
+				}
+				logger.Info(TC("%s"), application.c_str());
+				logger.Info(TC("  Is64Bit: %s"), (is64Bit ? TC("true") : TC("false")));
+				logger.Info(TC("  Size: %llu"), fileSize);
+				logger.Info(TC("  CasKey: %s"), CasKeyString(key).str);
+				return 0;
 			}
-			logger.Info(TC("%s"), application.c_str());
-			logger.Info(TC("  Is64Bit: %s"), (is64Bit ? TC("true") : TC("false")));
-			logger.Info(TC("  Size: %llu"), fileSize);
-			logger.Info(TC("  CasKey: %s"), CasKeyString(key).str);
-			return 0;
 		}
 
 		const tchar* dbgStr = TC("");
@@ -566,7 +582,7 @@ namespace uba
 		StorageServer& storageServer = *new StorageServer(storageInfo);
 		auto destroyStorage = MakeGuard([&]() { delete &storageServer; });
 
-		SessionServerCreateInfo info(storageServer, networkServer);
+		SessionServerCreateInfo info(storageServer, networkServer, logWriter);
 		info.useUniqueId = useScheduler;
 		info.traceEnabled = true;
 		//info.detailedTrace = true;
@@ -590,7 +606,7 @@ namespace uba
 
 		auto CreateCacheClient = [&]()
 			{
-				auto nc = new NetworkClient(ctorSuccess);
+				auto nc = new NetworkClient(ctorSuccess, {logWriter});
 				cacheClient = new CacheClient({logWriter, storageServer, *nc, sessionServer});
 			};
 
@@ -603,13 +619,21 @@ namespace uba
 				return -1;
 			}
 
+			if (!cacheCommand.empty())
+			{
+				LoggerWithWriter consoleLogger(g_consoleLogWriter);
+				if (!cacheClient->ExecuteCommand(consoleLogger, cacheCommand.data()))
+					return -1;
+				return 0;
+			}
+
 			if (writeCacheSummary)
 			{
 				StringBuffer<> tempFile(sessionServer.GetTempPath());
 				Guid guid;
 				CreateGuid(guid);
 				tempFile.Append(GuidToString(guid).str).Append(TC(".txt"));
-				if (!cacheClient->WriteCacheSummary(tempFile.data, cacheFilterString.data()))
+				if (!cacheClient->ExecuteCommand(logger, TC("content"), tempFile.data, cacheFilterString.data()))
 					return -1;
 				logger.Info(TC("Cache status summary written to %s"), tempFile.data);
 
diff --git a/Engine/Source/Programs/UnrealBuildAccelerator/Common/Private/UbaCacheClient.cpp b/Engine/Source/Programs/UnrealBuildAccelerator/Common/Private/UbaCacheClient.cpp
--- a/Engine/Source/Programs/UnrealBuildAccelerator/Common/Private/UbaCacheClient.cpp
+++ b/Engine/Source/Programs/UnrealBuildAccelerator/Common/Private/UbaCacheClient.cpp
@@ -59,16 +59,39 @@ namespace uba
 
 		m_client.RegisterOnConnected([this]()
 			{
-				StackBinaryWriter<1024> writer;
-				NetworkMessage msg(m_client, CacheServiceId, CacheMessageType_Connect, writer);
-				writer.WriteU32(CacheNetworkVersion);
-				StackBinaryReader<1024> reader;
-				if (!msg.Send(reader))
+				u32 retryCount = 0;
+				while (retryCount < 10)
 				{
-					m_logger.Error(TC("Failed to connect to cache server. Version mismatch?"));
-					return;
+					StackBinaryWriter<1024> writer;
+					NetworkMessage msg(m_client, CacheServiceId, CacheMessageType_Connect, writer);
+					writer.WriteU32(CacheNetworkVersion);
+					StackBinaryReader<1024> reader;
+					if (!msg.Send(reader))
+					{
+						m_logger.Info(TC("Failed to connect to cache server. Version mismatch?"));
+						return;
+					}
+					bool success = reader.ReadBool();
+					if (success)
+					{
+						if (retryCount != 0)
+							m_logger.Info(TC("Connected to cache server"));
+						m_connected = true;
+						return;
+					}
+
+					if (retryCount == 0)
+					{
+						StringBuffer<> reason;
+						reader.ReadString(reason);
+						m_logger.Info(TC("Cache server busy, retrying... (Reason: %s)"), reason.data);
+					}
+					Sleep(1000);
+					++retryCount;
 				}
-				m_connected = true;
+
+				m_logger.Info(TC("Failed to connect to cache server after %u retries. Giving up."), retryCount);
+
 			});
 
 		m_client.RegisterOnDisconnected([this]()
@@ -79,7 +102,7 @@ namespace uba
 
 	CacheClient::~CacheClient() = default;
 
-	bool CacheClient::WriteToCache(const RootPaths& rootPaths, u32 bucketId, const ProcessStartInfo& info, const u8* inputs, u64 inputsSize, const u8* outputs, u64 outputsSize)
+	bool CacheClient::WriteToCache(const RootPaths& rootPaths, u32 bucketId, const ProcessStartInfo& info, const u8* inputs, u64 inputsSize, const u8* outputs, u64 outputsSize, u32 processId)
 	{
 		if (!m_connected)
 			return false;
@@ -96,6 +119,12 @@ namespace uba
 			return false;
 		}
 
+		bool finished = false;
+		u64 bytesSent = 0;
+		if (processId)
+			m_session.GetTrace().CacheBeginWrite(processId);
+		auto tg = MakeGuard([&]() { if (processId) m_session.GetTrace().CacheEndWrite(processId, finished, bytesSent); });
+
 		BinaryReader inputsReader(inputs, 0, inputsSize);
 		BinaryReader outputsReader(outputs, 0, outputsSize);
 
@@ -145,8 +174,7 @@ namespace uba
 					success = false;
 					continue;
 				}
-
-				casKey = AsCompressed(casKey, true);
+				casKey = IsNormalized(casKey) ? AsCompressed(casKey, true) : CasKeyZero;
 			}
 			else if (path[path.count-1] == ':')
 			{
@@ -194,14 +222,6 @@ namespace uba
 				continue;
 			}
 
-			// Files that needs to be normalized
-			if (isOutput && ShouldNormalize(path))
-			{
-				casKey = rootPaths.NormalizeAndHashFile(m_logger, path.data);
-				if (casKey != CasKeyZero)
-					casKey = AsCompressed(casKey, true);
-			}
-
 			// Get file caskey using storage
 			if (casKey == CasKeyZero)
 			{
@@ -236,7 +256,7 @@ namespace uba
 			return false;
 
 		// actual cache entry now when we know server has the needed tables
-		if (!SendCacheEntry(bucket, rootPaths, cmdKey, inputsStringToCasKey, outputsStringToCasKey))
+		if (!SendCacheEntry(bucket, rootPaths, cmdKey, inputsStringToCasKey, outputsStringToCasKey, bytesSent))
 			return false;
 
 
@@ -262,6 +282,7 @@ namespace uba
 		m_logger.EndScope();
 		#endif
 
+		finished = true;
 		return true;
 	}
 
@@ -359,7 +380,7 @@ namespace uba
 							return false;
 						UBA_ASSERTF(IsCompressed(cacheCasKey), TC("Cache entry for %s has uncompressed cache key for path %s (%s)"), info.description, path.data, CasKeyString(cacheCasKey).str);
 
-						if (ShouldNormalize(path)) // Need to normalize caskey for these files since they contain absolute paths
+						if (IsNormalized(cacheCasKey)) // Need to normalize caskey for these files since they contain absolute paths
 						{
 							localCasKey = rootPaths.NormalizeAndHashFile(m_logger, path.data);
 							if (localCasKey != CasKeyZero)
@@ -418,7 +439,7 @@ namespace uba
 				FileFetcher fetcher { m_storage.m_bufferSlots };
 				fetcher.m_errorOnFail = false;
 
-				if (ShouldNormalize(path))
+				if (IsNormalized(casKey))
 				{
 					// Fetch into memory, file is in special format without absolute paths
 					MemoryBlock normalizedBlock(4*1024*1024);
@@ -426,61 +447,52 @@ namespace uba
 					if (!fetcher.RetrieveFile(m_logger, m_client, casKey, path.data, destinationIsCompressed, &normalizedBlock))
 						return false;
 
-					u8* memoryToWrite = normalizedBlock.memory;
-					u64 memoryToWriteSize = normalizedBlock.writtenSize;
 					MemoryBlock localBlock(4*1024*1024);
 
-					// It could be that file was actually not normalized.. because it had no absolute paths, in that case we don't have the header and can just write it to disk as is
-					if (normalizedBlock.writtenSize > 6 && *(u16*)normalizedBlock.memory == 0xFFFF)
-					{
-						u32 rootOffsets = *(u32*)(normalizedBlock.memory + sizeof(u16));
-						char* fileStart = (char*)(normalizedBlock.memory + sizeof(u16) + sizeof(u32));
-						UBA_ASSERT(rootOffsets <= normalizedBlock.writtenSize);
-
-						// "denormalize" fetched file into another memory block that will be written to disk
-						u64 lastWritten = 0;
-						BinaryReader reader2(normalizedBlock.memory, rootOffsets, normalizedBlock.writtenSize);
-						while (reader2.GetLeft())
-						{
-							u64 rootOffset = reader2.Read7BitEncoded();
-							if (u64 toWrite = rootOffset - lastWritten)
-								memcpy(localBlock.Allocate(toWrite, 1, TC("")), fileStart + lastWritten, toWrite);
-							u8 rootIndex = fileStart[rootOffset] - RootPaths::RootStartByte;
-							auto& root = rootPaths.GetRoot(rootIndex);
-
-							#if PLATFORM_WINDOWS
-							StringBuffer<> pathTemp;
-							pathTemp.Append(root.path);
-							char rootPath[512];
-							u32 rootPathLen = pathTemp.Parse(rootPath, sizeof_array(rootPath));
-							#else
-							const char* rootPath = root.path.data();
-							u32 rootPathLen = root.path.size();
-							#endif
-
-							if (u32 toWrite = rootPathLen - 1)
-								memcpy(localBlock.Allocate(toWrite, 1, TC("")), rootPath, toWrite);
-							lastWritten = rootOffset + 1;
-						}
+					u32 rootOffsets = *(u32*)(normalizedBlock.memory);
+					char* fileStart = (char*)(normalizedBlock.memory + sizeof(u32));
+					UBA_ASSERT(rootOffsets <= normalizedBlock.writtenSize);
 
-						u64 fileSize = rootOffsets - (sizeof(u16) + sizeof(u32));
-						if (u64 toWrite = fileSize - lastWritten)
+					// "denormalize" fetched file into another memory block that will be written to disk
+					u64 lastWritten = 0;
+					BinaryReader reader2(normalizedBlock.memory, rootOffsets, normalizedBlock.writtenSize);
+					while (reader2.GetLeft())
+					{
+						u64 rootOffset = reader2.Read7BitEncoded();
+						if (u64 toWrite = rootOffset - lastWritten)
 							memcpy(localBlock.Allocate(toWrite, 1, TC("")), fileStart + lastWritten, toWrite);
-
-						memoryToWrite = localBlock.memory;
-						memoryToWriteSize = localBlock.writtenSize;
+						u8 rootIndex = fileStart[rootOffset] - RootPaths::RootStartByte;
+						auto& root = rootPaths.GetRoot(rootIndex);
+
+						#if PLATFORM_WINDOWS
+						StringBuffer<> pathTemp;
+						pathTemp.Append(root.path);
+						char rootPath[512];
+						u32 rootPathLen = pathTemp.Parse(rootPath, sizeof_array(rootPath));
+						#else
+						const char* rootPath = root.path.data();
+						u32 rootPathLen = root.path.size();
+						#endif
+
+						if (u32 toWrite = rootPathLen - 1)
+							memcpy(localBlock.Allocate(toWrite, 1, TC("")), rootPath, toWrite);
+						lastWritten = rootOffset + 1;
 					}
 
+					u64 fileSize = rootOffsets - sizeof(u32);
+					if (u64 toWrite = fileSize - lastWritten)
+						memcpy(localBlock.Allocate(toWrite, 1, TC("")), fileStart + lastWritten, toWrite);
+
 					FileAccessor destFile(m_logger, path.data);
 					if (!destFile.CreateWrite())
 						return false;
-					if (!destFile.Write(memoryToWrite, memoryToWriteSize))
+					if (!destFile.Write(localBlock.memory, localBlock.writtenSize))
 						return false;
 					if (!destFile.Close(&fetcher.lastWritten))
 						return false;
 
-					fetcher.sizeOnDisk = memoryToWriteSize;
-					casKey = CalculateCasKey(memoryToWrite, memoryToWriteSize, false, nullptr);
+					fetcher.sizeOnDisk = localBlock.writtenSize;
+					casKey = CalculateCasKey(localBlock.memory, localBlock.writtenSize, false, nullptr);
 				}
 				else
 				{
@@ -519,22 +531,57 @@ namespace uba
 		return reader.ReadBool();
 	}
 
-	bool CacheClient::WriteCacheSummary(const tchar* destinationFile, const tchar* filterString)
+	bool CacheClient::ExecuteCommand(Logger& logger, const tchar* command, const tchar* destinationFile, const tchar* additionalInfo)
 	{
 		StackBinaryWriter<1024> writer;
-		NetworkMessage msg(m_client, CacheServiceId, CacheMessageType_CreateStatusFile, writer);
-		writer.WriteString(filterString ? filterString : TC(""));
-		StackBinaryReader<512> reader;
-		if (!msg.Send(reader))
-			return false;
-		CasKey statusFileCasKey = reader.ReadCasKey();
-		if (statusFileCasKey == CasKeyZero)
-			return false;
+		NetworkMessage msg(m_client, CacheServiceId, CacheMessageType_ExecuteCommand, writer);
+		writer.WriteString(command);
+		writer.WriteString(additionalInfo ? additionalInfo : TC(""));
+
+		CasKey statusFileCasKey;
+		{
+			StackBinaryReader<512> reader;
+			if (!msg.Send(reader))
+				return false;
+			statusFileCasKey = reader.ReadCasKey();
+			if (statusFileCasKey == CasKeyZero)
+				return false;
+		}
 
 		FileFetcher fetcher { m_storage.m_bufferSlots };
 		bool destinationIsCompressed = false;
-		if (!fetcher.RetrieveFile(m_logger, m_client, statusFileCasKey, destinationFile, destinationIsCompressed))
-			return false;
+		if (destinationFile)
+		{
+			if (!fetcher.RetrieveFile(m_logger, m_client, statusFileCasKey, destinationFile, destinationIsCompressed))
+				return false;
+		}
+		else
+		{
+			MemoryBlock block(4*1024*1024);
+			if (!fetcher.RetrieveFile(m_logger, m_client, statusFileCasKey, TC("CommandString"), destinationIsCompressed, &block))
+				return false;
+			BinaryReader reader(block.memory, 3, block.writtenSize); // Skipping bom
+
+			tchar line[1024];
+			tchar* it = line;
+			while (true)
+			{
+				tchar c = reader.ReadUtf8Char<tchar>();
+				if (c != '\n' && c != 0)
+				{
+					*it++ = c;
+					continue;
+				}
+
+				if (c == 0 && it == line)
+					break;
+				*it = 0;
+				logger.Log(LogEntryType_Info, line, u32(it - line));
+				it = line;
+				if (c == 0)
+					break;
+			}
+		}
 		return true;
 	}
 
@@ -586,7 +633,7 @@ namespace uba
 		return true;
 	}
 
-	bool CacheClient::SendCacheEntry(Bucket& bucket, const RootPaths& rootPaths, const CasKey& cmdKey, const Map<u32, u32>& inputsStringToCasKey, const Map<u32, u32>& outputsStringToCasKey)
+	bool CacheClient::SendCacheEntry(Bucket& bucket, const RootPaths& rootPaths, const CasKey& cmdKey, const Map<u32, u32>& inputsStringToCasKey, const Map<u32, u32>& outputsStringToCasKey, u64& outBytesSent)
 	{
 		StackBinaryReader<1024> reader;
 		{
@@ -626,6 +673,7 @@ namespace uba
 			StorageImpl::CasEntry* casEntry;
 			if (m_storage.HasCasFile(casKey, &casEntry))
 			{
+				UBA_ASSERT(!IsNormalized(casKey));
 				StringBuffer<> casKeyFileName;
 				if (!m_storage.GetCasFileName(casKeyFileName, casKey))
 					return false;
@@ -654,69 +702,59 @@ namespace uba
 
 				if (!SendFile(m_logger, m_client, casKey, fileData, fileSize, casKeyFileName.data))
 					return false;
+
+				outBytesSent += fileSize;
 			}
 			else // If we don't have the cas key it should be one of the normalized files.... otherwise there is a bug
 			{
-				if (ShouldNormalize(path))
-				{
-					FileAccessor file(m_logger, path.data);
-					if (!file.OpenMemoryRead())
-						return false;
-					MemoryBlock block(AlignUp(file.GetSize() + 16, 64*1024));
-					*(u16*)block.Allocate(sizeof(u16), 1, TC("")) = 0xFFFF; // Magic to be able to know if file was normalized or not
+				if (!IsNormalized(casKey))
+					return m_logger.Error(TC("Can't find output file %s to send to cache server"), path.data);
 
-					u32& rootOffsetsStart = *(u32*)block.Allocate(sizeof(u32), 1, TC(""));
-					rootOffsetsStart = 0;
-					Vector<u32> rootOffsets;
-					u32 rootOffsetsSize = 0;
+				FileAccessor file(m_logger, path.data);
+				if (!file.OpenMemoryRead())
+					return false;
+				MemoryBlock block(AlignUp(file.GetSize() + 16, 64*1024));
+				u32& rootOffsetsStart = *(u32*)block.Allocate(sizeof(u32), 1, TC(""));
+				rootOffsetsStart = 0;
+				Vector<u32> rootOffsets;
+				u32 rootOffsetsSize = 0;
 
-					bool wasNormalized = false;
-					auto handleString = [&](const char* str, u64 strLen, u32 rootPos)
+				auto handleString = [&](const char* str, u64 strLen, u32 rootPos)
+					{
+						void* mem = block.Allocate(strLen, 1, TC(""));
+						memcpy(mem, str, strLen);
+						if (rootPos != ~0u)
 						{
-							void* mem = block.Allocate(strLen, 1, TC(""));
-							memcpy(mem, str, strLen);
-							if (rootPos != ~0u)
-							{
-								wasNormalized = true;
-								rootOffsets.push_back(rootPos);
-								rootOffsetsSize += Get7BitEncodedCount(rootPos);
-							}
-						};
-
-					if (!rootPaths.NormalizeString<char>(m_logger, (const char*)file.GetData(), file.GetSize(), handleString, path.data))
-						return false;
+							rootOffsets.push_back(rootPos);
+							rootOffsetsSize += Get7BitEncodedCount(rootPos);
+						}
+					};
 
-					if (rootOffsetsSize)
-					{
-						u8* mem = (u8*)block.Allocate(rootOffsetsSize, 1, TC(""));
-						rootOffsetsStart = u32(mem - block.memory);
-						BinaryWriter writer(mem, 0, rootOffsetsSize);
-						for (u32 rootOffset : rootOffsets)
-							writer.Write7BitEncoded(rootOffset);
-					}
-					else
-						rootOffsetsStart = u32(block.writtenSize);
+				if (!rootPaths.NormalizeString<char>(m_logger, (const char*)file.GetData(), file.GetSize(), handleString, path.data))
+					return false;
+
+				if (rootOffsetsSize)
+				{
+					u8* mem = (u8*)block.Allocate(rootOffsetsSize, 1, TC(""));
+					rootOffsetsStart = u32(mem - block.memory);
+					BinaryWriter writer(mem, 0, rootOffsetsSize);
+					for (u32 rootOffset : rootOffsets)
+						writer.Write7BitEncoded(rootOffset);
+				}
+				else
+					rootOffsetsStart = u32(block.writtenSize);
 
 
-					auto& s = m_storage;
-					FileSender sender { m_logger, m_client, s.m_bufferSlots, s.Stats(), m_sendOneAtTheTimeLock, s.m_casCompressor, s.m_casCompressionLevel };
+				auto& s = m_storage;
+				FileSender sender { m_logger, m_client, s.m_bufferSlots, s.Stats(), m_sendOneAtTheTimeLock, s.m_casCompressor, s.m_casCompressionLevel };
 
-					u8* dataToSend = block.memory;
-					u64 sizeToSend = block.writtenSize;
-					if (!wasNormalized)
-					{
-						// magic + rootoffsetssize
-						dataToSend += 6;
-						sizeToSend -= 6;
-					}
+				u8* dataToSend = block.memory;
+				u64 sizeToSend = block.writtenSize;
 
-					if (!sender.SendFileCompressed(casKey, path.data, dataToSend, sizeToSend, TC("SendCacheEntry")))
-						return m_logger.Error(TC("Failed to send cas content for file %s"), path.data);
-				}
-				else
-				{
-					return m_logger.Error(TC("Can't find output file %s to send to cache server"), path.data);
-				}
+				if (!sender.SendFileCompressed(casKey, path.data, dataToSend, sizeToSend, TC("SendCacheEntry")))
+					return m_logger.Error(TC("Failed to send cas content for file %s"), path.data);
+
+				outBytesSent += sender.m_bytesSent;
 			}
 
 		}
diff --git a/Engine/Source/Programs/UnrealBuildAccelerator/Common/Private/UbaCacheServer.cpp b/Engine/Source/Programs/UnrealBuildAccelerator/Common/Private/UbaCacheServer.cpp
--- a/Engine/Source/Programs/UnrealBuildAccelerator/Common/Private/UbaCacheServer.cpp
+++ b/Engine/Source/Programs/UnrealBuildAccelerator/Common/Private/UbaCacheServer.cpp
@@ -10,7 +10,7 @@
 
 namespace uba
 {
-	static constexpr u32 CacheFileVersion = 2;
+	static constexpr u32 CacheFileVersion = 3;
 
 	bool IsCaseInsensitive(u64 id) { return (id & (1ull << 32)) == 0; }
 
@@ -73,6 +73,8 @@ namespace uba
 	,	m_server(server)
 	,	m_storage(storage)
 	{
+		m_startTime = GetTime();
+
 		m_rootDir.count = GetFullPathNameW(rootDir, m_rootDir.capacity, m_rootDir.data, NULL);
 		m_rootDir.Replace('/', PathSeparator).EnsureEndsWithSlash();
 
@@ -182,7 +184,7 @@ namespace uba
 		if (m_addsSinceMaintenance == 0)
 			return true;
 
-		SCOPED_WRITE_LOCK(m_maintenanceLock, lock);
+		SCOPED_WRITE_LOCK(m_connectionsLock, lock);
 		return SaveNoLock();
 	}
 
@@ -288,14 +290,20 @@ namespace uba
 
 	bool CacheServer::RunMaintenance(bool force, const Function<bool()>& shouldExit)
 	{
-		SCOPED_WRITE_LOCK(m_maintenanceLock, lock);
-		SCOPED_READ_LOCK(m_connectionsLock, lock2);
+		if (m_addsSinceMaintenance == 0 && !force)
+			return true;
+
+		SCOPED_WRITE_LOCK(m_connectionsLock, lock2);
 		if (!m_connections.empty())
 			return true;
+		m_isRunningMaintenance = true;
 		lock2.Leave();
 
-		if (m_addsSinceMaintenance == 0 && !force)
-			return true;
+		auto g = MakeGuard([&]()
+			{
+				SCOPED_WRITE_LOCK(m_connectionsLock, lock2);
+				m_isRunningMaintenance = false;
+			});
 
 		bool forceAllSteps = true;
 
@@ -628,8 +636,10 @@ namespace uba
 		}
 
 		u64 oldestTime = MsToTime(GetFileTimeAsSeconds(now - oldest)*1000);
-		m_logger.Info(TC("Maintenance done! (%s) CasFiles: %llu (%s) OldestEntry: %s"), TimeToText(GetTime() - startTime).str, totalCasCount - deletedCasCount, BytesToText(totalCasSize).str, TimeToText(oldestTime, true).str);
-
+		u64 duration = GetTime() - startTime;
+		m_logger.Info(TC("Maintenance done! (%s) CasFiles: %llu (%s) OldestEntry: %s"), TimeToText(duration).str, totalCasCount - deletedCasCount, BytesToText(totalCasSize).str, TimeToText(oldestTime, true).str);
+		
+		m_longestMaintenance = Max(m_longestMaintenance, duration);
 		return true;
 	}
 
@@ -637,7 +647,6 @@ namespace uba
 	{
 		if (!m_shutdownRequested)
 			return false;
-		SCOPED_WRITE_LOCK(m_maintenanceLock, lock);
 		SCOPED_READ_LOCK(m_connectionsLock, lock2);
 		if (!m_connections.empty() || m_addsSinceMaintenance)
 			return false;
@@ -669,6 +678,9 @@ namespace uba
 
 	bool CacheServer::HandleMessage(const ConnectionInfo& connectionInfo, u8 messageType, BinaryReader& reader, BinaryWriter& writer)
 	{
+		if (messageType != CacheMessageType_Connect && m_isRunningMaintenance)
+			return m_logger.Error(TC("Can't handle network message %s while running maintenance mode"), ToString(CacheMessageType(messageType)));
+
 		switch (messageType)
 		{
 		case CacheMessageType_Connect:
@@ -676,8 +688,15 @@ namespace uba
 			u32 clientVersion = reader.ReadU32();
 			if (clientVersion != CacheNetworkVersion)
 				return m_logger.Error(TC("Different network versions. Client: %u, Server: %u. Disconnecting"), clientVersion, CacheNetworkVersion);
-			SCOPED_READ_LOCK(m_maintenanceLock, lock);
-			SCOPED_WRITE_LOCK(m_connectionsLock, lock2);
+
+			SCOPED_WRITE_LOCK(m_connectionsLock, lock);
+			if (m_isRunningMaintenance)
+			{
+				writer.WriteBool(false);
+				writer.WriteString(TC("Running maintenance..."));
+			}
+
+			writer.WriteBool(true);
 			m_connections.try_emplace(connectionInfo.GetId());
 			return true;
 		}
@@ -731,8 +750,8 @@ namespace uba
 		case CacheMessageType_FetchEntries:
 			return HandleFetchEntries(reader, writer);
 
-		case CacheMessageType_CreateStatusFile:
-			return HandleCreateStatusFile(reader, writer);
+		case CacheMessageType_ExecuteCommand:
+			return HandleExecuteCommand(reader, writer);
 
 		case CacheMessageType_RequestShutdown:
 		{
@@ -984,12 +1003,13 @@ namespace uba
 		return true;
 	}
 
-	bool CacheServer::HandleCreateStatusFile(BinaryReader& reader, BinaryWriter& writer)
+	bool CacheServer::HandleExecuteCommand(BinaryReader& reader, BinaryWriter& writer)
 	{
-		StringBuffer<> filterString;
-		reader.ReadString(filterString);
+		StringBuffer<> command;
+		reader.ReadString(command);
 
-		SCOPED_READ_LOCK(m_maintenanceLock, lock);
+		StringBuffer<> additionalInfo;
+		reader.ReadString(additionalInfo);
 
 		StringBuffer<> tempFile(m_storage.GetTempPath());
 		Guid guid;
@@ -1000,8 +1020,8 @@ namespace uba
 		if (!file.CreateWrite())
 			return false;
 
-		bool success = true;
-		auto Write = [&](const void* data, u64 size) { success &= file.Write(data, size); };
+		bool writeSuccess = true;
+		auto Write = [&](const void* data, u64 size) { writeSuccess &= file.Write(data, size); };
 
 		u8 bom[] = {0xEF,0xBB,0xBF}; 
 		Write(bom, sizeof(bom));
@@ -1015,28 +1035,73 @@ namespace uba
 				Write(buffer, w.GetPosition());
 			};
 
-		writeLine(TC("UbaCache server summary"));
+		StringBuffer<> line;
 
-		u64 now = GetSystemTimeAsFileTime();
-
-		SCOPED_READ_LOCK(m_bucketsLock, bucketsLock);
-		for (auto& kv : m_buckets)
+		if (command.Equals(TC("content")))
 		{
-			Bucket& bucket = kv.second;
-			SCOPED_READ_LOCK(bucket.m_cacheEntryLookupLock, lock2);
+			writeLine(TC("UbaCache server summary"));
 
-			for (auto& kv2 : bucket.m_cacheEntryLookup)
+			StringBufferBase& filterString = additionalInfo;
+
+			u64 now = GetSystemTimeAsFileTime();
+
+			SCOPED_READ_LOCK(m_bucketsLock, bucketsLock);
+			for (auto& kv : m_buckets)
 			{
-				CacheEntries& entries = kv2.second;
-				SCOPED_READ_LOCK(entries.lock, lock3);
+				Bucket& bucket = kv.second;
+				SCOPED_READ_LOCK(bucket.m_cacheEntryLookupLock, lock2);
 
-				Set<u32> visibleIndices;
-				if (filterString.count)
+				for (auto& kv2 : bucket.m_cacheEntryLookup)
 				{
+					CacheEntries& entries = kv2.second;
+					SCOPED_READ_LOCK(entries.lock, lock3);
+
+					Set<u32> visibleIndices;
+					if (filterString.count)
+					{
+						u32 index = 0;
+						for (auto& entry : entries.entries)
+						{
+							auto findString = [&](const Vector<u8>& offsets)
+								{
+									BinaryReader reader2(offsets.data(), 0, offsets.size());
+									while (reader2.GetLeft())
+									{
+										u64 offset = reader2.Read7BitEncoded();
+										CasKey casKey;
+										StringBuffer<> path;
+										bucket.m_casKeyTable.GetPathAndKey(path, casKey, bucket.m_pathTable, offset);
+										if (path.Contains(filterString.data))
+											return true;
+										if (Contains(CasKeyString(casKey).str, filterString.data))
+											return true;
+									}
+									return false;
+								};
+
+							if (findString(entry.inputCasKeyOffsets) || findString(entry.outputCasKeyOffsets))
+								visibleIndices.insert(index);
+							++index;
+						}
+						if (visibleIndices.empty())
+							continue;
+					}
+
+
+					writeLine(CasKeyString(kv2.first).str);
 					u32 index = 0;
 					for (auto& entry : entries.entries)
 					{
-						auto findString = [&](const Vector<u8>& offsets)
+						if (!visibleIndices.empty() && visibleIndices.find(index) == visibleIndices.end())
+						{
+							++index;
+							continue;
+						}
+
+						u64 age = MsToTime(GetFileTimeAsSeconds(now - entry.creationTime)*1000);
+						writeLine(line.Clear().Appendf(TC("  #%u (%s ago)"), index, TimeToText(age, true).str).data);
+
+						auto writeOffsets = [&](const Vector<u8>& offsets)
 							{
 								BinaryReader reader2(offsets.data(), 0, offsets.size());
 								while (reader2.GetLeft())
@@ -1045,60 +1110,49 @@ namespace uba
 									CasKey casKey;
 									StringBuffer<> path;
 									bucket.m_casKeyTable.GetPathAndKey(path, casKey, bucket.m_pathTable, offset);
-									if (path.Contains(filterString.data))
-										return true;
-									if (Contains(CasKeyString(casKey).str, filterString.data))
-										return true;
+									writeLine(line.Clear().Appendf(TC("    %s - %s"), path.data, CasKeyString(casKey).str).data);
 								}
-								return false;
 							};
 
-						if (findString(entry.inputCasKeyOffsets) || findString(entry.outputCasKeyOffsets))
-							visibleIndices.insert(index);
+						writeLine(line.Clear().Append(TC("   Inputs:")).data);
+						writeOffsets(entry.inputCasKeyOffsets);
+						writeLine(line.Clear().Append(TC("   Outputs:")).data);
+						writeOffsets(entry.outputCasKeyOffsets);
 						++index;
 					}
-					if (visibleIndices.empty())
-						continue;
-				}
-
-
-				StringBuffer<> line;
-				writeLine(CasKeyString(kv2.first).str);
-				u32 index = 0;
-				for (auto& entry : entries.entries)
-				{
-					if (!visibleIndices.empty() && visibleIndices.find(index) == visibleIndices.end())
-					{
-						++index;
-						continue;
-					}
-
-					u64 age = MsToTime(GetFileTimeAsSeconds(now - entry.creationTime)*1000);
-					writeLine(line.Clear().Appendf(TC("  #%u (%s ago)"), index, TimeToText(age, true).str).data);
-
-					auto writeOffsets = [&](const Vector<u8>& offsets)
-						{
-							BinaryReader reader2(offsets.data(), 0, offsets.size());
-							while (reader2.GetLeft())
-							{
-								u64 offset = reader2.Read7BitEncoded();
-								CasKey casKey;
-								StringBuffer<> path;
-								bucket.m_casKeyTable.GetPathAndKey(path, casKey, bucket.m_pathTable, offset);
-								writeLine(line.Clear().Appendf(TC("    %s - %s"), path.data, CasKeyString(casKey).str).data);
-							}
-						};
-
-					writeLine(line.Clear().Append (TC("   Inputs:")).data);
-					writeOffsets(entry.inputCasKeyOffsets);
-					writeLine(line.Clear().Append (TC("   Outputs:")).data);
-					writeOffsets(entry.outputCasKeyOffsets);
-					++index;
 				}
 			}
 		}
+		else if (command.Equals(TC("status")))
+		{
+			writeLine(TC("UbaCacheServer status"));
+			writeLine(line.Clear().Appendf(TC("  Uptime: %s"), TimeToText(GetTime() - m_startTime).str).data);
+			writeLine(line.Clear().Appendf(TC("  Longest maintenance: %s"), TimeToText(m_longestMaintenance).str).data);
+			writeLine(line.Clear().Appendf(TC("  Buckets:")).data);
+			u32 index = 0;
+			for (auto& kv : m_buckets)
+			{
+				Bucket& bucket = kv.second;
+				writeLine(line.Clear().Appendf(TC("    #%u - %llu"), index++, kv.first).data);
+				writeLine(line.Clear().Appendf(TC("      PathTable: %s"), BytesToText(bucket.m_pathTable.GetSize()).str).data);
+				writeLine(line.Clear().Appendf(TC("      CasKeyTable: %s"), BytesToText(bucket.m_casKeyTable.GetSize()).str).data);
+			}
+
+			u64 totalCasSize = 0;
+			u64 totalCasCount = 0;
+			m_storage.TraverseAllCasFiles([&](const CasKey& casKey, u64 size) { ++totalCasCount; totalCasSize += size; });
+			writeLine(line.Clear().Appendf(TC("  CasDb:")).data);
+			writeLine(line.Clear().Appendf(TC("    Count: %llu"), totalCasCount).data);
+			writeLine(line.Clear().Appendf(TC("    Size: %s"), BytesToText(totalCasSize).str).data);
+		}
+		else
+		{
+			writeLine(line.Clear().Appendf(TC("Unknown command: %s"), command.data).data);
+		}
+
+		Write("", 1);
 
-		if (!success || !file.Close())
+		if (!writeSuccess || !file.Close())
 			return false;
 
 		CasKey key;
diff --git a/Engine/Source/Programs/UnrealBuildAccelerator/Common/Private/UbaRootPaths.cpp b/Engine/Source/Programs/UnrealBuildAccelerator/Common/Private/UbaRootPaths.cpp
--- a/Engine/Source/Programs/UnrealBuildAccelerator/Common/Private/UbaRootPaths.cpp
+++ b/Engine/Source/Programs/UnrealBuildAccelerator/Common/Private/UbaRootPaths.cpp
@@ -127,13 +127,7 @@ namespace uba
 		if (!NormalizeString<char>(logger, (const char*)file.GetData(), file.GetSize(), hashString, filename))
 			return CasKeyZero;
 
-		if (wasNormalized)
-		{
-			u16 magic = 0xFFFF; // We add some magic to make sure this does not match cas key of non-normalized files
-			hasher.Update(&magic, sizeof(u16));
-		}
-
-		return ToCasKey(hasher, false);
+		return AsNormalized(ToCasKey(hasher, false), wasNormalized);
 	}
 
 }
\ No newline at end of file
diff --git a/Engine/Source/Programs/UnrealBuildAccelerator/Common/Private/UbaStorageUtils.cpp b/Engine/Source/Programs/UnrealBuildAccelerator/Common/Private/UbaStorageUtils.cpp
--- a/Engine/Source/Programs/UnrealBuildAccelerator/Common/Private/UbaStorageUtils.cpp
+++ b/Engine/Source/Programs/UnrealBuildAccelerator/Common/Private/UbaStorageUtils.cpp
@@ -398,6 +398,7 @@ namespace uba
 
 		m_stats.sendCasBytesRaw += sourceSize;
 		m_stats.sendCasBytesComp += fileSize;
+		m_bytesSent = fileSize;
 
 		return true;
 	}
diff --git a/Engine/Source/Programs/UnrealBuildAccelerator/Common/Private/UbaTrace.cpp b/Engine/Source/Programs/UnrealBuildAccelerator/Common/Private/UbaTrace.cpp
--- a/Engine/Source/Programs/UnrealBuildAccelerator/Common/Private/UbaTrace.cpp
+++ b/Engine/Source/Programs/UnrealBuildAccelerator/Common/Private/UbaTrace.cpp
@@ -383,6 +383,20 @@ namespace uba
 		writer.WriteBytes(data, dataSize);
 	}
 
+	void Trace::CacheBeginWrite(u32 processId)
+	{
+		BEGIN_TRACE_ENTRY(TraceType_CacheBeginWrite);
+		writer.Write7BitEncoded(processId);
+	}
+
+	void Trace::CacheEndWrite(u32 processId, bool success, u64 bytesSent)
+	{
+		BEGIN_TRACE_ENTRY(TraceType_CacheEndWrite);
+		writer.Write7BitEncoded(processId);
+		writer.WriteBool(success);
+		writer.Write7BitEncoded(bytesSent);
+	}
+
 	TraceChannel::TraceChannel(Logger& logger) : m_logger(logger)
 	{
 	}
diff --git a/Engine/Source/Programs/UnrealBuildAccelerator/Common/Public/UbaCacheClient.h b/Engine/Source/Programs/UnrealBuildAccelerator/Common/Public/UbaCacheClient.h
--- a/Engine/Source/Programs/UnrealBuildAccelerator/Common/Public/UbaCacheClient.h
+++ b/Engine/Source/Programs/UnrealBuildAccelerator/Common/Public/UbaCacheClient.h
@@ -34,11 +34,11 @@ namespace uba
 		CacheClient(const CacheClientCreateInfo& info);
 		~CacheClient();
 
-		bool WriteToCache(const RootPaths& rootPaths, u32 bucketId, const ProcessStartInfo& info, const u8* inputs, u64 inputsSize, const u8* outputs, u64 outputsSize);
+		bool WriteToCache(const RootPaths& rootPaths, u32 bucketId, const ProcessStartInfo& info, const u8* inputs, u64 inputsSize, const u8* outputs, u64 outputsSize, u32 processId = 0);
 		bool FetchFromCache(const RootPaths& rootPaths, u32 bucketId, const ProcessStartInfo& info);
 		bool RequestServerShutdown(const tchar* reason);
 
-		bool WriteCacheSummary(const tchar* destinationFile, const tchar* filterString = nullptr);
+		bool ExecuteCommand(Logger& logger, const tchar* command, const tchar* destinationFile = nullptr, const tchar* additionalInfo = nullptr);
 
 		inline MutableLogger& GetLogger() { return m_logger; }
 		inline NetworkClient& GetClient() { return m_client; }
@@ -49,7 +49,7 @@ namespace uba
 
 		bool SendPathTable(Bucket& bucket, u32 requiredPathTableSize);
 		bool SendCasTable(Bucket& bucket, u32 requiredCasTableSize);
-		bool SendCacheEntry(Bucket& bucket, const RootPaths& rootPaths, const CasKey& cmdKey, const Map<u32, u32>& inputsStringToCasKey, const Map<u32, u32>& outputsStringToCasKey);
+		bool SendCacheEntry(Bucket& bucket, const RootPaths& rootPaths, const CasKey& cmdKey, const Map<u32, u32>& inputsStringToCasKey, const Map<u32, u32>& outputsStringToCasKey, u64& outBytesSent);
 		bool FetchCasTable(Bucket& bucket, CacheStats& stats, u32 requiredCasTableOffset);
 
 		CasKey GetCmdKey(const RootPaths& rootPaths, const ProcessStartInfo& info);
diff --git a/Engine/Source/Programs/UnrealBuildAccelerator/Common/Public/UbaCacheServer.h b/Engine/Source/Programs/UnrealBuildAccelerator/Common/Public/UbaCacheServer.h
--- a/Engine/Source/Programs/UnrealBuildAccelerator/Common/Public/UbaCacheServer.h
+++ b/Engine/Source/Programs/UnrealBuildAccelerator/Common/Public/UbaCacheServer.h
@@ -44,7 +44,7 @@ namespace uba
 		bool HandleFetchPathTable(BinaryReader& reader, BinaryWriter& writer);
 		bool HandleFetchCasTable(BinaryReader& reader, BinaryWriter& writer);
 		bool HandleFetchEntries(BinaryReader& reader, BinaryWriter& writer);
-		bool HandleCreateStatusFile(BinaryReader& reader, BinaryWriter& writer);
+		bool HandleExecuteCommand(BinaryReader& reader, BinaryWriter& writer);
 
 		MutableLogger m_logger;
 		NetworkServer& m_server;
@@ -52,8 +52,8 @@ namespace uba
 
 		StringBuffer<MaxPath> m_rootDir;
 
-		ReaderWriterLock m_maintenanceLock;
 		Atomic<u32> m_addsSinceMaintenance;
+		Atomic<bool> m_isRunningMaintenance;
 
 		ReaderWriterLock m_bucketsLock;
 		Map<u64, Bucket> m_buckets;
@@ -63,6 +63,9 @@ namespace uba
 
 		Atomic<bool> m_shutdownRequested = false;
 
+		u64 m_startTime;
+		u64 m_longestMaintenance = 0;
+
 		CacheServer(const CacheServer&) = delete;
 		CacheServer& operator=(const CacheServer&) = delete;
 	};
diff --git a/Engine/Source/Programs/UnrealBuildAccelerator/Common/Public/UbaNetwork.h b/Engine/Source/Programs/UnrealBuildAccelerator/Common/Public/UbaNetwork.h
--- a/Engine/Source/Programs/UnrealBuildAccelerator/Common/Public/UbaNetwork.h
+++ b/Engine/Source/Programs/UnrealBuildAccelerator/Common/Public/UbaNetwork.h
@@ -16,7 +16,7 @@ namespace uba
 	static constexpr u32 SystemNetworkVersion = 1339;
 	static constexpr u32 StorageNetworkVersion = 4;
 	static constexpr u32 SessionNetworkVersion = 32;
-	static constexpr u32 CacheNetworkVersion = 2;
+	static constexpr u32 CacheNetworkVersion = 3;
 
 	static constexpr u32 CachePathTableMaxSize = 32*1024*1024;
 	static constexpr u32 CacheCasKeyTableMaxSize = 32*1024*1024;
@@ -110,7 +110,7 @@ namespace uba
 		UBA_CACHE_MESSAGE(FetchPathTable) \
 		UBA_CACHE_MESSAGE(FetchCasTable) \
 		UBA_CACHE_MESSAGE(FetchEntries) \
-		UBA_CACHE_MESSAGE(CreateStatusFile) \
+		UBA_CACHE_MESSAGE(ExecuteCommand) \
 		UBA_CACHE_MESSAGE(RequestShutdown) \
 
 	enum CacheMessageType : u8
diff --git a/Engine/Source/Programs/UnrealBuildAccelerator/Common/Public/UbaRootPaths.h b/Engine/Source/Programs/UnrealBuildAccelerator/Common/Public/UbaRootPaths.h
--- a/Engine/Source/Programs/UnrealBuildAccelerator/Common/Public/UbaRootPaths.h
+++ b/Engine/Source/Programs/UnrealBuildAccelerator/Common/Public/UbaRootPaths.h
@@ -94,4 +94,19 @@ namespace uba
 			searchPos = pathStart + root->path.size();
 		}
 	}
+
+	inline bool IsNormalized(const CasKey& key)
+	{
+		UBA_ASSERT(key != CasKeyZero);
+		return (((u8*)&key)[19] & 2) == 2;
+	}
+
+	inline CasKey AsNormalized(const CasKey& key, bool normalized)
+	{
+		UBA_ASSERT(key != CasKeyZero);
+		CasKey newKey = key;
+		u8 flagField = ((u8*)&key)[19];
+		((u8*)&newKey)[19] = normalized ? (flagField | 2) : (flagField & ~2);
+		return newKey;
+	}
 }
diff --git a/Engine/Source/Programs/UnrealBuildAccelerator/Common/Public/UbaStorageUtils.h b/Engine/Source/Programs/UnrealBuildAccelerator/Common/Public/UbaStorageUtils.h
--- a/Engine/Source/Programs/UnrealBuildAccelerator/Common/Public/UbaStorageUtils.h
+++ b/Engine/Source/Programs/UnrealBuildAccelerator/Common/Public/UbaStorageUtils.h
@@ -28,9 +28,9 @@ namespace uba
 		BufferSlots& m_bufferSlots;
 		StorageStats& m_stats;
 		ReaderWriterLock& m_sendOneAtTheTimeLock;
-
 		u8 m_casCompressor = DefaultCompressor;
 		u8 m_casCompressionLevel = DefaultCompressionLevel;
+		u64 m_bytesSent = 0;
 	};
 
 	struct FileFetcher
diff --git a/Engine/Source/Programs/UnrealBuildAccelerator/Common/Public/UbaTrace.h b/Engine/Source/Programs/UnrealBuildAccelerator/Common/Public/UbaTrace.h
--- a/Engine/Source/Programs/UnrealBuildAccelerator/Common/Public/UbaTrace.h
+++ b/Engine/Source/Programs/UnrealBuildAccelerator/Common/Public/UbaTrace.h
@@ -54,6 +54,8 @@ namespace uba
 		TraceType_SessionNotification,
 		TraceType_CacheBeginFetch,
 		TraceType_CacheEndFetch,
+		TraceType_CacheBeginWrite,
+		TraceType_CacheEndWrite,
 	};
 
 	using Color = u32;
@@ -92,6 +94,8 @@ namespace uba
 
 		void CacheBeginFetch(u32 fetchId, const tchar* description);
 		void CacheEndFetch(u32 fetchId, bool success, const u8* data, u64 dataSize);
+		void CacheBeginWrite(u32 processId);
+		void CacheEndWrite(u32 processId, bool success, u64 bytesSent);
 
 		bool StopWrite(const tchar* writeFileName);
 
diff --git a/Engine/Source/Programs/UnrealBuildAccelerator/Core/Private/UbaHash.cpp b/Engine/Source/Programs/UnrealBuildAccelerator/Core/Private/UbaHash.cpp
--- a/Engine/Source/Programs/UnrealBuildAccelerator/Core/Private/UbaHash.cpp
+++ b/Engine/Source/Programs/UnrealBuildAccelerator/Core/Private/UbaHash.cpp
@@ -84,6 +84,7 @@ namespace uba
 		CasKeyHasher temp(hasher);
 		uint8_t output[BLAKE3_OUT_LEN];
 		blake3_hasher_finalize((blake3_hasher*)&temp.hasher, output, BLAKE3_OUT_LEN);
-		return AsCompressed((CasKey&)output, compressed);
+		output[19] = compressed ? 1 : 0;
+		return (CasKey&)output;
 	}
 }
diff --git a/Engine/Source/Programs/UnrealBuildAccelerator/Core/Public/UbaBinaryReaderWriter.h b/Engine/Source/Programs/UnrealBuildAccelerator/Core/Public/UbaBinaryReaderWriter.h
--- a/Engine/Source/Programs/UnrealBuildAccelerator/Core/Public/UbaBinaryReaderWriter.h
+++ b/Engine/Source/Programs/UnrealBuildAccelerator/Core/Public/UbaBinaryReaderWriter.h
@@ -68,6 +68,7 @@ namespace uba
 		inline CasKey ReadCasKey();
 		inline bool ReadBool() { return ReadByte() != 0; }
 		inline u64 Read7BitEncoded();
+		template<typename CharType> inline CharType ReadUtf8Char();
 		inline u32 PeekU32();
 		inline u64 PeekU64();
 		inline void Skip(u64 size) { m_pos += size; }
@@ -443,6 +444,34 @@ namespace uba
 		return result;
 	}
 
+	template<typename CharType>
+	CharType BinaryReader::ReadUtf8Char()
+	{
+		UBA_ASSERT_READ(1);
+		u8 a = *m_pos++;
+		if (a <= 127)
+			return CharType(a);
+
+		UBA_ASSERT_READ(1);
+		u8 b = *m_pos++;
+		if (a >= 192 && a <= 223)
+			return CharType((a-192)*64 + (b-128));
+
+		UBA_ASSERT_READ(1);
+		u8 c = *m_pos++;
+		if (a >= 224 && a <= 239)
+			return CharType((a-224)*4096 + (b-128)*64 + (c-128));
+
+		if (a >= 240 && a <= 253)
+		{
+			UBA_ASSERT(false); // Wide chars cannot exceed 16 bits
+			return CharType(~0);
+		}
+
+		UBA_ASSERT(false); // Wide chars cannot exceed 16 bits
+		return CharType(~0);
+	}
+
 	u32 BinaryReader::PeekU32()
 	{
 		return *(u32*)m_pos;
diff --git a/Engine/Source/Programs/UnrealBuildAccelerator/Core/Public/UbaHash.h b/Engine/Source/Programs/UnrealBuildAccelerator/Core/Public/UbaHash.h
--- a/Engine/Source/Programs/UnrealBuildAccelerator/Core/Public/UbaHash.h
+++ b/Engine/Source/Programs/UnrealBuildAccelerator/Core/Public/UbaHash.h
@@ -112,14 +112,15 @@ namespace uba
 	inline bool IsCompressed(const CasKey& key)
 	{
 		UBA_ASSERT(key != CasKeyZero);
-		return ((u8*)&key)[19] != 0;
+		return (((u8*)&key)[19] & 1) == 1;
 	}
 
 	inline CasKey AsCompressed(const CasKey& key, bool compressed)
 	{
 		UBA_ASSERT(key != CasKeyZero);
 		CasKey newKey = key;
-		((u8*)&newKey)[19] = compressed ? 1 : 0;
+		u8 flagField = ((u8*)&key)[19];
+		((u8*)&newKey)[19] = compressed ? (flagField | u8(1)) : (flagField & ~u8(1));
 		return newKey;
 	}
 
diff --git a/Engine/Source/Programs/UnrealBuildAccelerator/Host/Private/UbaExports.cpp b/Engine/Source/Programs/UnrealBuildAccelerator/Host/Private/UbaExports.cpp
--- a/Engine/Source/Programs/UnrealBuildAccelerator/Host/Private/UbaExports.cpp
+++ b/Engine/Source/Programs/UnrealBuildAccelerator/Host/Private/UbaExports.cpp
@@ -541,7 +541,7 @@ extern "C"
 		auto& si = process->GetStartInfo();
 		if (!si.trackInputs)
 			return false;
-		return cacheClient->WriteToCache(*rootPaths, bucket, process->GetStartInfo(), inputs, inputsSize, outputs, outputsSize);
+		return cacheClient->WriteToCache(*rootPaths, bucket, process->GetStartInfo(), inputs, inputsSize, outputs, outputsSize, process->GetId());
 	}
 
 	bool CacheClient_FetchFromCache(uba::CacheClient* cacheClient, uba::RootPaths* rootPaths, uba::u32 bucket, const uba::ProcessStartInfo& info)
diff --git a/Engine/Source/Programs/UnrealBuildAccelerator/Visualizer/Private/UbaTraceReader.cpp b/Engine/Source/Programs/UnrealBuildAccelerator/Visualizer/Private/UbaTraceReader.cpp
--- a/Engine/Source/Programs/UnrealBuildAccelerator/Visualizer/Private/UbaTraceReader.cpp
+++ b/Engine/Source/Programs/UnrealBuildAccelerator/Visualizer/Private/UbaTraceReader.cpp
@@ -922,6 +922,20 @@ namespace uba
 			memcpy(process.stats.data(), dataStart, dataEnd - dataStart);
 			break;
 		}
+		case TraceType_CacheBeginWrite:
+		{
+			u32 processId = u32(reader.Read7BitEncoded());
+			out.cacheWrites[processId].start = time;
+			break;
+		}
+		case TraceType_CacheEndWrite:
+		{
+			u32 processId = u32(reader.Read7BitEncoded());
+			TraceView::CacheWrite& write = out.cacheWrites[processId];
+			write.success = reader.ReadBool();
+			write.bytesSent = reader.Read7BitEncoded();
+			write.end = time;
+		}
 		}
 		return true;
 	}
diff --git a/Engine/Source/Programs/UnrealBuildAccelerator/Visualizer/Private/UbaTraceReader.h b/Engine/Source/Programs/UnrealBuildAccelerator/Visualizer/Private/UbaTraceReader.h
--- a/Engine/Source/Programs/UnrealBuildAccelerator/Visualizer/Private/UbaTraceReader.h
+++ b/Engine/Source/Programs/UnrealBuildAccelerator/Visualizer/Private/UbaTraceReader.h
@@ -118,13 +118,22 @@ namespace uba
 			bool operator==(const ProcessLocation& o) const { return sessionIndex == o.sessionIndex && processorIndex == o.processorIndex && processIndex == o.processIndex; }
 		};
 
+		struct CacheWrite
+		{
+			u64 start = 0;
+			u64 end = 0;
+			u64 bytesSent = 0;
+			bool success = false;
+		};
+
 		Process* GetProcess(const ProcessLocation& loc) { return &(sessions[loc.sessionIndex].processors[loc.processorIndex].processes[loc.processIndex]); }
-		void Clear() { sessions.clear(); workTracks.clear(); strings.clear(); statusMap.clear(); startTime = 0; finished = true; totalProcessActiveCount = 0; totalProcessExitedCount = 0; activeSessionCount = 0; };
+		void Clear() { sessions.clear(); workTracks.clear(); strings.clear(); statusMap.clear(); cacheWrites.clear(); startTime = 0; finished = true; totalProcessActiveCount = 0; totalProcessExitedCount = 0; activeSessionCount = 0; };
 
 		Vector<Session> sessions;
 		Vector<WorkTrack> workTracks;
 		Vector<TString> strings;
 		Map<u32, StatusUpdate> statusMap;
+		Map<u32, CacheWrite> cacheWrites;
 		u64 startTime = 0;
 		u64 frequency = 0;
 		u32 totalProcessActiveCount = 0;
diff --git a/Engine/Source/Programs/UnrealBuildAccelerator/Visualizer/Private/UbaVisualizer.cpp b/Engine/Source/Programs/UnrealBuildAccelerator/Visualizer/Private/UbaVisualizer.cpp
--- a/Engine/Source/Programs/UnrealBuildAccelerator/Visualizer/Private/UbaVisualizer.cpp
+++ b/Engine/Source/Programs/UnrealBuildAccelerator/Visualizer/Private/UbaVisualizer.cpp
@@ -1216,10 +1216,67 @@
 		m_contentWidth = ProgressRectLeft + int(TimeToS(lastStop != ~u64(0) ? lastStop : playTime) * scaleX);
 		m_contentHeight = posY - int(m_scrollPosY) + stepY + 14;
 
+		float timelineSelected = m_timelineSelected;
+
 		if (m_visibleComponents[ComponentType_Timeline])
 			PaintTimeline(hdc, clientRect);
-		if (m_visibleComponents[ComponentType_Cursor])
-			PaintCursor(hdc, clientRect);
+
+		if (m_visibleComponents[ComponentType_Cursor] && m_mouseOverWindow)
+		{
+			float timeScale = (m_horizontalScaleValue * m_zoomValue)*50.0f;
+			float startOffset = -(m_scrollPosX / timeScale);
+			POINT pos;
+			GetCursorPos(&pos);
+			ScreenToClient(m_hwnd, &pos);
+			timelineSelected = startOffset + (pos.x - ProgressRectLeft) / timeScale;
+		}
+
+		if (timelineSelected)
+		{
+			int posX = int(m_scrollPosX) + progressRect.left;
+			int left = int(posX + timelineSelected * scaleX);
+			int timelineTop = Min(posY, int(progressRect.bottom)) + 2;
+
+			// TODO: Draw line up
+			MoveToEx(hdc, left, 2, NULL);
+			LineTo(hdc, left, timelineTop);
+
+			if (timelineSelected >= 0)
+			{
+				StringBuffer<> b;
+				u32 milliseconds = u32(timelineSelected * 1000.0f);
+				u32 seconds = milliseconds / 1000;
+				milliseconds -= seconds * 1000;
+				u32 minutes = seconds / 60;
+				seconds -= minutes * 60;
+				u32 hours = minutes / 60;
+				minutes -= hours * 60;
+				if (hours)
+				{
+					b.AppendValue(hours).Append('h');
+					if (minutes < 10)
+						b.Append('0');
+				}
+				if (minutes || hours)
+				{
+					b.AppendValue(minutes).Append('m');
+					if (seconds < 10)
+						b.Append('0');
+				}
+				b.AppendValue(seconds).Append('.');
+				if (milliseconds < 100)
+					b.Append('0');
+				if (milliseconds < 10)
+					b.Append('0');
+				b.AppendValue(milliseconds);
+
+				SelectObject(hdc, m_popupFont);
+				DrawTextLogger logger(m_hwnd, hdc, m_popupFontHeight, m_tooltipBackgroundBrush);
+
+				logger.Info(L"%s", b.data);
+				logger.DrawAtPos(left + 4, timelineTop - 20);
+			}
+		}
 
 		{
 			int top = 5;
@@ -1248,6 +1305,8 @@
 			left -= 20;
 			if (m_visibleComponents[ComponentType_CpuMem])
 			{
+				SelectObject(hdc, m_font);
+
 				SetTextColor(hdc, m_cpuColor);
 				ExtTextOutW(hdc, left, top, 0, NULL, L"CPU", 3, NULL);
 				left -= 25;
@@ -1320,10 +1379,10 @@
 
 			logger.AddSpace(2);
 			logger.Info(L"  %ls", process.description.c_str());
-			logger.Info(L"  Start:     %ls", TimeToText(process.start, true).str);
-			logger.Info(L"  Duration:  %ls", TimeToText(duration, true).str);
+			logger.Info(L"  Start:    %7ls", TimeToText(process.start, true).str);
+			logger.Info(L"  Duration: %7ls", TimeToText(duration, true).str);
 			if (hasExited && process.exitCode != 0)
-				logger.Info(L"  ExitCode:  %u", process.exitCode);
+				logger.Info(L"  ExitCode: %7u", process.exitCode);
 
 			if (process.stop != ~u64(0) && !process.stats.empty())
 			{
@@ -1337,11 +1396,11 @@
 				if (process.cacheFetch)
 				{
 					if (process.returned)
-						logger.Info(L"  Cache:  Miss");
+						logger.Info(L"  Cache:       Miss");
 					else
-						logger.Info(L"  Cache:  Hit");
+						logger.Info(L"  Cache:        Hit");
 					cacheStats.Read(reader, m_traceView.version);
-					if (process.exitCode == 0)
+					if (reader.GetLeft())
 					{
 						storageStats.Read(reader);
 						systemStats.Read(reader, m_traceView.version);
@@ -1394,6 +1453,18 @@
 					systemStats.Print(logger, false, m_traceView.frequency);
 				}
 
+				auto findIt = m_traceView.cacheWrites.find(process.id);
+				if (findIt != m_traceView.cacheWrites.end())
+				{
+					TraceView::CacheWrite& write = findIt->second;
+					logger.Info(L"");
+					logger.Info(L"  -------- Cache write stats ----------");
+					logger.Info(L"  Duration                    %9s", TimeToText(write.end - write.start).str);
+					logger.Info(L"  Success                     %9s", write.success ? L"true" : L"false");
+					logger.Info(L"  Bytes sent                  %9s", BytesToText(write.bytesSent).str);
+
+				}
+
 				if (!logLines.empty())
 				{
 					logger.Info(L"");
@@ -1503,52 +1574,6 @@
 			logger.Info(L"%ls %ls", vis ? L"Hide" : L"Show", tooltip[m_buttonSelected]);
 			logger.DrawAtCursor();
 		}
-		else if (m_timelineSelected)
-		{
-			int posX = int(m_scrollPosX) + progressRect.left;
-			int left = int(posX + m_timelineSelected * scaleX);
-			int timelineTop = Min(posY, int(progressRect.bottom)) + 2;
-
-			// TODO: Draw line up
-			MoveToEx(hdc, left, 2, NULL);
-			LineTo(hdc, left, timelineTop);
-
-			if (m_timelineSelected >= 0)
-			{
-				StringBuffer<> b;
-				u32 milliseconds = u32(m_timelineSelected * 1000.0f);
-				u32 seconds = milliseconds / 1000;
-				milliseconds -= seconds * 1000;
-				u32 minutes = seconds / 60;
-				seconds -= minutes * 60;
-				u32 hours = minutes / 60;
-				minutes -= hours * 60;
-				if (hours)
-				{
-					b.AppendValue(hours).Append('h');
-					if (minutes < 10)
-						b.Append('0');
-				}
-				if (minutes || hours)
-				{
-					b.AppendValue(minutes).Append('m');
-					if (seconds < 10)
-						b.Append('0');
-				}
-				b.AppendValue(seconds).Append('.');
-				if (milliseconds < 100)
-					b.Append('0');
-				if (milliseconds < 10)
-					b.Append('0');
-				b.AppendValue(milliseconds);
-
-				SelectObject(hdc, m_popupFont);
-				DrawTextLogger logger(m_hwnd, hdc, m_popupFontHeight, m_tooltipBackgroundBrush);
-
-				logger.Info(L"%s", b.data);
-				logger.DrawAtPos(left + 4, timelineTop - 20);
-			}
-		}
 		else if (m_fetchedFilesSelected != ~0u)
 		{
 			auto& session = m_traceView.sessions[m_fetchedFilesSelected];
@@ -1761,20 +1786,6 @@
 		*/
 	}
 
-	void Visualizer::PaintCursor(HDC hdc, const RECT& clientRect)
-	{
-		// Get cursor position and transform into window coordinates
-		POINT cursorPos = {};
-		GetCursorPos(&cursorPos);
-		ScreenToClient(m_hwnd, &cursorPos);
-
-		// Draw cursor as vertical line
-		SelectObject(hdc, m_textPen);
-		const int cursorPosHorizontal = cursorPos.x;
-		MoveToEx(hdc, cursorPosHorizontal, clientRect.top, NULL);
-		LineTo(hdc, cursorPosHorizontal, clientRect.bottom);
-	}
-
 	void Visualizer::PaintDetailedStats(int& posY, const RECT& progressRect, TraceView::Session& session, bool isRemote, u64 playTime, const DrawTextFunc& drawTextFunc)
 	{
 		int stepY = FontHeight;
@@ -2224,7 +2235,7 @@
 
 	void Visualizer::UnselectAndRedraw()
 	{
-		if (Unselect())
+		if (Unselect() || m_visibleComponents[ComponentType_Cursor])
 			RedrawWindow(m_hwnd, NULL, NULL, RDW_INVALIDATE);
 	}
 
diff --git a/Engine/Source/Programs/UnrealBuildAccelerator/Visualizer/Private/UbaVisualizer.h b/Engine/Source/Programs/UnrealBuildAccelerator/Visualizer/Private/UbaVisualizer.h
--- a/Engine/Source/Programs/UnrealBuildAccelerator/Visualizer/Private/UbaVisualizer.h
+++ b/Engine/Source/Programs/UnrealBuildAccelerator/Visualizer/Private/UbaVisualizer.h
@@ -31,7 +31,6 @@ namespace uba
 		void PaintAll(HDC hdc, const RECT& clientRect);
 		void PaintProcessRect(TraceView::Process& process, HDC hdc, RECT rect, const RECT& progressRect, bool selected, bool writingBitmap);
 		void PaintTimeline(HDC hdc, const RECT& clientRect);
-		void PaintCursor(HDC hdc, const RECT& clientRect);
 		using DrawTextFunc = Function<void(const StringBufferBase& text, RECT& rect)>;
 		void PaintDetailedStats(int& posY, const RECT& progressRect, TraceView::Session& session, bool isRemote, u64 playTime, const DrawTextFunc& drawTextFunc);
 
