[Uba]
* Changed ParallelFor implementation to remove chance of thread calling ParallelFor picking up work queued by other threads

diff --git a/Engine/Source/Programs/UnrealBuildAccelerator/Common/Private/UbaCacheServer.cpp b/Engine/Source/Programs/UnrealBuildAccelerator/Common/Private/UbaCacheServer.cpp
--- a/Engine/Source/Programs/UnrealBuildAccelerator/Common/Private/UbaCacheServer.cpp
+++ b/Engine/Source/Programs/UnrealBuildAccelerator/Common/Private/UbaCacheServer.cpp
@@ -334,6 +334,8 @@ namespace uba
 		u64 now = GetSystemTimeAsFileTime();
 		u64 oldest = now;
 
+		u32 workerCount = m_server.GetWorkerCount();
+		u32 workerCountToUse = workerCount > 0 ? workerCount - 1 : 0;
 
 		Atomic<u64> deleteEntryCount;
 
@@ -344,7 +346,7 @@ namespace uba
 
 			ReaderWriterLock existingCasLock;
 
-			m_server.ParallelFor(16, m_buckets, [&](auto& it)
+			m_server.ParallelFor(workerCountToUse, m_buckets, [&](auto& it)
 			{
 				Vector<u64*> touchedCas;
 
@@ -474,7 +476,7 @@ namespace uba
 			return true;
 
 		Atomic<u32> bucketCounter;
-		m_server.ParallelFor(16, m_buckets, [&](auto& it)
+		m_server.ParallelFor(workerCountToUse, m_buckets, [&](auto& it)
 		{
 			u64 bucketStartTime = GetTime();
 
@@ -595,7 +597,7 @@ namespace uba
 						UBA_ASSERT(writer2.GetPosition() == newOffsetsSize);
 					};
 
-				m_server.ParallelFor(16, bucket.m_cacheEntryLookup, [&](auto& it)
+				m_server.ParallelFor(workerCountToUse, bucket.m_cacheEntryLookup, [&](auto& it)
 					{
 						CacheEntries& entries = it->second;(void)entries;
 						for (auto& entry : entries.entries)
diff --git a/Engine/Source/Programs/UnrealBuildAccelerator/Common/Public/UbaWorkManager.h b/Engine/Source/Programs/UnrealBuildAccelerator/Common/Public/UbaWorkManager.h
--- a/Engine/Source/Programs/UnrealBuildAccelerator/Common/Public/UbaWorkManager.h
+++ b/Engine/Source/Programs/UnrealBuildAccelerator/Common/Public/UbaWorkManager.h
@@ -2,6 +2,7 @@
 
 #pragma once
 
+#include "UbaEvent.h"
 #include "UbaMemory.h"
 
 namespace uba
@@ -74,30 +75,57 @@ namespace uba
 	template<typename TContainer, typename TFunc>
 	void WorkManager::ParallelFor(u32 workCount, TContainer& container, const TFunc& func, const tchar* description)
 	{
-		auto it = container.begin();
-		auto end = container.end();
-		ReaderWriterLock lock;
-		Atomic<u32> workLeft = workCount;
-
-		AddWork([&]()
+		struct Context
+		{
+			typename TContainer::iterator it;
+			typename TContainer::iterator end;
+			u32 refCount;
+			u32 activeCount;
+			bool isDone;
+			ReaderWriterLock lock;
+			Event* doneEvent;
+		};
+
+		Event doneEvent(true);
+
+		auto context = new Context();
+		context->it = container.begin();
+		context->end = container.end();
+		context->refCount = workCount + 1;
+		context->activeCount = 0;
+		context->isDone = false;
+		context->doneEvent = &doneEvent;
+
+		auto work = [context, &func]()
 			{
+				u32 active = 0;
 				while (true)
 				{
-					SCOPED_WRITE_LOCK(lock, l);
-					if (it == end)
+					SCOPED_WRITE_LOCK(context->lock, l);
+					context->activeCount -= active;
+					context->isDone = context->it == context->end;
+					if (context->isDone)
 					{
+						if (context->activeCount == 0)
+							context->doneEvent->Set();
+						if (--context->refCount)
+							return;
 						l.Leave();
-						--workLeft;
+						delete context;
 						return;
 					}
-					auto it2 = it++;
+
+					auto it2 = context->it++;
+					active = 1;
+					context->activeCount += active;
 					l.Leave();
 
 					func(it2);
 				}
-			}, workCount, description);
-
-		while (workLeft)
-			DoWork();
+			};
+		
+		AddWork(work, workCount, description);
+		work();
+		doneEvent.IsSet();
 	}
 }
