[Uba]
* First round of optimizations to uba cache server maintenance code.

diff --git a/Engine/Source/Programs/UnrealBuildAccelerator/CacheService/Private/UbaCacheService.cpp b/Engine/Source/Programs/UnrealBuildAccelerator/CacheService/Private/UbaCacheService.cpp
--- a/Engine/Source/Programs/UnrealBuildAccelerator/CacheService/Private/UbaCacheService.cpp
+++ b/Engine/Source/Programs/UnrealBuildAccelerator/CacheService/Private/UbaCacheService.cpp
@@ -39,7 +39,7 @@ namespace uba
 		}
 		logger.Info(TC(""));
 		logger.Info(TC("-------------------------------------------"));
-		logger.Info(TC("   UbaCacheService v%s"), Version);
+		logger.Info(TC("   UbaCacheService v%s (%u)"), Version, CacheNetworkVersion);
 		logger.Info(TC("-------------------------------------------"));
 		logger.Info(TC(""));
 		logger.Info(TC("  -dir=<rootdir>          The directory used to store data. Defaults to \"%s\""), DefaultRootDir);
@@ -182,7 +182,7 @@ namespace uba
 		#if UBA_DEBUG
 		dbgStr = TC(" (DEBUG)");
 		#endif
-		logger.Info(TC("UbaCacheService v%s%s (Rootdir: \"%s\", StoreCapacity: %uGb)\n"), Version, dbgStr, g_rootDir.data, storageCapacityGb);
+		logger.Info(TC("UbaCacheService v%s(%u)%s (Workers: %u, Rootdir: \"%s\", StoreCapacity: %uGb)\n"), Version, CacheNetworkVersion, dbgStr, GetLogicalProcessorCount(), g_rootDir.data, storageCapacityGb);
 
 		u64 storageCapacity = u64(storageCapacityGb)*1000*1000*1000;
 
diff --git a/Engine/Source/Programs/UnrealBuildAccelerator/Common/Private/UbaCacheServer.cpp b/Engine/Source/Programs/UnrealBuildAccelerator/Common/Private/UbaCacheServer.cpp
--- a/Engine/Source/Programs/UnrealBuildAccelerator/Common/Private/UbaCacheServer.cpp
+++ b/Engine/Source/Programs/UnrealBuildAccelerator/Common/Private/UbaCacheServer.cpp
@@ -6,6 +6,7 @@
 #include "UbaFileAccessor.h"
 #include "UbaNetworkServer.h"
 #include "UbaStorageServer.h"
+#include <algorithm>
 //#include <oodle2.h>
 
 namespace uba
@@ -483,113 +484,126 @@ namespace uba
 			Bucket& bucket = it->second;
 			u32 bucketIndex = bucketCounter++;
 
-			if (deleteEntryCount || forceAllSteps)
+			if (!deleteEntryCount && !forceAllSteps)
 			{
-				UnorderedSet<u32> usedCasKeyOffsets;
+				m_logger.Detail(TC("    Bucket %u skipped updating. No entries deleted"), bucketIndex);
+				return;
+			}
+			
+			MemoryBlock memoryBlock(128*1024*1024);
 
-				// Collect all caskeys that are used by cache entries.
-				for (auto& kv2 : bucket.m_cacheEntryLookup)
+			GrowingNoLockUnorderedSet<u32> usedCasKeyOffsets(&memoryBlock);
+			usedCasKeyOffsets.reserve(bucket.m_casKeyTable.GetKeyCount());
+
+			u64 collectUsedCasKeysStart = GetTime();
+
+			// Collect all caskeys that are used by cache entries.
+			for (auto& kv2 : bucket.m_cacheEntryLookup)
+			{
+				for (auto& entry : kv2.second.entries)
 				{
-					for (auto& entry : kv2.second.entries)
-					{
-						auto collectUsedCasKeyOffsets = [&](const Vector<u8>& offsets)
-							{
-								BinaryReader reader2(offsets.data(), 0, offsets.size());
-								while (reader2.GetLeft())
-									usedCasKeyOffsets.insert(u32(reader2.Read7BitEncoded()));
-							};
-						collectUsedCasKeyOffsets(entry.inputCasKeyOffsets);
-						collectUsedCasKeyOffsets(entry.outputCasKeyOffsets);
-					}
+					auto collectUsedCasKeyOffsets = [&](const Vector<u8>& offsets)
+						{
+							BinaryReader reader2(offsets.data(), 0, offsets.size());
+							while (reader2.GetLeft())
+								usedCasKeyOffsets.insert(u32(reader2.Read7BitEncoded()));
+						};
+					collectUsedCasKeyOffsets(entry.inputCasKeyOffsets);
+					collectUsedCasKeyOffsets(entry.outputCasKeyOffsets);
 				}
+			}
+			m_logger.Detail(TC("    Bucket %u Collected %llu used caskeys. (%s)"), bucketIndex, usedCasKeyOffsets.size(), TimeToText(GetTime() - collectUsedCasKeysStart).str);
 
-				u64 recreatePathTableStart = GetTime();
+			u64 recreatePathTableStart = GetTime();
 
-				// Traverse all caskeys in caskey table and figure out which ones we can delete
-				UnorderedSet<u32> usedPathOffsets;
-				usedPathOffsets.reserve(usedCasKeyOffsets.size());
+			// Traverse all caskeys in caskey table and figure out which ones we can delete
+			GrowingNoLockUnorderedSet<u32> usedPathOffsets(&memoryBlock);
+			usedPathOffsets.reserve(usedCasKeyOffsets.size());
 
-				for (u32 casKeyOffset : usedCasKeyOffsets)
-				{
-					BinaryReader reader2(bucket.m_casKeyTable.GetMemory(), casKeyOffset, bucket.m_casKeyTable.GetSize());
-					u32 pathOffset = u32(reader2.Read7BitEncoded());
-					usedPathOffsets.insert(pathOffset);
-				}
+			for (u32 casKeyOffset : usedCasKeyOffsets)
+			{
+				BinaryReader reader2(bucket.m_casKeyTable.GetMemory(), casKeyOffset, bucket.m_casKeyTable.GetSize());
+				u32 pathOffset = u32(reader2.Read7BitEncoded());
+				usedPathOffsets.insert(pathOffset);
+			}
 
-				// Build new path table based on used offsets
-				UnorderedMap<u32, u32> oldToNewPathOffset;
-				u32 oldSize = bucket.m_pathTable.GetSize();
-				{
-					CompactPathTable newPathTable(CachePathTableMaxSize, CompactPathTable::V1, bucket.m_pathTable.GetPathCount(), bucket.m_pathTable.GetSegmentCount());
-					oldToNewPathOffset.reserve(usedPathOffsets.size());
+			// Build new path table based on used offsets
+			GrowingNoLockUnorderedMap<u32, u32> oldToNewPathOffset(&memoryBlock);
+			u32 oldSize = bucket.m_pathTable.GetSize();
+			{
+				CompactPathTable newPathTable(CachePathTableMaxSize, CompactPathTable::V1, bucket.m_pathTable.GetPathCount(), bucket.m_pathTable.GetSegmentCount());
+				oldToNewPathOffset.reserve(usedPathOffsets.size());
 
-					for (u32 pathOffset : usedPathOffsets)
-					{
-						StringBuffer<> temp;
-						bucket.m_pathTable.GetString(temp, pathOffset);
-						u32 newOffset = newPathTable.AddNoLock(temp.data, temp.count);
-
-						#if 0
-						StringBuffer<> test;
-						newPathTable.GetString(test, newOffset);
-						UBA_ASSERT(test.Equals(temp.data));
-						#endif
-
-						auto res = oldToNewPathOffset.try_emplace(pathOffset, newOffset);
-						UBA_ASSERT(res.second);(void)res;
-					}
-					bucket.m_pathTable.Swap(newPathTable);
+				for (u32 pathOffset : usedPathOffsets)
+				{
+					StringBuffer<> temp;
+					bucket.m_pathTable.GetString(temp, pathOffset);
+					u32 newOffset = newPathTable.AddNoLock(temp.data, temp.count);
+
+					#if 0
+					StringBuffer<> test;
+					newPathTable.GetString(test, newOffset);
+					UBA_ASSERT(test.Equals(temp.data));
+					#endif
+
+					auto res = oldToNewPathOffset.try_emplace(pathOffset, newOffset);
+					UBA_ASSERT(res.second);(void)res;
 				}
-				m_logger.Detail(TC("    Bucket %u Recreated path table. %s -> %s (%s)"), bucketIndex, BytesToText(oldSize).str, BytesToText(bucket.m_pathTable.GetSize()).str, TimeToText(GetTime() - recreatePathTableStart).str);
+				bucket.m_pathTable.Swap(newPathTable);
+			}
+			m_logger.Detail(TC("    Bucket %u Recreated path table. %s -> %s (%s)"), bucketIndex, BytesToText(oldSize).str, BytesToText(bucket.m_pathTable.GetSize()).str, TimeToText(GetTime() - recreatePathTableStart).str);
 
 
-				// Build new caskey table based on used offsets
-				u64 recreateCasKeyTableStart = GetTime();
-				UnorderedMap<u32, u32> oldToNewCasKeyOffset;
-				oldSize = bucket.m_casKeyTable.GetSize();
+			// Build new caskey table based on used offsets
+			u64 recreateCasKeyTableStart = GetTime();
+			GrowingNoLockUnorderedMap<u32, u32> oldToNewCasKeyOffset(&memoryBlock);
+			oldSize = bucket.m_casKeyTable.GetSize();
+			{
+				oldToNewCasKeyOffset.reserve(usedCasKeyOffsets.size());
+				CompactCasKeyTable newCasKeyTable(CacheCasKeyTableMaxSize, usedCasKeyOffsets.size());
+				for (u32 casKeyOffset : usedCasKeyOffsets)
 				{
-					oldToNewCasKeyOffset.reserve(usedCasKeyOffsets.size());
-					CompactCasKeyTable newCasKeyTable(CacheCasKeyTableMaxSize, usedCasKeyOffsets.size());
-					for (u32 casKeyOffset : usedCasKeyOffsets)
-					{
-						BinaryReader reader2(bucket.m_casKeyTable.GetMemory(), casKeyOffset, bucket.m_casKeyTable.GetSize());
-						u32 oldPathOffset = u32(reader2.Read7BitEncoded());
-						CasKey casKey = reader2.ReadCasKey();
-						auto findIt = oldToNewPathOffset.find(oldPathOffset);
-						UBA_ASSERT(findIt != oldToNewPathOffset.end());
-						auto res = oldToNewCasKeyOffset.try_emplace(casKeyOffset, newCasKeyTable.Add(casKey, findIt->second));
-						UBA_ASSERT(res.second);(void)res;
-					}
-					bucket.m_casKeyTable.Swap(newCasKeyTable);
+					BinaryReader reader2(bucket.m_casKeyTable.GetMemory(), casKeyOffset, bucket.m_casKeyTable.GetSize());
+					u32 oldPathOffset = u32(reader2.Read7BitEncoded());
+					CasKey casKey = reader2.ReadCasKey();
+					auto findIt = oldToNewPathOffset.find(oldPathOffset);
+					UBA_ASSERT(findIt != oldToNewPathOffset.end());
+					u32 newCasKeyOffset = newCasKeyTable.Add(casKey, findIt->second);
+					if (casKeyOffset == newCasKeyOffset)
+						continue;
+					auto res = oldToNewCasKeyOffset.try_emplace(casKeyOffset, newCasKeyOffset);
+					UBA_ASSERT(res.second);(void)res;
 				}
-				m_logger.Detail(TC("    Bucket %u Recreated caskey table. %s -> %s (%s)"), bucketIndex, BytesToText(oldSize).str, BytesToText(bucket.m_casKeyTable.GetSize()).str, TimeToText(GetTime() - recreateCasKeyTableStart).str);
+				bucket.m_casKeyTable.Swap(newCasKeyTable);
+			}
+			m_logger.Detail(TC("    Bucket %u Recreated caskey table. %s -> %s (%s)"), bucketIndex, BytesToText(oldSize).str, BytesToText(bucket.m_casKeyTable.GetSize()).str, TimeToText(GetTime() - recreateCasKeyTableStart).str);
 
 
+			if (!oldToNewCasKeyOffset.empty())
+			{
 				// Update all casKeyOffsets
 				u64 updateEntriesStart = GetTime();
 
-				auto updateCasKeyOffsets = [&](Vector<u8>& offsets)
+				auto updateCasKeyOffsets = [&](Vector<u8>& offsets, Vector<u32>& newOffsets)
 					{
-						Set<u32> newOffsets; // Want these sorted since that is what we get from uploads
+						newOffsets.clear();
+						newOffsets.reserve(offsets.size()*4);
+
 						u32 newOffsetsSize = 0;
 						BinaryReader reader2(offsets.data(), 0, offsets.size());
 						while (reader2.GetLeft())
 						{
 							u32 oldOffset = u32(reader2.Read7BitEncoded());
+							u32 newOffset = oldOffset;
 							auto findIt = oldToNewCasKeyOffset.find(oldOffset);
-							UBA_ASSERT(findIt != oldToNewCasKeyOffset.end());
-							u32 newOffset = findIt->second;
-							if (newOffsets.insert(newOffset).second)
-								newOffsetsSize += Get7BitEncodedCount(newOffset);
-							else
-							{
-								StringBuffer<> temp;
-								CasKey key;
-								bucket.m_casKeyTable.GetPathAndKey(temp, key, bucket.m_pathTable, newOffset);
-								m_logger.Info(TC("Found duplicate of %s in cache entry"), temp.data);
-							}
+							if (findIt != oldToNewCasKeyOffset.end())
+								newOffset = findIt->second;
+							newOffsets.push_back(newOffset);
+							newOffsetsSize += Get7BitEncodedCount(newOffset);
 						}
 
+						std::sort(newOffsets.begin(), newOffsets.end());
+
 						offsets.resize(newOffsetsSize);
 						BinaryWriter writer2(offsets.data(), 0, newOffsetsSize);
 						for (u32 offset : newOffsets)
@@ -599,11 +613,12 @@ namespace uba
 
 				m_server.ParallelFor(workerCountToUse, bucket.m_cacheEntryLookup, [&](auto& it)
 					{
+						Vector<u32> newOffsets;
 						CacheEntries& entries = it->second;(void)entries;
 						for (auto& entry : entries.entries)
 						{
-							updateCasKeyOffsets(entry.inputCasKeyOffsets);
-							updateCasKeyOffsets(entry.outputCasKeyOffsets);
+							updateCasKeyOffsets(entry.inputCasKeyOffsets, newOffsets);
+							updateCasKeyOffsets(entry.outputCasKeyOffsets, newOffsets);
 						}
 					});
 
@@ -876,7 +891,7 @@ namespace uba
 				}
 				if (findIt->second != casKey)
 				{
-					m_logger.Warning(TC("Existing cache entry matches input but does not match output (%s has different caskey)"), path.data);
+					//m_logger.Warning(TC("Existing cache entry matches input but does not match output (%s has different caskey)"), path.data);
 					cacheEntries.entries.erase(matchingEntry);
 					shouldOverwrite = true;
 					break;
diff --git a/Engine/Source/Programs/UnrealBuildAccelerator/Common/Public/UbaCompactTables.h b/Engine/Source/Programs/UnrealBuildAccelerator/Common/Public/UbaCompactTables.h
--- a/Engine/Source/Programs/UnrealBuildAccelerator/Common/Public/UbaCompactTables.h
+++ b/Engine/Source/Programs/UnrealBuildAccelerator/Common/Public/UbaCompactTables.h
@@ -60,6 +60,8 @@ namespace uba
 		void ReadMem(BinaryReader& reader, bool populateLookup);
 		void Swap(CompactCasKeyTable& other);
 
+		u64 GetKeyCount() { return m_offsets.size(); }
+
 	private:
 		ReaderWriterLock m_lock;
 		MemoryBlock m_mem;
diff --git a/Engine/Source/Programs/UnrealBuildAccelerator/Core/Public/UbaMemory.h b/Engine/Source/Programs/UnrealBuildAccelerator/Core/Public/UbaMemory.h
--- a/Engine/Source/Programs/UnrealBuildAccelerator/Core/Public/UbaMemory.h
+++ b/Engine/Source/Programs/UnrealBuildAccelerator/Core/Public/UbaMemory.h
@@ -146,9 +146,47 @@ namespace uba
 		MemoryBlock* m_block;
 	};
 
+	template<typename Type>
+	class GrowingAllocatorNoLock
+	{
+	public:
+		using value_type = Type;
+
+		GrowingAllocatorNoLock(MemoryBlock* block) : m_block(block) {}
+		GrowingAllocatorNoLock(const GrowingAllocatorNoLock& o) : m_block(o.m_block) {}
+		GrowingAllocatorNoLock(GrowingAllocatorNoLock&& o) noexcept : m_block(o.m_block) {}
+		template <class _Other>
+		constexpr GrowingAllocatorNoLock(const GrowingAllocatorNoLock<_Other>& o) noexcept : m_block(o.m_block) {}
+
+		value_type* allocate(u64 n)
+		{
+			return (value_type*)m_block->AllocateNoLock(sizeof(value_type)*n, alignof(value_type), TC("GrowingAllocatorNoLock"));
+		}
+
+		/// @warning Naive implementation, assumes `p` is valid.
+		void deallocate(value_type*, u64)
+		{
+		}
+
+		u64 max_size() const
+		{
+			return static_cast<size_t>(-1) / sizeof(value_type);
+		}
+	
+		bool operator==(const GrowingAllocatorNoLock& o) const { return m_block == o.m_block; }
+
+		MemoryBlock* m_block;
+	};
 
 	template<typename Key, typename Value, typename Hash = std::hash<Key>, typename EqualTo = std::equal_to<Key>>
 	using GrowingUnorderedMap = std::unordered_map<Key, Value, Hash, EqualTo, GrowingAllocator<std::pair<const Key, Value>>>;
+	template<typename Key, typename Hash = std::hash<Key>, typename EqualTo = std::equal_to<Key>>
+	using GrowingUnorderedSet = std::unordered_set<Key, Hash, EqualTo, GrowingAllocator<Key>>;
+
+	template<typename Key, typename Value, typename Hash = std::hash<Key>, typename EqualTo = std::equal_to<Key>>
+	using GrowingNoLockUnorderedMap = std::unordered_map<Key, Value, Hash, EqualTo, GrowingAllocatorNoLock<std::pair<const Key, Value>>>;
+	template<typename Key, typename Hash = std::hash<Key>, typename EqualTo = std::equal_to<Key>>
+	using GrowingNoLockUnorderedSet = std::unordered_set<Key, Hash, EqualTo, GrowingAllocatorNoLock<Key>>;
 
 	template<typename Type>
 	struct BlockAllocator
