[UBA]
* Fixed fallback code for not finding caskey for file to compare caskeys with matching compression flags

diff --git a/Engine/Source/Programs/UnrealBuildAccelerator/Common/Private/UbaStorageServer.cpp b/Engine/Source/Programs/UnrealBuildAccelerator/Common/Private/UbaStorageServer.cpp
--- a/Engine/Source/Programs/UnrealBuildAccelerator/Common/Private/UbaStorageServer.cpp
+++ b/Engine/Source/Programs/UnrealBuildAccelerator/Common/Private/UbaStorageServer.cpp
@@ -428,7 +428,7 @@ namespace uba
 							}
 						}
 
-						if (checkedCasKey != casKey)
+						if (AsCompressed(checkedCasKey, m_storeCompressed) != casKey)
 						{
 							m_logger.Error(TC("FetchBegin failed for cas file %s (%s). File on disk has different cas %s"), CasKeyString(casKey).str, hint.data, CasKeyString(checkedCasKey).str);
 							writer.WriteU16(0);
