[UBA]
* Changed so UbaCacheServer stores ranges of input offsets for entries in the same key. It is likely we have many matching inputs between entries in the same key
* Added database creation time and store all times in relation to creation time
* Added lastUsedTime to entries to be able to add support for cleaning up entries that have not been used for x time (not implemented yet)

diff --git a/Engine/Source/Programs/UnrealBuildAccelerator/Common/Private/UbaCacheClient.cpp b/Engine/Source/Programs/UnrealBuildAccelerator/Common/Private/UbaCacheClient.cpp
--- a/Engine/Source/Programs/UnrealBuildAccelerator/Common/Private/UbaCacheClient.cpp
+++ b/Engine/Source/Programs/UnrealBuildAccelerator/Common/Private/UbaCacheClient.cpp
@@ -2,6 +2,7 @@
 
 #include "UbaCacheClient.h"
 #include "UbaApplicationRules.h"
+#include "UbaCacheEntry.h"
 #include "UbaCompactTables.h"
 #include "UbaFileAccessor.h"
 #include "UbaNetworkMessage.h"
@@ -354,161 +355,171 @@ namespace uba
 		struct MissInfo { TString path; u32 entryIndex; CasKey cache; CasKey local; };
 		Vector<MissInfo> misses;
 
-		for (u32 i=0; i!=entryCount; ++i)
+		CacheEntriesTraverser traverser(reader);
+
+		u32 entryIndex = 0;
+		for (; entryIndex!=entryCount; ++entryIndex)
 		{
-			u64 outputSize = 0;
 			{
 				TimerScope ts(cacheStats.testEntries);
 				bool isMatch = true;
-				u64 inputSize = reader.Read7BitEncoded();
-				const u8* inputEnd = reader.GetPositionData() + inputSize;
-				while (reader.GetPositionData() != inputEnd)
-				{
-					u32 casKeyOffset = u32(reader.Read7BitEncoded());
-					StringBuffer<MaxPath> path;
-
-					CasKey cacheCasKey;
-					CasKey localCasKey;
 
-					auto insres = offsetIsMatch.try_emplace(casKeyOffset);
-					if (insres.second)
+				bool result = traverser.TraverseEntryInputs([&](u32 casKeyOffset)
 					{
-						if (!FetchCasTable(bucket, cacheStats, casKeyOffset))
-							return false;
+						StringBuffer<MaxPath> path;
 
-						if (!GetLocalPathAndCasKey(bucket, rootPaths, path, cacheCasKey, bucket.serverCasKeyTable, bucket.serverPathTable, casKeyOffset))
-							return false;
-						UBA_ASSERTF(IsCompressed(cacheCasKey), TC("Cache entry for %s has uncompressed cache key for path %s (%s)"), info.description, path.data, CasKeyString(cacheCasKey).str);
+						CasKey cacheCasKey;
+						CasKey localCasKey;
 
-						if (IsNormalized(cacheCasKey)) // Need to normalize caskey for these files since they contain absolute paths
+						auto insres = offsetIsMatch.try_emplace(casKeyOffset);
+						if (insres.second)
 						{
-							localCasKey = rootPaths.NormalizeAndHashFile(m_logger, path.data);
-							if (localCasKey != CasKeyZero)
-								localCasKey = AsCompressed(localCasKey, true);
+							if (!FetchCasTable(bucket, cacheStats, casKeyOffset))
+								return false;
+
+							if (!GetLocalPathAndCasKey(bucket, rootPaths, path, cacheCasKey, bucket.serverCasKeyTable, bucket.serverPathTable, casKeyOffset))
+								return false;
+							UBA_ASSERTF(IsCompressed(cacheCasKey), TC("Cache entry for %s has uncompressed cache key for path %s (%s)"), info.description, path.data, CasKeyString(cacheCasKey).str);
+
+							if (IsNormalized(cacheCasKey)) // Need to normalize caskey for these files since they contain absolute paths
+							{
+								localCasKey = rootPaths.NormalizeAndHashFile(m_logger, path.data);
+								if (localCasKey != CasKeyZero)
+									localCasKey = AsCompressed(localCasKey, true);
+							}
+							else
+							{
+								bool deferCreation = true;
+								bool fileIsCompressed = IsFileCompressed(info, path.data, path.count);
+								m_storage.StoreCasFile(localCasKey, path.data, CasKeyZero, deferCreation, fileIsCompressed);
+								UBA_ASSERT(localCasKey == CasKeyZero || IsCompressed(localCasKey));
+							}
+
+							insres.first->second = localCasKey == cacheCasKey;
 						}
-						else
+
+						if (!insres.first->second)
 						{
-							bool deferCreation = true;
-							bool fileIsCompressed = IsFileCompressed(info, path.data, path.count);
-							m_storage.StoreCasFile(localCasKey, path.data, CasKeyZero, deferCreation, fileIsCompressed);
-							UBA_ASSERT(localCasKey == CasKeyZero || IsCompressed(localCasKey));
+							isMatch = false;
+							if (m_reportMissReason && path.count) // if empty this has already been reported
+								misses.push_back({TString(path.data), entryIndex, cacheCasKey, localCasKey });
+							return false;
 						}
+						return true;
+					});
 
-						insres.first->second = localCasKey == cacheCasKey;
-					}
-
-					if (!insres.first->second)
-					{
-						reader.Skip(inputEnd -  reader.GetPositionData());
-						isMatch = false;
-
-						if (m_reportMissReason && path.count) // if empty this has already been reported
-							misses.push_back({TString(path.data), i, cacheCasKey, localCasKey });
-
-						break;
-					}
-				}
-
-				outputSize = reader.Read7BitEncoded();
+				if (isMatch && !result) // Returned false before setting isMatch, something went wrong
+					return false;
 
 				// No match, test next entry
 				if (!isMatch)
 				{
-					reader.Skip(outputSize);
+					traverser.SkipEntryOutputs();
 					continue;
 				}
 			}
 
-			// Fetch output files from cache (and some files need to be "denormalized" before written to disk
-
-			const u8* outputEnd = reader.GetPositionData() + outputSize;
-			while (reader.GetPositionData() != outputEnd)
 			{
-				u32 casKeyOffset = u32(reader.Read7BitEncoded());
-				if (!FetchCasTable(bucket, cacheStats, casKeyOffset))
-					return false;
-
-				TimerScope ts(cacheStats.fetchOutput);
-
-				StringBuffer<MaxPath> path;
-				CasKey casKey;
-				if (!GetLocalPathAndCasKey(bucket, rootPaths, path, casKey, bucket.serverCasKeyTable, bucket.serverPathTable, casKeyOffset))
-					return false;
-				UBA_ASSERT(IsCompressed(casKey));
+				StackBinaryWriter<128> writer;
+				NetworkMessage msg(m_client, CacheServiceId, CacheMessageType_ReportUsedEntry, writer);
+				writer.Write7BitEncoded(MakeId(bucket.id));
+				writer.WriteCasKey(cmdKey);
+				writer.Write7BitEncoded(traverser.lastId);
+				msg.Send();
+			}
 
-				FileFetcher fetcher { m_storage.m_bufferSlots };
-				fetcher.m_errorOnFail = false;
+			// Fetch output files from cache (and some files need to be "denormalized" before written to disk
 
-				if (IsNormalized(casKey))
+			bool result = traverser.TraverseEntryOutputs([&](u32 casKeyOffset)
 				{
-					// Fetch into memory, file is in special format without absolute paths
-					MemoryBlock normalizedBlock(4*1024*1024);
-					bool destinationIsCompressed = false;
-					if (!fetcher.RetrieveFile(m_logger, m_client, casKey, path.data, destinationIsCompressed, &normalizedBlock))
+					if (!FetchCasTable(bucket, cacheStats, casKeyOffset))
 						return false;
 
-					MemoryBlock localBlock(4*1024*1024);
+					TimerScope ts(cacheStats.fetchOutput);
 
-					u32 rootOffsets = *(u32*)(normalizedBlock.memory);
-					char* fileStart = (char*)(normalizedBlock.memory + sizeof(u32));
-					UBA_ASSERT(rootOffsets <= normalizedBlock.writtenSize);
+					StringBuffer<MaxPath> path;
+					CasKey casKey;
+					if (!GetLocalPathAndCasKey(bucket, rootPaths, path, casKey, bucket.serverCasKeyTable, bucket.serverPathTable, casKeyOffset))
+						return false;
+					UBA_ASSERT(IsCompressed(casKey));
 
-					// "denormalize" fetched file into another memory block that will be written to disk
-					u64 lastWritten = 0;
-					BinaryReader reader2(normalizedBlock.memory, rootOffsets, normalizedBlock.writtenSize);
-					while (reader2.GetLeft())
+					FileFetcher fetcher { m_storage.m_bufferSlots };
+					fetcher.m_errorOnFail = false;
+
+					if (IsNormalized(casKey))
 					{
-						u64 rootOffset = reader2.Read7BitEncoded();
-						if (u64 toWrite = rootOffset - lastWritten)
+						// Fetch into memory, file is in special format without absolute paths
+						MemoryBlock normalizedBlock(4*1024*1024);
+						bool destinationIsCompressed = false;
+						if (!fetcher.RetrieveFile(m_logger, m_client, casKey, path.data, destinationIsCompressed, &normalizedBlock))
+							return false;
+
+						MemoryBlock localBlock(4*1024*1024);
+
+						u32 rootOffsets = *(u32*)(normalizedBlock.memory);
+						char* fileStart = (char*)(normalizedBlock.memory + sizeof(u32));
+						UBA_ASSERT(rootOffsets <= normalizedBlock.writtenSize);
+
+						// "denormalize" fetched file into another memory block that will be written to disk
+						u64 lastWritten = 0;
+						BinaryReader reader2(normalizedBlock.memory, rootOffsets, normalizedBlock.writtenSize);
+						while (reader2.GetLeft())
+						{
+							u64 rootOffset = reader2.Read7BitEncoded();
+							if (u64 toWrite = rootOffset - lastWritten)
+								memcpy(localBlock.Allocate(toWrite, 1, TC("")), fileStart + lastWritten, toWrite);
+							u8 rootIndex = fileStart[rootOffset] - RootPaths::RootStartByte;
+							auto& root = rootPaths.GetRoot(rootIndex);
+
+							#if PLATFORM_WINDOWS
+							StringBuffer<> pathTemp;
+							pathTemp.Append(root.path);
+							char rootPath[512];
+							u32 rootPathLen = pathTemp.Parse(rootPath, sizeof_array(rootPath));
+							#else
+							const char* rootPath = root.path.data();
+							u32 rootPathLen = root.path.size();
+							#endif
+
+							if (u32 toWrite = rootPathLen - 1)
+								memcpy(localBlock.Allocate(toWrite, 1, TC("")), rootPath, toWrite);
+							lastWritten = rootOffset + 1;
+						}
+
+						u64 fileSize = rootOffsets - sizeof(u32);
+						if (u64 toWrite = fileSize - lastWritten)
 							memcpy(localBlock.Allocate(toWrite, 1, TC("")), fileStart + lastWritten, toWrite);
-						u8 rootIndex = fileStart[rootOffset] - RootPaths::RootStartByte;
-						auto& root = rootPaths.GetRoot(rootIndex);
-
-						#if PLATFORM_WINDOWS
-						StringBuffer<> pathTemp;
-						pathTemp.Append(root.path);
-						char rootPath[512];
-						u32 rootPathLen = pathTemp.Parse(rootPath, sizeof_array(rootPath));
-						#else
-						const char* rootPath = root.path.data();
-						u32 rootPathLen = root.path.size();
-						#endif
-
-						if (u32 toWrite = rootPathLen - 1)
-							memcpy(localBlock.Allocate(toWrite, 1, TC("")), rootPath, toWrite);
-						lastWritten = rootOffset + 1;
+
+						FileAccessor destFile(m_logger, path.data);
+						if (!destFile.CreateWrite())
+							return false;
+						if (!destFile.Write(localBlock.memory, localBlock.writtenSize))
+							return false;
+						if (!destFile.Close(&fetcher.lastWritten))
+							return false;
+
+						fetcher.sizeOnDisk = localBlock.writtenSize;
+						casKey = CalculateCasKey(localBlock.memory, localBlock.writtenSize, false, nullptr);
+					}
+					else
+					{
+						bool destinationIsCompressed = IsFileCompressed(info, path.data, path.count);
+						if (!fetcher.RetrieveFile(m_logger, m_client, casKey, path.data, destinationIsCompressed))
+							return false;
 					}
 
-					u64 fileSize = rootOffsets - sizeof(u32);
-					if (u64 toWrite = fileSize - lastWritten)
-						memcpy(localBlock.Allocate(toWrite, 1, TC("")), fileStart + lastWritten, toWrite);
+					cacheStats.fetchBytesRaw += fetcher.sizeOnDisk;
+					cacheStats.fetchBytesComp += fetcher.bytesReceived;
 
-					FileAccessor destFile(m_logger, path.data);
-					if (!destFile.CreateWrite())
-						return false;
-					if (!destFile.Write(localBlock.memory, localBlock.writtenSize))
-						return false;
-					if (!destFile.Close(&fetcher.lastWritten))
+					if (!m_storage.FakeCopy(casKey, path.data, fetcher.sizeOnDisk, fetcher.lastWritten, false))
 						return false;
-
-					fetcher.sizeOnDisk = localBlock.writtenSize;
-					casKey = CalculateCasKey(localBlock.memory, localBlock.writtenSize, false, nullptr);
-				}
-				else
-				{
-					bool destinationIsCompressed = IsFileCompressed(info, path.data, path.count);
-					if (!fetcher.RetrieveFile(m_logger, m_client, casKey, path.data, destinationIsCompressed))
+					if (!m_session.RegisterNewFile(path.data))
 						return false;
-				}
+					return true;
+				});
 
-				cacheStats.fetchBytesRaw += fetcher.sizeOnDisk;
-				cacheStats.fetchBytesComp += fetcher.bytesReceived;
-
-				if (!m_storage.FakeCopy(casKey, path.data, fetcher.sizeOnDisk, fetcher.lastWritten, false))
-					return false;
-				if (!m_session.RegisterNewFile(path.data))
-					return false;
-			}
+			if (!result)
+				return false;
 
 			success = true;
 			return true;
diff --git a/Engine/Source/Programs/UnrealBuildAccelerator/Common/Private/UbaCacheEntry.cpp b/Engine/Source/Programs/UnrealBuildAccelerator/Common/Private/UbaCacheEntry.cpp
new file mode 100644
--- /dev/null
+++ b/Engine/Source/Programs/UnrealBuildAccelerator/Common/Private/UbaCacheEntry.cpp
@@ -0,0 +1,564 @@
+#include "UbaCacheEntry.h"
+#include "UbaFile.h"
+#include <algorithm>
+
+namespace uba
+{
+	u64 CacheEntries::GetSharedSize()
+	{
+		return sizeof(u16) + Get7BitEncodedCount(sharedInputCasKeyOffsets.size()) + sharedInputCasKeyOffsets.size();
+	}
+
+	u64 CacheEntries::GetEntrySize(CacheEntry& entry, bool toDisk)
+	{
+		u64 size = 0;
+
+		if (toDisk)
+			size += Get7BitEncodedCount(entry.creationTime) + Get7BitEncodedCount(entry.lastUsedTime);
+		else
+			size += Get7BitEncodedCount(entry.id);
+
+		auto& extra = entry.extraInputCasKeyOffsets;
+		size += Get7BitEncodedCount(extra.size()) + extra.size();
+
+		u64 rangeSize = 0;
+		auto& ranges = entry.sharedInputCasKeyOffsetRanges;
+		for (auto& range : ranges)
+			rangeSize += Get7BitEncodedCount(range.begin) + Get7BitEncodedCount(range.end);
+		size += Get7BitEncodedCount(rangeSize) + rangeSize;
+
+		auto& outputs = entry.outputCasKeyOffsets;
+		size += Get7BitEncodedCount(outputs.size()) + outputs.size();
+		return size;
+	}
+
+	u64 CacheEntries::GetTotalSize(bool toDisk)
+	{
+		u64 size = GetSharedSize();
+		for (auto& entry : entries)
+			size += GetEntrySize(entry, toDisk);
+		return size;
+	}
+
+	bool CacheEntries::Write(BinaryWriter& writer, u32 clientVersion, bool toDisk)
+	{
+		u16& entryCount = *(u16*)writer.AllocWrite(2);
+		entryCount = 0;
+
+		#if UBA_USE_OLD
+		if (clientVersion == 3)
+		{
+			UBA_ASSERT(!toDisk);
+
+			Vector<u8> flattenInputs;
+			for (auto& entry : entries)
+			{
+				Flatten(flattenInputs, entry);
+				auto& inputs = flattenInputs;
+				auto& outputs = entry.outputCasKeyOffsets;
+
+				u64 neededSize = Get7BitEncodedCount(inputs.size()) + inputs.size() + Get7BitEncodedCount(outputs.size()) + outputs.size();
+				if (neededSize > writer.GetCapacityLeft())
+					return true;
+
+				writer.Write7BitEncoded(inputs.size());
+				writer.WriteBytes(inputs.data(), inputs.size());
+				writer.Write7BitEncoded(outputs.size());
+				writer.WriteBytes(outputs.data(), outputs.size());
+				++entryCount;
+			}
+			return true;
+		}
+		#endif
+
+		{
+			auto& shared = sharedInputCasKeyOffsets;
+
+			if (!toDisk)
+			{
+				u64 neededSize = Get7BitEncodedCount(shared.size()) + shared.size();
+				if (neededSize > writer.GetCapacityLeft())
+					return true;
+			}
+			writer.Write7BitEncoded(shared.size());
+			writer.WriteBytes(shared.data(), shared.size());
+		}
+
+		for (auto& entry : entries)
+		{
+			auto& extra = entry.extraInputCasKeyOffsets;
+			auto& ranges = entry.sharedInputCasKeyOffsetRanges;
+			auto& outputs = entry.outputCasKeyOffsets;
+
+			u64 rangeSize = 0;
+			for (auto& range : ranges)
+				rangeSize += Get7BitEncodedCount(range.begin) + Get7BitEncodedCount(range.end);
+
+			if (toDisk)
+			{
+				writer.Write7BitEncoded(entry.creationTime);
+				writer.Write7BitEncoded(entry.lastUsedTime);
+			}
+			else
+			{
+				u64 neededSize = Get7BitEncodedCount(entry.id) + Get7BitEncodedCount(extra.size()) + extra.size();
+				neededSize += Get7BitEncodedCount(rangeSize) + rangeSize;
+				neededSize += Get7BitEncodedCount(outputs.size()) + outputs.size();
+
+				if (neededSize > writer.GetCapacityLeft())
+					return true;
+
+				writer.Write7BitEncoded(entry.id);
+			}
+
+			writer.Write7BitEncoded(extra.size());
+			writer.WriteBytes(extra.data(), extra.size());
+
+			writer.Write7BitEncoded(rangeSize);
+			for (auto& range : ranges)
+			{
+				writer.Write7BitEncoded(range.begin);
+				writer.Write7BitEncoded(range.end);
+			}
+
+			writer.Write7BitEncoded(outputs.size());
+			writer.WriteBytes(outputs.data(), outputs.size());
+
+			++entryCount;
+		}
+
+		return true;
+	}
+
+	bool CacheEntries::Read(Logger& logger, BinaryReader& reader, u32 databaseVersion)
+	{
+		if (databaseVersion == 3)
+		{
+			u32 cacheEntryCount = reader.ReadU32();
+			Vector<u32> temp;
+			while (cacheEntryCount--)
+			{
+				auto& cacheEntry = entries.emplace_back();
+				cacheEntry.id = idCounter++;
+				reader.ReadU64();
+				cacheEntry.creationTime = GetSystemTimeAsFileTime();
+				cacheEntry.lastUsedTime = GetSystemTimeAsFileTime();
+
+				u32 inputSize = reader.ReadU32();
+				#if UBA_USE_OLD
+				cacheEntry.inputCasKeyOffsets.resize(inputSize);
+				reader.ReadBytes(cacheEntry.inputCasKeyOffsets.data(), inputSize);
+				reader.SetPosition(reader.GetPosition() - inputSize);
+				#endif
+
+				u64 inputEnd = reader.GetPosition() + inputSize;
+				temp.clear();
+				while (reader.GetPosition() < inputEnd)
+					temp.push_back(u32(reader.Read7BitEncoded()));
+				BuildInputsT(cacheEntry, temp, entries.size() == 1);
+
+				u32 outputSize = reader.ReadU32();
+				cacheEntry.outputCasKeyOffsets.resize(outputSize);
+				reader.ReadBytes(cacheEntry.outputCasKeyOffsets.data(), outputSize);
+			}
+
+			#if UBA_USE_OLD
+			ValidateEntries(logger);
+			#endif
+			return true;
+		}
+
+		u16 entryCount = reader.ReadU16();
+		u64 sharedSize = reader.Read7BitEncoded();
+		sharedInputCasKeyOffsets.resize(sharedSize);
+		reader.ReadBytes(sharedInputCasKeyOffsets.data(), sharedSize);
+
+		while (entryCount--)
+		{
+			auto& entry = entries.emplace_back();
+			entry.id = idCounter++;
+			entry.creationTime = reader.Read7BitEncoded();
+			entry.lastUsedTime = reader.Read7BitEncoded();
+			u64 extraSize = reader.Read7BitEncoded();
+			entry.extraInputCasKeyOffsets.resize(extraSize);
+			reader.ReadBytes(entry.extraInputCasKeyOffsets.data(), extraSize);
+			u64 rangeSize = reader.Read7BitEncoded();
+			BinaryReader rangeReader(reader.GetPositionData(), 0, rangeSize);
+			reader.Skip(rangeSize);
+			while (rangeReader.GetLeft())
+			{
+				u32 begin = u32(rangeReader.Read7BitEncoded());
+				u32 end = u32(rangeReader.Read7BitEncoded());
+				entry.sharedInputCasKeyOffsetRanges.emplace_back(begin, end);
+			}
+
+			u64 outputSize = reader.Read7BitEncoded();
+			entry.outputCasKeyOffsets.resize(outputSize);
+			reader.ReadBytes(entry.outputCasKeyOffsets.data(), outputSize);
+		}
+
+		return true;
+	}
+
+	template<typename Container>
+	void CacheEntries::BuildInputsT(CacheEntry& entry, const Container& sortedInputs, bool populateShared)
+	{
+		if (populateShared)
+		{
+			u64 bytes = 0;
+			for (u32 i : sortedInputs)
+				bytes += Get7BitEncodedCount(i);
+			sharedInputCasKeyOffsets.resize(bytes);
+			BinaryWriter writer(sharedInputCasKeyOffsets.data(), 0, sharedInputCasKeyOffsets.size());
+			for (u32 i : sortedInputs)
+				writer.Write7BitEncoded(i);
+			entry.sharedInputCasKeyOffsetRanges.emplace_back(0, u32(writer.GetPosition()));
+			return;
+		}
+
+		auto inputsIt = sortedInputs.begin();
+		auto inputsEnd = sortedInputs.end();
+
+		BinaryReader sharedReader(sharedInputCasKeyOffsets.data(), 0, sharedInputCasKeyOffsets.size());
+
+		u32 sharedOffset = ~0u;
+		u32 offset = ~0u;
+
+		u32 rangeBegin = 0;
+		bool inRange = false;
+		bool unhandledOffset = false;
+		u32 lastSharedPos = ~0u;
+
+		Vector<u8> extraOffsets;
+		extraOffsets.resize(sortedInputs.size()*4);
+		BinaryWriter extraWriter(extraOffsets.data(), 0, extraOffsets.size());
+
+		while (true)
+		{
+			u32 sharedPos = u32(sharedReader.GetPosition());
+
+			if (!sharedReader.GetLeft())
+			{
+				// Add current range if there is one going
+				if (inRange)
+					entry.sharedInputCasKeyOffsetRanges.emplace_back(rangeBegin, sharedPos);
+
+				if (offset > sharedOffset)
+					extraWriter.Write7BitEncoded(offset);
+
+				// Populate rest in extraInputCasKeyOffsets
+				for (;inputsIt != inputsEnd ;++inputsIt)
+					extraWriter.Write7BitEncoded(*inputsIt);
+				break;
+			}
+
+			if (inputsIt == inputsEnd)
+			{
+				// Add current range if there is one going
+				if (inRange)
+					entry.sharedInputCasKeyOffsetRanges.emplace_back(rangeBegin, sharedPos);
+
+				if (offset > sharedOffset)
+					extraWriter.Write7BitEncoded(offset);
+				break;
+			}
+
+			if (sharedOffset < offset)
+			{
+				lastSharedPos = sharedPos;
+				sharedOffset = u32(sharedReader.Read7BitEncoded());
+			}
+			else if (offset < sharedOffset)
+			{
+				offset = *inputsIt++;
+				sharedPos = lastSharedPos;
+			}
+			else
+			{
+				lastSharedPos = sharedPos;
+				sharedOffset = u32(sharedReader.Read7BitEncoded());
+				offset = *inputsIt++;
+			}
+
+			if (sharedOffset == offset)
+			{
+				if (!inRange)
+				{
+					rangeBegin = sharedPos;
+					inRange = true;
+				}
+			}
+			else
+			{
+				if (inRange)
+				{
+					inRange = false;
+					entry.sharedInputCasKeyOffsetRanges.emplace_back(rangeBegin, sharedPos);
+				}
+				if (offset < sharedOffset)
+				{
+					extraWriter.Write7BitEncoded(offset);
+				}
+				else
+				{
+					unhandledOffset = true;
+				}
+			}
+		}
+
+		entry.extraInputCasKeyOffsets.resize(extraWriter.GetPosition());
+		memcpy(entry.extraInputCasKeyOffsets.data(), extraWriter.GetData(), extraWriter.GetPosition());
+	}
+
+	void CacheEntries::BuildInputs(CacheEntry& entry, const Set<u32>& inputs)
+	{
+		BuildInputsT(entry, inputs, entries.empty());
+	}
+
+	void CacheEntries::UpdateEntries()
+	{
+		if (entries.empty())
+			return;
+
+		// Checking if first entry is still the matching entry
+		auto& firstEntry = *entries.begin();
+		if (firstEntry.extraInputCasKeyOffsets.empty() && firstEntry.sharedInputCasKeyOffsetRanges.size() == 1)
+		{
+			auto& range = firstEntry.sharedInputCasKeyOffsetRanges[0];
+			if (range.begin == 0 && range.end == sharedInputCasKeyOffsets.size())
+				return;
+		}
+
+		// First entry is gone, need to refresh the shared table...
+
+
+		Vector<u8> oldShared;
+		bool isFirst = true;
+		Vector<u32> temp;
+		for (auto& entry : entries)
+		{
+			if (isFirst)
+			{
+				// Flatten first entry into new shared
+				Flatten(oldShared, entry);
+				oldShared.swap(sharedInputCasKeyOffsets);
+				entry.extraInputCasKeyOffsets.clear();
+				entry.sharedInputCasKeyOffsetRanges.clear();
+				entry.sharedInputCasKeyOffsetRanges.emplace_back(0u, u32(sharedInputCasKeyOffsets.size()));
+				isFirst = false;
+			}
+			else
+			{
+				// Flatten using old shared and rebuild it with new shared
+				Flatten(temp, entry, oldShared);
+				entry.extraInputCasKeyOffsets.clear();
+				entry.sharedInputCasKeyOffsetRanges.clear();
+				BuildInputsT(entry, temp, false);
+			}
+		}
+	}
+
+	void CacheEntries::UpdateEntries(Logger& logger, const GrowingNoLockUnorderedMap<u32, u32>& oldToNewCasKeyOffset)
+	{
+		if (entries.empty())
+			return;
+
+		Vector<u32> temp;
+
+		Vector<u8> oldShared;
+		oldShared.swap(sharedInputCasKeyOffsets);
+
+		bool isFirst = true;
+		for (auto& entry : entries)
+		{
+			if (isFirst)
+			{
+				// Flatten first entry into temp
+				Flatten(temp, entry, oldShared);
+
+				// Update temp with new offsets
+				u64 newSize = 0;
+				for (auto& offset : temp)
+				{
+					auto findIt = oldToNewCasKeyOffset.find(offset);
+					if (findIt != oldToNewCasKeyOffset.end())
+						offset = findIt->second;
+					newSize += Get7BitEncodedCount(offset);
+				}
+
+				// Sort temp now when it likely is out of order
+				std::sort(temp.begin(), temp.end());
+
+				sharedInputCasKeyOffsets.resize(newSize);
+				BinaryWriter writer(sharedInputCasKeyOffsets.data(), 0, newSize);
+				for (auto& offset : temp)
+					writer.Write7BitEncoded(offset);
+
+				entry.extraInputCasKeyOffsets.clear();
+				entry.sharedInputCasKeyOffsetRanges.clear();
+				entry.sharedInputCasKeyOffsetRanges.emplace_back(0u, u32(sharedInputCasKeyOffsets.size()));
+				isFirst = false;
+			}
+			else
+			{
+				// Flatten using old shared and rebuild it with new shared
+				Flatten(temp, entry, oldShared);
+				entry.extraInputCasKeyOffsets.clear();
+				entry.sharedInputCasKeyOffsetRanges.clear();
+				for (auto& offset : temp)
+				{
+					auto findIt = oldToNewCasKeyOffset.find(offset);
+					if (findIt != oldToNewCasKeyOffset.end())
+						offset = findIt->second;
+				}
+
+				// Sort temp now when it likely is out of order
+				std::sort(temp.begin(), temp.end());
+
+				BuildInputsT(entry, temp, false);
+			}
+		}
+
+		auto updateCasKeyOffsets = [&](Vector<u8>& offsets)
+			{
+				temp.clear();
+				temp.reserve(offsets.size()*4);
+
+				u32 newOffsetsSize = 0;
+				BinaryReader reader2(offsets.data(), 0, offsets.size());
+				while (reader2.GetLeft())
+				{
+					u32 oldOffset = u32(reader2.Read7BitEncoded());
+					u32 newOffset = oldOffset;
+					auto findIt = oldToNewCasKeyOffset.find(oldOffset);
+					if (findIt != oldToNewCasKeyOffset.end())
+						newOffset = findIt->second;
+					temp.push_back(newOffset);
+					newOffsetsSize += Get7BitEncodedCount(newOffset);
+				}
+
+				std::sort(temp.begin(), temp.end());
+
+				offsets.resize(newOffsetsSize);
+				BinaryWriter writer2(offsets.data(), 0, newOffsetsSize);
+				for (u32 offset : temp)
+					writer2.Write7BitEncoded(offset);
+				UBA_ASSERT(writer2.GetPosition() == newOffsetsSize);
+			};
+
+		for (auto& entry : entries)
+		{
+			#if UBA_USE_OLD
+			updateCasKeyOffsets(entry.inputCasKeyOffsets);
+			#endif
+			updateCasKeyOffsets(entry.outputCasKeyOffsets);
+		}
+
+		#if UBA_USE_OLD
+		ValidateEntries(logger);
+		#endif
+	}
+
+	#if UBA_USE_OLD
+	void CacheEntries::ValidateEntries(Logger& logger)
+	{
+		int index = 0;
+		for (auto& entry : entries)
+		{
+			Vector<u8> res;
+			Flatten(res, entry);
+			if (res.size() != entry.inputCasKeyOffsets.size() || memcmp(res.data(), entry.inputCasKeyOffsets.data(), res.size()) != 0)
+			{
+				BinaryReader reader1(entry.inputCasKeyOffsets.data(), 0, entry.inputCasKeyOffsets.size());
+				BinaryReader reader2(res.data(), 0, res.size());
+				while (reader1.GetLeft() || reader2.GetLeft())
+				{
+					u32 a = ~0u;
+					u32 b = ~0u;
+					if (reader1.GetLeft())
+						a = u32(reader1.Read7BitEncoded());
+					if (reader2.GetLeft())
+						b = u32(reader2.Read7BitEncoded());
+
+					logger.Detail(TC("A: %u B: %u"), a, b);
+				}
+				UBA_ASSERT(false);
+			}
+			++index;
+		}
+	}
+
+	#endif
+
+	void CacheEntries::Flatten(Vector<u8>& out, const CacheEntry& entry)
+	{
+		u64 size = entry.extraInputCasKeyOffsets.size();
+		for (auto& range : entry.sharedInputCasKeyOffsetRanges)
+			size += range.end - range.begin;
+		out.resize(size);
+		BinaryWriter writer(out.data(), 0, out.size());
+
+		BinaryReader extraReader(entry.extraInputCasKeyOffsets.data(), 0, entry.extraInputCasKeyOffsets.size());
+		u32 nextExtra = ~0u;
+		if (extraReader.GetLeft())
+			nextExtra = u32(extraReader.Read7BitEncoded());
+
+		auto writeExtra = [&](u32 prevOffset)
+			{
+				while (nextExtra < prevOffset)
+				{
+					writer.Write7BitEncoded(nextExtra);
+					nextExtra = ~0u;
+					if (extraReader.GetLeft())
+						nextExtra = u32(extraReader.Read7BitEncoded());
+				};
+			};
+
+		for (auto& range : entry.sharedInputCasKeyOffsetRanges)
+		{
+			BinaryReader inputReader(sharedInputCasKeyOffsets.data() + range.begin, 0, range.end - range.begin);
+			while (inputReader.GetLeft())
+			{
+				u32 offset = u32(inputReader.Read7BitEncoded());
+				writeExtra(offset);
+				writer.Write7BitEncoded(offset);
+			}
+		}
+
+		writeExtra(~0u);
+	}
+
+	void CacheEntries::Flatten(Vector<u32>& out, const CacheEntry& entry, const Vector<u8>& sharedOffsets)
+	{
+		out.clear();
+
+		BinaryReader extraReader(entry.extraInputCasKeyOffsets.data(), 0, entry.extraInputCasKeyOffsets.size());
+		u32 nextExtra = ~0u;
+		if (extraReader.GetLeft())
+			nextExtra = u32(extraReader.Read7BitEncoded());
+
+		auto writeExtra = [&](u32 prevOffset)
+			{
+				while (nextExtra < prevOffset)
+				{
+					out.push_back(nextExtra);
+					nextExtra = ~0u;
+					if (extraReader.GetLeft())
+						nextExtra = u32(extraReader.Read7BitEncoded());
+				};
+			};
+
+		for (auto& range : entry.sharedInputCasKeyOffsetRanges)
+		{
+			BinaryReader inputReader(sharedOffsets.data() + range.begin, 0, range.end - range.begin);
+			while (inputReader.GetLeft())
+			{
+				u32 offset = u32(inputReader.Read7BitEncoded());
+				writeExtra(offset);
+				out.push_back(offset);
+			}
+		}
+
+		writeExtra(~0u);
+	}
+}
diff --git a/Engine/Source/Programs/UnrealBuildAccelerator/Common/Private/UbaCacheServer.cpp b/Engine/Source/Programs/UnrealBuildAccelerator/Common/Private/UbaCacheServer.cpp
--- a/Engine/Source/Programs/UnrealBuildAccelerator/Common/Private/UbaCacheServer.cpp
+++ b/Engine/Source/Programs/UnrealBuildAccelerator/Common/Private/UbaCacheServer.cpp
@@ -1,27 +1,21 @@
 // Copyright Epic Games, Inc. All Rights Reserved.
 
 #include "UbaCacheServer.h"
+#include "UbaCacheEntry.h"
 #include "UbaCompactTables.h"
 #include "UbaBinaryReaderWriter.h"
 #include "UbaFileAccessor.h"
 #include "UbaNetworkServer.h"
 #include "UbaStorageServer.h"
-#include <algorithm>
 //#include <oodle2.h>
 
 namespace uba
 {
-	static constexpr u32 CacheFileVersion = 3;
+	static constexpr u32 CacheFileVersion = 4;
+	static constexpr u32 CacheFileCompatibilityVersion = 3;
 
 	bool IsCaseInsensitive(u64 id) { return (id & (1ull << 32)) == 0; }
 
-	struct CacheServer::CacheEntry
-	{
-		u64 creationTime;
-		Vector<u8> inputCasKeyOffsets;
-		Vector<u8> outputCasKeyOffsets;
-	};
-
 	struct CacheServer::ConnectionBucket
 	{
 		ConnectionBucket(u64 i) : pathTable(CachePathTableMaxSize, CompactPathTable::V1, IsCaseInsensitive(i)), casKeyTable(CacheCasKeyTableMaxSize), id(i) {}
@@ -35,15 +29,10 @@ namespace uba
 
 	struct CacheServer::Connection
 	{
+		u32 clientVersion;
 		UnorderedMap<u64, ConnectionBucket> buckets;
 	};
 
-	struct CacheServer::CacheEntries
-	{
-		ReaderWriterLock lock;
-		List<CacheEntry> entries;
-	};
-
 	struct CacheServer::Bucket
 	{
 		Bucket(u64 id) : m_pathTable(CachePathTableMaxSize, CompactPathTable::V1, IsCaseInsensitive(id)), m_casKeyTable(CacheCasKeyTableMaxSize) {}
@@ -111,13 +100,21 @@ namespace uba
 		if (!file.OpenMemoryRead(0, false))
 		{
 			m_logger.Detail(TC("No database found. Starting a new one at %s"), fileName.data);
+			m_creationTime = GetSystemTimeAsFileTime();
 			return true;
 		}
 		BinaryReader reader(file.GetData(), 0, file.GetSize());
 
-		u32 version = reader.ReadU32();
-		if (version != CacheFileVersion)
+		u32 databaseVersion = reader.ReadU32();
+		if (databaseVersion < CacheFileCompatibilityVersion || databaseVersion > CacheFileVersion)
+		{
+			m_logger.Detail(TC("Can't load database of version %u. Starting a new one at %s"), databaseVersion, fileName.data);
 			return true;
+		}
+		if (databaseVersion == 3)
+			m_creationTime = GetSystemTimeAsFileTime() - 1;
+		else
+			m_creationTime = reader.ReadU64();
 
 		u32 totalPathTableSize = 0;
 		u32 totalCasKeyTableSize = 0;
@@ -155,23 +152,8 @@ namespace uba
 				auto insres = bucket.m_cacheEntryLookup.try_emplace(reader.ReadCasKey());
 				UBA_ASSERT(insres.second);
 				auto& cacheEntries = insres.first->second;
-				u32 cacheEntryCount = reader.ReadU32();
-				totalCacheEntryCount += cacheEntryCount;
-
-				while (cacheEntryCount--)
-				{
-					auto& cacheEntry = cacheEntries.entries.emplace_back();
-
-					cacheEntry.creationTime = reader.ReadU64();
-
-					u32 inputSize = reader.ReadU32();
-					cacheEntry.inputCasKeyOffsets.resize(inputSize);
-					reader.ReadBytes(cacheEntry.inputCasKeyOffsets.data(), inputSize);
-
-					u32 outputSize = reader.ReadU32();
-					cacheEntry.outputCasKeyOffsets.resize(outputSize);
-					reader.ReadBytes(cacheEntry.outputCasKeyOffsets.data(), outputSize);
-				}
+				cacheEntries.Read(m_logger, reader, databaseVersion);
+				totalCacheEntryCount += cacheEntries.entries.size();
 			}
 		}
 
@@ -234,8 +216,12 @@ namespace uba
 
 		Write(CacheFileVersion);
 
+		Write(m_creationTime);
+
 		Write(u32(m_buckets.size()));
 
+		Vector<u8> temp;
+
 		for (auto& kv : m_buckets)
 		{
 			Bucket& bucket = kv.second;
@@ -257,21 +243,35 @@ namespace uba
 			{
 				Write(kv2.first);
 
-				u32 cacheEntryCount = u32(kv2.second.entries.size());
-				Write(cacheEntryCount);
-
-				for (CacheEntry& entry : kv2.second.entries)
+				#if UBA_USE_OLD
+				if (CacheFileVersion == 3)
 				{
-					WriteBytes(&entry.creationTime, sizeof(entry.creationTime));
+					u32 cacheEntryCount = u32(kv2.second.entries.size());
+					Write(cacheEntryCount);
 
-					u32 inputSize = u32(entry.inputCasKeyOffsets.size());
-					Write(inputSize);
-					WriteBytes(entry.inputCasKeyOffsets.data(), inputSize);
+					for (CacheEntry& entry : kv2.second.entries)
+					{
+						WriteBytes(&entry.creationTime, sizeof(entry.creationTime));
+
+						u32 inputSize = u32(entry.inputCasKeyOffsets.size());
+						Write(inputSize);
+						WriteBytes(entry.inputCasKeyOffsets.data(), inputSize);
 
-					u32 outputSize = u32(entry.outputCasKeyOffsets.size());
-					Write(outputSize);
-					WriteBytes(entry.outputCasKeyOffsets.data(), outputSize);
+						u32 outputSize = u32(entry.outputCasKeyOffsets.size());
+						Write(outputSize);
+						WriteBytes(entry.outputCasKeyOffsets.data(), outputSize);
+					}
 				}
+				else
+				#else
+				{
+					temp.resize(kv2.second.GetTotalSize(true));
+					BinaryWriter writer(temp.data(), 0, temp.size());
+					kv2.second.Write(writer, CacheNetworkVersion, true);
+					UBA_ASSERT(writer.GetPosition() == temp.size());
+					WriteBytes(temp.data(), temp.size());
+				}
+				#endif
 			}
 		}
 
@@ -357,20 +357,20 @@ namespace uba
 
 				for (auto li=bucket.m_cacheEntryLookup.begin(), le=bucket.m_cacheEntryLookup.end(); li!=le;)
 				{
+					CacheEntries& entries = li->second;
+
 					// There is currently no idea saving more than 256kb worth of entries per lookup key (because that is what fetch max returns).. so let's wipe out
 					// all the entries that overflow that number
-					u64 capacityLeft = SendMaxSize - 32;
+					u64 capacityLeft = SendMaxSize - 32 - entries.GetSharedSize();
 
-					CacheEntries& entries = li->second;
 					for (auto i=entries.entries.begin(), e=entries.entries.end(); i!=e;)
 					{
 						auto& entry = *i;
-						auto& inputs = entry.inputCasKeyOffsets;
 						auto& outputs = entry.outputCasKeyOffsets;
 
 						bool deleteEntry = false;
 
-						u64 neededSize = Get7BitEncodedCount(inputs.size()) + inputs.size() + Get7BitEncodedCount(outputs.size()) + outputs.size();
+						u64 neededSize = entries.GetEntrySize(entry, false);
 						if (neededSize > capacityLeft)
 						{
 							deleteEntry = true;
@@ -395,7 +395,12 @@ namespace uba
 							}
 						}
 
+						// This is an attempt at removing entries that has inputs that depends on other entries outputs.
+						// and that there is no point keeping them if the other entry is removed
+						// Example would be that there is no idea keeping entries that uses a pch if the entry producing the pch is gone
+						#if 0
 						// Check if there are keys that use removed caskey as input and delete those too
+						auto& inputs = entry.inputCasKeyOffsets;
 						if (!deleteEntry && !deletedCasFiles.empty())
 						{
 							BinaryReader reader2(inputs.data(), 0, inputs.size());
@@ -411,6 +416,7 @@ namespace uba
 								}
 							}
 						}
+						#endif
 
 						// Remove entry from entries list and skip increasing ref count of cas files
 						if (deleteEntry)
@@ -422,7 +428,6 @@ namespace uba
 						}
 
 						++bucket.totalEntryCount;
-						bucket.totalEntrySize += neededSize;
 
 						capacityLeft -= neededSize;
 
@@ -439,6 +444,8 @@ namespace uba
 
 					if (!entries.entries.empty())
 					{
+						//entries.UpdateEntries(); // Not needed.. 
+						bucket.totalEntrySize += entries.GetTotalSize(false);
 						++li;
 						continue;
 					}
@@ -500,15 +507,17 @@ namespace uba
 			// Collect all caskeys that are used by cache entries.
 			for (auto& kv2 : bucket.m_cacheEntryLookup)
 			{
+				auto collectUsedCasKeyOffsets = [&](const Vector<u8>& offsets)
+					{
+						BinaryReader reader2(offsets.data(), 0, offsets.size());
+						while (reader2.GetLeft())
+							usedCasKeyOffsets.insert(u32(reader2.Read7BitEncoded()));
+					};
+
+				collectUsedCasKeyOffsets(kv2.second.sharedInputCasKeyOffsets);
 				for (auto& entry : kv2.second.entries)
 				{
-					auto collectUsedCasKeyOffsets = [&](const Vector<u8>& offsets)
-						{
-							BinaryReader reader2(offsets.data(), 0, offsets.size());
-							while (reader2.GetLeft())
-								usedCasKeyOffsets.insert(u32(reader2.Read7BitEncoded()));
-						};
-					collectUsedCasKeyOffsets(entry.inputCasKeyOffsets);
+					collectUsedCasKeyOffsets(entry.extraInputCasKeyOffsets);
 					collectUsedCasKeyOffsets(entry.outputCasKeyOffsets);
 				}
 			}
@@ -584,42 +593,9 @@ namespace uba
 				// Update all casKeyOffsets
 				u64 updateEntriesStart = GetTime();
 
-				auto updateCasKeyOffsets = [&](Vector<u8>& offsets, Vector<u32>& newOffsets)
-					{
-						newOffsets.clear();
-						newOffsets.reserve(offsets.size()*4);
-
-						u32 newOffsetsSize = 0;
-						BinaryReader reader2(offsets.data(), 0, offsets.size());
-						while (reader2.GetLeft())
-						{
-							u32 oldOffset = u32(reader2.Read7BitEncoded());
-							u32 newOffset = oldOffset;
-							auto findIt = oldToNewCasKeyOffset.find(oldOffset);
-							if (findIt != oldToNewCasKeyOffset.end())
-								newOffset = findIt->second;
-							newOffsets.push_back(newOffset);
-							newOffsetsSize += Get7BitEncodedCount(newOffset);
-						}
-
-						std::sort(newOffsets.begin(), newOffsets.end());
-
-						offsets.resize(newOffsetsSize);
-						BinaryWriter writer2(offsets.data(), 0, newOffsetsSize);
-						for (u32 offset : newOffsets)
-							writer2.Write7BitEncoded(offset);
-						UBA_ASSERT(writer2.GetPosition() == newOffsetsSize);
-					};
-
 				m_server.ParallelFor(workerCountToUse, bucket.m_cacheEntryLookup, [&](auto& it)
 					{
-						Vector<u32> newOffsets;
-						CacheEntries& entries = it->second;(void)entries;
-						for (auto& entry : entries.entries)
-						{
-							updateCasKeyOffsets(entry.inputCasKeyOffsets, newOffsets);
-							updateCasKeyOffsets(entry.outputCasKeyOffsets, newOffsets);
-						}
+						it->second.UpdateEntries(m_logger, oldToNewCasKeyOffset);
 					});
 
 				#if 0
@@ -657,6 +633,7 @@ namespace uba
 		m_logger.Info(TC("Maintenance done! (%s) CasFiles: %llu (%s) OldestEntry: %s"), TimeToText(duration).str, totalCasCount - deletedCasCount, BytesToText(totalCasSize).str, TimeToText(oldestTime, true).str);
 		
 		m_longestMaintenance = Max(m_longestMaintenance, duration);
+
 		return true;
 	}
 
@@ -703,7 +680,7 @@ namespace uba
 		case CacheMessageType_Connect:
 		{
 			u32 clientVersion = reader.ReadU32();
-			if (clientVersion != CacheNetworkVersion)
+			if (clientVersion < 3 || clientVersion > CacheNetworkVersion)
 				return m_logger.Error(TC("Different network versions. Client: %u, Server: %u. Disconnecting"), clientVersion, CacheNetworkVersion);
 
 			SCOPED_WRITE_LOCK(m_connectionsLock, lock);
@@ -714,7 +691,7 @@ namespace uba
 			}
 
 			writer.WriteBool(true);
-			m_connections.try_emplace(connectionInfo.GetId());
+			m_connections.try_emplace(connectionInfo.GetId(), clientVersion);
 			return true;
 		}
 		case CacheMessageType_StorePathTable:
@@ -752,9 +729,7 @@ namespace uba
 				lock2.Leave();
 
 				SCOPED_WRITE_LOCK(cacheEntries.lock, lock4);
-				auto& cacheEntry = cacheEntries.entries.emplace_front(std::move(findIt->second));
-				(void)cacheEntry;
-				//m_logger.Info(TC("Added new cache entry"));
+				cacheEntries.entries.emplace_front(std::move(findIt->second));
 			}
 			return true;
 		}
@@ -765,11 +740,18 @@ namespace uba
 			return HandleFetchCasTable(reader, writer);
 
 		case CacheMessageType_FetchEntries:
-			return HandleFetchEntries(reader, writer);
-
+		{
+			SCOPED_READ_LOCK(m_connectionsLock, lock);
+			u32 clientVersion = m_connections[connectionInfo.GetId()].clientVersion;
+			lock.Leave();
+			return HandleFetchEntries(reader, writer, clientVersion);
+		}
 		case CacheMessageType_ExecuteCommand:
 			return HandleExecuteCommand(reader, writer);
 
+		case CacheMessageType_ReportUsedEntry:
+			return HandleReportUsedEntry(reader, writer);
+
 		case CacheMessageType_RequestShutdown:
 		{
 			TString reason = reader.ReadString();
@@ -847,15 +829,34 @@ namespace uba
 
 		SCOPED_WRITE_LOCK(cacheEntries.lock, lock2);
 		
+
+		// Create entry based on existing entry
+		CacheEntry newEntry;
+		cacheEntries.BuildInputs(newEntry, inputs);
+
 		List<CacheEntry>::iterator matchingEntry = cacheEntries.entries.end();
 		for (auto i=cacheEntries.entries.begin(), e=cacheEntries.entries.end(); i!=e; ++i)
 		{
-			if (i->inputCasKeyOffsets != inputCasKeyOffsets)
+			if (i->sharedInputCasKeyOffsetRanges != newEntry.sharedInputCasKeyOffsetRanges || i->extraInputCasKeyOffsets != newEntry.extraInputCasKeyOffsets)
 				continue;
 			matchingEntry = i;
 			break;
 		}
 
+		#if UBA_USE_OLD
+		List<CacheEntry>::iterator matchingEntry2 = cacheEntries.entries.end();
+		for (auto i=cacheEntries.entries.begin(), e=cacheEntries.entries.end(); i!=e; ++i)
+		{
+			if (i->inputCasKeyOffsets != inputCasKeyOffsets)
+				continue;
+			matchingEntry2 = i;
+			break;
+		}
+		if (matchingEntry2 != matchingEntry)
+			m_logger.Warning(L"CODE MISMATCH!!!");
+		#endif
+
+
 		// Already exists
 		if (matchingEntry != cacheEntries.entries.end())
 		{
@@ -901,9 +902,10 @@ namespace uba
 				return true;
 		}
 
+		#if UBA_USE_OLD
 		// Add new entry
-		CacheEntry newEntry;
 		newEntry.inputCasKeyOffsets.swap(inputCasKeyOffsets);
+		#endif
 
 		Set<u32> outputs;
 		u64 bytesForOutput = 0;
@@ -942,7 +944,8 @@ namespace uba
 			w2.Write7BitEncoded(output);
 
 
-		newEntry.creationTime = GetSystemTimeAsFileTime();
+		newEntry.creationTime = GetSystemTimeAsFileTime() - m_creationTime;
+		newEntry.id = cacheEntries.idCounter++;
 
 		// If cache server has all content we can put the new cache entry directly in the lookup.. otherwise we'll have to wait until client has uploaded content
 		if (hasAllContent)
@@ -985,37 +988,45 @@ namespace uba
 		return true;
 	}
 
-	bool CacheServer::HandleFetchEntries(BinaryReader& reader, BinaryWriter& writer)
+	bool CacheServer::HandleFetchEntries(BinaryReader& reader, BinaryWriter& writer, u32 clientVersion)
 	{
 		Bucket& bucket = GetBucket(reader);
 		CasKey cmdKey = reader.ReadCasKey();
 
-		u16& entryCount = *(u16*)writer.AllocWrite(2);
-		entryCount = 0;
-
 		SCOPED_READ_LOCK(bucket.m_cacheEntryLookupLock, lock);
 		auto findIt = bucket.m_cacheEntryLookup.find(cmdKey);
 		if (findIt == bucket.m_cacheEntryLookup.end())
+		{
+			writer.WriteU16(0);
 			return true;
+		}
 		auto& cacheEntries = findIt->second;
 		lock.Leave();
 
 		SCOPED_READ_LOCK(cacheEntries.lock, lock2);
+		return cacheEntries.Write(writer, clientVersion, false);
+	}
 
+	bool CacheServer::HandleReportUsedEntry(BinaryReader& reader, BinaryWriter& writer)
+	{
+		Bucket& bucket = GetBucket(reader);
+		CasKey cmdKey = reader.ReadCasKey();
+		u64 entryId = reader.Read7BitEncoded();
+
+		SCOPED_READ_LOCK(bucket.m_cacheEntryLookupLock, lock);
+		auto findIt = bucket.m_cacheEntryLookup.find(cmdKey);
+		if (findIt == bucket.m_cacheEntryLookup.end())
+			return true;
+		auto& cacheEntries = findIt->second;
+		lock.Leave();
+
+		SCOPED_WRITE_LOCK(cacheEntries.lock, lock2);
 		for (auto& entry : cacheEntries.entries)
 		{
-			auto& inputs = entry.inputCasKeyOffsets;
-			auto& outputs = entry.outputCasKeyOffsets;
-
-			u64 neededSize = Get7BitEncodedCount(inputs.size()) + inputs.size() + Get7BitEncodedCount(outputs.size()) + outputs.size();
-			if (neededSize > writer.GetCapacityLeft())
-				break;
-
-			writer.Write7BitEncoded(inputs.size());
-			writer.WriteBytes(inputs.data(), inputs.size());
-			writer.Write7BitEncoded(outputs.size());
-			writer.WriteBytes(outputs.data(), outputs.size());
-			++entryCount;
+			if (entryId != entry.id)
+				continue;
+			entry.lastUsedTime = GetSystemTimeAsFileTime() - m_creationTime;
+			break;
 		}
 		return true;
 	}
@@ -1062,6 +1073,8 @@ namespace uba
 
 			u64 now = GetSystemTimeAsFileTime();
 
+			Vector<u8> temp;
+
 			SCOPED_READ_LOCK(m_bucketsLock, bucketsLock);
 			for (auto& kv : m_buckets)
 			{
@@ -1077,26 +1090,27 @@ namespace uba
 					if (filterString.count)
 					{
 						u32 index = 0;
+						auto findString = [&](const Vector<u8>& offsets)
+							{
+								BinaryReader reader2(offsets.data(), 0, offsets.size());
+								while (reader2.GetLeft())
+								{
+									u64 offset = reader2.Read7BitEncoded();
+									CasKey casKey;
+									StringBuffer<> path;
+									bucket.m_casKeyTable.GetPathAndKey(path, casKey, bucket.m_pathTable, offset);
+									if (path.Contains(filterString.data))
+										return true;
+									if (Contains(CasKeyString(casKey).str, filterString.data))
+										return true;
+								}
+								return false;
+							};
+
 						for (auto& entry : entries.entries)
 						{
-							auto findString = [&](const Vector<u8>& offsets)
-								{
-									BinaryReader reader2(offsets.data(), 0, offsets.size());
-									while (reader2.GetLeft())
-									{
-										u64 offset = reader2.Read7BitEncoded();
-										CasKey casKey;
-										StringBuffer<> path;
-										bucket.m_casKeyTable.GetPathAndKey(path, casKey, bucket.m_pathTable, offset);
-										if (path.Contains(filterString.data))
-											return true;
-										if (Contains(CasKeyString(casKey).str, filterString.data))
-											return true;
-									}
-									return false;
-								};
-
-							if (findString(entry.inputCasKeyOffsets) || findString(entry.outputCasKeyOffsets))
+							entries.Flatten(temp, entry);
+							if (findString(temp) || findString(entry.outputCasKeyOffsets))
 								visibleIndices.insert(index);
 							++index;
 						}
@@ -1132,7 +1146,8 @@ namespace uba
 							};
 
 						writeLine(line.Clear().Append(TC("   Inputs:")).data);
-						writeOffsets(entry.inputCasKeyOffsets);
+						entries.Flatten(temp, entry);
+						writeOffsets(temp);
 						writeLine(line.Clear().Append(TC("   Outputs:")).data);
 						writeOffsets(entry.outputCasKeyOffsets);
 						++index;
@@ -1143,7 +1158,8 @@ namespace uba
 		else if (command.Equals(TC("status")))
 		{
 			writeLine(TC("UbaCacheServer status"));
-			writeLine(line.Clear().Appendf(TC("  Uptime: %s"), TimeToText(GetTime() - m_startTime).str).data);
+			writeLine(line.Clear().Appendf(TC("  CreationTime: %s"), TimeToText(MsToTime(1000*GetFileTimeAsSeconds(GetNowFileTime() - m_creationTime)), true).str).data);
+			writeLine(line.Clear().Appendf(TC("  UpTime: %s"), TimeToText(GetTime() - m_startTime, true).str).data);
 			writeLine(line.Clear().Appendf(TC("  Longest maintenance: %s"), TimeToText(m_longestMaintenance).str).data);
 			writeLine(line.Clear().Appendf(TC("  Buckets:")).data);
 			u32 index = 0;
diff --git a/Engine/Source/Programs/UnrealBuildAccelerator/Common/Public/UbaCacheEntry.h b/Engine/Source/Programs/UnrealBuildAccelerator/Common/Public/UbaCacheEntry.h
new file mode 100644
--- /dev/null
+++ b/Engine/Source/Programs/UnrealBuildAccelerator/Common/Public/UbaCacheEntry.h
@@ -0,0 +1,114 @@
+// Copyright Epic Games, Inc. All Rights Reserved.
+
+#pragma once
+
+#include "UbaBinaryReaderWriter.h"
+#include "UbaLogger.h"
+
+#define UBA_USE_OLD 0
+
+namespace uba
+{
+	struct CacheEntry
+	{
+		u64 creationTime = 0;
+		u64 lastUsedTime = 0;
+		u32 id = 0;
+
+		#if UBA_USE_OLD
+		Vector<u8> inputCasKeyOffsets;
+		#endif
+
+		struct Range { u32 begin; u32 end; bool operator==(const Range& o) const { return begin == o.begin && end == o.end; } };
+		Vector<Range> sharedInputCasKeyOffsetRanges;
+		Vector<u8> extraInputCasKeyOffsets;
+		Vector<u8> outputCasKeyOffsets;
+	};
+
+	struct CacheEntries
+	{
+		ReaderWriterLock lock;
+		List<CacheEntry> entries;
+		Vector<u8> sharedInputCasKeyOffsets;
+		u32 idCounter = 0;
+
+		u64 GetSharedSize();
+		u64 GetEntrySize(CacheEntry& entry, bool toDisk);
+		u64 GetTotalSize(bool toDisk);
+		bool Write(BinaryWriter& writer, u32 clientVersion, bool toDisk);
+		bool Read(Logger& logger, BinaryReader& reader, u32 databaseVersion);
+		void BuildInputs(CacheEntry& entry, const Set<u32>& inputs);
+		void UpdateEntries();
+		void UpdateEntries(Logger& logger, const GrowingNoLockUnorderedMap<u32, u32>& oldToNewCasKeyOffset);
+
+		#if UBA_USE_OLD
+		void ValidateEntries(Logger& logger);
+		#endif
+
+		void Flatten(Vector<u8>& out, const CacheEntry& entry);
+		void Flatten(Vector<u32>& out, const CacheEntry& entry, const Vector<u8>& sharedOffsets);
+
+		template<typename Container>
+		void BuildInputsT(CacheEntry& entry, const Container& sortedInputs, bool populateShared);
+	};
+
+
+	struct CacheEntriesTraverser
+	{
+		CacheEntriesTraverser(BinaryReader& r) : reader(r)
+		{
+			u64 sharedSize = r.Read7BitEncoded();
+			sharedInputOffsets = r.GetPositionData();
+			reader.Skip(sharedSize);
+		}
+
+		template<typename Func>
+		bool TraverseEntryInputs(const Func& func)
+		{
+			lastId = u32(reader.Read7BitEncoded());
+			u64 extraSize = reader.Read7BitEncoded();
+			BinaryReader extraReader(reader.GetPositionData(), 0, extraSize);
+			reader.Skip(extraSize);
+			u64 rangeSize = reader.Read7BitEncoded();
+			BinaryReader rangeReader(reader.GetPositionData(), 0, rangeSize);
+			reader.Skip(rangeSize);
+
+			while (extraReader.GetLeft())
+				if (!func(u32(extraReader.Read7BitEncoded())))
+					return false;
+
+			while (rangeReader.GetLeft())
+			{
+				u64 begin = rangeReader.Read7BitEncoded();
+				u64 end = rangeReader.Read7BitEncoded();
+				BinaryReader rangeReader2(sharedInputOffsets + begin, 0, end - begin);
+				while (rangeReader2.GetLeft())
+					if (!func(u32(rangeReader2.Read7BitEncoded())))
+						return false;
+			}
+			return true;
+		}
+
+		template<typename Func>
+		bool TraverseEntryOutputs(const Func& func)
+		{
+			u64 outSize = reader.Read7BitEncoded();
+			BinaryReader outReader(reader.GetPositionData(), 0, outSize);
+			reader.Skip(outSize);
+			while (outReader.GetLeft())
+				if (!func(u32(outReader.Read7BitEncoded())))
+					return false;
+			return true;
+		}
+
+		void SkipEntryOutputs()
+		{
+			u64 outSize = reader.Read7BitEncoded();
+			reader.Skip(outSize);
+		}
+
+		BinaryReader& reader;
+		const u8* sharedInputOffsets;
+		u32 lastId = 0;;
+	};
+}
diff --git a/Engine/Source/Programs/UnrealBuildAccelerator/Common/Public/UbaCacheServer.h b/Engine/Source/Programs/UnrealBuildAccelerator/Common/Public/UbaCacheServer.h
--- a/Engine/Source/Programs/UnrealBuildAccelerator/Common/Public/UbaCacheServer.h
+++ b/Engine/Source/Programs/UnrealBuildAccelerator/Common/Public/UbaCacheServer.h
@@ -11,6 +11,8 @@ namespace uba
 	class StorageServer;
 	struct BinaryReader;
 	struct BinaryWriter;
+	struct CacheEntry;
+	struct CacheEntries;
 	struct ConnectionInfo;
 
 	class CacheServer
@@ -32,8 +34,6 @@ namespace uba
 
 		struct Connection;
 		struct ConnectionBucket;
-		struct CacheEntry;
-		struct CacheEntries;
 		struct Bucket;
 
 		ConnectionBucket& GetConnectionBucket(const ConnectionInfo& connectionInfo, BinaryReader& reader);
@@ -43,7 +43,8 @@ namespace uba
 		bool HandleStoreEntry(ConnectionBucket& bucket, BinaryReader& reader, BinaryWriter& writer);
 		bool HandleFetchPathTable(BinaryReader& reader, BinaryWriter& writer);
 		bool HandleFetchCasTable(BinaryReader& reader, BinaryWriter& writer);
-		bool HandleFetchEntries(BinaryReader& reader, BinaryWriter& writer);
+		bool HandleFetchEntries(BinaryReader& reader, BinaryWriter& writer, u32 clientVersion);
+		bool HandleReportUsedEntry(BinaryReader& reader, BinaryWriter& writer);
 		bool HandleExecuteCommand(BinaryReader& reader, BinaryWriter& writer);
 
 		MutableLogger m_logger;
@@ -63,7 +64,8 @@ namespace uba
 
 		Atomic<bool> m_shutdownRequested = false;
 
-		u64 m_startTime;
+		u64 m_creationTime = 0;
+		u64 m_startTime = 0;
 		u64 m_longestMaintenance = 0;
 
 		CacheServer(const CacheServer&) = delete;
diff --git a/Engine/Source/Programs/UnrealBuildAccelerator/Common/Public/UbaNetwork.h b/Engine/Source/Programs/UnrealBuildAccelerator/Common/Public/UbaNetwork.h
--- a/Engine/Source/Programs/UnrealBuildAccelerator/Common/Public/UbaNetwork.h
+++ b/Engine/Source/Programs/UnrealBuildAccelerator/Common/Public/UbaNetwork.h
@@ -16,7 +16,7 @@ namespace uba
 	static constexpr u32 SystemNetworkVersion = 1339;
 	static constexpr u32 StorageNetworkVersion = 4;
 	static constexpr u32 SessionNetworkVersion = 32;
-	static constexpr u32 CacheNetworkVersion = 3;
+	static constexpr u32 CacheNetworkVersion = 4;
 
 	static constexpr u32 CachePathTableMaxSize = 32*1024*1024;
 	static constexpr u32 CacheCasKeyTableMaxSize = 32*1024*1024;
@@ -112,6 +112,7 @@ namespace uba
 		UBA_CACHE_MESSAGE(FetchEntries) \
 		UBA_CACHE_MESSAGE(ExecuteCommand) \
 		UBA_CACHE_MESSAGE(RequestShutdown) \
+		UBA_CACHE_MESSAGE(ReportUsedEntry) \
 
 	enum CacheMessageType : u8
 	{
