[UBA]
* Added so fails to write on UbaCacheClient just produce infos and not errors/warnings.. since we don't want a failed cache download to cause failed builds

diff --git a/Engine/Source/Programs/UnrealBuildAccelerator/Common/Private/UbaCacheClient.cpp b/Engine/Source/Programs/UnrealBuildAccelerator/Common/Private/UbaCacheClient.cpp
--- a/Engine/Source/Programs/UnrealBuildAccelerator/Common/Private/UbaCacheClient.cpp
+++ b/Engine/Source/Programs/UnrealBuildAccelerator/Common/Private/UbaCacheClient.cpp
@@ -430,6 +430,12 @@ namespace uba
 
 			// Fetch output files from cache (and some files need to be "denormalized" before written to disk
 
+			struct DowngradedLogger : public LoggerWithWriter
+			{
+				DowngradedLogger(LogWriter& writer, const tchar* prefix) : LoggerWithWriter(writer, prefix) {}
+				virtual void Log(LogEntryType type, const tchar* str, u32 strLen) override { LoggerWithWriter::Log(Max(type, LogEntryType_Info), str, strLen); }
+			};
+
 			bool result = traverser.TraverseEntryOutputs([&](u32 casKeyOffset)
 				{
 					if (!FetchCasTable(bucket, cacheStats, casKeyOffset))
@@ -448,11 +454,12 @@ namespace uba
 
 					if (IsNormalized(casKey))
 					{
+						DowngradedLogger logger(m_logger.m_writer, TC("UbaCacheClientNormalizedDownload"));
 						// Fetch into memory, file is in special format without absolute paths
 						MemoryBlock normalizedBlock(4*1024*1024);
 						bool destinationIsCompressed = false;
-						if (!fetcher.RetrieveFile(m_logger, m_client, casKey, path.data, destinationIsCompressed, &normalizedBlock))
-							return false;
+						if (!fetcher.RetrieveFile(logger, m_client, casKey, path.data, destinationIsCompressed, &normalizedBlock))
+							return logger.Error(TC("Failed to download cache output for %s"), info.description);
 
 						MemoryBlock localBlock(4*1024*1024);
 
@@ -490,9 +497,9 @@ namespace uba
 						if (u64 toWrite = fileSize - lastWritten)
 							memcpy(localBlock.Allocate(toWrite, 1, TC("")), fileStart + lastWritten, toWrite);
 
-						FileAccessor destFile(m_logger, path.data);
+						FileAccessor destFile(logger, path.data);
 						if (!destFile.CreateWrite())
-							return false;
+							return logger.Error(TC("Failed to create file for cache output %s for %s"), path.data, info.description);
 						if (!destFile.Write(localBlock.memory, localBlock.writtenSize))
 							return false;
 						if (!destFile.Close(&fetcher.lastWritten))
@@ -503,9 +510,10 @@ namespace uba
 					}
 					else
 					{
+						DowngradedLogger logger(m_logger.m_writer, TC("UbaCacheClientDownload"));
 						bool destinationIsCompressed = IsFileCompressed(info, path.data, path.count);
-						if (!fetcher.RetrieveFile(m_logger, m_client, casKey, path.data, destinationIsCompressed))
-							return false;
+						if (!fetcher.RetrieveFile(logger, m_client, casKey, path.data, destinationIsCompressed))
+							return logger.Error(TC("Failed to download cache output %s for %s"), path.data, info.description);
 					}
 
 					cacheStats.fetchBytesRaw += fetcher.sizeOnDisk;
