[UBA]
* Fixed macos compile error

diff --git a/Engine/Source/Programs/UnrealBuildAccelerator/Common/Private/UbaCacheServer.cpp b/Engine/Source/Programs/UnrealBuildAccelerator/Common/Private/UbaCacheServer.cpp
--- a/Engine/Source/Programs/UnrealBuildAccelerator/Common/Private/UbaCacheServer.cpp
+++ b/Engine/Source/Programs/UnrealBuildAccelerator/Common/Private/UbaCacheServer.cpp
@@ -691,7 +691,9 @@ namespace uba
 			}
 
 			writer.WriteBool(true);
-			m_connections.try_emplace(connectionInfo.GetId(), clientVersion);
+			auto insres = m_connections.try_emplace(connectionInfo.GetId());
+			auto& connection = insres.first->second;
+			connection.clientVersion = clientVersion;
 			return true;
 		}
 		case CacheMessageType_StorePathTable:
diff --git a/Engine/Source/Programs/UnrealBuildAccelerator/Common/Public/UbaCacheEntry.h b/Engine/Source/Programs/UnrealBuildAccelerator/Common/Public/UbaCacheEntry.h
--- a/Engine/Source/Programs/UnrealBuildAccelerator/Common/Public/UbaCacheEntry.h
+++ b/Engine/Source/Programs/UnrealBuildAccelerator/Common/Public/UbaCacheEntry.h
@@ -19,7 +19,12 @@ namespace uba
 		Vector<u8> inputCasKeyOffsets;
 		#endif
 
-		struct Range { u32 begin; u32 end; bool operator==(const Range& o) const { return begin == o.begin && end == o.end; } };
+		struct Range {
+			u32 begin; u32 end;
+			Range() = default;
+			Range(u32 b, u32 e) : begin(b), end(e) {}
+			bool operator==(const Range& o) const { return begin == o.begin && end == o.end; }
+		};
 		Vector<Range> sharedInputCasKeyOffsetRanges;
 		Vector<u8> extraInputCasKeyOffsets;
 		Vector<u8> outputCasKeyOffsets;
