[Uba]
* Fix for mistake where clientversion 3 was commented out for service

diff --git a/Engine/Source/Programs/UnrealBuildAccelerator/Common/Private/UbaCacheEntry.cpp b/Engine/Source/Programs/UnrealBuildAccelerator/Common/Private/UbaCacheEntry.cpp
--- a/Engine/Source/Programs/UnrealBuildAccelerator/Common/Private/UbaCacheEntry.cpp
+++ b/Engine/Source/Programs/UnrealBuildAccelerator/Common/Private/UbaCacheEntry.cpp
@@ -47,7 +47,6 @@ namespace uba
 		u16& entryCount = *(u16*)writer.AllocWrite(2);
 		entryCount = 0;
 
-		#if UBA_USE_OLD
 		if (clientVersion == 3)
 		{
 			UBA_ASSERT(!toDisk);
@@ -71,7 +70,6 @@ namespace uba
 			}
 			return true;
 		}
-		#endif
 
 		{
 			auto& shared = sharedInputCasKeyOffsets;
