[Uba]
* Removed GetNowFileTime. GetSystemTimeAsFileTime should be used instead
* Fixed some minor bugs in uba cache server related to times
* Added obliterate command to cache server

diff --git a/Engine/Source/Programs/UnrealBuildAccelerator/Common/Private/UbaCacheServer.cpp b/Engine/Source/Programs/UnrealBuildAccelerator/Common/Private/UbaCacheServer.cpp
--- a/Engine/Source/Programs/UnrealBuildAccelerator/Common/Private/UbaCacheServer.cpp
+++ b/Engine/Source/Programs/UnrealBuildAccelerator/Common/Private/UbaCacheServer.cpp
@@ -158,7 +158,7 @@ namespace uba
 		}
 
 		u64 duration = GetTime() - startTime;
-		m_logger.Detail(TC("Database loaded from %s in %s (%llu bucket(s) containing %s paths, %s keys, %llu cache entries)"), fileName.data, TimeToText(duration).str, m_buckets.size(), BytesToText(totalPathTableSize).str, BytesToText(totalCasKeyTableSize).str, totalCacheEntryCount);
+		m_logger.Detail(TC("Database (v%u) loaded from %s in %s (%llu bucket(s) containing %s paths, %s keys, %llu cache entries)"), databaseVersion, fileName.data, TimeToText(duration).str, m_buckets.size(), BytesToText(totalPathTableSize).str, BytesToText(totalCasKeyTableSize).str, totalCacheEntryCount);
 		return true;
 	}
 
@@ -306,9 +306,23 @@ namespace uba
 				m_isRunningMaintenance = false;
 			});
 
+
 		bool forceAllSteps = true;
 
-		m_logger.Info(TC("Maintenance started after %u added cache entries"), m_addsSinceMaintenance.load());
+		if (m_shouldWipe)
+		{
+			m_shouldWipe = false;
+			m_logger.Info(TC("Obliterating database"));
+			m_addsSinceMaintenance = 0;
+			m_longestMaintenance = 0;
+			m_buckets.clear();
+			forceAllSteps = true;
+			m_creationTime = GetSystemTimeAsFileTime();
+		}
+		else
+		{
+			m_logger.Info(TC("Maintenance started after %u added cache entries"), m_addsSinceMaintenance.load());
+		}
 		
 		u64 startTime = GetTime();
 
@@ -333,7 +347,7 @@ namespace uba
 			return true;
 
 		u64 now = GetSystemTimeAsFileTime();
-		u64 oldest = now;
+		u64 oldest = 0;
 
 		u32 workerCount = m_server.GetWorkerCount();
 		u32 workerCountToUse = workerCount > 0 ? workerCount - 1 : 0;
@@ -343,7 +357,7 @@ namespace uba
 		u64 deleteCacheEntriesStartTime = GetTime();
 		do
 		{
-			oldest = now;
+			oldest = 0;
 
 			ReaderWriterLock existingCasLock;
 
@@ -628,7 +642,7 @@ namespace uba
 			m_logger.Detail(TC("  Saved to disk (%s)"), TimeToText(GetTime() - saveStart).str);
 		}
 
-		u64 oldestTime = MsToTime(GetFileTimeAsSeconds(now - oldest)*1000);
+		u64 oldestTime = oldest ? GetFileTimeAsTime(now - (m_creationTime + oldest)) : 0;
 		u64 duration = GetTime() - startTime;
 		m_logger.Info(TC("Maintenance done! (%s) CasFiles: %llu (%s) OldestEntry: %s"), TimeToText(duration).str, totalCasCount - deletedCasCount, BytesToText(totalCasSize).str, TimeToText(oldestTime, true).str);
 		
@@ -1131,7 +1145,7 @@ namespace uba
 							continue;
 						}
 
-						u64 age = MsToTime(GetFileTimeAsSeconds(now - entry.creationTime)*1000);
+						u64 age = GetFileTimeAsTime(now - entry.creationTime);
 						writeLine(line.Clear().Appendf(TC("  #%u (%s ago)"), index, TimeToText(age, true).str).data);
 
 						auto writeOffsets = [&](const Vector<u8>& offsets)
@@ -1160,19 +1174,40 @@ namespace uba
 		else if (command.Equals(TC("status")))
 		{
 			writeLine(TC("UbaCacheServer status"));
-			writeLine(line.Clear().Appendf(TC("  CreationTime: %s"), TimeToText(MsToTime(1000*GetFileTimeAsSeconds(GetNowFileTime() - m_creationTime)), true).str).data);
+			writeLine(line.Clear().Appendf(TC("  CreationTime: %s ago"), TimeToText(GetFileTimeAsTime(GetSystemTimeAsFileTime() - m_creationTime), true).str).data);
 			writeLine(line.Clear().Appendf(TC("  UpTime: %s"), TimeToText(GetTime() - m_startTime, true).str).data);
 			writeLine(line.Clear().Appendf(TC("  Longest maintenance: %s"), TimeToText(m_longestMaintenance).str).data);
 			writeLine(line.Clear().Appendf(TC("  Buckets:")).data);
 			u32 index = 0;
-			for (auto& kv : m_buckets)
+
 			{
-				Bucket& bucket = kv.second;
-				writeLine(line.Clear().Appendf(TC("    #%u - %llu"), index++, kv.first).data);
-				writeLine(line.Clear().Appendf(TC("      PathTable: %s"), BytesToText(bucket.m_pathTable.GetSize()).str).data);
-				writeLine(line.Clear().Appendf(TC("      CasKeyTable: %s"), BytesToText(bucket.m_casKeyTable.GetSize()).str).data);
+				SCOPED_READ_LOCK(m_bucketsLock, bucketsLock);
+				for (auto& kv : m_buckets)
+				{
+					Bucket& bucket = kv.second;
+					SCOPED_READ_LOCK(bucket.m_cacheEntryLookupLock, lock2);
+					u64 mostEntries = 0;
+					u64 lastUsed = 0;
+					for (auto& kv2 : bucket.m_cacheEntryLookup)
+					{
+						CacheEntries& entries = kv2.second;
+						SCOPED_READ_LOCK(entries.lock, lock3);
+						mostEntries = Max(mostEntries, entries.entries.size());
+						for (auto& entry : entries.entries)
+							lastUsed = Max(lastUsed, entry.lastUsedTime);
+					}
+					lock2.Leave();
+					u64 lastUsedTime = 0;
+					if (lastUsed)
+						lastUsedTime = GetFileTimeAsTime(GetSystemTimeAsFileTime() - (m_creationTime + lastUsed));
+
+					writeLine(line.Clear().Appendf(TC("    #%u - %llu"), index++, kv.first).data);
+					writeLine(line.Clear().Appendf(TC("      PathTable: %s"), BytesToText(bucket.m_pathTable.GetSize()).str).data);
+					writeLine(line.Clear().Appendf(TC("      CasKeyTable: %s"), BytesToText(bucket.m_casKeyTable.GetSize()).str).data);
+					writeLine(line.Clear().Appendf(TC("      KeyMostEntries: %llu"), mostEntries).data);
+					writeLine(line.Clear().Appendf(TC("      LastEntryUsed: %s ago"), TimeToText(lastUsedTime, true).str).data);
+				}
 			}
-
 			u64 totalCasSize = 0;
 			u64 totalCasCount = 0;
 			m_storage.TraverseAllCasFiles([&](const CasKey& casKey, u64 size) { ++totalCasCount; totalCasSize += size; });
@@ -1180,6 +1215,12 @@ namespace uba
 			writeLine(line.Clear().Appendf(TC("    Count: %llu"), totalCasCount).data);
 			writeLine(line.Clear().Appendf(TC("    Size: %s"), BytesToText(totalCasSize).str).data);
 		}
+		else if (command.Equals(TC("obliterate")))
+		{
+			m_shouldWipe = true;
+			m_addsSinceMaintenance = 1;
+			writeLine(line.Clear().Appendf(TC("Cache server database obliteration queued!")).data);
+		}
 		else
 		{
 			writeLine(line.Clear().Appendf(TC("Unknown command: %s"), command.data).data);
diff --git a/Engine/Source/Programs/UnrealBuildAccelerator/Common/Private/UbaFile.cpp b/Engine/Source/Programs/UnrealBuildAccelerator/Common/Private/UbaFile.cpp
--- a/Engine/Source/Programs/UnrealBuildAccelerator/Common/Private/UbaFile.cpp
+++ b/Engine/Source/Programs/UnrealBuildAccelerator/Common/Private/UbaFile.cpp
@@ -796,18 +796,9 @@ namespace uba
 #endif
 	}
 
-	u64 GetNowFileTime()
+	u64 GetFileTimeAsTime(u64 fileTime)
 	{
-		#if PLATFORM_WINDOWS
-		FILETIME ft;
-		SYSTEMTIME st;
-		GetSystemTime(&st);
-		SystemTimeToFileTime(&st, &ft);
-		return (u64&)ft;
-		#else
-		// TODO
-		return 0ull;
-		#endif
+		return MsToTime(GetFileTimeAsSeconds(fileTime)*1000);
 	}
 
 	bool GetCurrentDirectoryW(StringBufferBase& out)
diff --git a/Engine/Source/Programs/UnrealBuildAccelerator/Common/Private/UbaSession.cpp b/Engine/Source/Programs/UnrealBuildAccelerator/Common/Private/UbaSession.cpp
--- a/Engine/Source/Programs/UnrealBuildAccelerator/Common/Private/UbaSession.cpp
+++ b/Engine/Source/Programs/UnrealBuildAccelerator/Common/Private/UbaSession.cpp
@@ -1912,7 +1912,7 @@ namespace uba
 					if (msg.mappingWritten)
 					{
 						writtenFile.mappingWritten = msg.mappingWritten;
-						writtenFile.lastWriteTime = GetNowFileTime();
+						writtenFile.lastWriteTime = GetSystemTimeAsFileTime();
 					}
 					addMapping = false;
 				}
@@ -1941,7 +1941,7 @@ namespace uba
 				writtenFile.mappingHandle = mappingHandle;
 				writtenFile.mappingWritten = msg.mappingWritten;
 				writtenFile.originalMappingHandle = msg.mappingHandle;
-				writtenFile.lastWriteTime = GetNowFileTime();
+				writtenFile.lastWriteTime = GetSystemTimeAsFileTime();
 			}
 
 			if (writtenFile.mappingHandle.IsValid())
@@ -2662,7 +2662,7 @@ namespace uba
 		UnmapViewOfFile(mem2, memoryBlock.writtenSize, TC(""));
 
 		StringKey symFileKey = CaseInsensitiveFs ? ToStringKeyLower(exportsFile) : ToStringKey(exportsFile);
-		u64 lastWriteTime = GetNowFileTime();
+		u64 lastWriteTime = GetSystemTimeAsFileTime();
 
 		if (!RegisterCreateFileForWrite(symFileKey, exportsFile.data, exportsFile.count, false, memoryBlock.writtenSize, lastWriteTime))
 			return false;
diff --git a/Engine/Source/Programs/UnrealBuildAccelerator/Common/Public/UbaCacheServer.h b/Engine/Source/Programs/UnrealBuildAccelerator/Common/Public/UbaCacheServer.h
--- a/Engine/Source/Programs/UnrealBuildAccelerator/Common/Public/UbaCacheServer.h
+++ b/Engine/Source/Programs/UnrealBuildAccelerator/Common/Public/UbaCacheServer.h
@@ -55,6 +55,7 @@ namespace uba
 
 		Atomic<u32> m_addsSinceMaintenance;
 		Atomic<bool> m_isRunningMaintenance;
+		Atomic<bool> m_shouldWipe;
 
 		ReaderWriterLock m_bucketsLock;
 		Map<u64, Bucket> m_buckets;
diff --git a/Engine/Source/Programs/UnrealBuildAccelerator/Common/Public/UbaFile.h b/Engine/Source/Programs/UnrealBuildAccelerator/Common/Public/UbaFile.h
--- a/Engine/Source/Programs/UnrealBuildAccelerator/Common/Public/UbaFile.h
+++ b/Engine/Source/Programs/UnrealBuildAccelerator/Common/Public/UbaFile.h
@@ -75,7 +75,7 @@ namespace uba
 	bool SearchPathW(const tchar* a, const tchar* b, const tchar* c, u32 d, tchar* e, tchar** f);
 	u64 GetSystemTimeAsFileTime();
 	u64 GetFileTimeAsSeconds(u64 fileTime);
-	u64 GetNowFileTime();
+	u64 GetFileTimeAsTime(u64 fileTime);
 	bool GetCurrentDirectoryW(StringBufferBase& out);
 
 
