[UBA]
* Fixed linux compile error

diff --git a/Engine/Source/Programs/UnrealBuildAccelerator/Common/Private/UbaCacheServer.cpp b/Engine/Source/Programs/UnrealBuildAccelerator/Common/Private/UbaCacheServer.cpp
--- a/Engine/Source/Programs/UnrealBuildAccelerator/Common/Private/UbaCacheServer.cpp
+++ b/Engine/Source/Programs/UnrealBuildAccelerator/Common/Private/UbaCacheServer.cpp
@@ -1192,7 +1192,7 @@ namespace uba
 					{
 						CacheEntries& entries = kv2.second;
 						SCOPED_READ_LOCK(entries.lock, lock3);
-						mostEntries = Max(mostEntries, entries.entries.size());
+						mostEntries = Max(mostEntries, u64(entries.entries.size()));
 						for (auto& entry : entries.entries)
 							lastUsed = Max(lastUsed, entry.lastUsedTime);
 					}
