[UBA]
* Change ranges array from array of structs to binary blob with 7bit encoded values
* Changed so temporary vector is reused over multiple container element iterations

diff --git a/Engine/Source/Programs/UnrealBuildAccelerator/Common/Private/UbaCacheEntry.cpp b/Engine/Source/Programs/UnrealBuildAccelerator/Common/Private/UbaCacheEntry.cpp
--- a/Engine/Source/Programs/UnrealBuildAccelerator/Common/Private/UbaCacheEntry.cpp
+++ b/Engine/Source/Programs/UnrealBuildAccelerator/Common/Private/UbaCacheEntry.cpp
@@ -23,11 +23,8 @@ namespace uba
 		auto& extra = entry.extraInputCasKeyOffsets;
 		size += Get7BitEncodedCount(extra.size()) + extra.size();
 
-		u64 rangeSize = 0;
 		auto& ranges = entry.sharedInputCasKeyOffsetRanges;
-		for (auto& range : ranges)
-			rangeSize += Get7BitEncodedCount(range.begin) + Get7BitEncodedCount(range.end);
-		size += Get7BitEncodedCount(rangeSize) + rangeSize;
+		size += Get7BitEncodedCount(ranges.size()) + ranges.size();
 
 		auto& outputs = entry.outputCasKeyOffsets;
 		size += Get7BitEncodedCount(outputs.size()) + outputs.size();
@@ -90,10 +87,6 @@ namespace uba
 			auto& ranges = entry.sharedInputCasKeyOffsetRanges;
 			auto& outputs = entry.outputCasKeyOffsets;
 
-			u64 rangeSize = 0;
-			for (auto& range : ranges)
-				rangeSize += Get7BitEncodedCount(range.begin) + Get7BitEncodedCount(range.end);
-
 			if (toDisk)
 			{
 				writer.Write7BitEncoded(entry.creationTime);
@@ -102,7 +95,7 @@ namespace uba
 			else
 			{
 				u64 neededSize = Get7BitEncodedCount(entry.id) + Get7BitEncodedCount(extra.size()) + extra.size();
-				neededSize += Get7BitEncodedCount(rangeSize) + rangeSize;
+				neededSize += Get7BitEncodedCount(ranges.size()) + ranges.size();
 				neededSize += Get7BitEncodedCount(outputs.size()) + outputs.size();
 
 				if (neededSize > writer.GetCapacityLeft())
@@ -114,12 +107,8 @@ namespace uba
 			writer.Write7BitEncoded(extra.size());
 			writer.WriteBytes(extra.data(), extra.size());
 
-			writer.Write7BitEncoded(rangeSize);
-			for (auto& range : ranges)
-			{
-				writer.Write7BitEncoded(range.begin);
-				writer.Write7BitEncoded(range.end);
-			}
+			writer.Write7BitEncoded(ranges.size());
+			writer.WriteBytes(ranges.data(), ranges.size());
 
 			writer.Write7BitEncoded(outputs.size());
 			writer.WriteBytes(outputs.data(), outputs.size());
@@ -179,18 +168,14 @@ namespace uba
 			entry.id = idCounter++;
 			entry.creationTime = reader.Read7BitEncoded();
 			entry.lastUsedTime = reader.Read7BitEncoded();
+
 			u64 extraSize = reader.Read7BitEncoded();
 			entry.extraInputCasKeyOffsets.resize(extraSize);
 			reader.ReadBytes(entry.extraInputCasKeyOffsets.data(), extraSize);
+
 			u64 rangeSize = reader.Read7BitEncoded();
-			BinaryReader rangeReader(reader.GetPositionData(), 0, rangeSize);
-			reader.Skip(rangeSize);
-			while (rangeReader.GetLeft())
-			{
-				u32 begin = u32(rangeReader.Read7BitEncoded());
-				u32 end = u32(rangeReader.Read7BitEncoded());
-				entry.sharedInputCasKeyOffsetRanges.emplace_back(begin, end);
-			}
+			entry.sharedInputCasKeyOffsetRanges.resize(rangeSize);
+			reader.ReadBytes(entry.sharedInputCasKeyOffsetRanges.data(), rangeSize);
 
 			u64 outputSize = reader.Read7BitEncoded();
 			entry.outputCasKeyOffsets.resize(outputSize);
@@ -203,6 +188,16 @@ namespace uba
 	template<typename Container>
 	void CacheEntries::BuildInputsT(CacheEntry& entry, const Container& sortedInputs, bool populateShared)
 	{
+		StackBinaryWriter<256*1024> rangeWriter;
+
+		auto g = MakeGuard([&]()
+			{
+				entry.sharedInputCasKeyOffsetRanges.resize(rangeWriter.GetPosition());
+				memcpy(entry.sharedInputCasKeyOffsetRanges.data(), rangeWriter.GetData(), rangeWriter.GetPosition());
+			});
+
+		auto writeRange = [&](u64 begin, u64 end) { rangeWriter.Write7BitEncoded(begin); rangeWriter.Write7BitEncoded(end); };
+
 		if (populateShared)
 		{
 			u64 bytes = 0;
@@ -212,7 +207,8 @@ namespace uba
 			BinaryWriter writer(sharedInputCasKeyOffsets.data(), 0, sharedInputCasKeyOffsets.size());
 			for (u32 i : sortedInputs)
 				writer.Write7BitEncoded(i);
-			entry.sharedInputCasKeyOffsetRanges.emplace_back(0, u32(writer.GetPosition()));
+			UBA_ASSERT(bytes == writer.GetPosition());
+			writeRange(0, bytes);
 			return;
 		}
 
@@ -241,7 +237,7 @@ namespace uba
 			{
 				// Add current range if there is one going
 				if (inRange)
-					entry.sharedInputCasKeyOffsetRanges.emplace_back(rangeBegin, sharedPos);
+					writeRange(rangeBegin, sharedPos);
 
 				if (offset > sharedOffset)
 					extraWriter.Write7BitEncoded(offset);
@@ -256,8 +252,7 @@ namespace uba
 			{
 				// Add current range if there is one going
 				if (inRange)
-					entry.sharedInputCasKeyOffsetRanges.emplace_back(rangeBegin, sharedPos);
-
+					writeRange(rangeBegin, sharedPos);
 				if (offset > sharedOffset)
 					extraWriter.Write7BitEncoded(offset);
 				break;
@@ -293,7 +288,7 @@ namespace uba
 				if (inRange)
 				{
 					inRange = false;
-					entry.sharedInputCasKeyOffsetRanges.emplace_back(rangeBegin, sharedPos);
+					writeRange(rangeBegin, sharedPos);
 				}
 				if (offset < sharedOffset)
 				{
@@ -322,10 +317,12 @@ namespace uba
 
 		// Checking if first entry is still the matching entry
 		auto& firstEntry = *entries.begin();
-		if (firstEntry.extraInputCasKeyOffsets.empty() && firstEntry.sharedInputCasKeyOffsetRanges.size() == 1)
+		if (firstEntry.extraInputCasKeyOffsets.empty())
 		{
-			auto& range = firstEntry.sharedInputCasKeyOffsetRanges[0];
-			if (range.begin == 0 && range.end == sharedInputCasKeyOffsets.size())
+			BinaryReader reader(firstEntry.sharedInputCasKeyOffsetRanges.data(), 0, firstEntry.sharedInputCasKeyOffsetRanges.size());
+			u64 begin = reader.Read7BitEncoded();
+			u64 end = reader.Read7BitEncoded();
+			if (begin == 0 && end == sharedInputCasKeyOffsets.size())
 				return;
 		}
 
@@ -343,8 +340,10 @@ namespace uba
 				Flatten(oldShared, entry);
 				oldShared.swap(sharedInputCasKeyOffsets);
 				entry.extraInputCasKeyOffsets.clear();
-				entry.sharedInputCasKeyOffsetRanges.clear();
-				entry.sharedInputCasKeyOffsetRanges.emplace_back(0u, u32(sharedInputCasKeyOffsets.size()));
+				entry.sharedInputCasKeyOffsetRanges.resize(1 + Get7BitEncodedCount(sharedInputCasKeyOffsets.size()));
+				BinaryWriter rangeWriter(entry.sharedInputCasKeyOffsetRanges.data(), 0, entry.sharedInputCasKeyOffsetRanges.size());
+				rangeWriter.Write7BitEncoded(0);
+				rangeWriter.Write7BitEncoded(sharedInputCasKeyOffsets.size());
 				isFirst = false;
 			}
 			else
@@ -358,12 +357,14 @@ namespace uba
 		}
 	}
 
-	void CacheEntries::UpdateEntries(Logger& logger, const GrowingNoLockUnorderedMap<u32, u32>& oldToNewCasKeyOffset)
+	void CacheEntries::UpdateEntries(Logger& logger, const GrowingNoLockUnorderedMap<u32, u32>& oldToNewCasKeyOffset, Vector<u32>& temp)
 	{
 		if (entries.empty())
 			return;
 
-		Vector<u32> temp;
+		#if UBA_USE_OLD
+		ValidateEntries(logger);
+		#endif
 
 		Vector<u8> oldShared;
 		oldShared.swap(sharedInputCasKeyOffsets);
@@ -395,16 +396,18 @@ namespace uba
 					writer.Write7BitEncoded(offset);
 
 				entry.extraInputCasKeyOffsets.clear();
-				entry.sharedInputCasKeyOffsetRanges.clear();
-				entry.sharedInputCasKeyOffsetRanges.emplace_back(0u, u32(sharedInputCasKeyOffsets.size()));
+				u64 rangeSize = 1 + Get7BitEncodedCount(newSize);
+				entry.sharedInputCasKeyOffsetRanges.resize(rangeSize);
+				BinaryWriter rangeWriter(entry.sharedInputCasKeyOffsetRanges.data(), 0, rangeSize);
+				rangeWriter.Write7BitEncoded(0);
+				rangeWriter.Write7BitEncoded(newSize);
+				UBA_ASSERT(rangeWriter.GetPosition() == rangeSize);
 				isFirst = false;
 			}
 			else
 			{
 				// Flatten using old shared and rebuild it with new shared
 				Flatten(temp, entry, oldShared);
-				entry.extraInputCasKeyOffsets.clear();
-				entry.sharedInputCasKeyOffsetRanges.clear();
 				for (auto& offset : temp)
 				{
 					auto findIt = oldToNewCasKeyOffset.find(offset);
@@ -415,6 +418,8 @@ namespace uba
 				// Sort temp now when it likely is out of order
 				std::sort(temp.begin(), temp.end());
 
+				entry.extraInputCasKeyOffsets.clear();
+				entry.sharedInputCasKeyOffsetRanges.clear();
 				BuildInputsT(entry, temp, false);
 			}
 		}
@@ -462,30 +467,29 @@ namespace uba
 	#if UBA_USE_OLD
 	void CacheEntries::ValidateEntries(Logger& logger)
 	{
-		int index = 0;
 		for (auto& entry : entries)
+			ValidateEntry(logger, entry);
+	}
+	void CacheEntries::ValidateEntry(Logger& logger, CacheEntry& entry)
+	{
+		Vector<u8> res;
+		Flatten(res, entry);
+		if (res.size() == entry.inputCasKeyOffsets.size() && memcmp(res.data(), entry.inputCasKeyOffsets.data(), res.size()) == 0)
+			return;
+		BinaryReader reader1(entry.inputCasKeyOffsets.data(), 0, entry.inputCasKeyOffsets.size());
+		BinaryReader reader2(res.data(), 0, res.size());
+		while (reader1.GetLeft() || reader2.GetLeft())
 		{
-			Vector<u8> res;
-			Flatten(res, entry);
-			if (res.size() != entry.inputCasKeyOffsets.size() || memcmp(res.data(), entry.inputCasKeyOffsets.data(), res.size()) != 0)
-			{
-				BinaryReader reader1(entry.inputCasKeyOffsets.data(), 0, entry.inputCasKeyOffsets.size());
-				BinaryReader reader2(res.data(), 0, res.size());
-				while (reader1.GetLeft() || reader2.GetLeft())
-				{
-					u32 a = ~0u;
-					u32 b = ~0u;
-					if (reader1.GetLeft())
-						a = u32(reader1.Read7BitEncoded());
-					if (reader2.GetLeft())
-						b = u32(reader2.Read7BitEncoded());
-
-					logger.Detail(TC("A: %u B: %u"), a, b);
-				}
-				UBA_ASSERT(false);
-			}
-			++index;
+			u32 a = ~0u;
+			u32 b = ~0u;
+			if (reader1.GetLeft())
+				a = u32(reader1.Read7BitEncoded());
+			if (reader2.GetLeft())
+				b = u32(reader2.Read7BitEncoded());
+
+			logger.Detail(TC("A: %u B: %u"), a, b);
 		}
+		UBA_ASSERT(false);
 	}
 
 	#endif
@@ -493,8 +497,16 @@ namespace uba
 	void CacheEntries::Flatten(Vector<u8>& out, const CacheEntry& entry)
 	{
 		u64 size = entry.extraInputCasKeyOffsets.size();
-		for (auto& range : entry.sharedInputCasKeyOffsetRanges)
-			size += range.end - range.begin;
+		{
+			BinaryReader rangeReader(entry.sharedInputCasKeyOffsetRanges.data(), 0, entry.sharedInputCasKeyOffsetRanges.size());
+			while (rangeReader.GetLeft())
+			{
+				u64 begin = rangeReader.Read7BitEncoded();
+				u64 end = rangeReader.Read7BitEncoded();
+				size += end - begin;
+			}
+		}
+
 		out.resize(size);
 		BinaryWriter writer(out.data(), 0, out.size());
 
@@ -514,9 +526,12 @@ namespace uba
 				};
 			};
 
-		for (auto& range : entry.sharedInputCasKeyOffsetRanges)
+		BinaryReader rangeReader(entry.sharedInputCasKeyOffsetRanges.data(), 0, entry.sharedInputCasKeyOffsetRanges.size());
+		while (rangeReader.GetLeft())
 		{
-			BinaryReader inputReader(sharedInputCasKeyOffsets.data() + range.begin, 0, range.end - range.begin);
+			u64 begin = rangeReader.Read7BitEncoded();
+			u64 end = rangeReader.Read7BitEncoded();
+			BinaryReader inputReader(sharedInputCasKeyOffsets.data() + begin, 0, end - begin);
 			while (inputReader.GetLeft())
 			{
 				u32 offset = u32(inputReader.Read7BitEncoded());
@@ -548,9 +563,12 @@ namespace uba
 				};
 			};
 
-		for (auto& range : entry.sharedInputCasKeyOffsetRanges)
+		BinaryReader rangeReader(entry.sharedInputCasKeyOffsetRanges.data(), 0, entry.sharedInputCasKeyOffsetRanges.size());
+		while (rangeReader.GetLeft())
 		{
-			BinaryReader inputReader(sharedOffsets.data() + range.begin, 0, range.end - range.begin);
+			u64 begin = rangeReader.Read7BitEncoded();
+			u64 end = rangeReader.Read7BitEncoded();
+			BinaryReader inputReader(sharedOffsets.data() + begin, 0, end - begin);
 			while (inputReader.GetLeft())
 			{
 				u32 offset = u32(inputReader.Read7BitEncoded());
diff --git a/Engine/Source/Programs/UnrealBuildAccelerator/Common/Private/UbaCacheServer.cpp b/Engine/Source/Programs/UnrealBuildAccelerator/Common/Private/UbaCacheServer.cpp
--- a/Engine/Source/Programs/UnrealBuildAccelerator/Common/Private/UbaCacheServer.cpp
+++ b/Engine/Source/Programs/UnrealBuildAccelerator/Common/Private/UbaCacheServer.cpp
@@ -243,35 +243,11 @@ namespace uba
 			{
 				Write(kv2.first);
 
-				#if UBA_USE_OLD
-				if (CacheFileVersion == 3)
-				{
-					u32 cacheEntryCount = u32(kv2.second.entries.size());
-					Write(cacheEntryCount);
-
-					for (CacheEntry& entry : kv2.second.entries)
-					{
-						WriteBytes(&entry.creationTime, sizeof(entry.creationTime));
-
-						u32 inputSize = u32(entry.inputCasKeyOffsets.size());
-						Write(inputSize);
-						WriteBytes(entry.inputCasKeyOffsets.data(), inputSize);
-
-						u32 outputSize = u32(entry.outputCasKeyOffsets.size());
-						Write(outputSize);
-						WriteBytes(entry.outputCasKeyOffsets.data(), outputSize);
-					}
-				}
-				else
-				#else
-				{
-					temp.resize(kv2.second.GetTotalSize(true));
-					BinaryWriter writer(temp.data(), 0, temp.size());
-					kv2.second.Write(writer, CacheNetworkVersion, true);
-					UBA_ASSERT(writer.GetPosition() == temp.size());
-					WriteBytes(temp.data(), temp.size());
-				}
-				#endif
+				temp.resize(kv2.second.GetTotalSize(true));
+				BinaryWriter writer(temp.data(), 0, temp.size());
+				kv2.second.Write(writer, CacheNetworkVersion, true);
+				UBA_ASSERT(writer.GetPosition() == temp.size());
+				WriteBytes(temp.data(), temp.size());
 			}
 		}
 
@@ -351,6 +327,7 @@ namespace uba
 
 		u32 workerCount = m_server.GetWorkerCount();
 		u32 workerCountToUse = workerCount > 0 ? workerCount - 1 : 0;
+		u32 workerCountToUseForBuckets = Min(workerCountToUse, u32(m_buckets.size()));
 
 		Atomic<u64> deleteEntryCount;
 
@@ -361,7 +338,7 @@ namespace uba
 
 			ReaderWriterLock existingCasLock;
 
-			m_server.ParallelFor(workerCountToUse, m_buckets, [&](auto& it)
+			m_server.ParallelFor(workerCountToUseForBuckets, m_buckets, [&](auto& it)
 			{
 				Vector<u64*> touchedCas;
 
@@ -498,7 +475,7 @@ namespace uba
 			return true;
 
 		Atomic<u32> bucketCounter;
-		m_server.ParallelFor(workerCountToUse, m_buckets, [&](auto& it)
+		m_server.ParallelFor(workerCountToUseForBuckets, m_buckets, [&](auto& it)
 		{
 			u64 bucketStartTime = GetTime();
 
@@ -607,9 +584,9 @@ namespace uba
 				// Update all casKeyOffsets
 				u64 updateEntriesStart = GetTime();
 
-				m_server.ParallelFor(workerCountToUse, bucket.m_cacheEntryLookup, [&](auto& it)
+				m_server.ParallelFor(workerCountToUse, bucket.m_cacheEntryLookup, [&, temp = Vector<u32>()](auto& it) mutable
 					{
-						it->second.UpdateEntries(m_logger, oldToNewCasKeyOffset);
+						it->second.UpdateEntries(m_logger, oldToNewCasKeyOffset, temp);
 					});
 
 				#if 0
@@ -921,6 +898,7 @@ namespace uba
 		#if UBA_USE_OLD
 		// Add new entry
 		newEntry.inputCasKeyOffsets.swap(inputCasKeyOffsets);
+		cacheEntries.ValidateEntry(m_logger, newEntry);
 		#endif
 
 		Set<u32> outputs;
diff --git a/Engine/Source/Programs/UnrealBuildAccelerator/Common/Public/UbaCacheEntry.h b/Engine/Source/Programs/UnrealBuildAccelerator/Common/Public/UbaCacheEntry.h
--- a/Engine/Source/Programs/UnrealBuildAccelerator/Common/Public/UbaCacheEntry.h
+++ b/Engine/Source/Programs/UnrealBuildAccelerator/Common/Public/UbaCacheEntry.h
@@ -18,14 +18,7 @@ namespace uba
 		#if UBA_USE_OLD
 		Vector<u8> inputCasKeyOffsets;
 		#endif
-
-		struct Range {
-			u32 begin; u32 end;
-			Range() = default;
-			Range(u32 b, u32 e) : begin(b), end(e) {}
-			bool operator==(const Range& o) const { return begin == o.begin && end == o.end; }
-		};
-		Vector<Range> sharedInputCasKeyOffsetRanges;
+		Vector<u8> sharedInputCasKeyOffsetRanges;
 		Vector<u8> extraInputCasKeyOffsets;
 		Vector<u8> outputCasKeyOffsets;
 	};
@@ -44,10 +37,11 @@ namespace uba
 		bool Read(Logger& logger, BinaryReader& reader, u32 databaseVersion);
 		void BuildInputs(CacheEntry& entry, const Set<u32>& inputs);
 		void UpdateEntries();
-		void UpdateEntries(Logger& logger, const GrowingNoLockUnorderedMap<u32, u32>& oldToNewCasKeyOffset);
+		void UpdateEntries(Logger& logger, const GrowingNoLockUnorderedMap<u32, u32>& oldToNewCasKeyOffset, Vector<u32>& temp);
 
 		#if UBA_USE_OLD
 		void ValidateEntries(Logger& logger);
+		void ValidateEntry(Logger& logger, CacheEntry& entry);
 		#endif
 
 		void Flatten(Vector<u8>& out, const CacheEntry& entry);
diff --git a/Engine/Source/Programs/UnrealBuildAccelerator/Common/Public/UbaWorkManager.h b/Engine/Source/Programs/UnrealBuildAccelerator/Common/Public/UbaWorkManager.h
--- a/Engine/Source/Programs/UnrealBuildAccelerator/Common/Public/UbaWorkManager.h
+++ b/Engine/Source/Programs/UnrealBuildAccelerator/Common/Public/UbaWorkManager.h
@@ -96,7 +96,7 @@ namespace uba
 		context->isDone = false;
 		context->doneEvent = &doneEvent;
 
-		auto work = [context, &func]()
+		auto work = [context, funcCopy = func]() mutable
 			{
 				u32 active = 0;
 				while (true)
@@ -120,7 +120,7 @@ namespace uba
 					context->activeCount += active;
 					l.Leave();
 
-					func(it2);
+					funcCopy(it2);
 				}
 			};
 		
diff --git a/Engine/Source/Programs/UnrealBuildAccelerator/Host/Private/UbaExports.cpp b/Engine/Source/Programs/UnrealBuildAccelerator/Host/Private/UbaExports.cpp
--- a/Engine/Source/Programs/UnrealBuildAccelerator/Host/Private/UbaExports.cpp
+++ b/Engine/Source/Programs/UnrealBuildAccelerator/Host/Private/UbaExports.cpp
@@ -510,13 +510,12 @@ extern "C"
 		using namespace uba;
 		LogWriter& writer = session->GetLogWriter();
 		StorageImpl& storage = (StorageImpl&)session->GetStorage();
-
-		auto networkBackend = new NetworkBackendTcp(writer);
+		auto& server = (NetworkServerWithBackend&)session->GetServer();
 
 		NetworkClientCreateInfo ncci(writer);
 		ncci.receiveTimeoutSeconds = 60;
 		bool ctorSuccess = false;
-		auto networkClient = new NetworkClientWithBackend(ctorSuccess, ncci, networkBackend);
+		auto networkClient = new NetworkClientWithBackend(ctorSuccess, ncci, server.backend);
 		if (!ctorSuccess)
 			return nullptr;
 		CacheClientCreateInfo info{writer, storage, *networkClient, *session};
