[UBA]
* Fixed visualizer bug where clicking checkboxes causes dragging to hang

diff --git a/Engine/Source/Programs/UnrealBuildAccelerator/Visualizer/Private/UbaVisualizer.cpp b/Engine/Source/Programs/UnrealBuildAccelerator/Visualizer/Private/UbaVisualizer.cpp
--- a/Engine/Source/Programs/UnrealBuildAccelerator/Visualizer/Private/UbaVisualizer.cpp
+++ b/Engine/Source/Programs/UnrealBuildAccelerator/Visualizer/Private/UbaVisualizer.cpp
@@ -566,13 +566,14 @@ namespace uba
 
 	void Visualizer::StopDragToScroll()
 	{
-		--m_dragToScrollCounter;
-		if (m_dragToScrollCounter == 0)
-		{
-			ReleaseCapture();
-			if (UpdateSelection())
-				RedrawWindow(m_hwnd, NULL, NULL, RDW_INVALIDATE);
-		}
+		if (m_dragToScrollCounter > 0)
+			--m_dragToScrollCounter;
+		if (m_dragToScrollCounter != 0)
+			return;
+
+		ReleaseCapture();
+		if (UpdateSelection())
+			RedrawWindow(m_hwnd, NULL, NULL, RDW_INVALIDATE);
 	}
 
 	void Visualizer::PaintClient(const Function<void(HDC hdc, HDC memDC, RECT& clientRect)>& paintFunc)
