[Uba]
* WorkManager - Fixed bug in ParallelFor. deleted event could be set
* CacheService - More optimizations around maintenance

diff --git a/Engine/Source/Programs/UnrealBuildAccelerator/Common/Private/UbaCacheServer.cpp b/Engine/Source/Programs/UnrealBuildAccelerator/Common/Private/UbaCacheServer.cpp
--- a/Engine/Source/Programs/UnrealBuildAccelerator/Common/Private/UbaCacheServer.cpp
+++ b/Engine/Source/Programs/UnrealBuildAccelerator/Common/Private/UbaCacheServer.cpp
@@ -42,8 +42,9 @@ namespace uba
 		CompactPathTable m_pathTable;
 		CompactCasKeyTable m_casKeyTable;
 
-		u64 totalEntryCount = 0;
-		u64 totalEntrySize = 0;
+		Atomic<u64> totalEntryCount;
+		Atomic<u64> totalEntrySize;
+		Atomic<bool> hasDeletedEntries;
 	};
 
 	const tchar* ToString(CacheMessageType type)
@@ -283,7 +284,7 @@ namespace uba
 			});
 
 
-		bool forceAllSteps = true;
+		bool forceAllSteps = false;
 
 		if (m_shouldWipe)
 		{
@@ -305,16 +306,22 @@ namespace uba
 		bool entriesAdded = m_addsSinceMaintenance != 0;
 		m_addsSinceMaintenance = 0;
 
-		Set<CasKey> deletedCasFiles;
+		UnorderedSet<CasKey> deletedCasFiles;
 		m_storage.HandleOverflow(&deletedCasFiles);
 		u64 deletedCasCount = deletedCasFiles.size();
 
 		struct CasFileInfo { u64 size; u64 refCount; };
-		Map<CasKey, CasFileInfo> existingCas;
+		UnorderedMap<CasKey, CasFileInfo> existingCas;
 		u64 totalCasSize = 0;
 		{
+			ReaderWriterLock existingCasLock;
 			u64 traverseStartTime = GetTime();
-			m_storage.TraverseAllCasFiles([&](const CasKey& casKey, u64 size) { totalCasSize += size; existingCas.try_emplace(casKey, CasFileInfo{size, 0ull}); });
+			m_storage.TraverseAllCasFiles([&](const CasKey& casKey, u64 size)
+				{
+					SCOPED_WRITE_LOCK(existingCasLock, lock);
+					totalCasSize += size;
+					existingCas.try_emplace(casKey, CasFileInfo{size, 0ull});
+				}, true);
 			m_logger.Detail(TC("  Found %llu cas files (%s)"), existingCas.size(), TimeToText(GetTime() - traverseStartTime).str);
 		}
 		u64 totalCasCount = existingCas.size() + deletedCasCount;
@@ -330,23 +337,28 @@ namespace uba
 		u32 workerCountToUseForBuckets = Min(workerCountToUse, u32(m_buckets.size()));
 
 		Atomic<u64> deleteEntryCount;
+		Atomic<u64> activeDropCount;
+		auto dropCasGuard = MakeGuard([&]() { while (activeDropCount != 0) Sleep(1); });
 
 		u64 deleteCacheEntriesStartTime = GetTime();
 		do
 		{
+			bool checkInputsForDeletes = m_checkInputsForDeletedCas && !deletedCasFiles.empty();
+
 			oldest = 0;
 
 			ReaderWriterLock existingCasLock;
 
 			m_server.ParallelFor(workerCountToUseForBuckets, m_buckets, [&](auto& it)
 			{
-				Vector<u64*> touchedCas;
-
 				Bucket& bucket = it->second;
 				bucket.totalEntryCount = 0;
 				bucket.totalEntrySize = 0;
 
-				for (auto li=bucket.m_cacheEntryLookup.begin(), le=bucket.m_cacheEntryLookup.end(); li!=le;)
+				ReaderWriterLock keysToEraseLock;
+				Vector<CasKey> keysToErase;
+
+				m_server.ParallelFor(workerCountToUse, bucket.m_cacheEntryLookup, [&, touchedCas = Vector<u64*>()](auto& li) mutable
 				{
 					CacheEntries& entries = li->second;
 
@@ -354,11 +366,26 @@ namespace uba
 					// all the entries that overflow that number
 					u64 capacityLeft = SendMaxSize - 32 - entries.GetSharedSize();
 
+					auto IsOffsetDeleted = [&](u64 offset) { CasKey casKey; bucket.m_casKeyTable.GetKey(casKey, offset); return deletedCasFiles.find(casKey) != deletedCasFiles.end(); };
+
+					// Check if any offset has been deleted in shared offsets..
+					bool offsetDeletedInShared = false;
+					auto& sharedOffsets = entries.sharedInputCasKeyOffsets;
+					if (checkInputsForDeletes)
+					{
+						BinaryReader reader2(sharedOffsets.data(), 0, sharedOffsets.size());
+						while (reader2.GetLeft())
+						{
+							if (!IsOffsetDeleted(reader2.Read7BitEncoded()))
+								continue;
+							offsetDeletedInShared = true;
+							break;
+						}
+					}
+
 					for (auto i=entries.entries.begin(), e=entries.entries.end(); i!=e;)
 					{
 						auto& entry = *i;
-						auto& outputs = entry.outputCasKeyOffsets;
-
 						bool deleteEntry = false;
 
 						u64 neededSize = entries.GetEntrySize(entry, false);
@@ -367,51 +394,72 @@ namespace uba
 							deleteEntry = true;
 							capacityLeft = 0;
 						}
-						else
+
+						// This is an attempt at removing entries that has inputs that depends on other entries outputs.
+						// and that there is no point keeping them if the other entry is removed
+						// Example would be that there is no idea keeping entries that uses a pch if the entry producing the pch is gone
+						if (checkInputsForDeletes)
 						{
-							// Traverse outputs and check if cas files exists for each output, if not, delete entry.
-							touchedCas.clear();
-							BinaryReader reader2(outputs.data(), 0, outputs.size());
-							while (reader2.GetLeft())
+							if (!deleteEntry && offsetDeletedInShared)
 							{
-								u64 offset = reader2.Read7BitEncoded();
-								CasKey casKey;
-								bucket.m_casKeyTable.GetKey(casKey, offset);
-								UBA_ASSERT(IsCompressed(casKey));
-								auto findIt = existingCas.find(casKey);
-								if (findIt == existingCas.end())
+								BinaryReader rangeReader(entry.sharedInputCasKeyOffsetRanges.data(), 0, entry.sharedInputCasKeyOffsetRanges.size());
+								while (!deleteEntry && rangeReader.GetLeft())
+								{
+									u64 begin = rangeReader.Read7BitEncoded();
+									u64 end = rangeReader.Read7BitEncoded();
+									BinaryReader inputReader(sharedOffsets.data() + begin, 0, end - begin);
+									while (inputReader.GetLeft())
+									{
+										if (!IsOffsetDeleted(inputReader.Read7BitEncoded()))
+											continue;
+										deleteEntry = true;
+										break;
+									}
+								}
+							}
+
+							if (!deleteEntry)
+							{
+								auto& extraInputs = entry.extraInputCasKeyOffsets;
+								BinaryReader extraReader(extraInputs.data(), 0, extraInputs.size());
+								while (extraReader.GetLeft())
+								{
+									if (!IsOffsetDeleted(extraReader.Read7BitEncoded()))
+										continue;
 									deleteEntry = true;
-								else
-									touchedCas.push_back(&findIt->second.refCount);
+									break;
+								}
 							}
 						}
 
-						// This is an attempt at removing entries that has inputs that depends on other entries outputs.
-						// and that there is no point keeping them if the other entry is removed
-						// Example would be that there is no idea keeping entries that uses a pch if the entry producing the pch is gone
-						#if 0
-						// Check if there are keys that use removed caskey as input and delete those too
-						auto& inputs = entry.inputCasKeyOffsets;
-						if (!deleteEntry && !deletedCasFiles.empty())
+						if (!deleteEntry)
 						{
-							BinaryReader reader2(inputs.data(), 0, inputs.size());
-							while (reader2.GetLeft())
+							// Traverse outputs and check if cas files exists for each output, if not, delete entry.
+							touchedCas.clear();
+
+							auto& outputs = entry.outputCasKeyOffsets;
+							BinaryReader outputsReader(outputs.data(), 0, outputs.size());
+							while (outputsReader.GetLeft())
 							{
-								u64 offset = reader2.Read7BitEncoded();
+								u64 offset = outputsReader.Read7BitEncoded();
 								CasKey casKey;
 								bucket.m_casKeyTable.GetKey(casKey, offset);
-								if (deletedCasFiles.find(casKey) != deletedCasFiles.end())
+								UBA_ASSERT(IsCompressed(casKey));
+								auto findIt = existingCas.find(casKey);
+								if (findIt != existingCas.end())
 								{
-									deleteEntry = true;
-									break;
+									touchedCas.push_back(&findIt->second.refCount);
+									continue;
 								}
+								deleteEntry = true;
+								break;
 							}
 						}
-						#endif
 
 						// Remove entry from entries list and skip increasing ref count of cas files
 						if (deleteEntry)
 						{
+							bucket.hasDeletedEntries = true;
 							++deleteEntryCount;
 							i = entries.entries.erase(i);
 							e = entries.entries.end();
@@ -433,17 +481,18 @@ namespace uba
 						++i;
 					}
 
-					if (!entries.entries.empty())
+
+					if (entries.entries.empty())
 					{
-						//entries.UpdateEntries(); // Not needed.. 
-						bucket.totalEntrySize += entries.GetTotalSize(false);
-						++li;
-						continue;
+						SCOPED_WRITE_LOCK(keysToEraseLock, lock2);
+						keysToErase.push_back(li->first);
 					}
+					else
+						bucket.totalEntrySize += entries.GetTotalSize(false);
+				});
 
-					li = bucket.m_cacheEntryLookup.erase(li);
-					le = bucket.m_cacheEntryLookup.end();
-				}
+				for (auto& key : keysToErase)
+					bucket.m_cacheEntryLookup.erase(key);
 			});
 
 			// Reset deleted cas files and update it again..
@@ -464,8 +513,12 @@ namespace uba
 				e = existingCas.end();
 			}
 
+			// Add drop cas as work so it can run in the background
 			for (auto& casKey : deletedCasFiles)
-				m_storage.DropCasFile(casKey, true, TC(""));
+			{
+				++activeDropCount;
+				m_server.AddWork([&, key = casKey]() { m_storage.DropCasFile(key, true, TC("")); --activeDropCount; }, 1, TC(""));
+			}
 		}
 		while (!deletedCasFiles.empty()); // if cas files are deleted we need to do another loop and check cache entry inputs to see if files were inputs
 
@@ -482,7 +535,7 @@ namespace uba
 			Bucket& bucket = it->second;
 			u32 bucketIndex = bucketCounter++;
 
-			if (!deleteEntryCount && !forceAllSteps)
+			if (!bucket.hasDeletedEntries && !forceAllSteps)
 			{
 				m_logger.Detail(TC("    Bucket %u skipped updating. No entries deleted"), bucketIndex);
 				return;
@@ -608,9 +661,12 @@ namespace uba
 				m_logger.Detail(TC("    Bucket %u Updated cache entries with new tables (%s)"), bucketIndex, TimeToText(GetTime() - updateEntriesStart).str);
 			}
 
-			m_logger.Info(TC("    Bucket %u Done (%s). CacheEntries: %llu (%s) PathTable: %s CasTable: %s"), bucketIndex, TimeToText(GetTime() - bucketStartTime).str, bucket.totalEntryCount, BytesToText(bucket.totalEntrySize).str, BytesToText(bucket.m_pathTable.GetSize()).str, BytesToText(bucket.m_casKeyTable.GetSize()).str);
+			m_logger.Info(TC("    Bucket %u Done (%s). CacheEntries: %llu (%s) PathTable: %s CasTable: %s"), bucketIndex, TimeToText(GetTime() - bucketStartTime).str, bucket.totalEntryCount.load(), BytesToText(bucket.totalEntrySize.load()).str, BytesToText(bucket.m_pathTable.GetSize()).str, BytesToText(bucket.m_casKeyTable.GetSize()).str);
 		});
 
+		// Need to make sure all cas entries are dropped before saving cas table
+		dropCasGuard.Execute();
+
 		if (entriesAdded || deletedCasCount || deleteEntryCount || forceAllSteps)
 		{
 			u64 saveStart = GetTime();
diff --git a/Engine/Source/Programs/UnrealBuildAccelerator/Common/Private/UbaStorage.cpp b/Engine/Source/Programs/UnrealBuildAccelerator/Common/Private/UbaStorage.cpp
--- a/Engine/Source/Programs/UnrealBuildAccelerator/Common/Private/UbaStorage.cpp
+++ b/Engine/Source/Programs/UnrealBuildAccelerator/Common/Private/UbaStorage.cpp
@@ -936,14 +936,14 @@ namespace uba
 			m_workManager->DoWork();
 	}
 
-	void StorageImpl::TraverseAllCasFiles(const Function<void(const CasKey& key, u64 size)>& func)
+	void StorageImpl::TraverseAllCasFiles(const Function<void(const CasKey& key, u64 size)>& func, bool allowParallel)
 	{
 		StringBuffer<> casRoot;
 		casRoot.Append(m_rootDir.data, m_rootDir.count - 1);
 		TraverseAllCasFiles(casRoot.data, [&](const StringBufferBase& fullPath, const DirectoryEntry& e)
 			{
 				func(CasKeyFromString(e.name), e.size);
-			});
+			}, allowParallel);
 	}
 
 	bool StorageImpl::CheckAllCasFiles(u64 checkContentOfFilesNewerThanTime)
@@ -1101,7 +1101,7 @@ namespace uba
 		return true;
 	}
 
-	void StorageImpl::HandleOverflow(Set<CasKey>* outDeletedFiles)
+	void StorageImpl::HandleOverflow(UnorderedSet<CasKey>* outDeletedFiles)
 	{
 		if (!m_casCapacityBytes)
 			return;
diff --git a/Engine/Source/Programs/UnrealBuildAccelerator/Common/Public/UbaCacheServer.h b/Engine/Source/Programs/UnrealBuildAccelerator/Common/Public/UbaCacheServer.h
--- a/Engine/Source/Programs/UnrealBuildAccelerator/Common/Public/UbaCacheServer.h
+++ b/Engine/Source/Programs/UnrealBuildAccelerator/Common/Public/UbaCacheServer.h
@@ -69,6 +69,8 @@ namespace uba
 		u64 m_startTime = 0;
 		u64 m_longestMaintenance = 0;
 
+		bool m_checkInputsForDeletedCas = true;
+
 		CacheServer(const CacheServer&) = delete;
 		CacheServer& operator=(const CacheServer&) = delete;
 	};
diff --git a/Engine/Source/Programs/UnrealBuildAccelerator/Common/Public/UbaStorage.h b/Engine/Source/Programs/UnrealBuildAccelerator/Common/Public/UbaStorage.h
--- a/Engine/Source/Programs/UnrealBuildAccelerator/Common/Public/UbaStorage.h
+++ b/Engine/Source/Programs/UnrealBuildAccelerator/Common/Public/UbaStorage.h
@@ -165,9 +165,9 @@ namespace uba
 		void AttachEntry(CasEntry& entry);
 		void DetachEntry(CasEntry& entry);
 		void TraverseAllCasFiles(const tchar* dir, const Function<void(const StringBufferBase& fullPath, const DirectoryEntry& e)>& func, bool allowParallel = false);
-		void TraverseAllCasFiles(const Function<void(const CasKey& key, u64 size)>& func);
+		void TraverseAllCasFiles(const Function<void(const CasKey& key, u64 size)>& func, bool allowParallel = false);
 		bool CheckAllCasFiles(u64 checkContentOfFilesNewerThanTime = ~u64(0));
-		void HandleOverflow(Set<CasKey>* outDeletedFiles);
+		void HandleOverflow(UnorderedSet<CasKey>* outDeletedFiles);
 		bool OpenCasDataFile(u32 index, u64 size);
 		bool CreateCasDataFiles();
 
diff --git a/Engine/Source/Programs/UnrealBuildAccelerator/Common/Public/UbaWorkManager.h b/Engine/Source/Programs/UnrealBuildAccelerator/Common/Public/UbaWorkManager.h
--- a/Engine/Source/Programs/UnrealBuildAccelerator/Common/Public/UbaWorkManager.h
+++ b/Engine/Source/Programs/UnrealBuildAccelerator/Common/Public/UbaWorkManager.h
@@ -106,8 +106,11 @@ namespace uba
 					context->isDone = context->it == context->end;
 					if (context->isDone)
 					{
-						if (context->activeCount == 0)
+						if (context->activeCount == 0 && context->doneEvent)
+						{
 							context->doneEvent->Set();
+							context->doneEvent = nullptr;
+						}
 						if (--context->refCount)
 							return;
 						l.Leave();
