[UBA]
* Added more times to steps to be able to track what is taking time

diff --git a/Engine/Source/Programs/UnrealBuildAccelerator/Common/Private/UbaCacheServer.cpp b/Engine/Source/Programs/UnrealBuildAccelerator/Common/Private/UbaCacheServer.cpp
--- a/Engine/Source/Programs/UnrealBuildAccelerator/Common/Private/UbaCacheServer.cpp
+++ b/Engine/Source/Programs/UnrealBuildAccelerator/Common/Private/UbaCacheServer.cpp
@@ -665,7 +665,11 @@ namespace uba
 		});
 
 		// Need to make sure all cas entries are dropped before saving cas table
+		u64 dropStartTime = GetTime();
 		dropCasGuard.Execute();
+		u64 dropCasDuration = GetTime() - dropStartTime;
+		if (TimeToMs(dropCasDuration) > 10)
+			m_logger.Detail(TC("  Done deleting cas files (%s)"), TimeToText(dropCasDuration).str);
 
 		if (entriesAdded || deletedCasCount || deleteEntryCount || forceAllSteps)
 		{
diff --git a/Engine/Source/Programs/UnrealBuildAccelerator/Common/Private/UbaStorage.cpp b/Engine/Source/Programs/UnrealBuildAccelerator/Common/Private/UbaStorage.cpp
--- a/Engine/Source/Programs/UnrealBuildAccelerator/Common/Private/UbaStorage.cpp
+++ b/Engine/Source/Programs/UnrealBuildAccelerator/Common/Private/UbaStorage.cpp
@@ -1105,6 +1105,7 @@ namespace uba
 	{
 		if (!m_casCapacityBytes)
 			return;
+		u64 startTime = GetTime();
 		u64 before = m_casTotalBytes;
 		while (m_casTotalBytes > m_casCapacityBytes)
 		{
@@ -1123,7 +1124,7 @@ namespace uba
 		}
 		u64 after = m_casTotalBytes;
 		if (before != after)
-			m_logger.Info(TC("Evicted %s from storage. Estimated new storage is now %s (there might be files db is not aware of)"), BytesToText(before - after).str, BytesToText(after).str);
+			m_logger.Info(TC("Evicted %s from storage (%s). Estimated new storage is now %s (there might be files db is not aware of)"), BytesToText(before - after).str, TimeToText(GetTime() - startTime).str, BytesToText(after).str);
 	}
 
 	bool StorageImpl::OpenCasDataFile(u32 index, u64 size)
