[UBA]
* Fixed but where variable tracking if entries have been deleted in bucket was never reset

diff --git a/Engine/Source/Programs/UnrealBuildAccelerator/Common/Private/UbaCacheServer.cpp b/Engine/Source/Programs/UnrealBuildAccelerator/Common/Private/UbaCacheServer.cpp
--- a/Engine/Source/Programs/UnrealBuildAccelerator/Common/Private/UbaCacheServer.cpp
+++ b/Engine/Source/Programs/UnrealBuildAccelerator/Common/Private/UbaCacheServer.cpp
@@ -540,7 +540,8 @@ namespace uba
 				m_logger.Detail(TC("    Bucket %u skipped updating. No entries deleted"), bucketIndex);
 				return;
 			}
-			
+			bucket.hasDeletedEntries = false;
+
 			MemoryBlock memoryBlock(128*1024*1024);
 
 			GrowingNoLockUnorderedSet<u32> usedCasKeyOffsets(&memoryBlock);
@@ -1226,6 +1227,7 @@ namespace uba
 					SCOPED_READ_LOCK(bucket.m_cacheEntryLookupLock, lock2);
 					u64 mostEntries = 0;
 					u64 lastUsed = 0;
+					u64 totalEntryCount = 0;
 					for (auto& kv2 : bucket.m_cacheEntryLookup)
 					{
 						CacheEntries& entries = kv2.second;
@@ -1233,6 +1235,7 @@ namespace uba
 						mostEntries = Max(mostEntries, u64(entries.entries.size()));
 						for (auto& entry : entries.entries)
 							lastUsed = Max(lastUsed, entry.lastUsedTime);
+						totalEntryCount += entries.entries.size();
 					}
 					lock2.Leave();
 					u64 lastUsedTime = 0;
@@ -1240,8 +1243,9 @@ namespace uba
 						lastUsedTime = GetFileTimeAsTime(GetSystemTimeAsFileTime() - (m_creationTime + lastUsed));
 
 					writeLine(line.Clear().Appendf(TC("    #%u - %llu"), index++, kv.first).data);
-					writeLine(line.Clear().Appendf(TC("      PathTable: %s"), BytesToText(bucket.m_pathTable.GetSize()).str).data);
-					writeLine(line.Clear().Appendf(TC("      CasKeyTable: %s"), BytesToText(bucket.m_casKeyTable.GetSize()).str).data);
+					writeLine(line.Clear().Appendf(TC("      PathTable: %lls (%s)"), bucket.m_pathTable.GetPathCount(), BytesToText(bucket.m_pathTable.GetSize()).str).data);
+					writeLine(line.Clear().Appendf(TC("      CasKeyTable: %llu (%s)"), bucket.m_cacheEntryLookup.size(), BytesToText(bucket.m_casKeyTable.GetSize()).str).data);
+					writeLine(line.Clear().Appendf(TC("      TotalEntries: %llu"), totalEntryCount).data);
 					writeLine(line.Clear().Appendf(TC("      KeyMostEntries: %llu"), mostEntries).data);
 					writeLine(line.Clear().Appendf(TC("      LastEntryUsed: %s ago"), TimeToText(lastUsedTime, true).str).data);
 				}
