[UBA]
* Added StringView type and replaced a ton of string+stringLen to StringView
* Fixed bug in obj file preloader.
* Fixed so clang targets can used compressed obj files
* Fixed bug in cache client when entry count is zero
* Added magic to compressed obj file header to be able to see if an obj file is compressed or not
* Moved ParseArguments to its own file and added so it can handle both char and wchar

diff --git a/Engine/Source/Programs/UnrealBuildAccelerator/Common/Private/UbaCacheClient.cpp b/Engine/Source/Programs/UnrealBuildAccelerator/Common/Private/UbaCacheClient.cpp
--- a/Engine/Source/Programs/UnrealBuildAccelerator/Common/Private/UbaCacheClient.cpp
+++ b/Engine/Source/Programs/UnrealBuildAccelerator/Common/Private/UbaCacheClient.cpp
@@ -227,7 +227,7 @@ namespace uba
 			if (casKey == CasKeyZero)
 			{
 				bool deferCreation = true;
-				bool fileIsCompressed = IsFileCompressed(info, path.data, path.count);
+				bool fileIsCompressed = IsFileCompressed(info, path);
 				if (!m_storage.StoreCasFile(casKey, path.data, CasKeyZero, deferCreation, fileIsCompressed))
 					return false;
 				if (casKey == CasKeyZero) // If file is not found it was a temporary file that was deleted and is not really an output
@@ -342,7 +342,8 @@ namespace uba
 		UnorderedMap<u32, bool> offsetIsMatch;
 
 		// Traverse entries and test inputs against local machine
-		u32 entryCount = reader.ReadU16();
+		CacheEntriesTraverser traverser(reader);
+		u32 entryCount = traverser.entryCount;
 
 		#if UBA_LOG_FETCH_CACHE_INFO
 		auto mg = MakeGuard([&]()
@@ -355,8 +356,6 @@ namespace uba
 		struct MissInfo { TString path; u32 entryIndex; CasKey cache; CasKey local; };
 		Vector<MissInfo> misses;
 
-		CacheEntriesTraverser traverser(reader);
-
 		u32 entryIndex = 0;
 		for (; entryIndex!=entryCount; ++entryIndex)
 		{
@@ -390,7 +389,7 @@ namespace uba
 							else
 							{
 								bool deferCreation = true;
-								bool fileIsCompressed = IsFileCompressed(info, path.data, path.count);
+								bool fileIsCompressed = IsFileCompressed(info, path);
 								m_storage.StoreCasFile(localCasKey, path.data, CasKeyZero, deferCreation, fileIsCompressed);
 								UBA_ASSERT(localCasKey == CasKeyZero || IsCompressed(localCasKey));
 							}
@@ -511,7 +510,7 @@ namespace uba
 					else
 					{
 						DowngradedLogger logger(m_logger.m_writer, TC("UbaCacheClientDownload"));
-						bool destinationIsCompressed = IsFileCompressed(info, path.data, path.count);
+						bool destinationIsCompressed = IsFileCompressed(info, path);
 						if (!fetcher.RetrieveFile(logger, m_client, casKey, path.data, destinationIsCompressed))
 							return logger.Error(TC("Failed to download cache output %s for %s"), path.data, info.description);
 					}
@@ -941,13 +940,13 @@ namespace uba
 		return true;
 	}
 
-	bool CacheClient::IsFileCompressed(const ProcessStartInfo& info, const tchar* filename, u64 filenameSize)
+	bool CacheClient::IsFileCompressed(const ProcessStartInfo& info, const StringView& filename)
 	{
 		if (!m_session.ShouldStoreObjFilesCompressed())
 			return false;
 		auto rules = info.rules;
 		if (!rules)
 			rules = m_session.GetRules(info);
-		return rules->StoreFileCompressed(filename, filenameSize);
+		return rules->StoreFileCompressed(filename);
 	}
 }
diff --git a/Engine/Source/Programs/UnrealBuildAccelerator/Common/Private/UbaObjectFileCoff.cpp b/Engine/Source/Programs/UnrealBuildAccelerator/Common/Private/UbaObjectFileCoff.cpp
--- a/Engine/Source/Programs/UnrealBuildAccelerator/Common/Private/UbaObjectFileCoff.cpp
+++ b/Engine/Source/Programs/UnrealBuildAccelerator/Common/Private/UbaObjectFileCoff.cpp
@@ -175,7 +175,7 @@ namespace uba
 	};
 
 	template<typename SymbolType>
-	StringView GetSymbolName(SymbolType& symbol, const u8* data, u32 stringTableMemPos)
+	ObjectFile::StringView GetSymbolName(SymbolType& symbol, const u8* data, u32 stringTableMemPos)
 	{
 		if (symbol.N.Name.Short == 0)
 		{
diff --git a/Engine/Source/Programs/UnrealBuildAccelerator/Common/Private/UbaProcess.cpp b/Engine/Source/Programs/UnrealBuildAccelerator/Common/Private/UbaProcess.cpp
--- a/Engine/Source/Programs/UnrealBuildAccelerator/Common/Private/UbaProcess.cpp
+++ b/Engine/Source/Programs/UnrealBuildAccelerator/Common/Private/UbaProcess.cpp
@@ -4,6 +4,7 @@
 #include "UbaFileAccessor.h"
 #include "UbaProtocol.h"
 #include "UbaProcessStats.h"
+#include "UbaProcessUtils.h"
 #include "UbaApplicationRules.h"
 
 #if PLATFORM_WINDOWS
@@ -611,7 +612,7 @@ namespace uba
 		SetFileAttributes(toName.data, DefaultAttributes());
 
 		StringKey toKey = ToStringKeyLower(toName);
-		m_session.RegisterCreateFileForWrite(toKey, toName.data, toName.count, true);
+		m_session.RegisterCreateFileForWrite(toKey, toName, true);
 
 		//m_writtenFiles.try_emplace(name);
 		//WrittenFile& writtenFile = m_writtenFiles[toName.data];
@@ -1546,7 +1547,7 @@ namespace uba
 			}
 
 			Vector<TString> arguments;
-			if (!ParseArguments(arguments, m_startInfo.arguments))
+			if (!ParseArguments(m_startInfo.arguments, [&](const tchar* arg, u32 argLen) { arguments.push_back(TString(arg, argLen)); }))
 			{
 				logger.Error("Failed to parse arguments: %s", m_startInfo.arguments);
 				return UBA_EXIT_CODE(16);
@@ -1937,61 +1938,4 @@ namespace uba
 				CloseFileMapping(pair.second.mappingHandle);
 		m_tempFiles.clear();
 	}
-
-	bool ParseArguments(Vector<TString>& outArguments, const tchar* argumentString)
-	{
-		const tchar* argStart = argumentString;
-		bool isInArg = false;
-		bool isInQuotes = false;
-		bool isEnd = *argumentString == 0;
-		tchar lastChar = 0;
-		for (const tchar* it = argumentString; !isEnd; lastChar = *it, ++it)
-		{
-			isEnd = *it == 0;
-			if (*it == ' ' || *it == '\t' || isEnd)
-			{
-				if (isInQuotes || !isInArg)
-					continue;
-
-				TString result(argStart, it);
-				tchar lastChar2 = 0;
-				for (auto i = result.begin(); i != result.end();)
-				{
-					if (*i == '\"')
-					{
-						if (lastChar2 == '\\')
-							i = result.erase(i - 1) + 1;
-						else
-							i = result.erase(i);
-						lastChar2 = 0;
-						continue;
-					}
-					lastChar2 = *i;
-					++i;
-				}
-
-				outArguments.push_back(std::move(result));
-				isInArg = false;
-				continue;
-			}
-
-			if (!isInArg)
-			{
-				isInArg = true;
-				argStart = it;
-				if (*it == '\"')
-					isInQuotes = true;
-				continue;
-			}
-
-			if (*it == '\"')
-			{
-				if (isInQuotes && lastChar == '\\')
-					continue;
-
-				isInQuotes = !isInQuotes;
-			}
-		}
-		return true;
-	}
 }
diff --git a/Engine/Source/Programs/UnrealBuildAccelerator/Common/Private/UbaSession.cpp b/Engine/Source/Programs/UnrealBuildAccelerator/Common/Private/UbaSession.cpp
--- a/Engine/Source/Programs/UnrealBuildAccelerator/Common/Private/UbaSession.cpp
+++ b/Engine/Source/Programs/UnrealBuildAccelerator/Common/Private/UbaSession.cpp
@@ -647,7 +647,7 @@ namespace uba
 	}
 
 
-	bool Session::RegisterCreateFileForWrite(StringKey fileNameKey, const tchar* fileName, u64 fileNameLen, bool registerRealFile, u64 fileSize, u64 lastWriteTime)
+	bool Session::RegisterCreateFileForWrite(StringKey fileNameKey, const StringView& fileName, bool registerRealFile, u64 fileSize, u64 lastWriteTime)
 	{
 		// Remote is not updating its own directory table
 		if (m_runningRemote)
@@ -658,7 +658,7 @@ namespace uba
 		StringKey dirKey;
 		const tchar* lastSlash;
 		StringBuffer<> dirName;
-		if (!GetDirKey(dirKey, dirName, lastSlash, fileName))
+		if (!GetDirKey(dirKey, dirName, lastSlash, fileName.data))
 			return true;
 			
 		#if 0//_DEBUG  // Bring this back, turned off right now because a few lines above the call to this method we add a mapping
@@ -675,7 +675,7 @@ namespace uba
 		}
 		#endif
 
-		bool shouldWriteToDisk = registerRealFile && ShouldWriteToDisk(fileName, fileNameLen);
+		bool shouldWriteToDisk = registerRealFile && ShouldWriteToDisk(fileName);
 
 		// When not writing to disk we need to populate lookup before adding non-written files.. otherwise they will be lost once lookup is actually populated
 		if (!shouldWriteToDisk)
@@ -718,7 +718,7 @@ namespace uba
 		if (fileNameKey == StringKeyZero)
 		{
 			StringBuffer<> forKey;
-			forKey.Append(fileName, fileNameLen);
+			forKey.Append(fileName);
 			if (CaseInsensitiveFs)
 				forKey.MakeLower();
 			fileNameKey = ToStringKey(forKey);
@@ -732,7 +732,7 @@ namespace uba
 		if (shouldWriteToDisk)
 		{
 			FileInformation info;
-			if (!GetFileInformation(info, m_logger, fileName))
+			if (!GetFileInformation(info, m_logger, fileName.data))
 				return m_logger.Error(TC("Failed to get file information for %s while checking file added for write. This should not happen! (%s)"), fileName, LastErrorToText().data);
 
 			attributes = info.attributes;
@@ -1288,7 +1288,7 @@ namespace uba
 		UBA_ASSERT(!m_runningRemote);
 		StringBuffer<> fixedFilePath;
 		auto key = GetKeyAndFixedName(fixedFilePath, filePath);
-		return RegisterCreateFileForWrite(key, fixedFilePath.data, fixedFilePath.count, true);
+		return RegisterCreateFileForWrite(key, fixedFilePath, true);
 	}
 
 	void Session::RegisterDeleteFile(const tchar* filePath)
@@ -1591,9 +1591,9 @@ namespace uba
 		return process.m_startInfo.rules->IsRarelyRead(fileName);
 	}
 
-	bool Session::IsRarelyReadAfterWritten(ProcessImpl& process, const tchar* fileName, u64 fileNameLen) const
+	bool Session::IsRarelyReadAfterWritten(ProcessImpl& process, const StringView& fileName) const
 	{
-		return process.m_startInfo.rules->IsRarelyReadAfterWritten(fileName, fileNameLen);
+		return process.m_startInfo.rules->IsRarelyReadAfterWritten(fileName);
 	}
 
 	bool Session::IsKnownSystemFile(const tchar* applicationName)
@@ -1605,11 +1605,11 @@ namespace uba
 #endif
 	}
 
-	bool Session::ShouldWriteToDisk(const tchar* fileName, u64 fileNameLen)
+	bool Session::ShouldWriteToDisk(const StringView& fileName)
 	{
 		if (m_shouldWriteToDisk)
 			return true;
-		return EndsWith(fileName, fileNameLen, TC(".h"));
+		return fileName.EndsWith(TC(".h"));
 	}
 
 	bool Session::PrepareProcess(const ProcessStartInfo& startInfo, bool isChild, StringBufferBase& outRealApplication, const tchar*& outRealWorkingDir)
@@ -1617,25 +1617,25 @@ namespace uba
 		return true;
 	}
 
-	u32 Session::GetMemoryMapAlignment(const tchar* fileName, u64 fileNameLen) const
+	u32 Session::GetMemoryMapAlignment(const StringView& fileName) const
 	{
 		// It is not necessarily better to make mem maps of everything.. only things that are read more than once in the build.
 		// Reason is because there is additional overhead to use memory mappings.
 		// Upside is that all things that are memory mapped can be stored compressed in cas storage so it saves space.
 
-		if (EndsWith(fileName, fileNameLen, TC(".pch")))
+		if (fileName.EndsWith(TC(".pch")))
 			return 64 * 1024; // pch needs 64k alignment
-		if (EndsWith(fileName, fileNameLen, TC(".h")) || EndsWith(fileName, fileNameLen, TC(".inl")) || EndsWith(fileName, fileNameLen, TC(".gch")))
+		if (fileName.EndsWith(TC(".h")) || fileName.EndsWith(TC(".inl")) || fileName.EndsWith(TC(".gch")))
 			return 4 * 1024; // clang seems to need 4k alignment? Is it a coincidence it works or what is happening inside the code? (msvc works with alignment 1byte here)
-		if (EndsWith(fileName, fileNameLen, TC(".h.obj")))
+		if (fileName.EndsWith(TC(".h.obj")))
 			return 4 * 1024;
-		//if (EndsWith(fileName, fileNameLen, TC(".lib")))
+		//if (fileName.EndsWith(TC(".lib")))
 		//	return 4 * 1024;
-		//if (EndsWith(fileName, fileNameLen, TC(".rc2.res")))
+		//if (fileName.EndsWith(TC(".rc2.res")))
 		//	return 64;
-		//if (EndsWith(fileName, fileNameLen, TC(".rsp")))
+		//if (fileName.EndsWith(TC(".rsp")))
 		//	return 4 * 1024; // rsp is read over and over again
-		//if (EndsWith(fileName, fileNameLen, TC(".h.obj")) || EndsWith(fileName, fileNameLen, TC(".lib")))
+		//if (fileName.EndsWith(TC(".h.obj")) || fileName.EndsWith(TC(".lib")))
 		//	return 4 * 1024; // rsp is read over and over again
 		return 0;
 	}
@@ -1776,7 +1776,7 @@ namespace uba
 			
 			if (m_allowMemoryMaps)
 			{
-				if (u64 alignment = GetMemoryMapAlignment(fileName.data, fileName.count))
+				if (u64 alignment = GetMemoryMapAlignment(fileName))
 				{
 					MemoryMap map;
 					if (CreateMemoryMapFromFile(map, fileNameKey, fileName.data, false, alignment))
@@ -1951,7 +1951,7 @@ namespace uba
 				lastWriteTime = writtenFile.lastWriteTime;
 			}
 
-			if (msg.process.m_extractExports && msg.process.m_startInfo.rules->ShouldExtractSymbols(file.name.c_str(), file.name.size()))
+			if (msg.process.m_extractExports && msg.process.m_startInfo.rules->ShouldExtractSymbols(file.name))
 				if (!ExtractSymbolsFromObjectFile(msg, name, fileSize))
 					return false;
 		}
@@ -1959,12 +1959,12 @@ namespace uba
 		if (!msg.newName.IsEmpty())
 		{
 			RegisterDeleteFile(file.nameKey, file.name.c_str());
-			RegisterCreateFileForWrite(msg.newNameKey, msg.newName.data, msg.newName.count, registerRealFile, fileSize, lastWriteTime);
+			RegisterCreateFileForWrite(msg.newNameKey, msg.newName, registerRealFile, fileSize, lastWriteTime);
 		}
 		else if (msg.deleteOnClose)
 			RegisterDeleteFile(file.nameKey, file.name.c_str());
 		else
-			RegisterCreateFileForWrite(file.nameKey, file.name.c_str(), file.name.size(), registerRealFile, fileSize, lastWriteTime);
+			RegisterCreateFileForWrite(file.nameKey, file.name, registerRealFile, fileSize, lastWriteTime);
 
 		out.directoryTableSize = GetDirectoryTableSize();
 		return true;
@@ -2012,7 +2012,7 @@ namespace uba
 		out.errorCode = ERROR_SUCCESS;
 		if (!out.result)
 			out.errorCode = GetLastError();
-		RegisterCreateFileForWrite(msg.toKey, msg.toName.data, msg.toName.count, true);
+		RegisterCreateFileForWrite(msg.toKey, msg.toName, true);
 		out.directoryTableSize = RegisterDeleteFile(msg.fromKey, msg.fromName.data);
 		return true;
 	}
@@ -2025,7 +2025,7 @@ namespace uba
 		out.errorCode = 0;
 		if (chmod(msg.fileName.data, (mode_t)msg.fileMode) == 0)
 		{
-			RegisterCreateFileForWrite(msg.fileNameKey, msg.fileName.data, msg.fileName.count, true);
+			RegisterCreateFileForWrite(msg.fileNameKey, msg.fileName, true);
 			return true;
 		}
 		out.errorCode = errno;
@@ -2037,7 +2037,7 @@ namespace uba
 	{
 		out.result = uba::CreateDirectoryW(msg.name.data);
 		if (out.result)
-			RegisterCreateFileForWrite(msg.nameKey, msg.name.data, msg.name.count, true);
+			RegisterCreateFileForWrite(msg.nameKey, msg.name, true);
 		out.errorCode = GetLastError();
 		return true;
 	}
@@ -2062,9 +2062,9 @@ namespace uba
 
 		auto writeFile = [&](WrittenFile& file)
 		{
-			bool shouldEvictFromMemory = IsRarelyReadAfterWritten(process, file.name.c_str(), file.name.size()) || file.mappingWritten > m_keepOutputFileMemoryMapsThreshold;
+			bool shouldEvictFromMemory = IsRarelyReadAfterWritten(process, file.name) || file.mappingWritten > m_keepOutputFileMemoryMapsThreshold;
 
-			if (ShouldWriteToDisk(file.name.c_str(), file.name.size()))
+			if (ShouldWriteToDisk(file.name))
 			{
 				// This is to kill I/O when writing lots of pdb/dlls in parallel
 				#if PLATFORM_WINDOWS
@@ -2088,7 +2088,7 @@ namespace uba
 
 				auto memClose = MakeGuard([&](){ UnmapViewOfFile(mem, fileSize, file.name.c_str()); });
 
-				if (m_storeObjFilesCompressed && process.m_startInfo.rules->StoreFileCompressed(file.name.c_str(), file.name.size()))
+				if (m_storeObjFilesCompressed && process.m_startInfo.rules->StoreFileCompressed(file.name))
 				{
 					Storage::WriteResult res;
 					CompressedObjFileHeader header { CalculateCasKey(mem, fileSize, true, m_workManager) };
@@ -2664,7 +2664,7 @@ namespace uba
 		StringKey symFileKey = CaseInsensitiveFs ? ToStringKeyLower(exportsFile) : ToStringKey(exportsFile);
 		u64 lastWriteTime = GetSystemTimeAsFileTime();
 
-		if (!RegisterCreateFileForWrite(symFileKey, exportsFile.data, exportsFile.count, false, memoryBlock.writtenSize, lastWriteTime))
+		if (!RegisterCreateFileForWrite(symFileKey, exportsFile, false, memoryBlock.writtenSize, lastWriteTime))
 			return false;
 
 		auto insres = msg.process.m_writtenFiles.try_emplace(exportsFile.data);
diff --git a/Engine/Source/Programs/UnrealBuildAccelerator/Common/Private/UbaSessionClient.cpp b/Engine/Source/Programs/UnrealBuildAccelerator/Common/Private/UbaSessionClient.cpp
--- a/Engine/Source/Programs/UnrealBuildAccelerator/Common/Private/UbaSessionClient.cpp
+++ b/Engine/Source/Programs/UnrealBuildAccelerator/Common/Private/UbaSessionClient.cpp
@@ -434,7 +434,7 @@ namespace uba
 			u32 memoryMapAlignment = 0;
 			if (m_allowMemoryMaps)
 			{
-				memoryMapAlignment = GetMemoryMapAlignment(fileName.data, fileName.count);
+				memoryMapAlignment = GetMemoryMapAlignment(fileName);
 				if (!memoryMapAlignment && !m_useStorage)
 					memoryMapAlignment = 64 * 1024;
 			}
@@ -563,7 +563,7 @@ namespace uba
 			if (!pair.second.mappingHandle.IsValid())
 				m_logger.Warning(TC("%s is not using file mapping"), pair.first.c_str());
 #endif
-			bool keepMappingInMemory = IsWindows && !IsRarelyReadAfterWritten(process, pair.first.c_str(), pair.first.size());
+			bool keepMappingInMemory = IsWindows && !IsRarelyReadAfterWritten(process, pair.first);
 			if (!SendFile(pair.second, pair.first.c_str(), process.GetId(), keepMappingInMemory))
 				return false;
 		}
diff --git a/Engine/Source/Programs/UnrealBuildAccelerator/Common/Private/UbaSessionServer.cpp b/Engine/Source/Programs/UnrealBuildAccelerator/Common/Private/UbaSessionServer.cpp
--- a/Engine/Source/Programs/UnrealBuildAccelerator/Common/Private/UbaSessionServer.cpp
+++ b/Engine/Source/Programs/UnrealBuildAccelerator/Common/Private/UbaSessionServer.cpp
@@ -1,6 +1,7 @@
 // Copyright Epic Games, Inc. All Rights Reserved.
 
-#include "UbaSessionServer.h"
+#include "UbaSessionServer.h"	
+#include "UbaApplicationRules.h"
 #include "UbaNetworkServer.h"
 #include "UbaProcess.h"
 #include "UbaProcessStartInfoHolder.h"
@@ -282,7 +283,7 @@ namespace uba
 
 				auto& ki = keys[keysIndex++];
 				ki.key = casKey;
-				ki.mappingAlignment = GetMemoryMapAlignment(fileName.data, fileName.count);
+				ki.mappingAlignment = GetMemoryMapAlignment(fileName);
 
 
 				// Update name to hash table
@@ -303,6 +304,8 @@ namespace uba
 			remoteProcess->m_knownInputsCount = keysIndex;
 		}
 
+		remoteProcess->startInfo.rules = GetRules(remoteProcess->startInfo);
+
 		ProcessHandle h(remoteProcess); // Keep ref count up even if process is removed by callbacks etc.
 
 		ScopedCriticalSection lock(m_remoteProcessAndSessionLock);
@@ -836,9 +839,19 @@ namespace uba
 						return true;
 					}
 
-					if (ShouldWriteToDisk(destination.data, destination.count))
+					if (ShouldWriteToDisk(destination))
 					{
-						bool writeCompressed = m_storeObjFilesCompressed && destination.EndsWith(TC(".obj"));
+						bool writeCompressed = false;
+						if (m_storeObjFilesCompressed)
+						{
+							SCOPED_READ_LOCK(m_processesLock, lock);
+							auto findIt = m_processes.find(processId);
+							if (findIt != m_processes.end())
+							{
+								auto& process = *(RemoteProcess*)findIt->second.m_process;
+								writeCompressed = process.startInfo.rules->StoreFileCompressed(destination);
+							}
+						}
 						success = m_storage.CopyOrLink(casKey, destination.data, attributes, writeCompressed);
 						if (!success)
 							m_logger.Error(TC("Failed to copy cas from %s to %s (%s)"), CasKeyString(casKey).str, destination.data, GetProcessDescription(processId).c_str());
@@ -857,7 +870,7 @@ namespace uba
 				if (success)
 				{
 					m_storage.DropCasFile(casKey, false, destination.data);
-					RegisterCreateFileForWrite(StringKeyZero, destination.data, destination.count, true);
+					RegisterCreateFileForWrite(StringKeyZero, destination, true);
 
 
 					SCOPED_WRITE_LOCK(m_processesLock, lock);
@@ -905,7 +918,7 @@ namespace uba
 				bool result = uba::CopyFileW(fromName.data, toName.data, false);
 				u32 errorCode = GetLastError();
 				if (result)
-					RegisterCreateFileForWrite(toNameKey, toName.data, toName.count, true);
+					RegisterCreateFileForWrite(toNameKey, toName, true);
 				writer.WriteU32(errorCode);
 				return true;
 			}
@@ -1738,7 +1751,7 @@ namespace uba
 			auto findIt = m_receivedFiles.find(msg.fileNameKey);
 			if (findIt != m_receivedFiles.end())
 			{
-				u64 memoryMapAlignment = GetMemoryMapAlignment(msg.fileName.data, msg.fileName.count);
+				u64 memoryMapAlignment = GetMemoryMapAlignment(msg.fileName);
 				if (!memoryMapAlignment)
 					memoryMapAlignment = 4096;
 				MemoryMap map;
diff --git a/Engine/Source/Programs/UnrealBuildAccelerator/Common/Private/UbaStorage.cpp b/Engine/Source/Programs/UnrealBuildAccelerator/Common/Private/UbaStorage.cpp
--- a/Engine/Source/Programs/UnrealBuildAccelerator/Common/Private/UbaStorage.cpp
+++ b/Engine/Source/Programs/UnrealBuildAccelerator/Common/Private/UbaStorage.cpp
@@ -524,7 +524,7 @@ namespace uba
 
 			if (fileIsCompressed)
 			{
-				CompressedObjFileHeader header;
+				CompressedObjFileHeader header(CasKeyZero);
 				if (!ReadFile(m_logger, fileName, readHandle, &header, sizeof(header)))
 					return false;
 				fileSize -= sizeof(header);
@@ -1925,7 +1925,7 @@ namespace uba
 		{
 			if (fileIsCompressed)
 			{
-				CompressedObjFileHeader header;
+				CompressedObjFileHeader header(CasKeyZero);
 				if (!ReadFile(m_logger, fileName, fileHandle, &header, sizeof(header)))
 					return m_logger.Error(TC("Failed to read header of compressed file %s (%s)"), fileName, LastErrorToText().data);
 				fileEntry.casKey = AsCompressed(header.casKey, m_storeCompressed);
@@ -2263,7 +2263,7 @@ namespace uba
 						if (!destinationFile.CreateMemoryWrite(false, fileAttributes, mappedView.size + sizeof(CompressedObjFileHeader), m_tempPath.data))
 							return false;
 						u8* mem = destinationFile.GetData();
-						*(CompressedObjFileHeader*)mem = { casKey };
+						*(CompressedObjFileHeader*)mem = CompressedObjFileHeader(casKey);
 						memcpy(mem + sizeof(CompressedObjFileHeader), mappedView.memory, mappedView.size);
 						return destinationFile.Close();
 					}
diff --git a/Engine/Source/Programs/UnrealBuildAccelerator/Common/Public/UbaCacheClient.h b/Engine/Source/Programs/UnrealBuildAccelerator/Common/Public/UbaCacheClient.h
--- a/Engine/Source/Programs/UnrealBuildAccelerator/Common/Public/UbaCacheClient.h
+++ b/Engine/Source/Programs/UnrealBuildAccelerator/Common/Public/UbaCacheClient.h
@@ -17,6 +17,7 @@ namespace uba
 	struct CacheStats;
 	struct CasKey;
 	struct ProcessStartInfo;
+	struct StringView;
 
 	struct CacheClientCreateInfo
 	{
@@ -56,7 +57,7 @@ namespace uba
 		bool ShouldNormalize(const StringBufferBase& path);
 
 		bool GetLocalPathAndCasKey(Bucket& bucket, const RootPaths& rootPaths, StringBufferBase& outPath, CasKey& outKey, CompactCasKeyTable& casKeyTable, CompactPathTable& pathTable, u32 offset);
-		bool IsFileCompressed(const ProcessStartInfo& info, const tchar* filename, u64 filenameSize);
+		bool IsFileCompressed(const ProcessStartInfo& info, const StringView& filename);
 
 		MutableLogger m_logger;
 		StorageImpl& m_storage;
diff --git a/Engine/Source/Programs/UnrealBuildAccelerator/Common/Public/UbaCacheEntry.h b/Engine/Source/Programs/UnrealBuildAccelerator/Common/Public/UbaCacheEntry.h
--- a/Engine/Source/Programs/UnrealBuildAccelerator/Common/Public/UbaCacheEntry.h
+++ b/Engine/Source/Programs/UnrealBuildAccelerator/Common/Public/UbaCacheEntry.h
@@ -56,6 +56,9 @@ namespace uba
 	{
 		CacheEntriesTraverser(BinaryReader& r) : reader(r)
 		{
+			entryCount = reader.ReadU16();
+			if (!reader.GetLeft())
+				return;
 			u64 sharedSize = r.Read7BitEncoded();
 			sharedInputOffsets = r.GetPositionData();
 			reader.Skip(sharedSize);
@@ -107,7 +110,8 @@ namespace uba
 		}
 
 		BinaryReader& reader;
-		const u8* sharedInputOffsets;
-		u32 lastId = 0;;
+		const u8* sharedInputOffsets = nullptr;
+		u32 entryCount = 0;
+		u32 lastId = 0;
 	};
 }
diff --git a/Engine/Source/Programs/UnrealBuildAccelerator/Common/Public/UbaObjectFile.h b/Engine/Source/Programs/UnrealBuildAccelerator/Common/Public/UbaObjectFile.h
--- a/Engine/Source/Programs/UnrealBuildAccelerator/Common/Public/UbaObjectFile.h
+++ b/Engine/Source/Programs/UnrealBuildAccelerator/Common/Public/UbaObjectFile.h
@@ -53,56 +53,56 @@ namespace uba
 		UnorderedSymbols m_imports;
 		UnorderedExports m_exports;
 		UnorderedSymbols m_potentialDuplicates;
-	};
-
-
-	struct StringView
-	{
-		const char* strBegin;
-		const char* strEnd;
 
-		u32 Length() const
+	public:
+		struct StringView
 		{
-			return u32(strEnd - strBegin);
-		}
+			const char* strBegin;
+			const char* strEnd;
 
-		bool StartsWith(const char* str, u32 strLen) const
-		{
-			if (strLen > Length())
-				return false;
-			return memcmp(strBegin, str, strLen) == 0;
-		}
+			u32 Length() const
+			{
+				return u32(strEnd - strBegin);
+			}
 
-		bool Contains(const char* str, u32 strLen) const
-		{
-			const char* it = strBegin;
-			const char* itEnd = strEnd - strLen + 1;
-			while (it < itEnd)
+			bool StartsWith(const char* str, u32 strLen) const
 			{
-				if (memcmp(it, str, strLen) == 0)
-					return true;
-				++it;
+				if (strLen > Length())
+					return false;
+				return memcmp(strBegin, str, strLen) == 0;
 			}
-			return false;
-		}
 
-		bool Equals(const char* str, u32 strLen) const
-		{
-			if (strLen != Length())
+			bool Contains(const char* str, u32 strLen) const
+			{
+				const char* it = strBegin;
+				const char* itEnd = strEnd - strLen + 1;
+				while (it < itEnd)
+				{
+					if (memcmp(it, str, strLen) == 0)
+						return true;
+					++it;
+				}
 				return false;
-			return memcmp(strBegin, str, strLen) == 0;
-		}
+			}
 
-		std::string ToString() const
-		{
-			return std::string(strBegin, strEnd);
-		}
+			bool Equals(const char* str, u32 strLen) const
+			{
+				if (strLen != Length())
+					return false;
+				return memcmp(strBegin, str, strLen) == 0;
+			}
 
-		std::string& ToString(std::string& out) const
-		{
-			out.assign(strBegin, strEnd);
-			return out;
-		}
+			std::string ToString() const
+			{
+				return std::string(strBegin, strEnd);
+			}
+
+			std::string& ToString(std::string& out) const
+			{
+				out.assign(strBegin, strEnd);
+				return out;
+			}
+		};
 	};
 
 	struct SymbolFile
diff --git a/Engine/Source/Programs/UnrealBuildAccelerator/Common/Public/UbaProcess.h b/Engine/Source/Programs/UnrealBuildAccelerator/Common/Public/UbaProcess.h
--- a/Engine/Source/Programs/UnrealBuildAccelerator/Common/Public/UbaProcess.h
+++ b/Engine/Source/Programs/UnrealBuildAccelerator/Common/Public/UbaProcess.h
@@ -185,6 +185,4 @@ namespace uba
 		ProcessImpl(const ProcessImpl&) = delete;
 		void operator=(const ProcessImpl&) = delete;
 	};
-
-	bool ParseArguments(Vector<TString>& outArguments, const tchar* argumentString);
 }
diff --git a/Engine/Source/Programs/UnrealBuildAccelerator/Common/Public/UbaSession.h b/Engine/Source/Programs/UnrealBuildAccelerator/Common/Public/UbaSession.h
--- a/Engine/Source/Programs/UnrealBuildAccelerator/Common/Public/UbaSession.h
+++ b/Engine/Source/Programs/UnrealBuildAccelerator/Common/Public/UbaSession.h
@@ -92,7 +92,7 @@ namespace uba
 		void StopTraceThread();
 		u32 GetDirectoryTableSize();
 		u32 GetFileMappingSize();
-		u32 GetMemoryMapAlignment(const tchar* fileName, u64 fileNameLen) const;
+		u32 GetMemoryMapAlignment(const StringView& fileName) const;
 
 		SessionStats& Stats();
 
@@ -100,9 +100,9 @@ namespace uba
 		bool GetBinaryModules(Vector<BinaryModule>& out, const tchar* application);
 		void Free(Vector<BinaryModule>& v);
 		bool IsRarelyRead(ProcessImpl& process, const StringBufferBase& fileName) const;
-		bool IsRarelyReadAfterWritten(ProcessImpl& process, const tchar* fileName, u64 fileNameLen) const;
+		bool IsRarelyReadAfterWritten(ProcessImpl& process, const StringView& fileName) const;
 		bool IsKnownSystemFile(const tchar* applicationName);
-		bool ShouldWriteToDisk(const tchar* fileName, u64 fileNameLen);
+		bool ShouldWriteToDisk(const StringView& fileName);
 		u32 WriteDirectoryEntries(const StringKey& dirKey, tchar* dirPath, u32& outTableOffset);
 		u32 AddFileMapping(StringKey fileNameKey, const tchar* fileName, const tchar* newFileName, u64 fileSize = InvalidValue);
 		
@@ -110,7 +110,7 @@ namespace uba
 		bool CreateMemoryMapFromFile(MemoryMap& out, StringKey fileNameKey, const tchar* fileName, bool isCompressed, u64 alignment);
 		bool CreateMemoryMapFromView(MemoryMap& out, StringKey fileNameKey, const tchar* fileName, const CasKey& casKey, u64 alignment);
 
-		bool RegisterCreateFileForWrite(StringKey fileNameKey, const tchar* fileName, u64 fileNameLen, bool registerRealFile, u64 fileSize = 0, u64 lastWriteTime = 0);
+		bool RegisterCreateFileForWrite(StringKey fileNameKey, const StringView& fileName, bool registerRealFile, u64 fileSize = 0, u64 lastWriteTime = 0);
 		u32 RegisterDeleteFile(StringKey fileNameKey, const tchar* fileName);
 
 		virtual bool PrepareProcess(const ProcessStartInfo& startInfo, bool isChild, StringBufferBase& outRealApplication, const tchar*& outRealWorkingDir);
diff --git a/Engine/Source/Programs/UnrealBuildAccelerator/Core/Private/UbaApplicationRules.cpp b/Engine/Source/Programs/UnrealBuildAccelerator/Core/Private/UbaApplicationRules.cpp
--- a/Engine/Source/Programs/UnrealBuildAccelerator/Core/Private/UbaApplicationRules.cpp
+++ b/Engine/Source/Programs/UnrealBuildAccelerator/Core/Private/UbaApplicationRules.cpp
@@ -30,17 +30,17 @@ namespace uba
 				return false;
 			return true;
 		}
-		virtual bool IsThrowAway(const tchar* fileName, u32 fileNameLen) const override
+		virtual bool IsThrowAway(const StringView& fileName) const override
 		{
-			return Contains(fileName, TC("vctip_")) || Super::IsThrowAway(fileName, fileNameLen);
+			return fileName.Contains(TC("vctip_")) || Super::IsThrowAway(fileName);
 		}
-		virtual bool KeepInMemory(const tchar* fileName, u32 fileNameLen, const tchar* systemTemp) const override
+		virtual bool KeepInMemory(const StringView& fileName, const tchar* systemTemp) const override
 		{
-			if (Contains(fileName, TC("\\vctip_")) != 0)
+			if (fileName.Contains(TC("\\vctip_")) != 0)
 				return true;
-			if (Contains(fileName, systemTemp))
+			if (fileName.Contains(systemTemp))
 				return true;
-			return Super::KeepInMemory(fileName, fileNameLen, systemTemp);
+			return Super::KeepInMemory(fileName, systemTemp);
 		}
 
 		virtual bool IsExitCodeSuccess(u32 exitCode) const override
@@ -52,12 +52,12 @@ namespace uba
 	class ApplicationRulesClExe : public ApplicationRulesVC
 	{
 	public:
-		virtual bool IsOutputFile(const tchar* file, u64 fileLen) const override
+		virtual bool IsOutputFile(const StringView& fileName) const override
 		{
-			return EndsWith(file, fileLen, TC(".obj"))
-				|| EndsWith(file, fileLen, TC(".dep.json"))
-				|| EndsWith(file, fileLen, TC(".rc2.res")) // Not really an obj file.. 
-				;// || EndsWith(file, fileLen, TC(".h.pch") // Not tested enough
+			return fileName.EndsWith(TC(".obj"))
+				|| fileName.EndsWith(TC(".dep.json"))
+				|| fileName.EndsWith(TC(".rc2.res")) // Not really an obj file.. 
+				;// || fileName.EndsWith(TC(".h.pch") // Not tested enough
 		}
 
 		virtual bool IsRarelyRead(const StringBufferBase& file) const override
@@ -66,11 +66,11 @@ namespace uba
 				|| file.EndsWith(TC(".obj.rsp"));
 		}
 
-		virtual bool IsRarelyReadAfterWritten(const tchar* fileName, u64 fileNameLen) const override
+		virtual bool IsRarelyReadAfterWritten(const StringView& fileName) const override
 		{
-			return EndsWith(fileName, fileNameLen, TC(".dep.json"))
-				|| EndsWith(fileName, fileNameLen, TC(".exe"))
-				|| EndsWith(fileName, fileNameLen, TC(".dll"));
+			return fileName.EndsWith(TC(".dep.json"))
+				|| fileName.EndsWith(TC(".exe"))
+				|| fileName.EndsWith(TC(".dll"));
 		}
 
 		virtual bool NeedsSharedMemory(const tchar* file) const override
@@ -83,14 +83,14 @@ namespace uba
 			return true;
 		}
 
-		virtual bool StoreFileCompressed(const tchar* fileName, u64 fileNameLen) const
+		virtual bool StoreFileCompressed(const StringView& fileName) const
 		{
-			return EndsWith(fileName, fileNameLen, TC(".obj"));
+			return fileName.EndsWith(TC(".obj"));
 		}
 
-		virtual bool ShouldExtractSymbols(const tchar* fileName, u64 fileNameLen) const
+		virtual bool ShouldExtractSymbols(const StringView& fileName) const
 		{
-			return EndsWith(fileName, fileNameLen, TC(".obj"));
+			return fileName.EndsWith(TC(".obj"));
 		}
 	};
 
@@ -98,19 +98,19 @@ namespace uba
 	{
 		using Super = ApplicationRulesVC;
 	public:
-		virtual bool IsOutputFile(const tchar* file, u64 fileLen) const override
+		virtual bool IsOutputFile(const StringView& fileName) const override
 		{
-			return EndsWith(file, fileLen, TC(".lib"))
-				|| EndsWith(file, fileLen, TC(".exp"))
-				|| EndsWith(file, fileLen, TC(".pdb"))
-				|| EndsWith(file, fileLen, TC(".dll"))
-				|| EndsWith(file, fileLen, TC(".exe"))
-				|| EndsWith(file, fileLen, TC(".rc2.res")); // Not really an obj file.. 
+			return fileName.EndsWith(TC(".lib"))
+				|| fileName.EndsWith(TC(".exp"))
+				|| fileName.EndsWith(TC(".pdb"))
+				|| fileName.EndsWith(TC(".dll"))
+				|| fileName.EndsWith(TC(".exe"))
+				|| fileName.EndsWith(TC(".rc2.res")); // Not really an obj file.. 
 		}
 
-		virtual bool IsThrowAway(const tchar* fileName, u32 fileNameLen) const override
+		virtual bool IsThrowAway(const StringView& fileName) const override
 		{
-			return Contains(fileName, TC(".sup.")); // .sup.lib/exp are throw-away files that we don't want created
+			return fileName.Contains(TC(".sup.")); // .sup.lib/exp are throw-away files that we don't want created
 		}
 
 		virtual bool CanExist(const tchar* file) const override
@@ -139,17 +139,22 @@ namespace uba
 			return Super::AllowStorageProxy(file);
 		}
 
-		virtual bool IsRarelyReadAfterWritten(const tchar* fileName, u64 fileNameLen) const override
+		virtual bool IsRarelyReadAfterWritten(const StringView& fileName) const override
 		{
-			return EndsWith(fileName, fileNameLen, TC(".pdb"))
-				|| EndsWith(fileName, fileNameLen, TC(".exe"))
-				|| EndsWith(fileName, fileNameLen, TC(".dll"));
+			return fileName.EndsWith(TC(".pdb"))
+				|| fileName.EndsWith(TC(".exe"))
+				|| fileName.EndsWith(TC(".dll"));
 		}
 
 		virtual bool IsCacheable() const override
 		{
 			return true;
 		}
+
+		virtual bool ShouldDecompressFiles(const StringView& fileName) const override
+		{
+			return fileName.IsEmpty() || fileName.EndsWith(TC(".obj"));
+		}
 	};
 
 	// ==== Clang tool chain ====
@@ -176,15 +181,15 @@ namespace uba
 			return true;
 		}
 
-		virtual bool IsOutputFile(const tchar* file, u64 fileLen) const override
+		virtual bool IsOutputFile(const StringView& fileName) const override
 		{
-			return EndsWith(file, fileLen, TC(".c.d"))
-				|| EndsWith(file, fileLen, TC(".h.d"))
-				|| EndsWith(file, fileLen, TC(".cc.d"))
-				|| EndsWith(file, fileLen, TC(".cpp.d"))
-				|| EndsWith(file, fileLen, TC(".o.tmp")) // Clang writes to tmp file and then move
-				|| EndsWith(file, fileLen, TC(".obj.tmp")) // Clang (verse) writes to tmp file and then move
-				;// || EndsWith(file, fileLen, TC(".gch.tmp")); // Need to fix "has been modified since the precompiled header"
+			return fileName.EndsWith(TC(".c.d"))
+				|| fileName.EndsWith(TC(".h.d"))
+				|| fileName.EndsWith(TC(".cc.d"))
+				|| fileName.EndsWith(TC(".cpp.d"))
+				|| fileName.EndsWith(TC(".o.tmp")) // Clang writes to tmp file and then move
+				|| fileName.EndsWith(TC(".obj.tmp")) // Clang (verse) writes to tmp file and then move
+				;// || fileName.EndsWith(TC(".gch.tmp")); // Need to fix "has been modified since the precompiled header"
 		}
 
 		virtual bool IsRarelyRead(const StringBufferBase& file) const override
@@ -193,9 +198,9 @@ namespace uba
 				|| file.EndsWith(TC(".o.rsp"));
 		}
 
-		virtual bool IsRarelyReadAfterWritten(const tchar* fileName, u64 fileNameLen) const override
+		virtual bool IsRarelyReadAfterWritten(const StringView& fileName) const override
 		{
-			return EndsWith(fileName, fileNameLen, TC(".d"));
+			return fileName.EndsWith(TC(".d"));
 		}
 
 		virtual bool AllowMiMalloc() const override
@@ -208,14 +213,14 @@ namespace uba
 			return true;
 		}
 
-		virtual bool StoreFileCompressed(const tchar* fileName, u64 fileNameLen) const
+		virtual bool StoreFileCompressed(const StringView& fileName) const
 		{
-			return EndsWith(fileName, fileNameLen, TC(".obj"));
+			return fileName.EndsWith(TC(".obj")) || fileName.EndsWith(TC(".o"));
 		}
 
-		virtual bool ShouldExtractSymbols(const tchar* fileName, u64 fileNameLen) const
+		virtual bool ShouldExtractSymbols(const StringView& fileName) const
 		{
-			return EndsWith(fileName, fileNameLen, TC(".obj")) || EndsWith(fileName, fileNameLen, TC(".o"));
+			return fileName.EndsWith(TC(".obj")) || fileName.EndsWith(TC(".o"));
 		}
 	};
 
@@ -223,9 +228,9 @@ namespace uba
 	{
 		using Super = ApplicationRulesClang;
 
-		virtual bool IsOutputFile(const tchar* file, u64 fileLen) const override
+		virtual bool IsOutputFile(const StringView& fileName) const override
 		{
-			return Contains(file, TC(".tmp")); // both .so.tmp and .tmp123456
+			return fileName.Contains(TC(".tmp")); // both .so.tmp and .tmp123456
 		}
 
 		virtual bool IsRarelyRead(const StringBufferBase& file) const override
@@ -237,20 +242,25 @@ namespace uba
 		{
 			return 14ull * 1024 * 1024 * 1024; // This is ridiculous (needed for asan targets)
 		}
+
+		virtual bool ShouldDecompressFiles(const StringView& fileName) const override
+		{
+			return fileName.IsEmpty() || fileName.EndsWith(TC(".o"));
+		}
 	};
 
 	class ApplicationRulesLlvmObjCopyExe : public ApplicationRulesClang
 	{
 		using Super = ApplicationRulesClang;
 
-		virtual bool IsOutputFile(const tchar* file, u64 fileLen) const override
+		virtual bool IsOutputFile(const StringView& fileName) const override
 		{
-			return Contains(file, TC(".temp-stream-"));
+			return fileName.Contains(TC(".temp-stream-"));
 		}
 
 		virtual u64 FileTypeMaxSize(const StringBufferBase& file, bool isSystemOrTempFile) const override
 		{
-			if (IsOutputFile(file.data, file.count))
+			if (IsOutputFile(file))
 				return 14ull * 1024 * 1024 * 1024; // This is ridiculous (needed for asan targets)
 			return Super::FileTypeMaxSize(file, isSystemOrTempFile);
 		}
@@ -258,9 +268,9 @@ namespace uba
 
 	class ApplicationRulesDumpSymsExe : public ApplicationRulesClang
 	{
-		virtual bool IsOutputFile(const tchar* file, u64 fileLen) const override
+		virtual bool IsOutputFile(const StringView& fileName) const override
 		{
-			return false; // EndsWith(file, fileLen, TC(".psym")); With psym as output file the BreakpadSymbolEncoder fails to output a .sym file
+			return false; // fileName.EndsWith(TC(".psym")); With psym as output file the BreakpadSymbolEncoder fails to output a .sym file
 		}
 	};
 
@@ -268,14 +278,14 @@ namespace uba
 	{
 		using Super = ApplicationRulesClangPlusPlusExe;
 
-		virtual bool IsThrowAway(const tchar* fileName, u32 fileNameLen) const override
+		virtual bool IsThrowAway(const StringView& fileName) const override
 		{
-			return EndsWith(fileName, fileNameLen, TC("-telemetry.json")) || Super::IsThrowAway(fileName, fileNameLen);
+			return fileName.EndsWith(TC("-telemetry.json")) || Super::IsThrowAway(fileName);
 		}
 
 		//virtual bool NeedsSharedMemory(const tchar* file)
 		//{
-		//	return Contains(file, TC("lto-llvm")); // Used by a clang based platform's link time optimization pass. Shared from lto process back to linker process
+		//	return fileName.Contains(TC("lto-llvm")); // Used by a clang based platform's link time optimization pass. Shared from lto process back to linker process
 		//}
 	};
 
@@ -283,38 +293,43 @@ namespace uba
 	{
 		using Super = ApplicationRules;
 
-		//virtual bool IsOutputFile(const tchar* file, u64 fileLen) const override
+		//virtual bool IsOutputFile(const StringView& fileName) const override
 		//{
-		//	return EndsWith(file, fileLen, TC(".self")) || Equals(file, TC("Symbols.map"));
+		//	return fileName.EndsWith(TC(".self")) || Equals(file, TC("Symbols.map"));
 		//}
-		virtual bool KeepInMemory(const tchar* fileName, u32 fileNameLen, const tchar* systemTemp) const override
+		virtual bool KeepInMemory(const StringView& fileName, const tchar* systemTemp) const override
 		{
-			return Super::KeepInMemory(fileName, fileNameLen, systemTemp)
-				|| Contains(fileName, TC("thinlto-"));// Used by a clang based platform's link time optimization pass. Shared from lto process back to linker process
+			return Super::KeepInMemory(fileName, systemTemp)
+				|| fileName.Contains(TC("thinlto-"));// Used by a clang based platform's link time optimization pass. Shared from lto process back to linker process
 		}
 		virtual bool NeedsSharedMemory(const tchar* file) const override
 		{
 			return Contains(file, TC("thinlto-")); // Used by a clang based platform's link time optimization pass. Shared from lto process back to linker process
 		}
+
+		virtual bool ShouldDecompressFiles(const StringView& fileName) const override
+		{
+			return fileName.IsEmpty() || fileName.EndsWith(TC(".o"));
+		}
 	};
 
 	class ApplicationRulesProsperoClangPlusPlusExe : public ApplicationRulesClangPlusPlusExe
 	{
 		using Super = ApplicationRulesClangPlusPlusExe;
 
-		virtual bool IsOutputFile(const tchar* file, u64 fileLen) const override
+		virtual bool IsOutputFile(const StringView& fileName) const override
 		{
-			return Contains(file, TC(".self")) || Super::IsOutputFile(file, fileLen);
+			return fileName.Contains(TC(".self")) || Super::IsOutputFile(fileName);
 		}
 
-		virtual bool IsThrowAway(const tchar* fileName, u32 fileNameLen) const override
+		virtual bool IsThrowAway(const StringView& fileName) const override
 		{
-			return EndsWith(fileName, fileNameLen, TC("-telemetry.json")) || Super::IsThrowAway(fileName, fileNameLen);
+			return fileName.EndsWith(TC("-telemetry.json")) || Super::IsThrowAway(fileName);
 		}
 
 		//virtual bool NeedsSharedMemory(const tchar* file)
 		//{
-		//	return Contains(file, TC("lto-llvm")); // Used by a clang based platform's link time optimization pass. Shared from lto process back to linker process
+		//	return fileName.Contains(TC("lto-llvm")); // Used by a clang based platform's link time optimization pass. Shared from lto process back to linker process
 		//}
 	};
 
@@ -322,20 +337,25 @@ namespace uba
 	{
 		using Super = ApplicationRules;
 
-		virtual bool IsOutputFile(const tchar* file, u64 fileLen) const override
+		virtual bool IsOutputFile(const StringView& fileName) const override
 		{
-			return Contains(file, TC(".self"));
+			return fileName.Contains(TC(".self"));
 		}
 
-		virtual bool IsThrowAway(const tchar* fileName, u32 fileNameLen) const override
+		virtual bool IsThrowAway(const StringView& fileName) const override
 		{
-			return EndsWith(fileName, fileNameLen, TC("-telemetry.json")) || Super::IsThrowAway(fileName, fileNameLen);
+			return fileName.EndsWith(TC("-telemetry.json")) || Super::IsThrowAway(fileName);
 		}
 
 		//virtual bool NeedsSharedMemory(const tchar* file)
 		//{
-		//	return Contains(file, TC("lto-llvm")); // Used by a clang based platform's link time optimization pass. Shared from lto process back to linker process
+		//	return fileName.Contains(TC("lto-llvm")); // Used by a clang based platform's link time optimization pass. Shared from lto process back to linker process
 		//}
+
+		virtual bool ShouldDecompressFiles(const StringView& fileName) const override
+		{
+			return fileName.IsEmpty() || fileName.EndsWith(TC(".o"));
+		}
 	};
 
 	// ====
@@ -347,12 +367,12 @@ namespace uba
 			return true;
 		}
 
-		virtual bool IsOutputFile(const tchar* file, u64 fileLen) const override
+		virtual bool IsOutputFile(const StringView& fileName) const override
 		{
-			return Contains(file, TC(".generated.dummy"))
-				|| EndsWith(file, fileLen, TC(".ispc.bc"))
-				|| EndsWith(file, fileLen, TC(".ispc.txt"))
-				|| EndsWith(file, fileLen, TC(".obj"));
+			return fileName.Contains(TC(".generated.dummy"))
+				|| fileName.EndsWith(TC(".ispc.bc"))
+				|| fileName.EndsWith(TC(".ispc.txt"))
+				|| fileName.EndsWith(TC(".obj"));
 		}
 
 		virtual bool IsCacheable() const override
@@ -360,37 +380,37 @@ namespace uba
 			return true;
 		}
 
-		virtual bool StoreFileCompressed(const tchar* fileName, u64 fileNameLen) const
+		virtual bool StoreFileCompressed(const StringView& fileName) const
 		{
-			return EndsWith(fileName, fileNameLen, TC(".obj"));
+			return fileName.EndsWith(TC(".obj"));
 		}
 
-		virtual bool ShouldExtractSymbols(const tchar* fileName, u64 fileNameLen) const
+		virtual bool ShouldExtractSymbols(const StringView& fileName) const
 		{
-			return EndsWith(fileName, fileNameLen, TC(".obj")) || EndsWith(fileName, fileNameLen, TC(".o"));
+			return fileName.EndsWith(TC(".obj")) || fileName.EndsWith(TC(".o"));
 		}
 	};
 
 	class ApplicationRulesUBTDll : public ApplicationRules
 	{
-		virtual bool IsOutputFile(const tchar* file, u64 fileLen) const override
+		virtual bool IsOutputFile(const StringView& fileName) const override
 		{
 			return false;
 			// TODO: These does not work when UnrealBuildTool creates these files multiple times in a row (building multiple targets)
 			// ... on output they get stored as file mappings.. and next execution of ubt opens them for write (writing file mappings not implemented right now)
-			//return EndsWith(file, fileLen, TC(".modules"))
-			//	|| EndsWith(file, fileLen, TC(".target"))
-			//	|| EndsWith(file, fileLen, TC(".version"));
+			//return fileName.EndsWith(TC(".modules"))
+			//	|| fileName.EndsWith(TC(".target"))
+			//	|| fileName.EndsWith(TC(".version"));
 		}
 	};
 
 	class ApplicationRulesPVSStudio : public ApplicationRules
 	{
-		virtual bool IsOutputFile(const tchar* file, u64 fileLen) const override
+		virtual bool IsOutputFile(const StringView& fileName) const override
 		{
-			return EndsWith(file, fileLen, TC(".PVS-Studio.log"))
-				|| EndsWith(file, fileLen, TC(".pvslog"))
-				|| EndsWith(file, fileLen, TC(".stacktrace.txt"));
+			return fileName.EndsWith(TC(".PVS-Studio.log"))
+				|| fileName.EndsWith(TC(".pvslog"))
+				|| fileName.EndsWith(TC(".stacktrace.txt"));
 		}
 		
 		virtual bool IsRarelyRead(const StringBufferBase& file) const override
@@ -424,14 +444,14 @@ namespace uba
 
 	class ApplicationRulesUbaObjTool : public ApplicationRules
 	{
-		virtual bool IsOutputFile(const tchar* file, u64 fileLen) const override
+		virtual bool IsOutputFile(const StringView& fileName) const override
 		{
-			return EndsWith(file, fileLen, TC(".obj"));
+			return fileName.EndsWith(TC(".obj"));
 		}
 
-		virtual bool StoreFileCompressed(const tchar* fileName, u64 fileNameLen) const
+		virtual bool StoreFileCompressed(const StringView& fileName) const
 		{
-			return EndsWith(fileName, fileNameLen, TC(".obj"));
+			return fileName.EndsWith(TC(".obj"));
 		}
 	};
 
diff --git a/Engine/Source/Programs/UnrealBuildAccelerator/Core/Private/UbaStringBuffer.cpp b/Engine/Source/Programs/UnrealBuildAccelerator/Core/Private/UbaStringBuffer.cpp
--- a/Engine/Source/Programs/UnrealBuildAccelerator/Core/Private/UbaStringBuffer.cpp
+++ b/Engine/Source/Programs/UnrealBuildAccelerator/Core/Private/UbaStringBuffer.cpp
@@ -146,6 +146,11 @@ namespace uba
 		return *this;
 	}
 
+	StringBufferBase& StringBufferBase::Append(const StringView& view)
+	{
+		return Append(view.data, view.count);
+	}
+
 	StringBufferBase& StringBufferBase::Appendf(const tchar* format, ...)
 	{
 		if (*format)
@@ -375,4 +380,9 @@ namespace uba
 		return (u32)destLen;
 	}
 	#endif
+
+	StringView ToView(const tchar* s)
+	{
+		return StringView(s, TStrlen(s));
+	}
 }
diff --git a/Engine/Source/Programs/UnrealBuildAccelerator/Core/Public/UbaApplicationRules.h b/Engine/Source/Programs/UnrealBuildAccelerator/Core/Public/UbaApplicationRules.h
--- a/Engine/Source/Programs/UnrealBuildAccelerator/Core/Public/UbaApplicationRules.h
+++ b/Engine/Source/Programs/UnrealBuildAccelerator/Core/Public/UbaApplicationRules.h
@@ -34,16 +34,16 @@ namespace uba
 		}
 
 		// Throw-away means that the file is temporary and will not be used after process exists. (By default these are kept in memory and never touch disk)
-		virtual bool IsThrowAway(const tchar* fileName, u32 fileNameLen) const
+		virtual bool IsThrowAway(const StringView& fileName) const
 		{
 			return false;
 		}
 
 		// Keep file in memory
 		// If this returns true it means that file will be kept in memory and never touch disk.
-		virtual bool KeepInMemory(const tchar* fileName, u32 fileNameLen, const tchar* systemTemp) const
+		virtual bool KeepInMemory(const StringView& fileName, const tchar* systemTemp) const
 		{
-			return IsThrowAway(fileName, fileNameLen);
+			return IsThrowAway(fileName);
 		}
 
 		// For files that are kept in memory but shared between process (temporary files where one process write and another read)
@@ -60,7 +60,7 @@ namespace uba
 		}
 
 		// Outputfile means that it is kept in memory and then sent back to session process which can decide to write it to disk or send it over network
-		virtual bool IsOutputFile(const tchar* file, u64 fileLen) const
+		virtual bool IsOutputFile(const StringView& fileName) const
 		{
 			return false;
 		}
@@ -80,7 +80,7 @@ namespace uba
 
 		// Return true if the file is never/rarely read after it was written.
 		// This is an optimization where the written file is not kept in file mappings after written
-		virtual bool IsRarelyReadAfterWritten(const tchar* fileName, u64 fileNameLen) const
+		virtual bool IsRarelyReadAfterWritten(const StringView& fileName) const
 		{
 			return false;
 		}
@@ -129,12 +129,17 @@ namespace uba
 			return false;
 		}
 
-		virtual bool StoreFileCompressed(const tchar* fileName, u64 fileNameLen) const
+		virtual bool StoreFileCompressed(const StringView& fileName) const
 		{
 			return false;
 		}
 
-		virtual bool ShouldExtractSymbols(const tchar* fileName, u64 fileNameLen) const
+		virtual bool ShouldDecompressFiles(const StringView& fileName) const
+		{
+			return false;
+		}
+
+		virtual bool ShouldExtractSymbols(const StringView& fileName) const
 		{
 			return false;
 		}
diff --git a/Engine/Source/Programs/UnrealBuildAccelerator/Core/Public/UbaCompressedObjFileHeader.h b/Engine/Source/Programs/UnrealBuildAccelerator/Core/Public/UbaCompressedObjFileHeader.h
--- a/Engine/Source/Programs/UnrealBuildAccelerator/Core/Public/UbaCompressedObjFileHeader.h
+++ b/Engine/Source/Programs/UnrealBuildAccelerator/Core/Public/UbaCompressedObjFileHeader.h
@@ -8,6 +8,16 @@ namespace uba
 {
 	struct CompressedObjFileHeader
 	{
+		inline static constexpr u8 Magic[] = { 'U', 'B', 'A', 0x01 };
+
+		CompressedObjFileHeader(const CasKey& key) : casKey(key)
+		{
+			memcpy(&magic, Magic, sizeof(magic));
+		}
+
+		bool IsValid() const { return memcmp(&magic, Magic, sizeof(Magic)) == 0; }
+
+		u8 magic[sizeof(Magic)];
 		CasKey casKey;
 	};
 }
diff --git a/Engine/Source/Programs/UnrealBuildAccelerator/Core/Public/UbaProcessUtils.h b/Engine/Source/Programs/UnrealBuildAccelerator/Core/Public/UbaProcessUtils.h
new file mode 100644
--- /dev/null
+++ b/Engine/Source/Programs/UnrealBuildAccelerator/Core/Public/UbaProcessUtils.h
@@ -0,0 +1,98 @@
+// Copyright Epic Games, Inc. All Rights Reserved.
+
+#pragma once
+
+#include "UbaStringBuffer.h"
+
+namespace uba
+{
+	// func is void(const CharType* arg, u32 strArg)
+	template<typename CharType, typename Func>
+	bool ParseArguments(const CharType* arguments, u64 argumentsLen, const Func& argumentFunc)
+	{
+		const CharType* argStart = arguments;
+		bool isInArg = false;
+		bool isInQuotes = false;
+		bool isEnd = *arguments == 0;
+		CharType lastChar = 0;
+		bool isBackslashOwned = false;
+		for (const CharType* it = arguments; !isEnd; lastChar = *it, ++it)
+		{
+			isEnd = *it == 0;
+			if (*it == ' ' || *it == '\t' || *it == '\n' || isEnd)
+			{
+				if (isInQuotes || !isInArg)
+					continue;
+
+				const CharType* argIt = argStart;
+				const CharType* argEnd = it;
+
+				if (*argEnd == '\n' && argStart != argEnd && *(argEnd-1) == '\r')
+					--argEnd;
+
+				CharType arg[16*1024];
+				CharType* argWrite = arg;
+				CharType lastChar2 = 0;
+				isBackslashOwned = false;
+				while (argIt != argEnd)
+				{
+					if (*argIt == '\"')
+					{
+						if (lastChar2 == '\\' && !isBackslashOwned)
+							argWrite[-1] = '\"';
+						lastChar2 = 0;
+						++argIt;
+						continue;
+					}
+
+					if (*argIt == '\\' && lastChar2 == '\\')
+						isBackslashOwned = !isBackslashOwned;
+					else
+						isBackslashOwned = false;
+
+					*argWrite++ = *argIt;
+					lastChar2 = *argIt;
+					++argIt;
+				}
+
+				if (arg != argWrite)
+				{
+					*argWrite = 0;
+					argumentFunc(arg, u32(argWrite - arg));
+				}
+				isInArg = false;
+				isBackslashOwned = false;
+				continue;
+			}
+
+			if (!isInArg)
+			{
+				isInArg = true;
+				argStart = it;
+				if (*it == '\"')
+					isInQuotes = true;
+				continue;
+			}
+
+			if (*it == '\"')
+			{
+				if (isInQuotes && lastChar == '\\' && !isBackslashOwned)
+					continue;
+
+				isInQuotes = !isInQuotes;
+			}
+
+			if (*it == '\\' && lastChar == '\\')
+				isBackslashOwned = !isBackslashOwned;
+			else
+				isBackslashOwned = false;
+		}
+		return true;
+	}
+
+	template<typename Func>
+	bool ParseArguments(const tchar* arguments, const Func& argumentFunc)
+	{
+		return ParseArguments(arguments, TStrlen(arguments), argumentFunc);
+	}
+}
diff --git a/Engine/Source/Programs/UnrealBuildAccelerator/Core/Public/UbaStringBuffer.h b/Engine/Source/Programs/UnrealBuildAccelerator/Core/Public/UbaStringBuffer.h
--- a/Engine/Source/Programs/UnrealBuildAccelerator/Core/Public/UbaStringBuffer.h
+++ b/Engine/Source/Programs/UnrealBuildAccelerator/Core/Public/UbaStringBuffer.h
@@ -6,6 +6,8 @@
 
 namespace uba
 {
+	struct StringView;
+
 	bool StartsWith(const tchar* data, const tchar* str, bool ignoreCase = true);
 	bool EndsWith(const tchar* str, u64 strLen, const tchar* value, bool ignoreCase = true);
 	bool Contains(const tchar* str, const tchar* sub, bool ignoreCase = true, const tchar** pos = nullptr);
@@ -25,6 +27,7 @@ namespace uba
 		StringBufferBase& Append(const tchar* str);
 		StringBufferBase& Append(const tchar c) { return Append(&c, 1); }
 		StringBufferBase& Append(const tchar* str, u64 charCount);
+		StringBufferBase& Append(const StringView& view);
 		StringBufferBase& Appendf(const tchar* format, ...);
 		StringBufferBase& AppendDir(const StringBufferBase& str);
 		StringBufferBase& AppendDir(const tchar* dir);
@@ -88,6 +91,22 @@ namespace uba
 		tchar buf[Capacity];
 	};
 
+	struct StringView
+	{
+		StringView() : data(nullptr), count(0) {}
+		StringView(const tchar* d, u32 c) : data(d), count(c) {}
+		StringView(const StringBufferBase& sb) : data(sb.data), count(sb.count) {}
+		StringView(const TString& str) : data(str.data()), count(u32(str.size())) {}
+		bool IsEmpty() const { return count == 0; }
+		bool EndsWith(const tchar* value, bool ignoreCase = true) const { return uba::EndsWith(data, count, value, ignoreCase); }
+		bool Contains(const tchar* str, bool ignoreCase = true) const { return uba::Contains(data, str, ignoreCase); }
+
+		const tchar* data;
+		u32 count;
+	};
+
+	StringView ToView(const tchar* s);
+
 	struct LastErrorToText : StringBuffer<256>
 	{
 		LastErrorToText();
diff --git a/Engine/Source/Programs/UnrealBuildAccelerator/Detours/Private/UbaDetoursShared.cpp b/Engine/Source/Programs/UnrealBuildAccelerator/Detours/Private/UbaDetoursShared.cpp
--- a/Engine/Source/Programs/UnrealBuildAccelerator/Detours/Private/UbaDetoursShared.cpp
+++ b/Engine/Source/Programs/UnrealBuildAccelerator/Detours/Private/UbaDetoursShared.cpp
@@ -239,7 +239,7 @@ namespace uba
 
 		memset(&outAttr.data, 0, sizeof(outAttr.data));
 
-		bool keepInMemory = KeepInMemory(fileName, fileNameForKey.count);
+		bool keepInMemory = KeepInMemory(StringView(fileName, fileNameForKey.count));
 		if (keepInMemory)
 		{
 			SCOPED_READ_LOCK(g_mappedFileTable.m_lookupLock, lock);
diff --git a/Engine/Source/Programs/UnrealBuildAccelerator/Detours/Private/UbaDetoursShared.h b/Engine/Source/Programs/UnrealBuildAccelerator/Detours/Private/UbaDetoursShared.h
--- a/Engine/Source/Programs/UnrealBuildAccelerator/Detours/Private/UbaDetoursShared.h
+++ b/Engine/Source/Programs/UnrealBuildAccelerator/Detours/Private/UbaDetoursShared.h
@@ -118,7 +118,7 @@ namespace uba
 	};
 
 	inline bool CanDetour(const tchar* file) { return !t_disallowDetour && g_rules->CanDetour(file); }
-	inline bool KeepInMemory(const tchar* fileName, u32 fileNameLen) { return g_allowKeepFilesInMemory && g_rules->KeepInMemory(fileName, fileNameLen, g_systemTemp.data); }
+	inline bool KeepInMemory(const StringView& fileName) { return g_allowKeepFilesInMemory && g_rules->KeepInMemory(fileName, g_systemTemp.data); }
 
 	void Shared_WriteConsole(const char* chars, u32 charCount, bool isError);
 	void Shared_WriteConsole(const wchar_t* chars, u32 charCount, bool isError);
diff --git a/Engine/Source/Programs/UnrealBuildAccelerator/Detours/Private/Windows/UbaDetoursFunctionsKernelBase.inl b/Engine/Source/Programs/UnrealBuildAccelerator/Detours/Private/Windows/UbaDetoursFunctionsKernelBase.inl
--- a/Engine/Source/Programs/UnrealBuildAccelerator/Detours/Private/Windows/UbaDetoursFunctionsKernelBase.inl
+++ b/Engine/Source/Programs/UnrealBuildAccelerator/Detours/Private/Windows/UbaDetoursFunctionsKernelBase.inl
@@ -971,7 +971,7 @@ DWORD Detoured_GetFileAttributesW(LPCWSTR lpFileName)
 BOOL Detoured_SetFileAttributesW(LPCWSTR lpFileName, DWORD dwFileAttributes)
 {
 	DETOURED_CALL(SetFileAttributesW);
-	if (KeepInMemory(lpFileName, u32(wcslen(lpFileName))))
+	if (KeepInMemory(StringView(lpFileName, u32(wcslen(lpFileName)))))
 	{
 		DEBUG_LOG_DETOURED(L"SetFileAttributesW", L"(%ls) %u", lpFileName, dwFileAttributes);
 		SetLastError(ERROR_SUCCESS);
@@ -1286,7 +1286,7 @@ BOOL Detoured_DeleteFileW(LPCWSTR lpFileName)
 		return True_DeleteFileW(original);
 	}
 
-	if (KeepInMemory(fixedNameLower.data, fixedNameLower.count))
+	if (KeepInMemory(fixedNameLower))
 	{
 		DEBUG_LOG_DETOURED(L"DeleteFileW", L"(INMEMORY) (%ls) -> Success", lpFileName);
 		SetLastError(ERROR_SUCCESS);
@@ -1341,7 +1341,7 @@ bool Shared_MoveFile(LPCWSTR lpExistingFileName, LPCWSTR lpNewFileName, DWORD dw
 
 	StringKey sourceKey = ToStringKeyLower(source);
 
-	if (KeepInMemory(source.data, source.count))
+	if (KeepInMemory(source))
 	{
 		SCOPED_WRITE_LOCK(g_mappedFileTable.m_lookupLock, lock);
 		auto it = g_mappedFileTable.m_lookup.find(sourceKey);
@@ -1349,7 +1349,7 @@ bool Shared_MoveFile(LPCWSTR lpExistingFileName, LPCWSTR lpNewFileName, DWORD dw
 		FileInfo& sourceInfo = it->second;
 		lock.Leave();
 
-		if (IsOutputFile(dest.data, dest.count, true))
+		if (IsOutputFile(dest, true))
 		{
 			sourceInfo.deleted = true;
 			UBA_ASSERT(!sourceInfo.memoryFile->isLocalOnly);
@@ -1360,7 +1360,7 @@ bool Shared_MoveFile(LPCWSTR lpExistingFileName, LPCWSTR lpNewFileName, DWORD dw
 			lock2.Leave();
 			FileInfo& destInfo = insres.first->second;
 			UBA_ASSERT(!insres.second); // This is here just to get a chance to investigate this scenario.. might work
-			UBA_ASSERTF(!destInfo.trueFileMapHandle && (!destInfo.memoryFile || g_rules->IsThrowAway(dest.data, dest.count)), TC("Moving file %s to %s that is an output file that is not a memory file is not supported"), source.data, lpNewFileName);
+			UBA_ASSERTF(!destInfo.trueFileMapHandle && (!destInfo.memoryFile || g_rules->IsThrowAway(dest)), TC("Moving file %s to %s that is an output file that is not a memory file is not supported"), source.data, lpNewFileName);
 			destInfo.memoryFile = sourceInfo.memoryFile;
 			sourceInfo.memoryFile = nullptr;
 			DEBUG_LOG_DETOURED(L"MoveFileExW", L"(memfile->memfile) %ls to %ls -> Success", lpExistingFileName, lpNewFileName);
@@ -1368,7 +1368,7 @@ bool Shared_MoveFile(LPCWSTR lpExistingFileName, LPCWSTR lpNewFileName, DWORD dw
 			return true;
 		}
 
-		UBA_ASSERT(!KeepInMemory(dest.data, dest.count));
+		UBA_ASSERT(!KeepInMemory(dest));
 
 		DEBUG_LOG_DETOURED(L"MoveFileExW", L"(memfile->file) %ls to %ls", lpExistingFileName, lpNewFileName);
 
diff --git a/Engine/Source/Programs/UnrealBuildAccelerator/Detours/Private/Windows/UbaDetoursFunctionsNtDll.inl b/Engine/Source/Programs/UnrealBuildAccelerator/Detours/Private/Windows/UbaDetoursFunctionsNtDll.inl
--- a/Engine/Source/Programs/UnrealBuildAccelerator/Detours/Private/Windows/UbaDetoursFunctionsNtDll.inl
+++ b/Engine/Source/Programs/UnrealBuildAccelerator/Detours/Private/Windows/UbaDetoursFunctionsNtDll.inl
@@ -683,7 +683,7 @@ NTSTATUS NTAPI Shared_NtCreateFile(bool IsCreateFunc, PHANDLE hFileHandle, ACCES
 
 	bool useContent = IsContentUse(DesiredAccess, CreateDisposition);
 	bool isWrite = IsWrite(DesiredAccess, CreateDisposition);
-	bool keepInMemory = (KeepInMemory(fileName.data, fileName.count) && useContent) || IsOutputFile(fileName.data, fileName.count, isWrite, isDeleteOnClose);
+	bool keepInMemory = (KeepInMemory(fileName) && useContent) || IsOutputFile(fileName, isWrite, isDeleteOnClose);
 
 #if UBA_DEBUG_LOG_ENABLED
 	const wchar_t* isWriteStr = isWrite ? L" WRITE" : L""; (void)isWriteStr;
@@ -1081,7 +1081,7 @@ NTSTATUS NTAPI Shared_NtCreateFile(bool IsCreateFunc, PHANDLE hFileHandle, ACCES
 			}
 			else
 			{
-				if (g_rules->IsThrowAway(fileName.data, fileName.count))
+				if (g_rules->IsThrowAway(fileName))
 				{
 					//isDeleteOnClose = true;
 				}
@@ -1092,7 +1092,7 @@ NTSTATUS NTAPI Shared_NtCreateFile(bool IsCreateFunc, PHANDLE hFileHandle, ACCES
 					return STATUS_OBJECT_NAME_NOT_FOUND;
 				}
 
-				bool isLocal = !IsOutputFile(fileName.data, fileName.count, isWrite, isDeleteOnClose);
+				bool isLocal = !IsOutputFile(fileName, isWrite, isDeleteOnClose);
 				//UBA_ASSERTF(CreateDisposition != FILE_OPEN || Contains(fileName.data, L"vctip_"), TC("Unsupported disposition %u for file %s"), CreateDisposition, fileName.data);
 				info.memoryFile = new MemoryFile(isLocal, FileTypeMaxSize(fileName, isSystemOrTempFile));
 			}
@@ -1343,7 +1343,7 @@ NTSTATUS NTAPI Detoured_NtClose(HANDLE handle)
 		mappingWritten = fi.memoryFile->writtenSize;
 
 		u32 orginalNameLen = TStrlen(fi.originalName);
-		if (IsOutputFile(fi.originalName, orginalNameLen, IsWrite(fo->desiredAccess, 0), fo->deleteOnClose) && !g_rules->IsThrowAway(fi.originalName, orginalNameLen))
+		if (IsOutputFile(StringView(fi.originalName, orginalNameLen), IsWrite(fo->desiredAccess, 0), fo->deleteOnClose) && !g_rules->IsThrowAway(StringView(fi.originalName, orginalNameLen)))
 		{
 			// Need to report this file to host so it can be tracked in directory table
 			if (!fi.memoryFile->isReported)
diff --git a/Engine/Source/Programs/UnrealBuildAccelerator/Detours/Private/Windows/UbaDetoursFunctionsWin.cpp b/Engine/Source/Programs/UnrealBuildAccelerator/Detours/Private/Windows/UbaDetoursFunctionsWin.cpp
--- a/Engine/Source/Programs/UnrealBuildAccelerator/Detours/Private/Windows/UbaDetoursFunctionsWin.cpp
+++ b/Engine/Source/Programs/UnrealBuildAccelerator/Detours/Private/Windows/UbaDetoursFunctionsWin.cpp
@@ -366,7 +366,7 @@ const wchar_t* HandleToName(HANDLE handle)
 
 bool NeedsSharedMemory(const wchar_t* file) { return g_allowKeepFilesInMemory && g_rules->NeedsSharedMemory(file); }
 u64 FileTypeMaxSize(const StringBufferBase& file, bool isSystemOrTempFile) { return g_rules->FileTypeMaxSize(file, isSystemOrTempFile); }
-bool IsOutputFile(LPCWSTR fileName, u64 fileNameLen, bool isWrite, bool isDeleteOnClose = false) { return (isWrite || isDeleteOnClose) && g_allowKeepFilesInMemory && g_rules->IsOutputFile(fileName, fileNameLen); }
+bool IsOutputFile(const StringView& fileName, bool isWrite, bool isDeleteOnClose = false) { return (isWrite || isDeleteOnClose) && g_allowKeepFilesInMemory && g_rules->IsOutputFile(fileName); }
 
 
 bool EnsureMapped(DetouredHandle& handle, DWORD dwFileOffsetHigh = 0, DWORD dwFileOffsetLow = 0, SIZE_T numberOfBytesToMap = 0, void* baseAddress = nullptr)
@@ -881,7 +881,7 @@ void Init(const DetoursPayload& payload, u64 startTime)
 	UBA_ASSERT(directoryTableMem);
 	g_directoryTable.Init(directoryTableMem, directoryTableCount, directoryTableSize);
 
-	if (payload.storeObjFilesCompressed && g_rulesIndex == 2) 	// Rules index 2 is link.exe.. compressed obj files is experimental so solution is a bit hacky
+	if (payload.storeObjFilesCompressed && g_rules->ShouldDecompressFiles(StringView()))
 		g_objFilesPreloader.Start(cmdLine);
 	else if (g_rulesIndex == 1 || g_rulesIndex == 7 || g_rulesIndex == 11 || g_rulesIndex == 14)
 		PrepopulatePchIncludedFiles(cmdLine, g_rulesIndex);
diff --git a/Engine/Source/Programs/UnrealBuildAccelerator/Detours/Private/Windows/UbaDetoursObjFilesPreloader.h b/Engine/Source/Programs/UnrealBuildAccelerator/Detours/Private/Windows/UbaDetoursObjFilesPreloader.h
--- a/Engine/Source/Programs/UnrealBuildAccelerator/Detours/Private/Windows/UbaDetoursObjFilesPreloader.h
+++ b/Engine/Source/Programs/UnrealBuildAccelerator/Detours/Private/Windows/UbaDetoursObjFilesPreloader.h
@@ -5,6 +5,7 @@
 #include "UbaCompressedObjFileHeader.h"
 #include "UbaDetoursUtilsWin.h"
 #include "UbaEvent.h"
+#include "UbaProcessUtils.h"
 #include <oodle2.h>
 #define _NTDEF_
 #include <ntsecapi.h>
@@ -53,22 +54,8 @@ namespace uba
 	class ObjFilesPreloader
 	{
 	public:
-		void Start(const wchar_t* cmdLine)
+		void ParseRsp(const StringView& rspFile)
 		{
-			const wchar_t* at = wcschr(cmdLine, '@');
-			if (!at)
-				return;
-
-			const wchar_t* end = at + wcslen(at);
-			StringBuffer<> rspFile;
-			++at;
-			if (*at == '\"')
-			{
-				++at;
-				--end;
-			}
-			rspFile.Append(at, end - at);
-
 			HANDLE rspFileHandle = CreateFileW(rspFile.data, GENERIC_READ, FILE_SHARE_READ, NULL, OPEN_EXISTING, FILE_ATTRIBUTE_NORMAL, NULL);
 			if (rspFileHandle == INVALID_HANDLE_VALUE)
 				return;
@@ -82,62 +69,101 @@ namespace uba
 			CloseHandle(rspFileHandle);
 			if (!rspFileMappingHandle)
 				return;
-			auto rspMem = (char*)MapViewOfFile(rspFileMappingHandle, FILE_MAP_READ, 0, 0, rspFileSize.QuadPart);
+			const void* rspMem = MapViewOfFile(rspFileMappingHandle, FILE_MAP_READ, 0, 0, rspFileSize.QuadPart);
 			CloseHandle(rspFileMappingHandle);
 			if (!rspMem)
 				return;
 
-			auto lineBegin = rspMem;
-			while (*lineBegin)
+			ParseArguments((const char*)rspMem, rspFileSize.QuadPart, [&](const char* arg, u32 argLen)
+				{
+					StringBuffer<> sb;
+					sb.Append(arg, argLen);
+					HandleLine(sb);
+				});
+
+			UnmapViewOfFile(rspMem);
+		}
+
+		void HandleLine(const StringView& line)
+		{
+			if (line.data[0] == '/' || line.data[0] == '-')
+				return;
+
+			if (line.data[0] == '@')
 			{
-				auto lineEnd = strchr(lineBegin, '\n');
-				if (!lineEnd)
-					lineEnd = lineBegin + strlen(lineBegin);
-				if (*lineBegin != '/')
+				const wchar_t* begin = line.data + 1;
+				const wchar_t* end = nullptr;
+				StringBuffer<> rspFile;
+				if (*begin == '\"')
 				{
-					char line[512];
-					u64 lineLen = lineEnd - lineBegin;
-					memcpy(line, lineBegin, lineLen);
-					line[lineLen] = 0;
-
-					if (char* ext = strstr(line, ".obj"))
-					{
-						ext[4] = 0;
-
-						char* objFile = line;
-						if (line[0] == '\"')
-							++objFile;
-
-						StringBuffer<> objFile2;
-						objFile2.Append(objFile);
-
-						StringBuffer<> objFileFull;
-						FixPath(objFileFull, objFile2.data);
-
-						HANDLE objFileHandle = CreateFileW(objFileFull.data, GENERIC_READ, FILE_SHARE_READ, NULL, OPEN_EXISTING, FILE_ATTRIBUTE_NORMAL, NULL);
-						UBA_ASSERT(objFileHandle != INVALID_HANDLE_VALUE);
-
-						StringKey fileNameKey = ToStringKey(objFileFull.MakeLower());
-						auto insres = m_preloadedObjFiles.try_emplace(fileNameKey);
-						UBA_ASSERT(insres.second);
-						insres.first->second.handle = objFileHandle;
-						insres.first->second.event.Create(true);
-
-						SCOPED_READ_LOCK(g_mappedFileTable.m_lookupLock, _);
-						auto findIt = g_mappedFileTable.m_lookup.find(fileNameKey);
-						insres.first->second.fileInfo = &findIt->second;
-					}
+					++begin;
+					end = wcschr(begin, '\"');
+					if (!end)
+						return;
 				}
-
-				if (!*lineEnd)
-					break;
-				lineBegin = lineEnd + 1;
+				else
+				{
+					end = wcschr(begin, ' ');
+					if (!end)
+						end = begin + wcslen(begin);
+				}
+				rspFile.Append(begin, end - begin);
+				ParseRsp(rspFile);
+				return;
 			}
 
+			StringBuffer<> file;
+			file.Append(line);
+			if (!g_rules->ShouldDecompressFiles(file))
+				return;
+			StringBuffer<> fileFull;
+			FixPath(fileFull, file.data);
+
+			HANDLE fileHandle = CreateFileW(fileFull.data, GENERIC_READ, FILE_SHARE_READ, NULL, OPEN_EXISTING, FILE_ATTRIBUTE_NORMAL, NULL);
+			UBA_ASSERT(fileHandle != INVALID_HANDLE_VALUE);
+
+			StringKey fileNameKey = ToStringKey(fileFull.MakeLower());
+			auto insres = m_preloadedObjFiles.try_emplace(fileNameKey);
+			UBA_ASSERT(insres.second);
+			insres.first->second.handle = fileHandle;
+			insres.first->second.event.Create(true);
+
+			SCOPED_READ_LOCK(g_mappedFileTable.m_lookupLock, _);
+			auto findIt = g_mappedFileTable.m_lookup.find(fileNameKey);
+			insres.first->second.fileInfo = &findIt->second;
+		}
+
+		void Start(const wchar_t* cmdLine)
+		{
+			Oodle_SetUsageWarnings(Oodle_UsageWarnings_Disabled);
+
+			auto handleLine = [&](const StringView& line)
+				{
+					StringBuffer<> file;
+					file.Append(line);
+					if (!g_rules->ShouldDecompressFiles(file))
+						return;
+					StringBuffer<> fileFull;
+					FixPath(fileFull, file.data);
+
+					HANDLE fileHandle = CreateFileW(fileFull.data, GENERIC_READ, FILE_SHARE_READ, NULL, OPEN_EXISTING, FILE_ATTRIBUTE_NORMAL, NULL);
+					UBA_ASSERT(fileHandle != INVALID_HANDLE_VALUE);
+
+					StringKey fileNameKey = ToStringKey(fileFull.MakeLower());
+					auto insres = m_preloadedObjFiles.try_emplace(fileNameKey);
+					UBA_ASSERT(insres.second);
+					insres.first->second.handle = fileHandle;
+					insres.first->second.event.Create(true);
+
+					SCOPED_READ_LOCK(g_mappedFileTable.m_lookupLock, _);
+					auto findIt = g_mappedFileTable.m_lookup.find(fileNameKey);
+					insres.first->second.fileInfo = &findIt->second;
+				};
+
+			ParseArguments(cmdLine, [&](const tchar* arg, u32 argLen) { HandleLine(StringView(arg, argLen)); });
+
 			if (!m_preloadedObjFiles.empty())
 				m_threadHandle = CreateThread(NULL, 0, [](LPVOID p) -> DWORD { ((ObjFilesPreloader*)p)->ThreadPreload(); return 0; }, this, 0, NULL);
-
-			UnmapViewOfFile(rspMem);
 		}
 
 		void Stop()
@@ -164,6 +190,8 @@ namespace uba
 				LARGE_INTEGER objFileSize;
 				BOOL res = GetFileSizeEx(objFileHandle, &objFileSize);
 				UBA_ASSERT(res);(void)res;
+				if (objFileSize.QuadPart < sizeof(CompressedObjFileHeader))
+					continue;
 
 				HANDLE objFileMappingHandle = CreateFileMappingW(objFileHandle, NULL, PAGE_READONLY, objFileSize.HighPart, objFileSize.LowPart, NULL);
 				UBA_ASSERT(objFileMappingHandle);
@@ -175,6 +203,11 @@ namespace uba
 				UBA_ASSERT(preload.objMem);
 				if (!preload.objMem)
 					continue;
+				if (!((CompressedObjFileHeader*)preload.objMem)->IsValid())
+				{
+					preload.event.Set();
+					continue;
+				}
 				preload.objCompressedSize = objFileSize.QuadPart;
 				preload.objDecompressedSize = *(u64*)(preload.objMem + sizeof(CompressedObjFileHeader));
 				preload.objReadOffset = 8 + sizeof(CompressedObjFileHeader);
diff --git a/Engine/Source/Programs/UnrealBuildAccelerator/Test/Private/UbaTestBasics.h b/Engine/Source/Programs/UnrealBuildAccelerator/Test/Private/UbaTestBasics.h
--- a/Engine/Source/Programs/UnrealBuildAccelerator/Test/Private/UbaTestBasics.h
+++ b/Engine/Source/Programs/UnrealBuildAccelerator/Test/Private/UbaTestBasics.h
@@ -7,9 +7,9 @@
 #include "UbaBinaryReaderWriter.h"
 #include "UbaCompactTables.h"
 #include "UbaFileAccessor.h"
-#include "UbaProcess.h"
 #include "UbaPathUtils.h"
 #include "UbaPlatform.h"
+#include "UbaProcessUtils.h"
 #include "UbaEvent.h"
 #include "UbaFile.h"
 #include "UbaLogger.h"
@@ -247,6 +247,8 @@ namespace uba
 
 	bool TestParseArguments(Logger& logger, const StringBufferBase& rootDir)
 	{
+		auto ParseArguments = [](Vector<TString>& a, const tchar* args) { uba::ParseArguments(args, [&](const tchar* arg, u32 argLen) { a.push_back({arg, argLen}); }); };
+
 		Vector<TString> arguments;
 		ParseArguments(arguments, TC("foo bar"));
 		UBA_TEST_CHECK(arguments.size() == 2, "ParseArguments 1 failed (%llu)", arguments.size());
@@ -280,6 +282,19 @@ namespace uba
 		ParseArguments(arguments7, TC("app \" \\\"foo\\\" bar\""));
 		UBA_TEST_CHECK(arguments7.size() == 2, "ParseArguments 7 failed");
 		UBA_TEST_CHECK(Contains(arguments7[1].data(), TC("\"")), "ParseArguments 7 failed");
+
+		Vector<TString> arguments8;
+		ParseArguments(arguments8, TC("\nline1\r\nline2\r\nline3\n\r\n"));
+		UBA_TEST_CHECK(arguments8.size() == 3, "ParseArguments 8 failed");
+		UBA_TEST_CHECK(Equals(arguments8[0].data(), TC("line1")), "ParseArguments 8 failed");
+		UBA_TEST_CHECK(Equals(arguments8[1].data(), TC("line2")), "ParseArguments 8 failed");
+		UBA_TEST_CHECK(Equals(arguments8[2].data(), TC("line3")), "ParseArguments 8 failed");
+
+		Vector<TString> arguments9;
+		ParseArguments(arguments9, TC("\"foo\\\\\" \"bar\\\\\""));
+		UBA_TEST_CHECK(arguments9.size() == 2, "ParseArguments 9 failed");
+		UBA_TEST_CHECK(Equals(arguments9[0].data(), TC("foo\\\\")), "ParseArguments 9 failed");
+		UBA_TEST_CHECK(Equals(arguments9[1].data(), TC("bar\\\\")), "ParseArguments 9 failed");
 		return true;
 	}
 
