[UBA]
* Fixed write path when store compressed obj files are enabled

diff --git a/Engine/Source/Programs/UnrealBuildAccelerator/Common/Private/UbaStorage.cpp b/Engine/Source/Programs/UnrealBuildAccelerator/Common/Private/UbaStorage.cpp
--- a/Engine/Source/Programs/UnrealBuildAccelerator/Common/Private/UbaStorage.cpp
+++ b/Engine/Source/Programs/UnrealBuildAccelerator/Common/Private/UbaStorage.cpp
@@ -2256,18 +2256,6 @@ namespace uba
 					if (!compressedData)
 						return m_logger.Error(TC("Failed to map view of mapping %s (%s)"), casFile.data, LastErrorToText().data);
 
-					if (writeCompressed)
-					{
-						// TODO: This should be more general.. send in header?
-						FileAccessor destinationFile(m_logger, destination);
-						if (!destinationFile.CreateMemoryWrite(false, fileAttributes, mappedView.size + sizeof(CompressedObjFileHeader), m_tempPath.data))
-							return false;
-						u8* mem = destinationFile.GetData();
-						*(CompressedObjFileHeader*)mem = CompressedObjFileHeader(casKey);
-						memcpy(mem + sizeof(CompressedObjFileHeader), mappedView.memory, mappedView.size);
-						return destinationFile.Close();
-					}
-
 					decompressedSize = *(u64*)compressedData;
 					readData = compressedData + sizeof(u64);
 				}
@@ -2279,9 +2267,6 @@ namespace uba
 #else
 					UBA_ASSERT(false);
 #endif
-					if (writeCompressed)
-						return m_logger.Error(TC("NOT IMPLEMENTED %s"), destination);
-
 					if (!OpenFileSequentialRead(m_logger, casFile.data, readHandle))
 						return m_logger.Error(TC("Failed to open file %s for read (%s)"), casFile.data, LastErrorToText().data);
 
@@ -2307,7 +2292,34 @@ namespace uba
 				BottleneckScope scope(bottleneck);
 				#endif
 
-				if (writeDirectlyToFile || !decompressedSize)
+				if (writeCompressed)
+				{
+					u64 compressedFileSize = mappedView.size;
+					if (!mappedView.memory)
+						if (!GetFileSizeEx(compressedFileSize, readHandle))
+							return m_logger.Error(TC("Failed to get file size of compressed file %s (%s)"), casFile.data, LastErrorToText().data);
+						
+
+					if (!destinationFile.CreateMemoryWrite(false, fileAttributes, compressedFileSize + sizeof(CompressedObjFileHeader), m_tempPath.data))
+						return false;
+					u8* writePos = destinationFile.GetData();
+					*(CompressedObjFileHeader*)writePos = CompressedObjFileHeader(casKey);
+					writePos += sizeof(CompressedObjFileHeader);
+
+					if (mappedView.memory)
+					{
+						memcpy(writePos, mappedView.memory, compressedFileSize);
+					}
+					else
+					{
+						*(u64*)writePos = decompressedSize;
+						writePos += sizeof(u64);
+						if (!ReadFile(m_logger, casFile.data, readHandle, writePos, compressedFileSize - sizeof(u64)))
+							return m_logger.Error(TC("Failed to read compressed file %s (%s)"), casFile.data, LastErrorToText().data);
+					}
+						
+				}
+				else if (writeDirectlyToFile || !decompressedSize)
 				{
 					if (!destinationFile.CreateWrite(allowRead, writeFlags | fileAttributes, decompressedSize, m_tempPath.data))
 						return false;
