[UBA]
* Changed algorithm for updating input offsets on cache entries in an attempt to make it faster when there are lots of entries on the same key
* Added unit tests

diff --git a/Engine/Source/Programs/UnrealBuildAccelerator/Common/Private/UbaCacheEntry.cpp b/Engine/Source/Programs/UnrealBuildAccelerator/Common/Private/UbaCacheEntry.cpp
--- a/Engine/Source/Programs/UnrealBuildAccelerator/Common/Private/UbaCacheEntry.cpp
+++ b/Engine/Source/Programs/UnrealBuildAccelerator/Common/Private/UbaCacheEntry.cpp
@@ -180,6 +180,10 @@ namespace uba
 			u64 outputSize = reader.Read7BitEncoded();
 			entry.outputCasKeyOffsets.resize(outputSize);
 			reader.ReadBytes(entry.outputCasKeyOffsets.data(), outputSize);
+
+			#if UBA_USE_OLD
+			Flatten(entry.inputCasKeyOffsets, entry);
+			#endif
 		}
 
 		return true;
@@ -222,7 +226,6 @@ namespace uba
 
 		u32 rangeBegin = 0;
 		bool inRange = false;
-		bool unhandledOffset = false;
 		u32 lastSharedPos = ~0u;
 
 		Vector<u8> extraOffsets;
@@ -294,15 +297,96 @@ namespace uba
 				{
 					extraWriter.Write7BitEncoded(offset);
 				}
+			}
+		}
+
+		entry.extraInputCasKeyOffsets.resize(extraWriter.GetPosition());
+		memcpy(entry.extraInputCasKeyOffsets.data(), extraWriter.GetData(), extraWriter.GetPosition());
+	}
+
+	template<typename Container>
+	void CacheEntries::BuildRangesFromExcludedT(CacheEntry& entry, const Container& sortedExcludedInputs)
+	{
+		StackBinaryWriter<256*1024> rangeWriter;
+
+		auto g = MakeGuard([&]()
+			{
+				entry.sharedInputCasKeyOffsetRanges.resize(rangeWriter.GetPosition());
+				memcpy(entry.sharedInputCasKeyOffsetRanges.data(), rangeWriter.GetData(), rangeWriter.GetPosition());
+			});
+
+		auto writeRange = [&](u64 begin, u64 end) { rangeWriter.Write7BitEncoded(begin); rangeWriter.Write7BitEncoded(end); };
+
+		auto excludedInputsIt = sortedExcludedInputs.begin();
+		auto excludedInputsEnd = sortedExcludedInputs.end();
+
+		BinaryReader sharedReader(sharedInputCasKeyOffsets.data(), 0, sharedInputCasKeyOffsets.size());
+
+		u32 sharedOffset = ~0u;
+		u32 offset = ~0u;
+
+		u32 excludeRangeEnd = 0;
+		u32 excludeRangeBegin = 0;
+		bool inExcludeRange = false;
+		u32 lastSharedPos = ~0u;
+
+		while (true)
+		{
+			u32 sharedPos = u32(sharedReader.GetPosition());
+
+			if (!sharedReader.GetLeft())
+			{
+				if (!inExcludeRange)
+					writeRange(excludeRangeEnd, excludeRangeBegin);
+				break;
+			}
+
+			if (offset <= sharedOffset && excludedInputsIt == excludedInputsEnd)
+			{
+				if (!inExcludeRange)
+					writeRange(excludeRangeEnd, sharedInputCasKeyOffsets.size());
 				else
+					writeRange(sharedPos, sharedInputCasKeyOffsets.size());
+				break;
+			}
+
+			if (sharedOffset < offset)
+			{
+				lastSharedPos = sharedPos;
+				sharedOffset = u32(sharedReader.Read7BitEncoded());
+			}
+			else if (offset < sharedOffset)
+			{
+				offset = *excludedInputsIt++;
+				sharedPos = lastSharedPos;
+			}
+			else
+			{
+				lastSharedPos = sharedPos;
+				sharedOffset = u32(sharedReader.Read7BitEncoded());
+				offset = *excludedInputsIt++;
+			}
+
+			if (sharedOffset == offset)
+			{
+				if (!inExcludeRange)
+				{
+					if (excludeRangeEnd != lastSharedPos)
+						writeRange(excludeRangeEnd, lastSharedPos);
+					excludeRangeBegin = sharedPos;
+					inExcludeRange = true;
+				}
+			}
+			else
+			{
+				if (inExcludeRange)
 				{
-					unhandledOffset = true;
+					inExcludeRange = false;
+					excludeRangeEnd = sharedPos;
 				}
 			}
 		}
 
-		entry.extraInputCasKeyOffsets.resize(extraWriter.GetPosition());
-		memcpy(entry.extraInputCasKeyOffsets.data(), extraWriter.GetData(), extraWriter.GetPosition());
 	}
 
 	void CacheEntries::BuildInputs(CacheEntry& entry, const Set<u32>& inputs)
@@ -336,6 +420,7 @@ namespace uba
 		{
 			if (isFirst)
 			{
+				primaryId = entry.id;
 				// Flatten first entry into new shared
 				Flatten(oldShared, entry);
 				oldShared.swap(sharedInputCasKeyOffsets);
@@ -357,7 +442,7 @@ namespace uba
 		}
 	}
 
-	void CacheEntries::UpdateEntries(Logger& logger, const GrowingNoLockUnorderedMap<u32, u32>& oldToNewCasKeyOffset, Vector<u32>& temp)
+	void CacheEntries::UpdateEntries(Logger& logger, const GrowingNoLockUnorderedMap<u32, u32>& oldToNewCasKeyOffset, Vector<u32>& temp, Vector<u8>& temp2)
 	{
 		if (entries.empty())
 			return;
@@ -366,97 +451,167 @@ namespace uba
 		ValidateEntries(logger);
 		#endif
 
-		Vector<u8> oldShared;
-		oldShared.swap(sharedInputCasKeyOffsets);
-
-		bool isFirst = true;
-		for (auto& entry : entries)
-		{
-			if (isFirst)
+		auto convertOffsets = [&](Vector<u8>& offsets)
 			{
-				// Flatten first entry into temp
-				Flatten(temp, entry, oldShared);
-
-				// Update temp with new offsets
-				u64 newSize = 0;
-				for (auto& offset : temp)
+				temp.clear();
+				u32 newOffsetsSize = 0;
+				BinaryReader reader(offsets.data(), 0, offsets.size());
+				while (reader.GetLeft())
 				{
-					auto findIt = oldToNewCasKeyOffset.find(offset);
+					u32 newOffset = u32(reader.Read7BitEncoded());
+					auto findIt = oldToNewCasKeyOffset.find(newOffset);
 					if (findIt != oldToNewCasKeyOffset.end())
-						offset = findIt->second;
-					newSize += Get7BitEncodedCount(offset);
+						newOffset = findIt->second;
+					temp.push_back(newOffset);
+					newOffsetsSize += Get7BitEncodedCount(newOffset);
 				}
 
-				// Sort temp now when it likely is out of order
 				std::sort(temp.begin(), temp.end());
 
-				sharedInputCasKeyOffsets.resize(newSize);
-				BinaryWriter writer(sharedInputCasKeyOffsets.data(), 0, newSize);
-				for (auto& offset : temp)
+				offsets.resize(newOffsetsSize);
+				BinaryWriter writer(offsets.data(), 0, newOffsetsSize);
+				for (u32 offset : temp)
 					writer.Write7BitEncoded(offset);
+			};
 
-				entry.extraInputCasKeyOffsets.clear();
+		for (auto& entry : entries)
+		{
+			#if UBA_USE_OLD
+			convertOffsets(entry.inputCasKeyOffsets);
+			#endif
+			convertOffsets(entry.outputCasKeyOffsets);
+		}
+
+		auto writePrimaryRange = [&](CacheEntry& entry, u64 newSize)
+			{
 				u64 rangeSize = 1 + Get7BitEncodedCount(newSize);
 				entry.sharedInputCasKeyOffsetRanges.resize(rangeSize);
 				BinaryWriter rangeWriter(entry.sharedInputCasKeyOffsetRanges.data(), 0, rangeSize);
 				rangeWriter.Write7BitEncoded(0);
 				rangeWriter.Write7BitEncoded(newSize);
 				UBA_ASSERT(rangeWriter.GetPosition() == rangeSize);
-				isFirst = false;
-			}
-			else
+			};
+
+		// If primary id is not set we use first entry as primaryId and base shared offsets off primary entry
+		if (entries.size() == 1 || primaryId == ~0u)
+		{
+			auto& oldShared = temp2;
+			oldShared = sharedInputCasKeyOffsets;
+
+			bool isFirst = true;
+			for (auto& entry : entries)
 			{
-				// Flatten using old shared and rebuild it with new shared
-				Flatten(temp, entry, oldShared);
-				for (auto& offset : temp)
+				if (isFirst)
 				{
-					auto findIt = oldToNewCasKeyOffset.find(offset);
-					if (findIt != oldToNewCasKeyOffset.end())
-						offset = findIt->second;
+					primaryId = entry.id;
+
+					// Flatten first entry into temp
+					Flatten(temp, entry, oldShared);
+
+					// Update temp with new offsets
+					u64 newSize = 0;
+					for (auto& offset : temp)
+					{
+						auto findIt = oldToNewCasKeyOffset.find(offset);
+						if (findIt != oldToNewCasKeyOffset.end())
+							offset = findIt->second;
+						newSize += Get7BitEncodedCount(offset);
+					}
+
+					// Sort temp now when it likely is out of order
+					std::sort(temp.begin(), temp.end());
+
+					// Write new shared
+					sharedInputCasKeyOffsets.resize(newSize);
+					BinaryWriter writer(sharedInputCasKeyOffsets.data(), 0, newSize);
+					for (auto& offset : temp)
+						writer.Write7BitEncoded(offset);
+
+					// Clear extra and set entire shared to range
+					entry.extraInputCasKeyOffsets.clear();
+					writePrimaryRange(entry, newSize);
+					isFirst = false;
+				}
+				else
+				{
+					// Flatten using old shared and rebuild it with new shared
+					Flatten(temp, entry, oldShared);
+					for (auto& offset : temp)
+					{
+						auto findIt = oldToNewCasKeyOffset.find(offset);
+						if (findIt != oldToNewCasKeyOffset.end())
+							offset = findIt->second;
+					}
+
+					// Sort temp now when it likely is out of order
+					std::sort(temp.begin(), temp.end());
+
+					entry.extraInputCasKeyOffsets.clear();
+					entry.sharedInputCasKeyOffsetRanges.clear();
+					BuildInputsT(entry, temp, false);
 				}
-
-				// Sort temp now when it likely is out of order
-				std::sort(temp.begin(), temp.end());
-
-				entry.extraInputCasKeyOffsets.clear();
-				entry.sharedInputCasKeyOffsetRanges.clear();
-				BuildInputsT(entry, temp, false);
 			}
 		}
+		else
+		{
+			// This approach should be faster if there are more than one entry since we expect entries to be very similar to each other
+			// It instead tracks removed offsets when calculating the shared offsets and build the ranges from that.
 
-		auto updateCasKeyOffsets = [&](Vector<u8>& offsets)
+			auto& oldShared = temp2;
+			oldShared = sharedInputCasKeyOffsets;
+			convertOffsets(sharedInputCasKeyOffsets);
+
+
+			for (auto& entry : entries)
 			{
+				// Collect all inputs that are removed from the shared inputs
+
+				auto collectInputs = [&](Vector<u32>& out, u32 rangeBegin, u32 rangeEnd)
+					{
+						BinaryReader excludedReader(oldShared.data() + rangeBegin, 0, rangeEnd - rangeBegin);
+						while (excludedReader.GetLeft())
+						{
+							u32 offset = u32(excludedReader.Read7BitEncoded());
+							auto findIt = oldToNewCasKeyOffset.find(offset);
+							if (findIt != oldToNewCasKeyOffset.end())
+								offset = findIt->second;
+							out.push_back(offset);
+						}
+					};
+
 				temp.clear();
-				temp.reserve(offsets.size()*4);
+				auto& excludedOffsets = temp;
 
-				u32 newOffsetsSize = 0;
-				BinaryReader reader2(offsets.data(), 0, offsets.size());
-				while (reader2.GetLeft())
+				BinaryReader rangeReader(entry.sharedInputCasKeyOffsetRanges.data(), 0, entry.sharedInputCasKeyOffsetRanges.size());
+				u32 lastEnd = 0;
+				while (rangeReader.GetLeft())
 				{
-					u32 oldOffset = u32(reader2.Read7BitEncoded());
-					u32 newOffset = oldOffset;
-					auto findIt = oldToNewCasKeyOffset.find(oldOffset);
-					if (findIt != oldToNewCasKeyOffset.end())
-						newOffset = findIt->second;
-					temp.push_back(newOffset);
-					newOffsetsSize += Get7BitEncodedCount(newOffset);
+					u32 begin = u32(rangeReader.Read7BitEncoded());
+					collectInputs(excludedOffsets, lastEnd, begin);
+					lastEnd = u32(rangeReader.Read7BitEncoded());
 				}
+				collectInputs(excludedOffsets, lastEnd, u32(oldShared.size()));
 
-				std::sort(temp.begin(), temp.end());
+				if (excludedOffsets.empty() && entry.extraInputCasKeyOffsets.empty())
+				{
+					writePrimaryRange(entry, sharedInputCasKeyOffsets.size());
+				}
+				else
+				{
+					// Sort excluded inputs..
+					std::sort(excludedOffsets.begin(), excludedOffsets.end());
 
-				offsets.resize(newOffsetsSize);
-				BinaryWriter writer2(offsets.data(), 0, newOffsetsSize);
-				for (u32 offset : temp)
-					writer2.Write7BitEncoded(offset);
-				UBA_ASSERT(writer2.GetPosition() == newOffsetsSize);
-			};
+					// Build new ranges based on shared and excluded offsets from shared
+					BuildRangesFromExcludedT(entry, excludedOffsets);
 
-		for (auto& entry : entries)
-		{
-			#if UBA_USE_OLD
-			updateCasKeyOffsets(entry.inputCasKeyOffsets);
-			#endif
-			updateCasKeyOffsets(entry.outputCasKeyOffsets);
+					// Create new extras
+					convertOffsets(entry.extraInputCasKeyOffsets);
+				}
+
+				#if UBA_USE_OLD
+				ValidateEntry(logger, entry);
+				#endif
+			}
 		}
 
 		#if UBA_USE_OLD
@@ -468,15 +623,17 @@ namespace uba
 	void CacheEntries::ValidateEntries(Logger& logger)
 	{
 		for (auto& entry : entries)
-			ValidateEntry(logger, entry);
+			ValidateEntry(logger, entry, entry.inputCasKeyOffsets);
 	}
-	void CacheEntries::ValidateEntry(Logger& logger, CacheEntry& entry)
+	#endif
+
+	void CacheEntries::ValidateEntry(Logger& logger, CacheEntry& entry, Vector<u8>& inputCasKeyOffsets)
 	{
 		Vector<u8> res;
 		Flatten(res, entry);
-		if (res.size() == entry.inputCasKeyOffsets.size() && memcmp(res.data(), entry.inputCasKeyOffsets.data(), res.size()) == 0)
+		if (res.size() == inputCasKeyOffsets.size() && memcmp(res.data(), inputCasKeyOffsets.data(), res.size()) == 0)
 			return;
-		BinaryReader reader1(entry.inputCasKeyOffsets.data(), 0, entry.inputCasKeyOffsets.size());
+		BinaryReader reader1(inputCasKeyOffsets.data(), 0, inputCasKeyOffsets.size());
 		BinaryReader reader2(res.data(), 0, res.size());
 		while (reader1.GetLeft() || reader2.GetLeft())
 		{
@@ -492,8 +649,6 @@ namespace uba
 		UBA_ASSERT(false);
 	}
 
-	#endif
-
 	void CacheEntries::Flatten(Vector<u8>& out, const CacheEntry& entry)
 	{
 		u64 size = entry.extraInputCasKeyOffsets.size();
diff --git a/Engine/Source/Programs/UnrealBuildAccelerator/Common/Private/UbaCacheServer.cpp b/Engine/Source/Programs/UnrealBuildAccelerator/Common/Private/UbaCacheServer.cpp
--- a/Engine/Source/Programs/UnrealBuildAccelerator/Common/Private/UbaCacheServer.cpp
+++ b/Engine/Source/Programs/UnrealBuildAccelerator/Common/Private/UbaCacheServer.cpp
@@ -284,7 +284,9 @@ namespace uba
 			});
 
 
-		bool forceAllSteps = false;
+		//m_forceAllSteps = true;
+		bool forceAllSteps = m_forceAllSteps;
+		m_forceAllSteps = false;
 
 		if (m_shouldWipe)
 		{
@@ -459,6 +461,8 @@ namespace uba
 						// Remove entry from entries list and skip increasing ref count of cas files
 						if (deleteEntry)
 						{
+							if (i->id == entries.primaryId)
+								entries.primaryId = ~0u;
 							bucket.hasDeletedEntries = true;
 							++deleteEntryCount;
 							i = entries.entries.erase(i);
@@ -638,9 +642,9 @@ namespace uba
 				// Update all casKeyOffsets
 				u64 updateEntriesStart = GetTime();
 
-				m_server.ParallelFor(workerCountToUse, bucket.m_cacheEntryLookup, [&, temp = Vector<u32>()](auto& it) mutable
+				m_server.ParallelFor(workerCountToUse, bucket.m_cacheEntryLookup, [&, temp = Vector<u32>(), temp2 = Vector<u8>()](auto& it) mutable
 					{
-						it->second.UpdateEntries(m_logger, oldToNewCasKeyOffset, temp);
+						it->second.UpdateEntries(m_logger, oldToNewCasKeyOffset, temp, temp2);
 					});
 
 				#if 0
@@ -1263,6 +1267,12 @@ namespace uba
 			m_addsSinceMaintenance = 1;
 			writeLine(line.Clear().Appendf(TC("Cache server database obliteration queued!")).data);
 		}
+		else if (command.Equals(TC("maintenance")))
+		{
+			m_forceAllSteps = true;
+			m_addsSinceMaintenance = 1;
+			writeLine(line.Clear().Appendf(TC("Cache server maintenance queued!")).data);
+		}
 		else
 		{
 			writeLine(line.Clear().Appendf(TC("Unknown command: %s"), command.data).data);
diff --git a/Engine/Source/Programs/UnrealBuildAccelerator/Common/Public/UbaCacheEntry.h b/Engine/Source/Programs/UnrealBuildAccelerator/Common/Public/UbaCacheEntry.h
--- a/Engine/Source/Programs/UnrealBuildAccelerator/Common/Public/UbaCacheEntry.h
+++ b/Engine/Source/Programs/UnrealBuildAccelerator/Common/Public/UbaCacheEntry.h
@@ -29,6 +29,7 @@ namespace uba
 		List<CacheEntry> entries;
 		Vector<u8> sharedInputCasKeyOffsets;
 		u32 idCounter = 0;
+		u32 primaryId = ~0u; // Id of entry that shared offsets was made from
 
 		u64 GetSharedSize();
 		u64 GetEntrySize(CacheEntry& entry, bool toDisk);
@@ -37,11 +38,10 @@ namespace uba
 		bool Read(Logger& logger, BinaryReader& reader, u32 databaseVersion);
 		void BuildInputs(CacheEntry& entry, const Set<u32>& inputs);
 		void UpdateEntries();
-		void UpdateEntries(Logger& logger, const GrowingNoLockUnorderedMap<u32, u32>& oldToNewCasKeyOffset, Vector<u32>& temp);
+		void UpdateEntries(Logger& logger, const GrowingNoLockUnorderedMap<u32, u32>& oldToNewCasKeyOffset, Vector<u32>& temp, Vector<u8>& temp2);
 
 		#if UBA_USE_OLD
 		void ValidateEntries(Logger& logger);
-		void ValidateEntry(Logger& logger, CacheEntry& entry);
 		#endif
 
 		void Flatten(Vector<u8>& out, const CacheEntry& entry);
@@ -49,6 +49,11 @@ namespace uba
 
 		template<typename Container>
 		void BuildInputsT(CacheEntry& entry, const Container& sortedInputs, bool populateShared);
+
+		template<typename Container>
+		void BuildRangesFromExcludedT(CacheEntry& entry, const Container& sortedExcludedInputs);
+
+		void ValidateEntry(Logger& logger, CacheEntry& entry, Vector<u8>& inputCasKeyOffsets);
 	};
 
 
diff --git a/Engine/Source/Programs/UnrealBuildAccelerator/Common/Public/UbaCacheServer.h b/Engine/Source/Programs/UnrealBuildAccelerator/Common/Public/UbaCacheServer.h
--- a/Engine/Source/Programs/UnrealBuildAccelerator/Common/Public/UbaCacheServer.h
+++ b/Engine/Source/Programs/UnrealBuildAccelerator/Common/Public/UbaCacheServer.h
@@ -55,7 +55,6 @@ namespace uba
 
 		Atomic<u32> m_addsSinceMaintenance;
 		Atomic<bool> m_isRunningMaintenance;
-		Atomic<bool> m_shouldWipe;
 
 		ReaderWriterLock m_bucketsLock;
 		Map<u64, Bucket> m_buckets;
@@ -71,6 +70,9 @@ namespace uba
 
 		bool m_checkInputsForDeletedCas = true;
 
+		bool m_shouldWipe = false;
+		bool m_forceAllSteps = false;
+
 		CacheServer(const CacheServer&) = delete;
 		CacheServer& operator=(const CacheServer&) = delete;
 	};
diff --git a/Engine/Source/Programs/UnrealBuildAccelerator/Test/Private/UbaTestAll.h b/Engine/Source/Programs/UnrealBuildAccelerator/Test/Private/UbaTestAll.h
--- a/Engine/Source/Programs/UnrealBuildAccelerator/Test/Private/UbaTestAll.h
+++ b/Engine/Source/Programs/UnrealBuildAccelerator/Test/Private/UbaTestAll.h
@@ -3,6 +3,7 @@
 #pragma once
 
 #include "UbaTestBasics.h"
+#include "UbaTestCache.h"
 #include "UbaTestNetwork.h"
 #include "UbaTestScheduler.h"
 #include "UbaTestSession.h"
@@ -29,6 +30,7 @@ namespace uba
 		UBA_TEST(TestDetouredTestApp) \
 		UBA_TEST(TestRemoteDetouredTestApp) \
 		UBA_TEST(TestCompactPathTable) \
+		UBA_TEST(TestCacheEntry) \
 
 
 #define UBA_POSIX_TESTS \
diff --git a/Engine/Source/Programs/UnrealBuildAccelerator/Test/Private/UbaTestCache.h b/Engine/Source/Programs/UnrealBuildAccelerator/Test/Private/UbaTestCache.h
new file mode 100644
--- /dev/null
+++ b/Engine/Source/Programs/UnrealBuildAccelerator/Test/Private/UbaTestCache.h
@@ -0,0 +1,42 @@
+// Copyright Epic Games, Inc. All Rights Reserved.
+
+#pragma once
+
+#include "UbaCacheEntry.h"
+
+namespace uba
+{
+	bool TestCacheEntry(Logger& logger, const StringBufferBase& rootDir)
+	{
+		CacheEntries entries;
+
+		auto AddEntry = [&](const Set<u32>& inputs)
+			{
+				Vector<u8> inputOffsets;
+				u64 bytes = 0;
+				for (u32 i : inputs)
+					bytes += Get7BitEncodedCount(i);
+				inputOffsets.resize(bytes);
+				BinaryWriter writer(inputOffsets.data(), 0, inputOffsets.size());
+				for (u32 i : inputs)
+					writer.Write7BitEncoded(i);
+
+				CacheEntry entry;
+				entries.BuildInputs(entry, inputs);
+				entries.entries.push_back(entry);
+				entries.ValidateEntry(logger, entry, inputOffsets);
+			};
+
+		AddEntry(Set<u32>{ 1, 4, 6 });
+		AddEntry(Set<u32>{ 0, 4, 6 });
+		AddEntry(Set<u32>{ 2, 4, 6 });
+		AddEntry(Set<u32>{ 1, 4, 5 });
+		AddEntry(Set<u32>{ 1, 4, 7 });
+		AddEntry(Set<u32>{ 1, 3, 6 });
+		AddEntry(Set<u32>{ 1, 5, 6 });
+		AddEntry(Set<u32>{ 1, 4, 6, 7 });
+		AddEntry(Set<u32>{ 0, 1, 4, 6 });
+
+		return true;
+	}
+}
\ No newline at end of file
