[UBA]
* Fixed static analysis errors

diff --git a/Engine/Source/Programs/UnrealBuildAccelerator/Build-StaticAnalysis.bat b/Engine/Source/Programs/UnrealBuildAccelerator/Build-StaticAnalysis.bat
--- a/Engine/Source/Programs/UnrealBuildAccelerator/Build-StaticAnalysis.bat
+++ b/Engine/Source/Programs/UnrealBuildAccelerator/Build-StaticAnalysis.bat
@@ -21,6 +21,9 @@ call Engine/Build/BatchFiles/RunUBT.bat -NoUba -NoUbaLocal -NoSNDBS -NoXGE -Arch
 	-Target="UbaHost Win64 %config%" ^
 	-Target="UbaStorageProxy Win64 %config%" ^
 	-Target="UbaVisualizer Win64 %config%" ^
+	-Target="UbaCacheService Win64 %config%" ^
+	-Target="UbaObjTool Win64 %config%" ^
+	-Target="UbaTest Win64 %config%" ^
 	-StaticAnalyzer=Clang
 
 popd
diff --git a/Engine/Source/Programs/UnrealBuildAccelerator/Common/Private/UbaCacheClient.cpp b/Engine/Source/Programs/UnrealBuildAccelerator/Common/Private/UbaCacheClient.cpp
--- a/Engine/Source/Programs/UnrealBuildAccelerator/Common/Private/UbaCacheClient.cpp
+++ b/Engine/Source/Programs/UnrealBuildAccelerator/Common/Private/UbaCacheClient.cpp
@@ -153,7 +153,11 @@ namespace uba
 			else
 				break;
 			
-			TString temp2 = path.data;
+			if (path.count < 2)
+			{
+				m_logger.Info(TC("Got messed up path from caller to WriteToCache: %s (%s)"), path.data, info.description);
+				success = false;
+			}
 
 			// For .exe and .dll we sometimes get relative paths so we need to expand them to full
 			if (path[1] != ':' && (path.EndsWith(TC(".dll")) || path.EndsWith(TC(".exe"))))
@@ -179,7 +183,7 @@ namespace uba
 			}
 			else if (path[path.count-1] == ':')
 			{
-				m_logger.Info(TC("GOT UNKNOWN RELATIVE PATH: %s"), path.data);
+				m_logger.Info(TC("GOT UNKNOWN RELATIVE PATH: %s (%s)"), path.data, info.description);
 				success = false;
 				continue;
 			}
@@ -188,7 +192,7 @@ namespace uba
 			auto root = rootPaths.FindRoot(path);
 			if (!root)
 			{
-				m_logger.Info(TC("FILE WITHOUT ROOT: %s"), path.data);
+				m_logger.Info(TC("FILE WITHOUT ROOT: %s (%s)"), path.data, info.description);
 				success = false;
 				continue;
 			}
diff --git a/Engine/Source/Programs/UnrealBuildAccelerator/Common/Private/UbaStorageServer.cpp b/Engine/Source/Programs/UnrealBuildAccelerator/Common/Private/UbaStorageServer.cpp
--- a/Engine/Source/Programs/UnrealBuildAccelerator/Common/Private/UbaStorageServer.cpp
+++ b/Engine/Source/Programs/UnrealBuildAccelerator/Common/Private/UbaStorageServer.cpp
@@ -755,10 +755,11 @@ namespace uba
 				{
 					StringBuffer<> casKeyName;
 					GetCasFileName(casKeyName, casKey);
-					fileAccessor = new FileAccessor(m_logger, TStrdup(casKeyName.data));
+					
+					const tchar* filename = TStrdup(casKeyName.data);
+					fileAccessor = new FileAccessor(m_logger, filename);
 					if (!fileAccessor->CreateMemoryWrite(false, DefaultAttributes(), fileSize, m_tempPath.data))
 					{
-						const tchar* filename = fileAccessor->GetFileName();
 						delete fileAccessor;
 						free((void*)filename);
 
@@ -766,6 +767,11 @@ namespace uba
 						casEntry.verified = false;
 						return false;
 					}
+
+					#ifdef __clang_analyzer__ // Seems clang analyzer gets lost
+					free((void*)filename);
+					#endif
+
 					mappedView.memory = fileAccessor->GetData();
 				}
 				else
diff --git a/Engine/Source/Programs/UnrealBuildAccelerator/Common/Private/UbaStorageUtils.cpp b/Engine/Source/Programs/UnrealBuildAccelerator/Common/Private/UbaStorageUtils.cpp
--- a/Engine/Source/Programs/UnrealBuildAccelerator/Common/Private/UbaStorageUtils.cpp
+++ b/Engine/Source/Programs/UnrealBuildAccelerator/Common/Private/UbaStorageUtils.cpp
@@ -418,7 +418,6 @@ namespace uba
 		u32 responseSize = 0;
 		bool isCompressed = false;
 		bool sendEnd = false;
-		u64 left = ~u64(0);
 
 		u32 sizeOfFirstMessage = 0;
 
@@ -446,8 +445,6 @@ namespace uba
 			isCompressed = (flags >> 0) & 1;
 			sendEnd = (flags >> 1) & 1;
 
-			left = fileSize;
-
 			responseSize = u32(reader.GetLeft());
 			readBuffer = (u8*)reader.GetPositionData();
 			readPosition = readBuffer;
@@ -515,7 +512,6 @@ namespace uba
 
 				leftCompressed -= responseSize;
 			}
-			actualSize = sizeof(header) + fileSize;
 		}
 		else if (actualSize)
 		{
@@ -531,7 +527,7 @@ namespace uba
 				bool isFirstInBlock = true;
 				u32 compressedSize = ~u32(0);
 				u32 uncompressedSize = ~u32(0);
-				left = 0;
+				u32 left = 0;
 				u32 overflow = 0;
 				do
 				{
diff --git a/Engine/Source/Programs/UnrealBuildAccelerator/Common/Public/UbaWorkManager.h b/Engine/Source/Programs/UnrealBuildAccelerator/Common/Public/UbaWorkManager.h
--- a/Engine/Source/Programs/UnrealBuildAccelerator/Common/Public/UbaWorkManager.h
+++ b/Engine/Source/Programs/UnrealBuildAccelerator/Common/Public/UbaWorkManager.h
@@ -75,6 +75,8 @@ namespace uba
 	template<typename TContainer, typename TFunc>
 	void WorkManager::ParallelFor(u32 workCount, TContainer& container, const TFunc& func, const tchar* description)
 	{
+		#if !defined( __clang_analyzer__ ) // Static analyzer claims context can leak but it can only leak when process terminates (and then it doesn't matter)
+
 		struct Context
 		{
 			typename TContainer::iterator it;
@@ -130,5 +132,7 @@ namespace uba
 		AddWork(work, workCount, description);
 		work();
 		doneEvent.IsSet();
+
+		#endif
 	}
 }
diff --git a/Engine/Source/Programs/UnrealBuildAccelerator/Core/Public/UbaHash.h b/Engine/Source/Programs/UnrealBuildAccelerator/Core/Public/UbaHash.h
--- a/Engine/Source/Programs/UnrealBuildAccelerator/Core/Public/UbaHash.h
+++ b/Engine/Source/Programs/UnrealBuildAccelerator/Core/Public/UbaHash.h
@@ -119,8 +119,10 @@ namespace uba
 	{
 		UBA_ASSERT(key != CasKeyZero);
 		CasKey newKey = key;
+		#ifndef __clang_analyzer__
 		u8 flagField = ((u8*)&key)[19];
 		((u8*)&newKey)[19] = compressed ? (flagField | u8(1)) : (flagField & ~u8(1));
+		#endif
 		return newKey;
 	}
 
diff --git a/Engine/Source/Programs/UnrealBuildAccelerator/Detours/Private/Windows/UbaDetoursMainWin.cpp b/Engine/Source/Programs/UnrealBuildAccelerator/Detours/Private/Windows/UbaDetoursMainWin.cpp
--- a/Engine/Source/Programs/UnrealBuildAccelerator/Detours/Private/Windows/UbaDetoursMainWin.cpp
+++ b/Engine/Source/Programs/UnrealBuildAccelerator/Detours/Private/Windows/UbaDetoursMainWin.cpp
@@ -1,5 +1,11 @@
 // Copyright Epic Games, Inc. All Rights Reserved.
 
+#if defined( __clang_analyzer__ )
+#include <corecrt.h>
+#undef __DEFINE_CPP_OVERLOAD_SECURE_FUNC_SPLITPATH
+#define __DEFINE_CPP_OVERLOAD_SECURE_FUNC_SPLITPATH(_ReturnType, _FuncName, _DstType, _Src)
+#endif
+
 #include "UbaProtocol.h"
 #include "UbaProcessStats.h"
 #include "UbaDetoursPayload.h"
diff --git a/Engine/Source/Programs/UnrealBuildAccelerator/Detours/Private/Windows/UbaDetoursObjFilesPreloader.h b/Engine/Source/Programs/UnrealBuildAccelerator/Detours/Private/Windows/UbaDetoursObjFilesPreloader.h
--- a/Engine/Source/Programs/UnrealBuildAccelerator/Detours/Private/Windows/UbaDetoursObjFilesPreloader.h
+++ b/Engine/Source/Programs/UnrealBuildAccelerator/Detours/Private/Windows/UbaDetoursObjFilesPreloader.h
@@ -137,29 +137,6 @@ namespace uba
 		{
 			Oodle_SetUsageWarnings(Oodle_UsageWarnings_Disabled);
 
-			auto handleLine = [&](const StringView& line)
-				{
-					StringBuffer<> file;
-					file.Append(line);
-					if (!g_rules->ShouldDecompressFiles(file))
-						return;
-					StringBuffer<> fileFull;
-					FixPath(fileFull, file.data);
-
-					HANDLE fileHandle = CreateFileW(fileFull.data, GENERIC_READ, FILE_SHARE_READ, NULL, OPEN_EXISTING, FILE_ATTRIBUTE_NORMAL, NULL);
-					UBA_ASSERT(fileHandle != INVALID_HANDLE_VALUE);
-
-					StringKey fileNameKey = ToStringKey(fileFull.MakeLower());
-					auto insres = m_preloadedObjFiles.try_emplace(fileNameKey);
-					UBA_ASSERT(insres.second);
-					insres.first->second.handle = fileHandle;
-					insres.first->second.event.Create(true);
-
-					SCOPED_READ_LOCK(g_mappedFileTable.m_lookupLock, _);
-					auto findIt = g_mappedFileTable.m_lookup.find(fileNameKey);
-					insres.first->second.fileInfo = &findIt->second;
-				};
-
 			ParseArguments(cmdLine, [&](const tchar* arg, u32 argLen) { HandleLine(StringView(arg, argLen)); });
 
 			if (!m_preloadedObjFiles.empty())
diff --git a/Engine/Source/Programs/UnrealBuildAccelerator/Host/Private/UbaExports.cpp b/Engine/Source/Programs/UnrealBuildAccelerator/Host/Private/UbaExports.cpp
--- a/Engine/Source/Programs/UnrealBuildAccelerator/Host/Private/UbaExports.cpp
+++ b/Engine/Source/Programs/UnrealBuildAccelerator/Host/Private/UbaExports.cpp
@@ -517,7 +517,10 @@ extern "C"
 		bool ctorSuccess = false;
 		auto networkClient = new NetworkClientWithBackend(ctorSuccess, ncci, server.backend);
 		if (!ctorSuccess)
+		{
+			delete networkClient;
 			return nullptr;
+		}
 		CacheClientCreateInfo info{writer, storage, *networkClient, *session};
 		info.reportMissReason = reportMissReason;
 		return new CacheClient(info);
diff --git a/Engine/Source/Programs/UnrealBuildAccelerator/Visualizer/Private/UbaVisualizer.cpp b/Engine/Source/Programs/UnrealBuildAccelerator/Visualizer/Private/UbaVisualizer.cpp
--- a/Engine/Source/Programs/UnrealBuildAccelerator/Visualizer/Private/UbaVisualizer.cpp
+++ b/Engine/Source/Programs/UnrealBuildAccelerator/Visualizer/Private/UbaVisualizer.cpp
@@ -1328,14 +1328,6 @@ namespace uba
 			SetTextColor(hdc, m_textColor);
 		}
 
-		auto GetClientCursorPos = [this](POINT& p)
-			{
-				GetCursorPos(&p);
-				ScreenToClient(m_hwnd, &p);
-				p.x += 3;
-				p.y += 3;
-			};
-
 		if (m_processSelected)
 		{
 			TraceView::Process& process = *m_traceView.GetProcess(m_processSelectedLocation);
