[UBA]
* Changed so cache server has one file per bucket since we are planning to have lots of buckets

diff --git a/Engine/Source/Programs/UnrealBuildAccelerator/Common/Private/UbaCacheServer.cpp b/Engine/Source/Programs/UnrealBuildAccelerator/Common/Private/UbaCacheServer.cpp
--- a/Engine/Source/Programs/UnrealBuildAccelerator/Common/Private/UbaCacheServer.cpp
+++ b/Engine/Source/Programs/UnrealBuildAccelerator/Common/Private/UbaCacheServer.cpp
@@ -4,6 +4,7 @@
 #include "UbaCacheEntry.h"
 #include "UbaCompactTables.h"
 #include "UbaBinaryReaderWriter.h"
+#include "UbaDirectoryIterator.h"
 #include "UbaFileAccessor.h"
 #include "UbaNetworkServer.h"
 #include "UbaStorageServer.h"
@@ -11,7 +12,7 @@
 
 namespace uba
 {
-	static constexpr u32 CacheFileVersion = 4;
+	static constexpr u32 CacheFileVersion = 5;
 	static constexpr u32 CacheFileCompatibilityVersion = 3;
 
 	bool IsCaseInsensitive(u64 id) { return (id & (1ull << 32)) == 0; }
@@ -45,6 +46,7 @@ namespace uba
 		Atomic<u64> totalEntryCount;
 		Atomic<u64> totalEntrySize;
 		Atomic<bool> hasDeletedEntries;
+		Atomic<bool> needsSave;
 	};
 
 	const tchar* ToString(CacheMessageType type)
@@ -90,6 +92,13 @@ namespace uba
 	{
 	}
 
+	struct CacheServer::LoadStats
+	{
+		Atomic<u32> totalPathTableSize;
+		Atomic<u32> totalCasKeyTableSize;
+		Atomic<u64> totalCacheEntryCount;
+	};
+
 	bool CacheServer::Load()
 	{
 		u64 startTime = GetTime();
@@ -102,6 +111,7 @@ namespace uba
 		{
 			m_logger.Detail(TC("No database found. Starting a new one at %s"), fileName.data);
 			m_creationTime = GetSystemTimeAsFileTime();
+			m_dbfileDirty = true;
 			return true;
 		}
 		BinaryReader reader(file.GetData(), 0, file.GetSize());
@@ -117,49 +127,89 @@ namespace uba
 		else
 			m_creationTime = reader.ReadU64();
 
-		u32 totalPathTableSize = 0;
-		u32 totalCasKeyTableSize = 0;
-		u64 totalCacheEntryCount = 0;
+		if (databaseVersion != CacheFileVersion)
+			m_dbfileDirty = true;
 
-		u32 bucketCount = reader.ReadU32();
-		while (bucketCount--)
-		{
-			u64 id = reader.ReadU64();
-			Bucket& bucket = m_buckets.try_emplace(id, id).first->second;
-
-			u32 pathTableSize = reader.ReadU32();
-			if (pathTableSize)
-			{
-				BinaryReader pathTableReader(reader.GetPositionData(), 0, pathTableSize);
-				bucket.m_pathTable.ReadMem(pathTableReader, true);
-				reader.Skip(pathTableSize);
-			}
-			totalPathTableSize += pathTableSize;
+		LoadStats stats;
 
-			u32 casKeyTableSize = reader.ReadU32();
-			if (casKeyTableSize)
+		if (databaseVersion == 4)
+		{
+			u32 bucketCount = reader.ReadU32();
+			while (bucketCount--)
 			{
-				BinaryReader casKeyTableReader(reader.GetPositionData(), 0, casKeyTableSize);
-				bucket.m_casKeyTable.ReadMem(casKeyTableReader, true);
-				reader.Skip(casKeyTableSize);
+				u64 id = reader.ReadU64();
+				Bucket& bucket = m_buckets.try_emplace(id, id).first->second;
+				LoadBucket(bucket, reader, databaseVersion, stats);
 			}
-			totalCasKeyTableSize += casKeyTableSize;
+		}
+		else
+		{
+			StringBuffer<MaxPath> bucketsDir(m_rootDir);
+			bucketsDir.EnsureEndsWithSlash().Append(TC("buckets"));
+			TraverseDir(m_logger, bucketsDir.data, [&](const DirectoryEntry& e)
+				{
+					StringBuffer<128> keyName;
+					keyName.Append(e.name, e.nameLen);
+					u64 key;
+					if (!keyName.Parse(key))
+						return;
+					m_buckets.try_emplace(key, key);
+				});
 
-			u32 entryLookupCount = reader.ReadU32();
-			bucket.m_cacheEntryLookup.reserve(entryLookupCount);
+			m_server.ParallelFor(GetBucketWorkerCount(), m_buckets, [&](auto& it)
+				{
+					u64 key = it->first;
 
-			while (entryLookupCount--)
-			{
-				auto insres = bucket.m_cacheEntryLookup.try_emplace(reader.ReadCasKey());
-				UBA_ASSERT(insres.second);
-				auto& cacheEntries = insres.first->second;
-				cacheEntries.Read(m_logger, reader, databaseVersion);
-				totalCacheEntryCount += cacheEntries.entries.size();
-			}
+					StringBuffer<MaxPath> bucketFilename(bucketsDir);
+					bucketFilename.EnsureEndsWithSlash().AppendValue(key);
+					FileAccessor bucketFile(m_logger, bucketFilename.data);
+					if (!bucketFile.OpenMemoryRead(0, false))
+					{
+						m_logger.Detail(TC("Failed to open bucket file %s"), bucketFilename.data);
+						return;
+					}
+					BinaryReader reader(bucketFile.GetData(), 0, bucketFile.GetSize());
+					u32 bucketVersion = reader.ReadU32();
+					LoadBucket(it->second, reader, bucketVersion, stats);
+				});
 		}
 
 		u64 duration = GetTime() - startTime;
-		m_logger.Detail(TC("Database (v%u) loaded from %s in %s (%llu bucket(s) containing %s paths, %s keys, %llu cache entries)"), databaseVersion, fileName.data, TimeToText(duration).str, m_buckets.size(), BytesToText(totalPathTableSize).str, BytesToText(totalCasKeyTableSize).str, totalCacheEntryCount);
+		m_logger.Detail(TC("Database (v%u) loaded from %s in %s (%llu bucket(s) containing %s paths, %s keys, %llu cache entries)"), databaseVersion, fileName.data, TimeToText(duration).str, m_buckets.size(), BytesToText(stats.totalPathTableSize).str, BytesToText(stats.totalCasKeyTableSize).str, stats.totalCacheEntryCount.load());
+		return true;
+	}
+
+	bool CacheServer::LoadBucket(Bucket& bucket, BinaryReader& reader, u32 databaseVersion, LoadStats& outStats)
+	{
+		u32 pathTableSize = reader.ReadU32();
+		if (pathTableSize)
+		{
+			BinaryReader pathTableReader(reader.GetPositionData(), 0, pathTableSize);
+			bucket.m_pathTable.ReadMem(pathTableReader, true);
+			reader.Skip(pathTableSize);
+		}
+		outStats.totalPathTableSize += pathTableSize;
+
+		u32 casKeyTableSize = reader.ReadU32();
+		if (casKeyTableSize)
+		{
+			BinaryReader casKeyTableReader(reader.GetPositionData(), 0, casKeyTableSize);
+			bucket.m_casKeyTable.ReadMem(casKeyTableReader, true);
+			reader.Skip(casKeyTableSize);
+		}
+		outStats.totalCasKeyTableSize += casKeyTableSize;
+
+		u32 entryLookupCount = reader.ReadU32();
+		bucket.m_cacheEntryLookup.reserve(entryLookupCount);
+
+		while (entryLookupCount--)
+		{
+			auto insres = bucket.m_cacheEntryLookup.try_emplace(reader.ReadCasKey());
+			UBA_ASSERT(insres.second);
+			auto& cacheEntries = insres.first->second;
+			cacheEntries.Read(m_logger, reader, databaseVersion);
+			outStats.totalCacheEntryCount += cacheEntries.entries.size();
+		}
 		return true;
 	}
 
@@ -172,98 +222,164 @@ namespace uba
 		return SaveNoLock();
 	}
 
-	bool CacheServer::SaveNoLock()
+	struct FileWriter
 	{
-		StringBuffer<MaxPath> fileName(m_rootDir);
-		fileName.EnsureEndsWithSlash().Append(TC("cachedb"));
-
-		StringBuffer<MaxPath> tempFileName;
-		tempFileName.Append(fileName).Append(TC(".tmp"));
-
-		FileAccessor file(m_logger, tempFileName.data);
-		if (!file.CreateWrite())
-			return false;
+		static constexpr u64 TempBufferSize = 1024*1024;
 
+		FileWriter(Logger& l, const tchar* fn)
+		:	logger(l)
+		,	fileName(fn)
+		,	tempFileName(StringBuffer<MaxPath>(fn).Append(TC(".tmp")).data)
+		,	file(logger, tempFileName.c_str())
+		{
+			tempBuffer = (u8*)malloc(TempBufferSize);
+		}
 
-		constexpr u64 tempBufferSize = 1024*1024;
-		u8* tempBuffer = (u8*)malloc(tempBufferSize);
-		auto g = MakeGuard([tempBuffer](){ free(tempBuffer); });
-		u64 tempBufferPos = 0;
+		~FileWriter()
+		{
+			free(tempBuffer);
+		}
 
-		bool success = true;
-		auto WriteBytes = [&](const void* data, u64 size)
+		void WriteBytes(const void* data, u64 size)
+		{
+			u8* readPos = (u8*)data;
+			u64 left = size;
+			while (left)
 			{
-				u8* readPos = (u8*)data;
-				u64 left = size;
-				while (left)
+				if (tempBufferPos != TempBufferSize)
 				{
-					if (tempBufferPos != tempBufferSize)
-					{
-						u64 toWrite = Min(tempBufferSize - tempBufferPos, left);
-						memcpy(tempBuffer+tempBufferPos, readPos, toWrite);
-						tempBufferPos += toWrite;
-						left -= toWrite;
-						readPos += toWrite;
-					}
-					else
-					{
-						success &= file.Write(tempBuffer, tempBufferPos);
-						tempBufferPos = 0;
-					}
+					u64 toWrite = Min(TempBufferSize - tempBufferPos, left);
+					memcpy(tempBuffer+tempBufferPos, readPos, toWrite);
+					tempBufferPos += toWrite;
+					left -= toWrite;
+					readPos += toWrite;
 				}
-			};
+				else
+				{
+					success &= file.Write(tempBuffer, tempBufferPos);
+					tempBufferPos = 0;
+				}
+			}
+		}
+
+		template<typename T>
+		void Write(const T& v)
+		{
+			WriteBytes(&v, sizeof(v));
+		}
+
+		bool Create() { return file.CreateWrite(); }
+
+		bool Close()
+		{
+			success &= file.Write(tempBuffer, tempBufferPos);
+
+			if (!success)
+				return false;
+
+			if (!file.Close())
+				return false;
+
+			if (!MoveFileExW(tempFileName.data(), fileName.data(), MOVEFILE_REPLACE_EXISTING))
+				return logger.Error(TC("Can't move file from %s to %s (%s)"), tempFileName.data(), fileName.data(), LastErrorToText().data);
+
+			return true;
+		}
+
+		Logger& logger;
+		bool success = true;
+		u8* tempBuffer = nullptr;
+		u64 tempBufferPos = 0;
+		TString fileName;
+		TString tempFileName;
+		FileAccessor file;
+	};
+
+	bool CacheServer::SaveBucket(u64 bucketId, Bucket& bucket)
+	{
+		StringBuffer<MaxPath> bucketsDir(m_rootDir);
+		bucketsDir.EnsureEndsWithSlash().Append(TC("buckets"));
+		if (!m_storage.CreateDirectory(bucketsDir.data))
+			return false;
+		bucketsDir.EnsureEndsWithSlash();
+		StringBuffer<MaxPath> bucketsFile(bucketsDir);
+		bucketsFile.AppendValue(bucketId);
 
-		auto Write = [&](auto v) { WriteBytes(&v, sizeof(v)); };
+		FileWriter file(m_logger, bucketsFile.data);
+		
+		if (!file.Create())
+			return false;
 
-		Write(CacheFileVersion);
+		file.Write(CacheFileVersion);
 
-		Write(m_creationTime);
+		u32 pathTableSize = bucket.m_pathTable.GetSize();
+		file.Write(pathTableSize);
+		file.WriteBytes(bucket.m_pathTable.GetMemory(), pathTableSize);
 
-		Write(u32(m_buckets.size()));
+		u32 casKeyTableSize = bucket.m_casKeyTable.GetSize();
+		file.Write(casKeyTableSize);
+		file.WriteBytes(bucket.m_casKeyTable.GetMemory(), casKeyTableSize);
+
+		u32 entryLookupCount = u32(bucket.m_cacheEntryLookup.size());
+		file.Write(entryLookupCount);
 
 		Vector<u8> temp;
 
-		for (auto& kv : m_buckets)
+		for (auto& kv2 : bucket.m_cacheEntryLookup)
 		{
-			Bucket& bucket = kv.second;
+			file.Write(kv2.first);
 
-			Write(kv.first);
+			temp.resize(kv2.second.GetTotalSize(true));
+			BinaryWriter writer(temp.data(), 0, temp.size());
+			kv2.second.Write(writer, CacheNetworkVersion, true);
+			UBA_ASSERT(writer.GetPosition() == temp.size());
+			file.WriteBytes(temp.data(), temp.size());
+		}
 
-			u32 pathTableSize = bucket.m_pathTable.GetSize();
-			Write(pathTableSize);
-			WriteBytes(bucket.m_pathTable.GetMemory(), pathTableSize);
+		return file.Close();
+	}
 
-			u32 casKeyTableSize = bucket.m_casKeyTable.GetSize();
-			Write(casKeyTableSize);
-			WriteBytes(bucket.m_casKeyTable.GetMemory(), casKeyTableSize);
+	bool CacheServer::SaveNoLock()
+	{
+		if (m_dbfileDirty)
+		{
+			StringBuffer<MaxPath> fileName(m_rootDir);
+			fileName.EnsureEndsWithSlash().Append(TC("cachedb"));
 
-			u32 entryLookupCount = u32(bucket.m_cacheEntryLookup.size());
-			Write(entryLookupCount);
+			FileWriter file(m_logger, fileName.data);
+		
+			if (!file.Create())
+				return false;
 
-			for (auto& kv2 : bucket.m_cacheEntryLookup)
-			{
-				Write(kv2.first);
+			file.Write(CacheFileVersion);
 
-				temp.resize(kv2.second.GetTotalSize(true));
-				BinaryWriter writer(temp.data(), 0, temp.size());
-				kv2.second.Write(writer, CacheNetworkVersion, true);
-				UBA_ASSERT(writer.GetPosition() == temp.size());
-				WriteBytes(temp.data(), temp.size());
-			}
-		}
+			file.Write(m_creationTime);
 
-		success &= file.Write(tempBuffer, tempBufferPos);
+			if (!file.Close())
+				return false;
+			m_dbfileDirty = false;
+		}
 
-		if (!success)
+		StringBuffer<MaxPath> bucketsDir(m_rootDir);
+		bucketsDir.EnsureEndsWithSlash().Append(TC("buckets"));
+		if (!m_storage.CreateDirectory(bucketsDir.data))
 			return false;
+		bucketsDir.EnsureEndsWithSlash();
 
-		if (!file.Close())
-			return false;
+		Atomic<bool> success = true;
 
-		if (!MoveFileExW(tempFileName.data, fileName.data, MOVEFILE_REPLACE_EXISTING))
-			return m_logger.Error(TC("Can't move file from %s to %s (%s)"), tempFileName.data, fileName.data, LastErrorToText().data);
+		m_server.ParallelFor(GetBucketWorkerCount(), m_buckets, [&, temp = Vector<u8>()](auto& it) mutable
+			{
+				Bucket& bucket = it->second;
+				if (!bucket.needsSave)
+					return;
+				if (SaveBucket(it->first, bucket))
+					bucket.needsSave = false;
+				else
+					success = false;
+			});
 
-		return true;
+		return success;
 	}
 
 	bool CacheServer::RunMaintenance(bool force, const Function<bool()>& shouldExit)
@@ -666,6 +782,8 @@ namespace uba
 				m_logger.Detail(TC("    Bucket %u Updated cache entries with new tables (%s)"), bucketIndex, TimeToText(GetTime() - updateEntriesStart).str);
 			}
 
+			bucket.needsSave = true;
+
 			m_logger.Info(TC("    Bucket %u Done (%s). CacheEntries: %llu (%s) PathTable: %s CasTable: %s"), bucketIndex, TimeToText(GetTime() - bucketStartTime).str, bucket.totalEntryCount.load(), BytesToText(bucket.totalEntrySize.load()).str, BytesToText(bucket.m_pathTable.GetSize()).str, BytesToText(bucket.m_casKeyTable.GetSize()).str);
 		});
 
@@ -726,6 +844,13 @@ namespace uba
 		
 	}
 
+	u32 CacheServer::GetBucketWorkerCount()
+	{
+		u32 workerCount = m_server.GetWorkerCount();
+		u32 workerCountToUse = workerCount > 0 ? workerCount - 1 : 0;
+		return Min(workerCountToUse, u32(m_buckets.size()));
+	}
+
 	bool CacheServer::HandleMessage(const ConnectionInfo& connectionInfo, u8 messageType, BinaryReader& reader, BinaryWriter& writer)
 	{
 		if (messageType != CacheMessageType_Connect && m_isRunningMaintenance)
@@ -1019,6 +1144,7 @@ namespace uba
 		}
 
 		//m_logger.Info(TC("Added new cache entry (%u inputs and %u outputs)"), u32(inputs.size()), outputCount);
+		bucket.needsSave = true;
 
 		++m_addsSinceMaintenance;
 
diff --git a/Engine/Source/Programs/UnrealBuildAccelerator/Common/Public/UbaCacheServer.h b/Engine/Source/Programs/UnrealBuildAccelerator/Common/Public/UbaCacheServer.h
--- a/Engine/Source/Programs/UnrealBuildAccelerator/Common/Public/UbaCacheServer.h
+++ b/Engine/Source/Programs/UnrealBuildAccelerator/Common/Public/UbaCacheServer.h
@@ -29,15 +29,20 @@ namespace uba
 		bool ShouldShutdown();
 
 	private:
+		struct Bucket;
+		struct LoadStats;
+
+		bool LoadBucket(Bucket& bucket, BinaryReader& reader, u32 databaseVersion, LoadStats& outStats);
+		bool SaveBucket(u64 bucketId, Bucket& bucket);
 		bool SaveNoLock();
 		void OnDisconnected(u32 clientId);
 
 		struct Connection;
 		struct ConnectionBucket;
-		struct Bucket;
 
 		ConnectionBucket& GetConnectionBucket(const ConnectionInfo& connectionInfo, BinaryReader& reader);
 		Bucket& GetBucket(BinaryReader& reader);
+		u32 GetBucketWorkerCount();
 
 		bool HandleMessage(const ConnectionInfo& connectionInfo, u8 messageType, BinaryReader& reader, BinaryWriter& writer);
 		bool HandleStoreEntry(ConnectionBucket& bucket, BinaryReader& reader, BinaryWriter& writer);
@@ -67,6 +72,7 @@ namespace uba
 		u64 m_creationTime = 0;
 		u64 m_startTime = 0;
 		u64 m_longestMaintenance = 0;
+		bool m_dbfileDirty = false;
 
 		bool m_checkInputsForDeletedCas = true;
 
