[UBA]
* Changed so when fetching cache entries we only calculate cas for localized files once

diff --git a/Engine/Source/Programs/UnrealBuildAccelerator/Common/Private/UbaCacheClient.cpp b/Engine/Source/Programs/UnrealBuildAccelerator/Common/Private/UbaCacheClient.cpp
--- a/Engine/Source/Programs/UnrealBuildAccelerator/Common/Private/UbaCacheClient.cpp
+++ b/Engine/Source/Programs/UnrealBuildAccelerator/Common/Private/UbaCacheClient.cpp
@@ -360,6 +360,8 @@ namespace uba
 		struct MissInfo { TString path; u32 entryIndex; CasKey cache; CasKey local; };
 		Vector<MissInfo> misses;
 
+		UnorderedMap<StringKey, CasKey> normalizedCasKeys;
+
 		u32 entryIndex = 0;
 		for (; entryIndex!=entryCount; ++entryIndex)
 		{
@@ -386,9 +388,17 @@ namespace uba
 
 							if (IsNormalized(cacheCasKey)) // Need to normalize caskey for these files since they contain absolute paths
 							{
-								localCasKey = rootPaths.NormalizeAndHashFile(m_logger, path.data);
-								if (localCasKey != CasKeyZero)
-									localCasKey = AsCompressed(localCasKey, true);
+								auto insres2 = normalizedCasKeys.try_emplace(ToStringKeyNoCheck(path.data, path.count));
+								if (insres2.second)
+								{
+									localCasKey = rootPaths.NormalizeAndHashFile(m_logger, path.data);
+									if (localCasKey != CasKeyZero)
+										localCasKey = AsCompressed(localCasKey, true);
+									insres2.first->second = localCasKey;
+								}
+								else
+									localCasKey = insres2.first->second;
+
 							}
 							else
 							{
diff --git a/Engine/Source/Programs/UnrealBuildAccelerator/Common/Private/UbaCacheServer.cpp b/Engine/Source/Programs/UnrealBuildAccelerator/Common/Private/UbaCacheServer.cpp
--- a/Engine/Source/Programs/UnrealBuildAccelerator/Common/Private/UbaCacheServer.cpp
+++ b/Engine/Source/Programs/UnrealBuildAccelerator/Common/Private/UbaCacheServer.cpp
@@ -47,6 +47,8 @@ namespace uba
 		Atomic<u64> totalEntrySize;
 		Atomic<bool> hasDeletedEntries;
 		Atomic<bool> needsSave;
+
+		u32 index = ~0u;
 	};
 
 	const tchar* ToString(CacheMessageType type)
@@ -137,8 +139,7 @@ namespace uba
 			u32 bucketCount = reader.ReadU32();
 			while (bucketCount--)
 			{
-				u64 id = reader.ReadU64();
-				Bucket& bucket = m_buckets.try_emplace(id, id).first->second;
+				Bucket& bucket = GetBucket(reader.ReadU64());
 				LoadBucket(bucket, reader, databaseVersion, stats);
 			}
 		}
@@ -150,10 +151,10 @@ namespace uba
 				{
 					StringBuffer<128> keyName;
 					keyName.Append(e.name, e.nameLen);
-					u64 key;
-					if (!keyName.Parse(key))
+					u64 id;
+					if (!keyName.Parse(id))
 						return;
-					m_buckets.try_emplace(key, key);
+					GetBucket(id);
 				});
 
 			m_server.ParallelFor(GetBucketWorkerCount(), m_buckets, [&](auto& it)
@@ -373,10 +374,12 @@ namespace uba
 				Bucket& bucket = it->second;
 				if (!bucket.needsSave)
 					return;
+				u64 saveStart = GetTime();
 				if (SaveBucket(it->first, bucket))
 					bucket.needsSave = false;
 				else
 					success = false;
+				m_logger.Detail(TC("    Bucket %u saved (%s)"), bucket.index, TimeToText(GetTime() - saveStart).str);
 			});
 
 		return success;
@@ -647,17 +650,15 @@ namespace uba
 		if (shouldExit())
 			return true;
 
-		Atomic<u32> bucketCounter;
 		m_server.ParallelFor(workerCountToUseForBuckets, m_buckets, [&](auto& it)
 		{
 			u64 bucketStartTime = GetTime();
 
 			Bucket& bucket = it->second;
-			u32 bucketIndex = bucketCounter++;
 
 			if (!bucket.hasDeletedEntries && !forceAllSteps)
 			{
-				m_logger.Detail(TC("    Bucket %u skipped updating. No entries deleted"), bucketIndex);
+				m_logger.Detail(TC("    Bucket %u skipped updating. No entries deleted"), bucket.index);
 				return;
 			}
 			bucket.hasDeletedEntries = false;
@@ -686,7 +687,7 @@ namespace uba
 					collectUsedCasKeyOffsets(entry.outputCasKeyOffsets);
 				}
 			}
-			m_logger.Detail(TC("    Bucket %u Collected %llu used caskeys. (%s)"), bucketIndex, usedCasKeyOffsets.size(), TimeToText(GetTime() - collectUsedCasKeysStart).str);
+			m_logger.Detail(TC("    Bucket %u Collected %llu used caskeys. (%s)"), bucket.index, usedCasKeyOffsets.size(), TimeToText(GetTime() - collectUsedCasKeysStart).str);
 
 			u64 recreatePathTableStart = GetTime();
 
@@ -725,7 +726,7 @@ namespace uba
 				}
 				bucket.m_pathTable.Swap(newPathTable);
 			}
-			m_logger.Detail(TC("    Bucket %u Recreated path table. %s -> %s (%s)"), bucketIndex, BytesToText(oldSize).str, BytesToText(bucket.m_pathTable.GetSize()).str, TimeToText(GetTime() - recreatePathTableStart).str);
+			m_logger.Detail(TC("    Bucket %u Recreated path table. %s -> %s (%s)"), bucket.index, BytesToText(oldSize).str, BytesToText(bucket.m_pathTable.GetSize()).str, TimeToText(GetTime() - recreatePathTableStart).str);
 
 
 			// Build new caskey table based on used offsets
@@ -750,7 +751,7 @@ namespace uba
 				}
 				bucket.m_casKeyTable.Swap(newCasKeyTable);
 			}
-			m_logger.Detail(TC("    Bucket %u Recreated caskey table. %s -> %s (%s)"), bucketIndex, BytesToText(oldSize).str, BytesToText(bucket.m_casKeyTable.GetSize()).str, TimeToText(GetTime() - recreateCasKeyTableStart).str);
+			m_logger.Detail(TC("    Bucket %u Recreated caskey table. %s -> %s (%s)"), bucket.index, BytesToText(oldSize).str, BytesToText(bucket.m_casKeyTable.GetSize()).str, TimeToText(GetTime() - recreateCasKeyTableStart).str);
 
 
 			if (!oldToNewCasKeyOffset.empty())
@@ -779,12 +780,12 @@ namespace uba
 				#endif
 
 
-				m_logger.Detail(TC("    Bucket %u Updated cache entries with new tables (%s)"), bucketIndex, TimeToText(GetTime() - updateEntriesStart).str);
+				m_logger.Detail(TC("    Bucket %u Updated cache entries with new tables (%s)"), bucket.index, TimeToText(GetTime() - updateEntriesStart).str);
 			}
 
 			bucket.needsSave = true;
 
-			m_logger.Info(TC("    Bucket %u Done (%s). CacheEntries: %llu (%s) PathTable: %s CasTable: %s"), bucketIndex, TimeToText(GetTime() - bucketStartTime).str, bucket.totalEntryCount.load(), BytesToText(bucket.totalEntrySize.load()).str, BytesToText(bucket.m_pathTable.GetSize()).str, BytesToText(bucket.m_casKeyTable.GetSize()).str);
+			m_logger.Info(TC("    Bucket %u Done (%s). CacheEntries: %llu (%s) PathTable: %s CasTable: %s"), bucket.index, TimeToText(GetTime() - bucketStartTime).str, bucket.totalEntryCount.load(), BytesToText(bucket.totalEntrySize.load()).str, BytesToText(bucket.m_pathTable.GetSize()).str, BytesToText(bucket.m_casKeyTable.GetSize()).str);
 		});
 
 		// Need to make sure all cas entries are dropped before saving cas table
@@ -797,9 +798,10 @@ namespace uba
 		if (entriesAdded || deletedCasCount || deleteEntryCount || forceAllSteps)
 		{
 			u64 saveStart = GetTime();
+			m_logger.Detail(TC("  Saving to disk"));
 			m_storage.SaveCasTable(false, false);
 			SaveNoLock();
-			m_logger.Detail(TC("  Saved to disk (%s)"), TimeToText(GetTime() - saveStart).str);
+			m_logger.Detail(TC("  Save Done (%s)"), TimeToText(GetTime() - saveStart).str);
 		}
 
 		u64 oldestTime = oldest ? GetFileTimeAsTime(now - (m_creationTime + oldest)) : 0;
@@ -839,11 +841,19 @@ namespace uba
 	CacheServer::Bucket& CacheServer::GetBucket(BinaryReader& reader)
 	{
 		u64 id = reader.Read7BitEncoded();
-		SCOPED_WRITE_LOCK(m_bucketsLock, bucketsLock);
-		return m_buckets.try_emplace(id, id).first->second;
+		return GetBucket(id);
 		
 	}
 
+	CacheServer::Bucket& CacheServer::GetBucket(u64 id)
+	{
+		SCOPED_WRITE_LOCK(m_bucketsLock, bucketsLock);
+		auto insres = m_buckets.try_emplace(id, id);
+		if (insres.second)
+			insres.first->second.index = u32(m_buckets.size() - 1);
+		return insres.first->second;
+	}
+
 	u32 CacheServer::GetBucketWorkerCount()
 	{
 		u32 workerCount = m_server.GetWorkerCount();
@@ -902,9 +912,7 @@ namespace uba
 			if (findIt != connectionBucket.cacheEntryLookup.end())
 			{
 				u64 id = connectionBucket.id;
-				SCOPED_WRITE_LOCK(m_bucketsLock, bucketsLock);
-				Bucket& bucket = m_buckets.try_emplace(id, id).first->second;
-				bucketsLock.Leave();
+				Bucket& bucket = GetBucket(id);
 
 				SCOPED_WRITE_LOCK(bucket.m_cacheEntryLookupLock, lock3);
 				auto insres = bucket.m_cacheEntryLookup.try_emplace(cmdKey);
@@ -961,9 +969,7 @@ namespace uba
 
 		u64 outputStartOffset = reader.GetPosition();
 		u64 id = connectionBucket.id;
-		SCOPED_WRITE_LOCK(m_bucketsLock, bucketsLock);
-		Bucket& bucket = m_buckets.try_emplace(id, id).first->second;
-		bucketsLock.Leave();
+		Bucket& bucket = GetBucket(id);
 
 		while (reader.GetLeft())
 		{
diff --git a/Engine/Source/Programs/UnrealBuildAccelerator/Common/Public/UbaCacheServer.h b/Engine/Source/Programs/UnrealBuildAccelerator/Common/Public/UbaCacheServer.h
--- a/Engine/Source/Programs/UnrealBuildAccelerator/Common/Public/UbaCacheServer.h
+++ b/Engine/Source/Programs/UnrealBuildAccelerator/Common/Public/UbaCacheServer.h
@@ -42,6 +42,7 @@ namespace uba
 
 		ConnectionBucket& GetConnectionBucket(const ConnectionInfo& connectionInfo, BinaryReader& reader);
 		Bucket& GetBucket(BinaryReader& reader);
+		Bucket& GetBucket(u64 id);
 		u32 GetBucketWorkerCount();
 
 		bool HandleMessage(const ConnectionInfo& connectionInfo, u8 messageType, BinaryReader& reader, BinaryWriter& writer);
