[UBA]
* Made sure buckets are saved on version change

diff --git a/Engine/Source/Programs/UnrealBuildAccelerator/Common/Private/UbaCacheServer.cpp b/Engine/Source/Programs/UnrealBuildAccelerator/Common/Private/UbaCacheServer.cpp
--- a/Engine/Source/Programs/UnrealBuildAccelerator/Common/Private/UbaCacheServer.cpp
+++ b/Engine/Source/Programs/UnrealBuildAccelerator/Common/Private/UbaCacheServer.cpp
@@ -182,6 +182,9 @@ namespace uba
 
 	bool CacheServer::LoadBucket(Bucket& bucket, BinaryReader& reader, u32 databaseVersion, LoadStats& outStats)
 	{
+		if (databaseVersion != CacheFileVersion)
+			bucket.needsSave = true;
+
 		u32 pathTableSize = reader.ReadU32();
 		if (pathTableSize)
 		{
