[UBA]
* Fixed bug in argument parser when reading rsp file that does not have a string termination character in the end

diff --git a/Engine/Source/Programs/UnrealBuildAccelerator/Core/Public/UbaProcessUtils.h b/Engine/Source/Programs/UnrealBuildAccelerator/Core/Public/UbaProcessUtils.h
--- a/Engine/Source/Programs/UnrealBuildAccelerator/Core/Public/UbaProcessUtils.h
+++ b/Engine/Source/Programs/UnrealBuildAccelerator/Core/Public/UbaProcessUtils.h
@@ -10,16 +10,23 @@ namespace uba
 	template<typename CharType, typename Func>
 	bool ParseArguments(const CharType* arguments, u64 argumentsLen, const Func& argumentFunc)
 	{
+		if (argumentsLen == 0)
+			return true;
+
 		const CharType* argStart = arguments;
 		bool isInArg = false;
 		bool isInQuotes = false;
 		bool isEnd = *arguments == 0;
+		CharType currentChar = 0;
 		CharType lastChar = 0;
 		bool isBackslashOwned = false;
-		for (const CharType* it = arguments; !isEnd; lastChar = *it, ++it)
+		for (const CharType* it = arguments; !isEnd; lastChar = currentChar, ++it)
 		{
-			isEnd = *it == 0;
-			if (*it == ' ' || *it == '\t' || *it == '\n' || isEnd)
+			bool pastEnd = u64(it - arguments) == argumentsLen;
+			if (!pastEnd)
+				currentChar = *it;
+			isEnd = pastEnd || currentChar == 0;
+			if (isEnd || currentChar == ' ' || currentChar == '\t' || currentChar == '\n')
 			{
 				if (isInQuotes || !isInArg)
 					continue;
@@ -27,7 +34,7 @@ namespace uba
 				const CharType* argIt = argStart;
 				const CharType* argEnd = it;
 
-				if (*argEnd == '\n' && argStart != argEnd && *(argEnd-1) == '\r')
+				if ((pastEnd || *argEnd == '\n') && argStart != argEnd && *(argEnd-1) == '\r')
 					--argEnd;
 
 				CharType arg[16*1024];
@@ -69,7 +76,7 @@ namespace uba
 			{
 				isInArg = true;
 				argStart = it;
-				if (*it == '\"')
+				if (currentChar == '\"')
 					isInQuotes = true;
 				continue;
 			}
