[UBA]
* Store cas file lookup after first call and then reuse it since we know no cas files go away while we're running

diff --git a/Engine/Source/Programs/UnrealBuildAccelerator/Common/Private/UbaCacheServer.cpp b/Engine/Source/Programs/UnrealBuildAccelerator/Common/Private/UbaCacheServer.cpp
--- a/Engine/Source/Programs/UnrealBuildAccelerator/Common/Private/UbaCacheServer.cpp
+++ b/Engine/Source/Programs/UnrealBuildAccelerator/Common/Private/UbaCacheServer.cpp
@@ -434,9 +434,8 @@ namespace uba
 		m_storage.HandleOverflow(&deletedCasFiles);
 		u64 deletedCasCount = deletedCasFiles.size();
 
-		struct CasFileInfo { u64 size; u64 refCount; };
-		UnorderedMap<CasKey, CasFileInfo> existingCas;
 		u64 totalCasSize = 0;
+		if (m_existingCas.empty())
 		{
 			ReaderWriterLock existingCasLock;
 			u64 traverseStartTime = GetTime();
@@ -444,11 +443,19 @@ namespace uba
 				{
 					SCOPED_WRITE_LOCK(existingCasLock, lock);
 					totalCasSize += size;
-					existingCas.try_emplace(casKey, CasFileInfo{size, 0ull});
+					m_existingCas.try_emplace(casKey, CasFileInfo{size, 0ull});
 				}, true);
-			m_logger.Detail(TC("  Found %llu cas files (%s)"), existingCas.size(), TimeToText(GetTime() - traverseStartTime).str);
+			m_logger.Detail(TC("  Found %llu cas files (%s)"), m_existingCas.size(), TimeToText(GetTime() - traverseStartTime).str);
 		}
-		u64 totalCasCount = existingCas.size() + deletedCasCount;
+		else
+		{
+			for (const CasKey& deletedCasFile : deletedCasFiles)
+				m_existingCas.erase(deletedCasFile);
+			for (auto& kv : m_existingCas)
+				totalCasSize += kv.second.size;
+		}
+
+		u64 totalCasCount = m_existingCas.size() + deletedCasCount;
 
 		if (shouldExit())
 			return true;
@@ -569,8 +576,8 @@ namespace uba
 								CasKey casKey;
 								bucket.m_casKeyTable.GetKey(casKey, offset);
 								UBA_ASSERT(IsCompressed(casKey));
-								auto findIt = existingCas.find(casKey);
-								if (findIt != existingCas.end())
+								auto findIt = m_existingCas.find(casKey);
+								if (findIt != m_existingCas.end())
 								{
 									touchedCas.push_back(&findIt->second.refCount);
 									continue;
@@ -624,7 +631,7 @@ namespace uba
 			// Reset deleted cas files and update it again..
 			deletedCasFiles.clear();
 
-			for (auto i=existingCas.begin(), e=existingCas.end(); i!=e;)
+			for (auto i=m_existingCas.begin(), e=m_existingCas.end(); i!=e;)
 			{
 				if (i->second.refCount != 0)
 				{
@@ -635,8 +642,8 @@ namespace uba
 				deletedCasFiles.insert(i->first);
 				++deletedCasCount;
 				totalCasSize -= i->second.size;
-				i = existingCas.erase(i);
-				e = existingCas.end();
+				i = m_existingCas.erase(i);
+				e = m_existingCas.end();
 			}
 
 			// Add drop cas as work so it can run in the background
diff --git a/Engine/Source/Programs/UnrealBuildAccelerator/Common/Public/UbaCacheServer.h b/Engine/Source/Programs/UnrealBuildAccelerator/Common/Public/UbaCacheServer.h
--- a/Engine/Source/Programs/UnrealBuildAccelerator/Common/Public/UbaCacheServer.h
+++ b/Engine/Source/Programs/UnrealBuildAccelerator/Common/Public/UbaCacheServer.h
@@ -2,6 +2,7 @@
 
 #pragma once
 
+#include "UbaHash.h"
 #include "UbaLogger.h"
 #include "UbaStringBuffer.h"
 
@@ -62,6 +63,9 @@ namespace uba
 		Atomic<u32> m_addsSinceMaintenance;
 		Atomic<bool> m_isRunningMaintenance;
 
+		struct CasFileInfo { u64 size; u64 refCount; };
+		UnorderedMap<CasKey, CasFileInfo> m_existingCas;
+
 		ReaderWriterLock m_bucketsLock;
 		Map<u64, Bucket> m_buckets;
 
