[UBA]
* Added -expiration=<seconds> to UbaCacheService.. This controls when cache entries should be deleted. If an entry has been unused for more than <expiration> time it gets deleted. Defaults to two days.
* Added reason why cache entries are deleted
* Changed so cas database is always checked in startup.

diff --git a/Engine/Source/Programs/UnrealBuildAccelerator/CacheService/Private/UbaCacheService.cpp b/Engine/Source/Programs/UnrealBuildAccelerator/CacheService/Private/UbaCacheService.cpp
--- a/Engine/Source/Programs/UnrealBuildAccelerator/CacheService/Private/UbaCacheService.cpp
+++ b/Engine/Source/Programs/UnrealBuildAccelerator/CacheService/Private/UbaCacheService.cpp
@@ -45,6 +45,7 @@ namespace uba
 		logger.Info(TC("  -dir=<rootdir>          The directory used to store data. Defaults to \"%s\""), DefaultRootDir);
 		logger.Info(TC("  -port=[<host>:]<port>   The ip/name and port (default: %u) to listen for clients on"), DefaultCachePort);
 		logger.Info(TC("  -capacity=<gigaby>      Capacity of local store. Defaults to %u gigabytes"), DefaultCapacityGb);
+		logger.Info(TC("  -expiration=<seconds>   Time until unused cache entries get deleted. Defaults to 2 days (172800 seconds)"));
 		logger.Info(TC(""));
 		return -1;
 	}
@@ -115,6 +116,7 @@ namespace uba
 		u16 port = DefaultCachePort;
 		bool quiet = false;
 		bool storeCompressed = true;
+		u32 expirationTimeSeconds = 2*24*60*60;
 
 		for (int i=1; i!=argc; ++i)
 		{
@@ -158,6 +160,11 @@ namespace uba
 				if (!value.Parse(storageCapacityGb))
 					return PrintHelp(TC("Invalid value for -capacity"));
 			}
+			else if (name.Equals(TC("-expiration")))
+			{
+				if (!value.Parse(expirationTimeSeconds))
+					return PrintHelp(TC("Invalid value for -expire"));
+			}
 			else if (name.Equals(TC("-?")))
 			{
 				return PrintHelp(TC(""));
@@ -182,7 +189,7 @@ namespace uba
 		#if UBA_DEBUG
 		dbgStr = TC(" (DEBUG)");
 		#endif
-		logger.Info(TC("UbaCacheService v%s(%u)%s (Workers: %u, Rootdir: \"%s\", StoreCapacity: %uGb)\n"), Version, CacheNetworkVersion, dbgStr, GetLogicalProcessorCount(), g_rootDir.data, storageCapacityGb);
+		logger.Info(TC("UbaCacheService v%s(%u)%s (Workers: %u, Rootdir: \"%s\", StoreCapacity: %uGb, Expiration: %s)\n"), Version, CacheNetworkVersion, dbgStr, GetLogicalProcessorCount(), g_rootDir.data, storageCapacityGb, TimeToText(MsToTime(expirationTimeSeconds)*1000, true).str);
 
 		u64 storageCapacity = u64(storageCapacityGb)*1000*1000*1000;
 
@@ -224,10 +231,12 @@ namespace uba
 		storageInfo.writeRecievedCasFilesToDisk = true;
 		StorageServer storageServer(storageInfo);
 
-		if (!storageServer.LoadCasTable(true))
+		if (!storageServer.LoadCasTable(true, true))
 			return -1;
 
-		CacheServer cacheServer(logWriter, g_rootDir.data, networkServer, storageServer);
+		CacheServerCreateInfo cacheInfo(logWriter, storageServer, g_rootDir.data);
+		cacheInfo.expirationTimeSeconds = expirationTimeSeconds;
+		CacheServer cacheServer(cacheInfo);
 
 		if (!cacheServer.Load())
 			return -1;
@@ -250,6 +259,7 @@ namespace uba
 			}
 		}
 
+		storageServer.SaveCasTable(true);
 		cacheServer.Save();
 		return 0;
 	}
diff --git a/Engine/Source/Programs/UnrealBuildAccelerator/Common/Private/UbaCacheServer.cpp b/Engine/Source/Programs/UnrealBuildAccelerator/Common/Private/UbaCacheServer.cpp
--- a/Engine/Source/Programs/UnrealBuildAccelerator/Common/Private/UbaCacheServer.cpp
+++ b/Engine/Source/Programs/UnrealBuildAccelerator/Common/Private/UbaCacheServer.cpp
@@ -48,6 +48,9 @@ namespace uba
 		Atomic<bool> hasDeletedEntries;
 		Atomic<bool> needsSave;
 
+		Atomic<u64> lastSavedTime;
+		Atomic<u64> lastUsedTime;
+
 		u32 index = ~0u;
 	};
 
@@ -63,14 +66,17 @@ namespace uba
 		}
 	}
 
-	CacheServer::CacheServer(LogWriter& writer, const tchar* rootDir, NetworkServer& server, StorageServer& storage)
-	:	m_logger(writer, TC("UbaCacheServer"))
-	,	m_server(server)
-	,	m_storage(storage)
+	CacheServer::CacheServer(const CacheServerCreateInfo& info)
+	:	m_logger(info.writer, TC("UbaCacheServer"))
+	,	m_server(info.storage.GetServer())
+	,	m_storage(info.storage)
 	{
+		m_checkInputsForDeletedCas = info.checkInputsForDeletedCas;
 		m_startTime = GetTime();
 
-		m_rootDir.count = GetFullPathNameW(rootDir, m_rootDir.capacity, m_rootDir.data, NULL);
+		m_expirationTimeSeconds = info.expirationTimeSeconds;
+
+		m_rootDir.count = GetFullPathNameW(info.rootDir, m_rootDir.capacity, m_rootDir.data, NULL);
 		m_rootDir.Replace('/', PathSeparator).EnsureEndsWithSlash();
 
 		m_server.RegisterService(CacheServiceId,
@@ -219,10 +225,13 @@ namespace uba
 
 	bool CacheServer::Save()
 	{
-		if (m_addsSinceMaintenance == 0)
-			return true;
+		for (auto& kv : m_buckets)
+		{
+			Bucket& bucket = kv.second;
+			if (bucket.lastSavedTime < bucket.lastUsedTime)
+				bucket.needsSave = true;
+		}
 
-		SCOPED_WRITE_LOCK(m_connectionsLock, lock);
 		return SaveNoLock();
 	}
 
@@ -301,6 +310,8 @@ namespace uba
 
 	bool CacheServer::SaveBucket(u64 bucketId, Bucket& bucket)
 	{
+		u64 saveStart = GetTime();
+
 		StringBuffer<MaxPath> bucketsDir(m_rootDir);
 		bucketsDir.EnsureEndsWithSlash().Append(TC("buckets"));
 		if (!m_storage.CreateDirectory(bucketsDir.data))
@@ -340,7 +351,13 @@ namespace uba
 			file.WriteBytes(temp.data(), temp.size());
 		}
 
-		return file.Close();
+		if (!file.Close())
+			return false;
+
+		bucket.lastSavedTime = GetSystemTimeAsFileTime() - m_creationTime;
+
+		m_logger.Detail(TC("    Bucket %u saved (%s)"), bucket.index, TimeToText(GetTime() - saveStart).str);
+		return true;
 	}
 
 	bool CacheServer::SaveNoLock()
@@ -377,12 +394,10 @@ namespace uba
 				Bucket& bucket = it->second;
 				if (!bucket.needsSave)
 					return;
-				u64 saveStart = GetTime();
 				if (SaveBucket(it->first, bucket))
 					bucket.needsSave = false;
 				else
 					success = false;
-				m_logger.Detail(TC("    Bucket %u saved (%s)"), bucket.index, TimeToText(GetTime() - saveStart).str);
 			});
 
 		return success;
@@ -435,27 +450,22 @@ namespace uba
 		u64 deletedCasCount = deletedCasFiles.size();
 
 		u64 totalCasSize = 0;
-		if (m_existingCas.empty())
-		{
-			ReaderWriterLock existingCasLock;
-			u64 traverseStartTime = GetTime();
-			m_storage.TraverseAllCasFiles([&](const CasKey& casKey, u64 size)
-				{
-					SCOPED_WRITE_LOCK(existingCasLock, lock);
-					totalCasSize += size;
-					m_existingCas.try_emplace(casKey, CasFileInfo{size, 0ull});
-				}, true);
-			m_logger.Detail(TC("  Found %llu cas files (%s)"), m_existingCas.size(), TimeToText(GetTime() - traverseStartTime).str);
-		}
-		else
+
+		struct CasFileInfo { u64 size; u64 refCount; };
+		UnorderedMap<CasKey, CasFileInfo> existingCas;
+		ReaderWriterLock existingCasLock;
+
 		{
-			for (const CasKey& deletedCasFile : deletedCasFiles)
-				m_existingCas.erase(deletedCasFile);
-			for (auto& kv : m_existingCas)
+			SCOPED_WRITE_LOCK(m_storage.m_casLookupLock, lookupLock);
+			for (auto& kv : m_storage.m_casLookup)
+			{
 				totalCasSize += kv.second.size;
+				existingCas.try_emplace(kv.first, CasFileInfo{kv.second.size, 0ull});
+			}
 		}
 
-		u64 totalCasCount = m_existingCas.size() + deletedCasCount;
+		m_logger.Detail(TC("  Found %llu cas files (%s)"), existingCas.size(), BytesToText(totalCasSize).str);
+		u64 totalCasCount = existingCas.size() + deletedCasCount;
 
 		if (shouldExit())
 			return true;
@@ -463,11 +473,22 @@ namespace uba
 		u64 now = GetSystemTimeAsFileTime();
 		u64 oldest = 0;
 
+		u64 lastUseTimeLimit = 0;
+		if (m_expirationTimeSeconds && GetFileTimeAsSeconds(now - m_creationTime) > m_expirationTimeSeconds)
+			lastUseTimeLimit = (now - m_creationTime) - GetSecondsAsFileTime(m_expirationTimeSeconds);
+
+
 		u32 workerCount = m_server.GetWorkerCount();
 		u32 workerCountToUse = workerCount > 0 ? workerCount - 1 : 0;
 		u32 workerCountToUseForBuckets = Min(workerCountToUse, u32(m_buckets.size()));
 
+		Atomic<u64> totalEntryCount;
 		Atomic<u64> deleteEntryCount;
+		Atomic<u64> expiredEntryCount;
+		Atomic<u64> overflowedEntryCount;
+		Atomic<u64> missingOutputEntryCount;
+		Atomic<u64> missingInputEntryCount;
+
 		Atomic<u64> activeDropCount;
 		auto dropCasGuard = MakeGuard([&]() { while (activeDropCount != 0) Sleep(1); });
 
@@ -477,8 +498,7 @@ namespace uba
 			bool checkInputsForDeletes = m_checkInputsForDeletedCas && !deletedCasFiles.empty();
 
 			oldest = 0;
-
-			ReaderWriterLock existingCasLock;
+			totalEntryCount = 0;
 
 			m_server.ParallelFor(workerCountToUseForBuckets, m_buckets, [&](auto& it)
 			{
@@ -524,6 +544,13 @@ namespace uba
 						{
 							deleteEntry = true;
 							capacityLeft = 0;
+							++overflowedEntryCount;
+						}
+
+						if (!deleteEntry && entry.creationTime < lastUseTimeLimit && entry.lastUsedTime < lastUseTimeLimit)
+						{
+							deleteEntry = true;
+							++expiredEntryCount;
 						}
 
 						// This is an attempt at removing entries that has inputs that depends on other entries outputs.
@@ -544,6 +571,7 @@ namespace uba
 										if (!IsOffsetDeleted(inputReader.Read7BitEncoded()))
 											continue;
 										deleteEntry = true;
+										++missingInputEntryCount;
 										break;
 									}
 								}
@@ -558,6 +586,7 @@ namespace uba
 									if (!IsOffsetDeleted(extraReader.Read7BitEncoded()))
 										continue;
 									deleteEntry = true;
+									++missingInputEntryCount;
 									break;
 								}
 							}
@@ -576,13 +605,14 @@ namespace uba
 								CasKey casKey;
 								bucket.m_casKeyTable.GetKey(casKey, offset);
 								UBA_ASSERT(IsCompressed(casKey));
-								auto findIt = m_existingCas.find(casKey);
-								if (findIt != m_existingCas.end())
+								auto findIt = existingCas.find(casKey);
+								if (findIt != existingCas.end())
 								{
 									touchedCas.push_back(&findIt->second.refCount);
 									continue;
 								}
 								deleteEntry = true;
+								++missingOutputEntryCount;
 								break;
 							}
 						}
@@ -626,12 +656,14 @@ namespace uba
 
 				for (auto& key : keysToErase)
 					bucket.m_cacheEntryLookup.erase(key);
+
+				totalEntryCount += bucket.totalEntryCount;
 			});
 
 			// Reset deleted cas files and update it again..
 			deletedCasFiles.clear();
 
-			for (auto i=m_existingCas.begin(), e=m_existingCas.end(); i!=e;)
+			for (auto i=existingCas.begin(), e=existingCas.end(); i!=e;)
 			{
 				if (i->second.refCount != 0)
 				{
@@ -642,8 +674,8 @@ namespace uba
 				deletedCasFiles.insert(i->first);
 				++deletedCasCount;
 				totalCasSize -= i->second.size;
-				i = m_existingCas.erase(i);
-				e = m_existingCas.end();
+				i = existingCas.erase(i);
+				e = existingCas.end();
 			}
 
 			// Add drop cas as work so it can run in the background
@@ -655,6 +687,15 @@ namespace uba
 		}
 		while (!deletedCasFiles.empty()); // if cas files are deleted we need to do another loop and check cache entry inputs to see if files were inputs
 
+		if (overflowedEntryCount)
+			m_logger.Detail(TC("  Found %llu overflowed cache entries"), overflowedEntryCount.load());
+		if (expiredEntryCount)
+			m_logger.Detail(TC("  Found %llu expired cache entries (older than %s)"), expiredEntryCount.load(), TimeToText(MsToTime(m_expirationTimeSeconds*1000), true).str);
+		if (missingOutputEntryCount)
+			m_logger.Detail(TC("  Found %llu cache entries with missing output cas"), missingOutputEntryCount.load());
+		if (missingInputEntryCount)
+			m_logger.Detail(TC("  Found %llu cache entries with missing input cas"), missingInputEntryCount.load());
+
 		m_logger.Detail(TC("  Deleted %llu cas files and %llu cache entries (%s)"), deletedCasCount, deleteEntryCount.load(), TimeToText(GetTime() - deleteCacheEntriesStartTime).str);
 
 		if (shouldExit())
@@ -668,7 +709,7 @@ namespace uba
 
 			if (!bucket.hasDeletedEntries && !forceAllSteps)
 			{
-				m_logger.Detail(TC("    Bucket %u skipped updating. No entries deleted"), bucket.index);
+				m_logger.Detail(TC("    Bucket %u skipped updating. (%llu entries)"), bucket.index, bucket.totalEntryCount.load());
 				return;
 			}
 			bucket.hasDeletedEntries = false;
@@ -816,7 +857,7 @@ namespace uba
 
 		u64 oldestTime = oldest ? GetFileTimeAsTime(now - (m_creationTime + oldest)) : 0;
 		u64 duration = GetTime() - startTime;
-		m_logger.Info(TC("Maintenance done! (%s) CasFiles: %llu (%s) OldestEntry: %s"), TimeToText(duration).str, totalCasCount - deletedCasCount, BytesToText(totalCasSize).str, TimeToText(oldestTime, true).str);
+		m_logger.Info(TC("Maintenance done! (%s) CasFiles: %llu (%s) Entries: %llu OldestEntry: %s"), TimeToText(duration).str, totalCasCount - deletedCasCount, BytesToText(totalCasSize).str, totalEntryCount.load(), TimeToText(oldestTime, true).str);
 		
 		m_longestMaintenance = Max(m_longestMaintenance, duration);
 
@@ -1226,7 +1267,9 @@ namespace uba
 		{
 			if (entryId != entry.id)
 				continue;
-			entry.lastUsedTime = GetSystemTimeAsFileTime() - m_creationTime;
+			u64 fileTime = GetSystemTimeAsFileTime() - m_creationTime;
+			entry.lastUsedTime = fileTime;
+			bucket.lastUsedTime = fileTime;
 			break;
 		}
 		return true;
diff --git a/Engine/Source/Programs/UnrealBuildAccelerator/Common/Private/UbaFile.cpp b/Engine/Source/Programs/UnrealBuildAccelerator/Common/Private/UbaFile.cpp
--- a/Engine/Source/Programs/UnrealBuildAccelerator/Common/Private/UbaFile.cpp
+++ b/Engine/Source/Programs/UnrealBuildAccelerator/Common/Private/UbaFile.cpp
@@ -801,6 +801,15 @@ namespace uba
 		return MsToTime(GetFileTimeAsSeconds(fileTime)*1000);
 	}
 
+	u64 GetSecondsAsFileTime(u64 seconds)
+	{
+#if PLATFORM_WINDOWS
+		return MsToTime(seconds*1000);
+#else
+		return seconds * 10'000'000ull;
+#endif
+	}
+
 	bool GetCurrentDirectoryW(StringBufferBase& out)
 	{
 #if PLATFORM_WINDOWS
diff --git a/Engine/Source/Programs/UnrealBuildAccelerator/Common/Private/UbaStorage.cpp b/Engine/Source/Programs/UnrealBuildAccelerator/Common/Private/UbaStorage.cpp
--- a/Engine/Source/Programs/UnrealBuildAccelerator/Common/Private/UbaStorage.cpp
+++ b/Engine/Source/Programs/UnrealBuildAccelerator/Common/Private/UbaStorage.cpp
@@ -956,8 +956,6 @@ namespace uba
 		StringBuffer<> casRoot;
 		casRoot.Append(m_rootDir.data, m_rootDir.count - 1);
 
-		m_logger.Info(TC("Previous run was not gracefully shutdown. Reparsing cas directory %s to check for added/missing files"), casRoot.data);
-
 		Atomic<bool> success = true;
 		TraverseAllCasFiles(casRoot.data, [&](const StringBufferBase& fullPath, const DirectoryEntry& e)
 			{
@@ -1242,7 +1240,7 @@ namespace uba
 		return logger.Error(TC("Needs exclusive access to storage %s. Another process is running"), rootDir.data);
 	}
 
-	bool StorageImpl::LoadCasTable(bool logStats)
+	bool StorageImpl::LoadCasTable(bool logStats, bool alwaysCheckAllFiles)
 	{
 		static bool isExclusive = CheckExclusive(m_logger, m_rootDir);
 		if (!isExclusive)
@@ -1416,6 +1414,8 @@ namespace uba
 			if (fh != InvalidFileHandle)
 				GetFileLastWriteTime(fileTime, fh);
 
+			m_logger.Info(TC("Previous run was not gracefully shutdown. Reparsing cas directory %s to check for added/missing files"), m_rootDir.data);
+
 			if (!CheckAllCasFiles(fileTime))
 				return false;
 			resave = true;
@@ -1423,6 +1423,8 @@ namespace uba
 			if (fh != InvalidFileHandle)
 				SetFileLastWriteTime(fh, GetSystemTimeAsFileTime());
 		}
+		else if (alwaysCheckAllFiles)
+			CheckAllCasFiles();
 
 		if (!m_manuallyHandleOverflow)
 			HandleOverflow(nullptr);
diff --git a/Engine/Source/Programs/UnrealBuildAccelerator/Common/Public/UbaCacheServer.h b/Engine/Source/Programs/UnrealBuildAccelerator/Common/Public/UbaCacheServer.h
--- a/Engine/Source/Programs/UnrealBuildAccelerator/Common/Public/UbaCacheServer.h
+++ b/Engine/Source/Programs/UnrealBuildAccelerator/Common/Public/UbaCacheServer.h
@@ -16,10 +16,29 @@ namespace uba
 	struct CacheEntries;
 	struct ConnectionInfo;
 
+	struct CacheServerCreateInfo
+	{
+		// Log writer
+		LogWriter& writer;
+
+		// Storage server
+		StorageServer& storage;
+
+		// Root dir
+		const tchar* rootDir = nullptr;
+
+		// Will check cache entry inputs of they depend on cas files that have been deleted
+		bool checkInputsForDeletedCas = true;
+
+		// The time cache entries will stay around after they were last used in hours (defaults to two days)
+		// Set to zero to never expire
+		u64 expirationTimeSeconds = 2*24*60*60;
+	};
+
 	class CacheServer
 	{
 	public:
-		CacheServer(LogWriter& writer, const tchar* rootDir, NetworkServer& server, StorageServer& storage);
+		CacheServer(const CacheServerCreateInfo& info);
 		~CacheServer();
 
 		bool Load();
@@ -63,9 +82,6 @@ namespace uba
 		Atomic<u32> m_addsSinceMaintenance;
 		Atomic<bool> m_isRunningMaintenance;
 
-		struct CasFileInfo { u64 size; u64 refCount; };
-		UnorderedMap<CasKey, CasFileInfo> m_existingCas;
-
 		ReaderWriterLock m_bucketsLock;
 		Map<u64, Bucket> m_buckets;
 
@@ -77,6 +93,7 @@ namespace uba
 		u64 m_creationTime = 0;
 		u64 m_startTime = 0;
 		u64 m_longestMaintenance = 0;
+		u64 m_expirationTimeSeconds = 0;
 		bool m_dbfileDirty = false;
 
 		bool m_checkInputsForDeletedCas = true;
diff --git a/Engine/Source/Programs/UnrealBuildAccelerator/Common/Public/UbaFile.h b/Engine/Source/Programs/UnrealBuildAccelerator/Common/Public/UbaFile.h
--- a/Engine/Source/Programs/UnrealBuildAccelerator/Common/Public/UbaFile.h
+++ b/Engine/Source/Programs/UnrealBuildAccelerator/Common/Public/UbaFile.h
@@ -76,6 +76,7 @@ namespace uba
 	u64 GetSystemTimeAsFileTime();
 	u64 GetFileTimeAsSeconds(u64 fileTime);
 	u64 GetFileTimeAsTime(u64 fileTime);
+	u64 GetSecondsAsFileTime(u64 seconds);
 	bool GetCurrentDirectoryW(StringBufferBase& out);
 
 
diff --git a/Engine/Source/Programs/UnrealBuildAccelerator/Common/Public/UbaStorage.h b/Engine/Source/Programs/UnrealBuildAccelerator/Common/Public/UbaStorage.h
--- a/Engine/Source/Programs/UnrealBuildAccelerator/Common/Public/UbaStorage.h
+++ b/Engine/Source/Programs/UnrealBuildAccelerator/Common/Public/UbaStorage.h
@@ -103,7 +103,7 @@ namespace uba
 		StorageImpl(const StorageCreateInfo& info, const tchar* logPrefix = TC("UbaStorage"));
 		virtual ~StorageImpl();
 
-		bool LoadCasTable(bool logStats = true);
+		bool LoadCasTable(bool logStats = true, bool alwaysCheckAllFiles = false);
 		bool CheckCasContent(u32 workerCount);
 		bool CheckFileTable(const tchar* searchPath, u32 workerCount);
 		const tchar* GetTempPath();
diff --git a/Engine/Source/Programs/UnrealBuildAccelerator/Common/Public/UbaStorageServer.h b/Engine/Source/Programs/UnrealBuildAccelerator/Common/Public/UbaStorageServer.h
--- a/Engine/Source/Programs/UnrealBuildAccelerator/Common/Public/UbaStorageServer.h
+++ b/Engine/Source/Programs/UnrealBuildAccelerator/Common/Public/UbaStorageServer.h
@@ -24,6 +24,13 @@ namespace uba
 
 		bool RegisterDisallowedPath(const tchar* path);
 
+		// Get the network server used by this session
+		NetworkServer& GetServer() { return m_server; }
+
+		using StorageImpl::StoreCasFile;
+
+	protected:
+
 		virtual bool GetZone(StringBufferBase& out) override;
 		virtual bool RetrieveCasFile(RetrieveResult& out, const CasKey& casKey, const tchar* hint, FileMappingBuffer* mappingBuffer = nullptr, u64 memoryMapAlignment = 1, bool allowProxy = true) override;
 		virtual bool StoreCasFile(CasKey& out, StringKey fileNameKey, const tchar* fileName, FileMappingHandle mappingHandle, u64 mappingOffset, u64 fileSize, const tchar* hint, bool deferCreation = false, bool keepMappingInMemory = false) override;
@@ -31,10 +38,10 @@ namespace uba
 		virtual bool IsDisallowedPath(const tchar* fileName) override;
 		virtual void SetTrace(Trace* trace, bool detailed) override;
 		virtual bool HasProxy(u32 clientId) override;
+
 		void OnDisconnected(u32 clientId);
 		bool HandleMessage(const ConnectionInfo& connectionInfo, u8 messageType, BinaryReader& reader, BinaryWriter& writer);
 		bool WaitForWritten(CasEntry& casEntry, ScopedWriteLock& entryLock, const tchar* hint);
-		using StorageImpl::StoreCasFile;
 
 		u16 PopId();
 		void PushId(u16 id);
diff --git a/Engine/Source/Programs/UnrealBuildAccelerator/Test/Private/UbaTestBasics.h b/Engine/Source/Programs/UnrealBuildAccelerator/Test/Private/UbaTestBasics.h
--- a/Engine/Source/Programs/UnrealBuildAccelerator/Test/Private/UbaTestBasics.h
+++ b/Engine/Source/Programs/UnrealBuildAccelerator/Test/Private/UbaTestBasics.h
@@ -45,6 +45,10 @@ namespace uba
 		consoleLogger.Info(TC("Slept ms: %llu"), ms);
 #endif
 
+		u64 seconds = 15;
+		u64 fileTime = GetSecondsAsFileTime(seconds);
+		u64 seconds2 = GetFileTimeAsSeconds(fileTime);
+		UBA_TEST_CHECK(seconds == seconds2, "GetSecondsAsFileTime does not match GetFileTimeAsSeconds");
 		return true;
 	}
 
