[UBA]
* UbaCacheService - Changed expiration time to 3 days to not evict all cache over weekends (maybe this should be handled differently?)

diff --git a/Engine/Source/Programs/UnrealBuildAccelerator/CacheService/Private/UbaCacheService.cpp b/Engine/Source/Programs/UnrealBuildAccelerator/CacheService/Private/UbaCacheService.cpp
--- a/Engine/Source/Programs/UnrealBuildAccelerator/CacheService/Private/UbaCacheService.cpp
+++ b/Engine/Source/Programs/UnrealBuildAccelerator/CacheService/Private/UbaCacheService.cpp
@@ -19,6 +19,7 @@ namespace uba
 {
 	const tchar*	Version = GetVersionString();
 	u32				DefaultCapacityGb = 500;
+	u32				DefaultExpiration = 3*24*60*60;
 	const tchar*	DefaultRootDir = []() {
 		static tchar buf[256];
 		if (IsWindows)
@@ -45,7 +46,7 @@ namespace uba
 		logger.Info(TC("  -dir=<rootdir>          The directory used to store data. Defaults to \"%s\""), DefaultRootDir);
 		logger.Info(TC("  -port=[<host>:]<port>   The ip/name and port (default: %u) to listen for clients on"), DefaultCachePort);
 		logger.Info(TC("  -capacity=<gigaby>      Capacity of local store. Defaults to %u gigabytes"), DefaultCapacityGb);
-		logger.Info(TC("  -expiration=<seconds>   Time until unused cache entries get deleted. Defaults to 2 days (172800 seconds)"));
+		logger.Info(TC("  -expiration=<seconds>   Time until unused cache entries get deleted. Defaults to %s (%u seconds)"), TimeToText(MsToTime(DefaultExpiration*1000)).str, DefaultExpiration);
 		logger.Info(TC(""));
 		return -1;
 	}
@@ -116,7 +117,7 @@ namespace uba
 		u16 port = DefaultCachePort;
 		bool quiet = false;
 		bool storeCompressed = true;
-		u32 expirationTimeSeconds = 2*24*60*60;
+		u32 expirationTimeSeconds = DefaultExpiration;
 
 		for (int i=1; i!=argc; ++i)
 		{
