[UBA]
* Fixed bug where cache service was showing cas entries that doesn't exist

diff --git a/Engine/Source/Programs/UnrealBuildAccelerator/Common/Private/UbaCacheServer.cpp b/Engine/Source/Programs/UnrealBuildAccelerator/Common/Private/UbaCacheServer.cpp
--- a/Engine/Source/Programs/UnrealBuildAccelerator/Common/Private/UbaCacheServer.cpp
+++ b/Engine/Source/Programs/UnrealBuildAccelerator/Common/Private/UbaCacheServer.cpp
@@ -456,11 +456,19 @@ namespace uba
 		ReaderWriterLock existingCasLock;
 
 		{
+			// TODO: Make this cleaner... 
 			SCOPED_WRITE_LOCK(m_storage.m_casLookupLock, lookupLock);
-			for (auto& kv : m_storage.m_casLookup)
+			for (auto i=m_storage.m_casLookup.begin(), e=m_storage.m_casLookup.end(); i!=e;)
 			{
-				totalCasSize += kv.second.size;
-				existingCas.try_emplace(kv.first, CasFileInfo{kv.second.size, 0ull});
+				if (i->second.verified && !i->second.exists)
+				{
+					i = m_storage.m_casLookup.erase(i);
+					e=m_storage.m_casLookup.end();
+					continue;
+				}
+				totalCasSize += i->second.size;
+				existingCas.try_emplace(i->first, CasFileInfo{i->second.size, 0ull});
+				++i;
 			}
 		}
 
