[UBA]
* Optimizations in UbaCacheService maintenance code

diff --git a/Engine/Source/Programs/UnrealBuildAccelerator/Common/Private/UbaCacheServer.cpp b/Engine/Source/Programs/UnrealBuildAccelerator/Common/Private/UbaCacheServer.cpp
--- a/Engine/Source/Programs/UnrealBuildAccelerator/Common/Private/UbaCacheServer.cpp
+++ b/Engine/Source/Programs/UnrealBuildAccelerator/Common/Private/UbaCacheServer.cpp
@@ -458,12 +458,15 @@ namespace uba
 		{
 			// TODO: Make this cleaner... 
 			SCOPED_WRITE_LOCK(m_storage.m_casLookupLock, lookupLock);
+			
+			existingCas.reserve(m_storage.m_casLookup.size());
+
 			for (auto i=m_storage.m_casLookup.begin(), e=m_storage.m_casLookup.end(); i!=e;)
 			{
 				if (i->second.verified && !i->second.exists)
 				{
 					i = m_storage.m_casLookup.erase(i);
-					e=m_storage.m_casLookup.end();
+					e = m_storage.m_casLookup.end();
 					continue;
 				}
 				totalCasSize += i->second.size;
