[UBA]
* Removed check not needed anymore

diff --git a/Engine/Source/Programs/UnrealBuildAccelerator/Host/Private/UbaExports.cpp b/Engine/Source/Programs/UnrealBuildAccelerator/Host/Private/UbaExports.cpp
--- a/Engine/Source/Programs/UnrealBuildAccelerator/Host/Private/UbaExports.cpp
+++ b/Engine/Source/Programs/UnrealBuildAccelerator/Host/Private/UbaExports.cpp
@@ -539,10 +539,6 @@ extern "C"
 
 	bool CacheClient_WriteToCache(uba::CacheClient* cacheClient, uba::RootPaths* rootPaths, uba::u32 bucket, const uba::ProcessHandle* process, const uba::u8* inputs, uba::u32 inputsSize, const uba::u8* outputs, uba::u32 outputsSize)
 	{
-		using namespace uba;
-		auto& si = process->GetStartInfo();
-		if (!si.trackInputs)
-			return false;
 		return cacheClient->WriteToCache(*rootPaths, bucket, process->GetStartInfo(), inputs, inputsSize, outputs, outputsSize, process->GetId());
 	}
 
