[UBA]
* Added more stats to cache fetch to be able to figure out where time goes
* Changed code to use FileExists instead of GetFileInformation.. since we only need information that can be fetched with FileExists
* Removed usage of sprintf when creating cas filename

diff --git a/Engine/Source/Programs/UnrealBuildAccelerator/Common/Private/UbaCacheClient.cpp b/Engine/Source/Programs/UnrealBuildAccelerator/Common/Private/UbaCacheClient.cpp
--- a/Engine/Source/Programs/UnrealBuildAccelerator/Common/Private/UbaCacheClient.cpp
+++ b/Engine/Source/Programs/UnrealBuildAccelerator/Common/Private/UbaCacheClient.cpp
@@ -314,7 +314,7 @@ namespace uba
 		bool success = false;
 		auto tg = MakeGuard([&]()
 			{
-				cacheStats.testEntries.time -= cacheStats.fetchCasTable.time;
+				cacheStats.testEntry.time -= (cacheStats.fetchCasTable.time + cacheStats.normalizeFile.time);
 				BinaryWriter writer(memory, 0, sizeof_array(memory));
 				cacheStats.Write(writer);
 				if (success)
@@ -366,7 +366,7 @@ namespace uba
 		for (; entryIndex!=entryCount; ++entryIndex)
 		{
 			{
-				TimerScope ts(cacheStats.testEntries);
+				TimerScope ts(cacheStats.testEntry);
 				bool isMatch = true;
 
 				bool result = traverser.TraverseEntryInputs([&](u32 casKeyOffset)
@@ -391,6 +391,7 @@ namespace uba
 								auto insres2 = normalizedCasKeys.try_emplace(ToStringKeyNoCheck(path.data, path.count));
 								if (insres2.second)
 								{
+									TimerScope ts(cacheStats.normalizeFile);
 									localCasKey = rootPaths.NormalizeAndHashFile(m_logger, path.data);
 									if (localCasKey != CasKeyZero)
 										localCasKey = AsCompressed(localCasKey, true);
diff --git a/Engine/Source/Programs/UnrealBuildAccelerator/Common/Private/UbaFile.cpp b/Engine/Source/Programs/UnrealBuildAccelerator/Common/Private/UbaFile.cpp
--- a/Engine/Source/Programs/UnrealBuildAccelerator/Common/Private/UbaFile.cpp
+++ b/Engine/Source/Programs/UnrealBuildAccelerator/Common/Private/UbaFile.cpp
@@ -219,7 +219,7 @@ namespace uba
 #endif
 	}
 
-	bool FileExists(Logger& logger, const tchar* fileName, u64* outSize, u32* outAttributes)
+	bool FileExists(Logger& logger, const tchar* fileName, u64* outSize, u32* outAttributes, u64* lastWriteTime)
 	{
 #if PLATFORM_WINDOWS
 		MAKE_LONG_FILENAME(fileName);
@@ -241,13 +241,20 @@ namespace uba
 		}
 		if (outAttributes)
 			*outAttributes = data.dwFileAttributes;
+
+		if (lastWriteTime)
+			*lastWriteTime = (u64&)data.ftLastWriteTime;
+
 		return true;
 #else
 		struct stat attr;
 		if (stat(fileName, &attr) == -1)
 		{
 			if (errno == ENOENT)
+			{
+				SetLastError(ERROR_FILE_NOT_FOUND);
 				return false;
+			}
 			UBA_ASSERTF(false, TC("FileExists error handling implemented"));
 			return false;
 		}
@@ -256,6 +263,8 @@ namespace uba
 			*outSize = attr.st_size;
 		if (outAttributes)
 			*outAttributes = attr.st_mode;
+		if (lastWriteTime)
+			*lastWriteTime = FromTimeSpec(attr.st_mtimespec);
 		return true;
 #endif
 	}
diff --git a/Engine/Source/Programs/UnrealBuildAccelerator/Common/Private/UbaStorage.cpp b/Engine/Source/Programs/UnrealBuildAccelerator/Common/Private/UbaStorage.cpp
--- a/Engine/Source/Programs/UnrealBuildAccelerator/Common/Private/UbaStorage.cpp
+++ b/Engine/Source/Programs/UnrealBuildAccelerator/Common/Private/UbaStorage.cpp
@@ -1883,24 +1883,19 @@ namespace uba
 		}
 		fileEntry.verified = true;
 
-		FileHandle fileHandle;
-		if (!OpenFileSequentialRead(m_logger, fileName, fileHandle))
+		u64 fileSize = 0;
+		u64 lastWritten = 0;
+
+		if (!FileExists(m_logger, fileName, &fileSize, nullptr, &lastWritten))
 		{
 			fileEntry.casKey = CasKeyZero;
+
+			u32 lastError = GetLastError();
+			if (lastError != ERROR_FILE_NOT_FOUND && lastError != ERROR_PATH_NOT_FOUND)
+				return m_logger.Error(TC("FileExists failed on %s (%s)"), fileName, LastErrorToText(lastError).data);
 			out = CasKeyZero;
 			return true;
 		}
-		auto fileGuard = MakeGuard([&](){ CloseFile(fileName, fileHandle); });
-
-		FileInformation info;
-		if (!GetFileInformationByHandle(info, m_logger, fileName, fileHandle))
-		{
-			fileEntry.casKey = CasKeyZero;
-			return m_logger.Error(TC("GetFileInformationByHandle failed on %s"), fileName);
-		}
-
-		u64 fileSize = info.size;
-		u64 lastWritten = info.lastWriteTime;
 
 		if (fileEntry.casKey != CasKeyZero)
 		{
@@ -1925,6 +1920,11 @@ namespace uba
 		fileEntry.lastWritten = lastWritten;
 		if (casKeyOverride == CasKeyZero)
 		{
+			FileHandle fileHandle;
+			if (!OpenFileSequentialRead(m_logger, fileName, fileHandle))
+				return false;
+			auto fileGuard = MakeGuard([&](){ CloseFile(fileName, fileHandle); });
+
 			if (fileIsCompressed)
 			{
 				CompressedObjFileHeader header(CasKeyZero);
@@ -2044,7 +2044,7 @@ namespace uba
 #if !UBA_USE_SPARSEFILE
 	bool StorageImpl::GetCasFileName(StringBufferBase& out, const CasKey& casKey)
 	{
-		out.Appendf(TC("%s%02x"), m_rootDir.data, ((const u8*)&casKey)[0]);
+		out.Append(m_rootDir.data).AppendHex(((const u8*)&casKey)[0]);
 		if (!CreateDirectory(out.data))
 			return false;
 		out.Append(PathSeparator).Append(CasKeyString(casKey).str);
diff --git a/Engine/Source/Programs/UnrealBuildAccelerator/Common/Public/UbaFile.h b/Engine/Source/Programs/UnrealBuildAccelerator/Common/Public/UbaFile.h
--- a/Engine/Source/Programs/UnrealBuildAccelerator/Common/Public/UbaFile.h
+++ b/Engine/Source/Programs/UnrealBuildAccelerator/Common/Public/UbaFile.h
@@ -47,7 +47,7 @@ namespace uba
 
 	bool ReadFile(Logger& logger, const tchar* fileName, FileHandle fileHandle, void* b, u64 bufferLen);
 	bool OpenFileSequentialRead(Logger& logger, const tchar* fileName, FileHandle& outHandle, bool fileNotFoundIsError = true, bool overlapped = false);
-	bool FileExists(Logger& logger, const tchar* fileName, u64* outSize = nullptr, u32* outAttributes = nullptr);
+	bool FileExists(Logger& logger, const tchar* fileName, u64* outSize = nullptr, u32* outAttributes = nullptr, u64* lastWriteTime = nullptr);
 	bool SetEndOfFile(Logger& logger, const tchar* fileName, FileHandle handle, u64 size);
 	bool GetDirectoryOfCurrentModule(Logger& logger, StringBufferBase& out);
 	bool DeleteAllFiles(Logger& logger, const tchar* dir, bool deleteDir = true, u32* count = nullptr);
diff --git a/Engine/Source/Programs/UnrealBuildAccelerator/Common/Public/UbaStats.h b/Engine/Source/Programs/UnrealBuildAccelerator/Common/Public/UbaStats.h
--- a/Engine/Source/Programs/UnrealBuildAccelerator/Common/Public/UbaStats.h
+++ b/Engine/Source/Programs/UnrealBuildAccelerator/Common/Public/UbaStats.h
@@ -121,7 +121,8 @@ namespace uba
 	#define UBA_CACHE_STATS \
 		UBA_CACHE_STAT(Timer, fetchEntries, 0) \
 		UBA_CACHE_STAT(Timer, fetchCasTable, 0) \
-		UBA_CACHE_STAT(Timer, testEntries, 0) \
+		UBA_CACHE_STAT(Timer, normalizeFile, 28) \
+		UBA_CACHE_STAT(Timer, testEntry, 0) \
 		UBA_CACHE_STAT(Timer, fetchOutput, 0) \
 		UBA_CACHE_STAT(AtomicU64, fetchBytesRaw, 26) \
 		UBA_CACHE_STAT(AtomicU64, fetchBytesComp, 26) \
diff --git a/Engine/Source/Programs/UnrealBuildAccelerator/Common/Public/UbaTrace.h b/Engine/Source/Programs/UnrealBuildAccelerator/Common/Public/UbaTrace.h
--- a/Engine/Source/Programs/UnrealBuildAccelerator/Common/Public/UbaTrace.h
+++ b/Engine/Source/Programs/UnrealBuildAccelerator/Common/Public/UbaTrace.h
@@ -61,7 +61,7 @@ namespace uba
 	using Color = u32;
 	inline Color toColor(u8 r, u8 g, u8 b) { return (r << 16) + (g << 8) + b; }
 
-	static constexpr u32 TraceVersion = 27;
+	static constexpr u32 TraceVersion = 28;
 	static constexpr u32 TraceReadCompatibilityVersion = 6;
 
 	class Trace : public WorkTracker
