[UBA]
* Fixed macos compile error

diff --git a/Engine/Source/Programs/UnrealBuildAccelerator/CacheService/Private/UbaCacheService.cpp b/Engine/Source/Programs/UnrealBuildAccelerator/CacheService/Private/UbaCacheService.cpp
--- a/Engine/Source/Programs/UnrealBuildAccelerator/CacheService/Private/UbaCacheService.cpp
+++ b/Engine/Source/Programs/UnrealBuildAccelerator/CacheService/Private/UbaCacheService.cpp
@@ -235,7 +235,7 @@ namespace uba
 		if (!storageServer.LoadCasTable(true, true))
 			return -1;
 
-		CacheServerCreateInfo cacheInfo(logWriter, storageServer, g_rootDir.data);
+		CacheServerCreateInfo cacheInfo(storageServer, g_rootDir.data, logWriter);
 		cacheInfo.expirationTimeSeconds = expirationTimeSeconds;
 		CacheServer cacheServer(cacheInfo);
 
diff --git a/Engine/Source/Programs/UnrealBuildAccelerator/Common/Private/UbaCacheServer.cpp b/Engine/Source/Programs/UnrealBuildAccelerator/Common/Private/UbaCacheServer.cpp
--- a/Engine/Source/Programs/UnrealBuildAccelerator/Common/Private/UbaCacheServer.cpp
+++ b/Engine/Source/Programs/UnrealBuildAccelerator/Common/Private/UbaCacheServer.cpp
@@ -67,7 +67,7 @@ namespace uba
 	}
 
 	CacheServer::CacheServer(const CacheServerCreateInfo& info)
-	:	m_logger(info.writer, TC("UbaCacheServer"))
+	:	m_logger(info.logWriter, TC("UbaCacheServer"))
 	,	m_server(info.storage.GetServer())
 	,	m_storage(info.storage)
 	{
diff --git a/Engine/Source/Programs/UnrealBuildAccelerator/Common/Public/UbaCacheServer.h b/Engine/Source/Programs/UnrealBuildAccelerator/Common/Public/UbaCacheServer.h
--- a/Engine/Source/Programs/UnrealBuildAccelerator/Common/Public/UbaCacheServer.h
+++ b/Engine/Source/Programs/UnrealBuildAccelerator/Common/Public/UbaCacheServer.h
@@ -18,8 +18,7 @@ namespace uba
 
 	struct CacheServerCreateInfo
 	{
-		// Log writer
-		LogWriter& writer;
+		CacheServerCreateInfo(StorageServer& s, const tchar* rd, LogWriter& w = g_consoleLogWriter) : storage(s), rootDir(rd), logWriter(w) {}
 
 		// Storage server
 		StorageServer& storage;
@@ -27,6 +26,9 @@ namespace uba
 		// Root dir
 		const tchar* rootDir = nullptr;
 
+		// Log writer
+		LogWriter& logWriter;
+
 		// Will check cache entry inputs of they depend on cas files that have been deleted
 		bool checkInputsForDeletedCas = true;
 
