[UBA]
* Fixed bug where files that was opened with certain create disposition mode and no write flags were not properly treated as writes causing weird errors. Only seen on windows 11 helpers

diff --git a/Engine/Source/Programs/UnrealBuildAccelerator/Detours/Private/Windows/UbaDetoursFunctionsKernelBase.inl b/Engine/Source/Programs/UnrealBuildAccelerator/Detours/Private/Windows/UbaDetoursFunctionsKernelBase.inl
--- a/Engine/Source/Programs/UnrealBuildAccelerator/Detours/Private/Windows/UbaDetoursFunctionsKernelBase.inl
+++ b/Engine/Source/Programs/UnrealBuildAccelerator/Detours/Private/Windows/UbaDetoursFunctionsKernelBase.inl
@@ -3260,7 +3260,7 @@ DWORD Detoured_ExpandEnvironmentStringsW(LPCWSTR lpSrc, LPWSTR lpDst, DWORD nSiz
 UINT Detoured_GetTempFileNameW(LPCWSTR lpPathName, LPCWSTR lpPrefixString, UINT uUnique, LPTSTR lpTempFileName)
 {
 	DETOURED_CALL(GetTempFileNameW);
-	DEBUG_LOG_TRUE(L"GetTempFileNameW", L"");
+	DEBUG_LOG_TRUE(L"GetTempFileNameW", L"%s %s", lpPathName, lpPrefixString);
 	return True_GetTempFileNameW(lpPathName, lpPrefixString, uUnique, lpTempFileName);
 }
 
diff --git a/Engine/Source/Programs/UnrealBuildAccelerator/Detours/Private/Windows/UbaDetoursFunctionsNtDll.inl b/Engine/Source/Programs/UnrealBuildAccelerator/Detours/Private/Windows/UbaDetoursFunctionsNtDll.inl
--- a/Engine/Source/Programs/UnrealBuildAccelerator/Detours/Private/Windows/UbaDetoursFunctionsNtDll.inl
+++ b/Engine/Source/Programs/UnrealBuildAccelerator/Detours/Private/Windows/UbaDetoursFunctionsNtDll.inl
@@ -24,12 +24,12 @@ bool IsWrite(u32 desiredAccess, u32 createDisposition)
 	return IsContentWrite(desiredAccess, createDisposition) || (desiredAccess & (FILE_WRITE_ATTRIBUTES | FILE_WRITE_EA)) != 0;
 }
 
-u8 GetFileAccessFlags(DWORD desiredAccess)
+u8 GetFileAccessFlags(DWORD desiredAccess, u32 createDisposition)
 {
 	u8 access = 0;
-	if (IsContentRead(desiredAccess, 0))
+	if (IsContentRead(desiredAccess, createDisposition))
 		access |= AccessFlag_Read;
-	if (IsWrite(desiredAccess, 0))
+	if (IsWrite(desiredAccess, createDisposition))
 		access |= AccessFlag_Write;
 	return access;
 }
@@ -903,7 +903,7 @@ NTSTATUS NTAPI Shared_NtCreateFile(bool IsCreateFunc, PHANDLE hFileHandle, ACCES
 		info.name = info.originalName;
 		if (!keepInMemory && !isSystemOrTempFile)
 		{
-			u8 access = GetFileAccessFlags(DesiredAccess);
+			u8 access = GetFileAccessFlags(DesiredAccess, CreateDisposition);
 			wchar_t newFileName[512];
 			Rpc_CreateFileW(lpFileName, fileNameKey, access, newFileName, sizeof_array(newFileName), size, closeId, false);
 			info.name = g_memoryBlock.Strdup(newFileName);
@@ -931,7 +931,7 @@ NTSTATUS NTAPI Shared_NtCreateFile(bool IsCreateFunc, PHANDLE hFileHandle, ACCES
 				u64 size = InvalidValue;
 				info.deleted = false;
 				wchar_t newFileName[1024];
-				u8 access = GetFileAccessFlags(DesiredAccess);
+				u8 access = GetFileAccessFlags(DesiredAccess, CreateDisposition);
 				Rpc_CreateFileW(lpFileName, fileNameKey, access, newFileName, sizeof_array(newFileName), size, closeId, false);
 				info.name = g_memoryBlock.Strdup(newFileName);
 				//info.size = size; // TODO: Should this be set?
