[UBA]
* Fixed bug in previous submit.. FileExists would return true for directories.. while OpenForRead would return false for directories

diff --git a/Engine/Source/Programs/UnrealBuildAccelerator/Common/Private/UbaStorage.cpp b/Engine/Source/Programs/UnrealBuildAccelerator/Common/Private/UbaStorage.cpp
--- a/Engine/Source/Programs/UnrealBuildAccelerator/Common/Private/UbaStorage.cpp
+++ b/Engine/Source/Programs/UnrealBuildAccelerator/Common/Private/UbaStorage.cpp
@@ -1884,12 +1884,12 @@ namespace uba
 		fileEntry.verified = true;
 
 		u64 fileSize = 0;
+		u32 attributes = 0;
 		u64 lastWritten = 0;
 
-		if (!FileExists(m_logger, fileName, &fileSize, nullptr, &lastWritten))
+		if (!FileExists(m_logger, fileName, &fileSize, &attributes, &lastWritten))
 		{
 			fileEntry.casKey = CasKeyZero;
-
 			u32 lastError = GetLastError();
 			if (lastError != ERROR_FILE_NOT_FOUND && lastError != ERROR_PATH_NOT_FOUND)
 				return m_logger.Error(TC("FileExists failed on %s (%s)"), fileName, LastErrorToText(lastError).data);
@@ -1897,6 +1897,13 @@ namespace uba
 			return true;
 		}
 
+		if (IsDirectory(attributes))
+		{
+			fileEntry.casKey = CasKeyZero;
+			out = CasKeyZero;
+			return true;
+		}
+
 		if (fileEntry.casKey != CasKeyZero)
 		{
 			if (casKeyOverride != CasKeyZero && casKeyOverride != fileEntry.casKey)
