[UBA]
* Add more logging for ubacacheserver to try to figure out why cas files that are inputs to cache entries are deleted.

diff --git a/Engine/Source/Programs/UnrealBuildAccelerator/Common/Private/UbaCacheServer.cpp b/Engine/Source/Programs/UnrealBuildAccelerator/Common/Private/UbaCacheServer.cpp
--- a/Engine/Source/Programs/UnrealBuildAccelerator/Common/Private/UbaCacheServer.cpp
+++ b/Engine/Source/Programs/UnrealBuildAccelerator/Common/Private/UbaCacheServer.cpp
@@ -455,6 +455,7 @@ namespace uba
 		UnorderedMap<CasKey, CasFileInfo> existingCas;
 		ReaderWriterLock existingCasLock;
 
+		u32 removedNonExisting = 0;
 		{
 			// TODO: Make this cleaner... 
 			SCOPED_WRITE_LOCK(m_storage.m_casLookupLock, lookupLock);
@@ -465,6 +466,7 @@ namespace uba
 			{
 				if (i->second.verified && !i->second.exists)
 				{
+					++removedNonExisting;
 					i = m_storage.m_casLookup.erase(i);
 					e = m_storage.m_casLookup.end();
 					continue;
@@ -475,7 +477,10 @@ namespace uba
 			}
 		}
 
-		m_logger.Detail(TC("  Found %llu cas files (%s)"), existingCas.size(), BytesToText(totalCasSize).str);
+		if (removedNonExisting)
+			m_logger.Detail(TC("  Removed %u cas entries (marked as not existing)"), removedNonExisting);
+
+		m_logger.Detail(TC("  Found %llu cas files and %llu deleted by overflow (%s)"), existingCas.size(), deletedCasFiles.size(), BytesToText(totalCasSize).str);
 		u64 totalCasCount = existingCas.size() + deletedCasCount;
 
 		if (shouldExit())
