[UBA]
* Added support for explicit id for root paths. This will make it easier to make things backward compatible

diff --git a/Engine/Source/Programs/UnrealBuildAccelerator/Common/Private/UbaCacheClient.cpp b/Engine/Source/Programs/UnrealBuildAccelerator/Common/Private/UbaCacheClient.cpp
--- a/Engine/Source/Programs/UnrealBuildAccelerator/Common/Private/UbaCacheClient.cpp
+++ b/Engine/Source/Programs/UnrealBuildAccelerator/Common/Private/UbaCacheClient.cpp
@@ -490,16 +490,18 @@ namespace uba
 							if (u64 toWrite = rootOffset - lastWritten)
 								memcpy(localBlock.Allocate(toWrite, 1, TC("")), fileStart + lastWritten, toWrite);
 							u8 rootIndex = fileStart[rootOffset] - RootPaths::RootStartByte;
-							auto& root = rootPaths.GetRoot(rootIndex);
+							const TString& root = rootPaths.GetRoot(rootIndex);
+							if (root.empty())
+								return logger.Error(TC("Cache entry uses root path index %u which is not set for this startupinfo (%s)"), rootIndex, info.description);
 
 							#if PLATFORM_WINDOWS
 							StringBuffer<> pathTemp;
-							pathTemp.Append(root.path);
+							pathTemp.Append(root);
 							char rootPath[512];
 							u32 rootPathLen = pathTemp.Parse(rootPath, sizeof_array(rootPath));
 							#else
-							const char* rootPath = root.path.data();
-							u32 rootPathLen = root.path.size();
+							const char* rootPath = root.data();
+							u32 rootPathLen = root.size();
 							#endif
 
 							if (u32 toWrite = rootPathLen - 1)
@@ -948,10 +950,10 @@ namespace uba
 		UBA_ASSERT(normalizedPath.count);
 
 		u32 rootIndex = normalizedPath[0] - RootPaths::RootStartByte;
-		auto& root = rootPaths.GetRoot(rootIndex);
+		const TString& root = rootPaths.GetRoot(rootIndex);
 
 		StringBuffer<MaxPath> path;
-		outPath.Append(root.path).Append(normalizedPath.data + 1);
+		outPath.Append(root).Append(normalizedPath.data + 1);
 		return true;
 	}
 
diff --git a/Engine/Source/Programs/UnrealBuildAccelerator/Common/Private/UbaRootPaths.cpp b/Engine/Source/Programs/UnrealBuildAccelerator/Common/Private/UbaRootPaths.cpp
--- a/Engine/Source/Programs/UnrealBuildAccelerator/Common/Private/UbaRootPaths.cpp
+++ b/Engine/Source/Programs/UnrealBuildAccelerator/Common/Private/UbaRootPaths.cpp
@@ -9,7 +9,7 @@
 
 namespace uba
 {
-	bool RootPaths::RegisterRoot(Logger& logger, const tchar* rootPath, bool includeInKey)
+	bool RootPaths::RegisterRoot(Logger& logger, const tchar* rootPath, bool includeInKey, u8 id)
 	{
 		// Register rootPath both with single path separators and double path separators on windows because text files store them with double path separators
 		#if PLATFORM_WINDOWS
@@ -28,11 +28,19 @@ namespace uba
 
 		for (const tchar* rp : rootPaths)
 		{
-			u32 index = u32(m_roots.size());
+			u8 index = id;
+			if (index == 0)
+				index = u8(m_roots.size());
 			if (index == '~' - ' ') // This is not really true.. as long as value is under 256 we're good
-				return logger.Error(TC("Too many roots added (%llu)"), index);
+				return logger.Error(TC("Too many roots added (%u)"), index);
+
+			if (index >= m_roots.size())
+				m_roots.resize(index+1);
+
+			auto& root = m_roots[index];
+			if (!root.path.empty())
+				return logger.Error(TC("Root at index %u already added (existing as %s, added as %s)"), root.path.c_str(), rp);
 
-			auto& root = m_roots.emplace_back();
 			root.index = index;
 			root.path = rp;
 
@@ -55,26 +63,28 @@ namespace uba
 		return true;
 	}
 
-	bool RootPaths::RegisterSystemRoots(Logger& logger)
+	bool RootPaths::RegisterSystemRoots(Logger& logger, u8 startId)
 	{
 		#if PLATFORM_WINDOWS
+		u8 id = startId;
+		auto GetId = [&]() { u8 res = id; if (id) id += 2; return res; };
 		StringBuffer<MaxPath> dir;
 		dir.count = GetSystemDirectory(dir.data, dir.capacity);
-		RegisterRoot(logger, dir.EnsureEndsWithSlash().data, false); // Ignore files from here.. we do expect them not to affect the output of a process
+		RegisterRoot(logger, dir.EnsureEndsWithSlash().data, false, GetId()); // Ignore files from here.. we do expect them not to affect the output of a process
 		
 		dir.count = GetEnvironmentVariable(TC("ProgramW6432"), dir.Clear().data, dir.capacity);
-		RegisterRoot(logger, dir.EnsureEndsWithSlash().data, true);
+		RegisterRoot(logger, dir.EnsureEndsWithSlash().data, true, GetId());
 
 		dir.count = GetEnvironmentVariable(TC("ProgramFiles(x86)"), dir.Clear().data, dir.capacity);
-		RegisterRoot(logger, dir.EnsureEndsWithSlash().data, true);
+		RegisterRoot(logger, dir.EnsureEndsWithSlash().data, true, GetId());
 
 		dir.count = GetEnvironmentVariable(TC("ProgramFiles(x86)"), dir.Clear().data, dir.capacity);
-		RegisterRoot(logger, dir.EnsureEndsWithSlash().data, true);
+		RegisterRoot(logger, dir.EnsureEndsWithSlash().data, true, GetId());
 
 		PWSTR path;
 		if (!SUCCEEDED(SHGetKnownFolderPath(FOLDERID_ProgramData, 0, NULL, &path)))
 			return false;
-		RegisterRoot(logger, dir.Clear().Append(path).EnsureEndsWithSlash().data, true);
+		RegisterRoot(logger, dir.Clear().Append(path).EnsureEndsWithSlash().data, true, GetId());
 		CoTaskMemFree(path);
 
 		#else
@@ -106,9 +116,13 @@ namespace uba
 		return nullptr;
 	}
 
-	const RootPaths::Root& RootPaths::GetRoot(u32 index) const
+	static const TString EmptyString;
+
+	const TString& RootPaths::GetRoot(u32 index) const
 	{
-		return m_roots[index];
+		if (m_roots.size() <= index)
+			return EmptyString;
+		return m_roots[index].path;
 	}
 
 	CasKey RootPaths::NormalizeAndHashFile(Logger& logger, const tchar* filename) const
diff --git a/Engine/Source/Programs/UnrealBuildAccelerator/Common/Public/UbaRootPaths.h b/Engine/Source/Programs/UnrealBuildAccelerator/Common/Public/UbaRootPaths.h
--- a/Engine/Source/Programs/UnrealBuildAccelerator/Common/Public/UbaRootPaths.h
+++ b/Engine/Source/Programs/UnrealBuildAccelerator/Common/Public/UbaRootPaths.h
@@ -11,19 +11,19 @@ namespace uba
 	{
 	public:
 
-		bool RegisterRoot(Logger& logger, const tchar* rootPath, bool includeInKey = true);
-		bool RegisterSystemRoots(Logger& logger);
+		bool RegisterRoot(Logger& logger, const tchar* rootPath, bool includeInKey = true, u8 id = 0);
+		bool RegisterSystemRoots(Logger& logger, u8 startId = 0);
 
 		struct Root
 		{
 			TString path;
 			StringKey shortestPathKey;
-			u32 index;
+			u8 index;
 			bool includeInKey;
 		};
 
 		const Root* FindRoot(const StringBufferBase& path) const;
-		const Root& GetRoot(u32 index) const;
+		const TString& GetRoot(u32 index) const;
 
 		template<typename CharType, typename Func>
 		bool NormalizeString(Logger& logger, const CharType* str, u64 strLen, const Func& func, const tchar* hint, const tchar* hint2 = TC("")) const;
diff --git a/Engine/Source/Programs/UnrealBuildAccelerator/Host/Private/UbaExports.cpp b/Engine/Source/Programs/UnrealBuildAccelerator/Host/Private/UbaExports.cpp
--- a/Engine/Source/Programs/UnrealBuildAccelerator/Host/Private/UbaExports.cpp
+++ b/Engine/Source/Programs/UnrealBuildAccelerator/Host/Private/UbaExports.cpp
@@ -417,18 +417,18 @@ extern "C"
 		return new uba::RootPathsWithLogger(writer);
 	}
 
-	bool RootPaths_RegisterRoot(uba::RootPaths* rootPaths, const uba::tchar* path, bool includeInKey)
+	bool RootPaths_RegisterRoot(uba::RootPaths* rootPaths, const uba::tchar* path, bool includeInKey, uba::u8 id)
 	{
 		using namespace uba;
 		auto rp = (uba::RootPathsWithLogger*)rootPaths;
-		return rp->RegisterRoot(rp->logger, path, includeInKey);
+		return rp->RegisterRoot(rp->logger, path, includeInKey, id);
 	}
 
-	bool RootPaths_RegisterSystemRoots(uba::RootPaths* rootPaths)
+	bool RootPaths_RegisterSystemRoots(uba::RootPaths* rootPaths, uba::u8 startId)
 	{
 		using namespace uba;
 		auto rp = (uba::RootPathsWithLogger*)rootPaths;
-		return rp->RegisterSystemRoots(rp->logger);
+		return rp->RegisterSystemRoots(rp->logger, startId);
 	}
 
 	void RootPaths_Destroy(uba::RootPaths* rootPaths)
diff --git a/Engine/Source/Programs/UnrealBuildAccelerator/Host/Public/UbaExports.h b/Engine/Source/Programs/UnrealBuildAccelerator/Host/Public/UbaExports.h
--- a/Engine/Source/Programs/UnrealBuildAccelerator/Host/Public/UbaExports.h
+++ b/Engine/Source/Programs/UnrealBuildAccelerator/Host/Public/UbaExports.h
@@ -61,8 +61,8 @@ extern "C"
 
 	// RootPaths
 	UBA_API uba::RootPaths* RootPaths_Create(uba::LogWriter& writer);
-	UBA_API bool RootPaths_RegisterRoot(uba::RootPaths* rootPaths, const uba::tchar* path, bool includeInKey);
-	UBA_API bool RootPaths_RegisterSystemRoots(uba::RootPaths* rootPaths);
+	UBA_API bool RootPaths_RegisterRoot(uba::RootPaths* rootPaths, const uba::tchar* path, bool includeInKey, uba::u8 id = 0);
+	UBA_API bool RootPaths_RegisterSystemRoots(uba::RootPaths* rootPaths, uba::u8 startId = 0);
 	UBA_API void RootPaths_Destroy(uba::RootPaths* rootPaths);
 
 	// ProcessStartInfo
