[UBA]
* Bugfix for when root id was sent in on windows (where id needs to be incremented by 2)

diff --git a/Engine/Source/Programs/UnrealBuildAccelerator/Common/Private/UbaRootPaths.cpp b/Engine/Source/Programs/UnrealBuildAccelerator/Common/Private/UbaRootPaths.cpp
--- a/Engine/Source/Programs/UnrealBuildAccelerator/Common/Private/UbaRootPaths.cpp
+++ b/Engine/Source/Programs/UnrealBuildAccelerator/Common/Private/UbaRootPaths.cpp
@@ -39,7 +39,7 @@ namespace uba
 
 			auto& root = m_roots[index];
 			if (!root.path.empty())
-				return logger.Error(TC("Root at index %u already added (existing as %s, added as %s)"), root.path.c_str(), rp);
+				return logger.Error(TC("Root at index %u already added (existing as %s, added as %s)"), index, root.path.c_str(), rp);
 
 			root.index = index;
 			root.path = rp;
@@ -59,6 +59,9 @@ namespace uba
 			}
 			else
 				root.shortestPathKey = ToStringKeyNoCheck(root.path.data(), m_shortestRoot);
+
+			if (id)
+				++id;
 		}
 		return true;
 	}
