[UBA]
* Added rootpaths version to bucket id in order to make sure we can change system root paths and still be backward compatible

diff --git a/Engine/Source/Programs/UnrealBuildAccelerator/Common/Private/UbaCacheClient.cpp b/Engine/Source/Programs/UnrealBuildAccelerator/Common/Private/UbaCacheClient.cpp
--- a/Engine/Source/Programs/UnrealBuildAccelerator/Common/Private/UbaCacheClient.cpp
+++ b/Engine/Source/Programs/UnrealBuildAccelerator/Common/Private/UbaCacheClient.cpp
@@ -17,7 +17,7 @@
 
 namespace uba
 {
-	u64 MakeId(u32 bucketId) { return u64(bucketId) | (u64(!CaseInsensitiveFs) << 32); }
+	u64 MakeId(u32 bucketId) { return u64(bucketId) | ((u64(!CaseInsensitiveFs) + (RootPathsVersion << 1)) << 32); }
 
 	struct CacheClient::Bucket
 	{
diff --git a/Engine/Source/Programs/UnrealBuildAccelerator/Common/Private/UbaRootPaths.cpp b/Engine/Source/Programs/UnrealBuildAccelerator/Common/Private/UbaRootPaths.cpp
--- a/Engine/Source/Programs/UnrealBuildAccelerator/Common/Private/UbaRootPaths.cpp
+++ b/Engine/Source/Programs/UnrealBuildAccelerator/Common/Private/UbaRootPaths.cpp
@@ -68,31 +68,46 @@ namespace uba
 
 	bool RootPaths::RegisterSystemRoots(Logger& logger, u8 startId)
 	{
-		#if PLATFORM_WINDOWS
-		u8 id = startId;
-		auto GetId = [&]() { u8 res = id; if (id) id += 2; return res; };
-		StringBuffer<MaxPath> dir;
-		dir.count = GetSystemDirectory(dir.data, dir.capacity);
-		RegisterRoot(logger, dir.EnsureEndsWithSlash().data, false, GetId()); // Ignore files from here.. we do expect them not to affect the output of a process
-		
-		dir.count = GetEnvironmentVariable(TC("ProgramW6432"), dir.Clear().data, dir.capacity);
-		RegisterRoot(logger, dir.EnsureEndsWithSlash().data, true, GetId());
-
-		dir.count = GetEnvironmentVariable(TC("ProgramFiles(x86)"), dir.Clear().data, dir.capacity);
-		RegisterRoot(logger, dir.EnsureEndsWithSlash().data, true, GetId());
+#if PLATFORM_WINDOWS
 
-		dir.count = GetEnvironmentVariable(TC("ProgramFiles(x86)"), dir.Clear().data, dir.capacity);
-		RegisterRoot(logger, dir.EnsureEndsWithSlash().data, true, GetId());
+		static StringBuffer<64> systemDir;
+		static StringBuffer<64> programW6432;
+		static StringBuffer<64> programFiles86;
+		static StringBuffer<64> programData;
 
-		PWSTR path;
-		if (!SUCCEEDED(SHGetKnownFolderPath(FOLDERID_ProgramData, 0, NULL, &path)))
+		bool init = []()
+			{
+				systemDir.count = GetSystemDirectory(systemDir.data, systemDir.capacity);
+				systemDir.EnsureEndsWithSlash();
+				programW6432.count = GetEnvironmentVariable(TC("ProgramW6432"), programW6432.data, programW6432.capacity);
+				programW6432.EnsureEndsWithSlash();
+				programFiles86.count = GetEnvironmentVariable(TC("ProgramFiles(x86)"), programFiles86.data, programFiles86.capacity);
+				programFiles86.EnsureEndsWithSlash();
+
+				PWSTR path;
+				if (!SUCCEEDED(SHGetKnownFolderPath(FOLDERID_ProgramData, 0, NULL, &path)))
+					return false;
+				programData.Append(path).EnsureEndsWithSlash();
+				CoTaskMemFree(path);
+				return true;
+			}();
+
+		if (!init)
 			return false;
-		RegisterRoot(logger, dir.Clear().Append(path).EnsureEndsWithSlash().data, true, GetId());
-		CoTaskMemFree(path);
 
-		#else
+		u8 id = startId;
+		if (!id)
+			id = 80;
+
+		auto GetId = [&]() { u8 res = id; if (id) id += 2; return res; };
+		RegisterRoot(logger, systemDir.data, false, GetId()); // Ignore files from here.. we do expect them not to affect the output of a process
+		RegisterRoot(logger, programW6432.data, true, GetId());
+		RegisterRoot(logger, programFiles86.data, true, GetId());
+		RegisterRoot(logger, programData.data, true, GetId());
+
+#else
 		// no system roots
-		#endif
+#endif
 		return true;
 	}
 
@@ -119,7 +134,7 @@ namespace uba
 		return nullptr;
 	}
 
-	static const TString EmptyString;
+	static TString& EmptyString = *new TString(); // Need to leak to prevent shutdown hangs when running in managed process
 
 	const TString& RootPaths::GetRoot(u32 index) const
 	{
diff --git a/Engine/Source/Programs/UnrealBuildAccelerator/Common/Public/UbaRootPaths.h b/Engine/Source/Programs/UnrealBuildAccelerator/Common/Public/UbaRootPaths.h
--- a/Engine/Source/Programs/UnrealBuildAccelerator/Common/Public/UbaRootPaths.h
+++ b/Engine/Source/Programs/UnrealBuildAccelerator/Common/Public/UbaRootPaths.h
@@ -7,6 +7,8 @@
 
 namespace uba
 {
+	static constexpr u8 RootPathsVersion = 1;
+
 	class RootPaths
 	{
 	public:
