[UBA]
* Added so empty cache buckets are deleted on save

diff --git a/Engine/Source/Programs/UnrealBuildAccelerator/Common/Private/UbaCacheServer.cpp b/Engine/Source/Programs/UnrealBuildAccelerator/Common/Private/UbaCacheServer.cpp
--- a/Engine/Source/Programs/UnrealBuildAccelerator/Common/Private/UbaCacheServer.cpp
+++ b/Engine/Source/Programs/UnrealBuildAccelerator/Common/Private/UbaCacheServer.cpp
@@ -389,6 +389,22 @@ namespace uba
 
 		Atomic<bool> success = true;
 
+		for (auto it=m_buckets.begin(); it!=m_buckets.end();)
+		{
+			Bucket& bucket = it->second;
+			if (!bucket.m_cacheEntryLookup.empty())
+			{
+				++it;
+				continue;
+			}
+
+			StringBuffer<MaxPath> bucketsFile(m_rootDir);
+			bucketsFile.EnsureEndsWithSlash().Append(TC("buckets")).EnsureEndsWithSlash().AppendValue(it->first);
+			DeleteFileW(bucketsFile.data);
+			m_logger.Detail(TC("    Bucket %u was empty. Deleted"), bucket.index);
+			it = m_buckets.erase(it);
+		}
+
 		m_server.ParallelFor(GetBucketWorkerCount(), m_buckets, [&, temp = Vector<u8>()](auto& it) mutable
 			{
 				Bucket& bucket = it->second;
