[UBA]
* Added support for weirdly space-escaped paths found in .d files

diff --git a/Engine/Source/Programs/UnrealBuildAccelerator/Common/Private/UbaRootPaths.cpp b/Engine/Source/Programs/UnrealBuildAccelerator/Common/Private/UbaRootPaths.cpp
--- a/Engine/Source/Programs/UnrealBuildAccelerator/Common/Private/UbaRootPaths.cpp
+++ b/Engine/Source/Programs/UnrealBuildAccelerator/Common/Private/UbaRootPaths.cpp
@@ -14,14 +14,22 @@ namespace uba
 		// Register rootPath both with single path separators and double path separators on windows because text files store them with double path separators
 		#if PLATFORM_WINDOWS
 		StringBuffer<> doubleSlash;
+		StringBuffer<> spaceEscaped;
+		bool hasSpace = false;
 		for (const tchar* it=rootPath; *it; ++it)
 		{
+			if (*it == ' ')
+			{
+				hasSpace = true;
+				spaceEscaped.Append('\\');
+			}
+			spaceEscaped.Append(*it);
 			doubleSlash.Append(*it);
 			if (*it == PathSeparator)
 				doubleSlash.Append(PathSeparator);
 		}
 
-		const tchar* rootPaths[] = { rootPath, doubleSlash.data };
+		const tchar* rootPaths[] = { rootPath, doubleSlash.data, hasSpace ? spaceEscaped.data : TC("") };
 		#else
 		const tchar* rootPaths[] = { rootPath };
 		#endif
@@ -29,37 +37,8 @@ namespace uba
 		for (const tchar* rp : rootPaths)
 		{
 			u8 index = id;
-			if (index == 0)
-				index = u8(m_roots.size());
-			if (index == '~' - ' ') // This is not really true.. as long as value is under 256 we're good
-				return logger.Error(TC("Too many roots added (%u)"), index);
-
-			if (index >= m_roots.size())
-				m_roots.resize(index+1);
-
-			auto& root = m_roots[index];
-			if (!root.path.empty())
-				return logger.Error(TC("Root at index %u already added (existing as %s, added as %s)"), index, root.path.c_str(), rp);
-
-			root.index = index;
-			root.path = rp;
-
-			if (CaseInsensitiveFs)
-				ToLower(root.path.data());
-
-			root.includeInKey = includeInKey;
-
-			m_longestRoot = Max(u32(root.path.size()), m_longestRoot);
-
-			if (!m_shortestRoot || root.path.size() < m_shortestRoot)
-			{
-				m_shortestRoot = u32(root.path.size());
-				for (auto& r : m_roots)
-					r.shortestPathKey = ToStringKeyNoCheck(r.path.data(), m_shortestRoot);
-			}
-			else
-				root.shortestPathKey = ToStringKeyNoCheck(root.path.data(), m_shortestRoot);
-
+			if (!InternalRegisterRoot(logger, rp, includeInKey, index))
+				return false;
 			if (id)
 				++id;
 		}
@@ -75,7 +54,7 @@ namespace uba
 		static StringBuffer<64> programFiles86;
 		static StringBuffer<64> programData;
 
-		bool init = []()
+		static bool init = []()
 			{
 				systemDir.count = GetSystemDirectory(systemDir.data, systemDir.capacity);
 				systemDir.EnsureEndsWithSlash();
@@ -96,9 +75,6 @@ namespace uba
 			return false;
 
 		u8 id = startId;
-		if (!id)
-			id = 80;
-
 		auto GetId = [&]() { u8 res = id; if (id) id += 2; return res; };
 		RegisterRoot(logger, systemDir.data, false, GetId()); // Ignore files from here.. we do expect them not to affect the output of a process
 		RegisterRoot(logger, programW6432.data, true, GetId());
@@ -162,4 +138,41 @@ namespace uba
 		return AsNormalized(ToCasKey(hasher, false), wasNormalized);
 	}
 
+	bool RootPaths::InternalRegisterRoot(Logger& logger, const tchar* rootPath, bool includeInKey, u8 index)
+	{
+		if (index == 0)
+			index = u8(m_roots.size());
+		if (index == '~' - ' ') // This is not really true.. as long as value is under 256 we're good
+			return logger.Error(TC("Too many roots added (%u)"), index);
+
+		if (index >= m_roots.size())
+			m_roots.resize(index+1);
+
+		if (!*rootPath)
+			return true;
+
+		auto& root = m_roots[index];
+		if (!root.path.empty())
+			return logger.Error(TC("Root at index %u already added (existing as %s, added as %s)"), index, root.path.c_str(), rootPath);
+
+		root.index = index;
+		root.path = rootPath;
+
+		if (CaseInsensitiveFs)
+			ToLower(root.path.data());
+
+		root.includeInKey = includeInKey;
+
+		m_longestRoot = Max(u32(root.path.size()), m_longestRoot);
+
+		if (!m_shortestRoot || root.path.size() < m_shortestRoot)
+		{
+			m_shortestRoot = u32(root.path.size());
+			for (auto& r : m_roots)
+				r.shortestPathKey = ToStringKeyNoCheck(r.path.data(), m_shortestRoot);
+		}
+		else
+			root.shortestPathKey = ToStringKeyNoCheck(root.path.data(), m_shortestRoot);
+		return true;
+	}
 }
\ No newline at end of file
diff --git a/Engine/Source/Programs/UnrealBuildAccelerator/Common/Public/UbaRootPaths.h b/Engine/Source/Programs/UnrealBuildAccelerator/Common/Public/UbaRootPaths.h
--- a/Engine/Source/Programs/UnrealBuildAccelerator/Common/Public/UbaRootPaths.h
+++ b/Engine/Source/Programs/UnrealBuildAccelerator/Common/Public/UbaRootPaths.h
@@ -35,6 +35,8 @@ namespace uba
 		static constexpr u8 RootStartByte = ' ';
 
 	private:
+		bool InternalRegisterRoot(Logger& logger, const tchar* rootPath, bool includeInKey, u8 index);
+
 		Vector<Root> m_roots;
 		u32 m_shortestRoot = 0;
 		u32 m_longestRoot = 0;
