[UBA]
* Added code (disabled by default) to enable SIO_LOOPBACK_FAST_PATH on sockets to speed up tcp communication when running cache service locally

diff --git a/Engine/Source/Programs/UnrealBuildAccelerator/Common/Private/UbaNetworkBackendTcp.cpp b/Engine/Source/Programs/UnrealBuildAccelerator/Common/Private/UbaNetworkBackendTcp.cpp
--- a/Engine/Source/Programs/UnrealBuildAccelerator/Common/Private/UbaNetworkBackendTcp.cpp
+++ b/Engine/Source/Programs/UnrealBuildAccelerator/Common/Private/UbaNetworkBackendTcp.cpp
@@ -14,6 +14,7 @@
 #if PLATFORM_WINDOWS
 #include <iphlpapi.h>
 #include <ipifcons.h>
+#include <Mstcpip.h>
 #pragma comment (lib, "Netapi32.lib")
 #pragma comment (lib, "Ws2_32.lib")
 #pragma comment(lib, "IPHLPAPI.lib") // For GetAdaptersInfo
@@ -90,6 +91,7 @@ namespace uba
 	bool SetKeepAlive(Logger& logger, SOCKET socket);
 	bool SetBlocking(Logger& logger, SOCKET socket, bool blocking);
 	bool DisableNagle(Logger& logger, SOCKET socket);
+	bool EnableFastLoopback(Logger& logger, SOCKET socket);
 	bool SendSocket(Logger& logger, SOCKET socket, const void* b, u64 bufferLen);
 	bool RecvSocket(Logger& logger, SOCKET socket, void* b, u32 bufferLen, u32 timeoutMs, const Guid& connection, const tchar* hint1, const tchar* hint2, bool isFirstCall);
 
@@ -404,6 +406,8 @@ namespace uba
 				continue;
 			}
 
+			EnableFastLoopback(logger, clientSocket);
+
 			SCOPED_WRITE_LOCK(m_connectionsLock, lock);
 			auto it = m_connections.emplace(m_connections.end(), logger, clientSocket);
 			auto& conn = *it;
@@ -649,6 +653,9 @@ namespace uba
 		if (!SetKeepAlive(logger, socketFd))
 			return false;
 
+		EnableFastLoopback(logger, socketFd);
+
+
 		socketClose.Cancel();
 
 		SCOPED_WRITE_LOCK(m_connectionsLock, lock);
@@ -708,6 +715,32 @@ namespace uba
 		return true;
 	}
 
+	bool EnableFastLoopback(Logger& logger, SOCKET socket)
+	{
+#if 0 // PLATFORM_WINDOWS // Disabled for now because it seems like it is not supported on windows 11
+		static bool mightBeSupported = true;
+		if (!mightBeSupported)
+			return true;
+		int optionValue = 1;
+		DWORD ret = 0;
+		int status = WSAIoctl(socket, SIO_LOOPBACK_FAST_PATH, &optionValue, sizeof(optionValue), NULL, 0, &ret, 0, 0);
+		if (status != SOCKET_ERROR)
+			return true;
+		u32 lastError = GetLastError();
+		if (lastError == WSAEOPNOTSUPP)
+		{
+			mightBeSupported = false;
+			return true;
+		}
+		#if UBA_DEBUG
+		logger.Warning(TC("WSAIoctl SIO_LOOPBACK_FAST_PATH failed (error: %s)"), LastErrorToText(lastError).data);
+		#endif
+		return false;
+#else
+		return true;
+#endif
+	}
+
 	bool SetKeepAlive(Logger& logger, SOCKET socket)
 	{
 		u32 value = 1;
