[UBA]
* Added stats tracking for get file attributes to get more information on what is stalling on the farm
* Enabled storage and system stats for cache misses

diff --git a/Engine/Source/Programs/UnrealBuildAccelerator/Common/Private/UbaCacheClient.cpp b/Engine/Source/Programs/UnrealBuildAccelerator/Common/Private/UbaCacheClient.cpp
--- a/Engine/Source/Programs/UnrealBuildAccelerator/Common/Private/UbaCacheClient.cpp
+++ b/Engine/Source/Programs/UnrealBuildAccelerator/Common/Private/UbaCacheClient.cpp
@@ -317,11 +317,8 @@ namespace uba
 				cacheStats.testEntry.time -= (cacheStats.fetchCasTable.time + cacheStats.normalizeFile.time);
 				BinaryWriter writer(memory, 0, sizeof_array(memory));
 				cacheStats.Write(writer);
-				if (success)
-				{
-					storageStats.Write(writer);
-					systemStats.Write(writer);
-				}
+				storageStats.Write(writer);
+				systemStats.Write(writer);
 				m_session.GetTrace().CacheEndFetch(fetchId, success, memory, writer.GetPosition());
 			});
 
diff --git a/Engine/Source/Programs/UnrealBuildAccelerator/Common/Private/UbaFile.cpp b/Engine/Source/Programs/UnrealBuildAccelerator/Common/Private/UbaFile.cpp
--- a/Engine/Source/Programs/UnrealBuildAccelerator/Common/Private/UbaFile.cpp
+++ b/Engine/Source/Programs/UnrealBuildAccelerator/Common/Private/UbaFile.cpp
@@ -163,6 +163,7 @@ namespace uba
 
 	bool GetFileInformationByHandle(FileInformation& out, Logger& logger, const tchar* fileName, FileHandle hFile)
 	{
+		ExtendedTimerScope ts(SystemStats::GetCurrent().getFileAttributes);
 #if PLATFORM_WINDOWS
 		BY_HANDLE_FILE_INFORMATION info;
 		if (!::GetFileInformationByHandle(asHANDLE(hFile), &info))
@@ -221,6 +222,7 @@ namespace uba
 
 	bool FileExists(Logger& logger, const tchar* fileName, u64* outSize, u32* outAttributes, u64* lastWriteTime)
 	{
+		ExtendedTimerScope ts(SystemStats::GetCurrent().getFileAttributes);
 #if PLATFORM_WINDOWS
 		MAKE_LONG_FILENAME(fileName);
 		WIN32_FILE_ATTRIBUTE_DATA data;
@@ -632,6 +634,7 @@ namespace uba
 
 	bool GetFileSizeEx(u64& outFileSize, FileHandle hFile)
 	{
+		ExtendedTimerScope ts(SystemStats::GetCurrent().getFileAttributes);
 #if PLATFORM_WINDOWS
 		LARGE_INTEGER lpFileSize;
 		if (!::GetFileSizeEx(asHANDLE(hFile), &lpFileSize))
@@ -654,6 +657,7 @@ namespace uba
 
 	u32 GetFileAttributesW(const tchar* fileName)
 	{
+		ExtendedTimerScope ts(SystemStats::GetCurrent().getFileAttributes);
 #if PLATFORM_WINDOWS
 		MAKE_LONG_FILENAME(fileName);
 		return ::GetFileAttributesW(fileName);
diff --git a/Engine/Source/Programs/UnrealBuildAccelerator/Common/Public/UbaTrace.h b/Engine/Source/Programs/UnrealBuildAccelerator/Common/Public/UbaTrace.h
--- a/Engine/Source/Programs/UnrealBuildAccelerator/Common/Public/UbaTrace.h
+++ b/Engine/Source/Programs/UnrealBuildAccelerator/Common/Public/UbaTrace.h
@@ -61,7 +61,7 @@ namespace uba
 	using Color = u32;
 	inline Color toColor(u8 r, u8 g, u8 b) { return (r << 16) + (g << 8) + b; }
 
-	static constexpr u32 TraceVersion = 28;
+	static constexpr u32 TraceVersion = 29;
 	static constexpr u32 TraceReadCompatibilityVersion = 6;
 
 	class Trace : public WorkTracker
diff --git a/Engine/Source/Programs/UnrealBuildAccelerator/Core/Public/UbaProcessStats.h b/Engine/Source/Programs/UnrealBuildAccelerator/Core/Public/UbaProcessStats.h
--- a/Engine/Source/Programs/UnrealBuildAccelerator/Core/Public/UbaProcessStats.h
+++ b/Engine/Source/Programs/UnrealBuildAccelerator/Core/Public/UbaProcessStats.h
@@ -133,6 +133,7 @@ namespace uba
 		UBA_SYSTEM_STAT(writeFile, 0) \
 		UBA_SYSTEM_STAT(readFile, 0) \
 		UBA_SYSTEM_STAT(setFileInfo, 0) \
+		UBA_SYSTEM_STAT(getFileAttributes, 29) \
 		UBA_SYSTEM_STAT(createFileMapping, 0) \
 		UBA_SYSTEM_STAT(mapViewOfFile, 0) \
 		UBA_SYSTEM_STAT(unmapViewOfFile, 0) \
diff --git a/Engine/Source/Programs/UnrealBuildAccelerator/Visualizer/Private/UbaTraceReader.cpp b/Engine/Source/Programs/UnrealBuildAccelerator/Visualizer/Private/UbaTraceReader.cpp
--- a/Engine/Source/Programs/UnrealBuildAccelerator/Visualizer/Private/UbaTraceReader.cpp
+++ b/Engine/Source/Programs/UnrealBuildAccelerator/Visualizer/Private/UbaTraceReader.cpp
@@ -912,7 +912,7 @@ namespace uba
 			StorageStats storageStats;
 			const u8* dataStart = reader.GetPositionData();
 			cacheStats.Read(reader, out.version);
-			if (success)
+			if (success || out.version >= 29)
 			{
 				storageStats.Read(reader);
 				systemStats.Read(reader, out.version);
diff --git a/Engine/Source/Programs/UnrealBuildAccelerator/Visualizer/Private/UbaVisualizer.cpp b/Engine/Source/Programs/UnrealBuildAccelerator/Visualizer/Private/UbaVisualizer.cpp
--- a/Engine/Source/Programs/UnrealBuildAccelerator/Visualizer/Private/UbaVisualizer.cpp
+++ b/Engine/Source/Programs/UnrealBuildAccelerator/Visualizer/Private/UbaVisualizer.cpp
@@ -1425,18 +1425,18 @@ namespace uba
 					sessionStats.Print(logger, m_traceView.frequency);
 				}
 
-				if (!storageStats.IsEmpty())
+				if (!cacheStats.IsEmpty())
 				{
 					logger.Info(L"");
-					logger.Info(L"  ----------- Storage stats -----------");
-					storageStats.Print(logger, m_traceView.frequency);
+					logger.Info(L"  ------------ Cache stats ------------");
+					cacheStats.Print(logger, m_traceView.frequency);
 				}
 
-				if (!cacheStats.IsEmpty())
+				if (!storageStats.IsEmpty())
 				{
 					logger.Info(L"");
-					logger.Info(L"  ------------ Cache stats ------------");
-					cacheStats.Print(logger, m_traceView.frequency);
+					logger.Info(L"  ----------- Storage stats -----------");
+					storageStats.Print(logger, m_traceView.frequency);
 				}
 
 				if (!systemStats.IsEmpty())
