[UBA]
* Optimized cache entry matching
* Added StoreCasKey and changed some StoreCasFile to StoreCasKey in CacheClient which does not need to store cas files
* Changed some parts of the code to use string keys instead of strings since the keys are designed to have correct casing

diff --git a/Engine/Source/Programs/UnrealBuildAccelerator/Common/Private/UbaCacheClient.cpp b/Engine/Source/Programs/UnrealBuildAccelerator/Common/Private/UbaCacheClient.cpp
--- a/Engine/Source/Programs/UnrealBuildAccelerator/Common/Private/UbaCacheClient.cpp
+++ b/Engine/Source/Programs/UnrealBuildAccelerator/Common/Private/UbaCacheClient.cpp
@@ -15,6 +15,8 @@
 #define UBA_LOG_WRITE_CACHE_INFO 0 // 0 = Disabled, 1 = Normal, 2 = Detailed
 #define UBA_LOG_FETCH_CACHE_INFO 0 // 0 = Disabled, 1 = Misses, 2 = Both misses and hits
 
+#define UBA_OLD_TEST 0
+
 namespace uba
 {
 	u64 MakeId(u32 bucketId) { return u64(bucketId) | ((u64(!CaseInsensitiveFs) + (RootPathsVersion << 1)) << 32); }
@@ -232,8 +234,17 @@ namespace uba
 			{
 				bool deferCreation = true;
 				bool fileIsCompressed = IsFileCompressed(info, path);
-				if (!m_storage.StoreCasFile(casKey, path.data, CasKeyZero, deferCreation, fileIsCompressed))
-					return false;
+				if (isOutput)
+				{
+					if (!m_storage.StoreCasFile(casKey, path.data, CasKeyZero, deferCreation, fileIsCompressed))
+						return false;
+				}
+				else
+				{
+					if (!m_storage.StoreCasKey(casKey, path.data, CasKeyZero, fileIsCompressed))
+						return false;
+				}
+
 				if (casKey == CasKeyZero) // If file is not found it was a temporary file that was deleted and is not really an output
 				{
 					//m_logger.Warning(TC("Can't find file %s"), path.data); 
@@ -339,12 +350,12 @@ namespace uba
 				return false;
 		}
 
-
-		UnorderedMap<u32, bool> offsetIsMatch;
-
 		// Traverse entries and test inputs against local machine
-		CacheEntriesTraverser traverser(reader);
-		u32 entryCount = traverser.entryCount;
+#if UBA_OLD_TEST
+		BinaryReader traverserReader(reader.GetPositionData(), 0, reader.GetLeft());
+#endif
+
+		u32 entryCount = reader.ReadU16();
 
 		#if UBA_LOG_FETCH_CACHE_INFO
 		auto mg = MakeGuard([&]()
@@ -354,105 +365,249 @@ namespace uba
 			});
 		#endif
 
+		if (!entryCount)
+			return false;
+
 		struct MissInfo { TString path; u32 entryIndex; CasKey cache; CasKey local; };
 		Vector<MissInfo> misses;
 
 		UnorderedMap<StringKey, CasKey> normalizedCasKeys;
+		UnorderedMap<u32, bool> isCasKeyMatchCache;
 
-		u32 entryIndex = 0;
-		for (; entryIndex!=entryCount; ++entryIndex)
-		{
+		auto IsCasKeyMatch = [&](bool& outIsMatch, u32 casKeyOffset, u32 entryIndex, bool useLookup)
 			{
-				TimerScope ts(cacheStats.testEntry);
-				bool isMatch = true;
+				outIsMatch = false;
+				StringBuffer<MaxPath> path;
 
-				bool result = traverser.TraverseEntryInputs([&](u32 casKeyOffset)
+				CasKey cacheCasKey;
+				CasKey localCasKey;
+
+				bool* cachedIsMatch = nullptr;
+				if (useLookup)
+				{
+					auto insres = isCasKeyMatchCache.try_emplace(casKeyOffset);
+					if (!insres.second)
 					{
-						StringBuffer<MaxPath> path;
+						outIsMatch = insres.first->second;
+						return true;
+					}
+					cachedIsMatch = &insres.first->second;
+				}
 
-						CasKey cacheCasKey;
-						CasKey localCasKey;
+				if (!FetchCasTable(bucket, cacheStats, casKeyOffset))
+					return false;
+
+				if (!GetLocalPathAndCasKey(bucket, rootPaths, path, cacheCasKey, bucket.serverCasKeyTable, bucket.serverPathTable, casKeyOffset))
+					return false;
+				UBA_ASSERTF(IsCompressed(cacheCasKey), TC("Cache entry for %s has uncompressed cache key for path %s (%s)"), info.description, path.data, CasKeyString(cacheCasKey).str);
 
-						auto insres = offsetIsMatch.try_emplace(casKeyOffset);
-						if (insres.second)
+				if (IsNormalized(cacheCasKey)) // Need to normalize caskey for these files since they contain absolute paths
+				{
+					auto insres2 = normalizedCasKeys.try_emplace(ToStringKeyNoCheck(path.data, path.count));
+					if (insres2.second)
+					{
+						TimerScope ts(cacheStats.normalizeFile);
+						localCasKey = rootPaths.NormalizeAndHashFile(m_logger, path.data);
+						if (localCasKey != CasKeyZero)
+							localCasKey = AsCompressed(localCasKey, true);
+						insres2.first->second = localCasKey;
+					}
+					else
+						localCasKey = insres2.first->second;
+
+				}
+				else
+				{
+					bool fileIsCompressed = IsFileCompressed(info, path);
+					m_storage.StoreCasKey(localCasKey, path.data, CasKeyZero, fileIsCompressed);
+					UBA_ASSERT(localCasKey == CasKeyZero || IsCompressed(localCasKey));
+				}
+
+				outIsMatch = localCasKey == cacheCasKey;
+				if (useLookup)
+					*cachedIsMatch = outIsMatch;
+
+				if (!outIsMatch)
+					if (m_reportMissReason && path.count) // if empty this has already been reported
+						misses.push_back({TString(path.data), entryIndex, cacheCasKey, localCasKey });
+				return true;
+			};
+
+
+#if UBA_OLD_TEST
+		bool testIsMatch = false;
+		u32 matchingId = ~0u;
+		{
+			CacheEntriesTraverser traverser(traverserReader);
+			u32 entryIndex = 0;
+			for (; entryIndex!=entryCount; ++entryIndex)
+			{
+				{
+					bool isMatch = true;
+					bool result = traverser.TraverseEntryInputs([&](u32 casKeyOffset)
 						{
-							if (!FetchCasTable(bucket, cacheStats, casKeyOffset))
+							bool entryMatch;
+							if (!IsCasKeyMatch(entryMatch, casKeyOffset, entryIndex, true))
 								return false;
+							isMatch &= entryMatch;
+							return true;
+						});
 
-							if (!GetLocalPathAndCasKey(bucket, rootPaths, path, cacheCasKey, bucket.serverCasKeyTable, bucket.serverPathTable, casKeyOffset))
-								return false;
-							UBA_ASSERTF(IsCompressed(cacheCasKey), TC("Cache entry for %s has uncompressed cache key for path %s (%s)"), info.description, path.data, CasKeyString(cacheCasKey).str);
-
-							if (IsNormalized(cacheCasKey)) // Need to normalize caskey for these files since they contain absolute paths
-							{
-								auto insres2 = normalizedCasKeys.try_emplace(ToStringKeyNoCheck(path.data, path.count));
-								if (insres2.second)
-								{
-									TimerScope ts(cacheStats.normalizeFile);
-									localCasKey = rootPaths.NormalizeAndHashFile(m_logger, path.data);
-									if (localCasKey != CasKeyZero)
-										localCasKey = AsCompressed(localCasKey, true);
-									insres2.first->second = localCasKey;
-								}
-								else
-									localCasKey = insres2.first->second;
-
-							}
-							else
-							{
-								bool deferCreation = true;
-								bool fileIsCompressed = IsFileCompressed(info, path);
-								m_storage.StoreCasFile(localCasKey, path.data, CasKeyZero, deferCreation, fileIsCompressed);
-								UBA_ASSERT(localCasKey == CasKeyZero || IsCompressed(localCasKey));
-							}
-
-							insres.first->second = localCasKey == cacheCasKey;
-						}
+					if (isMatch && !result) // Returned false before setting isMatch, something went wrong
+						return false;
 
-						if (!insres.first->second)
-						{
-							isMatch = false;
-							if (m_reportMissReason && path.count) // if empty this has already been reported
-								misses.push_back({TString(path.data), entryIndex, cacheCasKey, localCasKey });
-							return false;
-						}
-						return true;
-					});
+					// No match, test next entry
+					if (!isMatch)
+					{
+						traverser.SkipEntryOutputs();
+						continue;
+					}
+				}
 
-				if (isMatch && !result) // Returned false before setting isMatch, something went wrong
+				testIsMatch = true;
+				matchingId = traverser.lastId;
+			}
+		}
+#endif
+
+
+
+
+		struct Range
+		{
+			u32 begin;
+			u32 end;
+		};
+		Vector<Range> sharedMatchingRanges;
+
+
+		// Create ranges out of shared offsets that matches local state
+		{
+			TimerScope ts(cacheStats.testEntry);
+			u64 sharedSize = reader.Read7BitEncoded();
+
+			BinaryReader sharedReader(reader.GetPositionData(), 0, sharedSize);
+			reader.Skip(sharedSize);
+
+			u32 rangeBegin = 0;
+
+			auto addRange = [&](u32 rangeEnd)
+				{
+					if (rangeBegin != rangeEnd)
+					{
+						auto& range = sharedMatchingRanges.emplace_back();
+						range.begin = rangeBegin;
+						range.end = rangeEnd;
+					}
+				};
+			while (sharedReader.GetLeft())
+			{
+				u32 position = u32(sharedReader.GetPosition());
+				bool isMatch;
+				if (!IsCasKeyMatch(isMatch, u32(sharedReader.Read7BitEncoded()), 0, false))
 					return false;
 
-				// No match, test next entry
-				if (!isMatch)
+				if (isMatch)
 				{
-					traverser.SkipEntryOutputs();
-					continue;
+					if (rangeBegin != ~0u)
+						continue;
+					rangeBegin = position;
+				}
+				else
+				{
+					if (rangeBegin == ~0u)
+						continue;
+					addRange(position);
+					rangeBegin = ~0u;
 				}
 			}
+			if (rangeBegin != ~0u)
+				addRange(u32(sharedReader.GetPosition()));
+			UBA_ASSERT(!sharedMatchingRanges.empty());
+		}
 
+		// Read entries
+		{
+			--cacheStats.testEntry.count; // Remove the shared one
+
+			BinaryReader entryReader(reader.GetPositionData(), 0, reader.GetLeft());
+			u32 entryIndex = 0;
+			for (; entryIndex!=entryCount; ++entryIndex)
 			{
-				StackBinaryWriter<128> writer;
-				NetworkMessage msg(m_client, CacheServiceId, CacheMessageType_ReportUsedEntry, writer);
-				writer.Write7BitEncoded(MakeId(bucket.id));
-				writer.WriteCasKey(cmdKey);
-				writer.Write7BitEncoded(traverser.lastId);
-				msg.Send();
-			}
+				u32 entryId = u32(reader.Read7BitEncoded());
+				u64 extraSize = reader.Read7BitEncoded();
+				BinaryReader extraReader(reader.GetPositionData(), 0, extraSize);
+				reader.Skip(extraSize);
+				u64 rangeSize = reader.Read7BitEncoded();
+				BinaryReader rangeReader(reader.GetPositionData(), 0, rangeSize);
+				reader.Skip(rangeSize);
+				u64 outSize = reader.Read7BitEncoded();
+				BinaryReader outputsReader(reader.GetPositionData(), 0, outSize);
+				reader.Skip(outSize);
 
-			// Fetch output files from cache (and some files need to be "denormalized" before written to disk
+				{
+					TimerScope ts(cacheStats.testEntry);
 
-			struct DowngradedLogger : public LoggerWithWriter
-			{
-				DowngradedLogger(LogWriter& writer, const tchar* prefix) : LoggerWithWriter(writer, prefix) {}
-				virtual void Log(LogEntryType type, const tchar* str, u32 strLen) override { LoggerWithWriter::Log(Max(type, LogEntryType_Info), str, strLen); }
-			};
+					bool isMatch = true;
+
+					// Check ranges first
+					auto sharedRangeIt = sharedMatchingRanges.begin();
+					while (isMatch && rangeReader.GetLeft())
+					{
+						u64 begin = rangeReader.Read7BitEncoded();
+						u64 end = rangeReader.Read7BitEncoded();
+					
+						Range matchingRange = *sharedRangeIt;
+
+						while (matchingRange.end <= begin)
+						{
+							++sharedRangeIt;
+							if (sharedRangeIt == sharedMatchingRanges.end())
+								break;
+							matchingRange = *sharedRangeIt;
+						}
+
+						isMatch = matchingRange.begin <= begin && matchingRange.end >= end;
+					}
+
+					// Check extra keys after
+					while (isMatch && extraReader.GetLeft())
+						if (!IsCasKeyMatch(isMatch, u32(extraReader.Read7BitEncoded()), entryIndex, true))
+							return false;
+
+					if (!isMatch)
+						continue;
+				}
+
+#if UBA_OLD_TEST
+				UBA_ASSERTF(testIsMatch, TC("%s"), info.description);
+				UBA_ASSERTF(matchingId == entryId, TC("%u vs %u"), matchingId, entryId);
+#endif
 
-			bool result = traverser.TraverseEntryOutputs([&](u32 casKeyOffset)
 				{
+					StackBinaryWriter<128> writer;
+					NetworkMessage msg(m_client, CacheServiceId, CacheMessageType_ReportUsedEntry, writer);
+					writer.Write7BitEncoded(MakeId(bucket.id));
+					writer.WriteCasKey(cmdKey);
+					writer.Write7BitEncoded(entryId);
+					msg.Send();
+				}
+
+				// Fetch output files from cache (and some files need to be "denormalized" before written to disk
+
+				struct DowngradedLogger : public LoggerWithWriter
+				{
+					DowngradedLogger(LogWriter& writer, const tchar* prefix) : LoggerWithWriter(writer, prefix) {}
+					virtual void Log(LogEntryType type, const tchar* str, u32 strLen) override { LoggerWithWriter::Log(Max(type, LogEntryType_Info), str, strLen); }
+				};
+
+				while (outputsReader.GetLeft())
+				{
+					u32 casKeyOffset = u32(outputsReader.Read7BitEncoded());
 					if (!FetchCasTable(bucket, cacheStats, casKeyOffset))
 						return false;
 
-					TimerScope ts(cacheStats.fetchOutput);
+					TimerScope fts(cacheStats.fetchOutput);
 
 					StringBuffer<MaxPath> path;
 					CasKey casKey;
@@ -536,16 +691,16 @@ namespace uba
 						return false;
 					if (!m_session.RegisterNewFile(path.data))
 						return false;
-					return true;
-				});
-
-			if (!result)
-				return false;
-
-			success = true;
-			return true;
+				}
+				success = true;
+				return true;
+			}
 		}
 
+#if UBA_OLD_TEST
+		UBA_ASSERTF(!testIsMatch, TC("%s"), info.description);
+#endif
+
 		for (auto& miss : misses)
 			m_logger.Info(TC("Cache miss on %s because of mismatch of %s (entry: %u, local: %s cache: %s)"), info.description, miss.path.data(), miss.entryIndex, CasKeyString(miss.local).str, CasKeyString(miss.cache).str);
 
@@ -888,9 +1043,8 @@ namespace uba
 		{
 			// Add hash of application binary to key
 			CasKey applicationCasKey;
-			bool deferCreation = true;
 			bool fileIsCompressed = false;
-			if (!m_storage.StoreCasFile(applicationCasKey, info.application, CasKeyZero, deferCreation, fileIsCompressed))
+			if (!m_storage.StoreCasKey(applicationCasKey, info.application, CasKeyZero, fileIsCompressed))
 				return CasKeyZero;
 			hasher.Update(&applicationCasKey, sizeof(CasKey));
 		}
diff --git a/Engine/Source/Programs/UnrealBuildAccelerator/Common/Private/UbaCacheServer.cpp b/Engine/Source/Programs/UnrealBuildAccelerator/Common/Private/UbaCacheServer.cpp
--- a/Engine/Source/Programs/UnrealBuildAccelerator/Common/Private/UbaCacheServer.cpp
+++ b/Engine/Source/Programs/UnrealBuildAccelerator/Common/Private/UbaCacheServer.cpp
@@ -482,6 +482,7 @@ namespace uba
 			{
 				if (i->second.verified && !i->second.exists)
 				{
+					m_storage.DetachEntry(i->second);
 					++removedNonExisting;
 					i = m_storage.m_casLookup.erase(i);
 					e = m_storage.m_casLookup.end();
diff --git a/Engine/Source/Programs/UnrealBuildAccelerator/Common/Private/UbaNetworkClient.cpp b/Engine/Source/Programs/UnrealBuildAccelerator/Common/Private/UbaNetworkClient.cpp
--- a/Engine/Source/Programs/UnrealBuildAccelerator/Common/Private/UbaNetworkClient.cpp
+++ b/Engine/Source/Programs/UnrealBuildAccelerator/Common/Private/UbaNetworkClient.cpp
@@ -291,6 +291,7 @@ namespace uba
 		{
 			SCOPED_WRITE_LOCK(m_connectionsLock, lock2);
 			m_connections.clear();
+			m_connectionsIt = m_connections.end();
 		}
 
 		FlushWork();
diff --git a/Engine/Source/Programs/UnrealBuildAccelerator/Common/Private/UbaSession.cpp b/Engine/Source/Programs/UnrealBuildAccelerator/Common/Private/UbaSession.cpp
--- a/Engine/Source/Programs/UnrealBuildAccelerator/Common/Private/UbaSession.cpp
+++ b/Engine/Source/Programs/UnrealBuildAccelerator/Common/Private/UbaSession.cpp
@@ -1804,7 +1804,7 @@ namespace uba
 		}
 		
 		// if ((message.Access & FileAccess.Write) != 0)
-		m_storage.ReportFileWrite(fileName.data);
+		m_storage.ReportFileWrite(fileNameKey, fileName.data);
 
 		if (m_runningRemote && !fileName.StartsWith(m_tempPath.data))
 		{
diff --git a/Engine/Source/Programs/UnrealBuildAccelerator/Common/Private/UbaStorage.cpp b/Engine/Source/Programs/UnrealBuildAccelerator/Common/Private/UbaStorage.cpp
--- a/Engine/Source/Programs/UnrealBuildAccelerator/Common/Private/UbaStorage.cpp
+++ b/Engine/Source/Programs/UnrealBuildAccelerator/Common/Private/UbaStorage.cpp
@@ -835,7 +835,7 @@ namespace uba
 	}
 
 
-	bool StorageImpl::AddCasFile(const tchar* fileName, const CasKey& casKey, bool deferCreation, bool fileIsCompressed)
+	bool StorageImpl::AddCasFile(StringKey fileNameKey, const tchar* fileName, const CasKey& casKey, bool deferCreation, bool fileIsCompressed)
 	{
 		UBA_ASSERTF(IsCompressed(casKey) == m_storeCompressed, TC("CasKey compress mode must match storage compress mode (%s)"), fileName);
 		SCOPED_WRITE_LOCK(m_casLookupLock, lookupLock);
@@ -872,8 +872,8 @@ namespace uba
 				DeferedCasCreation& dcc = res.first->second;
 				dcc.fileName = fileName;
 				dcc.fileIsCompressed = fileIsCompressed;
-				auto res2 = m_deferredCasCreationLookupByName.try_emplace(dcc.fileName.c_str(), casKey);
-				UBA_ASSERT(res2.second); (void)res2;
+				auto res2 = m_deferredCasCreationLookupByName.try_emplace(fileNameKey, casKey);
+				UBA_ASSERTF(res2.second, TC("%s - existing: %s added: %s"), dcc.fileName.c_str(), CasKeyString(res2.first->second).str, CasKeyString(casKey).str); (void)res2;
 			}
 			return true;
 		}
@@ -1564,6 +1564,9 @@ namespace uba
 					}
 #endif
 
+					if (writer.GetCapacityLeft() < entrySize + sizeof(CasKey))
+						return m_logger.Error(TC("This should not happen, somehow there are more valid entries in access list than lookup. (Lookup has %llu entries)"), m_casLookup.size());
+
 					UBA_ASSERT(entry.key != CasKeyZero);
 					writer.WriteCasKey(entry.key);
 					writer.WriteU64(entry.size);
@@ -1875,7 +1878,7 @@ namespace uba
 			if (fileEntry.casKey != CasKeyZero)
 			{
 				UBA_ASSERT(casKeyOverride == CasKeyZero || casKeyOverride == fileEntry.casKey);
-				if (!AddCasFile(fileName, fileEntry.casKey, deferCreation, fileIsCompressed))
+				if (!AddCasFile(fileNameKey, fileName, fileEntry.casKey, deferCreation, fileIsCompressed))
 					return false;
 			}
 			out = fileEntry.casKey;
@@ -1909,14 +1912,14 @@ namespace uba
 			if (casKeyOverride != CasKeyZero && casKeyOverride != fileEntry.casKey)
 			{
 				fileEntry.casKey = casKeyOverride;
-				if (!AddCasFile(fileName, fileEntry.casKey, deferCreation, fileIsCompressed))
+				if (!AddCasFile(fileNameKey, fileName, fileEntry.casKey, deferCreation, fileIsCompressed))
 					return false;
 				out = fileEntry.casKey;
 				return true;
 			}
 			if (fileSize == fileEntry.size && lastWritten == fileEntry.lastWritten)
 			{
-				if (!AddCasFile(fileName, fileEntry.casKey, deferCreation, fileIsCompressed))
+				if (!AddCasFile(fileNameKey, fileName, fileEntry.casKey, deferCreation, fileIsCompressed))
 					return false;
 				out = fileEntry.casKey;
 				return true;
@@ -1950,7 +1953,96 @@ namespace uba
 		if (fileEntry.casKey == CasKeyZero)
 			return false;
 
-		if (!AddCasFile(fileName, fileEntry.casKey, deferCreation, fileIsCompressed))
+		if (!AddCasFile(fileNameKey, fileName, fileEntry.casKey, deferCreation, fileIsCompressed))
+			return false;
+
+		out = fileEntry.casKey;
+		return true;
+	}
+
+	bool StorageImpl::StoreCasKey(CasKey& out, const tchar* fileName, const CasKey& casKeyOverride, bool fileIsCompressed)
+	{
+		StringBuffer<> forKey;
+		forKey.Append(fileName);
+		if (CaseInsensitiveFs)
+			forKey.MakeLower();
+		StringKey fileNameKey = ToStringKey(forKey);
+
+		SCOPED_WRITE_LOCK(m_fileTableLookupLock, lookupLock);
+		auto insres = m_fileTableLookup.try_emplace(fileNameKey);
+		FileEntry& fileEntry = insres.first->second;
+		lookupLock.Leave();
+
+		SCOPED_WRITE_LOCK(fileEntry.lock, entryLock);
+	
+		if (fileEntry.verified)
+		{
+			out = fileEntry.casKey;
+			return true;
+		}
+		fileEntry.verified = true;
+
+		u64 fileSize = 0;
+		u32 attributes = 0;
+		u64 lastWritten = 0;
+
+		if (!FileExists(m_logger, fileName, &fileSize, &attributes, &lastWritten))
+		{
+			fileEntry.casKey = CasKeyZero;
+			u32 lastError = GetLastError();
+			if (lastError != ERROR_FILE_NOT_FOUND && lastError != ERROR_PATH_NOT_FOUND)
+				return m_logger.Error(TC("FileExists failed on %s (%s)"), fileName, LastErrorToText(lastError).data);
+			out = CasKeyZero;
+			return true;
+		}
+
+		if (IsDirectory(attributes))
+		{
+			fileEntry.casKey = CasKeyZero;
+			out = CasKeyZero;
+			return true;
+		}
+
+		if (fileEntry.casKey != CasKeyZero)
+		{
+			if (casKeyOverride != CasKeyZero && casKeyOverride != fileEntry.casKey)
+			{
+				fileEntry.casKey = casKeyOverride;
+				out = fileEntry.casKey;
+				return true;
+			}
+			if (fileSize == fileEntry.size && lastWritten == fileEntry.lastWritten)
+			{
+				out = fileEntry.casKey;
+				return true;
+			}
+		}
+
+		fileEntry.size = fileSize;
+		fileEntry.lastWritten = lastWritten;
+		if (casKeyOverride == CasKeyZero)
+		{
+			FileHandle fileHandle;
+			if (!OpenFileSequentialRead(m_logger, fileName, fileHandle))
+				return false;
+			auto fileGuard = MakeGuard([&](){ CloseFile(fileName, fileHandle); });
+
+			if (fileIsCompressed)
+			{
+				CompressedObjFileHeader header(CasKeyZero);
+				if (!ReadFile(m_logger, fileName, fileHandle, &header, sizeof(header)))
+					return m_logger.Error(TC("Failed to read header of compressed file %s (%s)"), fileName, LastErrorToText().data);
+				fileEntry.casKey = AsCompressed(header.casKey, m_storeCompressed);
+			}
+			else
+			{
+				fileEntry.casKey = CalculateCasKey(fileName, fileHandle, fileSize, m_storeCompressed);
+			}
+		}
+		else
+			fileEntry.casKey = casKeyOverride;
+
+		if (fileEntry.casKey == CasKeyZero)
 			return false;
 
 		out = fileEntry.casKey;
@@ -1985,9 +2077,10 @@ namespace uba
 		if (findIt == m_deferredCasCreationLookup.end())
 			return false;
 		casEntry.verified = true;
+		StringKey fileNameKey = findIt->second.fileNameKey;
 		StringBuffer<> deferredCreation(findIt->second.fileName);
 		bool fileIsCompressed = findIt->second.fileIsCompressed;
-		m_deferredCasCreationLookupByName.erase(deferredCreation.data);
+		m_deferredCasCreationLookupByName.erase(fileNameKey);
 		m_deferredCasCreationLookup.erase(findIt);
 		deferredLock.Leave();
 		WriteResult res;
@@ -2434,11 +2527,11 @@ namespace uba
 		return true;
 	}
 
-	void StorageImpl::ReportFileWrite(const tchar* fileName)
+	void StorageImpl::ReportFileWrite(StringKey fileNameKey, const tchar* fileName)
 	{
 		// If a defered cas creation is queued up while the source file is about to be modified we need to flush out the cas creation before modifying the file
 		SCOPED_READ_LOCK(m_deferredCasCreationLookupLock, deferredLock);
-		auto findIt = m_deferredCasCreationLookupByName.find(fileName);
+		auto findIt = m_deferredCasCreationLookupByName.find(fileNameKey);
 		if (findIt == m_deferredCasCreationLookupByName.end())
 			return;
 		deferredLock.Leave();
diff --git a/Engine/Source/Programs/UnrealBuildAccelerator/Common/Private/UbaStorageServer.cpp b/Engine/Source/Programs/UnrealBuildAccelerator/Common/Private/UbaStorageServer.cpp
--- a/Engine/Source/Programs/UnrealBuildAccelerator/Common/Private/UbaStorageServer.cpp
+++ b/Engine/Source/Programs/UnrealBuildAccelerator/Common/Private/UbaStorageServer.cpp
@@ -213,7 +213,7 @@ namespace uba
 		externalFileLock.Leave();
 
 		bool fileIsCompressed = false;
-		if (!AddCasFile(fileName, fileEntry.casKey, deferCreation, fileIsCompressed))
+		if (!AddCasFile(fileNameKey, fileName, fileEntry.casKey, deferCreation, fileIsCompressed))
 			return false;
 
 		return true;
@@ -403,9 +403,9 @@ namespace uba
 						// Last resort.. use hint to load file into cas (hint should be renamed since it is now a critical parameter)
 						// We better check the caskey first to make sure it is matching on the server
 
+						StringKey fileNameKey = CaseInsensitiveFs ? ToStringKeyLower(hint) : ToStringKey(hint);
 						CasKey checkedCasKey;
 						{
-							StringKey fileNameKey = CaseInsensitiveFs ? ToStringKeyLower(hint) : ToStringKey(hint);
 							SCOPED_READ_LOCK(m_fileTableLookupLock, lookupLock);
 							auto findIt = m_fileTableLookup.find(fileNameKey);
 							if (findIt != m_fileTableLookup.end())
@@ -437,7 +437,7 @@ namespace uba
 
 						bool deferCreation = false;
 						bool fileIsCompressed = false;
-						if (!AddCasFile(hint.data, casKey, deferCreation, fileIsCompressed))
+						if (!AddCasFile(fileNameKey, hint.data, casKey, deferCreation, fileIsCompressed))
 						{
 							m_logger.Error(TC("FetchBegin failed for cas file %s (%s). Can't add cas file to database"), CasKeyString(casKey).str, hint.data);
 							writer.WriteU16(0);
diff --git a/Engine/Source/Programs/UnrealBuildAccelerator/Common/Public/UbaStorage.h b/Engine/Source/Programs/UnrealBuildAccelerator/Common/Public/UbaStorage.h
--- a/Engine/Source/Programs/UnrealBuildAccelerator/Common/Public/UbaStorage.h
+++ b/Engine/Source/Programs/UnrealBuildAccelerator/Common/Public/UbaStorage.h
@@ -57,7 +57,7 @@ namespace uba
 		virtual MappedView MapView(const CasKey& casKey, const tchar* hint) = 0;
 		virtual void UnmapView(const MappedView& view, const tchar* hint) = 0;
 
-		virtual void ReportFileWrite(const tchar* fileName) = 0;
+		virtual void ReportFileWrite(StringKey fileNameKey, const tchar* fileName) = 0;
 		
 		virtual StorageStats& Stats() = 0;
 		virtual void AddStats(StorageStats& stats) = 0;
@@ -135,7 +135,7 @@ namespace uba
 		virtual MappedView MapView(const CasKey& casKey, const tchar* hint) override;
 		virtual void UnmapView(const MappedView& view, const tchar* hint) override;
 
-		virtual void ReportFileWrite(const tchar* fileName) override;
+		virtual void ReportFileWrite(StringKey fileNameKey, const tchar* fileName) override;
 		
 		virtual StorageStats& Stats() final;
 		virtual void AddStats(StorageStats& stats) override;
@@ -146,6 +146,7 @@ namespace uba
 		bool EnsureCasFile(const CasKey& casKey, const tchar* fileName);
 		CasKey CalculateCasKey(const tchar* fileName, FileHandle fileHandle, u64 fileSize, bool storeCompressed);
 		CasKey CalculateCasKey(u8* fileMem, u64 fileSize, bool storeCompressed);
+		bool StoreCasKey(CasKey& out, const tchar* fileName, const CasKey& casKeyOverride, bool fileIsCompressed);
 
 		virtual bool WriteCompressed(WriteResult& out, const tchar* from, const tchar* toFile);
 		virtual bool WriteCompressed(WriteResult& out, const tchar* from, FileHandle readHandle, u8* readMem, u64 fileSize, const tchar* toFile, const void* header, u64 headerSize) final;
@@ -153,7 +154,7 @@ namespace uba
 		bool WriteCasFileNoCheck(WriteResult& out, const tchar* fileName, bool fileIsCompressed, const tchar* casFile, bool storeCompressed);
 		bool WriteCasFile(WriteResult& out, const tchar* fileName, bool fileIsCompressed, const CasKey& casKey);
 		bool VerifyExisting(bool& outReturnValue, ScopedWriteLock& entryLock, const CasKey& casKey, CasEntry& casEntry, StringBufferBase& casFile);
-		bool AddCasFile(const tchar* fileName, const CasKey& casKey, bool deferCreation, bool fileIsCompressed);
+		bool AddCasFile(StringKey fileNameKey, const tchar* fileName, const CasKey& casKey, bool deferCreation, bool fileIsCompressed);
 		void CasEntryAccessed(const CasKey& casKey);
 		virtual bool IsDisallowedPath(const tchar* fileName);
 		virtual bool DecompressMemoryToMemory(u8* compressedData, u8* writeData, u64 decompressedSize, const tchar* readHint) override;
@@ -227,9 +228,9 @@ namespace uba
 		FileMappingBuffer m_casDataBuffer;
 
 		ReaderWriterLock m_deferredCasCreationLookupLock;
-		struct DeferedCasCreation { TString fileName; bool fileIsCompressed; };
+		struct DeferedCasCreation { StringKey fileNameKey; TString fileName; bool fileIsCompressed; };
 		UnorderedMap<CasKey, DeferedCasCreation> m_deferredCasCreationLookup;
-		UnorderedMap<const tchar*, CasKey, HashString, EqualString> m_deferredCasCreationLookupByName;
+		UnorderedMap<StringKey, CasKey> m_deferredCasCreationLookupByName;
 
 		DirectoryCache m_dirCache;
 
