[UBA]
* Fixed compile error in test code

diff --git a/Engine/Source/Programs/UnrealBuildAccelerator/Test/Private/UbaTestBasics.h b/Engine/Source/Programs/UnrealBuildAccelerator/Test/Private/UbaTestBasics.h
--- a/Engine/Source/Programs/UnrealBuildAccelerator/Test/Private/UbaTestBasics.h
+++ b/Engine/Source/Programs/UnrealBuildAccelerator/Test/Private/UbaTestBasics.h
@@ -420,7 +420,7 @@ namespace uba
 
 		StringBuffer<> newStr;
 		auto& root = paths.GetRoot(rootPos - RootPaths::RootStartByte);
-		newStr.Append(root.path).Append(temp);
+		newStr.Append(root.c_str()).Append(temp);
 		if (!newStr.Equals(str))
 			return false;
 
