[UBA]
* Fixed bug in visualizer for selecting times on the timeline

diff --git a/Engine/Source/Programs/UnrealBuildAccelerator/Visualizer/Private/UbaVisualizer.cpp b/Engine/Source/Programs/UnrealBuildAccelerator/Visualizer/Private/UbaVisualizer.cpp
--- a/Engine/Source/Programs/UnrealBuildAccelerator/Visualizer/Private/UbaVisualizer.cpp
+++ b/Engine/Source/Programs/UnrealBuildAccelerator/Visualizer/Private/UbaVisualizer.cpp
@@ -2575,11 +2575,12 @@ namespace uba
 			{
 				if (!m_fileName.IsEmpty()) // Only works for files right now
 				{
+					float timelineSelected = m_timelineSelected;
 					Reset();
 					if (!m_trace.ReadFile(m_traceView, m_fileName.data, true))
 						return false;
 					bool changed;
-					u64 time = MsToTime(u64(m_timelineSelected * 1000.0));
+					u64 time = MsToTime(u64(timelineSelected * 1000.0));
 					m_traceView.finished = false;
 					m_trace.UpdateReadFile(m_traceView, time, changed);
 					m_pauseStart = m_startTime + time;
@@ -2587,6 +2588,7 @@ namespace uba
 					{
 						m_pauseTime = 0;
 						m_paused = true;
+						m_replay = 1;
 					}
 					else
 					{
@@ -2700,6 +2702,7 @@ namespace uba
 
 			case Popup_Replay:
 				m_replay = 1;
+				m_paused = false;
 				PostMessage(m_hwnd, WM_NEWTRACE, 0, 0);
 				break;
 
