[UBA]
* Bumped directory table mem size from 40mb to 48mb to be able to handle builds that includes tons and tons of targets

diff --git a/Engine/Source/Programs/UnrealBuildAccelerator/Core/Public/UbaProtocol.h b/Engine/Source/Programs/UnrealBuildAccelerator/Core/Public/UbaProtocol.h
--- a/Engine/Source/Programs/UnrealBuildAccelerator/Core/Public/UbaProtocol.h
+++ b/Engine/Source/Programs/UnrealBuildAccelerator/Core/Public/UbaProtocol.h
@@ -47,7 +47,7 @@ namespace uba
 	inline constexpr u32 CommunicationMemSize = IsWindows ? 64*1024 : 64*1024*2; // Macos expands some commandlines to be crazy long
 
 	inline constexpr u32 FileMappingTableMemSize = 16 * 1024 * 1024;
-	inline constexpr u32 DirTableMemSize = 40 * 1024 * 1024;
+	inline constexpr u32 DirTableMemSize = 48 * 1024 * 1024;
 }
 
 // Currently only used for detoured process
