[UBA]
* Changed so posix targets use clock_gettime(CLOCK_MONOTONIC, x) instead of gettimeofday in order to make sure time is always showing increasing number (and I guess it is more accurate)

diff --git a/Engine/Source/Programs/UnrealBuildAccelerator/Core/Private/UbaEvent.cpp b/Engine/Source/Programs/UnrealBuildAccelerator/Core/Private/UbaEvent.cpp
--- a/Engine/Source/Programs/UnrealBuildAccelerator/Core/Private/UbaEvent.cpp
+++ b/Engine/Source/Programs/UnrealBuildAccelerator/Core/Private/UbaEvent.cpp
@@ -10,9 +10,30 @@
 #include <new>
 #endif
 
+#define UBA_TEST_WAIT_QUALITY 0
+#define UBA_USE_TIME_OF_DAY_CONDITION 1//PLATFORM_MAC
+
 namespace uba
 {
 #if !PLATFORM_WINDOWS
+
+	u64 GetMonoticTimeNs();
+
+	#if UBA_USE_TIME_OF_DAY_CONDITION
+	static u64 GetTimeOfDayOffset = []()
+		{
+			timeval tv;
+			gettimeofday(&tv, NULL);
+				
+			u64 tsNs = GetMonoticTimeNs();
+			u64 tvNs = u64(tv.tv_sec)*1'000'000'000 + tv.tv_usec * 1'000;
+
+			return tvNs - tsNs;
+		}();
+	#else
+	static constexpr u64 GetTimeOfDayOffset = 0;
+	#endif
+
 	struct EventImpl
 	{
 		EventImpl() = default;
@@ -64,6 +85,14 @@ namespace uba
 				return false;
 			}
 
+			#if !UBA_USE_TIME_OF_DAY_CONDITION
+			if (pthread_condattr_setclock(&attrcond, CLOCK_MONOTONIC) != 0)
+			{
+				UBA_ASSERTF(false, "pthread_condattr_setclock failed");
+				return false;
+			}
+			#endif
+
 			if (shared)
 			{
 				if (pthread_condattr_setpshared(&attrcond, PTHREAD_PROCESS_SHARED) != 0)
@@ -139,18 +168,49 @@ namespace uba
 			UnlockEventMutex();
 		}
 
+		inline u64 ToNanoSeconds(const struct timespec& ts)
+		{
+			return u64(ts.tv_sec)*1'000'000'000 + ts.tv_nsec;
+		}
+
+		inline struct timespec ToTimeSpec(u64 nanoSeconds)
+		{
+			struct timespec ts;
+			ts.tv_sec = nanoSeconds / 1'000'000'000;
+			ts.tv_nsec = nanoSeconds % 1'000'000'000;
+			return ts;
+		}
+
 		bool IsSet(u32 timeoutMs = ~0u)
 		{
 			if (!m_initialized)
 				return false;
 
-			struct timeval startTime;
+			u64 startTimeNs = 0;
 
 			// We need to know the start time if we're going to do a timed wait.
 			if ((timeoutMs > 0) && (timeoutMs != ~0u))  // not polling and not infinite wait.
-				gettimeofday(&startTime, NULL);
+				startTimeNs = GetMonoticTimeNs();
+
+			u64 timeoutNs = u64(timeoutMs) * 1'000'000;
+
+			#if UBA_TEST_WAIT_QUALITY
+			u32 loop = 0;
+			bool timedOut = false;
+			u64 initialStartTimeNs = startTimeNs;
+			u64 initialTimeoutNs = timeoutNs;
+			auto qg = MakeGuard([&]()
+				{
+					if (startTimeNs && timedOut)
+					{
+						u64 overtimeNs = GetMonoticTimeNs() - initialStartTimeNs - initialTimeoutNs;
+						printf("Loops: %u TimeOut: %ums Over-wait: %lluus\n", loop, timeoutMs, overtimeNs/1000);
+					}
+				});
+			#endif
 
 			LockEventMutex();
+			auto unlock = MakeGuard([this]() { UnlockEventMutex(); });
 
 			// loop in case we fall through the Condition signal but someone else claims the event.
 			do
@@ -158,18 +218,16 @@ namespace uba
 				if (m_triggered == TriggerType_One)
 				{
 					m_triggered = TriggerType_None;
-					UnlockEventMutex();
 					return true;
 				}
 
 				if (m_triggered == TriggerType_All)
 				{
-					UnlockEventMutex();
 					return true;
 				}
 
 				// No event signalled yet.
-				if (timeoutMs != 0)  // not just polling, wait on the condition variable.
+				if (timeoutNs != 0)  // not just polling, wait on the condition variable.
 				{
 					++m_waitingThreads;
 					if (timeoutMs == ~0u) // infinite wait?
@@ -179,28 +237,29 @@ namespace uba
 					}
 					else  // timed wait.
 					{
-						struct timespec TimeOut;
-						const int ms = int(startTime.tv_usec / 1000) + int(timeoutMs);
-						TimeOut.tv_sec = (startTime.tv_sec) + (ms / 1000);
-						TimeOut.tv_nsec = (ms % 1000) * 1000000;  // remainder of milliseconds converted to nanoseconds.
-						int rc = pthread_cond_timedwait(&m_condition, &m_mutex, &TimeOut);    // unlocks Mutex while blocking...
+						struct timespec timeout = ToTimeSpec(startTimeNs + timeoutNs + GetTimeOfDayOffset);
+						int rc = pthread_cond_timedwait(&m_condition, &m_mutex, &timeout);    // unlocks Mutex while blocking...
 						UBA_ASSERTF((rc == 0) || (rc == ETIMEDOUT), "pthread_cond_timedwait failed"); (void)rc;
 
-						// Update timeoutMs and startTime in case we have to go again...
-						struct timeval now, difference;
-						gettimeofday(&now, NULL);
-						SubtractTimevals(&now, &startTime, &difference);
-						const u32 differenceMS = ((u32(difference.tv_sec) * 1000) + (u32(difference.tv_usec) / 1000));
-						timeoutMs = ((differenceMS >= timeoutMs) ? 0 : (timeoutMs - differenceMS));
-						startTime = now;
+						#if UBA_TEST_WAIT_QUALITY
+						++loop;
+						#endif
+
+						u64 nowNs = GetMonoticTimeNs();
+						u64 diffNs = nowNs - startTimeNs;
+
+						timeoutNs = diffNs >= timeoutNs ? 0 : (timeoutNs - diffNs);
+						startTimeNs = nowNs;
 					}
 					--m_waitingThreads;
 					UBA_ASSERTF(m_waitingThreads >= 0, "m_waitingThreads less than 0");
 				}
 
-			} while (timeoutMs != 0);
+			} while (timeoutNs != 0);
 
-			UnlockEventMutex();
+			#if UBA_TEST_WAIT_QUALITY
+			timedOut = true;
+			#endif
 			return false;
 		}
 
@@ -225,26 +284,6 @@ namespace uba
 			int res = pthread_mutex_unlock(&m_mutex);
 			UBA_ASSERTF(res == 0, "pthread_mutex_unlock failed (error code %i)", res);
 		}
-
-		static inline void SubtractTimevals(const struct timeval* FromThis, struct timeval* SubThis, struct timeval* Difference)
-		{
-			if (FromThis->tv_usec < SubThis->tv_usec)
-			{
-				int nsec = int((SubThis->tv_usec - FromThis->tv_usec) / 1000000) + 1;
-				SubThis->tv_usec -= 1000000 * nsec;
-				SubThis->tv_sec += nsec;
-			}
-
-			if (FromThis->tv_usec - SubThis->tv_usec > 1000000)
-			{
-				int nsec = int((FromThis->tv_usec - SubThis->tv_usec) / 1000000);
-				SubThis->tv_usec += 1000000 * nsec;
-				SubThis->tv_sec -= nsec;
-			}
-
-			Difference->tv_sec = FromThis->tv_sec - SubThis->tv_sec;
-			Difference->tv_usec = FromThis->tv_usec - SubThis->tv_usec;
-		}
 	};
 #endif
 
diff --git a/Engine/Source/Programs/UnrealBuildAccelerator/Core/Private/UbaMemory.cpp b/Engine/Source/Programs/UnrealBuildAccelerator/Core/Private/UbaMemory.cpp
--- a/Engine/Source/Programs/UnrealBuildAccelerator/Core/Private/UbaMemory.cpp
+++ b/Engine/Source/Programs/UnrealBuildAccelerator/Core/Private/UbaMemory.cpp
@@ -5,9 +5,6 @@
 
 namespace uba
 {
-	ANALYSIS_NORETURN void FatalError(u32 code, const tchar* format, ...);
-
-
 	MemoryBlock::MemoryBlock(u64 reserveSize_, void* baseAddress_)
 	{
 		Init(reserveSize_, baseAddress_);
diff --git a/Engine/Source/Programs/UnrealBuildAccelerator/Core/Private/UbaTimer.cpp b/Engine/Source/Programs/UnrealBuildAccelerator/Core/Private/UbaTimer.cpp
--- a/Engine/Source/Programs/UnrealBuildAccelerator/Core/Private/UbaTimer.cpp
+++ b/Engine/Source/Programs/UnrealBuildAccelerator/Core/Private/UbaTimer.cpp
@@ -8,34 +8,47 @@
 #include <sys/time.h>
 #endif
 
+#define UBA_USE_GETTIMEOFDAY 0
+
 namespace uba
 {
+	#if !PLATFORM_WINDOWS
+	u64 GetMonoticTimeNs()
+	{
+		struct timespec ts;
+		if (clock_gettime(CLOCK_MONOTONIC, &ts) == -1)
+			FatalError(1401, TC("clock_gettime(CLOCK_MONOTONIC) failed"));
+		return u64(ts.tv_sec * 1'000'000'000LL + ts.tv_nsec);
+	}
+	#endif
+
 	u64 GetTime()
 	{
-		#if PLATFORM_WINDOWS
+	#if PLATFORM_WINDOWS
 		LARGE_INTEGER li;
 		QueryPerformanceCounter(&li);
 		return li.QuadPart;
-		//#elif PLATFORM_LINUX
-		//struct timespec tp;
-		//clock_gettime(CLOCK_REALTIME, &tp);
-		//u64 result = u64(tp.tv_sec * 10'000'000LL + tp.tv_nsec/100); // Stored as a 10th of a microsecond
-		//return result;
-		#else
+	#elif UBA_USE_GETTIMEOFDAY
 		timeval tv;
 		gettimeofday(&tv, NULL); // Returns time in microseconds since 1 Jan 1970
 		return u64(tv.tv_sec) * 1'000'000ull + u64(tv.tv_usec);
-		#endif
+	#else
+		return GetMonoticTimeNs();
+	#endif
 	}
 
 	u64 GetFrequency()
 	{
-		#if PLATFORM_WINDOWS
+	#if PLATFORM_WINDOWS
 		static u64 frequency = []() { LARGE_INTEGER li; QueryPerformanceFrequency(&li); return li.QuadPart; }();
 		return frequency;
-		#else
+	#elif UBA_USE_GETTIMEOFDAY
 		return 1000000LL;
-		#endif
+	#else
+		return 1'000'000'000LL;
+		//static u64 frequency = []() { timespec ts; clock_getres(CLOCK_MONOTONIC, &ts); return u64(ts.tv_sec * 1'000'000'000LL + ts.tv_nsec); }();
+		//return frequency;
+	#endif
 	}
 
 	u64 GetSystemTimeUs()
@@ -46,7 +59,9 @@ namespace uba
 		GetSystemTimeAsFileTime(&st);
 		return *(u64*)&st / 10 - (EPOCH_DIFF*1'000'000ull);
 		#else
-		return GetTime();
+		timeval tv;
+		gettimeofday(&tv, NULL); // Returns time in microseconds since 1 Jan 1970
+		return u64(tv.tv_sec) * 1'000'000ull + u64(tv.tv_usec);
 		#endif
 	}
 
diff --git a/Engine/Source/Programs/UnrealBuildAccelerator/Core/Public/UbaPlatform.h b/Engine/Source/Programs/UnrealBuildAccelerator/Core/Public/UbaPlatform.h
--- a/Engine/Source/Programs/UnrealBuildAccelerator/Core/Public/UbaPlatform.h
+++ b/Engine/Source/Programs/UnrealBuildAccelerator/Core/Public/UbaPlatform.h
@@ -72,6 +72,7 @@ namespace uba
 
 	void WriteAssertInfo(class StringBufferBase& out, const tchar* text, const char* file, u32 line, const char* expr, u32 skipCallstack);
 	ANALYSIS_NORETURN void UbaAssert(const tchar* text, const char* file, u32 line, const char* expr, u32 terminateCode);
+	ANALYSIS_NORETURN void FatalError(u32 code, const tchar* format, ...);
 	using CustomAssertHandler = void(const tchar* text);
 	void SetCustomAssertHandler(CustomAssertHandler* handler);
 	bool CreateGuid(Guid& out);
diff --git a/Engine/Source/Programs/UnrealBuildAccelerator/Detours/Private/UbaDetoursShared.h b/Engine/Source/Programs/UnrealBuildAccelerator/Detours/Private/UbaDetoursShared.h
--- a/Engine/Source/Programs/UnrealBuildAccelerator/Detours/Private/UbaDetoursShared.h
+++ b/Engine/Source/Programs/UnrealBuildAccelerator/Detours/Private/UbaDetoursShared.h
@@ -17,8 +17,6 @@ namespace uba
 	DWORD Local_GetLongPathNameW(LPCWSTR lpszShortPath, LPWSTR lpszLongPath, DWORD cchBuffer);
 	#endif
 
-	ANALYSIS_NORETURN void FatalError(u32 code, const tchar* format, ...);
-
 	void Rpc_WriteLog(const tchar* text, u64 textCharLength, bool printInSession, bool isError);
 	void Rpc_WriteLogf(const tchar* format, ...);
 
diff --git a/Engine/Source/Programs/UnrealBuildAccelerator/Test/Private/UbaTestBasics.h b/Engine/Source/Programs/UnrealBuildAccelerator/Test/Private/UbaTestBasics.h
--- a/Engine/Source/Programs/UnrealBuildAccelerator/Test/Private/UbaTestBasics.h
+++ b/Engine/Source/Programs/UnrealBuildAccelerator/Test/Private/UbaTestBasics.h
@@ -54,29 +54,47 @@ namespace uba
 
 	bool TestEvents(Logger& logger, const StringBufferBase& rootDir)
 	{
-		Event ev(true);
-		Thread t([&]()
-			{
-				Sleep(500);
-				//logger.Info(TC("Setting event"));
-				ev.Set();
-				Sleep(500);
-				return true;
-			});
-
-		if (ev.IsSet(0))
-			return false;
-
-		//logger.Info(TC("Waiting for event"));
-		if (!ev.IsSet(10000))
-			return false;
-		//logger.Info(TC("Event was set"));
+		for (uint i=0; i!=2; ++i)
+		{
+			Event ev;
+			if (!ev.Create(true, i == 1))
+				return logger.Error(TC("Failed to create event"));
 
-		if (t.Wait(0))
-			return false;
+			Thread t([&]()
+				{
+					Sleep(500);
+					//logger.Info(TC("Setting event"));
+					ev.Set();
+					Sleep(500);
+					return true;
+				});
+
+			if (ev.IsSet(1))
+				return logger.Error(TC("Event was set after 1ms timeout where it should take 500ms"));
+
+			if (ev.IsSet(0))
+				return logger.Error(TC("Event was set after no timeout where it should take 500ms"));
+
+			//logger.Info(TC("Waiting for event"));
+			if (!ev.IsSet(2000))
+				return logger.Error(TC("Event was not set after 2000ms where it should take 500ms"));
+			//logger.Info(TC("Event was set"));
+
+			if (t.Wait(0))
+				return logger.Error(TC("Thread wait timed out. Should already be done after 2000ms"));
+
+			if (!t.Wait(2000))
+				return logger.Error(TC("Thread wait did not timed out should be done after 2000ms"));
+
+			#if 0 // Long time test... disabled by default
+			u64 time = GetTime();
+			Event ev2(true);
+			ev2.IsSet(10 * 60 * 1000);
+			if (TimeToMs((GetTime() - time) < 9 * 60 * 1000))
+				return logger.Error(TC("Event timeout was way too fast"));
+			#endif
+		}
 
-		if (!t.Wait(2000))
-			return false;
 		return true;
 	}
 
