[UBA]
* Added ToStringKeyLower that takes in char* and length

diff --git a/Engine/Source/Programs/UnrealBuildAccelerator/Core/Private/UbaHash.cpp b/Engine/Source/Programs/UnrealBuildAccelerator/Core/Private/UbaHash.cpp
--- a/Engine/Source/Programs/UnrealBuildAccelerator/Core/Private/UbaHash.cpp
+++ b/Engine/Source/Programs/UnrealBuildAccelerator/Core/Private/UbaHash.cpp
@@ -34,12 +34,17 @@ namespace uba
 		return ToStringKey(b.data, b.count);
 	}
 
-	StringKey ToStringKeyLower(const StringBufferBase& b)
+	StringKey ToStringKeyLower(const tchar* str, u64 strLen)
 	{
 		StringBuffer<> temp;
-		temp.Append(b).MakeLower();
+		temp.Append(str, strLen).MakeLower();
 		return ToStringKey(temp.data, temp.count);
 	}
+
+	StringKey ToStringKeyLower(const StringBufferBase& b)
+	{
+		return ToStringKeyLower(b.data, b.count);
+	}
 	StringKey ToStringKey(const StringKeyHasher& hasher, const tchar* str, u64 strLen)
 	{
 		CHECK_PATH(str);
diff --git a/Engine/Source/Programs/UnrealBuildAccelerator/Core/Public/UbaHash.h b/Engine/Source/Programs/UnrealBuildAccelerator/Core/Public/UbaHash.h
--- a/Engine/Source/Programs/UnrealBuildAccelerator/Core/Public/UbaHash.h
+++ b/Engine/Source/Programs/UnrealBuildAccelerator/Core/Public/UbaHash.h
@@ -100,6 +100,7 @@ namespace uba
 	};
 
 	StringKey ToStringKey(const tchar* str, u64 strLen);
+	StringKey ToStringKeyLower(const tchar* str, u64 strLen);
 	StringKey ToStringKey(const StringBufferBase& b);
 	StringKey ToStringKeyLower(const StringBufferBase& b);
 	StringKey ToStringKey(const StringKeyHasher& hasher, const tchar* str, u64 strLen);
