[UBA]
* Renamed system stats to kernel stats
* Changed all stats storing to use bitfield first to say which fields that are non-zero

diff --git a/Engine/Source/Programs/UnrealBuildAccelerator/Agent/Private/UbaAgent.cpp b/Engine/Source/Programs/UnrealBuildAccelerator/Agent/Private/UbaAgent.cpp
--- a/Engine/Source/Programs/UnrealBuildAccelerator/Agent/Private/UbaAgent.cpp
+++ b/Engine/Source/Programs/UnrealBuildAccelerator/Agent/Private/UbaAgent.cpp
@@ -1338,7 +1338,7 @@ namespace uba
 					sessionClient->PrintSummary(logger);
 					storageClient->PrintSummary(logger);
 					client->PrintSummary(logger);
-					SystemStats::GetGlobal().Print(logger, true);
+					KernelStats::GetGlobal().Print(logger, true);
 				}
 
 				logger.Info(TC("----------- Session %s done! -----------"), sessionClient->GetId());
diff --git a/Engine/Source/Programs/UnrealBuildAccelerator/Cli/Private/UbaCli.cpp b/Engine/Source/Programs/UnrealBuildAccelerator/Cli/Private/UbaCli.cpp
--- a/Engine/Source/Programs/UnrealBuildAccelerator/Cli/Private/UbaCli.cpp
+++ b/Engine/Source/Programs/UnrealBuildAccelerator/Cli/Private/UbaCli.cpp
@@ -591,7 +591,7 @@ namespace uba
 		//info.shouldWriteToDisk = shouldWriteToDisk;
 		info.rootDir = g_rootDir.data;
 		//info.traceName.Append(TC("TESTTRACE"));
-		//info.storeObjFilesCompressed = true;
+		info.storeObjFilesCompressed = true;
 		#if UBA_DEBUG_LOG_ENABLED
 		info.remoteLogEnabled = true;
 		#endif
@@ -891,7 +891,7 @@ namespace uba
 			sessionServer.PrintSummary(logger);
 			storageServer.PrintSummary(logger);
 			networkServer.PrintSummary(logger);
-			SystemStats::GetGlobal().Print(logger, true);
+			KernelStats::GetGlobal().Print(logger, true);
 		}
 		logger.EndScope();
 
diff --git a/Engine/Source/Programs/UnrealBuildAccelerator/Common/Private/UbaCacheClient.cpp b/Engine/Source/Programs/UnrealBuildAccelerator/Common/Private/UbaCacheClient.cpp
--- a/Engine/Source/Programs/UnrealBuildAccelerator/Common/Private/UbaCacheClient.cpp
+++ b/Engine/Source/Programs/UnrealBuildAccelerator/Common/Private/UbaCacheClient.cpp
@@ -309,10 +309,11 @@ namespace uba
 
 		CacheStats cacheStats;
 		StorageStats storageStats;
-		SystemStats systemStats;
+		KernelStats kernelStats;
+		auto kg = MakeGuard([&]() { KernelStats::GetGlobal().Add(kernelStats); m_storage.AddStats(storageStats); });
 
 		StorageStatsScope __(storageStats);
-		SystemStatsScope _(systemStats);
+		KernelStatsScope _(kernelStats);
 
 		CasKey cmdKey = GetCmdKey(rootPaths, info);
 		if (cmdKey == CasKeyZero)
@@ -329,7 +330,7 @@ namespace uba
 				BinaryWriter writer(memory, 0, sizeof_array(memory));
 				cacheStats.Write(writer);
 				storageStats.Write(writer);
-				systemStats.Write(writer);
+				kernelStats.Write(writer);
 				m_session.GetTrace().CacheEndFetch(fetchId, success, memory, writer.GetPosition());
 			});
 
@@ -615,7 +616,7 @@ namespace uba
 						return false;
 					UBA_ASSERT(IsCompressed(casKey));
 
-					FileFetcher fetcher { m_storage.m_bufferSlots };
+					FileFetcher fetcher { m_storage.m_bufferSlots, storageStats };
 					fetcher.m_errorOnFail = false;
 
 					if (IsNormalized(casKey))
@@ -735,7 +736,8 @@ namespace uba
 				return false;
 		}
 
-		FileFetcher fetcher { m_storage.m_bufferSlots };
+		StorageStats storageStats;
+		FileFetcher fetcher { m_storage.m_bufferSlots, storageStats };
 		bool destinationIsCompressed = false;
 		if (destinationFile)
 		{
diff --git a/Engine/Source/Programs/UnrealBuildAccelerator/Common/Private/UbaFile.cpp b/Engine/Source/Programs/UnrealBuildAccelerator/Common/Private/UbaFile.cpp
--- a/Engine/Source/Programs/UnrealBuildAccelerator/Common/Private/UbaFile.cpp
+++ b/Engine/Source/Programs/UnrealBuildAccelerator/Common/Private/UbaFile.cpp
@@ -96,7 +96,8 @@ namespace uba
 
 	bool ReadFile(Logger& logger, const tchar* fileName, FileHandle fileHandle, void* b, u64 bufferLen)
 	{
-		ExtendedTimerScope ts(SystemStats::GetCurrent().readFile);
+		auto& stats = KernelStats::GetCurrent();
+		ExtendedTimerScope ts(stats.readFile);
 		u8* buffer = (u8*)b;
 		u64 readLeft = bufferLen;
 		u64 firstZeroReadTime = 0;
@@ -128,6 +129,8 @@ namespace uba
 			readLeft -= wasRead;
 			buffer += wasRead;
 		}
+
+		stats.readFile.bytes += bufferLen;
 		return true;
 	}
 
@@ -163,7 +166,7 @@ namespace uba
 
 	bool GetFileInformationByHandle(FileInformation& out, Logger& logger, const tchar* fileName, FileHandle hFile)
 	{
-		ExtendedTimerScope ts(SystemStats::GetCurrent().getFileAttributes);
+		ExtendedTimerScope ts(KernelStats::GetCurrent().getFileInfo);
 #if PLATFORM_WINDOWS
 		BY_HANDLE_FILE_INFORMATION info;
 		if (!::GetFileInformationByHandle(asHANDLE(hFile), &info))
@@ -222,7 +225,7 @@ namespace uba
 
 	bool FileExists(Logger& logger, const tchar* fileName, u64* outSize, u32* outAttributes, u64* lastWriteTime)
 	{
-		ExtendedTimerScope ts(SystemStats::GetCurrent().getFileAttributes);
+		ExtendedTimerScope ts(KernelStats::GetCurrent().getFileInfo);
 #if PLATFORM_WINDOWS
 		MAKE_LONG_FILENAME(fileName);
 		WIN32_FILE_ATTRIBUTE_DATA data;
@@ -273,7 +276,7 @@ namespace uba
 
 	bool SetEndOfFile(Logger& logger, const tchar* fileName, FileHandle handle, u64 size)
 	{
-		ExtendedTimerScope ts(SystemStats::GetCurrent().setFileInfo);
+		ExtendedTimerScope ts(KernelStats::GetCurrent().setFileInfo);
 #if PLATFORM_WINDOWS
 		FILE_END_OF_FILE_INFO info;
 		info.EndOfFile = ToLargeInteger(size);
@@ -420,7 +423,7 @@ namespace uba
 
 	FileHandle CreateFileW(const tchar* fileName, u32 desiredAccess, u32 shareMode, u32 createDisp, u32 flagsAndAttributes)
 	{
-		ExtendedTimerScope ts(SystemStats::GetCurrent().createFile);
+		ExtendedTimerScope ts(KernelStats::GetCurrent().createFile);
 	#if PLATFORM_WINDOWS
 		MAKE_LONG_FILENAME(fileName);
 		return (FileHandle)(u64)::CreateFileW(fileName, desiredAccess, shareMode, NULL, createDisp, flagsAndAttributes, NULL);
@@ -479,7 +482,7 @@ namespace uba
 
 	bool CloseFile(const tchar* fileName, FileHandle h)
 	{
-		ExtendedTimerScope ts(SystemStats::GetCurrent().closeFile);
+		ExtendedTimerScope ts(KernelStats::GetCurrent().closeFile);
 #if PLATFORM_WINDOWS
 		return ::CloseHandle(asHANDLE(h));
 #else
@@ -586,7 +589,7 @@ namespace uba
 
 	bool GetFileLastWriteTime(u64& outTime, FileHandle hFile)
 	{
-		ExtendedTimerScope ts(SystemStats::GetCurrent().getFileTime);
+		ExtendedTimerScope ts(KernelStats::GetCurrent().getFileTime);
 #if PLATFORM_WINDOWS
 		FILETIME lastWriteTime;
 		auto res = ::GetFileTime(asHANDLE(hFile), NULL, NULL, &lastWriteTime);
@@ -634,7 +637,7 @@ namespace uba
 
 	bool GetFileSizeEx(u64& outFileSize, FileHandle hFile)
 	{
-		ExtendedTimerScope ts(SystemStats::GetCurrent().getFileAttributes);
+		ExtendedTimerScope ts(KernelStats::GetCurrent().getFileInfo);
 #if PLATFORM_WINDOWS
 		LARGE_INTEGER lpFileSize;
 		if (!::GetFileSizeEx(asHANDLE(hFile), &lpFileSize))
@@ -657,7 +660,7 @@ namespace uba
 
 	u32 GetFileAttributesW(const tchar* fileName)
 	{
-		ExtendedTimerScope ts(SystemStats::GetCurrent().getFileAttributes);
+		ExtendedTimerScope ts(KernelStats::GetCurrent().getFileInfo);
 #if PLATFORM_WINDOWS
 		MAKE_LONG_FILENAME(fileName);
 		return ::GetFileAttributesW(fileName);
diff --git a/Engine/Source/Programs/UnrealBuildAccelerator/Common/Private/UbaFileAccessor.cpp b/Engine/Source/Programs/UnrealBuildAccelerator/Common/Private/UbaFileAccessor.cpp
--- a/Engine/Source/Programs/UnrealBuildAccelerator/Common/Private/UbaFileAccessor.cpp
+++ b/Engine/Source/Programs/UnrealBuildAccelerator/Common/Private/UbaFileAccessor.cpp
@@ -26,7 +26,7 @@ namespace uba
 #if PLATFORM_WINDOWS
 	bool SetDeleteOnClose(Logger& logger, const tchar* fileName, FileHandle& handle, bool value)
 	{
-		ExtendedTimerScope ts(SystemStats::GetCurrent().setFileInfo);
+		ExtendedTimerScope ts(KernelStats::GetCurrent().setFileInfo);
 		FILE_DISPOSITION_INFO info;
 		info.DeleteFile = value;
 		if (!::SetFileInformationByHandle(asHANDLE(handle), FileDispositionInfo, &info, sizeof(info)))
@@ -161,11 +161,14 @@ namespace uba
 
 	bool FileAccessor::Write(const void* data, u64 dataLen, u64 offset)
 	{
-		ExtendedTimerScope ts(SystemStats::GetCurrent().writeFile);
+		auto& stats = KernelStats::GetCurrent();
+		ExtendedTimerScope ts(stats.writeFile);
 
 		if (!m_isWrite)
 			return false;
 
+		stats.writeFile.bytes += dataLen;
+
 		#if UBA_USE_WRITE_THROUGH
 		bool useWriteThrough = true;
 		u8 writeThroughBuffer[4 * 1024];
diff --git a/Engine/Source/Programs/UnrealBuildAccelerator/Common/Private/UbaFileMapping.cpp b/Engine/Source/Programs/UnrealBuildAccelerator/Common/Private/UbaFileMapping.cpp
--- a/Engine/Source/Programs/UnrealBuildAccelerator/Common/Private/UbaFileMapping.cpp
+++ b/Engine/Source/Programs/UnrealBuildAccelerator/Common/Private/UbaFileMapping.cpp
@@ -42,7 +42,7 @@ namespace uba
 
 	FileMappingHandle CreateMemoryMappingW(Logger& logger, u32 flProtect, u64 maxSize, const tchar* name)
 	{
-		ExtendedTimerScope ts(SystemStats::GetCurrent().createFileMapping);
+		ExtendedTimerScope ts(KernelStats::GetCurrent().createFileMapping);
 #if PLATFORM_WINDOWS
 		return { InternalCreateFileMappingW(INVALID_HANDLE_VALUE, flProtect, (DWORD)ToHigh(maxSize), ToLow(maxSize), name) };
 #else
@@ -168,7 +168,7 @@ namespace uba
 
 	FileMappingHandle CreateFileMappingW(FileHandle file, u32 protect, u64 maxSize, const tchar* hint)
 	{
-		ExtendedTimerScope ts(SystemStats::GetCurrent().createFileMapping);
+		ExtendedTimerScope ts(KernelStats::GetCurrent().createFileMapping);
 #if PLATFORM_WINDOWS
 		return { InternalCreateFileMappingW(asHANDLE(file), protect, (DWORD)ToHigh(maxSize), ToLow(maxSize), NULL) };
 #else
@@ -206,7 +206,7 @@ namespace uba
 
 	u8* MapViewOfFile(FileMappingHandle fileMappingObject, u32 desiredAccess, u64 offset, u64 bytesToMap)
 	{
-		ExtendedTimerScope ts(SystemStats::GetCurrent().mapViewOfFile);
+		ExtendedTimerScope ts(KernelStats::GetCurrent().mapViewOfFile);
 #if PLATFORM_WINDOWS
 		return (u8*)::MapViewOfFile(fileMappingObject.handle, desiredAccess, (DWORD)ToHigh(offset), ToLow(offset), bytesToMap);
 #else
@@ -237,7 +237,7 @@ namespace uba
 
 	bool UnmapViewOfFile(const void* lpBaseAddress, u64 bytesToUnmap, const tchar* hint)
 	{
-		ExtendedTimerScope ts(SystemStats::GetCurrent().unmapViewOfFile);
+		ExtendedTimerScope ts(KernelStats::GetCurrent().unmapViewOfFile);
 #if PLATFORM_WINDOWS
 		(void)bytesToUnmap; return ::UnmapViewOfFile(lpBaseAddress);
 #else
diff --git a/Engine/Source/Programs/UnrealBuildAccelerator/Common/Private/UbaNetworkServer.cpp b/Engine/Source/Programs/UnrealBuildAccelerator/Common/Private/UbaNetworkServer.cpp
--- a/Engine/Source/Programs/UnrealBuildAccelerator/Common/Private/UbaNetworkServer.cpp
+++ b/Engine/Source/Programs/UnrealBuildAccelerator/Common/Private/UbaNetworkServer.cpp
@@ -643,9 +643,9 @@ namespace uba
 			for (auto& c : m_connections)
 			{
 				success = c.Stop() && success;
-				m_sendTimer.Add(c.m_sendTimer);
-				m_encryptTimer.Add(c.m_encryptTimer);
-				m_decryptTimer.Add(c.m_decryptTimer);
+				m_sendTimer += c.m_sendTimer;
+				m_encryptTimer += c.m_encryptTimer;
+				m_decryptTimer += c.m_decryptTimer;
 			}
 			lock.Leave();
 
@@ -1070,7 +1070,7 @@ namespace uba
 				++it;
 				continue;
 			}
-			m_sendTimer.Add(con.m_sendTimer);
+			m_sendTimer += con.m_sendTimer;
 			it = m_connections.erase(it);
 		}
 	}
diff --git a/Engine/Source/Programs/UnrealBuildAccelerator/Common/Private/UbaProcess.cpp b/Engine/Source/Programs/UnrealBuildAccelerator/Common/Private/UbaProcess.cpp
--- a/Engine/Source/Programs/UnrealBuildAccelerator/Common/Private/UbaProcess.cpp
+++ b/Engine/Source/Programs/UnrealBuildAccelerator/Common/Private/UbaProcess.cpp
@@ -419,7 +419,7 @@ namespace uba
 
 	void ProcessImpl::ThreadRun(bool runningRemote, void* environment)
 	{
-		SystemStatsScope systemStatsScope(m_systemStats);
+		KernelStatsScope kernelStatsScope(m_kernelStats);
 		StorageStatsScope storageStatsScope(m_storageStats);
 		SessionStatsScope sessionStatsScope(m_sessionStats);
 
@@ -518,7 +518,7 @@ namespace uba
 	{
 		m_processStats.wallTime = GetTime() - m_startTime;
 
-		SystemStats::GetGlobal().Add(m_systemStats);
+		KernelStats::GetGlobal().Add(m_kernelStats);
 
 		// For some reason a parent can exit before a child. Need to figure out repro for this but I've seen it happen on ClangEditor win64
 		for (auto& child : m_childProcesses)
@@ -1023,7 +1023,11 @@ namespace uba
 		ProcessStats stats;
 		stats.Read(reader, ~0u);
 
+		KernelStats kernelStats;
+		kernelStats.Read(reader, ~0u);
+
 		m_processStats.Add(stats);
+		m_kernelStats.Add(kernelStats);
 
 		if (!IsCancelled())
 			if (m_startInfo.writeOutputFilesOnFail || m_startInfo.rules->IsExitCodeSuccess(m_nativeProcessExitCode))
@@ -1034,7 +1038,7 @@ namespace uba
 			m_parentProcess->m_processStats.Add(m_processStats);
 			m_parentProcess->m_sessionStats.Add(m_sessionStats);
 			m_parentProcess->m_storageStats.Add(m_storageStats);
-			m_parentProcess->m_systemStats.Add(m_systemStats);
+			m_parentProcess->m_kernelStats.Add(m_kernelStats);
 		}
 
 		if (m_startInfo.outputStatsThresholdMs && TimeToMs(m_processStats.GetTotalTime()) > m_startInfo.outputStatsThresholdMs)
@@ -1077,11 +1081,14 @@ namespace uba
 		processStats.wallTime = GetTime() - m_startTime;
 		processStats.cpuTime = 0;
 		processStats.hostTotalTime = m_processStats.hostTotalTime;
+		
+		KernelStats kernelStats;
+		kernelStats.Read(reader, TraceVersion);
 					
 		processStats.Write(statsWriter);
 		m_sessionStats.Write(statsWriter);
 		m_storageStats.Write(statsWriter);
-		m_systemStats.Write(statsWriter);
+		m_kernelStats.Write(statsWriter);
 		BinaryReader statsReader(statsWriter.GetData(), 0, statsWriter.GetPosition());
 
 		bool newProcess = false;
@@ -1110,7 +1117,7 @@ namespace uba
 		m_processStats = {};
 		m_sessionStats = {};
 		m_storageStats = {};
-		m_systemStats = {};
+		m_kernelStats = {};
 
 		m_startTime = GetTime();
 
diff --git a/Engine/Source/Programs/UnrealBuildAccelerator/Common/Private/UbaSession.cpp b/Engine/Source/Programs/UnrealBuildAccelerator/Common/Private/UbaSession.cpp
--- a/Engine/Source/Programs/UnrealBuildAccelerator/Common/Private/UbaSession.cpp
+++ b/Engine/Source/Programs/UnrealBuildAccelerator/Common/Private/UbaSession.cpp
@@ -1416,6 +1416,8 @@ namespace uba
 		{
 			StackBinaryWriter<1024> writer;
 			process.m_processStats.Write(writer);
+			process.m_storageStats.Write(writer);
+			process.m_kernelStats.Write(writer);
 			u32 exitCode = process.GetExitCode();
 			Vector<ProcessLogLine> emptyLines;
 			auto& logLines = (exitCode != 0 || m_detailedTrace) ? process.m_logLines : emptyLines;
@@ -1530,7 +1532,7 @@ namespace uba
 	void Session::PrintSummary(Logger& logger)
 	{
 		logger.BeginScope();
-		logger.Info(TC("  ----- Uba process stats summary -----"));
+		logger.Info(TC("  ------- Detours stats summary -------"));
 		m_processStats.Print(logger);
 		logger.Info(TC(""));
 
@@ -1546,7 +1548,7 @@ namespace uba
 		}
 		logger.Info(TC(""));
 
-		logger.Info(TC("  ----- Uba session stats summary -----"));
+		logger.Info(TC("  ------- Session stats summary -------"));
 
 		PrintSessionStats(logger);
 		logger.EndScope();
@@ -2270,12 +2272,12 @@ namespace uba
 		process.m_processStats.Write(writer);
 		process.m_sessionStats.Write(writer);
 		process.m_storageStats.Write(writer);
-		process.m_systemStats.Write(writer);
+		process.m_kernelStats.Write(writer);
 		m_trace.ProcessEnvironmentUpdated(process.GetId(), reason, writer.GetData(), writer.GetPosition());
 		process.m_processStats = {};
 		process.m_sessionStats = {};
 		process.m_storageStats = {};
-		process.m_systemStats = {};
+		process.m_kernelStats = {};
 		return true;
 	}
 
diff --git a/Engine/Source/Programs/UnrealBuildAccelerator/Common/Private/UbaSessionClient.cpp b/Engine/Source/Programs/UnrealBuildAccelerator/Common/Private/UbaSessionClient.cpp
--- a/Engine/Source/Programs/UnrealBuildAccelerator/Common/Private/UbaSessionClient.cpp
+++ b/Engine/Source/Programs/UnrealBuildAccelerator/Common/Private/UbaSessionClient.cpp
@@ -1301,7 +1301,7 @@ namespace uba
 		process.m_processStats.Write(writer);
 		process.m_sessionStats.Write(writer);
 		process.m_storageStats.Write(writer);
-		process.m_systemStats.Write(writer);
+		process.m_kernelStats.Write(writer);
 
 		StackBinaryReader<16> reader;
 		if (!msg.Send(reader, m_stats.procFinishedMsg) && m_loop)
@@ -1379,7 +1379,7 @@ namespace uba
 				PrintSummary(logger);
 				m_storage.PrintSummary(logger);
 				m_client.PrintSummary(logger);
-				SystemStats::GetGlobal().Print(logger, true);
+				KernelStats::GetGlobal().Print(logger, true);
 				if (extraInfo)
 					extraInfo(logger);
 			});
@@ -1765,7 +1765,7 @@ namespace uba
 					PrintSummary(logger);
 					m_storage.PrintSummary(logger);
 					m_client.PrintSummary(logger);
-					SystemStats::GetGlobal().Print(logger, true);
+					KernelStats::GetGlobal().Print(logger, true);
 				});
 			m_trace.SessionSummary(0, writer.GetData(), writer.GetPosition());
 
@@ -1933,12 +1933,12 @@ namespace uba
 			process.m_processStats.Write(writer);
 			process.m_sessionStats.Write(writer);
 			process.m_storageStats.Write(writer);
-			process.m_systemStats.Write(writer);
+			process.m_kernelStats.Write(writer);
 
 			process.m_processStats = {};
 			process.m_sessionStats = {};
 			process.m_storageStats = {};
-			process.m_systemStats = {};
+			process.m_kernelStats = {};
 
 			if (!msg.Send(reader, m_stats.customMsg))
 				return false;
diff --git a/Engine/Source/Programs/UnrealBuildAccelerator/Common/Private/UbaSessionServer.cpp b/Engine/Source/Programs/UnrealBuildAccelerator/Common/Private/UbaSessionServer.cpp
--- a/Engine/Source/Programs/UnrealBuildAccelerator/Common/Private/UbaSessionServer.cpp
+++ b/Engine/Source/Programs/UnrealBuildAccelerator/Common/Private/UbaSessionServer.cpp
@@ -220,7 +220,7 @@ namespace uba
 					PrintSummary(logger);
 					m_storage.PrintSummary(logger);
 					m_server.PrintSummary(logger);
-					SystemStats::GetGlobal().Print(logger, true);
+					KernelStats::GetGlobal().Print(logger, true);
 				});
 			m_trace.SessionSummary(0, writer.GetData(), writer.GetPosition());
 		}
diff --git a/Engine/Source/Programs/UnrealBuildAccelerator/Common/Private/UbaStats.cpp b/Engine/Source/Programs/UnrealBuildAccelerator/Common/Private/UbaStats.cpp
--- a/Engine/Source/Programs/UnrealBuildAccelerator/Common/Private/UbaStats.cpp
+++ b/Engine/Source/Programs/UnrealBuildAccelerator/Common/Private/UbaStats.cpp
@@ -5,6 +5,20 @@
 
 namespace uba
 {
+	inline void Write(BinaryWriter& writer, u64 v) { writer.Write7BitEncoded(v); }
+	inline void Write(BinaryWriter& writer, u32 v) { writer.Write7BitEncoded(v); }
+
+	inline void Read(BinaryReader& reader, u64& v, u32 version) { v = reader.Read7BitEncoded(); }
+	inline void Read(BinaryReader& reader, u32& v, u32 version) { v = u32(reader.Read7BitEncoded()); }
+	inline void Read(BinaryReader& reader, Timer& timer, u32 version) { timer.time = reader.Read7BitEncoded(); timer.count = (u32)reader.Read7BitEncoded(); }
+	inline void Read(BinaryReader& reader, AtomicU64& v, u32 version) { v = reader.Read7BitEncoded(); }
+	inline void Read(BinaryReader& reader, TimeAndBytes& timer, u32 version) { timer.time = reader.Read7BitEncoded(); timer.count = (u32)reader.Read7BitEncoded(); if (version >= 30) timer.bytes = reader.Read7BitEncoded(); }
+
+	inline bool IsEmpty(u64 v) { return v == 0; }
+	inline bool IsEmpty(u32 v) { return v == 0; }
+
+	const tchar EmptyCharArray[] = TC("                   ");
+
 	void ProcessStats::Print(Logger& logger, u64 frequency)
 	{
 		logger.Info(TC("  Total              %8u %9s"), GetTotalCount(), TimeToText(GetTotalTime(), false, frequency).str);
@@ -20,10 +34,9 @@ namespace uba
 			#undef UBA_PROCESS_STAT
 		};
 
-		const tchar empty[] = TC("                   ");
 		for (Stat& s : stats)
 			if (s.timer.count)
-				logger.Info(TC("  %c") PERCENT_HS TC("%s %8u %9s"), ToUpper(s.name[0]), s.name + 1, empty + s.nameLen, s.timer.count.load(), TimeToText(s.timer.time, false, frequency).str);
+				logger.Info(TC("  %c") PERCENT_HS TC("%s %8u %9s"), ToUpper(s.name[0]), s.name + 1, EmptyCharArray + s.nameLen, s.timer.count.load(), TimeToText(s.timer.time, false, frequency).str);
 
 		logger.Info(TC(""));
 
@@ -34,74 +47,167 @@ namespace uba
 		logger.Info(TC("  Wall Time                   %9s"), TimeToText(wallTime, false, frequency).str);
 	}
 
-	void SystemStats::Print(Logger& logger, bool writeHeader, u64 frequency)
+	u64 ProcessStats::GetTotalTime()
 	{
-		if (writeHeader)
-			logger.Info(TC("  --- Platform system stats summary ---"));
+		return 
+		#define UBA_PROCESS_STAT(T, ver) + T.time
+		UBA_PROCESS_STATS
+		#undef UBA_PROCESS_STAT
+			;
+	}
 
-		struct Stat { const char* name; u64 nameLen; const Timer& timer; };
-		Stat stats[] =
+	u32 ProcessStats::GetTotalCount()
+	{
+		return 
+		#define UBA_PROCESS_STAT(T, ver) + T.count
+			UBA_PROCESS_STATS
+		#undef UBA_PROCESS_STAT
+			;
+	}
+
+	void ProcessStats::Read(BinaryReader& reader, u32 version)
+	{
+		uba::Read(reader, waitOnResponse, version);
+
+		if (version < 30)
 		{
-			#define UBA_SYSTEM_STAT(T, ver) { #T, sizeof(#T), T },
-			UBA_SYSTEM_STATS
-			#undef UBA_SYSTEM_STAT
-		};
+			#define UBA_PROCESS_STAT(T, ver) if (ver <= version) uba::Read(reader, T, version);
+			UBA_PROCESS_STATS
+			#undef UBA_PROCESS_STAT
+		}
+		else
+		{
+			u64 bits = reader.Read7BitEncoded();
+			#define UBA_PROCESS_STAT(var, ver) if (bits & (1 << Bit_##var)) uba::Read(reader, var, version);
+			UBA_PROCESS_STATS
+			#undef UBA_PROCESS_STAT
+		}
 
-		const tchar empty[] = TC("                   ");
-		for (Stat& s : stats)
-			if (s.timer.count)
-				logger.Info(TC("  %c") PERCENT_HS TC("%s %8u %9s"), ToUpper(s.name[0]), s.name+1, empty + s.nameLen, s.timer.count.load(), TimeToText(s.timer.time, false, frequency).str);
+		startupTime = reader.ReadU64();
+		exitTime = reader.ReadU64();
+		wallTime = reader.ReadU64();
+		cpuTime = reader.ReadU64();
+		usedMemory = reader.ReadU32();
+		hostTotalTime = reader.ReadU64();
+	}
+
+	void ProcessStats::Add(const ProcessStats& other)
+	{
+		waitOnResponse += other.waitOnResponse;
+			
+		#define UBA_PROCESS_STAT(T, ver) T += other.T;
+		UBA_PROCESS_STATS
+		#undef UBA_PROCESS_STAT
+
+		startupTime += other.startupTime;
+		exitTime += other.exitTime;
+		wallTime += other.wallTime;
+		cpuTime += other.cpuTime;
+		usedMemory = Max(usedMemory, other.usedMemory);
+		hostTotalTime += other.hostTotalTime;
+	}
+
+	template<typename T>
+	void LogStat(Logger& logger, const char* name, const T&, u64 frequency) {}
+
+	void LogStat(Logger& logger, const char* name, const Timer& timer, u64 frequency)
+	{
+		if (!timer.count)
+			return;
+		logger.Info(TC("  %c") PERCENT_HS TC("%s %8u %9s"), ToUpper(name[0]), name+1, EmptyCharArray + strlen(name)+1, timer.count.load(), TimeToText(timer.time, false, frequency).str);
+	}
+
+	void LogStat(Logger& logger, const char* name, const ExtendedTimer& timer, u64 frequency)
+	{
+		LogStat(logger, name, (const Timer&)timer, frequency);
+	}
+
+	void LogStat(Logger& logger, const char* name, const TimeAndBytes& timer, u64 frequency)
+	{
+		if (!timer.count)
+			return;
+		logger.Info(TC("  %c") PERCENT_HS TC("%s %8u %9s"), ToUpper(name[0]), name+1, EmptyCharArray + strlen(name)+1, timer.count.load(), TimeToText(timer.time, false, frequency).str);
+		if (timer.bytes)
+			logger.Info(TC("     Bytes                    %9s"), BytesToText(timer.bytes).str);
+	}
+
+	void KernelStats::Print(Logger& logger, bool writeHeader, u64 frequency)
+	{
+		if (writeHeader)
+			logger.Info(TC("  ------- Kernel stats summary --------"));
+
+		#define UBA_KERNEL_STAT(type, var, ver) LogStat(logger, #var, var, frequency);
+		UBA_KERNEL_STATS
+		#undef UBA_KERNEL_STAT
 
 		if (writeHeader)
 			logger.Info(TC(""));
 	}
 
-	bool SystemStats::IsEmpty()
+	bool KernelStats::IsEmpty()
 	{
-		#define UBA_SYSTEM_STAT(T, ver) if (T.count) return false;
-		UBA_SYSTEM_STATS
-		#undef UBA_SYSTEM_STAT
+		#define UBA_KERNEL_STAT(type, var, ver) if (var.count) return false;
+		UBA_KERNEL_STATS
+		#undef UBA_KERNEL_STAT
 		return true;
 	}
 
-	void SystemStats::Add(const SystemStats& other)
+	void KernelStats::Add(const KernelStats& other)
 	{
-		#define UBA_SYSTEM_STAT(var, ver) var += other.var;
-		UBA_SYSTEM_STATS
-		#undef UBA_SYSTEM_STAT
+		#define UBA_KERNEL_STAT(type, var, ver) var += other.var;
+		UBA_KERNEL_STATS
+		#undef UBA_KERNEL_STAT
 	}
 
-	void SystemStats::Write(BinaryWriter& writer)
+	void KernelStats::Read(BinaryReader& reader, u32 version)
 	{
-		#define UBA_SYSTEM_STAT(var, ver) uba::Write(writer, var);
-		UBA_SYSTEM_STATS
-		#undef UBA_SYSTEM_STAT
-	}
+		if (version < 30)
+		{
+			#define UBA_KERNEL_STAT(type, var, ver) if (ver <= version) uba::Read(reader, var, version);
+			UBA_KERNEL_STATS
+			#undef UBA_KERNEL_STAT
+			return;
+		}
 
-	void SystemStats::Read(BinaryReader& reader, u32 version)
-	{
-		#define UBA_SYSTEM_STAT(var, ver) if (ver <= version) uba::Read(reader, var);
-		UBA_SYSTEM_STATS
-		#undef UBA_SYSTEM_STAT
+		u16 bits = reader.ReadU16();
+		#define UBA_KERNEL_STAT(type, var, ver) if (bits & (1 << Bit_##var)) uba::Read(reader, var, version);
+		UBA_KERNEL_STATS
+		#undef UBA_KERNEL_STAT
 	}
 
 	void StorageStats::Write(BinaryWriter& writer)
 	{
-		#define UBA_STORAGE_STAT(type, var) uba::Write(writer, var);
+		u64 bits = 0;
+		#define UBA_STORAGE_STAT(type, var, ver) if (!uba::IsEmpty(var)) bits |= (1 << Bit_##var);
+		UBA_STORAGE_STATS
+		#undef UBA_STORAGE_STAT
+
+		writer.Write7BitEncoded(bits);
+
+		#define UBA_STORAGE_STAT(type, var, ver) if (!uba::IsEmpty(var)) uba::Write(writer, var);
 		UBA_STORAGE_STATS
 		#undef UBA_STORAGE_STAT
 	}
 
-	void StorageStats::Read(BinaryReader& reader)
+	void StorageStats::Read(BinaryReader& reader, u32 version)
 	{
-		#define UBA_STORAGE_STAT(type, var) uba::Read(reader, var);
+		if (version < 30)
+		{
+			#define UBA_STORAGE_STAT(type, var, ver) uba::Read(reader, var, version);
+			UBA_STORAGE_STATS
+			#undef UBA_STORAGE_STAT
+			return;
+		}
+
+		u64 bits = reader.Read7BitEncoded();
+		#define UBA_STORAGE_STAT(type, var, ver) if (bits & (1 << Bit_##var)) uba::Read(reader, var, version);
 		UBA_STORAGE_STATS
 		#undef UBA_STORAGE_STAT
 	}
 
 	void StorageStats::Add(const StorageStats& other)
 	{
-		#define UBA_STORAGE_STAT(type, var) var += other.var;
+		#define UBA_STORAGE_STAT(type, var, ver) var += other.var;
 		UBA_STORAGE_STATS
 		#undef UBA_STORAGE_STAT
 	}
@@ -135,13 +241,18 @@ namespace uba
 			logger.Info(TC("  CopyOrLink           %6u %9s"), copyOrLink.count.load(), TimeToText(copyOrLink.time, false, frequency).str);
 		if (copyOrLinkWait.count)
 			logger.Info(TC("  CopyOrLinkWait       %6u %9s"), copyOrLinkWait.count.load(), TimeToText(copyOrLinkWait.time, false, frequency).str);
+		if (compressWrite.count)
+			logger.Info(TC("  CompressToMem        %6u %9s"), compressWrite.count.load(), TimeToText(compressWrite.time, false, frequency).str);
 		if (decompressToMem.count)
 			logger.Info(TC("  DecompressToMem      %6u %9s"), decompressToMem.count.load(), TimeToText(decompressToMem.time, false, frequency).str);
+
+		if (memoryCopy.count)
+			logger.Info(TC("  MemoryCopy           %6u %9s"), memoryCopy.count.load(), TimeToText(memoryCopy.time, false, frequency).str);
 	}
 
 	bool StorageStats::IsEmpty()
 	{
-		#define UBA_STORAGE_STAT(type, var) if (var != type()) return false;
+		#define UBA_STORAGE_STAT(type, var, ver) if (var != type()) return false;
 		UBA_STORAGE_STATS
 		#undef UBA_STORAGE_STAT
 		return true;
@@ -168,14 +279,30 @@ namespace uba
 
 	void SessionStats::Write(BinaryWriter& writer)
 	{
-		#define UBA_SESSION_STAT(type, var, ver) uba::Write(writer, var);
+		u16 bits = 0;
+		#define UBA_SESSION_STAT(type, var, ver) if (var.count) bits |= (1 << Bit_##var);
+		UBA_SESSION_STATS
+		#undef UBA_SESSION_STAT
+
+		writer.WriteU16(bits);
+
+		#define UBA_SESSION_STAT(type, var, ver) if (var.count) uba::Write(writer, var);
 		UBA_SESSION_STATS
 		#undef UBA_SESSION_STAT
 	}
 
 	void SessionStats::Read(BinaryReader& reader, u32 version)
 	{
-		#define UBA_SESSION_STAT(type, var, ver) if (ver <= version) uba::Read(reader, var);
+		if (version < 30)
+		{
+			#define UBA_SESSION_STAT(type, var, ver) if (ver <= version) uba::Read(reader, var, version);
+			UBA_SESSION_STATS
+			#undef UBA_SESSION_STAT
+			return;
+		}
+
+		u16 bits = reader.ReadU16();
+		#define UBA_SESSION_STAT(type, var, ver) if (bits & (1 << Bit_##var)) uba::Read(reader, var, version);
 		UBA_SESSION_STATS
 		#undef UBA_SESSION_STAT
 	}
@@ -187,17 +314,6 @@ namespace uba
 		#undef UBA_SESSION_STAT
 	}
 
-	template<typename T>
-	void LogStat(Logger& logger, const char* name, const T&, u64 frequency) {}
-
-	void LogStat(Logger& logger, const char* name, const Timer& timer, u64 frequency)
-	{
-		if (!timer.count)
-			return;
-		const tchar empty[] = TC("                   ");
-		logger.Info(TC("  %c") PERCENT_HS TC("%s %8u %9s"), ToUpper(name[0]), name+1, empty + strlen(name)+1, timer.count.load(), TimeToText(timer.time, false, frequency).str);
-	}
-
 	void SessionStats::Print(Logger& logger, u64 frequency)
 	{
 		#define UBA_SESSION_STAT(type, var, ver) LogStat(logger, #var, var, frequency);
@@ -242,7 +358,7 @@ namespace uba
 	void SessionSummaryStats::Read(BinaryReader& reader, u32 version)
 	{
 		stats.Read(reader, version);
-		#define UBA_SESSION_SUMMARY_STAT(type, var) uba::Read(reader, var);
+		#define UBA_SESSION_SUMMARY_STAT(type, var) uba::Read(reader, var, version);
 		UBA_SESSION_SUMMARY_STATS
 		#undef UBA_SESSION_SUMMARY_STAT
 	}
@@ -267,7 +383,7 @@ namespace uba
 
 	void CacheStats::Read(BinaryReader& reader, u32 version)
 	{
-		#define UBA_CACHE_STAT(type, var, ver) if (ver <= version) uba::Read(reader, var);
+		#define UBA_CACHE_STAT(type, var, ver) if (ver <= version) uba::Read(reader, var, version);
 		UBA_CACHE_STATS
 		#undef UBA_CACHE_STAT
 	}
diff --git a/Engine/Source/Programs/UnrealBuildAccelerator/Common/Private/UbaStorage.cpp b/Engine/Source/Programs/UnrealBuildAccelerator/Common/Private/UbaStorage.cpp
--- a/Engine/Source/Programs/UnrealBuildAccelerator/Common/Private/UbaStorage.cpp
+++ b/Engine/Source/Programs/UnrealBuildAccelerator/Common/Private/UbaStorage.cpp
@@ -344,6 +344,7 @@ namespace uba
 				if (!destinationFile.Write(destBuf, writeBytes))
 					return false;
 #else
+				TimerScope cts(stats.memoryCopy);
 				memcpy(toMem + totalWritten, destBuf, writeBytes);
 #endif
 
@@ -469,6 +470,7 @@ namespace uba
 				if (!rec->destination->Write(compressSlotBuffer, writeBytes))
 					rec->error = true;
 #else
+				TimerScope cts(stats.memoryCopy);
 				memcpy(rec->mem + rec->memPos, compressSlotBuffer, writeBytes);
 				rec->memPos += writeBytes;
 #endif
@@ -2578,7 +2580,7 @@ namespace uba
 
 	void StorageImpl::PrintSummary(Logger& logger)
 	{
-		logger.Info(TC("  ----- Uba storage stats summary -----"));
+		logger.Info(TC("  ------- Storage stats summary -------"));
 		if (m_casLookup.empty())
 		{
 			logger.Info(TC("  Storage not loaded"));
diff --git a/Engine/Source/Programs/UnrealBuildAccelerator/Common/Private/UbaStorageServer.cpp b/Engine/Source/Programs/UnrealBuildAccelerator/Common/Private/UbaStorageServer.cpp
--- a/Engine/Source/Programs/UnrealBuildAccelerator/Common/Private/UbaStorageServer.cpp
+++ b/Engine/Source/Programs/UnrealBuildAccelerator/Common/Private/UbaStorageServer.cpp
@@ -575,7 +575,7 @@ namespace uba
 				{
 					*fetchId = u16(~0);
 					u64 sendCasTime = GetTime() - start;
-					stats.sendCas.Add(Timer{sendCasTime, 1});
+					stats.sendCas += Timer{sendCasTime, 1};
 					return true;
 				}
 
@@ -643,7 +643,7 @@ namespace uba
 				PushId(fetchId);
 
 				sendCasTime += GetTime() - start;
-				Stats().sendCas.Add(Timer{sendCasTime, 1});
+				Stats().sendCas += Timer{sendCasTime, 1};
 				return true;
 			}
 
@@ -861,7 +861,7 @@ namespace uba
 					activeStore.recvCasTime += GetTime() - time2;
 
 					StorageStats& stats = Stats();
-					stats.recvCas.Add(Timer{activeStore.recvCasTime, 1});
+					stats.recvCas += Timer{activeStore.recvCasTime, 1};
 					stats.recvCasBytesComp += activeStore.fileSize;
 					stats.recvCasBytesRaw += activeStore.actualSize;
 
diff --git a/Engine/Source/Programs/UnrealBuildAccelerator/Common/Private/UbaStorageUtils.cpp b/Engine/Source/Programs/UnrealBuildAccelerator/Common/Private/UbaStorageUtils.cpp
--- a/Engine/Source/Programs/UnrealBuildAccelerator/Common/Private/UbaStorageUtils.cpp
+++ b/Engine/Source/Programs/UnrealBuildAccelerator/Common/Private/UbaStorageUtils.cpp
@@ -405,6 +405,7 @@ namespace uba
 
 	bool FileFetcher::RetrieveFile(Logger& logger, NetworkClient& client, const CasKey& casKey, const tchar* destination, bool destinationIsCompressed, MemoryBlock* destinationMem)
 	{
+		TimerScope ts(m_stats.recvCas);
 		u8* slot = m_bufferSlots.Pop();
 		auto sg = MakeGuard([&](){ m_bufferSlots.Push(slot); });
 
@@ -478,16 +479,6 @@ namespace uba
 				return true;
 			};
 
-		//u8* writePos = writeMem;
-
-		// This is here just to prevent server from getting a million messages at the same time.
-		// In theory we could have 10 clients with 48 processes each where each one of the processes asks for a large file (64 messages in flight)
-		// So worst case in that scenario would be 10*48*64 = 30000 messages.
-		bool oneAtTheTime = false;//left > client->GetMessageMaxSize() * 2;
-		if (oneAtTheTime)
-			m_retrieveOneBatchAtTheTimeLock.EnterWrite();
-		auto oatg = MakeGuard([&]() { if (oneAtTheTime) m_retrieveOneBatchAtTheTimeLock.LeaveWrite(); });
-
 		u32 readIndex = 0;
 
 		if (destinationIsCompressed)
@@ -630,6 +621,9 @@ namespace uba
 			if (!destinationFile.Close(&lastWritten))
 				return false;
 
+		m_stats.recvCasBytesRaw += actualSize;
+		m_stats.recvCasBytesComp += fileSize;
+
 		return true;
 	}
 
diff --git a/Engine/Source/Programs/UnrealBuildAccelerator/Common/Public/UbaDirectoryIterator.h b/Engine/Source/Programs/UnrealBuildAccelerator/Common/Public/UbaDirectoryIterator.h
--- a/Engine/Source/Programs/UnrealBuildAccelerator/Common/Public/UbaDirectoryIterator.h
+++ b/Engine/Source/Programs/UnrealBuildAccelerator/Common/Public/UbaDirectoryIterator.h
@@ -84,7 +84,7 @@ namespace uba
 		RtlInitUnicodeString(&uniName, str.data);
 
 		Timer timer(0, 1);
-		auto mg = MakeGuard([&]() { SystemStats::GetCurrent().traverseDir.Add(timer); });
+		auto mg = MakeGuard([&]() { KernelStats::GetCurrent().traverseDir += timer; });
 
 		HANDLE handle;
 
@@ -96,7 +96,7 @@ namespace uba
 
 		u64 startTime = GetTime();
 		NTSTATUS res = NtCreateFile(&handle, FILE_READ_ATTRIBUTES | FILE_LIST_DIRECTORY | FILE_ATTRIBUTE_UNPINNED, &ObjectAttributes, &IoStatusBlock, 0, FILE_ATTRIBUTE_NORMAL, ShareAccess, FILE_OPEN, FILE_DIRECTORY_FILE | FILE_SYNCHRONOUS_IO_ALERT, NULL, 0);
-		timer.time += startTime - GetTime();
+		timer.time += GetTime() - startTime;
 
 		if (res == STATUS_OBJECT_NAME_NOT_FOUND || res == STATUS_OBJECT_PATH_NOT_FOUND)
 			return !errorOnNotFound;
@@ -134,7 +134,7 @@ namespace uba
 		{
 			startTime = GetTime();
 			res = NtQueryDirectoryFile(handle, 0, NULL, NULL, &IoStatusBlock, buff, sizeof(buff) - 2, (FILE_INFORMATION_CLASS)FileIdBothDirectoryInformation, FALSE, NULL, FALSE);
-			timer.time += startTime - GetTime();
+			timer.time += GetTime() - startTime;
 
 			if (res != STATUS_SUCCESS)
 			{
diff --git a/Engine/Source/Programs/UnrealBuildAccelerator/Common/Public/UbaProcess.h b/Engine/Source/Programs/UnrealBuildAccelerator/Common/Public/UbaProcess.h
--- a/Engine/Source/Programs/UnrealBuildAccelerator/Common/Public/UbaProcess.h
+++ b/Engine/Source/Programs/UnrealBuildAccelerator/Common/Public/UbaProcess.h
@@ -176,7 +176,7 @@ namespace uba
 		SessionStats m_sessionStats;
 		StorageStats m_storageStats;
 		ProcessStats m_processStats;
-		SystemStats m_systemStats;
+		KernelStats m_kernelStats;
 
 		Thread m_messageThread;
 
diff --git a/Engine/Source/Programs/UnrealBuildAccelerator/Common/Public/UbaStats.h b/Engine/Source/Programs/UnrealBuildAccelerator/Common/Public/UbaStats.h
--- a/Engine/Source/Programs/UnrealBuildAccelerator/Common/Public/UbaStats.h
+++ b/Engine/Source/Programs/UnrealBuildAccelerator/Common/Public/UbaStats.h
@@ -8,41 +8,46 @@ namespace uba
 {
 	class Logger;
 
-	inline void Write(BinaryWriter& writer, const AtomicU64& v) { writer.Write7BitEncoded(v.load()); }
-	inline void Read(BinaryReader& reader, AtomicU64& v) { v = reader.Read7BitEncoded(); }
-
 	#define UBA_STORAGE_STATS \
-		UBA_STORAGE_STAT(Timer, calculateCasKey) \
-		UBA_STORAGE_STAT(Timer, copyOrLink) \
-		UBA_STORAGE_STAT(Timer, copyOrLinkWait) \
-		UBA_STORAGE_STAT(Timer, ensureCas) \
-		UBA_STORAGE_STAT(Timer, sendCas) \
-		UBA_STORAGE_STAT(Timer, recvCas) \
-		UBA_STORAGE_STAT(Timer, compressWrite) \
-		UBA_STORAGE_STAT(Timer, compressSend) \
-		UBA_STORAGE_STAT(Timer, decompressRecv) \
-		UBA_STORAGE_STAT(Timer, decompressToMem) \
-		UBA_STORAGE_STAT(Timer, handleOverflow) \
-		UBA_STORAGE_STAT(AtomicU64, sendCasBytesRaw) \
-		UBA_STORAGE_STAT(AtomicU64, sendCasBytesComp) \
-		UBA_STORAGE_STAT(AtomicU64, recvCasBytesRaw) \
-		UBA_STORAGE_STAT(AtomicU64, recvCasBytesComp) \
-		UBA_STORAGE_STAT(Timer, createCas) \
-		UBA_STORAGE_STAT(AtomicU64, createCasBytesRaw) \
-		UBA_STORAGE_STAT(AtomicU64, createCasBytesComp) \
+		UBA_STORAGE_STAT(Timer, calculateCasKey, 0) \
+		UBA_STORAGE_STAT(Timer, copyOrLink, 0) \
+		UBA_STORAGE_STAT(Timer, copyOrLinkWait, 0) \
+		UBA_STORAGE_STAT(Timer, ensureCas, 0) \
+		UBA_STORAGE_STAT(Timer, sendCas, 0) \
+		UBA_STORAGE_STAT(Timer, recvCas, 0) \
+		UBA_STORAGE_STAT(Timer, compressWrite, 0) \
+		UBA_STORAGE_STAT(Timer, compressSend, 0) \
+		UBA_STORAGE_STAT(Timer, decompressRecv, 0) \
+		UBA_STORAGE_STAT(Timer, decompressToMem, 0) \
+		UBA_STORAGE_STAT(Timer, memoryCopy, 30) \
+		UBA_STORAGE_STAT(Timer, handleOverflow, 0) \
+		UBA_STORAGE_STAT(AtomicU64, sendCasBytesRaw, 0) \
+		UBA_STORAGE_STAT(AtomicU64, sendCasBytesComp, 0) \
+		UBA_STORAGE_STAT(AtomicU64, recvCasBytesRaw, 0) \
+		UBA_STORAGE_STAT(AtomicU64, recvCasBytesComp, 0) \
+		UBA_STORAGE_STAT(Timer, createCas, 0) \
+		UBA_STORAGE_STAT(AtomicU64, createCasBytesRaw, 0) \
+		UBA_STORAGE_STAT(AtomicU64, createCasBytesComp, 0) \
 
     struct StorageStats
 	{
-		#define UBA_STORAGE_STAT(type, var) type var;
+		#define UBA_STORAGE_STAT(type, var, ver) type var;
 		UBA_STORAGE_STATS
 		#undef UBA_STORAGE_STAT
 
 		void Write(BinaryWriter& writer);
-		void Read(BinaryReader& reader);
+		void Read(BinaryReader& reader, u32 version);
 		void Add(const StorageStats& other);
 		void Print(Logger& logger, u64 frequency = GetFrequency());
 		bool IsEmpty();
 		static StorageStats* GetCurrent();
+
+		enum
+		{
+			#define UBA_STORAGE_STAT(type, var, ver) Bit_##var,
+			UBA_STORAGE_STATS
+			#undef UBA_STORAGE_STAT
+		};
 	};
 
 	struct StorageStatsScope
@@ -80,6 +85,13 @@ namespace uba
 		void Print(Logger& logger, u64 frequency = GetFrequency());
 		bool IsEmpty();
 		static SessionStats* GetCurrent();
+
+		enum
+		{
+			#define UBA_SESSION_STAT(type, var, ver) Bit_##var,
+			UBA_SESSION_STATS
+			#undef UBA_SESSION_STAT
+		};
 	};
 
 	struct SessionStatsScope
diff --git a/Engine/Source/Programs/UnrealBuildAccelerator/Common/Public/UbaStorageUtils.h b/Engine/Source/Programs/UnrealBuildAccelerator/Common/Public/UbaStorageUtils.h
--- a/Engine/Source/Programs/UnrealBuildAccelerator/Common/Public/UbaStorageUtils.h
+++ b/Engine/Source/Programs/UnrealBuildAccelerator/Common/Public/UbaStorageUtils.h
@@ -38,8 +38,8 @@ namespace uba
 		bool RetrieveFile(Logger& logger, NetworkClient& client, const CasKey& casKey, const tchar* destination, bool destinationIsCompressed, MemoryBlock* destinationMem = nullptr);
 
 		BufferSlots& m_bufferSlots;
+		StorageStats& m_stats;
 		StringBuffer<> m_tempPath;
-		ReaderWriterLock m_retrieveOneBatchAtTheTimeLock;
 		bool m_errorOnFail = true;
 
 		u64 lastWritten = 0;
diff --git a/Engine/Source/Programs/UnrealBuildAccelerator/Common/Public/UbaTrace.h b/Engine/Source/Programs/UnrealBuildAccelerator/Common/Public/UbaTrace.h
--- a/Engine/Source/Programs/UnrealBuildAccelerator/Common/Public/UbaTrace.h
+++ b/Engine/Source/Programs/UnrealBuildAccelerator/Common/Public/UbaTrace.h
@@ -61,7 +61,7 @@ namespace uba
 	using Color = u32;
 	inline Color toColor(u8 r, u8 g, u8 b) { return (r << 16) + (g << 8) + b; }
 
-	static constexpr u32 TraceVersion = 29;
+	static constexpr u32 TraceVersion = 30;
 	static constexpr u32 TraceReadCompatibilityVersion = 6;
 
 	class Trace : public WorkTracker
diff --git a/Engine/Source/Programs/UnrealBuildAccelerator/Core/Private/UbaPlatform.cpp b/Engine/Source/Programs/UnrealBuildAccelerator/Core/Private/UbaPlatform.cpp
--- a/Engine/Source/Programs/UnrealBuildAccelerator/Core/Private/UbaPlatform.cpp
+++ b/Engine/Source/Programs/UnrealBuildAccelerator/Core/Private/UbaPlatform.cpp
@@ -14,27 +14,27 @@ extern const char* __progname;
 
 namespace uba
 {
-	SystemStats g_systemStats;
-	thread_local SystemStats* t_systemStats;
+	KernelStats g_kernelStatsInternal;
+	thread_local KernelStats* t_kernelStats;
 
-	SystemStats& SystemStats::GetCurrent()
+	KernelStats& KernelStats::GetCurrent()
 	{
-		SystemStats* stats = t_systemStats;
-		return stats ? *stats : g_systemStats;
+		KernelStats* stats = t_kernelStats;
+		return stats ? *stats : g_kernelStatsInternal;
 	}
 
-	SystemStats& SystemStats::GetGlobal()
+	KernelStats& KernelStats::GetGlobal()
 	{
-		return g_systemStats;
+		return g_kernelStatsInternal;
 	}
 
-	SystemStatsScope::SystemStatsScope(SystemStats& s) : stats(s)
+	KernelStatsScope::KernelStatsScope(KernelStats& s) : stats(s)
 	{
-		t_systemStats = &stats;
+		t_kernelStats = &stats;
 	}
-	SystemStatsScope::~SystemStatsScope()
+	KernelStatsScope::~KernelStatsScope()
 	{
-		t_systemStats = nullptr;
+		t_kernelStats = nullptr;
 	}
 
 	bool CreateGuid(Guid& out)
diff --git a/Engine/Source/Programs/UnrealBuildAccelerator/Core/Public/UbaProcessStats.h b/Engine/Source/Programs/UnrealBuildAccelerator/Core/Public/UbaProcessStats.h
--- a/Engine/Source/Programs/UnrealBuildAccelerator/Core/Public/UbaProcessStats.h
+++ b/Engine/Source/Programs/UnrealBuildAccelerator/Core/Public/UbaProcessStats.h
@@ -10,11 +10,9 @@ namespace uba
 	class Logger;
 
 	inline void Write(BinaryWriter& writer, const Timer& timer) { writer.Write7BitEncoded(timer.time); writer.Write7BitEncoded(timer.count); }
-	inline void Write(BinaryWriter& writer, u64 v) { writer.Write7BitEncoded(v); }
-	inline void Write(BinaryWriter& writer, u32 v) { writer.Write7BitEncoded(v); }
-	inline void Read(BinaryReader& reader, Timer& timer) { timer.time = reader.Read7BitEncoded(); timer.count = (u32)reader.Read7BitEncoded(); }
-	inline void Read(BinaryReader& reader, u64& v) { v = reader.Read7BitEncoded(); }
-	inline void Read(BinaryReader& reader, u32& v) { v = u32(reader.Read7BitEncoded()); }
+	inline void Write(BinaryWriter& writer, const AtomicU64& v) { writer.Write7BitEncoded(v.load()); }
+	inline bool IsEmpty(const Timer& timer) { return timer.count == 0; }
+	inline bool IsEmpty(const AtomicU64& v) { return v.load() == 0; }
 
 	struct ProcessStats
 	{
@@ -41,64 +39,50 @@ namespace uba
 		UBA_PROCESS_STAT(sendFiles, 0) \
 		UBA_PROCESS_STAT(writeFiles, 19) \
 		UBA_PROCESS_STAT(queryCache, 24) \
+		UBA_PROCESS_STAT(waitDecompress, 30) \
 
 
 		#define UBA_PROCESS_STAT(T, ver) Timer T;
 		UBA_PROCESS_STATS
 		#undef UBA_PROCESS_STAT
 
-		AtomicU64 startupTime = 0;
-		AtomicU64 exitTime = 0;
+		enum
+		{
+			#define UBA_PROCESS_STAT(var, ver) Bit_##var,
+			UBA_PROCESS_STATS
+			#undef UBA_PROCESS_STAT
+		};
+
+		AtomicU64 startupTime;
+		AtomicU64 exitTime;
 
 		// Don't add in GetTotalTime()
-		AtomicU64 wallTime = 0;
-		AtomicU64 cpuTime = 0;
+		AtomicU64 wallTime;
+		AtomicU64 cpuTime;
 
-		AtomicU64 usedMemory = 0;
+		AtomicU64 usedMemory;
 
 		AtomicU64 hostTotalTime;
 
 		void Print(Logger& logger, u64 frequency = GetFrequency());
 
-		u64 GetTotalTime()
-		{
-			return 
-			#define UBA_PROCESS_STAT(T, ver) + T.time
-			UBA_PROCESS_STATS
-			#undef UBA_PROCESS_STAT
-				;
-		}
-
-		u32 GetTotalCount()
-		{
-			return 
-			#define UBA_PROCESS_STAT(T, ver) + T.count
-				UBA_PROCESS_STATS
-			#undef UBA_PROCESS_STAT
-				;
-		}
+		u64 GetTotalTime();
+		u32 GetTotalCount();
+		void Read(BinaryReader& reader, u32 version);
+		void Add(const ProcessStats& other);
 
-		void Read(BinaryReader& reader, u32 version)
+		void Write(BinaryWriter& writer) const
 		{
-			uba::Read(reader, waitOnResponse);
+			uba::Write(writer, waitOnResponse);
 
-			#define UBA_PROCESS_STAT(T, ver) if (ver <= version) uba::Read(reader, T);
+			u64 bits = 0;
+			#define UBA_PROCESS_STAT(var, ver) if (!uba::IsEmpty(var)) bits |= (1 << Bit_##var);
 			UBA_PROCESS_STATS
 			#undef UBA_PROCESS_STAT
 
-			startupTime = reader.ReadU64();
-			exitTime = reader.ReadU64();
-			wallTime = reader.ReadU64();
-			cpuTime = reader.ReadU64();
-			usedMemory = reader.ReadU32();
-			hostTotalTime = reader.ReadU64();
-		}
+			writer.Write7BitEncoded(bits);
 
-		void Write(BinaryWriter& writer) const
-		{
-			uba::Write(writer, waitOnResponse);
-
-			#define UBA_PROCESS_STAT(T, ver) uba::Write(writer, T);
+			#define UBA_PROCESS_STAT(var, ver) if (!uba::IsEmpty(var)) uba::Write(writer, var);
 			UBA_PROCESS_STATS
 			#undef UBA_PROCESS_STAT
 
@@ -109,59 +93,71 @@ namespace uba
 			writer.WriteU32(u32(usedMemory));
 			writer.WriteU64(hostTotalTime);
 		}
+	};
 
-		void Add(const ProcessStats& other)
-		{
-			waitOnResponse.Add(other.waitOnResponse);
-			
-			#define UBA_PROCESS_STAT(T, ver) T.Add(other.T);
-			UBA_PROCESS_STATS
-			#undef UBA_PROCESS_STAT
-
-			startupTime += other.startupTime;
-			exitTime += other.exitTime;
-			wallTime += other.wallTime;
-			cpuTime += other.cpuTime;
-			usedMemory = Max(usedMemory, other.usedMemory);
-			hostTotalTime += other.hostTotalTime;
-		}
+	struct TimeAndBytes : ExtendedTimer
+	{
+		void operator+=(const TimeAndBytes& o) { time += o.time; count += o.count; bytes += o.bytes; }
+		AtomicU64 bytes;
 	};
 
-	#define UBA_SYSTEM_STATS \
-		UBA_SYSTEM_STAT(createFile, 0) \
-		UBA_SYSTEM_STAT(closeFile, 0) \
-		UBA_SYSTEM_STAT(writeFile, 0) \
-		UBA_SYSTEM_STAT(readFile, 0) \
-		UBA_SYSTEM_STAT(setFileInfo, 0) \
-		UBA_SYSTEM_STAT(getFileAttributes, 29) \
-		UBA_SYSTEM_STAT(createFileMapping, 0) \
-		UBA_SYSTEM_STAT(mapViewOfFile, 0) \
-		UBA_SYSTEM_STAT(unmapViewOfFile, 0) \
-		UBA_SYSTEM_STAT(getFileTime, 0) \
-		UBA_SYSTEM_STAT(closeHandle, 0) \
-		UBA_SYSTEM_STAT(traverseDir, 27) \
-
-
-	struct SystemStats
+	inline void Write(BinaryWriter& writer, const TimeAndBytes& timer) { writer.Write7BitEncoded(timer.time); writer.Write7BitEncoded(timer.count); writer.Write7BitEncoded(timer.bytes); }
+
+	struct KernelStats
 	{
-		#define UBA_SYSTEM_STAT(T, ver) ExtendedTimer T;
-		UBA_SYSTEM_STATS
-		#undef UBA_SYSTEM_STAT
+		#define UBA_KERNEL_STATS \
+			UBA_KERNEL_STAT(ExtendedTimer, createFile, 0) \
+			UBA_KERNEL_STAT(ExtendedTimer, closeFile, 0) \
+			UBA_KERNEL_STAT(TimeAndBytes, writeFile, 0) \
+			UBA_KERNEL_STAT(TimeAndBytes, memoryCopy, 30) \
+			UBA_KERNEL_STAT(TimeAndBytes, readFile, 0) \
+			UBA_KERNEL_STAT(ExtendedTimer, setFileInfo, 0) \
+			UBA_KERNEL_STAT(ExtendedTimer, getFileInfo, 29) \
+			UBA_KERNEL_STAT(ExtendedTimer, createFileMapping, 0) \
+			UBA_KERNEL_STAT(ExtendedTimer, mapViewOfFile, 0) \
+			UBA_KERNEL_STAT(ExtendedTimer, unmapViewOfFile, 0) \
+			UBA_KERNEL_STAT(ExtendedTimer, getFileTime, 0) \
+			UBA_KERNEL_STAT(ExtendedTimer, closeHandle, 0) \
+			UBA_KERNEL_STAT(ExtendedTimer, traverseDir, 27) \
+
+		#define UBA_KERNEL_STAT(type, var, ver) type var;
+		UBA_KERNEL_STATS
+		#undef UBA_KERNEL_STAT
 
 		void Read(BinaryReader& reader, u32 version);
-		void Write(BinaryWriter& writer);
 		void Print(Logger& logger, bool writeHeader, u64 frequency = GetFrequency());
 		bool IsEmpty();
-		void Add(const SystemStats& other);
-		static SystemStats& GetCurrent();
-		static SystemStats& GetGlobal();
+		void Add(const KernelStats& other);
+		static KernelStats& GetCurrent();
+		static KernelStats& GetGlobal();
+
+		enum
+		{
+			#define UBA_KERNEL_STAT(type, var, ver) Bit_##var,
+			UBA_KERNEL_STATS
+			#undef UBA_KERNEL_STAT
+		};
+
+		void Write(BinaryWriter& writer)
+		{
+			u16 bits = 0;
+			#define UBA_KERNEL_STAT(type, var, ver) if (var.count) bits |= (1 << Bit_##var);
+			UBA_KERNEL_STATS
+			#undef UBA_KERNEL_STAT
+
+			writer.WriteU16(bits);
+
+			#define UBA_KERNEL_STAT(type, var, ver) if (var.count) uba::Write(writer, var);
+			UBA_KERNEL_STATS
+			#undef UBA_KERNEL_STAT
+		}
 	};
 
-	struct SystemStatsScope
+	struct KernelStatsScope
 	{
-		SystemStatsScope(SystemStats& stats);
-		~SystemStatsScope();
+		KernelStatsScope(KernelStats& stats);
+		~KernelStatsScope();
 
-		SystemStats& stats;
+		KernelStats& stats;
 	};
 }
diff --git a/Engine/Source/Programs/UnrealBuildAccelerator/Core/Public/UbaTimer.h b/Engine/Source/Programs/UnrealBuildAccelerator/Core/Public/UbaTimer.h
--- a/Engine/Source/Programs/UnrealBuildAccelerator/Core/Public/UbaTimer.h
+++ b/Engine/Source/Programs/UnrealBuildAccelerator/Core/Public/UbaTimer.h
@@ -35,7 +35,6 @@ namespace uba
 		void operator=(const Timer& o) { time.store(o.time.load()); count.store(o.count.load()); }
 		Atomic<u64> time;
 		Atomic<u32> count;
-		void Add(const Timer& o) { time += o.time; count += o.count; }
 		void operator+=(const Timer& o) { time += o.time; count += o.count; }
 		bool operator==(const Timer& o) const { return time == o.time && count == o.count; }
 	};
diff --git a/Engine/Source/Programs/UnrealBuildAccelerator/Detours/Private/PosixOS/UbaDetoursFunctionsPosix.cpp b/Engine/Source/Programs/UnrealBuildAccelerator/Detours/Private/PosixOS/UbaDetoursFunctionsPosix.cpp
--- a/Engine/Source/Programs/UnrealBuildAccelerator/Detours/Private/PosixOS/UbaDetoursFunctionsPosix.cpp
+++ b/Engine/Source/Programs/UnrealBuildAccelerator/Detours/Private/PosixOS/UbaDetoursFunctionsPosix.cpp
@@ -1958,6 +1958,7 @@ namespace uba
 		writer.WriteU32(0); // Exit code
 		writer.WriteString(""); // Log name
 		g_stats.Write(writer);
+		g_kernelStats.Write(writer);
 
 		// This can't wait for response since the session process might move on and reuse shared memory with someone else
 		// Note, if we start using memory mapped files we need to change this to true for child processes since Exit message is writing files to disk..
@@ -1998,6 +1999,7 @@ namespace uba
 		writer.WriteU32(terminateCode); // Exit code
 		writer.WriteString(""); // Log name
 		g_stats.Write(writer);
+		g_kernelStats.Write(writer);
 		writer.Flush(false);
 
 		CloseCom();
@@ -2021,7 +2023,7 @@ extern "C"
 			writer.WriteByte(MessageType_GetNextProcess);
 			writer.WriteU32(prevExitCode);
 			g_stats.Write(writer);
-
+			g_kernelStats.Write(writer);
 
 			writer.Flush();
 			BinaryReader reader;
diff --git a/Engine/Source/Programs/UnrealBuildAccelerator/Detours/Private/UbaDetoursShared.cpp b/Engine/Source/Programs/UnrealBuildAccelerator/Detours/Private/UbaDetoursShared.cpp
--- a/Engine/Source/Programs/UnrealBuildAccelerator/Detours/Private/UbaDetoursShared.cpp
+++ b/Engine/Source/Programs/UnrealBuildAccelerator/Detours/Private/UbaDetoursShared.cpp
@@ -10,6 +10,7 @@ namespace uba
 	VARIABLE_MEM(StringBuffer<512>, g_virtualApplication);
 	VARIABLE_MEM(StringBuffer<512>, g_virtualApplicationDir);
 	VARIABLE_MEM(ProcessStats, g_stats);
+	VARIABLE_MEM(KernelStats, g_kernelStats);
 	VARIABLE_MEM(ReaderWriterLock, g_communicationLock);
 	VARIABLE_MEM(StringBuffer<256>, g_logName);
 	VARIABLE_MEM(StringBuffer<512>, g_virtualWorkingDir);
@@ -31,6 +32,7 @@ namespace uba
 		g_virtualApplicationMem.Create();
 		g_virtualApplicationDirMem.Create();
 		g_statsMem.Create();
+		g_kernelStatsMem.Create();
 		g_communicationLockMem.Create();
 		g_logNameMem.Create();
 		g_virtualWorkingDirMem.Create();
diff --git a/Engine/Source/Programs/UnrealBuildAccelerator/Detours/Private/UbaDetoursShared.h b/Engine/Source/Programs/UnrealBuildAccelerator/Detours/Private/UbaDetoursShared.h
--- a/Engine/Source/Programs/UnrealBuildAccelerator/Detours/Private/UbaDetoursShared.h
+++ b/Engine/Source/Programs/UnrealBuildAccelerator/Detours/Private/UbaDetoursShared.h
@@ -56,6 +56,7 @@ namespace uba
 	extern StringBuffer<128>& g_systemRoot;
 
 	extern ProcessStats& g_stats;
+	extern KernelStats& g_kernelStats;
 	extern bool g_echoOn;
 	extern ReaderWriterLock& g_communicationLock;
 	extern MemoryBlock& g_memoryBlock;
diff --git a/Engine/Source/Programs/UnrealBuildAccelerator/Detours/Private/Windows/UbaDetoursFunctionsKernelBase.inl b/Engine/Source/Programs/UnrealBuildAccelerator/Detours/Private/Windows/UbaDetoursFunctionsKernelBase.inl
--- a/Engine/Source/Programs/UnrealBuildAccelerator/Detours/Private/Windows/UbaDetoursFunctionsKernelBase.inl
+++ b/Engine/Source/Programs/UnrealBuildAccelerator/Detours/Private/Windows/UbaDetoursFunctionsKernelBase.inl
@@ -327,6 +327,7 @@ BOOL Detoured_ReadFile(HANDLE hFile, LPVOID lpBuffer, DWORD nNumberOfBytesToRead
 		trueHandle = dh.trueHandle;
 	}
 
+	TimerScope ts(g_kernelStats.readFile);
 	BOOL res = True_ReadFile(trueHandle, lpBuffer, nNumberOfBytesToRead, lpNumberOfBytesRead, lpOverlapped);
 	DEBUG_LOG_TRUE(L"ReadFile", L"%llu %u/%u (%ls) -> %ls", uintptr_t(hFile), lpNumberOfBytesRead ? *lpNumberOfBytesRead : ~0u, nNumberOfBytesToRead, HandleToName(hFile), ToString(res));
 	return res;
@@ -594,6 +595,7 @@ BOOL Detoured_WriteFile(HANDLE hFile, LPCVOID lpBuffer, DWORD nNumberOfBytesToWr
 		//return True_WriteFile(trueHandle, lpBuffer, nNumberOfBytesToWrite, lpNumberOfBytesWritten, lpOverlapped);
 	}
 
+	TimerScope ts(g_kernelStats.writeFile);
 	BOOL res = True_WriteFile(trueHandle, lpBuffer, nNumberOfBytesToWrite, lpNumberOfBytesWritten, lpOverlapped);
 	DEBUG_LOG_TRUE(L"WriteFile", L"%llu (%ls) -> %ls", uintptr_t(hFile), HandleToName(hFile), ToString(res));
 	return res;
@@ -607,6 +609,8 @@ BOOL Detoured_WriteFileEx(HANDLE hFile, LPCVOID lpBuffer, DWORD nNumberOfBytesTo
 	UBA_ASSERT(isDetouredHandle(hFile));
 	DetouredHandle& h = asDetouredHandle(hFile);
 	UBA_ASSERT(h.trueHandle != INVALID_HANDLE_VALUE);
+
+	TimerScope ts(g_kernelStats.writeFile);
 	return True_WriteFileEx(h.trueHandle, lpBuffer, nNumberOfBytesToWrite, lpOverlapped, lpCompletionRoutine);
 }
 
@@ -905,6 +909,7 @@ BOOL Shared_GetFileAttributesExW(LPCWSTR lpFileName, GET_FILEEX_INFO_LEVELS fInf
 	if (!CanDetour(fixedFileName))
 	{
 		DEBUG_LOG_TRUE(L"GetFileAttributesExW", L"(%ls)", lpFileName);
+		TimerScope ts(g_kernelStats.getFileInfo);
 		return True_GetFileAttributesExW(lpFileName, fInfoLevelId, lpFileInformation);
 	}
 
@@ -914,6 +919,7 @@ BOOL Shared_GetFileAttributesExW(LPCWSTR lpFileName, GET_FILEEX_INFO_LEVELS fInf
 	if (!attr.useCache)
 	{
 		DEBUG_LOG_TRUE(L"GetFileAttributesExW", L"(%ls)", lpFileName);
+		TimerScope ts(g_kernelStats.getFileInfo);
 		return True_GetFileAttributesExW(realName, fInfoLevelId, lpFileInformation);
 	}
 
@@ -932,6 +938,7 @@ BOOL Detoured_GetFileAttributesExW(LPCWSTR lpFileName, GET_FILEEX_INFO_LEVELS fI
 	{
 		UBA_ASSERT(!g_runningRemote);
 		DEBUG_LOG_TRUE(L"GetFileAttributesExW", L"(%ls)", lpFileName);
+		TimerScope ts(g_kernelStats.getFileInfo);
 		return True_GetFileAttributesExW(lpFileName, fInfoLevelId, lpFileInformation);
 	}
 
@@ -952,6 +959,7 @@ DWORD Detoured_GetFileAttributesW(LPCWSTR lpFileName)
 	DETOURED_CALL(GetFileAttributesW);
 	if (t_disallowDetour != 0 || Equals(lpFileName, L"nul"))
 	{
+		TimerScope ts(g_kernelStats.getFileInfo);
 		DWORD res = True_GetFileAttributesW(lpFileName);
 		DEBUG_LOG_TRUE(L"GetFileAttributesW", L"(NODETOUR) (%ls) -> %u", lpFileName, res);
 		return res;
@@ -978,6 +986,7 @@ BOOL Detoured_SetFileAttributesW(LPCWSTR lpFileName, DWORD dwFileAttributes)
 		return true;
 	}
 	DEBUG_LOG_TRUE(L"SetFileAttributesW", L"(%ls) %u", lpFileName, dwFileAttributes);
+	TimerScope ts(g_kernelStats.setFileInfo);
 	return True_SetFileAttributesW(lpFileName, dwFileAttributes);
 }
 
@@ -1831,6 +1840,7 @@ BOOL Detoured_GetFileInformationByHandleEx(HANDLE hFile, FILE_INFO_BY_HANDLE_CLA
 		}
 	}
 	DEBUG_LOG_TRUE(L"GetFileInformationByHandleEx", L"(%ls)", HandleToName(hFile));
+	TimerScope ts(g_kernelStats.getFileInfo);
 	return True_GetFileInformationByHandleEx(trueHandle, fileInformationClass, lpFileInformation, dwBufferSize); /// calls GetFileInformationByHandleEx
 }
 
@@ -1982,6 +1992,7 @@ BOOL Detoured_GetFileInformationByHandle(HANDLE hFile, LPBY_HANDLE_FILE_INFORMAT
 		trueHandle = dh.trueHandle;
 	}
 
+	TimerScope ts(g_kernelStats.getFileInfo);
 	auto res = True_GetFileInformationByHandle(trueHandle, lpFileInformation); // Calls NtQueryInformationFile
 	DEBUG_LOG_TRUE(L"GetFileInformationByHandle", L"%llu (%ls) -> %u", uintptr_t(hFile), HandleToName(hFile), res);
 	return res;
@@ -2080,6 +2091,7 @@ HANDLE Detoured_CreateFileMappingW(HANDLE hFile, LPSECURITY_ATTRIBUTES lpFileMap
 		trueHandle = dh.trueHandle;
 	}
 
+	TimerScope ts(g_kernelStats.createFileMapping);
 	HANDLE mappingHandle = True_CreateFileMappingW(trueHandle, lpFileMappingAttributes, flProtect, dwMaximumSizeHigh, dwMaximumSizeLow, lpName);
 	if (!mappingHandle)
 	{
@@ -2176,6 +2188,7 @@ LPVOID Detoured_MapViewOfFileEx(HANDLE hFileMappingObject, DWORD dwDesiredAccess
 					u32 counter = 0;
 					do
 					{
+						TimerScope ts(g_kernelStats.mapViewOfFile);
 						res = (u8*)True_MapViewOfFileEx(trueMappingObject, dwDesiredAccess, ToHigh(offset), ToLow(offset), dwNumberOfBytesToMap, lpBaseAddress);
 						if (res)
 							break;
@@ -2212,6 +2225,8 @@ LPVOID Detoured_MapViewOfFileEx(HANDLE hFileMappingObject, DWORD dwDesiredAccess
 		UBA_ASSERT(dh.trueHandle != INVALID_HANDLE_VALUE);
 		trueMappingObject = dh.trueHandle;
 	}
+
+	TimerScope ts(g_kernelStats.mapViewOfFile);
 	void* res = True_MapViewOfFileEx(trueMappingObject, dwDesiredAccess, dwFileOffsetHigh, dwFileOffsetLow, dwNumberOfBytesToMap, lpBaseAddress);
 	DEBUG_LOG_TRUE(L"MapViewOfFileEx", L"%llu (size %llu) (%ls) -> 0x%llx", uintptr_t(hFileMappingObject), dwNumberOfBytesToMap, HandleToName(hFileMappingObject), uintptr_t(res));
 
@@ -3132,6 +3147,7 @@ BOOL Detoured_ReadFileEx(HANDLE hFile, LPVOID lpBuffer, DWORD nNumberOfBytesToRe
 	DEBUG_LOG_TRUE(L"ReadFileEx", L"%llu (%ls)", uintptr_t(hFile), HandleToName(hFile));
 	UBA_ASSERT(!isDetouredHandle(hFile));
 	UBA_ASSERT(!isListDirectoryHandle(hFile));
+	TimerScope ts(g_kernelStats.readFile);
 	return True_ReadFileEx(hFile, lpBuffer, nNumberOfBytesToRead, lpOverlapped, lpCompletionRoutine);
 }
 
diff --git a/Engine/Source/Programs/UnrealBuildAccelerator/Detours/Private/Windows/UbaDetoursFunctionsNtDll.inl b/Engine/Source/Programs/UnrealBuildAccelerator/Detours/Private/Windows/UbaDetoursFunctionsNtDll.inl
--- a/Engine/Source/Programs/UnrealBuildAccelerator/Detours/Private/Windows/UbaDetoursFunctionsNtDll.inl
+++ b/Engine/Source/Programs/UnrealBuildAccelerator/Detours/Private/Windows/UbaDetoursFunctionsNtDll.inl
@@ -213,6 +213,7 @@ NTSTATUS Detoured_NtQueryInformationFile(HANDLE FileHandle, PIO_STATUS_BLOCK IoS
 		}
 	}
 
+	TimerScope ts(g_kernelStats.getFileInfo);
 	auto res = True_NtQueryInformationFile(TrueHandle, IoStatusBlock, FileInformation, Length, FileInformationClass);
 	DEBUG_LOG_TRUE(L"NtQueryInformationFile", L"(%u) %llu (%ls) -> %ls", FileInformationClass, uintptr_t(FileHandle), HandleToName(FileHandle), ToString(res));
 	return res;
@@ -468,6 +469,8 @@ NTSTATUS Detoured_NtSetInformationFile(HANDLE FileHandle, PIO_STATUS_BLOCK IoSta
 			return STATUS_SUCCESS;
 		}
 	}
+
+	TimerScope ts(g_kernelStats.setFileInfo);
 	auto res = True_NtSetInformationFile(trueHandle, IoStatusBlock, FileInformation, Length, FileInformationClass);
 	DEBUG_LOG_TRUE(L"NtSetInformationFile", L"(%u) %llu (%ls) -> %ls", FileInformationClass, uintptr_t(FileHandle), HandleToName(FileHandle), ToString(res));
 	return res;
@@ -561,6 +564,8 @@ NTSTATUS NTAPI Local_NtCreateFile(bool IsCreateFunc, PHANDLE hFileHandle, ACCESS
 	}
 #endif
 
+	TimerScope ts(g_kernelStats.createFile);
+
 	//if (!Contains(ObjectAttributes->ObjectName->Buffer, L".dll") && !Contains(ObjectAttributes->ObjectName->Buffer, L".mui"))
 	//Rpc_WriteLogf(L"NtCreateFile: %ls", ObjectAttributes->ObjectName->Buffer);
 	auto res = True_NtCreateFile(hFileHandle, DesiredAccess, ObjectAttributes, IoStatusBlock, AllocationSize, FileAttributes, ShareAccess, CreateDisposition, CreateOptions, EaBuffer, EaLength);
@@ -806,6 +811,7 @@ NTSTATUS NTAPI Shared_NtCreateFile(bool IsCreateFunc, PHANDLE hFileHandle, ACCES
 	{
 		if (isWrite || !g_allowListDirectoryHandle)
 		{
+			TimerScope ts(g_kernelStats.createFile);
 			UBA_ASSERT(!g_runningRemote);
 			NTSTATUS res = True_NtCreateFile(hFileHandle, DesiredAccess, ObjectAttributes, IoStatusBlock, AllocationSize, FileAttributes, ShareAccess, CreateDisposition, CreateOptions, EaBuffer, EaLength);
 
@@ -1068,6 +1074,8 @@ NTSTATUS NTAPI Shared_NtCreateFile(bool IsCreateFunc, PHANDLE hFileHandle, ACCES
 						True_DuplicateHandle(g_hostProcess, (HANDLE)mappingHandle, GetCurrentProcess(), &mf.mappingHandle, 0, FALSE, DUPLICATE_SAME_ACCESS);
 						UBA_ASSERTF(mf.mappingHandle, L"DuplicateHandle failed when opening temp file %ls (%u)", fileName.data, GetLastError());
 						mf.writtenSize = mappingHandleSize;
+
+						TimerScope ts2(g_kernelStats.mapViewOfFile);
 						mf.baseAddress = (u8*)True_MapViewOfFile(mf.mappingHandle, FILE_MAP_READ, 0, 0, mappingHandleSize);
 						UBA_ASSERTF(mf.baseAddress, L"MapViewOfFile failed when opening temp file %ls (%u)", fileName.data, GetLastError());
 						mf.committedSize = mappingHandleSize;
@@ -1252,7 +1260,10 @@ NTSTATUS NTAPI Detoured_NtClose(HANDLE handle)
 	DETOURED_CALL(NtClose);
 
 	if (handle == INVALID_HANDLE_VALUE || handle == PseudoHandle)
+	{
+		TimerScope ts(g_kernelStats.closeHandle);
 		return True_NtClose(handle);
+	}
 
 	if (isListDirectoryHandle(handle))
 	{
@@ -1273,6 +1284,7 @@ NTSTATUS NTAPI Detoured_NtClose(HANDLE handle)
 
 	if (!isDetouredHandle(handle))
 	{
+		TimerScope ts(g_kernelStats.closeHandle);
 		auto res = True_NtClose(handle);
 		DEBUG_LOG_TRUE(L"NtClose", L"%llu (%ls) -> %ls", uintptr_t(handle), HandleToName(handle), ToString(res));
 		return res;
@@ -1286,7 +1298,10 @@ NTSTATUS NTAPI Detoured_NtClose(HANDLE handle)
 	NTSTATUS res = STATUS_SUCCESS;
 
 	if (dh.trueHandle != INVALID_HANDLE_VALUE)
+	{
+		TimerScope ts(g_kernelStats.closeFile);
 		res = True_NtClose(dh.trueHandle);
+	}
 
 	FileObject* fo = dh.fileObject;
 	if (!fo)
diff --git a/Engine/Source/Programs/UnrealBuildAccelerator/Detours/Private/Windows/UbaDetoursFunctionsWin.cpp b/Engine/Source/Programs/UnrealBuildAccelerator/Detours/Private/Windows/UbaDetoursFunctionsWin.cpp
--- a/Engine/Source/Programs/UnrealBuildAccelerator/Detours/Private/Windows/UbaDetoursFunctionsWin.cpp
+++ b/Engine/Source/Programs/UnrealBuildAccelerator/Detours/Private/Windows/UbaDetoursFunctionsWin.cpp
@@ -282,6 +282,7 @@ void MemoryFile::EnsureCommited(DetouredHandle& handle, u64 size)
 		{
 			True_UnmapViewOfFile(baseAddress);
 			mappedSize = Min(reserveSize, AlignUp(Max(size, mappedSize * 4), g_pageSize));
+			TimerScope ts(g_kernelStats.mapViewOfFile);
 			baseAddress = (u8*)True_MapViewOfFile(mappingHandle, FILE_MAP_READ | FILE_MAP_WRITE, 0, 0, mappedSize);
 			if (!baseAddress)
 				FatalError(1347, L"MapViewOfFile failed trying to map %llu for %ls. ReservedSize: %llu (Error code: %u)", mappedSize, HandleToName(handle), reserveSize, GetLastError());
@@ -391,10 +392,12 @@ bool EnsureMapped(DetouredHandle& handle, DWORD dwFileOffsetHigh = 0, DWORD dwFi
 		alignedOffsetStart = AlignUp(offset - (g_pageSize - 1), g_pageSize);
 		u64 alignedOffsetEnd = AlignUp(endOffset, g_pageSize);
 		u64 mapSize = alignedOffsetEnd - alignedOffsetStart;
+		TimerScope ts(g_kernelStats.mapViewOfFile);
 		info.fileMapMem = (u8*)True_MapViewOfFileEx(info.trueFileMapHandle, info.fileMapViewDesiredAccess, ToHigh(alignedOffsetStart), ToLow(alignedOffsetStart), mapSize, baseAddress);
 	}
 	else
 	{
+		TimerScope ts(g_kernelStats.mapViewOfFile);
 		info.fileMapMem = (u8*)True_MapViewOfFileEx(info.trueFileMapHandle, info.fileMapViewDesiredAccess, 0, 0, numberOfBytesToMap, baseAddress);
 	}
 
@@ -498,6 +501,7 @@ void SendExitMessage(DWORD exitCode, u64 startTime)
 	g_stats.detach.count = 1;
 
 	g_stats.Write(writer);
+	g_kernelStats.Write(writer);
 
 	// We must flush here if this is a child because,
 	// if there is a parent process waiting for this to finish,
@@ -870,14 +874,22 @@ void Init(const DetoursPayload& payload, u64 startTime)
 	if (!True_DuplicateHandle(g_hostProcess, mappedFileTableHandle, GetCurrentProcess(), &mappedFileTableHandle, 0, FALSE, DUPLICATE_SAME_ACCESS))
 		UBA_ASSERTF(false, L"Failed to duplicate filetable handle (%u)", GetLastError());
 
-	u8* mappedFileTableMem = (u8*)True_MapViewOfFile(mappedFileTableHandle, FILE_MAP_READ, 0, 0, 0);
+	u8* mappedFileTableMem;
+	{
+		TimerScope ts(g_kernelStats.mapViewOfFile);
+		mappedFileTableMem = (u8*)True_MapViewOfFile(mappedFileTableHandle, FILE_MAP_READ, 0, 0, 0);
+	}
 	UBA_ASSERT(mappedFileTableMem);
 	g_mappedFileTable.Init(mappedFileTableMem, mappedFileTableCount, mappedFileTableSize);
 
 	if (!True_DuplicateHandle(g_hostProcess, directoryTableHandle, GetCurrentProcess(), &directoryTableHandle, 0, FALSE, DUPLICATE_SAME_ACCESS))
 		UBA_ASSERTF(false, L"Failed to duplicate directorytable handle (%u)", GetLastError());
 
-	u8* directoryTableMem = (u8*)True_MapViewOfFile(directoryTableHandle, FILE_MAP_READ, 0, 0, 0);
+	u8* directoryTableMem;
+	{
+		TimerScope ts(g_kernelStats.mapViewOfFile);
+		directoryTableMem = (u8*)True_MapViewOfFile(directoryTableHandle, FILE_MAP_READ, 0, 0, 0);
+	}
 	UBA_ASSERT(directoryTableMem);
 	g_directoryTable.Init(directoryTableMem, directoryTableCount, directoryTableSize);
 
@@ -999,7 +1011,7 @@ extern "C"
 			writer.WriteByte(MessageType_GetNextProcess);
 			writer.WriteU32(prevExitCode);
 			g_stats.Write(writer);
-
+			g_kernelStats.Write(writer);
 
 			writer.Flush();
 			BinaryReader reader;
diff --git a/Engine/Source/Programs/UnrealBuildAccelerator/Detours/Private/Windows/UbaDetoursObjFilesPreloader.h b/Engine/Source/Programs/UnrealBuildAccelerator/Detours/Private/Windows/UbaDetoursObjFilesPreloader.h
--- a/Engine/Source/Programs/UnrealBuildAccelerator/Detours/Private/Windows/UbaDetoursObjFilesPreloader.h
+++ b/Engine/Source/Programs/UnrealBuildAccelerator/Detours/Private/Windows/UbaDetoursObjFilesPreloader.h
@@ -285,15 +285,8 @@ namespace uba
 			auto findIt = m_preloadedObjFiles.find(key);
 			if (findIt == m_preloadedObjFiles.end())
 				return;
-			//u64 start = GetTime();
+			TimerScope ts(g_stats.waitDecompress);
 			findIt->second.event.IsSet();
-			//u64 waitTime = GetTime() - start;
-			//if (TimeToMs(waitTime) > 2)
-			//{
-			//	StringBuffer b;
-			//	b.Appendf(L"Waited %s for preload", TimeToText(waitTime).str);
-			//	Rpc_WriteLog(b.data, b.count, true, false);
-			//}
 		}
 
 		struct Preload
diff --git a/Engine/Source/Programs/UnrealBuildAccelerator/Detours/Private/Windows/UbaDetoursUtilsWin.h b/Engine/Source/Programs/UnrealBuildAccelerator/Detours/Private/Windows/UbaDetoursUtilsWin.h
--- a/Engine/Source/Programs/UnrealBuildAccelerator/Detours/Private/Windows/UbaDetoursUtilsWin.h
+++ b/Engine/Source/Programs/UnrealBuildAccelerator/Detours/Private/Windows/UbaDetoursUtilsWin.h
@@ -70,10 +70,14 @@ namespace uba
 			}
 			else
 			{
-				mappedSize = 32 * 1024 * 1024;
-				mappingHandle = True_CreateFileMappingW(INVALID_HANDLE_VALUE, NULL, PAGE_READWRITE | SEC_RESERVE, ToHigh(reserveSize), ToLow(reserveSize), NULL);
-				if (!mappingHandle)
-					FatalError(1348, L"CreateFileMappingW failed trying to reserve %llu. (Error code: %u)", reserveSize, GetLastError());
+				{
+					mappedSize = 32 * 1024 * 1024;
+					TimerScope ts(g_kernelStats.createFileMapping);
+					mappingHandle = True_CreateFileMappingW(INVALID_HANDLE_VALUE, NULL, PAGE_READWRITE | SEC_RESERVE, ToHigh(reserveSize), ToLow(reserveSize), NULL);
+					if (!mappingHandle)
+						FatalError(1348, L"CreateFileMappingW failed trying to reserve %llu. (Error code: %u)", reserveSize, GetLastError());
+				}
+				TimerScope ts(g_kernelStats.mapViewOfFile);
 				baseAddress = (u8*)True_MapViewOfFile(mappingHandle, FILE_MAP_READ | FILE_MAP_WRITE, 0, 0, mappedSize);
 				if (!baseAddress)
 					FatalError(1353, L"MapViewOfFile failed trying to map %llu. ReservedSize: %llu (Error code: %u)", mappedSize, reserveSize, GetLastError());
diff --git a/Engine/Source/Programs/UnrealBuildAccelerator/Host/Private/UbaExports.cpp b/Engine/Source/Programs/UnrealBuildAccelerator/Host/Private/UbaExports.cpp
--- a/Engine/Source/Programs/UnrealBuildAccelerator/Host/Private/UbaExports.cpp
+++ b/Engine/Source/Programs/UnrealBuildAccelerator/Host/Private/UbaExports.cpp
@@ -364,7 +364,7 @@ extern "C"
 		server->PrintSummary(logger);
 		server->GetStorage().PrintSummary(logger);
 		server->GetServer().PrintSummary(logger);
-		uba::SystemStats::GetGlobal().Print(logger, true);
+		uba::KernelStats::GetGlobal().Print(logger, true);
 		uba::PrintContentionSummary(logger);
 	}
 	void SessionServer_CancelAll(uba::SessionServer* server)
diff --git a/Engine/Source/Programs/UnrealBuildAccelerator/Test/Private/UbaTestBasics.h b/Engine/Source/Programs/UnrealBuildAccelerator/Test/Private/UbaTestBasics.h
--- a/Engine/Source/Programs/UnrealBuildAccelerator/Test/Private/UbaTestBasics.h
+++ b/Engine/Source/Programs/UnrealBuildAccelerator/Test/Private/UbaTestBasics.h
@@ -54,7 +54,7 @@ namespace uba
 
 	bool TestEvents(Logger& logger, const StringBufferBase& rootDir)
 	{
-		for (uint i=0; i!=2; ++i)
+		for (u32 i=0; i!=2; ++i)
 		{
 			Event ev;
 			if (!ev.Create(true, i == 1))
diff --git a/Engine/Source/Programs/UnrealBuildAccelerator/Visualizer/Private/UbaTraceReader.cpp b/Engine/Source/Programs/UnrealBuildAccelerator/Visualizer/Private/UbaTraceReader.cpp
--- a/Engine/Source/Programs/UnrealBuildAccelerator/Visualizer/Private/UbaTraceReader.cpp
+++ b/Engine/Source/Programs/UnrealBuildAccelerator/Visualizer/Private/UbaTraceReader.cpp
@@ -498,17 +498,25 @@ namespace uba
 			ProcessStats processStats;
 			SessionStats sessionStats;
 			StorageStats storageStats;
-			SystemStats systemStats;
+			KernelStats kernelStats;
 
 			const u8* dataStart = reader.GetPositionData();
 			processStats.Read(reader, out.version);
 
-			bool isRemote = sessionIndex != 0;
-			if (isRemote && out.version >= 7)
+			process.isRemote = sessionIndex != 0;
+			if (process.isRemote)
 			{
-				sessionStats.Read(reader, out.version);
-				storageStats.Read(reader);
-				systemStats.Read(reader, out.version);
+				if (out.version >= 7)
+				{
+					sessionStats.Read(reader, out.version);
+					storageStats.Read(reader, out.version);
+					kernelStats.Read(reader, out.version);
+				}
+			}
+			else if (out.version >= 30)
+			{
+				storageStats.Read(reader, out.version);
+				kernelStats.Read(reader, out.version);
 			}
 			const u8* dataEnd = reader.GetPositionData();
 			process.stats.resize(dataEnd - dataStart);
@@ -563,11 +571,11 @@ namespace uba
 			ProcessStats processStats;
 			SessionStats sessionStats;
 			StorageStats storageStats;
-			SystemStats systemStats;
+			KernelStats kernelStats;
 			processStats.Read(reader, out.version);
 			sessionStats.Read(reader, out.version);
-			storageStats.Read(reader);
-			systemStats.Read(reader, out.version);
+			storageStats.Read(reader, out.version);
+			kernelStats.Read(reader, out.version);
 
 			const u8* dataEnd = reader.GetPositionData();
 			process.stats.resize(dataEnd - dataStart);
@@ -908,14 +916,14 @@ namespace uba
 			process.returned = !success;
 
 			CacheStats cacheStats;
-			SystemStats systemStats;
+			KernelStats kernelStats;
 			StorageStats storageStats;
 			const u8* dataStart = reader.GetPositionData();
 			cacheStats.Read(reader, out.version);
 			if (success || out.version >= 29)
 			{
-				storageStats.Read(reader);
-				systemStats.Read(reader, out.version);
+				storageStats.Read(reader, out.version);
+				kernelStats.Read(reader, out.version);
 			}
 			const u8* dataEnd = reader.GetPositionData();
 			process.stats.resize(dataEnd - dataStart);
diff --git a/Engine/Source/Programs/UnrealBuildAccelerator/Visualizer/Private/UbaTraceReader.h b/Engine/Source/Programs/UnrealBuildAccelerator/Visualizer/Private/UbaTraceReader.h
--- a/Engine/Source/Programs/UnrealBuildAccelerator/Visualizer/Private/UbaTraceReader.h
+++ b/Engine/Source/Programs/UnrealBuildAccelerator/Visualizer/Private/UbaTraceReader.h
@@ -25,6 +25,7 @@ namespace uba
 			bool bitmapDirty = true;
 			bool returned = false;
 			bool cacheFetch = false;
+			bool isRemote = false;
 			u64 createFilesTime = 0;
 			u64 writeFilesTime = 0;
 			Vector<u8> stats;
diff --git a/Engine/Source/Programs/UnrealBuildAccelerator/Visualizer/Private/UbaVisualizer.cpp b/Engine/Source/Programs/UnrealBuildAccelerator/Visualizer/Private/UbaVisualizer.cpp
--- a/Engine/Source/Programs/UnrealBuildAccelerator/Visualizer/Private/UbaVisualizer.cpp
+++ b/Engine/Source/Programs/UnrealBuildAccelerator/Visualizer/Private/UbaVisualizer.cpp
@@ -1383,7 +1383,7 @@ namespace uba
 				ProcessStats processStats;
 				SessionStats sessionStats;
 				StorageStats storageStats;
-				SystemStats systemStats;
+				KernelStats kernelStats;
 				CacheStats cacheStats;
 
 				if (process.cacheFetch)
@@ -1395,8 +1395,8 @@ namespace uba
 					cacheStats.Read(reader, m_traceView.version);
 					if (reader.GetLeft())
 					{
-						storageStats.Read(reader);
-						systemStats.Read(reader, m_traceView.version);
+						storageStats.Read(reader, m_traceView.version);
+						kernelStats.Read(reader, m_traceView.version);
 					}
 				}
 				else
@@ -1405,16 +1405,17 @@ namespace uba
 
 					if (reader.GetLeft())
 					{
-						sessionStats.Read(reader, m_traceView.version);
-						storageStats.Read(reader);
-						systemStats.Read(reader, m_traceView.version);
+						if (process.isRemote)
+							sessionStats.Read(reader, m_traceView.version);
+						storageStats.Read(reader, m_traceView.version);
+						kernelStats.Read(reader, m_traceView.version);
 					}
 				}
 
 				if (processStats.hostTotalTime)
 				{
 					logger.Info(L"");
-					logger.Info(L"  ----------- Process stats -----------");
+					logger.Info(L"  ----------- Detours stats -----------");
 					processStats.Print(logger, m_traceView.frequency);
 				}
 
@@ -1439,11 +1440,11 @@ namespace uba
 					storageStats.Print(logger, m_traceView.frequency);
 				}
 
-				if (!systemStats.IsEmpty())
+				if (!kernelStats.IsEmpty())
 				{
 					logger.Info(L"");
-					logger.Info(L"  ----------- System stats ------------");
-					systemStats.Print(logger, false, m_traceView.frequency);
+					logger.Info(L"  ----------- Kernel stats ------------");
+					kernelStats.Print(logger, false, m_traceView.frequency);
 				}
 
 				auto findIt = m_traceView.cacheWrites.find(process.id);
@@ -2165,7 +2166,6 @@ namespace uba
 
 	void Visualizer::WriteProcessStats(Logger& out, TraceView::Process& process)
 	{
-		bool hasStorageStats = true;
 		bool hasExited = process.stop != ~u64(0);
 		out.Info(L"  %ls", process.description.c_str());
 		out.Info(L"  Start:     %ls", TimeToText(process.start, true).str);
@@ -2182,29 +2182,39 @@ namespace uba
 			ProcessStats processStats;
 			SessionStats sessionStats;
 			StorageStats storageStats;
-			SystemStats systemStats;
+			KernelStats kernelStats;
 
 			processStats.Read(reader, m_traceView.version);
 			if (reader.GetLeft())
 			{
-				sessionStats.Read(reader, m_traceView.version);
-				storageStats.Read(reader);
-				systemStats.Read(reader, m_traceView.version);
+				if (process.isRemote)
+					sessionStats.Read(reader, m_traceView.version);
+				storageStats.Read(reader, m_traceView.version);
+				kernelStats.Read(reader, m_traceView.version);
 			}
 
-			out.Info(L"  ----------- Process stats -----------");
+			out.Info(L"  ----------- Detours stats -----------");
 			processStats.Print(out, m_traceView.frequency);
-			if (hasStorageStats)
+
+			if (!sessionStats.IsEmpty())
 			{
 				out.Info(L"");
 				out.Info(L"  ----------- Session stats -----------");
 				sessionStats.Print(out, m_traceView.frequency);
+			}
+
+			if (!storageStats.IsEmpty())
+			{
 				out.Info(L"");
 				out.Info(L"  ----------- Storage stats -----------");
 				storageStats.Print(out, m_traceView.frequency);
+			}
+
+			if (!kernelStats.IsEmpty())
+			{
 				out.Info(L"");
-				out.Info(L"  ----------- System stats ------------");
-				systemStats.Print(out, false, m_traceView.frequency);
+				out.Info(L"  ----------- Kernel stats ------------");
+				kernelStats.Print(out, false, m_traceView.frequency);
 			}
 		}
 	}
