[UBA]
* Added support for adding high priorty work
* Changed so cache server puts bucket maintenance work at high priority

diff --git a/Engine/Source/Programs/UnrealBuildAccelerator/Common/Private/UbaCacheServer.cpp b/Engine/Source/Programs/UnrealBuildAccelerator/Common/Private/UbaCacheServer.cpp
--- a/Engine/Source/Programs/UnrealBuildAccelerator/Common/Private/UbaCacheServer.cpp
+++ b/Engine/Source/Programs/UnrealBuildAccelerator/Common/Private/UbaCacheServer.cpp
@@ -691,7 +691,7 @@ namespace uba
 					bucket.m_cacheEntryLookup.erase(key);
 
 				totalEntryCount += bucket.totalEntryCount;
-			});
+			}, TC(""), true);
 
 			// Reset deleted cas files and update it again..
 			deletedCasFiles.clear();
@@ -870,7 +870,7 @@ namespace uba
 			bucket.needsSave = true;
 
 			m_logger.Info(TC("    Bucket %u Done (%s). CacheEntries: %llu (%s) PathTable: %s CasTable: %s"), bucket.index, TimeToText(GetTime() - bucketStartTime).str, bucket.totalEntryCount.load(), BytesToText(bucket.totalEntrySize.load()).str, BytesToText(bucket.m_pathTable.GetSize()).str, BytesToText(bucket.m_casKeyTable.GetSize()).str);
-		});
+		}, TC(""), true);
 
 		// Need to make sure all cas entries are dropped before saving cas table
 		u64 dropStartTime = GetTime();
diff --git a/Engine/Source/Programs/UnrealBuildAccelerator/Common/Private/UbaNetworkServer.cpp b/Engine/Source/Programs/UnrealBuildAccelerator/Common/Private/UbaNetworkServer.cpp
--- a/Engine/Source/Programs/UnrealBuildAccelerator/Common/Private/UbaNetworkServer.cpp
+++ b/Engine/Source/Programs/UnrealBuildAccelerator/Common/Private/UbaNetworkServer.cpp
@@ -779,14 +779,23 @@ namespace uba
 		}
 	}
 
-	void NetworkServer::AddWork(const Function<void()>& work, u32 count, const tchar* desc)
+	void NetworkServer::AddWork(const Function<void()>& work, u32 count, const tchar* desc, bool highPriority)
 	{
 		SCOPED_WRITE_LOCK(m_additionalWorkLock, lock);
 		for (u32 i = 0; i != count; ++i)
 		{
-			m_additionalWork.push_back({ work });
-			if (m_workTracker)
-				m_additionalWork.back().desc = desc;
+			if (highPriority)
+			{
+				m_additionalWork.push_front({ work });
+				if (m_workTracker)
+					m_additionalWork.front().desc = desc;
+			}
+			else
+			{
+				m_additionalWork.push_back({ work });
+				if (m_workTracker)
+					m_additionalWork.back().desc = desc;
+			}
 		}
 		lock.Leave();
 
diff --git a/Engine/Source/Programs/UnrealBuildAccelerator/Common/Private/UbaWorkManager.cpp b/Engine/Source/Programs/UnrealBuildAccelerator/Common/Private/UbaWorkManager.cpp
--- a/Engine/Source/Programs/UnrealBuildAccelerator/Common/Private/UbaWorkManager.cpp
+++ b/Engine/Source/Programs/UnrealBuildAccelerator/Common/Private/UbaWorkManager.cpp
@@ -99,14 +99,24 @@ namespace uba
 	}
 
 
-	void WorkManagerImpl::AddWork(const Function<void()>& work, u32 count, const tchar* desc)
+	void WorkManagerImpl::AddWork(const Function<void()>& work, u32 count, const tchar* desc, bool highPriority)
 	{
 		SCOPED_WRITE_LOCK(m_workLock, lock);
+		bool trackWork = m_workTracker.load();
 		for (u32 i = 0; i != count; ++i)
 		{
-			m_work.push_back({ work });
-			if (m_workTracker.load())
-				m_work.back().desc = desc;
+			if (highPriority)
+			{
+				m_work.push_front({ work });
+				if (trackWork)
+					m_work.front().desc = desc;
+			}
+			else
+			{
+				m_work.push_back({ work });
+				if (trackWork)
+					m_work.back().desc = desc;
+			}
 		}
 		lock.Leave();
 
diff --git a/Engine/Source/Programs/UnrealBuildAccelerator/Common/Public/UbaNetworkServer.h b/Engine/Source/Programs/UnrealBuildAccelerator/Common/Public/UbaNetworkServer.h
--- a/Engine/Source/Programs/UnrealBuildAccelerator/Common/Public/UbaNetworkServer.h
+++ b/Engine/Source/Programs/UnrealBuildAccelerator/Common/Public/UbaNetworkServer.h
@@ -72,7 +72,7 @@ namespace uba
 		//void RegisterOnConnection(u8 id, const OnConnectionFunction& func);
 		//void UnregisterOnConnection(u8 id);
 
-		virtual void AddWork(const Function<void()>& work, u32 count, const tchar* desc) override final;
+		virtual void AddWork(const Function<void()>& work, u32 count, const tchar* desc, bool highPriority = false) override final;
 		virtual void DoWork(u32 count = 1) override final;
 		virtual u32 GetWorkerCount() override final;
 
diff --git a/Engine/Source/Programs/UnrealBuildAccelerator/Common/Public/UbaWorkManager.h b/Engine/Source/Programs/UnrealBuildAccelerator/Common/Public/UbaWorkManager.h
--- a/Engine/Source/Programs/UnrealBuildAccelerator/Common/Public/UbaWorkManager.h
+++ b/Engine/Source/Programs/UnrealBuildAccelerator/Common/Public/UbaWorkManager.h
@@ -17,7 +17,7 @@ namespace uba
 	class WorkManager
 	{
 	public:
-		virtual void AddWork(const Function<void()>& work, u32 count, const tchar* desc) = 0;
+		virtual void AddWork(const Function<void()>& work, u32 count, const tchar* desc, bool highPriority = false) = 0;
 		virtual u32 GetWorkerCount() = 0;
 		virtual void DoWork(u32 count = 1) = 0;
 
@@ -28,7 +28,7 @@ namespace uba
 		WorkTracker* GetWorkTracker() { return m_workTracker; }
 
 		template<typename TContainer, typename TFunc>
-		void ParallelFor(u32 workCount, TContainer& container, const TFunc& func, const tchar* description = TC(""));
+		void ParallelFor(u32 workCount, TContainer& container, const TFunc& func, const tchar* description = TC(""), bool highPriority = false);
 
 	protected:
 		Atomic<WorkTracker*> m_workTracker;
@@ -40,7 +40,7 @@ namespace uba
 	public:
 		WorkManagerImpl(u32 workerCount);
 		virtual ~WorkManagerImpl();
-		virtual void AddWork(const Function<void()>& work, u32 count, const tchar* desc) override;
+		virtual void AddWork(const Function<void()>& work, u32 count, const tchar* desc, bool highPriority = false) override;
 		virtual u32 GetWorkerCount() override;
 		virtual void DoWork(u32 count = 1) override;
 		void FlushWork();
@@ -73,7 +73,7 @@ namespace uba
 
 
 	template<typename TContainer, typename TFunc>
-	void WorkManager::ParallelFor(u32 workCount, TContainer& container, const TFunc& func, const tchar* description)
+	void WorkManager::ParallelFor(u32 workCount, TContainer& container, const TFunc& func, const tchar* description, bool highPriority)
 	{
 		#if !defined( __clang_analyzer__ ) // Static analyzer claims context can leak but it can only leak when process terminates (and then it doesn't matter)
 
@@ -129,7 +129,7 @@ namespace uba
 				}
 			};
 		
-		AddWork(work, workCount, description);
+		AddWork(work, workCount, description, highPriority);
 		work();
 		doneEvent.IsSet();
 
