[UBA]
* Changed so SendBatchMessages can take in a functor that can be used to run code in between sending messages and waiting for response
* Changed so FileFetcher::RetrieveFile use file mapping when writing to file. This will reduce syscall count and speed up file writing

diff --git a/Engine/Source/Programs/UnrealBuildAccelerator/Common/Private/UbaStorageUtils.cpp b/Engine/Source/Programs/UnrealBuildAccelerator/Common/Private/UbaStorageUtils.cpp
--- a/Engine/Source/Programs/UnrealBuildAccelerator/Common/Private/UbaStorageUtils.cpp
+++ b/Engine/Source/Programs/UnrealBuildAccelerator/Common/Private/UbaStorageUtils.cpp
@@ -101,7 +101,7 @@ namespace uba
 		return ToCasKey(hasher, storeCompressed);
 	}
 
-	bool SendBatchMessages(Logger& logger, NetworkClient& client, u16 fetchId, u8* slot, u64 capacity, u64 left, u32 messageMaxSize, u32& readIndex, u32& responseSize)
+	bool SendBatchMessages(Logger& logger, NetworkClient& client, u16 fetchId, u8* slot, u64 capacity, u64 left, u32 messageMaxSize, u32& readIndex, u32& responseSize, const Function<bool()>& runInWaitFunc)
 	{
 		responseSize = 0;
 
@@ -144,6 +144,10 @@ namespace uba
 			break;
 		}
 
+		if (runInWaitFunc)
+			if (!runInWaitFunc())
+				return false;
+
 		for (u32 i=0; i!=inFlightCount; ++i)
 		{
 			Entry& entry = entries[i];
@@ -456,26 +460,50 @@ namespace uba
 		}
 
 		bytesReceived = fileSize;
-		sizeOnDisk = destinationIsCompressed ? fileSize : actualSize;
+		sizeOnDisk = destinationIsCompressed ? (sizeof(CompressedObjFileHeader) + fileSize) : actualSize;
 
 		FileAccessor destinationFile(logger, destination);
+
+		constexpr bool useFileMapping = true;
+		u8* fileMappingMem = nullptr;
+
 		if (!destinationMem)
-			if (!destinationFile.CreateWrite(false, DefaultAttributes(), sizeOnDisk, m_tempPath.data))
-				return false;
+		{
+			if (useFileMapping)
+			{
+				if (!destinationFile.CreateMemoryWrite(false, DefaultAttributes(), sizeOnDisk))
+					return false;
+				fileMappingMem = destinationFile.GetData();
+			}
+			else
+			{
+				if (!destinationFile.CreateWrite(false, DefaultAttributes(), sizeOnDisk, m_tempPath.data))
+					return false;
+			}
+		}
 
+		u64 destOffset = 0;
 
-		auto WriteDestination = [&](const void* source, u64 sourceSize, u64 sourceOffset = 0)
+		auto WriteDestination = [&](const void* source, u64 sourceSize)
 			{
-				if (!destinationMem)
+				if (fileMappingMem)
 				{
-					if (!destinationFile.Write(source, sourceSize, sourceOffset))
-						return false;
+					TimerScope ts(m_stats.memoryCopy);
+					MapMemoryCopy(fileMappingMem + destOffset, source, sourceSize);
+					destOffset += sourceSize;
 				}
-				else
+				else if (destinationMem)
 				{
+					TimerScope ts(m_stats.memoryCopy);
 					void* mem = destinationMem->Allocate(sourceSize, 1, TC(""));
 					memcpy(mem, source, sourceSize);
 				}
+				else
+				{
+					if (!destinationFile.Write(source, sourceSize, destOffset))
+						return false;
+					destOffset += sourceSize;
+				}
 				return true;
 			};
 
@@ -487,22 +515,28 @@ namespace uba
 			if (!WriteDestination(&header, sizeof(header)))
 				return false;
 
-			if (!WriteDestination(readBuffer, responseSize))
-				return false;
+			u8* source = slot;
+			u8* lastSource = readBuffer;
+			u64 lastResponseSize = responseSize;
+			auto writePrev = [&]() { return WriteDestination(lastSource, lastResponseSize); };
 
 			u64 leftCompressed = fileSize - responseSize;
 			while (leftCompressed)
 			{
 				if (fetchId == u16(~0))
 					return logger.Error(TC("Cas content error (2). Server believes %s was only one segment but client sees more. "));//UncompressedSize: %llu LeftUncompressed: %llu Size: %llu Left to read: %llu ResponseSize: %u. (%s)"), destination, actualSize, leftUncompressed, fileSize, left, responseSize, CasKeyString(casKey).str);
-				if (!SendBatchMessages(logger, client, fetchId, slot, BufferSlotSize, leftCompressed, sizeOfFirstMessage, readIndex, responseSize))
+
+				if (!SendBatchMessages(logger, client, fetchId, source, BufferSlotHalfSize, leftCompressed, sizeOfFirstMessage, readIndex, responseSize, writePrev))
 					return logger.Error(TC("Failed to send batched messages to server (%s)"), CasKeyString(casKey).str);
 
-				if (!WriteDestination(slot, responseSize))
-					return false;
+				lastSource = source;
+				lastResponseSize = responseSize;
+				source = source == slot ? slot + BufferSlotHalfSize : slot;
 
 				leftCompressed -= responseSize;
 			}
+			if (!writePrev())
+				return false;
 		}
 		else if (actualSize)
 		{
@@ -585,12 +619,14 @@ namespace uba
 				{
 					u8* decompressBuffer = slot + BufferSlotHalfSize;
 
-					//TimerScope ts(stats.decompressRecv);
-					OO_SINTa decompLen = OodleLZ_Decompress(readBuffer, int(compressedSize), decompressBuffer, int(uncompressedSize));
-					if (decompLen != uncompressedSize)
-						return logger.Error(TC("Expected %u but got %i when decompressing %u bytes for file %s"), uncompressedSize, int(decompLen), compressedSize, destination);
+					{
+						TimerScope ts2(m_stats.decompressRecv);
+						OO_SINTa decompLen = OodleLZ_Decompress(readBuffer, int(compressedSize), decompressBuffer, int(uncompressedSize));
+						if (decompLen != uncompressedSize)
+							return logger.Error(TC("Expected %u but got %i when decompressing %u bytes for file %s"), uncompressedSize, int(decompLen), compressedSize, destination);
+					}
 
-					if (!WriteDestination(decompressBuffer, uncompressedSize, actualSize - leftUncompressed))
+					if (!WriteDestination(decompressBuffer, uncompressedSize))
 						return false;
 
 					leftUncompressed -= uncompressedSize;
diff --git a/Engine/Source/Programs/UnrealBuildAccelerator/Common/Public/UbaStorageUtils.h b/Engine/Source/Programs/UnrealBuildAccelerator/Common/Public/UbaStorageUtils.h
--- a/Engine/Source/Programs/UnrealBuildAccelerator/Common/Public/UbaStorageUtils.h
+++ b/Engine/Source/Programs/UnrealBuildAccelerator/Common/Public/UbaStorageUtils.h
@@ -16,7 +16,7 @@ namespace uba
 	static constexpr u8 DefaultCompressionLevel = OodleLZ_CompressionLevel_SuperFast;
 
 	CasKey CalculateCasKey(u8* fileMem, u64 fileSize, bool storeCompressed, WorkManager* workManager);
-	bool SendBatchMessages(Logger& logger, NetworkClient& client, u16 fetchId, u8* slot, u64 capacity, u64 left, u32 messageMaxSize, u32& readIndex, u32& responseSize);
+	bool SendBatchMessages(Logger& logger, NetworkClient& client, u16 fetchId, u8* slot, u64 capacity, u64 left, u32 messageMaxSize, u32& readIndex, u32& responseSize, const Function<bool()>& runInWaitFunc = {});
 	bool SendFile(Logger& logger, NetworkClient& client, const CasKey& casKey, const u8* sourceMem, u64 sourceSize, const tchar* hint);
 
 	struct FileSender
