[UBA]
* Added MapMemoryCopy function that wraps memcpy and add timer

diff --git a/Engine/Source/Programs/UnrealBuildAccelerator/Common/Private/UbaFileMapping.cpp b/Engine/Source/Programs/UnrealBuildAccelerator/Common/Private/UbaFileMapping.cpp
--- a/Engine/Source/Programs/UnrealBuildAccelerator/Common/Private/UbaFileMapping.cpp
+++ b/Engine/Source/Programs/UnrealBuildAccelerator/Common/Private/UbaFileMapping.cpp
@@ -296,6 +296,14 @@ namespace uba
 #endif
 	}
 
+	void MapMemoryCopy(void* dest, const void* source, u64 size)
+	{
+		auto& stats = KernelStats::GetCurrent();
+		ExtendedTimerScope ts(stats.memoryCopy);
+		stats.memoryCopy.bytes += size;
+		memcpy(dest, source, size);
+	}
+
 	FileMappingBuffer::FileMappingBuffer(Logger& logger, WorkManager* workManager)
 	:	m_logger(logger)
 	,	m_workManager(workManager)
diff --git a/Engine/Source/Programs/UnrealBuildAccelerator/Common/Private/UbaSession.cpp b/Engine/Source/Programs/UnrealBuildAccelerator/Common/Private/UbaSession.cpp
--- a/Engine/Source/Programs/UnrealBuildAccelerator/Common/Private/UbaSession.cpp
+++ b/Engine/Source/Programs/UnrealBuildAccelerator/Common/Private/UbaSession.cpp
@@ -2119,7 +2119,7 @@ namespace uba
 					{
 						if (!destinationFile.CreateMemoryWrite(false, attributes, fileSize, m_tempPath.data))
 							return false;
-						memcpy(destinationFile.GetData(), mem, fileSize);
+						MapMemoryCopy(destinationFile.GetData(), mem, fileSize);
 					}
 					else
 					{
@@ -2660,7 +2660,7 @@ namespace uba
 		if (!mem2)
 			return false;
 
-		memcpy(mem2, memoryBlock.memory, memoryBlock.writtenSize);
+		MapMemoryCopy(mem2, memoryBlock.memory, memoryBlock.writtenSize);
 		UnmapViewOfFile(mem2, memoryBlock.writtenSize, TC(""));
 
 		StringKey symFileKey = CaseInsensitiveFs ? ToStringKeyLower(exportsFile) : ToStringKey(exportsFile);
diff --git a/Engine/Source/Programs/UnrealBuildAccelerator/Common/Private/UbaStorageClient.cpp b/Engine/Source/Programs/UnrealBuildAccelerator/Common/Private/UbaStorageClient.cpp
--- a/Engine/Source/Programs/UnrealBuildAccelerator/Common/Private/UbaStorageClient.cpp
+++ b/Engine/Source/Programs/UnrealBuildAccelerator/Common/Private/UbaStorageClient.cpp
@@ -489,7 +489,7 @@ namespace uba
 					}
 					else
 					{
-						memcpy(writePos, readBuffer, responseSize);
+						MapMemoryCopy(writePos, readBuffer, responseSize);
 						writePos += responseSize;
 					}
 
@@ -608,7 +608,7 @@ namespace uba
 							}
 							else
 							{
-								memcpy(writePos, decompressBuffer, uncompressedSize);
+								MapMemoryCopy(writePos, decompressBuffer, uncompressedSize);
 								writePos += uncompressedSize;
 							}
 
diff --git a/Engine/Source/Programs/UnrealBuildAccelerator/Common/Public/UbaFileMapping.h b/Engine/Source/Programs/UnrealBuildAccelerator/Common/Public/UbaFileMapping.h
--- a/Engine/Source/Programs/UnrealBuildAccelerator/Common/Public/UbaFileMapping.h
+++ b/Engine/Source/Programs/UnrealBuildAccelerator/Common/Public/UbaFileMapping.h
@@ -54,6 +54,7 @@ namespace uba
 	bool UnmapViewOfFile(const void* lpBaseAddress, u64 bytesToUnmap, const tchar* hint);
 	bool CloseFileMapping(FileMappingHandle h);
 	bool DuplicateFileMapping(ProcHandle hSourceProcessHandle, FileMappingHandle hSourceHandle, ProcHandle hTargetProcessHandle, FileMappingHandle* lpTargetHandle, u32 dwDesiredAccess, bool bInheritHandle, u32 dwOptions);
+	void MapMemoryCopy(void* dest, const void* source, u64 size);
 
 	struct MappedView
 	{
