[UBA]
* Changed so writing to file/memory happens in parallel with batched messages to fetch data

diff --git a/Engine/Source/Programs/UnrealBuildAccelerator/Common/Private/UbaStorageUtils.cpp b/Engine/Source/Programs/UnrealBuildAccelerator/Common/Private/UbaStorageUtils.cpp
--- a/Engine/Source/Programs/UnrealBuildAccelerator/Common/Private/UbaStorageUtils.cpp
+++ b/Engine/Source/Programs/UnrealBuildAccelerator/Common/Private/UbaStorageUtils.cpp
@@ -515,7 +515,7 @@ namespace uba
 			if (!WriteDestination(&header, sizeof(header)))
 				return false;
 
-			u8* source = slot;
+			u8* source = slot + BufferSlotHalfSize;
 			u8* lastSource = readBuffer;
 			u64 lastResponseSize = responseSize;
 			auto writePrev = [&]() { return WriteDestination(lastSource, lastResponseSize); };
@@ -545,13 +545,24 @@ namespace uba
 			readBuffer += sizeof(u64); // Size is stored first
 			u64 maxReadSize = BufferSlotHalfSize - sizeof(u64);
 
+			u8* decompressBuffer = slot + BufferSlotHalfSize;
+			u32 lastDecompressSize = 0;
+			auto tryWriteDecompressed = [&]()
+				{
+					if (!lastDecompressSize)
+						return true;
+					u32 toWrite = lastDecompressSize;
+					lastDecompressSize = 0;
+					return WriteDestination(decompressBuffer, toWrite);
+				};
+
 			u64 leftCompressed = fileSize - responseSize;
 			do
 			{
 				// First read in a full decompressable block
 				bool isFirstInBlock = true;
 				u32 compressedSize = ~u32(0);
-				u32 uncompressedSize = ~u32(0);
+				u32 decompressedSize = ~u32(0);
 				u32 left = 0;
 				u32 overflow = 0;
 				do
@@ -560,7 +571,7 @@ namespace uba
 					{
 						if (fetchId == u16(~0))
 							return logger.Error(TC("Cas content error (2). Server believes %s was only one segment but client sees more. UncompressedSize: %llu LeftUncompressed: %llu Size: %llu Left to read: %llu ResponseSize: %u. (%s)"), destination, actualSize, leftUncompressed, fileSize, left, responseSize, CasKeyString(casKey).str);
-						if (!SendBatchMessages(logger, client, fetchId, readPosition, maxReadSize - u32(readPosition - readBuffer), leftCompressed, sizeOfFirstMessage, readIndex, responseSize))
+						if (!SendBatchMessages(logger, client, fetchId, readPosition, maxReadSize - u32(readPosition - readBuffer), leftCompressed, sizeOfFirstMessage, readIndex, responseSize, tryWriteDecompressed))
 							return logger.Error(TC("Failed to send batched messages to server (%s)"), CasKeyString(casKey).str);
 						leftCompressed -= responseSize;
 					}
@@ -576,7 +587,7 @@ namespace uba
 						isFirstInBlock = false;
 						u32* blockSize = (u32*)readBuffer;
 						compressedSize = blockSize[0];
-						uncompressedSize = blockSize[1];
+						decompressedSize = blockSize[1];
 						readBuffer += sizeof(u32) * 2;
 						maxReadSize = BufferSlotHalfSize - sizeof(u32) * 2;
 						u32 read = (responseSize + u32(readPosition - readBuffer));
@@ -615,21 +626,39 @@ namespace uba
 					}
 				} while (left);
 
-				// Then decompress
+
+				// Second, decompress
+				while (true)
 				{
-					u8* decompressBuffer = slot + BufferSlotHalfSize;
+					tryWriteDecompressed();
 
 					{
 						TimerScope ts2(m_stats.decompressRecv);
-						OO_SINTa decompLen = OodleLZ_Decompress(readBuffer, int(compressedSize), decompressBuffer, int(uncompressedSize));
-						if (decompLen != uncompressedSize)
-							return logger.Error(TC("Expected %u but got %i when decompressing %u bytes for file %s"), uncompressedSize, int(decompLen), compressedSize, destination);
+						OO_SINTa decompLen = OodleLZ_Decompress(readBuffer, int(compressedSize), decompressBuffer, int(decompressedSize));
+						if (decompLen != decompressedSize)
+							return logger.Error(TC("Expected %u but got %i when decompressing %u bytes for file %s"), decompressedSize, int(decompLen), compressedSize, destination);
 					}
 
-					if (!WriteDestination(decompressBuffer, uncompressedSize))
-						return false;
+					lastDecompressSize = decompressedSize;
+					leftUncompressed -= decompressedSize;
 
-					leftUncompressed -= uncompressedSize;
+					constexpr bool decompressMultiple = false; // This does not seem to be a win.. it batches more but didn't save any time
+
+					if (!decompressMultiple)
+						break;
+
+					if (overflow < 8)
+						break;
+					u8* nextBlock = readBuffer + compressedSize;
+					u32* blockSize = (u32*)nextBlock;
+					u32 compressedSize2 = blockSize[0];
+					if (overflow < compressedSize2 + 8)
+						break;
+					readBuffer += compressedSize + 8;
+
+					decompressedSize = blockSize[1];
+					compressedSize = compressedSize2;
+					overflow -= compressedSize + 8;
 				}
 
 				// Move overflow back to the beginning of the buffer and start the next block (if there is one)
@@ -642,6 +671,9 @@ namespace uba
 				if (overflow)
 					responseSize = 0;
 			} while (leftUncompressed);
+
+			if (!tryWriteDecompressed())
+				return false;
 		}
 
 		if (sendEnd)
