[UBA]
* Fixed bug in backward compatibility trace reader code

diff --git a/Engine/Source/Programs/UnrealBuildAccelerator/Common/Private/UbaStats.cpp b/Engine/Source/Programs/UnrealBuildAccelerator/Common/Private/UbaStats.cpp
--- a/Engine/Source/Programs/UnrealBuildAccelerator/Common/Private/UbaStats.cpp
+++ b/Engine/Source/Programs/UnrealBuildAccelerator/Common/Private/UbaStats.cpp
@@ -193,7 +193,7 @@ namespace uba
 	{
 		if (version < 30)
 		{
-			#define UBA_STORAGE_STAT(type, var, ver) uba::Read(reader, var, version);
+			#define UBA_STORAGE_STAT(type, var, ver) if (ver <= version) uba::Read(reader, var, version);
 			UBA_STORAGE_STATS
 			#undef UBA_STORAGE_STAT
 			return;
diff --git a/Engine/Source/Programs/UnrealBuildAccelerator/Common/Public/UbaTrace.h b/Engine/Source/Programs/UnrealBuildAccelerator/Common/Public/UbaTrace.h
--- a/Engine/Source/Programs/UnrealBuildAccelerator/Common/Public/UbaTrace.h
+++ b/Engine/Source/Programs/UnrealBuildAccelerator/Common/Public/UbaTrace.h
@@ -28,34 +28,39 @@ namespace uba
 		void* m_mem = nullptr;
 	};
 
+	#define UBA_TRACE_TYPES \
+		UBA_TRACE_TYPE(SessionAdded) \
+		UBA_TRACE_TYPE(SessionUpdate) \
+		UBA_TRACE_TYPE(ProcessAdded) \
+		UBA_TRACE_TYPE(ProcessExited) \
+		UBA_TRACE_TYPE(ProcessReturned) \
+		UBA_TRACE_TYPE(FileBeginFetch) \
+		UBA_TRACE_TYPE(FileEndFetch) \
+		UBA_TRACE_TYPE(FileBeginStore) \
+		UBA_TRACE_TYPE(FileEndStore) \
+		UBA_TRACE_TYPE(Summary) \
+		UBA_TRACE_TYPE(BeginWork) \
+		UBA_TRACE_TYPE(EndWork) \
+		UBA_TRACE_TYPE(String) \
+		UBA_TRACE_TYPE(SessionSummary) \
+		UBA_TRACE_TYPE(ProcessEnvironmentUpdated) \
+		UBA_TRACE_TYPE(SessionDisconnect) \
+		UBA_TRACE_TYPE(ProxyCreated) \
+		UBA_TRACE_TYPE(ProxyUsed) \
+		UBA_TRACE_TYPE(FileFetchLight) \
+		UBA_TRACE_TYPE(FileStoreLight) \
+		UBA_TRACE_TYPE(StatusUpdate) \
+		UBA_TRACE_TYPE(SessionNotification) \
+		UBA_TRACE_TYPE(CacheBeginFetch) \
+		UBA_TRACE_TYPE(CacheEndFetch) \
+		UBA_TRACE_TYPE(CacheBeginWrite) \
+		UBA_TRACE_TYPE(CacheEndWrite) \
+
 	enum TraceType : u8
 	{
-		TraceType_SessionAdded,
-		TraceType_SessionUpdate,
-		TraceType_ProcessAdded,
-		TraceType_ProcessExited,
-		TraceType_ProcessReturned,
-		TraceType_FileBeginFetch,
-		TraceType_FileEndFetch,
-		TraceType_FileBeginStore,
-		TraceType_FileEndStore,
-		TraceType_Summary,
-		TraceType_BeginWork,
-		TraceType_EndWork,
-		TraceType_String,
-		TraceType_SessionSummary,
-		TraceType_ProcessEnvironmentUpdated,
-		TraceType_SessionDisconnect,
-		TraceType_ProxyCreated,
-		TraceType_ProxyUsed,
-		TraceType_FileFetchLight,
-		TraceType_FileStoreLight,
-		TraceType_StatusUpdate,
-		TraceType_SessionNotification,
-		TraceType_CacheBeginFetch,
-		TraceType_CacheEndFetch,
-		TraceType_CacheBeginWrite,
-		TraceType_CacheEndWrite,
+		#define UBA_TRACE_TYPE(name) TraceType_##name,
+		UBA_TRACE_TYPES
+		#undef UBA_TRACE_TYPE
 	};
 
 	using Color = u32;
diff --git a/Engine/Source/Programs/UnrealBuildAccelerator/Visualizer/Private/UbaTraceReader.cpp b/Engine/Source/Programs/UnrealBuildAccelerator/Visualizer/Private/UbaTraceReader.cpp
--- a/Engine/Source/Programs/UnrealBuildAccelerator/Visualizer/Private/UbaTraceReader.cpp
+++ b/Engine/Source/Programs/UnrealBuildAccelerator/Visualizer/Private/UbaTraceReader.cpp
@@ -279,9 +279,13 @@ namespace uba
 		}
 
 		#if 0
-		StringBuffer<> type;
-		type.Appendf(L"TraceType: %u\r\n", traceType);
-		OutputDebugString(type.data);
+		static const tchar* traceTypeNames[] =
+		{
+			#define UBA_TRACE_TYPE(name) TC(#name),
+			UBA_TRACE_TYPES
+			#undef UBA_TRACE_TYPE
+		};
+		OutputDebugString(StringBuffer<>().Appendf(L"TraceType: %s\r\n", traceTypeNames[traceType]).data);
 		#endif
 
 		switch (traceType)
