[UBA]
* Added directory preparsing in UbaCacheClient where it parse the entire folder when needing the hash of a file

diff --git a/Engine/Source/Programs/UnrealBuildAccelerator/Common/Private/UbaCacheClient.cpp b/Engine/Source/Programs/UnrealBuildAccelerator/Common/Private/UbaCacheClient.cpp
--- a/Engine/Source/Programs/UnrealBuildAccelerator/Common/Private/UbaCacheClient.cpp
+++ b/Engine/Source/Programs/UnrealBuildAccelerator/Common/Private/UbaCacheClient.cpp
@@ -4,6 +4,7 @@
 #include "UbaApplicationRules.h"
 #include "UbaCacheEntry.h"
 #include "UbaCompactTables.h"
+#include "UbaDirectoryIterator.h"
 #include "UbaFileAccessor.h"
 #include "UbaNetworkMessage.h"
 #include "UbaProcessStartInfo.h"
@@ -60,6 +61,8 @@ namespace uba
 		m_reportMissReason = true;
 		#endif
 
+		m_useDirectoryPreParsing = info.useDirectoryPreparsing;
+
 		m_client.RegisterOnConnected([this]()
 			{
 				u32 retryCount = 0;
@@ -302,8 +305,10 @@ namespace uba
 		return true;
 	}
 
-	bool CacheClient::FetchFromCache(const RootPaths& rootPaths, u32 bucketId, const ProcessStartInfo& info)
+	bool CacheClient::FetchFromCache(bool& outCacheHit, const RootPaths& rootPaths, u32 bucketId, const ProcessStartInfo& info)
 	{
+		outCacheHit = false;
+
 		if (!m_connected)
 			return false;
 
@@ -419,8 +424,17 @@ namespace uba
 				}
 				else
 				{
+					StringBuffer<MaxPath> forKey;
+					forKey.Append(path);
+					if (CaseInsensitiveFs)
+						forKey.MakeLower();
+					StringKey fileNameKey = ToStringKey(forKey);
+
+					if (m_useDirectoryPreParsing)
+						PreparseDirectory(fileNameKey, path);
+
 					bool fileIsCompressed = IsFileCompressed(info, path);
-					m_storage.StoreCasKey(localCasKey, path.data, CasKeyZero, fileIsCompressed);
+					m_storage.StoreCasKey(localCasKey, fileNameKey, path.data, CasKeyZero, fileIsCompressed);
 					UBA_ASSERT(localCasKey == CasKeyZero || IsCompressed(localCasKey));
 				}
 
@@ -667,10 +681,21 @@ namespace uba
 							memcpy(localBlock.Allocate(toWrite, 1, TC("")), fileStart + lastWritten, toWrite);
 
 						FileAccessor destFile(logger, path.data);
-						if (!destFile.CreateWrite())
-							return logger.Error(TC("Failed to create file for cache output %s for %s"), path.data, info.description);
-						if (!destFile.Write(localBlock.memory, localBlock.writtenSize))
-							return false;
+
+						bool useFileMapping = true;
+						if (useFileMapping)
+						{
+							if (!destFile.CreateMemoryWrite(false, DefaultAttributes(), localBlock.writtenSize))
+								return logger.Error(TC("Failed to create file for cache output %s for %s"), path.data, info.description);
+							MapMemoryCopy(destFile.GetData(), localBlock.memory, localBlock.writtenSize);
+						}
+						else
+						{
+							if (!destFile.CreateWrite())
+								return logger.Error(TC("Failed to create file for cache output %s for %s"), path.data, info.description);
+							if (!destFile.Write(localBlock.memory, localBlock.writtenSize))
+								return false;
+						}
 						if (!destFile.Close(&fetcher.lastWritten))
 							return false;
 
@@ -693,6 +718,7 @@ namespace uba
 					if (!m_session.RegisterNewFile(path.data))
 						return false;
 				}
+				outCacheHit = true;
 				success = true;
 				return true;
 			}
@@ -1105,7 +1131,6 @@ namespace uba
 		u32 rootIndex = normalizedPath[0] - RootPaths::RootStartByte;
 		const TString& root = rootPaths.GetRoot(rootIndex);
 
-		StringBuffer<MaxPath> path;
 		outPath.Append(root).Append(normalizedPath.data + 1);
 		return true;
 	}
@@ -1119,4 +1144,49 @@ namespace uba
 			rules = m_session.GetRules(info);
 		return rules->StoreFileCompressed(filename);
 	}
+
+	void CacheClient::PreparseDirectory(const StringKey& fileNameKey, const StringBufferBase& filePath)
+	{
+		const tchar* lastSep = filePath.Last(PathSeparator);
+		if (!lastSep)
+			return;
+
+		StringBuffer<MaxPath> path;
+		path.Append(filePath.data, lastSep - filePath.data);
+		if (CaseInsensitiveFs)
+			path.MakeLower();
+
+		StringKeyHasher dirHasher;
+		dirHasher.Update(path.data, path.count);
+		StringKey pathKey = ToStringKey(dirHasher);
+
+		SCOPED_WRITE_LOCK(m_directoryPreparserLock, preparserLock);
+		auto insres = m_directoryPreparser.try_emplace(pathKey);
+		PreparedDir& dir = insres.first->second;
+		preparserLock.Leave();
+
+		SCOPED_WRITE_LOCK(dir.lock, preparserLock2);
+		if (dir.done)
+			return;
+
+		dir.done = true;
+
+		// It is likely this folder has already been handled by session if this file is verified
+		if (m_storage.IsFileVerified(fileNameKey))
+			return;
+
+		u32 dirLength = path.count;
+
+		// Traverse all files in directory and report the file information... but only if it has not been reported before.. we don't want to interfere with other reports
+		TraverseDir(m_logger, path.data, 
+			[&](const DirectoryEntry& e)
+			{
+				path.Resize(dirLength).Append('\\').Append(e.name, e.nameLen);
+				if (CaseInsensitiveFs)
+					path.MakeLower();
+
+				StringKey fileNameKey = ToStringKey(dirHasher, path.data, path.count);
+				m_storage.ReportFileInfoWeak(fileNameKey, e.lastWritten, e.size);
+			});
+	}
 }
diff --git a/Engine/Source/Programs/UnrealBuildAccelerator/Common/Private/UbaScheduler.cpp b/Engine/Source/Programs/UnrealBuildAccelerator/Common/Private/UbaScheduler.cpp
--- a/Engine/Source/Programs/UnrealBuildAccelerator/Common/Private/UbaScheduler.cpp
+++ b/Engine/Source/Programs/UnrealBuildAccelerator/Common/Private/UbaScheduler.cpp
@@ -406,7 +406,8 @@ namespace uba
 						ProcessStartInfo& si = exitInfo->startInfo->startInfo;
 						u64 startTime = GetTime();
 
-						if (m_cacheClient->FetchFromCache(*m_rootPaths[0], 0, si))
+						bool cacheHit = false;
+						if (m_cacheClient->FetchFromCache(cacheHit, *m_rootPaths[0], 0, si) && cacheHit)
 						{
 							auto process = new CachedProcess(si);
 							ProcessHandle ph(process);
diff --git a/Engine/Source/Programs/UnrealBuildAccelerator/Common/Private/UbaStorage.cpp b/Engine/Source/Programs/UnrealBuildAccelerator/Common/Private/UbaStorage.cpp
--- a/Engine/Source/Programs/UnrealBuildAccelerator/Common/Private/UbaStorage.cpp
+++ b/Engine/Source/Programs/UnrealBuildAccelerator/Common/Private/UbaStorage.cpp
@@ -1962,6 +1962,33 @@ namespace uba
 		return true;
 	}
 
+	bool StorageImpl::IsFileVerified(const StringKey& fileNameKey)
+	{
+		SCOPED_READ_LOCK(m_fileTableLookupLock, lookupLock);
+		auto findIt = m_fileTableLookup.find(fileNameKey);
+		if (findIt == m_fileTableLookup.end())
+			return false;
+		FileEntry& fileEntry = findIt->second;
+		lookupLock.Leave();
+		SCOPED_READ_LOCK(fileEntry.lock, entryLock);
+		return fileEntry.verified;
+	}
+
+	void StorageImpl::ReportFileInfoWeak(const StringKey& fileNameKey, u64 verifiedLastWriteTime, u64 verifiedSize)
+	{
+		SCOPED_READ_LOCK(m_fileTableLookupLock, lookupLock);
+		auto findIt = m_fileTableLookup.find(fileNameKey);
+		if (findIt == m_fileTableLookup.end())
+			return;
+		FileEntry& fileEntry = findIt->second;
+		lookupLock.Leave();
+
+		SCOPED_WRITE_LOCK(fileEntry.lock, entryLock);
+		if (fileEntry.verified)
+			return;
+		fileEntry.verified = fileEntry.lastWritten == verifiedLastWriteTime && fileEntry.size == verifiedSize;
+	}
+
 	bool StorageImpl::StoreCasKey(CasKey& out, const tchar* fileName, const CasKey& casKeyOverride, bool fileIsCompressed)
 	{
 		StringBuffer<> forKey;
@@ -1969,7 +1996,11 @@ namespace uba
 		if (CaseInsensitiveFs)
 			forKey.MakeLower();
 		StringKey fileNameKey = ToStringKey(forKey);
+		return StoreCasKey(out, fileNameKey, fileName, casKeyOverride, fileIsCompressed);
+	}
 
+	bool StorageImpl::StoreCasKey(CasKey& out, const StringKey& fileNameKey, const tchar* fileName, const CasKey& casKeyOverride, bool fileIsCompressed)
+	{
 		SCOPED_WRITE_LOCK(m_fileTableLookupLock, lookupLock);
 		auto insres = m_fileTableLookup.try_emplace(fileNameKey);
 		FileEntry& fileEntry = insres.first->second;
@@ -2412,7 +2443,8 @@ namespace uba
 
 					if (mappedView.memory)
 					{
-						memcpy(writePos, mappedView.memory, compressedFileSize);
+						TimerScope cts(stats.memoryCopy);
+						MapMemoryCopy(writePos, mappedView.memory, compressedFileSize);
 					}
 					else
 					{
diff --git a/Engine/Source/Programs/UnrealBuildAccelerator/Common/Public/UbaCacheClient.h b/Engine/Source/Programs/UnrealBuildAccelerator/Common/Public/UbaCacheClient.h
--- a/Engine/Source/Programs/UnrealBuildAccelerator/Common/Public/UbaCacheClient.h
+++ b/Engine/Source/Programs/UnrealBuildAccelerator/Common/Public/UbaCacheClient.h
@@ -2,6 +2,7 @@
 
 #pragma once
 
+#include "UbaHash.h"
 #include "UbaLogger.h"
 #include "UbaMemory.h"
 
@@ -27,6 +28,7 @@ namespace uba
 		NetworkClient& client;
 		Session& session;
 		bool reportMissReason = false;
+		bool useDirectoryPreparsing = true; // This is used to minimize syscalls. GetFileAttributes can be very expensive on cloud machines and we can enable this to minimize syscall count
 	};
 
 	class CacheClient
@@ -36,7 +38,7 @@ namespace uba
 		~CacheClient();
 
 		bool WriteToCache(const RootPaths& rootPaths, u32 bucketId, const ProcessStartInfo& info, const u8* inputs, u64 inputsSize, const u8* outputs, u64 outputsSize, u32 processId = 0);
-		bool FetchFromCache(const RootPaths& rootPaths, u32 bucketId, const ProcessStartInfo& info);
+		bool FetchFromCache(bool& outCacheHit, const RootPaths& rootPaths, u32 bucketId, const ProcessStartInfo& info);
 		bool RequestServerShutdown(const tchar* reason);
 
 		bool ExecuteCommand(Logger& logger, const tchar* command, const tchar* destinationFile = nullptr, const tchar* additionalInfo = nullptr);
@@ -58,12 +60,14 @@ namespace uba
 
 		bool GetLocalPathAndCasKey(Bucket& bucket, const RootPaths& rootPaths, StringBufferBase& outPath, CasKey& outKey, CompactCasKeyTable& casKeyTable, CompactPathTable& pathTable, u32 offset);
 		bool IsFileCompressed(const ProcessStartInfo& info, const StringView& filename);
+		void PreparseDirectory(const StringKey& fileNameKey, const StringBufferBase& filePath);
 
 		MutableLogger m_logger;
 		StorageImpl& m_storage;
 		NetworkClient& m_client;
 		Session& m_session;
 		bool m_reportMissReason;
+		bool m_useDirectoryPreParsing;
 
 		Atomic<bool> m_connected;
 
@@ -72,6 +76,10 @@ namespace uba
 
 		ReaderWriterLock m_sendOneAtTheTimeLock;
 
+		ReaderWriterLock m_directoryPreparserLock;
+		struct PreparedDir { ReaderWriterLock lock; bool done = false; };
+		UnorderedMap<StringKey, PreparedDir> m_directoryPreparser;
+
 		CacheClient(const CacheClient&) = delete;
 		CacheClient& operator=(const CacheClient&) = delete;
 	};
diff --git a/Engine/Source/Programs/UnrealBuildAccelerator/Common/Public/UbaStorage.h b/Engine/Source/Programs/UnrealBuildAccelerator/Common/Public/UbaStorage.h
--- a/Engine/Source/Programs/UnrealBuildAccelerator/Common/Public/UbaStorage.h
+++ b/Engine/Source/Programs/UnrealBuildAccelerator/Common/Public/UbaStorage.h
@@ -147,6 +147,9 @@ namespace uba
 		CasKey CalculateCasKey(const tchar* fileName, FileHandle fileHandle, u64 fileSize, bool storeCompressed);
 		CasKey CalculateCasKey(u8* fileMem, u64 fileSize, bool storeCompressed);
 		bool StoreCasKey(CasKey& out, const tchar* fileName, const CasKey& casKeyOverride, bool fileIsCompressed);
+		bool StoreCasKey(CasKey& out, const StringKey& fileNameKey, const tchar* fileName, const CasKey& casKeyOverride, bool fileIsCompressed);
+		bool IsFileVerified(const StringKey& fileNameKey);
+		void ReportFileInfoWeak(const StringKey& fileNameKey, u64 verifiedLastWriteTime, u64 verifiedSize);
 
 		virtual bool WriteCompressed(WriteResult& out, const tchar* from, const tchar* toFile);
 		virtual bool WriteCompressed(WriteResult& out, const tchar* from, FileHandle readHandle, u8* readMem, u64 fileSize, const tchar* toFile, const void* header, u64 headerSize) final;
diff --git a/Engine/Source/Programs/UnrealBuildAccelerator/Host/Private/UbaExports.cpp b/Engine/Source/Programs/UnrealBuildAccelerator/Host/Private/UbaExports.cpp
--- a/Engine/Source/Programs/UnrealBuildAccelerator/Host/Private/UbaExports.cpp
+++ b/Engine/Source/Programs/UnrealBuildAccelerator/Host/Private/UbaExports.cpp
@@ -545,7 +545,9 @@ extern "C"
 	bool CacheClient_FetchFromCache(uba::CacheClient* cacheClient, uba::RootPaths* rootPaths, uba::u32 bucket, const uba::ProcessStartInfo& info)
 	{
 		using namespace uba;
-		return cacheClient->FetchFromCache(*rootPaths, bucket, info);
+		bool cacheHit = false;
+		bool res = cacheClient->FetchFromCache(cacheHit, *rootPaths, bucket, info);
+		return res && cacheHit;
 	}
 
 	void CacheClient_RequestServerShutdown(uba::CacheClient* cacheClient, const uba::tchar* reason)
diff --git a/Engine/Source/Programs/UnrealBuildAccelerator/Test/Private/UbaTestBasics.h b/Engine/Source/Programs/UnrealBuildAccelerator/Test/Private/UbaTestBasics.h
--- a/Engine/Source/Programs/UnrealBuildAccelerator/Test/Private/UbaTestBasics.h
+++ b/Engine/Source/Programs/UnrealBuildAccelerator/Test/Private/UbaTestBasics.h
@@ -421,7 +421,7 @@ namespace uba
 				{
 					if (strLen != 1)
 						success = false;
-					if (str[0] != RootPaths::RootStartByte + (IsWindows ? 2 : 1))
+					if (str[0] != RootPaths::RootStartByte + (IsWindows ? 3 : 1)) // 3 is because each root has three entries on windows
 						success = false;
 					rootPos = str[0];
 				}
