[UBA]
* Added so cache server can handle StoreEntryDone messages with bool saying if it succeeded or not

diff --git a/Engine/Source/Programs/UnrealBuildAccelerator/Common/Private/UbaCacheServer.cpp b/Engine/Source/Programs/UnrealBuildAccelerator/Common/Private/UbaCacheServer.cpp
--- a/Engine/Source/Programs/UnrealBuildAccelerator/Common/Private/UbaCacheServer.cpp
+++ b/Engine/Source/Programs/UnrealBuildAccelerator/Common/Private/UbaCacheServer.cpp
@@ -991,21 +991,28 @@ namespace uba
 			auto& connectionBucket = GetConnectionBucket(connectionInfo, reader);
 			CasKey cmdKey = reader.ReadCasKey();
 
+			bool success = true;
+			if (reader.GetLeft())
+				success = reader.ReadBool();
+
 			SCOPED_WRITE_LOCK(connectionBucket.cacheEntryLookupLock, lock2);
 			auto findIt = connectionBucket.cacheEntryLookup.find(cmdKey);
-			if (findIt != connectionBucket.cacheEntryLookup.end())
-			{
-				u64 id = connectionBucket.id;
-				Bucket& bucket = GetBucket(id);
+			if (findIt == connectionBucket.cacheEntryLookup.end())
+				return true;
+			auto g = MakeGuard([&]() { connectionBucket.cacheEntryLookup.erase(findIt); });
+			if (!success)
+				return true;
 
-				SCOPED_WRITE_LOCK(bucket.m_cacheEntryLookupLock, lock3);
-				auto insres = bucket.m_cacheEntryLookup.try_emplace(cmdKey);
-				auto& cacheEntries = insres.first->second;
-				lock2.Leave();
+			u64 id = connectionBucket.id;
+			Bucket& bucket = GetBucket(id);
 
-				SCOPED_WRITE_LOCK(cacheEntries.lock, lock4);
-				cacheEntries.entries.emplace_front(std::move(findIt->second));
-			}
+			SCOPED_WRITE_LOCK(bucket.m_cacheEntryLookupLock, lock3);
+			auto insres = bucket.m_cacheEntryLookup.try_emplace(cmdKey);
+			auto& cacheEntries = insres.first->second;
+			lock3.Leave();
+			
+			SCOPED_WRITE_LOCK(cacheEntries.lock, lock4);
+			cacheEntries.entries.emplace_front(std::move(findIt->second));
 			return true;
 		}
 		case CacheMessageType_FetchPathTable:
