[UBA]
* Removed duplicated code
* Added asserts to make sure no compressed obj files are used to calculate cas key
* Added missing locks around some code that should have locks

diff --git a/Engine/Source/Programs/UnrealBuildAccelerator/Common/Private/UbaCacheClient.cpp b/Engine/Source/Programs/UnrealBuildAccelerator/Common/Private/UbaCacheClient.cpp
--- a/Engine/Source/Programs/UnrealBuildAccelerator/Common/Private/UbaCacheClient.cpp
+++ b/Engine/Source/Programs/UnrealBuildAccelerator/Common/Private/UbaCacheClient.cpp
@@ -700,7 +700,7 @@ namespace uba
 							return false;
 
 						fetcher.sizeOnDisk = localBlock.writtenSize;
-						casKey = CalculateCasKey(localBlock.memory, localBlock.writtenSize, false, nullptr);
+						casKey = CalculateCasKey(localBlock.memory, localBlock.writtenSize, false, nullptr, path.data);
 					}
 					else
 					{
diff --git a/Engine/Source/Programs/UnrealBuildAccelerator/Common/Private/UbaSession.cpp b/Engine/Source/Programs/UnrealBuildAccelerator/Common/Private/UbaSession.cpp
--- a/Engine/Source/Programs/UnrealBuildAccelerator/Common/Private/UbaSession.cpp
+++ b/Engine/Source/Programs/UnrealBuildAccelerator/Common/Private/UbaSession.cpp
@@ -2093,7 +2093,7 @@ namespace uba
 				if (m_storeObjFilesCompressed && process.m_startInfo.rules->StoreFileCompressed(file.name))
 				{
 					Storage::WriteResult res;
-					CompressedObjFileHeader header { CalculateCasKey(mem, fileSize, true, m_workManager) };
+					CompressedObjFileHeader header { CalculateCasKey(mem, fileSize, true, m_workManager, file.name.c_str()) };
 
 					if (!m_storage.WriteCompressed(res, TC("MemoryMap"), InvalidFileHandle, mem, fileSize, file.name.c_str(), &header, sizeof(header)))
 						return false;
diff --git a/Engine/Source/Programs/UnrealBuildAccelerator/Common/Private/UbaStorage.cpp b/Engine/Source/Programs/UnrealBuildAccelerator/Common/Private/UbaStorage.cpp
--- a/Engine/Source/Programs/UnrealBuildAccelerator/Common/Private/UbaStorage.cpp
+++ b/Engine/Source/Programs/UnrealBuildAccelerator/Common/Private/UbaStorage.cpp
@@ -1510,10 +1510,12 @@ namespace uba
 				BinaryWriter writer(fileTableData.data(), 0, fileTableData.size());
 				for (auto& pair : m_fileTableLookup)
 				{
+					FileEntry& entry = pair.second;
+					SCOPED_READ_LOCK(entry.lock, entryLock);
 					writer.WriteStringKey(pair.first);
-					writer.WriteU64(pair.second.size);
-					writer.WriteU64(pair.second.lastWritten);
-					writer.WriteCasKey(pair.second.casKey);
+					writer.WriteU64(entry.size);
+					writer.WriteU64(entry.lastWritten);
+					writer.WriteCasKey(entry.casKey);
 				}
 				if (!tempFile.Write(fileTableData.data(), writer.GetPosition()))
 					return false;
@@ -2483,7 +2485,7 @@ namespace uba
 				u64 lastWriteTime = 0;
 				if (!destinationFile.Close(&lastWriteTime))
 					return false;
-
+				UBA_ASSERT(lastWriteTime);
 				if (lastWriteTime)
 				{
 					entry.casKey = casKey;
@@ -2553,7 +2555,10 @@ namespace uba
 		StringKey key = ToStringKey(forKey);
 		SCOPED_WRITE_LOCK(m_fileTableLookupLock, lock);
 		auto insres = m_fileTableLookup.try_emplace(key);
+		lock.Leave();
+
 		FileEntry& entry = insres.first->second;
+		SCOPED_WRITE_LOCK(entry.lock, lock2);
 		entry.casKey = casKey;
 		entry.lastWritten = lastWritten;
 		entry.size = size;
@@ -2642,7 +2647,7 @@ namespace uba
 	{
 		StorageStats& stats = Stats();
 		TimerScope ts(stats.calculateCasKey);
-		return uba::CalculateCasKey(fileMem, fileSize, storeCompressed, m_workManager);
+		return uba::CalculateCasKey(fileMem, fileSize, storeCompressed, m_workManager, nullptr);
 	}
 
 
@@ -2651,10 +2656,6 @@ namespace uba
 		StorageStats& stats = Stats();
 		TimerScope ts(stats.calculateCasKey);
 
-		CasKeyHasher hasher;
-
-		#ifndef __clang_analyzer__
-
 		if (fileSize > BufferSlotSize) // Note that when filesize is larger than BufferSlotSize the hash becomes a hash of hashes
 		{
 			FileMappingHandle fileMapping = uba::CreateFileMappingW(fileHandle, PAGE_READONLY, fileSize, fileName);
@@ -2679,98 +2680,22 @@ namespace uba
 						UnmapViewOfFile(fileData, fileSize, fileName);
 				});
 
-			struct WorkRec
-			{
-				Atomic<u64> refCount;
-				Atomic<u64> counter;
-				Atomic<u64> doneCounter;
-				u8* fileData = nullptr;
-				u64 workCount = 0;
-				u64 fileSize = 0;
-				bool error = false;
-				Vector<CasKey> keys;
-				const tchar* fileName = nullptr;
-				Event done;
-			};
-
-			u32 workCount = u32((fileSize + BufferSlotSize - 1) / BufferSlotSize);
-
-			WorkRec* rec = new WorkRec();
-			rec->fileData = fileData;
-			rec->workCount = workCount;
-			rec->fileSize = fileSize;
-			rec->fileName = fileName;
-			rec->keys.resize(workCount);
-			rec->done.Create(true);
-
-			auto work = [rec]()
-			{
-				while (true)
-				{
-					u64 index = rec->counter++;
-					if (index >= rec->workCount)
-					{
-						if (!--rec->refCount)
-							delete rec;
-						return 0;
-					}
-
-					u64 startOffset = BufferSlotSize*index;
-					u64 toRead = Min(BufferSlotSize, rec->fileSize - startOffset);
-					u8* slot = rec->fileData + startOffset;
-					CasKeyHasher hasher;
-					hasher.Update(slot, toRead);
-					rec->keys[index] = ToCasKey(hasher, false);
-
-					if (++rec->doneCounter == rec->workCount)
-						rec->done.Set();
-				}
-				return 0;
-			};
-
-			if (m_workManager)
-			{
-				u32 workerCount = Min(workCount, m_workManager->GetWorkerCount());
-				workerCount = Min(workerCount, MaxWorkItemsPerAction); // Cap this to not starve other things
-
-				rec->refCount = workerCount + 1; // We need to keep refcount up 1 to make sure it is not deleted before we read rec->written
-				m_workManager->AddWork(work, workerCount-1, TC("CalculateKey")); // We are a worker ourselves
-			}
-			else
-			{
-				rec->refCount = 2;
-			}
-
-			work();
-			rec->done.IsSet();
-
-			hasher.Update(rec->keys.data(), rec->keys.size()*sizeof(CasKey));
-
-			bool error = rec->error;
-
-			if (!--rec->refCount)
-				delete rec;
-
-			if (error)
-				return CasKeyZero;
+			return uba::CalculateCasKey(fileData, fileSize, storeCompressed, m_workManager, fileName);
 		}
-		else
+
+		CasKeyHasher hasher;
+		u8* slot = m_bufferSlots.Pop();
+		auto _ = MakeGuard([&](){ m_bufferSlots.Push(slot); });
+		u64 left = fileSize;
+		while (left)
 		{
-			u8* slot = m_bufferSlots.Pop();
-			auto _ = MakeGuard([&](){ m_bufferSlots.Push(slot); });
-			u64 left = fileSize;
-			while (left)
-			{
-				u32 toRead = u32(Min(left, BufferSlotSize));
-				if (!ReadFile(m_logger, fileName, fileHandle, slot, toRead))
-					return CasKeyZero;
-				hasher.Update(slot, toRead);
-				left -= toRead;
-			}
+			u32 toRead = u32(Min(left, BufferSlotSize));
+			if (!ReadFile(m_logger, fileName, fileHandle, slot, toRead))
+				return CasKeyZero;
+			hasher.Update(slot, toRead);
+			left -= toRead;
 		}
 
-		#endif // __clang_analyzer__
-
 		return ToCasKey(hasher, storeCompressed);
 	}
 
diff --git a/Engine/Source/Programs/UnrealBuildAccelerator/Common/Private/UbaStorageUtils.cpp b/Engine/Source/Programs/UnrealBuildAccelerator/Common/Private/UbaStorageUtils.cpp
--- a/Engine/Source/Programs/UnrealBuildAccelerator/Common/Private/UbaStorageUtils.cpp
+++ b/Engine/Source/Programs/UnrealBuildAccelerator/Common/Private/UbaStorageUtils.cpp
@@ -9,7 +9,7 @@
 
 namespace uba
 {
-	CasKey CalculateCasKey(u8* fileMem, u64 fileSize, bool storeCompressed, WorkManager* workManager)
+	CasKey CalculateCasKey(u8* fileMem, u64 fileSize, bool storeCompressed, WorkManager* workManager, const tchar* hint)
 	{
 		constexpr u32 MaxWorkItemsPerAction2 = 128; // Cap this to not starve other things
 		CasKeyHasher hasher;
@@ -17,6 +17,8 @@ namespace uba
 		if (fileSize == 0)
 			return ToCasKey(hasher, storeCompressed);
 
+		UBA_ASSERTF(fileSize < sizeof(CompressedObjFileHeader) || !((CompressedObjFileHeader*)fileMem)->IsValid(), TC("CalculateCasKey should never run on compressed file"));
+
 		#ifndef __clang_analyzer__
 
 		if (fileSize > BufferSlotSize) // Note that when filesize is larger than BufferSlotSize the hash becomes a hash of hashes
diff --git a/Engine/Source/Programs/UnrealBuildAccelerator/Common/Public/UbaStorageUtils.h b/Engine/Source/Programs/UnrealBuildAccelerator/Common/Public/UbaStorageUtils.h
--- a/Engine/Source/Programs/UnrealBuildAccelerator/Common/Public/UbaStorageUtils.h
+++ b/Engine/Source/Programs/UnrealBuildAccelerator/Common/Public/UbaStorageUtils.h
@@ -15,7 +15,7 @@ namespace uba
 	static constexpr u8 DefaultCompressor = OodleLZ_Compressor_Kraken;
 	static constexpr u8 DefaultCompressionLevel = OodleLZ_CompressionLevel_SuperFast;
 
-	CasKey CalculateCasKey(u8* fileMem, u64 fileSize, bool storeCompressed, WorkManager* workManager);
+	CasKey CalculateCasKey(u8* fileMem, u64 fileSize, bool storeCompressed, WorkManager* workManager, const tchar* hint);
 	bool SendBatchMessages(Logger& logger, NetworkClient& client, u16 fetchId, u8* slot, u64 capacity, u64 left, u32 messageMaxSize, u32& readIndex, u32& responseSize, const Function<bool()>& runInWaitFunc = {});
 	bool SendFile(Logger& logger, NetworkClient& client, const CasKey& casKey, const u8* sourceMem, u64 sourceSize, const tchar* hint);
 
