[UBA]
* Fixed annoying assert when pressing ctrl-c when running in debug
* Fixed so trace stream writes out processed stopped through cancel

diff --git a/Engine/Source/Programs/UnrealBuildAccelerator/Common/Private/UbaSessionServer.cpp b/Engine/Source/Programs/UnrealBuildAccelerator/Common/Private/UbaSessionServer.cpp
--- a/Engine/Source/Programs/UnrealBuildAccelerator/Common/Private/UbaSessionServer.cpp
+++ b/Engine/Source/Programs/UnrealBuildAccelerator/Common/Private/UbaSessionServer.cpp
@@ -519,7 +519,7 @@ namespace uba
 			m_trace.SessionDisconnect(sessionId);
 
 			sessionName.Append(s.name);
-			UBA_ASSERTF(s.usedSlotCount == returnCount, TC("Used slot count different than return count (%u vs %u)"), s.usedSlotCount, returnCount);
+			UBA_ASSERTF(s.usedSlotCount == returnCount || m_logger.isMuted, TC("Used slot count different than return count (%u vs %u)"), s.usedSlotCount, returnCount);
 			s.usedSlotCount -= returnCount;
 
 			if (s.enabled)
@@ -1678,6 +1678,13 @@ namespace uba
 			m_activeRemoteProcesses.erase(process);
 			queueLock.Leave();
 
+			StackBinaryWriter<1024> writer;
+			ProcessStats().Write(writer);
+			SessionStats().Write(writer);
+			StorageStats().Write(writer);
+			KernelStats().Write(writer);
+			m_trace.ProcessExited(process->m_processId, process->m_exitCode, writer.GetData(), writer.GetPosition(), Vector<ProcessLogLine>());
+
 			m_logger.Warning(TC("Cancelling remote active processes has not been tested. Notify devs"));
 
 			{
