[UBA]
* Fixed UbaCli linux compile error

diff --git a/Engine/Source/Programs/UnrealBuildAccelerator/Cli/Private/UbaCli.cpp b/Engine/Source/Programs/UnrealBuildAccelerator/Cli/Private/UbaCli.cpp
--- a/Engine/Source/Programs/UnrealBuildAccelerator/Cli/Private/UbaCli.cpp
+++ b/Engine/Source/Programs/UnrealBuildAccelerator/Cli/Private/UbaCli.cpp
@@ -520,7 +520,7 @@ namespace uba
 				u8* data = fa.GetData();
 				bool is64Bit = true;
 
-				CasKey key = CalculateCasKey(data, fileSize, false, nullptr);
+				CasKey key = CalculateCasKey(data, fileSize, false, nullptr, application.c_str());
 
 				if (data[0] != 'M' || data[1] != 'Z')
 					is64Bit = false;
