[UBA]
* Changed individual cancel message from Info to Detail

diff --git a/Engine/Source/Programs/UnrealBuildAccelerator/Common/Private/UbaSessionServer.cpp b/Engine/Source/Programs/UnrealBuildAccelerator/Common/Private/UbaSessionServer.cpp
--- a/Engine/Source/Programs/UnrealBuildAccelerator/Common/Private/UbaSessionServer.cpp
+++ b/Engine/Source/Programs/UnrealBuildAccelerator/Common/Private/UbaSessionServer.cpp
@@ -1192,7 +1192,7 @@ namespace uba
 				if (session.enabled)
 					++m_availableRemoteSlotCount;
 
-				m_logger.Info(TC("Client %s returned process to queue (%s)"), session.name.c_str(), reason.data);
+				m_logger.Detail(TC("Client %s returned process to queue (%s)"), session.name.c_str(), reason.data);
 				++m_returnedRemoteProcessCount;
 
 				process->m_executingHost.clear();
