[UBA]
* Fixed so Cache server saves cas in parallel with buckets

diff --git a/Engine/Source/Programs/UnrealBuildAccelerator/Common/Private/UbaCacheServer.cpp b/Engine/Source/Programs/UnrealBuildAccelerator/Common/Private/UbaCacheServer.cpp
--- a/Engine/Source/Programs/UnrealBuildAccelerator/Common/Private/UbaCacheServer.cpp
+++ b/Engine/Source/Programs/UnrealBuildAccelerator/Common/Private/UbaCacheServer.cpp
@@ -883,8 +883,10 @@ namespace uba
 		{
 			u64 saveStart = GetTime();
 			m_logger.Detail(TC("  Saving to disk"));
-			m_storage.SaveCasTable(false, false);
+			Event saveCasEvent(true);
+			m_server.AddWork([&]() { m_storage.SaveCasTable(false, false); saveCasEvent.Set(); }, 1, TC("SaveCas"));
 			SaveNoLock();
+			saveCasEvent.IsSet();
 			m_logger.Detail(TC("  Save Done (%s)"), TimeToText(GetTime() - saveStart).str);
 		}
 
