[UBA]
* Fixed bug in Thread code where two concurrent waits could race and cause delete of functor at the same time

diff --git a/Engine/Source/Programs/UnrealBuildAccelerator/Common/Private/UbaThread.cpp b/Engine/Source/Programs/UnrealBuildAccelerator/Common/Private/UbaThread.cpp
--- a/Engine/Source/Programs/UnrealBuildAccelerator/Common/Private/UbaThread.cpp
+++ b/Engine/Source/Programs/UnrealBuildAccelerator/Common/Private/UbaThread.cpp
@@ -94,6 +94,7 @@ namespace uba
 
 	bool Thread::Wait(u32 milliseconds, Event* wakeupEvent)
 	{
+		SCOPED_READ_LOCK(m_funcLock, readLock);
 		if (!m_handle)
 			return true;
 
@@ -110,7 +111,6 @@ namespace uba
 			if (WaitForSingleObject(m_handle, milliseconds) == WAIT_TIMEOUT)
 				return false;
 		}
-		CloseHandle(m_handle);
 #else
 		if (!m_finished.IsSet(milliseconds))
 			return false;
@@ -118,6 +118,17 @@ namespace uba
 		int res = pthread_join(*(pthread_t*)&m_handle, (void**)&ptr);
 		UBA_ASSERT(res == 0);
 #endif
+
+		readLock.Leave();
+
+		SCOPED_WRITE_LOCK(m_funcLock, lock);
+		if (!m_handle)
+			return true;
+
+		#if PLATFORM_WINDOWS
+		CloseHandle(m_handle);
+		#endif
+
 		m_func = {};
 		m_handle = nullptr;
 		return true;
diff --git a/Engine/Source/Programs/UnrealBuildAccelerator/Common/Public/UbaThread.h b/Engine/Source/Programs/UnrealBuildAccelerator/Common/Public/UbaThread.h
--- a/Engine/Source/Programs/UnrealBuildAccelerator/Common/Public/UbaThread.h
+++ b/Engine/Source/Programs/UnrealBuildAccelerator/Common/Public/UbaThread.h
@@ -2,8 +2,8 @@
 
 #pragma once
 
+#include "UbaBinaryReaderWriter.h"
 #include "UbaEvent.h"
-#include "UbaMemory.h"
 
 namespace uba
 {
@@ -24,6 +24,7 @@ namespace uba
 	private:
 		Function<u32()>	m_func;
 		void* m_handle = nullptr;
+		ReaderWriterLock m_funcLock;
 
 		#if !PLATFORM_WINDOWS
 		Event m_finished;
