[UBA]
* Fixed bug in code that prefetches dir content

diff --git a/Engine/Source/Programs/UnrealBuildAccelerator/Common/Private/UbaCacheClient.cpp b/Engine/Source/Programs/UnrealBuildAccelerator/Common/Private/UbaCacheClient.cpp
--- a/Engine/Source/Programs/UnrealBuildAccelerator/Common/Private/UbaCacheClient.cpp
+++ b/Engine/Source/Programs/UnrealBuildAccelerator/Common/Private/UbaCacheClient.cpp
@@ -1175,13 +1175,11 @@ namespace uba
 		if (m_storage.IsFileVerified(fileNameKey))
 			return;
 
-		u32 dirLength = path.count;
-
 		// Traverse all files in directory and report the file information... but only if it has not been reported before.. we don't want to interfere with other reports
 		TraverseDir(m_logger, path.data, 
 			[&](const DirectoryEntry& e)
 			{
-				path.Resize(dirLength).Append('\\').Append(e.name, e.nameLen);
+				path.Clear().Append('\\').Append(e.name, e.nameLen);
 				if (CaseInsensitiveFs)
 					path.MakeLower();
 
