[UBA]
* Added extra validation when writing cache entries
* Disabled directory preparsing when fetching cache entries for now

diff --git a/Engine/Source/Programs/UnrealBuildAccelerator/Common/Private/UbaCacheClient.cpp b/Engine/Source/Programs/UnrealBuildAccelerator/Common/Private/UbaCacheClient.cpp
--- a/Engine/Source/Programs/UnrealBuildAccelerator/Common/Private/UbaCacheClient.cpp
+++ b/Engine/Source/Programs/UnrealBuildAccelerator/Common/Private/UbaCacheClient.cpp
@@ -4,6 +4,7 @@
 #include "UbaApplicationRules.h"
 #include "UbaCacheEntry.h"
 #include "UbaCompactTables.h"
+#include "UbaCompressedObjFileHeader.h"
 #include "UbaDirectoryIterator.h"
 #include "UbaFileAccessor.h"
 #include "UbaNetworkMessage.h"
@@ -62,6 +63,8 @@ namespace uba
 		#endif
 
 		m_useDirectoryPreParsing = info.useDirectoryPreparsing;
+		m_validateCacheWritesInput = info.validateCacheWritesInput;
+		m_validateCacheWritesOutput = info.validateCacheWritesOutput;
 
 		m_client.RegisterOnConnected([this]()
 			{
@@ -248,12 +251,55 @@ namespace uba
 						return false;
 				}
 
+				bool shouldValidate = (m_validateCacheWritesInput && !isOutput) || (m_validateCacheWritesOutput && isOutput);
+
 				if (casKey == CasKeyZero) // If file is not found it was a temporary file that was deleted and is not really an output
 				{
+					
+					if (shouldValidate && FileExists(m_logger, path.data))
+					{
+						m_logger.Warning(TC("CasDb claims file %s does not exist but it does! Will not populate cache for %s"), path.data, info.description); 
+						return false;
+					}
+
 					//m_logger.Warning(TC("Can't find file %s"), path.data); 
 					stringToCasKey.erase(insres.first);
 					continue; // m_logger.Info(TC("This should never happen! (%s)"), path.data);
 				}
+
+
+				if (shouldValidate)
+				{
+					FileAccessor fa(m_logger, path.data);
+					if (!fa.OpenMemoryRead())
+					{
+						m_logger.Warning(TC("CasDb claims file %s does exist but can't open it. Will not populate cache for %s"), path.data, info.description); 
+						return false;
+					}
+
+					CasKey oldKey = AsCompressed(casKey, false);
+					CasKey newKey;
+
+					u64 fileSize = fa.GetSize();
+					u8* fileMem = fa.GetData();
+
+					if (fileSize > sizeof(CompressedObjFileHeader) && ((CompressedObjFileHeader*)fileMem)->IsValid())
+						newKey = AsCompressed(((CompressedObjFileHeader*)fileMem)->casKey, false);
+					else
+						newKey = CalculateCasKey(fileMem, fileSize, false, nullptr, path.data);
+
+					if (newKey != oldKey)
+					{
+						FileInformation fileInfo;
+						fa.GetFileInformationByHandle(fileInfo);
+
+						auto& fileEntry = m_storage.GetOrCreateFileEntry(CaseInsensitiveFs ? ToStringKeyLower(path) : ToStringKey(path));
+						SCOPED_READ_LOCK(fileEntry.lock, lock);
+
+						m_logger.Warning(TC("CasDb claims file %s has caskey %s but recalculting it gives us %s (FileEntry: %llu/%llu/%s, Real: %llu/%llu). Will not populate cache for %s"), path.data, CasKeyString(oldKey).str, CasKeyString(newKey).str, fileEntry.size, fileEntry.lastWritten, fileEntry.verified ? TC("true") : TC("false"), fileSize, fileInfo.lastWriteTime, info.description); 
+						return false;
+					}
+				}
 			}
 
 			UBA_ASSERT(IsCompressed(casKey));
diff --git a/Engine/Source/Programs/UnrealBuildAccelerator/Common/Private/UbaFileAccessor.cpp b/Engine/Source/Programs/UnrealBuildAccelerator/Common/Private/UbaFileAccessor.cpp
--- a/Engine/Source/Programs/UnrealBuildAccelerator/Common/Private/UbaFileAccessor.cpp
+++ b/Engine/Source/Programs/UnrealBuildAccelerator/Common/Private/UbaFileAccessor.cpp
@@ -336,7 +336,7 @@ namespace uba
 			return errorOnFail ? m_logger.Error(TC("Failed to open file %s for read"), m_fileName) : false;
 
 		FileInformation info;
-		if (!GetFileInformationByHandle(info, m_logger, m_fileName, m_fileHandle))
+		if (!GetFileInformationByHandle(info))
 			return m_logger.Error(TC("GetFileInformationByHandle failed on %s"), m_fileName);
 
 		m_size = info.size;
@@ -360,6 +360,10 @@ namespace uba
 		return true;
 	}
 
+	bool FileAccessor::GetFileInformationByHandle(FileInformation& out)
+	{
+		return uba::GetFileInformationByHandle(out, m_logger, m_fileName, m_fileHandle);
+	}
 
 	bool FileAccessor::InternalClose(bool success, u64* lastWriteTime)
 	{
diff --git a/Engine/Source/Programs/UnrealBuildAccelerator/Common/Private/UbaStorage.cpp b/Engine/Source/Programs/UnrealBuildAccelerator/Common/Private/UbaStorage.cpp
--- a/Engine/Source/Programs/UnrealBuildAccelerator/Common/Private/UbaStorage.cpp
+++ b/Engine/Source/Programs/UnrealBuildAccelerator/Common/Private/UbaStorage.cpp
@@ -1242,6 +1242,12 @@ namespace uba
 		return logger.Error(TC("Needs exclusive access to storage %s. Another process is running"), rootDir.data);
 	}
 
+	StorageImpl::FileEntry& StorageImpl::GetOrCreateFileEntry(StringKey fileNameKey)
+	{
+		SCOPED_WRITE_LOCK(m_fileTableLookupLock, lock);
+		return m_fileTableLookup.try_emplace(fileNameKey).first->second;
+	}
+
 	bool StorageImpl::LoadCasTable(bool logStats, bool alwaysCheckAllFiles)
 	{
 		static bool isExclusive = CheckExclusive(m_logger, m_rootDir);
@@ -1870,11 +1876,7 @@ namespace uba
 			forKey.MakeLower();
 		StringKey fileNameKey = ToStringKey(forKey);
 
-		SCOPED_WRITE_LOCK(m_fileTableLookupLock, lookupLock);
-		auto insres = m_fileTableLookup.try_emplace(fileNameKey);
-		FileEntry& fileEntry = insres.first->second;
-		lookupLock.Leave();
-
+		FileEntry& fileEntry = GetOrCreateFileEntry(fileNameKey);
 		SCOPED_WRITE_LOCK(fileEntry.lock, entryLock);
 	
 		if (fileEntry.verified)
@@ -2003,11 +2005,7 @@ namespace uba
 
 	bool StorageImpl::StoreCasKey(CasKey& out, const StringKey& fileNameKey, const tchar* fileName, const CasKey& casKeyOverride, bool fileIsCompressed)
 	{
-		SCOPED_WRITE_LOCK(m_fileTableLookupLock, lookupLock);
-		auto insres = m_fileTableLookup.try_emplace(fileNameKey);
-		FileEntry& fileEntry = insres.first->second;
-		lookupLock.Leave();
-
+		FileEntry& fileEntry = GetOrCreateFileEntry(fileNameKey);
 		SCOPED_WRITE_LOCK(fileEntry.lock, entryLock);
 	
 		if (fileEntry.verified)
@@ -2345,10 +2343,7 @@ namespace uba
 		if (CaseInsensitiveFs)
 			forKey.MakeLower();
 		StringKey key = ToStringKey(forKey);
-		SCOPED_WRITE_LOCK(m_fileTableLookupLock, lock);
-		auto insres = m_fileTableLookup.try_emplace(key);
-		FileEntry& entry = insres.first->second;
-		lock.Leave();
+		FileEntry& entry = GetOrCreateFileEntry(key);
 
 		TimerScope ts(stats.copyOrLink);
 
@@ -2553,11 +2548,7 @@ namespace uba
 		if (CaseInsensitiveFs)
 			forKey.MakeLower();
 		StringKey key = ToStringKey(forKey);
-		SCOPED_WRITE_LOCK(m_fileTableLookupLock, lock);
-		auto insres = m_fileTableLookup.try_emplace(key);
-		lock.Leave();
-
-		FileEntry& entry = insres.first->second;
+		FileEntry& entry = GetOrCreateFileEntry(key);
 		SCOPED_WRITE_LOCK(entry.lock, lock2);
 		entry.casKey = casKey;
 		entry.lastWritten = lastWritten;
diff --git a/Engine/Source/Programs/UnrealBuildAccelerator/Common/Private/UbaStorageServer.cpp b/Engine/Source/Programs/UnrealBuildAccelerator/Common/Private/UbaStorageServer.cpp
--- a/Engine/Source/Programs/UnrealBuildAccelerator/Common/Private/UbaStorageServer.cpp
+++ b/Engine/Source/Programs/UnrealBuildAccelerator/Common/Private/UbaStorageServer.cpp
@@ -199,10 +199,7 @@ namespace uba
 		if (out == CasKeyZero)
 			return false;
 
-		SCOPED_WRITE_LOCK(m_fileTableLookupLock, lookupLock);
-		auto insres = m_fileTableLookup.try_emplace(fileNameKey);
-		FileEntry& fileEntry = insres.first->second;
-		lookupLock.Leave();
+		FileEntry& fileEntry = GetOrCreateFileEntry(fileNameKey);
 		SCOPED_WRITE_LOCK(fileEntry.lock, entryLock);
 		fileEntry.verified = true;
 		fileEntry.casKey = out;
diff --git a/Engine/Source/Programs/UnrealBuildAccelerator/Common/Public/UbaCacheClient.h b/Engine/Source/Programs/UnrealBuildAccelerator/Common/Public/UbaCacheClient.h
--- a/Engine/Source/Programs/UnrealBuildAccelerator/Common/Public/UbaCacheClient.h
+++ b/Engine/Source/Programs/UnrealBuildAccelerator/Common/Public/UbaCacheClient.h
@@ -28,7 +28,9 @@ namespace uba
 		NetworkClient& client;
 		Session& session;
 		bool reportMissReason = false;
-		bool useDirectoryPreparsing = true; // This is used to minimize syscalls. GetFileAttributes can be very expensive on cloud machines and we can enable this to minimize syscall count
+		bool useDirectoryPreparsing = false; // This is used to minimize syscalls. GetFileAttributes can be very expensive on cloud machines and we can enable this to minimize syscall count
+		bool validateCacheWritesInput = false; // Set to true to validate cas of all input files before sent to cache
+		bool validateCacheWritesOutput = true; // Set to true to validate cas of all output files before sent to cache
 	};
 
 	class CacheClient
@@ -68,6 +70,8 @@ namespace uba
 		Session& m_session;
 		bool m_reportMissReason;
 		bool m_useDirectoryPreParsing;
+		bool m_validateCacheWritesInput;
+		bool m_validateCacheWritesOutput;
 
 		Atomic<bool> m_connected;
 
diff --git a/Engine/Source/Programs/UnrealBuildAccelerator/Common/Public/UbaFileAccessor.h b/Engine/Source/Programs/UnrealBuildAccelerator/Common/Public/UbaFileAccessor.h
--- a/Engine/Source/Programs/UnrealBuildAccelerator/Common/Public/UbaFileAccessor.h
+++ b/Engine/Source/Programs/UnrealBuildAccelerator/Common/Public/UbaFileAccessor.h
@@ -29,6 +29,9 @@ namespace uba
 		inline u8* GetData() { return m_data; }
 		inline u64 GetSize() { return m_size; }
 
+		// Can be called if file is opened
+		bool GetFileInformationByHandle(FileInformation& out);
+
 	private:
 		bool InternalClose(bool success, u64* lastWriteTime);
 		Logger& m_logger;
diff --git a/Engine/Source/Programs/UnrealBuildAccelerator/Common/Public/UbaStorage.h b/Engine/Source/Programs/UnrealBuildAccelerator/Common/Public/UbaStorage.h
--- a/Engine/Source/Programs/UnrealBuildAccelerator/Common/Public/UbaStorage.h
+++ b/Engine/Source/Programs/UnrealBuildAccelerator/Common/Public/UbaStorage.h
@@ -175,6 +175,9 @@ namespace uba
 		bool OpenCasDataFile(u32 index, u64 size);
 		bool CreateCasDataFiles();
 
+		struct FileEntry;
+		FileEntry& GetOrCreateFileEntry(StringKey fileNameKey);
+
 
 		WorkManager* m_workManager;
 		MutableLogger m_logger;
