[UBA]
* Attempt at fixing problem with cache client seeing the wrong file when writing to cache. Problem is related to a race condition between directory crawler threads and writing files. We can end up in this scenario
1. Crawler query directory, file information gets stored in temporary state. and crawler starts traversing state
2. File is written and new date is set to file on disk, and file db entry is set to verified false.
3. Crawler hits the same file and now checks temporary state with file db entry and notice the lastwritten is the same, and set it back to verified. (but it is not the same)

diff --git a/Engine/Source/Programs/UnrealBuildAccelerator/Common/Private/UbaSession.cpp b/Engine/Source/Programs/UnrealBuildAccelerator/Common/Private/UbaSession.cpp
--- a/Engine/Source/Programs/UnrealBuildAccelerator/Common/Private/UbaSession.cpp
+++ b/Engine/Source/Programs/UnrealBuildAccelerator/Common/Private/UbaSession.cpp
@@ -770,6 +770,11 @@ namespace uba
 			}
 		}
 
+		// There are directory crawlers happening in parallel so we need to really make sure to invalidate this one since a crawler can actually
+		// hit this file with information from a query before it was written.. and then it will turn it back to "verified" using old info
+		if (registerRealFile)
+			m_storage.InvalidateCachedFileInfo(fileNameKey);
+
 		FileEntryAdded(fileNameKey, lastWriteTime, fileSize);
 
 		u8 temp[1024];
@@ -2144,8 +2149,9 @@ namespace uba
 						return false;
 				}
 
-				Storage::CachedFileInfo cfi;
-				m_storage.VerifyAndGetCachedFileInfo(cfi, file.key, file.lastWriteTime, fileSize);
+				// There are directory crawlers happening in parallel so we need to really make sure to invalidate this one since a crawler can actually
+				// hit this file with information from a query before it was written.. and then it will turn it back to "verified" using old info
+				m_storage.InvalidateCachedFileInfo(file.key);
 			}
 			else
 			{
diff --git a/Engine/Source/Programs/UnrealBuildAccelerator/Common/Private/UbaStorage.cpp b/Engine/Source/Programs/UnrealBuildAccelerator/Common/Private/UbaStorage.cpp
--- a/Engine/Source/Programs/UnrealBuildAccelerator/Common/Private/UbaStorage.cpp
+++ b/Engine/Source/Programs/UnrealBuildAccelerator/Common/Private/UbaStorage.cpp
@@ -1341,19 +1341,19 @@ namespace uba
 		{
 			StringKey fileNameKey = reader.ReadStringKey();
 			auto insres = m_fileTableLookup.try_emplace(fileNameKey);
-			FileEntry& entry = insres.first->second;
-			entry.verified = false;
+			FileEntry& fileEntry = insres.first->second;
+			fileEntry.verified = false;
 			if (reader.GetPosition() + 24 > fileSize)
 			{
 				m_fileTableLookup.clear();
 				m_logger.Warning(TC("CasTable file %s is corrupt"), fileName.data);
 				return false;
 			}
-			entry.size = reader.ReadU64();
-			entry.lastWritten = reader.ReadU64();
+			fileEntry.size = reader.ReadU64();
+			fileEntry.lastWritten = reader.ReadU64();
 			CasKey key = reader.ReadCasKey();
 			if (key != CasKeyZero)
-				entry.casKey = AsCompressed(key, m_storeCompressed);
+				fileEntry.casKey = AsCompressed(key, m_storeCompressed);
 		}
 
 
@@ -1861,13 +1861,33 @@ namespace uba
 		lookupLock.Leave();
 
 		SCOPED_WRITE_LOCK(fileEntry.lock, entryLock);
-		fileEntry.verified = fileEntry.lastWritten == verifiedLastWriteTime && fileEntry.size == verifiedSize;
+		
+		// If fileEntry.lastWritten is newer than our verifiedLastWriteTime it means that something else updated the file in between us getting info from file system
+		// and ending up here. If that is the case we can't touch the verified stamp.
+		if (verifiedLastWriteTime >= fileEntry.lastWritten)
+			fileEntry.verified = fileEntry.lastWritten == verifiedLastWriteTime && fileEntry.size == verifiedSize && fileEntry.casKey != CasKeyInvalid;
+
 		if (!fileEntry.verified)
 			return false;
 		out.casKey = fileEntry.casKey;
 		return true;
 	}
 
+	bool StorageImpl::InvalidateCachedFileInfo(StringKey fileNameKey)
+	{
+		SCOPED_READ_LOCK(m_fileTableLookupLock, lookupLock);
+		auto findIt = m_fileTableLookup.find(fileNameKey);
+		if (findIt == m_fileTableLookup.end())
+			return false;
+		FileEntry& fileEntry = findIt->second;
+		lookupLock.Leave();
+
+		SCOPED_WRITE_LOCK(fileEntry.lock, entryLock);
+		fileEntry.verified = false;
+		fileEntry.casKey = CasKeyInvalid;
+		return true;
+	}
+
 	bool StorageImpl::StoreCasFile(CasKey& out, const tchar* fileName, const CasKey& casKeyOverride, bool deferCreation, bool fileIsCompressed)
 	{
 		StringBuffer<> forKey;
@@ -1881,6 +1901,8 @@ namespace uba
 	
 		if (fileEntry.verified)
 		{
+			UBA_ASSERT(fileEntry.casKey != CasKeyInvalid);
+
 			if (fileEntry.casKey != CasKeyZero)
 			{
 				UBA_ASSERT(casKeyOverride == CasKeyZero || casKeyOverride == fileEntry.casKey);
@@ -1923,7 +1945,7 @@ namespace uba
 				out = fileEntry.casKey;
 				return true;
 			}
-			if (fileSize == fileEntry.size && lastWritten == fileEntry.lastWritten)
+			if (fileSize == fileEntry.size && lastWritten == fileEntry.lastWritten && fileEntry.casKey != CasKeyInvalid)
 			{
 				if (!AddCasFile(fileNameKey, fileName, fileEntry.casKey, deferCreation, fileIsCompressed))
 					return false;
@@ -2044,7 +2066,7 @@ namespace uba
 				out = fileEntry.casKey;
 				return true;
 			}
-			if (fileSize == fileEntry.size && lastWritten == fileEntry.lastWritten)
+			if (fileSize == fileEntry.size && lastWritten == fileEntry.lastWritten && fileEntry.casKey != CasKeyInvalid)
 			{
 				out = fileEntry.casKey;
 				return true;
@@ -2343,7 +2365,7 @@ namespace uba
 		if (CaseInsensitiveFs)
 			forKey.MakeLower();
 		StringKey key = ToStringKey(forKey);
-		FileEntry& entry = GetOrCreateFileEntry(key);
+		FileEntry& fileEntry = GetOrCreateFileEntry(key);
 
 		TimerScope ts(stats.copyOrLink);
 
@@ -2414,8 +2436,8 @@ namespace uba
 
 				FileAccessor destinationFile(m_logger, destination);
 
-				SCOPED_WRITE_LOCK(entry.lock, entryLock);
-				entry.verified = false;
+				SCOPED_WRITE_LOCK(fileEntry.lock, entryLock);
+				fileEntry.verified = false;
 
 				// This is to reduce number of active CreateFiles.. seems like machines don't like tons of CreateFile at the same time
 				#if PLATFORM_WINDOWS
@@ -2483,10 +2505,10 @@ namespace uba
 				UBA_ASSERT(lastWriteTime);
 				if (lastWriteTime)
 				{
-					entry.casKey = casKey;
-					entry.lastWritten = lastWriteTime;
-					entry.size = decompressedSize;
-					entry.verified = true;
+					fileEntry.casKey = casKey;
+					fileEntry.lastWritten = lastWriteTime;
+					fileEntry.size = decompressedSize;
+					fileEntry.verified = true;
 				}
 				return true;
 			}
@@ -2499,8 +2521,8 @@ namespace uba
 			UBA_ASSERT(false);
 			#endif
 
-			SCOPED_WRITE_LOCK(entry.lock, entryLock);
-			entry.verified = false;
+			SCOPED_WRITE_LOCK(fileEntry.lock, entryLock);
+			fileEntry.verified = false;
 
 			bool firstTry = true;
 			while (true)
diff --git a/Engine/Source/Programs/UnrealBuildAccelerator/Common/Public/UbaStorage.h b/Engine/Source/Programs/UnrealBuildAccelerator/Common/Public/UbaStorage.h
--- a/Engine/Source/Programs/UnrealBuildAccelerator/Common/Public/UbaStorage.h
+++ b/Engine/Source/Programs/UnrealBuildAccelerator/Common/Public/UbaStorage.h
@@ -43,6 +43,7 @@ namespace uba
 
 		struct CachedFileInfo { CasKey casKey; };
 		virtual bool VerifyAndGetCachedFileInfo(CachedFileInfo& out, StringKey fileNameKey, u64 verifiedLastWriteTime, u64 verifiedSize) = 0;
+		virtual bool InvalidateCachedFileInfo(StringKey fileNameKey) = 0;
 
 		virtual bool StoreCasFile(CasKey& out, const tchar* fileName, const CasKey& casKeyOverride, bool deferCreation, bool fileIsCompressed) = 0;
 		virtual bool StoreCasFile(CasKey& out, StringKey fileNameKey, const tchar* fileName, FileMappingHandle mappingHandle, u64 mappingOffset, u64 fileSize, const tchar* hint, bool deferCreation = false, bool keepMappingInMemory = false) = 0;
@@ -122,6 +123,7 @@ namespace uba
 		virtual bool DeleteCasForFile(const tchar* file) override;
 		virtual bool RetrieveCasFile(RetrieveResult& out, const CasKey& casKey, const tchar* hint, FileMappingBuffer* mappingBuffer = nullptr, u64 memoryMapAlignment = 1, bool allowProxy = true) override;
 		virtual bool VerifyAndGetCachedFileInfo(CachedFileInfo& out, StringKey fileNameKey, u64 verifiedLastWriteTime, u64 verifiedSize) override;
+		virtual bool InvalidateCachedFileInfo(StringKey fileNameKey) override;
 		virtual bool StoreCasFile(CasKey& out, const tchar* fileName, const CasKey& casKeyOverride, bool deferCreation, bool fileIsCompressed) override;
 		virtual bool StoreCasFile(CasKey& out, StringKey fileNameKey, const tchar* fileName, FileMappingHandle mappingHandle, u64 mappingOffset, u64 fileSize, const tchar* hint, bool deferCreation = false, bool keepMappingInMemory = false) override;
 		virtual bool DropCasFile(const CasKey& casKey, bool forceDelete, const tchar* hint) override;
diff --git a/Engine/Source/Programs/UnrealBuildAccelerator/Core/Public/UbaHash.h b/Engine/Source/Programs/UnrealBuildAccelerator/Core/Public/UbaHash.h
--- a/Engine/Source/Programs/UnrealBuildAccelerator/Core/Public/UbaHash.h
+++ b/Engine/Source/Programs/UnrealBuildAccelerator/Core/Public/UbaHash.h
@@ -35,6 +35,7 @@ namespace uba
 #pragma pack(pop)
 
 	constexpr CasKey CasKeyZero;
+	constexpr CasKey CasKeyInvalid(~0ull, ~0ull, ~0u);
 
 	// Use 36+1 characters
 	inline void GuidToStr(tchar* out, u32 capacity, const Guid& g)
