[UBA]
* Reserve memory for elements that we know will be allocated in dir table and filemapping table

diff --git a/Engine/Source/Programs/UnrealBuildAccelerator/Core/Private/UbaMemory.cpp b/Engine/Source/Programs/UnrealBuildAccelerator/Core/Private/UbaMemory.cpp
--- a/Engine/Source/Programs/UnrealBuildAccelerator/Core/Private/UbaMemory.cpp
+++ b/Engine/Source/Programs/UnrealBuildAccelerator/Core/Private/UbaMemory.cpp
@@ -86,6 +86,21 @@ namespace uba
 		return ret;
 	}
 
+	void MemoryBlock::ReserveNoLock(u64 bytes, const tchar* hint)
+	{
+		#if PLATFORM_WINDOWS
+		u64 newPos = Min(writtenSize + bytes, reserveSize);
+		if (newPos <= mappedSize)
+			return;
+		u64 toCommit = AlignUp(newPos - mappedSize, 1024 * 1024);
+		if (mappedSize + toCommit > reserveSize)
+			toCommit = reserveSize - mappedSize;
+		if (!VirtualAlloc(memory + mappedSize, toCommit, MEM_COMMIT, PAGE_READWRITE))
+			FatalError(9883, TC("Failed to commit virtual memory for memory block. Total size %llu (%u) (%s)"), mappedSize + toCommit, GetLastError(), hint);
+		mappedSize += toCommit;
+		#endif
+	}
+
 	void MemoryBlock::Free(void* p)
 	{
 		if (!memory)
diff --git a/Engine/Source/Programs/UnrealBuildAccelerator/Core/Public/UbaDirectoryTable.h b/Engine/Source/Programs/UnrealBuildAccelerator/Core/Public/UbaDirectoryTable.h
--- a/Engine/Source/Programs/UnrealBuildAccelerator/Core/Public/UbaDirectoryTable.h
+++ b/Engine/Source/Programs/UnrealBuildAccelerator/Core/Public/UbaDirectoryTable.h
@@ -19,6 +19,8 @@ namespace uba
 		{
 			m_memory = mem;
 			m_lookup.reserve(tableCount + 100);
+			if (m_memoryBlock)
+				m_memoryBlock->ReserveNoLock(tableCount*(sizeof(GrowingUnorderedMap<StringKey, Directory>::value_type)+16), TC(""));
 			ParseDirectoryTable(tableSize);
 		}
 
diff --git a/Engine/Source/Programs/UnrealBuildAccelerator/Core/Public/UbaMemory.h b/Engine/Source/Programs/UnrealBuildAccelerator/Core/Public/UbaMemory.h
--- a/Engine/Source/Programs/UnrealBuildAccelerator/Core/Public/UbaMemory.h
+++ b/Engine/Source/Programs/UnrealBuildAccelerator/Core/Public/UbaMemory.h
@@ -99,6 +99,7 @@ namespace uba
 		void Deinit();
 		void* Allocate(u64 bytes, u64 alignment, const tchar* hint);
 		void* AllocateNoLock(u64 bytes, u64 alignment, const tchar* hint);
+		void ReserveNoLock(u64 bytes, const tchar* hint);
 		void Free(void* p);
 		tchar* Strdup(const tchar* str);
 		void Swap(MemoryBlock& other);
diff --git a/Engine/Source/Programs/UnrealBuildAccelerator/Detours/Private/UbaDetoursFileMappingTable.cpp b/Engine/Source/Programs/UnrealBuildAccelerator/Detours/Private/UbaDetoursFileMappingTable.cpp
--- a/Engine/Source/Programs/UnrealBuildAccelerator/Detours/Private/UbaDetoursFileMappingTable.cpp
+++ b/Engine/Source/Programs/UnrealBuildAccelerator/Detours/Private/UbaDetoursFileMappingTable.cpp
@@ -14,6 +14,7 @@ namespace uba
 	{
 		m_mem = mem;
 		m_lookup.reserve(tableCount + 100);
+		m_memoryBlock.ReserveNoLock(tableCount*(sizeof(GrowingUnorderedMap<StringKey, FileInfo>::value_type)+16), TC(""));
 		ParseNoLock(tableSize);
 	}
 
