[UBA]
* Added code to make cache server return result without waiting for unmapview of file being done

diff --git a/Engine/Source/Programs/UnrealBuildAccelerator/Common/Private/UbaCacheServer.cpp b/Engine/Source/Programs/UnrealBuildAccelerator/Common/Private/UbaCacheServer.cpp
--- a/Engine/Source/Programs/UnrealBuildAccelerator/Common/Private/UbaCacheServer.cpp
+++ b/Engine/Source/Programs/UnrealBuildAccelerator/Common/Private/UbaCacheServer.cpp
@@ -471,6 +471,8 @@ namespace uba
 		UnorderedMap<CasKey, CasFileInfo> existingCas;
 		ReaderWriterLock existingCasLock;
 
+		m_storage.WaitForActiveWork();
+
 		u32 removedNonExisting = 0;
 		{
 			// TODO: Make this cleaner... 
diff --git a/Engine/Source/Programs/UnrealBuildAccelerator/Common/Private/UbaStorageServer.cpp b/Engine/Source/Programs/UnrealBuildAccelerator/Common/Private/UbaStorageServer.cpp
--- a/Engine/Source/Programs/UnrealBuildAccelerator/Common/Private/UbaStorageServer.cpp
+++ b/Engine/Source/Programs/UnrealBuildAccelerator/Common/Private/UbaStorageServer.cpp
@@ -44,6 +44,7 @@ namespace uba
 
 	StorageServer::~StorageServer()
 	{
+		WaitForActiveWork();
 		UBA_ASSERT(m_waitEntries.empty());
 		UBA_ASSERT(m_proxies.empty());
 		m_server.UnregisterOnClientDisconnected(ServiceId);
@@ -56,6 +57,12 @@ namespace uba
 		return true;
 	}
 
+	void StorageServer::WaitForActiveWork()
+	{
+		while (m_activeUnmap)
+			Sleep(5);
+	}
+
 	bool StorageServer::GetZone(StringBufferBase& out)
 	{
 		if (m_zone.empty())
@@ -112,9 +119,14 @@ namespace uba
 		{
 			if (ownsMapping)
 			{
-				UnmapViewOfFile(memoryBegin, mappedView.size, TC(""));
-				CloseFileMapping(mappedView.handle);
-				CloseFile(nullptr, readFileHandle);
+				++server.m_activeUnmap;
+				server.GetServer().AddWork([&server, mb = memoryBegin, mp = mappedView, rfh = readFileHandle]()
+					{
+						UnmapViewOfFile(mb, mp.size, TC(""));
+						CloseFileMapping(mp.handle);
+						CloseFile(nullptr, rfh);
+						--server.m_activeUnmap;
+					}, 1, TC("ActiveFetchRelease"));
 			}
 			else
 				server.m_casDataBuffer.UnmapView(mappedView, TC("OnDisconnected"));
diff --git a/Engine/Source/Programs/UnrealBuildAccelerator/Common/Public/UbaStorageServer.h b/Engine/Source/Programs/UnrealBuildAccelerator/Common/Public/UbaStorageServer.h
--- a/Engine/Source/Programs/UnrealBuildAccelerator/Common/Public/UbaStorageServer.h
+++ b/Engine/Source/Programs/UnrealBuildAccelerator/Common/Public/UbaStorageServer.h
@@ -29,6 +29,8 @@ namespace uba
 
 		using StorageImpl::StoreCasFile;
 
+		void WaitForActiveWork();
+
 	protected:
 
 		virtual bool GetZone(StringBufferBase& out) override;
@@ -100,6 +102,7 @@ namespace uba
 		};
 		ReaderWriterLock m_activeFetchesLock;
 		UnorderedMap<u16, ActiveFetch> m_activeFetches;
+		Atomic<u32> m_activeUnmap;
 
 
 		ReaderWriterLock m_availableIdsLock;
