[UBA]
* Added more debug logging to writetocache mismatch
* Renamed all "entry" to either casEntry or fileEntry in order to be able to more easily do find-in-files
* Fixed so store entry done message is sent on fail as well (so server can clean up)

diff --git a/Engine/Source/Programs/UnrealBuildAccelerator/Common/Private/UbaCacheClient.cpp b/Engine/Source/Programs/UnrealBuildAccelerator/Common/Private/UbaCacheClient.cpp
--- a/Engine/Source/Programs/UnrealBuildAccelerator/Common/Private/UbaCacheClient.cpp
+++ b/Engine/Source/Programs/UnrealBuildAccelerator/Common/Private/UbaCacheClient.cpp
@@ -238,8 +238,28 @@ namespace uba
 			// Get file caskey using storage
 			if (casKey == CasKeyZero)
 			{
+				bool shouldValidate = (m_validateCacheWritesInput && !isOutput) || (m_validateCacheWritesOutput && isOutput);
 				bool deferCreation = true;
 				bool fileIsCompressed = IsFileCompressed(info, path);
+
+
+
+				CasKey oldOldKey;
+				u64 oldOldSize = 0;
+				u64 oldOldLastWritten = 0;
+				bool oldOldVerified = false;
+
+				if (shouldValidate)
+				{
+					auto& fileEntry = m_storage.GetOrCreateFileEntry(CaseInsensitiveFs ? ToStringKeyLower(path) : ToStringKey(path));
+					SCOPED_READ_LOCK(fileEntry.lock, lock);
+					oldOldKey = fileEntry.casKey;
+					oldOldSize = fileEntry.size;
+					oldOldLastWritten = fileEntry.lastWritten;
+					oldOldVerified = fileEntry.verified;
+				}
+
+
 				if (isOutput)
 				{
 					if (!m_storage.StoreCasFile(casKey, path.data, CasKeyZero, deferCreation, fileIsCompressed))
@@ -251,8 +271,6 @@ namespace uba
 						return false;
 				}
 
-				bool shouldValidate = (m_validateCacheWritesInput && !isOutput) || (m_validateCacheWritesOutput && isOutput);
-
 				if (casKey == CasKeyZero) // If file is not found it was a temporary file that was deleted and is not really an output
 				{
 					
@@ -296,7 +314,9 @@ namespace uba
 						auto& fileEntry = m_storage.GetOrCreateFileEntry(CaseInsensitiveFs ? ToStringKeyLower(path) : ToStringKey(path));
 						SCOPED_READ_LOCK(fileEntry.lock, lock);
 
-						m_logger.Warning(TC("CasDb claims file %s has caskey %s but recalculting it gives us %s (FileEntry: %llu/%llu/%s, Real: %llu/%llu). Will not populate cache for %s"), path.data, CasKeyString(oldKey).str, CasKeyString(newKey).str, fileEntry.size, fileEntry.lastWritten, fileEntry.verified ? TC("true") : TC("false"), fileSize, fileInfo.lastWriteTime, info.description); 
+						auto ToString = [](bool b) { return b ? TC("true") : TC("false"); };
+						m_logger.Warning(TC("CasDb claims file %s has caskey %s but recalculating it gives us %s (OldEntry: %s/%llu/%llu/%s FileEntry: %llu/%llu/%s, Real: %llu/%llu). Will not populate cache for %s"),
+							path.data, CasKeyString(oldKey).str, CasKeyString(newKey).str, CasKeyString(oldOldKey).str, oldOldSize, oldOldLastWritten, ToString(oldOldVerified), fileEntry.size, fileEntry.lastWritten, ToString(fileEntry.verified), fileSize, fileInfo.lastWriteTime, info.description);
 						return false;
 					}
 				}
@@ -919,6 +939,18 @@ namespace uba
 		if (!reader.GetLeft())
 			return true;
 
+		bool success = false;
+		auto doneGuard = MakeGuard([&]()
+			{
+				// Send done.. confirm to server
+				StackBinaryWriter<SendMaxSize> writer;
+				NetworkMessage msg(m_client, CacheServiceId, CacheMessageType_StoreEntryDone, writer);
+				writer.Write7BitEncoded(MakeId(bucket.id));
+				writer.WriteCasKey(cmdKey);
+				writer.WriteBool(success);
+				return msg.Send(reader);
+			});
+
 		// There is content we need to upload to server
 		while (reader.GetLeft())
 		{
@@ -1020,15 +1052,8 @@ namespace uba
 
 		}
 
-		// Send done.. confirm to server
-		StackBinaryWriter<SendMaxSize> writer;
-		NetworkMessage msg(m_client, CacheServiceId, CacheMessageType_StoreEntryDone, writer);
-		writer.Write7BitEncoded(MakeId(bucket.id));
-		writer.WriteCasKey(cmdKey);
-		if (!msg.Send(reader))
-			return false;
-
-		return true;
+		success = true;
+		return doneGuard.Execute();
 	}
 
 	bool CacheClient::FetchCasTable(Bucket& bucket, CacheStats& stats, u32 requiredCasTableOffset)
diff --git a/Engine/Source/Programs/UnrealBuildAccelerator/Common/Private/UbaStorage.cpp b/Engine/Source/Programs/UnrealBuildAccelerator/Common/Private/UbaStorage.cpp
--- a/Engine/Source/Programs/UnrealBuildAccelerator/Common/Private/UbaStorage.cpp
+++ b/Engine/Source/Programs/UnrealBuildAccelerator/Common/Private/UbaStorage.cpp
@@ -52,52 +52,52 @@ namespace uba
 			free(slot);
 	}
 
-	void StorageImpl::CasEntryAccessed(CasEntry& entry)
+	void StorageImpl::CasEntryAccessed(CasEntry& casEntry)
 	{
 		bool hasMapping;
 		{
-			SCOPED_READ_LOCK(entry.lock, l); // Note, this lock is taken again outside CasEntryAccessed.. so if this takes a long time it won't help to remove this lock
-			hasMapping = entry.mappingHandle.IsValid();
+			SCOPED_READ_LOCK(casEntry.lock, l); // Note, this lock is taken again outside CasEntryAccessed.. so if this takes a long time it won't help to remove this lock
+			hasMapping = casEntry.mappingHandle.IsValid();
 		}
 		if (hasMapping)
 			return;
 
 		SCOPED_WRITE_LOCK(m_accessLock, lock);
 
-		CasEntry* prevAccessed = entry.prevAccessed;
+		CasEntry* prevAccessed = casEntry.prevAccessed;
 		if (prevAccessed == nullptr)
 		{
-			if (m_newestAccessed == &entry) // We are already first
+			if (m_newestAccessed == &casEntry) // We are already first
 				return;
 		}
 		else
-			prevAccessed->nextAccessed = entry.nextAccessed;
+			prevAccessed->nextAccessed = casEntry.nextAccessed;
 
-		if (entry.nextAccessed)
-			entry.nextAccessed->prevAccessed = prevAccessed;
+		if (casEntry.nextAccessed)
+			casEntry.nextAccessed->prevAccessed = prevAccessed;
 		else if (prevAccessed)
 			m_oldestAccessed = prevAccessed;
 		else if (!m_oldestAccessed)
-			m_oldestAccessed = &entry;
+			m_oldestAccessed = &casEntry;
 
 		if (m_newestAccessed)
-			m_newestAccessed->prevAccessed = &entry;
-		entry.nextAccessed = m_newestAccessed;
-		entry.prevAccessed = nullptr;
-		m_newestAccessed = &entry;
+			m_newestAccessed->prevAccessed = &casEntry;
+		casEntry.nextAccessed = m_newestAccessed;
+		casEntry.prevAccessed = nullptr;
+		m_newestAccessed = &casEntry;
 	}
 
-	void StorageImpl::CasEntryWritten(CasEntry& entry, u64 size)
+	void StorageImpl::CasEntryWritten(CasEntry& casEntry, u64 size)
 	{
 		SCOPED_WRITE_LOCK(m_accessLock, lock);
 
-		m_casTotalBytes += size - entry.size;
+		m_casTotalBytes += size - casEntry.size;
 		m_casMaxBytes = Max(m_casTotalBytes, m_casMaxBytes);
 
-		entry.size = size;
+		casEntry.size = size;
 
 #if !UBA_USE_SPARSEFILE
-		UBA_ASSERT(!entry.mappingHandle.IsValid());
+		UBA_ASSERT(!casEntry.mappingHandle.IsValid());
 #endif
 
 		if (!m_casCapacityBytes || m_overflowReported || m_casTotalBytes <= m_casCapacityBytes || m_manuallyHandleOverflow)
@@ -111,7 +111,7 @@ namespace uba
 		UBA_ASSERT(!m_newestAccessed || !m_newestAccessed->prevAccessed);
 		UBA_ASSERT(!m_oldestAccessed || !m_oldestAccessed->nextAccessed);
 
-		struct Rec { CasEntry& entry; u64 size; };
+		struct Rec { CasEntry& casEntry; u64 size; };
 		Vector<Rec> toDelete;
 
 		for (CasEntry* it = m_oldestAccessed; it;)
@@ -151,7 +151,7 @@ namespace uba
 		{
 #if !UBA_USE_SPARSEFILE
 			StringBuffer<> casFile;
-			StorageImpl::GetCasFileName(casFile, rec.entry.key);
+			StorageImpl::GetCasFileName(casFile, rec.casEntry.key);
 
 			if (!DeleteFileW(casFile.data))
 			{
@@ -159,9 +159,9 @@ namespace uba
 				if (error != ERROR_FILE_NOT_FOUND && error != ERROR_PATH_NOT_FOUND)
 				{
 					m_logger.Error(TC("Failed to delete %s while handling overflow (%s)"), casFile.data, LastErrorToText(error).data);
-					rec.entry.exists = true;
-					rec.entry.size = rec.size;
-					rec.entry.lock.LeaveWrite();
+					rec.casEntry.exists = true;
+					rec.casEntry.size = rec.size;
+					rec.casEntry.lock.LeaveWrite();
 
 					// TODO: Should this instead set some overflow
 					/*
@@ -169,7 +169,7 @@ namespace uba
 					m_casEvictedBytes -= rec.size;
 					--m_casEvictedCount;
 					m_casTotalBytes += rec.size;
-					AttachEntry(&rec.entry); // TODO: This should be re-added in the end
+					AttachEntry(&rec.casEntry); // TODO: This should be re-added in the end
 					*/
 					continue;
 				}
@@ -179,46 +179,46 @@ namespace uba
 			//UBA_ASSERT(false);
 #endif
 
-			entry.verified = true; // Verified to be deleted
+			casEntry.verified = true; // Verified to be deleted
 
-			rec.entry.lock.LeaveWrite();
+			rec.casEntry.lock.LeaveWrite();
 		}
 	}
 
-	void StorageImpl::CasEntryDeleted(CasEntry& entry, u64 size)
+	void StorageImpl::CasEntryDeleted(CasEntry& casEntry, u64 size)
 	{
 		SCOPED_WRITE_LOCK(m_accessLock, lock);
 		m_casTotalBytes -= size;
-		entry.size = 0;
-		DetachEntry(entry);
+		casEntry.size = 0;
+		DetachEntry(casEntry);
 	}
 
-	void StorageImpl::AttachEntry(CasEntry& entry)
+	void StorageImpl::AttachEntry(CasEntry& casEntry)
 	{
 		if (m_oldestAccessed)
-			m_oldestAccessed->nextAccessed = &entry;
-		entry.prevAccessed = m_oldestAccessed;
-		entry.nextAccessed = nullptr;
+			m_oldestAccessed->nextAccessed = &casEntry;
+		casEntry.prevAccessed = m_oldestAccessed;
+		casEntry.nextAccessed = nullptr;
 		if (!m_newestAccessed)
-			m_newestAccessed = &entry;
-		m_oldestAccessed = &entry;
+			m_newestAccessed = &casEntry;
+		m_oldestAccessed = &casEntry;
 	}
 
-	void StorageImpl::DetachEntry(CasEntry& entry)
+	void StorageImpl::DetachEntry(CasEntry& casEntry)
 	{
-		CasEntry* prevAccessed = entry.prevAccessed;
+		CasEntry* prevAccessed = casEntry.prevAccessed;
 		if (prevAccessed)
-			prevAccessed->nextAccessed = entry.nextAccessed;
-		else if (m_newestAccessed == &entry)
-			m_newestAccessed = entry.nextAccessed;
+			prevAccessed->nextAccessed = casEntry.nextAccessed;
+		else if (m_newestAccessed == &casEntry)
+			m_newestAccessed = casEntry.nextAccessed;
 
-		if (entry.nextAccessed)
-			entry.nextAccessed->prevAccessed = prevAccessed;
-		else if (m_oldestAccessed == &entry)
+		if (casEntry.nextAccessed)
+			casEntry.nextAccessed->prevAccessed = prevAccessed;
+		else if (m_oldestAccessed == &casEntry)
 			m_oldestAccessed = prevAccessed;
 
-		entry.prevAccessed = nullptr;
-		entry.nextAccessed = nullptr;
+		casEntry.prevAccessed = nullptr;
+		casEntry.nextAccessed = nullptr;
 	}
 
 	bool StorageImpl::WriteCompressed(WriteResult& out, const tchar* from, const tchar* to)
@@ -1022,24 +1022,24 @@ namespace uba
 
 				SCOPED_WRITE_LOCK(m_casLookupLock, lookupLock);
 				auto insres = m_casLookup.try_emplace(casKey);
-				CasEntry& entry = insres.first->second;
-				entry.verified = true;
-				entry.exists = true;
+				CasEntry& casEntry = insres.first->second;
+				casEntry.verified = true;
+				casEntry.exists = true;
 
 				m_casTotalBytes += size;
 				if (insres.second)
 				{
-					entry.key = casKey;
-					entry.size = size;
-					AttachEntry(entry);
+					casEntry.key = casKey;
+					casEntry.size = size;
+					AttachEntry(casEntry);
 				}
 				else
 				{
-					UBA_ASSERT(entry.key == casKey);
+					UBA_ASSERT(casEntry.key == casKey);
 					// We should probably delete this one.. something is wrong
-					if (entry.size != 0 && entry.size != size && !deleteFile)
-						m_logger.Detail(TC("Found cas entry which has a different size than what the table thought! Was %llu, is %llu (%s)"), entry.size, size, e.name);
-					entry.size = size;
+					if (casEntry.size != 0 && casEntry.size != size && !deleteFile)
+						m_logger.Detail(TC("Found cas entry which has a different size than what the table thought! Was %llu, is %llu (%s)"), casEntry.size, size, e.name);
+					casEntry.size = size;
 				}
 
 				if (!size && casKey != ToCasKey(CasKeyHasher(), IsCompressed(casKey)))
@@ -1051,7 +1051,7 @@ namespace uba
 				if (!deleteFile)
 					return;
 
-				DetachEntry(entry);
+				DetachEntry(casEntry);
 				m_casLookup.erase(insres.first);
 				m_casTotalBytes -= size;
 				lookupLock.Leave();
@@ -1071,15 +1071,15 @@ namespace uba
 		// All files we saw is tagged as "handled") so let's see if there are cas entries we didn't find
 		for (auto it = m_casLookup.begin(); it!=m_casLookup.end();)
 		{
-			CasEntry& entry = it->second;
-			if (entry.verified)
+			CasEntry& casEntry = it->second;
+			if (casEntry.verified)
 			{
 				++it;
-				entry.verified = false; // Unhandle the entries to be able to evict later
+				casEntry.verified = false; // Unhandle the entries to be able to evict later
 				continue;
 			}
-			entry.size = 0;
-			DetachEntry(entry);
+			casEntry.size = 0;
+			DetachEntry(casEntry);
 			it = m_casLookup.erase(it);
 			didNotExistCount++;
 		}
@@ -1109,18 +1109,18 @@ namespace uba
 		u64 before = m_casTotalBytes;
 		while (m_casTotalBytes > m_casCapacityBytes)
 		{
-			CasEntry* entry = m_oldestAccessed;
-			if (!entry)
+			CasEntry* casEntry = m_oldestAccessed;
+			if (!casEntry)
 			{
 				UBA_ASSERT(m_casLookup.empty());
 				m_casTotalBytes = 0;
 				break;
 			}
-			DropCasFile(entry->key, true, TC("HandleOverflow"));
+			DropCasFile(casEntry->key, true, TC("HandleOverflow"));
 			if (outDeletedFiles)
-				outDeletedFiles->insert(entry->key);
-			DetachEntry(*entry);
-			m_casLookup.erase(entry->key);
+				outDeletedFiles->insert(casEntry->key);
+			DetachEntry(*casEntry);
+			m_casLookup.erase(casEntry->key);
 		}
 		u64 after = m_casTotalBytes;
 		if (before != after)
@@ -1382,11 +1382,11 @@ namespace uba
 				m_newestAccessed = nullptr;
 				return false;
 			}
-			CasEntry& entry = insres.first->second;
-			entry.key = casKey;
-			entry.size = reader.ReadU64();
-			entry.exists = true;
-			m_casTotalBytes += entry.size;
+			CasEntry& casEntry = insres.first->second;
+			casEntry.key = casKey;
+			casEntry.size = reader.ReadU64();
+			casEntry.exists = true;
+			m_casTotalBytes += casEntry.size;
 
 
 #if UBA_USE_SPARSEFILE
@@ -1394,22 +1394,22 @@ namespace uba
 			u64 mappingOffset = reader.ReadU64();
 			u64 mappingSize = reader.ReadU64();
 
-			entry.mappingHandle = m_casDataBuffer.GetPersistentHandle(mappingFileIndex);
-			entry.mappingOffset = mappingOffset;
-			entry.mappingSize = mappingSize;
-			entry.verified = true;
-			entry.exists = true;
+			casEntry.mappingHandle = m_casDataBuffer.GetPersistentHandle(mappingFileIndex);
+			casEntry.mappingOffset = mappingOffset;
+			casEntry.mappingSize = mappingSize;
+			casEntry.verified = true;
+			casEntry.exists = true;
 #endif
 
 
 			if (prev)
 			{
-				prev->nextAccessed = &entry;
-				entry.prevAccessed = prev;
+				prev->nextAccessed = &casEntry;
+				casEntry.prevAccessed = prev;
 			}
 			else
-				m_newestAccessed = &entry;
-			prev = &entry;
+				m_newestAccessed = &casEntry;
+			prev = &casEntry;
 		}
 		m_oldestAccessed = prev;
 
@@ -1862,10 +1862,7 @@ namespace uba
 
 		SCOPED_WRITE_LOCK(fileEntry.lock, entryLock);
 		
-		// If fileEntry.lastWritten is newer than our verifiedLastWriteTime it means that something else updated the file in between us getting info from file system
-		// and ending up here. If that is the case we can't touch the verified stamp.
-		if (verifiedLastWriteTime >= fileEntry.lastWritten)
-			fileEntry.verified = fileEntry.lastWritten == verifiedLastWriteTime && fileEntry.size == verifiedSize && fileEntry.casKey != CasKeyInvalid;
+		fileEntry.verified = fileEntry.lastWritten == verifiedLastWriteTime && fileEntry.size == verifiedSize && fileEntry.casKey != CasKeyInvalid;
 
 		if (!fileEntry.verified)
 			return false;
@@ -2012,7 +2009,7 @@ namespace uba
 		SCOPED_WRITE_LOCK(fileEntry.lock, entryLock);
 		if (fileEntry.verified)
 			return;
-		fileEntry.verified = fileEntry.lastWritten == verifiedLastWriteTime && fileEntry.size == verifiedSize;
+		fileEntry.verified = fileEntry.lastWritten == verifiedLastWriteTime && fileEntry.size == verifiedSize && fileEntry.casKey != CasKeyInvalid;
 	}
 
 	bool StorageImpl::StoreCasKey(CasKey& out, const tchar* fileName, const CasKey& casKeyOverride, bool fileIsCompressed)
@@ -2570,12 +2567,12 @@ namespace uba
 		if (CaseInsensitiveFs)
 			forKey.MakeLower();
 		StringKey key = ToStringKey(forKey);
-		FileEntry& entry = GetOrCreateFileEntry(key);
-		SCOPED_WRITE_LOCK(entry.lock, lock2);
-		entry.casKey = casKey;
-		entry.lastWritten = lastWritten;
-		entry.size = size;
-		entry.verified = true;
+		FileEntry& fileEntry = GetOrCreateFileEntry(key);
+		SCOPED_WRITE_LOCK(fileEntry.lock, lock2);
+		fileEntry.casKey = casKey;
+		fileEntry.lastWritten = lastWritten;
+		fileEntry.size = size;
+		fileEntry.verified = true;
 		return true;
 	}
 
