[UBA]
* Reverted back to use OpenFile/GetFileInformation instead of GetFileAttributes because weirdly enough it seems to be much faster on cloud setups like AWS.. and I don't know why. You would think 1 syscall is faster than three (open,getinfo.close). GetFileInformation even contains more info than GetFileAttributes :-)

diff --git a/Engine/Source/Programs/UnrealBuildAccelerator/Common/Private/UbaStorage.cpp b/Engine/Source/Programs/UnrealBuildAccelerator/Common/Private/UbaStorage.cpp
--- a/Engine/Source/Programs/UnrealBuildAccelerator/Common/Private/UbaStorage.cpp
+++ b/Engine/Source/Programs/UnrealBuildAccelerator/Common/Private/UbaStorage.cpp
@@ -238,7 +238,7 @@ namespace uba
 		return WriteCompressed(out, from, readHandle, 0, fileSize, to, nullptr, 0);
 	}
 
-	bool StorageImpl::WriteCompressed(WriteResult& out, const tchar* from, FileHandle readHandle, u8* readMem, u64 fileSize, const tchar* to, const void* header, u64 headerSize)
+	bool StorageImpl::WriteCompressed(WriteResult& out, const tchar* from, FileHandle readHandle, u8* readMem, u64 fileSize, const tchar* to, const void* header, u64 headerSize, u64 lastWriteTime)
 	{
 		StorageStats& stats = Stats();
 
@@ -355,6 +355,9 @@ namespace uba
 		}
 
 #if !UBA_USE_SPARSEFILE
+		if (lastWriteTime)
+			if (!SetFileLastWriteTime(destinationFile.GetHandle(), lastWriteTime))
+				return m_logger.Error(TC("Failed to set file time on filehandle for %s"), to);
 		if (!destinationFile.Close())
 			return false;
 #endif
@@ -515,6 +518,7 @@ namespace uba
 			FileHandle readHandle;
 			if (!OpenFileSequentialRead(m_logger, fileName, readHandle))
 				return m_logger.Error(TC("Failed to open file %s for read (%s)"), fileName, LastErrorToText().data);
+			auto fileGuard = MakeGuard([&](){ CloseFile(fileName, readHandle); });
 
 			u64 fileSize;
 			if (!uba::GetFileSizeEx(fileSize, readHandle))
@@ -1911,27 +1915,26 @@ namespace uba
 		}
 		fileEntry.verified = true;
 
-		u64 fileSize = 0;
-		u32 attributes = 0;
-		u64 lastWritten = 0;
-
-		if (!FileExists(m_logger, fileName, &fileSize, &attributes, &lastWritten))
+		// Use OpenFile+GetFileInformationByHandle+close instead of Getting file attributes because it is actually faster on cloud setups (weirdly enough)
+		FileHandle fileHandle;
+		if (!OpenFileSequentialRead(m_logger, fileName, fileHandle))
 		{
 			fileEntry.casKey = CasKeyZero;
-			u32 lastError = GetLastError();
-			if (lastError != ERROR_FILE_NOT_FOUND && lastError != ERROR_PATH_NOT_FOUND)
-				return m_logger.Error(TC("FileExists failed on %s (%s)"), fileName, LastErrorToText(lastError).data);
 			out = CasKeyZero;
 			return true;
 		}
+		auto fileGuard = MakeGuard([&](){ CloseFile(fileName, fileHandle); });
 
-		if (IsDirectory(attributes))
+		FileInformation info;
+		if (!GetFileInformationByHandle(info, m_logger, fileName, fileHandle))
 		{
 			fileEntry.casKey = CasKeyZero;
-			out = CasKeyZero;
-			return true;
+			return m_logger.Error(TC("GetFileInformationByHandle failed on %s"), fileName);
 		}
 
+		u64 fileSize = info.size;
+		u64 lastWritten = info.lastWriteTime;
+
 		if (fileEntry.casKey != CasKeyZero)
 		{
 			if (casKeyOverride != CasKeyZero && casKeyOverride != fileEntry.casKey)
@@ -1955,11 +1958,6 @@ namespace uba
 		fileEntry.lastWritten = lastWritten;
 		if (casKeyOverride == CasKeyZero)
 		{
-			FileHandle fileHandle;
-			if (!OpenFileSequentialRead(m_logger, fileName, fileHandle))
-				return false;
-			auto fileGuard = MakeGuard([&](){ CloseFile(fileName, fileHandle); });
-
 			if (fileIsCompressed)
 			{
 				CompressedObjFileHeader header(CasKeyZero);
@@ -2034,27 +2032,26 @@ namespace uba
 		}
 		fileEntry.verified = true;
 
-		u64 fileSize = 0;
-		u32 attributes = 0;
-		u64 lastWritten = 0;
-
-		if (!FileExists(m_logger, fileName, &fileSize, &attributes, &lastWritten))
+		// Use OpenFile+GetFileInformationByHandle+close instead of Getting file attributes because it is actually faster on cloud setups (weirdly enough)
+		FileHandle fileHandle;
+		if (!OpenFileSequentialRead(m_logger, fileName, fileHandle))
 		{
 			fileEntry.casKey = CasKeyZero;
-			u32 lastError = GetLastError();
-			if (lastError != ERROR_FILE_NOT_FOUND && lastError != ERROR_PATH_NOT_FOUND)
-				return m_logger.Error(TC("FileExists failed on %s (%s)"), fileName, LastErrorToText(lastError).data);
 			out = CasKeyZero;
 			return true;
 		}
+		auto fileGuard = MakeGuard([&](){ CloseFile(fileName, fileHandle); });
 
-		if (IsDirectory(attributes))
+		FileInformation info;
+		if (!GetFileInformationByHandle(info, m_logger, fileName, fileHandle))
 		{
 			fileEntry.casKey = CasKeyZero;
-			out = CasKeyZero;
-			return true;
+			return m_logger.Error(TC("GetFileInformationByHandle failed on %s"), fileName);
 		}
 
+		u64 fileSize = info.size;
+		u64 lastWritten = info.lastWriteTime;
+
 		if (fileEntry.casKey != CasKeyZero)
 		{
 			if (casKeyOverride != CasKeyZero && casKeyOverride != fileEntry.casKey)
@@ -2074,11 +2071,6 @@ namespace uba
 		fileEntry.lastWritten = lastWritten;
 		if (casKeyOverride == CasKeyZero)
 		{
-			FileHandle fileHandle;
-			if (!OpenFileSequentialRead(m_logger, fileName, fileHandle))
-				return false;
-			auto fileGuard = MakeGuard([&](){ CloseFile(fileName, fileHandle); });
-
 			if (fileIsCompressed)
 			{
 				CompressedObjFileHeader header(CasKeyZero);
diff --git a/Engine/Source/Programs/UnrealBuildAccelerator/Common/Public/UbaStorage.h b/Engine/Source/Programs/UnrealBuildAccelerator/Common/Public/UbaStorage.h
--- a/Engine/Source/Programs/UnrealBuildAccelerator/Common/Public/UbaStorage.h
+++ b/Engine/Source/Programs/UnrealBuildAccelerator/Common/Public/UbaStorage.h
@@ -69,7 +69,7 @@ namespace uba
 		virtual void Ping() {}
 
 		struct WriteResult { FileMappingHandle mappingHandle; u64 size = InvalidValue; u64 offset = InvalidValue; };
-		virtual bool WriteCompressed(WriteResult& out, const tchar* from, FileHandle readHandle, u8* readMem, u64 fileSize, const tchar* toFile, const void* header, u64 headerSize) = 0;
+		virtual bool WriteCompressed(WriteResult& out, const tchar* from, FileHandle readHandle, u8* readMem, u64 fileSize, const tchar* toFile, const void* header, u64 headerSize, u64 lastWriteTime = 0) = 0;
 	};
 
 	struct StorageCreateInfo
@@ -154,7 +154,7 @@ namespace uba
 		void ReportFileInfoWeak(const StringKey& fileNameKey, u64 verifiedLastWriteTime, u64 verifiedSize);
 
 		virtual bool WriteCompressed(WriteResult& out, const tchar* from, const tchar* toFile);
-		virtual bool WriteCompressed(WriteResult& out, const tchar* from, FileHandle readHandle, u8* readMem, u64 fileSize, const tchar* toFile, const void* header, u64 headerSize) final;
+		virtual bool WriteCompressed(WriteResult& out, const tchar* from, FileHandle readHandle, u8* readMem, u64 fileSize, const tchar* toFile, const void* header, u64 headerSize, u64 lastWriteTime = 0) final;
 		bool WriteMemToCompressedFile(FileAccessor& destination, u32 workCount, const u8* uncompressedData, u64 fileSize, u64 maxUncompressedBlock, u64& totalWritten);
 		bool WriteCasFileNoCheck(WriteResult& out, const tchar* fileName, bool fileIsCompressed, const tchar* casFile, bool storeCompressed);
 		bool WriteCasFile(WriteResult& out, const tchar* fileName, bool fileIsCompressed, const CasKey& casKey);
