[UBA]
* Added very basic solution for loading a config file. There are lots of internal things that we want to be able to config without having to release new binaries all the time
* Added CacheClientCreateInfo support for config file

diff --git a/Engine/Source/Programs/UnrealBuildAccelerator/Common/Private/UbaCacheClient.cpp b/Engine/Source/Programs/UnrealBuildAccelerator/Common/Private/UbaCacheClient.cpp
--- a/Engine/Source/Programs/UnrealBuildAccelerator/Common/Private/UbaCacheClient.cpp
+++ b/Engine/Source/Programs/UnrealBuildAccelerator/Common/Private/UbaCacheClient.cpp
@@ -5,6 +5,7 @@
 #include "UbaCacheEntry.h"
 #include "UbaCompactTables.h"
 #include "UbaCompressedObjFileHeader.h"
+#include "UbaConfig.h"
 #include "UbaDirectoryIterator.h"
 #include "UbaFileAccessor.h"
 #include "UbaNetworkMessage.h"
@@ -21,6 +22,16 @@
 
 namespace uba
 {
+	void CacheClientCreateInfo::Apply(Config& config)
+	{
+		const ConfigTable& table = config.GetTable(TC("CacheClient"));
+		table.GetValueAsBool(useDirectoryPreparsing, TC("UseDirectoryPreparsing"));
+		table.GetValueAsBool(validateCacheWritesInput, TC("ValidateCacheWritesInput"));
+		table.GetValueAsBool(validateCacheWritesOutput, TC("ValidateCacheWritesOutput"));
+		table.GetValueAsBool(reportMissReason, TC("ReportMissReason"));
+	}
+
+
 	u64 MakeId(u32 bucketId) { return u64(bucketId) | ((u64(!CaseInsensitiveFs) + (RootPathsVersion << 1)) << 32); }
 
 	struct CacheClient::Bucket
@@ -1250,6 +1261,9 @@ namespace uba
 		TraverseDir(m_logger, path.data, 
 			[&](const DirectoryEntry& e)
 			{
+				if (IsDirectory(e.attributes))
+					return;
+
 				path.Clear().Append('\\').Append(e.name, e.nameLen);
 				if (CaseInsensitiveFs)
 					path.MakeLower();
diff --git a/Engine/Source/Programs/UnrealBuildAccelerator/Common/Private/UbaConfig.cpp b/Engine/Source/Programs/UnrealBuildAccelerator/Common/Private/UbaConfig.cpp
new file mode 100644
--- /dev/null
+++ b/Engine/Source/Programs/UnrealBuildAccelerator/Common/Private/UbaConfig.cpp
@@ -0,0 +1,160 @@
+// Copyright Epic Games, Inc. All Rights Reserved.
+
+#include "UbaConfig.h"
+#include "UbaFileAccessor.h"
+
+namespace uba
+{
+	static ConfigTable& EmptyConfigTable = *new ConfigTable();
+
+	bool ConfigTable::GetValueAsString(StringBufferBase& out, const tchar* key) const
+	{
+		return false;
+	}
+
+	bool ConfigTable::GetValueAsU32(u32& out, const tchar* key) const
+	{
+		return false;
+	}
+
+	bool ConfigTable::GetValueAsBool(bool& out, const tchar* key) const
+	{
+		auto findIt = m_values.find(key);
+		if (findIt == m_values.end())
+			return false;
+		const tchar* value = findIt->second.c_str();
+		if (Equals(value, TC("true")) || Equals(value, TC("1")))
+		{
+			out = true;
+			return true;
+		}
+		if (Equals(value, TC("false")) || Equals(value, TC("0")))
+		{
+			out = false;
+			return true;
+		}
+
+		return false;
+	}
+
+	bool Config::LoadFromFile(Logger& logger, const tchar* configFile)
+	{
+		m_isLoaded = true;
+
+		FileAccessor fa(logger, configFile);
+		if (!fa.OpenMemoryRead(0, false))
+			return false;
+		return LoadFromText(logger, (const char*)fa.GetData(), fa.GetSize());
+	}
+
+	bool Config::LoadFromText(Logger& logger, const char* text, u64 textLen)
+	{
+		m_isLoaded = true;
+
+		const char* i = text;
+		const char* e = i + textLen;
+
+		auto consumeEmpty = [&]() -> char
+			{
+				while (i != e)
+				{
+					if (*i != ' ' && *i != '\t'&& *i != '\r')
+						return *i;
+					++i;
+				}
+				return 0;
+			};
+
+		auto consumeIdentifier = [&](StringBufferBase& out) -> char
+			{
+				while (i != e)
+				{
+					if (!(*i >= 'a' && *i <= 'z' || *i >= 'A' && *i <= 'Z'))
+						return *i;
+					out.Append(*i);
+					++i;
+				}
+				return 0;
+			};
+
+		auto consumeLine = [&](StringBufferBase& out) -> char
+			{
+				while (i != e)
+				{
+					if (*i == '\n')
+						return *i;
+					if (*i != '\r')
+						out.Append(*i);
+					++i;
+				}
+				return 0;
+			};
+
+		ConfigTable globalTable;
+		ConfigTable* activeTable = &globalTable;
+		while (true)
+		{
+			char token = consumeEmpty();
+			if (token == 0)
+				break;
+			if (token == '\n')
+			{
+				++i;
+			}
+			else if (token == '[')
+			{
+				++i;
+				StringBuffer<128> tableName;
+				token = consumeIdentifier(tableName);
+				if (token != ']')
+					return logger.Error(TC("No end token after group name %s"), tableName.data);
+				++i;
+				token = consumeEmpty();
+				if (token == 0)
+					break;
+				if (token != '\n')
+					return logger.Error(TC("Unexpected token %c after group %s"), tableName.data);
+				++i;
+				activeTable = &m_tables.try_emplace(tableName.data).first->second;
+			}
+			else
+			{
+				StringBuffer<128> key;
+				consumeIdentifier(key);
+				token = consumeEmpty();
+				if (token != '=')
+					return logger.Error(TC("Unexpected equals sign after key name %s"), key);
+				++i;
+				token = consumeEmpty();
+				if (token == '\"')
+				{
+					return logger.Error(TC("Strings not supported yet %s"), key);
+				}
+				else
+				{
+					StringBuffer<128> value;
+					token = consumeIdentifier(value);
+					activeTable->m_values[key.data] = value.data;
+					if (token == 0)
+						break;
+					++i;
+				}
+			}
+		}
+
+		return true;
+	}
+
+	bool Config::IsLoaded() const
+	{
+		return m_isLoaded;
+	}
+
+	const ConfigTable& Config::GetTable(const tchar* name) const
+	{
+		auto findIt = m_tables.find(name);
+		if (findIt == m_tables.end())
+			return EmptyConfigTable;
+		return findIt->second;
+	}
+}
\ No newline at end of file
diff --git a/Engine/Source/Programs/UnrealBuildAccelerator/Common/Public/UbaCacheClient.h b/Engine/Source/Programs/UnrealBuildAccelerator/Common/Public/UbaCacheClient.h
--- a/Engine/Source/Programs/UnrealBuildAccelerator/Common/Public/UbaCacheClient.h
+++ b/Engine/Source/Programs/UnrealBuildAccelerator/Common/Public/UbaCacheClient.h
@@ -10,6 +10,7 @@ namespace uba
 {
 	class CompactCasKeyTable;
 	class CompactPathTable;
+	class Config;
 	class NetworkClient;
 	class ProcessHandle;
 	class RootPaths;
@@ -27,6 +28,9 @@ namespace uba
 		StorageImpl& storage;
 		NetworkClient& client;
 		Session& session;
+
+		void Apply(Config& config);
+
 		bool reportMissReason = false;
 		bool useDirectoryPreparsing = false; // This is used to minimize syscalls. GetFileAttributes can be very expensive on cloud machines and we can enable this to minimize syscall count
 		bool validateCacheWritesInput = false; // Set to true to validate cas of all input files before sent to cache
diff --git a/Engine/Source/Programs/UnrealBuildAccelerator/Common/Public/UbaConfig.h b/Engine/Source/Programs/UnrealBuildAccelerator/Common/Public/UbaConfig.h
new file mode 100644
--- /dev/null
+++ b/Engine/Source/Programs/UnrealBuildAccelerator/Common/Public/UbaConfig.h
@@ -0,0 +1,33 @@
+// Copyright Epic Games, Inc. All Rights Reserved.
+
+#pragma once
+
+#include "UbaStringBuffer.h"
+
+namespace uba
+{
+	class Logger;
+
+	class ConfigTable
+	{
+	public:
+		bool GetValueAsString(StringBufferBase& out, const tchar* key) const;
+		bool GetValueAsU32(u32& out, const tchar* key) const;
+		bool GetValueAsBool(bool& out, const tchar* key) const;
+
+		UnorderedMap<TString, TString> m_values;
+	};
+
+	class Config : public ConfigTable
+	{
+	public:
+		bool LoadFromFile(Logger& logger, const tchar* configFile);
+		bool LoadFromText(Logger& logger, const char* text, u64 textLen);
+		bool IsLoaded() const;
+
+		const ConfigTable& GetTable(const tchar* name) const;
+
+		bool m_isLoaded = false;
+		UnorderedMap<TString, ConfigTable> m_tables;
+	};
+}
diff --git a/Engine/Source/Programs/UnrealBuildAccelerator/Host/Private/UbaExports.cpp b/Engine/Source/Programs/UnrealBuildAccelerator/Host/Private/UbaExports.cpp
--- a/Engine/Source/Programs/UnrealBuildAccelerator/Host/Private/UbaExports.cpp
+++ b/Engine/Source/Programs/UnrealBuildAccelerator/Host/Private/UbaExports.cpp
@@ -3,6 +3,7 @@
 #include "UbaExports.h"
 #include "UbaAWS.h"
 #include "UbaCacheClient.h"
+#include "UbaConfig.h"
 #include "UbaNetworkBackendQuic.h"
 #include "UbaNetworkBackendTcp.h"
 #include "UbaNetworkClient.h"
@@ -84,6 +85,24 @@ namespace uba
 		UbaAssert(desc.data, "", 0, "", -1);
 	}
 	#endif
+
+	Config& GetConfig(const tchar* fileName = nullptr)
+	{
+		static Config config;
+		if (config.IsLoaded())
+			return config;
+		LoggerWithWriter logger(g_nullLogWriter);
+		StringBuffer<> temp;
+		if (!fileName)
+		{
+			GetDirectoryOfCurrentModule(logger, temp);
+			temp.EnsureEndsWithSlash().Append(TC("UbaHost.toml"));
+			fileName = temp.data;
+		}
+
+		config.LoadFromFile(logger, fileName);
+		return config;
+	}
 }
 
 extern "C"
@@ -114,6 +133,12 @@ extern "C"
 			delete writer;
 	}
 
+	bool Config_Load(const uba::tchar* configFile)
+	{
+		uba::GetConfig(configFile);
+		return true;
+	}
+
 	uba::NetworkServer* NetworkServer_Create(uba::LogWriter& writer, uba::u32 workerCount, uba::u32 sendSize, uba::u32 receiveTimeoutSeconds, bool useQuic)
 	{
 		using namespace uba;
@@ -522,6 +547,8 @@ extern "C"
 			return nullptr;
 		}
 		CacheClientCreateInfo info{writer, storage, *networkClient, *session};
+		info.apply(GetConfig);
+
 		info.reportMissReason = reportMissReason;
 		return new CacheClient(info);
 	}
diff --git a/Engine/Source/Programs/UnrealBuildAccelerator/Host/Public/UbaExports.h b/Engine/Source/Programs/UnrealBuildAccelerator/Host/Public/UbaExports.h
--- a/Engine/Source/Programs/UnrealBuildAccelerator/Host/Public/UbaExports.h
+++ b/Engine/Source/Programs/UnrealBuildAccelerator/Host/Public/UbaExports.h
@@ -45,6 +45,9 @@ extern "C"
 	UBA_API uba::LogWriter* CreateCallbackLogWriter(uba::CallbackLogWriter::BeginScopeCallback begin, uba::CallbackLogWriter::EndScopeCallback end, uba::CallbackLogWriter::LogCallback log);
 	UBA_API void DestroyCallbackLogWriter(uba::LogWriter* writer);
 
+	// Config
+	UBA_API bool Config_Load(const uba::tchar* configFile);
+
 	// NetworkServer
 	UBA_API uba::NetworkServer* NetworkServer_Create(uba::LogWriter& writer = uba::g_consoleLogWriter, uba::u32 workerCount = 64, uba::u32 sendSize = uba::SendDefaultSize, uba::u32 receiveTimeoutSeconds = 60, bool useQuic = false);
 	UBA_API void NetworkServer_Destroy(uba::NetworkServer* server);
diff --git a/Engine/Source/Programs/UnrealBuildAccelerator/Test/Private/UbaTestAll.h b/Engine/Source/Programs/UnrealBuildAccelerator/Test/Private/UbaTestAll.h
--- a/Engine/Source/Programs/UnrealBuildAccelerator/Test/Private/UbaTestAll.h
+++ b/Engine/Source/Programs/UnrealBuildAccelerator/Test/Private/UbaTestAll.h
@@ -4,6 +4,7 @@
 
 #include "UbaTestBasics.h"
 #include "UbaTestCache.h"
+#include "UbaTestConfig.h"
 #include "UbaTestNetwork.h"
 #include "UbaTestScheduler.h"
 #include "UbaTestSession.h"
@@ -31,6 +32,7 @@ namespace uba
 		UBA_TEST(TestRemoteDetouredTestApp) \
 		UBA_TEST(TestCompactPathTable) \
 		UBA_TEST(TestCacheEntry) \
+		UBA_TEST(TestConfig) \
 
 
 #define UBA_POSIX_TESTS \
diff --git a/Engine/Source/Programs/UnrealBuildAccelerator/Test/Private/UbaTestConfig.h b/Engine/Source/Programs/UnrealBuildAccelerator/Test/Private/UbaTestConfig.h
new file mode 100644
--- /dev/null
+++ b/Engine/Source/Programs/UnrealBuildAccelerator/Test/Private/UbaTestConfig.h
@@ -0,0 +1,28 @@
+// Copyright Epic Games, Inc. All Rights Reserved.
+
+#pragma once
+
+#include "UbaConfig.h"
+
+namespace uba
+{
+	bool TestConfig(Logger& logger, const StringBufferBase& rootDir)
+	{
+		static const char* configText =
+			"[CacheClient]\r\n"
+			"UseDirectoryPreparsing = true\r\n"
+			"";
+
+		Config config;
+		if (!config.LoadFromText(logger, configText, strlen(configText)))
+			return false;
+
+		const ConfigTable& table = config.GetTable(TC("CacheClient"));
+		bool test = false;
+		if (!table.GetValueAsBool(test, TC("UseDirectoryPreparsing")))
+			return false;
+		if (test != true)
+			return false;
+		return true;
+	}
+}
\ No newline at end of file
