Fixed typo in ubaexports

diff --git a/Engine/Source/Programs/UnrealBuildAccelerator/Host/Private/UbaExports.cpp b/Engine/Source/Programs/UnrealBuildAccelerator/Host/Private/UbaExports.cpp
--- a/Engine/Source/Programs/UnrealBuildAccelerator/Host/Private/UbaExports.cpp
+++ b/Engine/Source/Programs/UnrealBuildAccelerator/Host/Private/UbaExports.cpp
@@ -547,7 +547,7 @@ extern "C"
 			return nullptr;
 		}
 		CacheClientCreateInfo info{writer, storage, *networkClient, *session};
-		info.apply(GetConfig);
+		info.Apply(GetConfig);
 
 		info.reportMissReason = reportMissReason;
 		return new CacheClient(info);
