[UBA]
* Fixed macos compile errors

diff --git a/Engine/Source/Programs/UnrealBuildAccelerator/Common/Private/UbaConfig.cpp b/Engine/Source/Programs/UnrealBuildAccelerator/Common/Private/UbaConfig.cpp
--- a/Engine/Source/Programs/UnrealBuildAccelerator/Common/Private/UbaConfig.cpp
+++ b/Engine/Source/Programs/UnrealBuildAccelerator/Common/Private/UbaConfig.cpp
@@ -69,7 +69,7 @@ namespace uba
 			{
 				while (i != e)
 				{
-					if (!(*i >= 'a' && *i <= 'z' || *i >= 'A' && *i <= 'Z'))
+					if (!((*i >= 'a' && *i <= 'z') || (*i >= 'A' && *i <= 'Z')))
 						return *i;
 					out.Append(*i);
 					++i;
