[UBA]
* Changed so casdb entry is not invalidated when file is stored from helper since we already have caskey and everything ready
* Added options to disable usage of cache hits in cacheclient.. this is a debugging feature
* Added option to not use path roots in cache client
* Added more debug information to try to figure out bug where cache client sees wrong cas key of file

diff --git a/Engine/Source/Programs/UnrealBuildAccelerator/Common/Private/UbaCacheClient.cpp b/Engine/Source/Programs/UnrealBuildAccelerator/Common/Private/UbaCacheClient.cpp
--- a/Engine/Source/Programs/UnrealBuildAccelerator/Common/Private/UbaCacheClient.cpp
+++ b/Engine/Source/Programs/UnrealBuildAccelerator/Common/Private/UbaCacheClient.cpp
@@ -29,11 +29,11 @@ namespace uba
 		table.GetValueAsBool(validateCacheWritesInput, TC("ValidateCacheWritesInput"));
 		table.GetValueAsBool(validateCacheWritesOutput, TC("ValidateCacheWritesOutput"));
 		table.GetValueAsBool(reportMissReason, TC("ReportMissReason"));
+		table.GetValueAsBool(useRoots, TC("UseRoots"));
+		table.GetValueAsBool(useCacheHit, TC("UseCacheHit"));
 	}
 
 
-	u64 MakeId(u32 bucketId) { return u64(bucketId) | ((u64(!CaseInsensitiveFs) + (RootPathsVersion << 1)) << 32); }
-
 	struct CacheClient::Bucket
 	{
 		Bucket(u32 id_)
@@ -76,6 +76,8 @@ namespace uba
 		m_useDirectoryPreParsing = info.useDirectoryPreparsing;
 		m_validateCacheWritesInput = info.validateCacheWritesInput;
 		m_validateCacheWritesOutput = info.validateCacheWritesOutput;
+		m_useCacheHit = info.useCacheHit;
+		m_useRoots = info.useRoots;
 
 		m_client.RegisterOnConnected([this]()
 			{
@@ -158,6 +160,8 @@ namespace uba
 		Bucket& bucket = m_buckets.try_emplace(bucketId, bucketId).first->second;
 		bucketsLock.Leave();
 
+		TString qualifiedPath;
+
 		// Traverse all inputs and outputs. to create cache entry that we can send to server
 		while (true)
 		{
@@ -207,21 +211,29 @@ namespace uba
 				continue;
 			}
 
-			// Find root for path in order to be able to normalize it.
-			auto root = rootPaths.FindRoot(path);
-			if (!root)
+			if (m_useRoots)
 			{
-				m_logger.Info(TC("FILE WITHOUT ROOT: %s (%s)"), path.data, info.description);
-				success = false;
-				continue;
-			}
+				// Find root for path in order to be able to normalize it.
+				auto root = rootPaths.FindRoot(path);
+				if (!root)
+				{
+					m_logger.Info(TC("FILE WITHOUT ROOT: %s (%s)"), path.data, info.description);
+					success = false;
+					continue;
+				}
 
-			if (!root->includeInKey)
-				continue;
+				if (!root->includeInKey)
+					continue;
+
+				u32 rootLen = u32(root->path.size());
+				qualifiedPath = path.data + rootLen - 1;
+				qualifiedPath[0] = tchar(RootPaths::RootStartByte + root->index);
+			}
+			else
+			{
+				qualifiedPath = path.data;
+			}
 
-			u32 rootLen = u32(root->path.size());
-			TString qualifiedPath = path.data + rootLen - 1;
-			qualifiedPath[0] = tchar(RootPaths::RootStartByte + root->index);
 
 			u32 pathOffset = bucket.sendPathTable.Add(qualifiedPath.c_str(), u32(qualifiedPath.size()), &requiredPathTableSize);
 
@@ -258,6 +270,7 @@ namespace uba
 				CasKey oldOldKey;
 				u64 oldOldSize = 0;
 				u64 oldOldLastWritten = 0;
+				u64 lastInvalidationTime = 0;
 				bool oldOldVerified = false;
 
 				if (shouldValidate)
@@ -268,6 +281,7 @@ namespace uba
 					oldOldSize = fileEntry.size;
 					oldOldLastWritten = fileEntry.lastWritten;
 					oldOldVerified = fileEntry.verified;
+					lastInvalidationTime = fileEntry.lastInvalidationTime;
 				}
 
 
@@ -326,8 +340,8 @@ namespace uba
 						SCOPED_READ_LOCK(fileEntry.lock, lock);
 
 						auto ToString = [](bool b) { return b ? TC("true") : TC("false"); };
-						m_logger.Warning(TC("CasDb claims file %s has caskey %s but recalculating it gives us %s (OldEntry: %s/%llu/%llu/%s FileEntry: %llu/%llu/%s, Real: %llu/%llu). Will not populate cache for %s"),
-							path.data, CasKeyString(oldKey).str, CasKeyString(newKey).str, CasKeyString(oldOldKey).str, oldOldSize, oldOldLastWritten, ToString(oldOldVerified), fileEntry.size, fileEntry.lastWritten, ToString(fileEntry.verified), fileSize, fileInfo.lastWriteTime, info.description);
+						m_logger.Warning(TC("CasDb claims file %s has caskey %s but recalculating it gives us %s (OldEntry: %s/%llu/%llu/%s FileEntry: %llu/%llu/%s, Real: %llu/%llu. Last invalidation: %llu). Will not populate cache for %s"),
+							path.data, CasKeyString(oldKey).str, CasKeyString(newKey).str, CasKeyString(oldOldKey).str, oldOldSize, oldOldLastWritten, ToString(oldOldVerified), fileEntry.size, fileEntry.lastWritten, ToString(fileEntry.verified), fileSize, fileInfo.lastWriteTime, lastInvalidationTime, info.description);
 						return false;
 					}
 				}
@@ -382,6 +396,12 @@ namespace uba
 		return true;
 	}
 
+	u64 CacheClient::MakeId(u32 bucketId)
+	{
+		return u64(bucketId) | ((u64(!CaseInsensitiveFs) + (RootPathsVersion << 1) + (u8(m_useRoots) << 2)) << 32);
+	}
+
+
 	bool CacheClient::FetchFromCache(bool& outCacheHit, const RootPaths& rootPaths, u32 bucketId, const ProcessStartInfo& info)
 	{
 		outCacheHit = false;
@@ -615,7 +635,12 @@ namespace uba
 			}
 			if (rangeBegin != ~0u)
 				addRange(u32(sharedReader.GetPosition()));
-			UBA_ASSERT(!sharedMatchingRanges.empty());
+			if (sharedMatchingRanges.empty())
+			{
+				auto& range = sharedMatchingRanges.emplace_back();
+				range.begin = 0;
+				range.end = 0;
+			}
 		}
 
 		// Read entries
@@ -676,6 +701,9 @@ namespace uba
 				UBA_ASSERTF(matchingId == entryId, TC("%u vs %u"), matchingId, entryId);
 #endif
 
+				if (!m_useCacheHit)
+					return false;
+
 				{
 					StackBinaryWriter<128> writer;
 					NetworkMessage msg(m_client, CacheServiceId, CacheMessageType_ReportUsedEntry, writer);
@@ -1189,6 +1217,8 @@ namespace uba
 
 	bool CacheClient::ShouldNormalize(const StringBufferBase& path)
 	{
+		if (!m_useRoots)
+			return false;
 		if (path.EndsWith(TC(".dep.json"))) // Contains absolute paths (dep file for msvc)
 			return true;
 		if (path.EndsWith(TC(".d"))) // Contains absolute paths (dep file for clang)
@@ -1213,7 +1243,7 @@ namespace uba
 		u32 rootIndex = normalizedPath[0] - RootPaths::RootStartByte;
 		const TString& root = rootPaths.GetRoot(rootIndex);
 
-		outPath.Append(root).Append(normalizedPath.data + 1);
+		outPath.Append(root).Append(normalizedPath.data + u32(m_useRoots)); // If we use root paths, then first byte is root path table index
 		return true;
 	}
 
diff --git a/Engine/Source/Programs/UnrealBuildAccelerator/Common/Private/UbaSession.cpp b/Engine/Source/Programs/UnrealBuildAccelerator/Common/Private/UbaSession.cpp
--- a/Engine/Source/Programs/UnrealBuildAccelerator/Common/Private/UbaSession.cpp
+++ b/Engine/Source/Programs/UnrealBuildAccelerator/Common/Private/UbaSession.cpp
@@ -647,7 +647,7 @@ namespace uba
 	}
 
 
-	bool Session::RegisterCreateFileForWrite(StringKey fileNameKey, const StringView& fileName, bool registerRealFile, u64 fileSize, u64 lastWriteTime)
+	bool Session::RegisterCreateFileForWrite(StringKey fileNameKey, const StringView& fileName, bool registerRealFile, u64 fileSize, u64 lastWriteTime, bool invalidateStorage)
 	{
 		// Remote is not updating its own directory table
 		if (m_runningRemote)
@@ -772,7 +772,7 @@ namespace uba
 
 		// There are directory crawlers happening in parallel so we need to really make sure to invalidate this one since a crawler can actually
 		// hit this file with information from a query before it was written.. and then it will turn it back to "verified" using old info
-		if (registerRealFile)
+		if (registerRealFile && invalidateStorage)
 			m_storage.InvalidateCachedFileInfo(fileNameKey);
 
 		FileEntryAdded(fileNameKey, lastWriteTime, fileSize);
@@ -2100,7 +2100,7 @@ namespace uba
 					Storage::WriteResult res;
 					CompressedObjFileHeader header { CalculateCasKey(mem, fileSize, true, m_workManager, file.name.c_str()) };
 
-					if (!m_storage.WriteCompressed(res, TC("MemoryMap"), InvalidFileHandle, mem, fileSize, file.name.c_str(), &header, sizeof(header)))
+					if (!m_storage.WriteCompressed(res, TC("MemoryMap"), InvalidFileHandle, mem, fileSize, file.name.c_str(), &header, sizeof(header), file.lastWriteTime))
 						return false;
 					shouldEvictFromMemory = true;
 				}
diff --git a/Engine/Source/Programs/UnrealBuildAccelerator/Common/Private/UbaSessionServer.cpp b/Engine/Source/Programs/UnrealBuildAccelerator/Common/Private/UbaSessionServer.cpp
--- a/Engine/Source/Programs/UnrealBuildAccelerator/Common/Private/UbaSessionServer.cpp
+++ b/Engine/Source/Programs/UnrealBuildAccelerator/Common/Private/UbaSessionServer.cpp
@@ -815,6 +815,8 @@ namespace uba
 					auto logType = connectionInfo.ShouldDisconnect() ? LogEntryType_Info : LogEntryType_Warning;
 					m_logger.Logf(logType, TC("Failed to retrieve cas for %s from client (Needed to write %s)"), CasKeyString(casKey).str, destination.data);
 				}
+
+				bool shouldWriteToDisk = ShouldWriteToDisk(destination);
 				if (success)
 				{
 					if (destination.StartsWith(TC("<log>")))
@@ -839,7 +841,7 @@ namespace uba
 						return true;
 					}
 
-					if (ShouldWriteToDisk(destination))
+					if (shouldWriteToDisk)
 					{
 						bool writeCompressed = false;
 						if (m_storeObjFilesCompressed)
@@ -870,7 +872,8 @@ namespace uba
 				if (success)
 				{
 					m_storage.DropCasFile(casKey, false, destination.data);
-					RegisterCreateFileForWrite(StringKeyZero, destination, true);
+					bool invalidateStorage = false; // No need, already handled in m_storage.CopyOrLink
+					RegisterCreateFileForWrite(StringKeyZero, destination, shouldWriteToDisk, 0, 0, invalidateStorage);
 
 
 					SCOPED_WRITE_LOCK(m_processesLock, lock);
diff --git a/Engine/Source/Programs/UnrealBuildAccelerator/Common/Private/UbaStorage.cpp b/Engine/Source/Programs/UnrealBuildAccelerator/Common/Private/UbaStorage.cpp
--- a/Engine/Source/Programs/UnrealBuildAccelerator/Common/Private/UbaStorage.cpp
+++ b/Engine/Source/Programs/UnrealBuildAccelerator/Common/Private/UbaStorage.cpp
@@ -1886,6 +1886,7 @@ namespace uba
 		SCOPED_WRITE_LOCK(fileEntry.lock, entryLock);
 		fileEntry.verified = false;
 		fileEntry.casKey = CasKeyInvalid;
+		fileEntry.lastInvalidationTime = GetSystemTimeAsFileTime();
 		return true;
 	}
 
diff --git a/Engine/Source/Programs/UnrealBuildAccelerator/Common/Public/UbaCacheClient.h b/Engine/Source/Programs/UnrealBuildAccelerator/Common/Public/UbaCacheClient.h
--- a/Engine/Source/Programs/UnrealBuildAccelerator/Common/Public/UbaCacheClient.h
+++ b/Engine/Source/Programs/UnrealBuildAccelerator/Common/Public/UbaCacheClient.h
@@ -35,6 +35,8 @@ namespace uba
 		bool useDirectoryPreparsing = false; // This is used to minimize syscalls. GetFileAttributes can be very expensive on cloud machines and we can enable this to minimize syscall count
 		bool validateCacheWritesInput = false; // Set to true to validate cas of all input files before sent to cache
 		bool validateCacheWritesOutput = true; // Set to true to validate cas of all output files before sent to cache
+		bool useRoots = true;
+		bool useCacheHit = true;
 	};
 
 	class CacheClient
@@ -55,6 +57,7 @@ namespace uba
 
 	private:
 		struct Bucket;
+		u64 MakeId(u32 bucketId);
 
 		bool SendPathTable(Bucket& bucket, u32 requiredPathTableSize);
 		bool SendCasTable(Bucket& bucket, u32 requiredCasTableSize);
@@ -76,6 +79,8 @@ namespace uba
 		bool m_useDirectoryPreParsing;
 		bool m_validateCacheWritesInput;
 		bool m_validateCacheWritesOutput;
+		bool m_useRoots;
+		bool m_useCacheHit;
 
 		Atomic<bool> m_connected;
 
diff --git a/Engine/Source/Programs/UnrealBuildAccelerator/Common/Public/UbaSession.h b/Engine/Source/Programs/UnrealBuildAccelerator/Common/Public/UbaSession.h
--- a/Engine/Source/Programs/UnrealBuildAccelerator/Common/Public/UbaSession.h
+++ b/Engine/Source/Programs/UnrealBuildAccelerator/Common/Public/UbaSession.h
@@ -110,7 +110,7 @@ namespace uba
 		bool CreateMemoryMapFromFile(MemoryMap& out, StringKey fileNameKey, const tchar* fileName, bool isCompressed, u64 alignment);
 		bool CreateMemoryMapFromView(MemoryMap& out, StringKey fileNameKey, const tchar* fileName, const CasKey& casKey, u64 alignment);
 
-		bool RegisterCreateFileForWrite(StringKey fileNameKey, const StringView& fileName, bool registerRealFile, u64 fileSize = 0, u64 lastWriteTime = 0);
+		bool RegisterCreateFileForWrite(StringKey fileNameKey, const StringView& fileName, bool registerRealFile, u64 fileSize = 0, u64 lastWriteTime = 0, bool invalidateStorage = true);
 		u32 RegisterDeleteFile(StringKey fileNameKey, const tchar* fileName);
 
 		virtual bool PrepareProcess(const ProcessStartInfo& startInfo, bool isChild, StringBufferBase& outRealApplication, const tchar*& outRealWorkingDir);
diff --git a/Engine/Source/Programs/UnrealBuildAccelerator/Common/Public/UbaStorage.h b/Engine/Source/Programs/UnrealBuildAccelerator/Common/Public/UbaStorage.h
--- a/Engine/Source/Programs/UnrealBuildAccelerator/Common/Public/UbaStorage.h
+++ b/Engine/Source/Programs/UnrealBuildAccelerator/Common/Public/UbaStorage.h
@@ -188,7 +188,15 @@ namespace uba
 		StringBuffer<> m_tempPath;
 
 		ReaderWriterLock m_fileTableLookupLock;
-		struct FileEntry { ReaderWriterLock lock; CasKey casKey = CasKeyZero; u64 size = 0; u64 lastWritten = 0; bool verified = false; };
+		struct FileEntry
+		{
+			ReaderWriterLock lock;
+			CasKey casKey = CasKeyZero;
+			u64 size = 0;
+			u64 lastWritten = 0;
+			bool verified = false;
+			u64 lastInvalidationTime = 0; // Remove when we have figured out a bug
+		};
 		UnorderedMap<StringKey, FileEntry> m_fileTableLookup;
 
 		ReaderWriterLock m_casLookupLock;
