[UBA]
* Fixed so cacheclient useRoots set to true does not affect bucket id

diff --git a/Engine/Source/Programs/UnrealBuildAccelerator/Common/Private/UbaCacheClient.cpp b/Engine/Source/Programs/UnrealBuildAccelerator/Common/Private/UbaCacheClient.cpp
--- a/Engine/Source/Programs/UnrealBuildAccelerator/Common/Private/UbaCacheClient.cpp
+++ b/Engine/Source/Programs/UnrealBuildAccelerator/Common/Private/UbaCacheClient.cpp
@@ -398,7 +398,7 @@ namespace uba
 
 	u64 CacheClient::MakeId(u32 bucketId)
 	{
-		return u64(bucketId) | ((u64(!CaseInsensitiveFs) + (RootPathsVersion << 1) + (u8(m_useRoots) << 2)) << 32);
+		return u64(bucketId) | ((u64(!CaseInsensitiveFs) + (RootPathsVersion << 1) + (u8(!m_useRoots) << 2)) << 32);
 	}
 
 
