[UBA]
* Added crypto support on linux/macos

diff --git a/Engine/Source/Programs/UnrealBuildAccelerator/Common/Private/UbaCrypto.cpp b/Engine/Source/Programs/UnrealBuildAccelerator/Common/Private/UbaCrypto.cpp
--- a/Engine/Source/Programs/UnrealBuildAccelerator/Common/Private/UbaCrypto.cpp
+++ b/Engine/Source/Programs/UnrealBuildAccelerator/Common/Private/UbaCrypto.cpp
@@ -8,13 +8,22 @@
 #if PLATFORM_WINDOWS
 #define UBA_CRYPTO_TYPE 1
 #else
-#define UBA_CRYPTO_TYPE 0 
+#define UBA_CRYPTO_TYPE 2
 #endif
 
 #if UBA_CRYPTO_TYPE == 1
 	#include <winternl.h>
 	#include <bcrypt.h>
 	#pragma comment (lib, "bcrypt.lib")
+#elif UBA_CRYPTO_TYPE == 2
+	#include "openssl/aes.h"
+
+	struct KeyData
+	{
+		AES_KEY encryptKey;
+		AES_KEY decryptKey;
+	};
+
 #endif // UBA_CRYPTO_TYPE
 
 namespace uba
@@ -56,6 +65,11 @@ namespace uba
 		}
 
 		return (CryptoKey)(u64)keyHandle;
+#elif UBA_CRYPTO_TYPE == 2
+		auto data = new KeyData;
+		AES_set_encrypt_key(key128, 128, &data->encryptKey);
+		AES_set_decrypt_key(key128, 128, &data->decryptKey);
+		return (CryptoKey)uintptr_t(data);
 #else
 		logger.Error(TC("ERROR: Crypto not supported on non-windows platforms"));
 		return InvalidCryptoKey;
@@ -73,6 +87,9 @@ namespace uba
 			return (CryptoKey)(u64)newKey;
 		logger.Error(L"ERROR: BCryptDuplicateKey failed (0x%x)", res);
 		return InvalidCryptoKey;
+#elif UBA_CRYPTO_TYPE == 2
+		auto data = new KeyData(*(KeyData*)uintptr_t(original));
+		return (CryptoKey)uintptr_t(data);
 #else
 		return InvalidCryptoKey;
 #endif // UBA_CRYPTO_TYPE
@@ -82,12 +99,14 @@ namespace uba
 	{
 #if UBA_CRYPTO_TYPE == 1
 		BCryptDestroyKey((BCRYPT_KEY_HANDLE)key);
+#elif UBA_CRYPTO_TYPE == 2
+		delete (KeyData*)uintptr_t(key);
 #endif // UBA_CRYPTO_TYPE
 	}
 
+#if UBA_CRYPTO_TYPE == 1
 	bool BCryptEncryptDecrypt(Logger& logger, bool encrypt, CryptoKey key, u8* data, u32 size)
 	{
-#if UBA_CRYPTO_TYPE == 1
 		u8 objectBuffer[1024];
 		u32 objectBufferLen = sizeof(objectBuffer);
 
@@ -118,18 +137,38 @@ namespace uba
 			logger.Error(L"ERROR: %s cipher text length does not match aligned size", (encrypt ? L"BCryptEncrypt" : L"BCryptDecrypt"));
 			return false;
 		}
-#endif
 		return true;
 	}
+#endif
 
 	bool Crypto::Encrypt(Logger& logger, CryptoKey key, u8* data, u32 size)
 	{
+#if UBA_CRYPTO_TYPE == 1
 		return BCryptEncryptDecrypt(logger, true, key, data, size);
+#elif UBA_CRYPTO_TYPE == 2
+		auto& keyData = *(KeyData*)uintptr_t(key);
+		unsigned char iv[AES_BLOCK_SIZE];
+		memset(iv, 0x00, AES_BLOCK_SIZE);
+		AES_cbc_encrypt(data, data, size, &keyData.encryptKey, iv, AES_ENCRYPT);
+		return true;
+#else
+		return false;
+#endif
 	}
 
 	bool Crypto::Decrypt(Logger& logger, CryptoKey key, u8* data, u32 size)
 	{
+#if UBA_CRYPTO_TYPE == 1
 		return BCryptEncryptDecrypt(logger, false, key, data, size);
+#elif UBA_CRYPTO_TYPE == 2
+		auto& keyData = *(KeyData*)uintptr_t(key);
+		unsigned char iv[AES_BLOCK_SIZE];
+		memset(iv, 0x00, AES_BLOCK_SIZE);
+		AES_cbc_encrypt(data, data, size, &keyData.decryptKey, iv, AES_DECRYPT);
+		return true;
+#else
+		return false;
+#endif
 	}
 
 }
\ No newline at end of file
diff --git a/Engine/Source/Programs/UnrealBuildAccelerator/Common/UbaCommon.Build.cs b/Engine/Source/Programs/UnrealBuildAccelerator/Common/UbaCommon.Build.cs
--- a/Engine/Source/Programs/UnrealBuildAccelerator/Common/UbaCommon.Build.cs
+++ b/Engine/Source/Programs/UnrealBuildAccelerator/Common/UbaCommon.Build.cs
@@ -36,6 +36,13 @@ public class UbaCommon : ModuleRules
 			PublicSystemLibraries.Add("dl");
 		}
 
+		if (Target.Platform != UnrealTargetPlatform.Win64)
+		{
+			PrivateDependencyModuleNames.AddRange(new string[] {
+				"OpenSSL",
+			});
+		}
+
 		PrivateDefinitions.AddRange(new string[] {
 			"_CONSOLE",
 		});
diff --git a/Engine/Source/Programs/UnrealBuildAccelerator/Test/Private/UbaTestAll.h b/Engine/Source/Programs/UnrealBuildAccelerator/Test/Private/UbaTestAll.h
--- a/Engine/Source/Programs/UnrealBuildAccelerator/Test/Private/UbaTestAll.h
+++ b/Engine/Source/Programs/UnrealBuildAccelerator/Test/Private/UbaTestAll.h
@@ -5,6 +5,7 @@
 #include "UbaTestBasics.h"
 #include "UbaTestCache.h"
 #include "UbaTestConfig.h"
+#include "UbaTestCrypto.h"
 #include "UbaTestNetwork.h"
 #include "UbaTestScheduler.h"
 #include "UbaTestSession.h"
@@ -33,6 +34,7 @@ namespace uba
 		UBA_TEST(TestCompactPathTable) \
 		UBA_TEST(TestCacheEntry) \
 		UBA_TEST(TestConfig) \
+		UBA_TEST(TestCrypto) \
 
 
 #define UBA_POSIX_TESTS \
diff --git a/Engine/Source/Programs/UnrealBuildAccelerator/Test/Private/UbaTestCrypto.h b/Engine/Source/Programs/UnrealBuildAccelerator/Test/Private/UbaTestCrypto.h
new file mode 100644
--- /dev/null
+++ b/Engine/Source/Programs/UnrealBuildAccelerator/Test/Private/UbaTestCrypto.h
@@ -0,0 +1,44 @@
+// Copyright Epic Games, Inc. All Rights Reserved.
+
+#pragma once
+
+#include "UbaCrypto.h"
+#include "UbaNetwork.h"
+
+namespace uba
+{
+	bool TestCrypto(Logger& logger, const StringBufferBase& rootDir)
+	{
+		u64 key128[] = { 0x1234567812345678llu, 0x1234567812345678llu };
+
+		CryptoKey key = Crypto::CreateKey(logger, (const u8*)key128);
+
+		u8 encryptedData[sizeof(EncryptionHandshakeString)];
+		memcpy(encryptedData, EncryptionHandshakeString, sizeof(EncryptionHandshakeString));
+
+		for (u32 i=0; i!=3; ++i)
+		{
+			if (!Crypto::Encrypt(logger, key, encryptedData, sizeof(encryptedData)))
+				return false;
+
+			if (memcmp(encryptedData, EncryptionHandshakeString, sizeof(EncryptionHandshakeString)) == 0)
+				return false;
+
+			if (!Crypto::Decrypt(logger, key, encryptedData, sizeof(encryptedData)))
+				return false;
+
+			if (memcmp(encryptedData, EncryptionHandshakeString, sizeof(EncryptionHandshakeString)) != 0)
+				return false;
+
+			if (i == 1)
+			{
+				CryptoKey newKey = Crypto::DuplicateKey(logger, key);
+				Crypto::DestroyKey(key);
+				key = newKey;
+			}
+		}
+
+		Crypto::DestroyKey(key);
+		return true;
+	}
+}
\ No newline at end of file
