[UBA]
* Fixed bug where renamed file reported wrong stringkey for name to session from detoured process
* Fixed so compressed obj files wrote the right size on disk to file lookuptable

diff --git a/Engine/Source/Programs/UnrealBuildAccelerator/Common/Private/UbaStorage.cpp b/Engine/Source/Programs/UnrealBuildAccelerator/Common/Private/UbaStorage.cpp
--- a/Engine/Source/Programs/UnrealBuildAccelerator/Common/Private/UbaStorage.cpp
+++ b/Engine/Source/Programs/UnrealBuildAccelerator/Common/Private/UbaStorage.cpp
@@ -2436,6 +2436,8 @@ namespace uba
 				BottleneckScope scope(bottleneck);
 				#endif
 
+				u64 sizeOnDisk = decompressedSize;
+
 				if (writeCompressed)
 				{
 					u64 compressedFileSize = mappedView.size;
@@ -2443,8 +2445,9 @@ namespace uba
 						if (!GetFileSizeEx(compressedFileSize, readHandle))
 							return m_logger.Error(TC("Failed to get file size of compressed file %s (%s)"), casFile.data, LastErrorToText().data);
 						
+					sizeOnDisk = compressedFileSize + sizeof(CompressedObjFileHeader);
 
-					if (!destinationFile.CreateMemoryWrite(false, fileAttributes, compressedFileSize + sizeof(CompressedObjFileHeader), m_tempPath.data))
+					if (!destinationFile.CreateMemoryWrite(false, fileAttributes, sizeOnDisk, m_tempPath.data))
 						return false;
 					u8* writePos = destinationFile.GetData();
 					*(CompressedObjFileHeader*)writePos = CompressedObjFileHeader(casKey);
@@ -2497,7 +2500,7 @@ namespace uba
 				{
 					fileEntry.casKey = casKey;
 					fileEntry.lastWritten = lastWriteTime;
-					fileEntry.size = decompressedSize;
+					fileEntry.size = sizeOnDisk;
 					fileEntry.verified = true;
 				}
 				return true;
diff --git a/Engine/Source/Programs/UnrealBuildAccelerator/Detours/Private/Windows/UbaDetoursFunctionsNtDll.inl b/Engine/Source/Programs/UnrealBuildAccelerator/Detours/Private/Windows/UbaDetoursFunctionsNtDll.inl
--- a/Engine/Source/Programs/UnrealBuildAccelerator/Detours/Private/Windows/UbaDetoursFunctionsNtDll.inl
+++ b/Engine/Source/Programs/UnrealBuildAccelerator/Detours/Private/Windows/UbaDetoursFunctionsNtDll.inl
@@ -1371,6 +1371,9 @@ NTSTATUS NTAPI Detoured_NtClose(HANDLE handle)
 				StringBuffer<> fixedName;
 				FixPath(fixedName, fileName);
 				StringKey fileNameKey = fi.fileNameKey;
+				if (!fo->newName.empty())
+					fileNameKey = ToStringKeyLower(fixedName);
+
 				u64 size;
 				Rpc_CreateFileW(fixedName.data, fileNameKey, AccessFlag_Write, temp, sizeof_array(temp), size, fo->closeId, true);
 			}
