[UBA]
* Removed lots of debug code now when bug has been found

diff --git a/Engine/Source/Programs/UnrealBuildAccelerator/Common/Private/UbaCacheClient.cpp b/Engine/Source/Programs/UnrealBuildAccelerator/Common/Private/UbaCacheClient.cpp
--- a/Engine/Source/Programs/UnrealBuildAccelerator/Common/Private/UbaCacheClient.cpp
+++ b/Engine/Source/Programs/UnrealBuildAccelerator/Common/Private/UbaCacheClient.cpp
@@ -265,26 +265,6 @@ namespace uba
 				bool deferCreation = true;
 				bool fileIsCompressed = IsFileCompressed(info, path);
 
-
-
-				CasKey oldOldKey;
-				u64 oldOldSize = 0;
-				u64 oldOldLastWritten = 0;
-				u64 lastInvalidationTime = 0;
-				bool oldOldVerified = false;
-
-				if (shouldValidate)
-				{
-					auto& fileEntry = m_storage.GetOrCreateFileEntry(CaseInsensitiveFs ? ToStringKeyLower(path) : ToStringKey(path));
-					SCOPED_READ_LOCK(fileEntry.lock, lock);
-					oldOldKey = fileEntry.casKey;
-					oldOldSize = fileEntry.size;
-					oldOldLastWritten = fileEntry.lastWritten;
-					oldOldVerified = fileEntry.verified;
-					lastInvalidationTime = fileEntry.lastInvalidationTime;
-				}
-
-
 				if (isOutput)
 				{
 					if (!m_storage.StoreCasFile(casKey, path.data, CasKeyZero, deferCreation, fileIsCompressed))
@@ -340,8 +320,8 @@ namespace uba
 						SCOPED_READ_LOCK(fileEntry.lock, lock);
 
 						auto ToString = [](bool b) { return b ? TC("true") : TC("false"); };
-						m_logger.Warning(TC("CasDb claims file %s has caskey %s but recalculating it gives us %s (OldEntry: %s/%llu/%llu/%s FileEntry: %llu/%llu/%s, Real: %llu/%llu. Last invalidation: %llu). Will not populate cache for %s"),
-							path.data, CasKeyString(oldKey).str, CasKeyString(newKey).str, CasKeyString(oldOldKey).str, oldOldSize, oldOldLastWritten, ToString(oldOldVerified), fileEntry.size, fileEntry.lastWritten, ToString(fileEntry.verified), fileSize, fileInfo.lastWriteTime, lastInvalidationTime, info.description);
+						m_logger.Warning(TC("CasDb claims file %s has caskey %s but recalculating it gives us %s (FileEntry: %llu/%llu/%s, Real: %llu/%llu). Will not populate cache for %s"),
+							path.data, CasKeyString(oldKey).str, CasKeyString(newKey).str, fileEntry.size, fileEntry.lastWritten, ToString(fileEntry.verified), fileSize, fileInfo.lastWriteTime, info.description);
 						return false;
 					}
 				}
diff --git a/Engine/Source/Programs/UnrealBuildAccelerator/Common/Private/UbaStorage.cpp b/Engine/Source/Programs/UnrealBuildAccelerator/Common/Private/UbaStorage.cpp
--- a/Engine/Source/Programs/UnrealBuildAccelerator/Common/Private/UbaStorage.cpp
+++ b/Engine/Source/Programs/UnrealBuildAccelerator/Common/Private/UbaStorage.cpp
@@ -1520,12 +1520,12 @@ namespace uba
 				BinaryWriter writer(fileTableData.data(), 0, fileTableData.size());
 				for (auto& pair : m_fileTableLookup)
 				{
-					FileEntry& entry = pair.second;
-					SCOPED_READ_LOCK(entry.lock, entryLock);
+					FileEntry& fileEntry = pair.second;
+					SCOPED_READ_LOCK(fileEntry.lock, entryLock);
 					writer.WriteStringKey(pair.first);
-					writer.WriteU64(entry.size);
-					writer.WriteU64(entry.lastWritten);
-					writer.WriteCasKey(entry.casKey);
+					writer.WriteU64(fileEntry.size);
+					writer.WriteU64(fileEntry.lastWritten);
+					writer.WriteCasKey(fileEntry.casKey);
 				}
 				if (!tempFile.Write(fileTableData.data(), writer.GetPosition()))
 					return false;
@@ -1551,14 +1551,14 @@ namespace uba
 				for (CasEntry* it = m_newestAccessed; it; it = it->nextAccessed)
 				{
 					last = it;
-					CasEntry& entry = *it;
-					if (entry.verified && !entry.exists)
+					CasEntry& casEntry = *it;
+					if (casEntry.verified && !casEntry.exists)
 						continue;
-					if (entry.dropped)
+					if (casEntry.dropped)
 					{
 #if !UBA_USE_SPARSEFILE
 						StringBuffer<512> casFileName;
-						if (!StorageImpl::GetCasFileName(casFileName, entry.key))
+						if (!StorageImpl::GetCasFileName(casFileName, casEntry.key))
 							continue;
 						DeleteFileW(casFileName.data);
 #else
@@ -1568,12 +1568,12 @@ namespace uba
 					}
 
 #if UBA_USE_SPARSEFILE
-					auto findIt = handleToIndex.find(entry.mappingHandle);
+					auto findIt = handleToIndex.find(casEntry.mappingHandle);
 					if (findIt == handleToIndex.end())
 					{
-						if (m_deferredCasCreationLookup.find(entry.key) != m_deferredCasCreationLookup.end())
+						if (m_deferredCasCreationLookup.find(casEntry.key) != m_deferredCasCreationLookup.end())
 							continue;
-						m_logger.Error(TC("Can't find cas database file with mappingHandle %llu"), uintptr_t(entry.mappingHandle));
+						m_logger.Error(TC("Can't find cas database file with mappingHandle %llu"), uintptr_t(casEntry.mappingHandle));
 						continue;
 					}
 #endif
@@ -1581,14 +1581,14 @@ namespace uba
 					if (writer.GetCapacityLeft() < entrySize + sizeof(CasKey))
 						return m_logger.Error(TC("This should not happen, somehow there are more valid entries in access list than lookup. (Lookup has %llu entries)"), m_casLookup.size());
 
-					UBA_ASSERT(entry.key != CasKeyZero);
-					writer.WriteCasKey(entry.key);
-					writer.WriteU64(entry.size);
+					UBA_ASSERT(casEntry.key != CasKeyZero);
+					writer.WriteCasKey(casEntry.key);
+					writer.WriteU64(casEntry.size);
 
 #if UBA_USE_SPARSEFILE
 					writer.WriteU32(findIt->second);
-					writer.WriteU64(entry.mappingOffset);
-					writer.WriteU64(entry.mappingSize);
+					writer.WriteU64(casEntry.mappingOffset);
+					writer.WriteU64(casEntry.mappingSize);
 #endif
 				}
 				writer.WriteCasKey(CasKeyZero);
@@ -1886,7 +1886,6 @@ namespace uba
 		SCOPED_WRITE_LOCK(fileEntry.lock, entryLock);
 		fileEntry.verified = false;
 		fileEntry.casKey = CasKeyInvalid;
-		fileEntry.lastInvalidationTime = GetSystemTimeAsFileTime();
 		return true;
 	}
 
diff --git a/Engine/Source/Programs/UnrealBuildAccelerator/Common/Public/UbaCacheClient.h b/Engine/Source/Programs/UnrealBuildAccelerator/Common/Public/UbaCacheClient.h
--- a/Engine/Source/Programs/UnrealBuildAccelerator/Common/Public/UbaCacheClient.h
+++ b/Engine/Source/Programs/UnrealBuildAccelerator/Common/Public/UbaCacheClient.h
@@ -31,12 +31,12 @@ namespace uba
 
 		void Apply(Config& config);
 
-		bool reportMissReason = false;
-		bool useDirectoryPreparsing = false; // This is used to minimize syscalls. GetFileAttributes can be very expensive on cloud machines and we can enable this to minimize syscall count
+		bool reportMissReason = false; // Report the reason no matching cache entry was found
+		bool useDirectoryPreparsing = true; // This is used to minimize syscalls. GetFileAttributes can be very expensive on cloud machines and we can enable this to minimize syscall count
 		bool validateCacheWritesInput = false; // Set to true to validate cas of all input files before sent to cache
-		bool validateCacheWritesOutput = true; // Set to true to validate cas of all output files before sent to cache
-		bool useRoots = true;
-		bool useCacheHit = true;
+		bool validateCacheWritesOutput = false; // Set to true to validate cas of all output files before sent to cache
+		bool useRoots = true; // Set this to false to allow paths that are not under roots and to not fix them up
+		bool useCacheHit = true; // Set this to false to ignore found cache hits.. this is for debugging/testing only
 	};
 
 	class CacheClient
diff --git a/Engine/Source/Programs/UnrealBuildAccelerator/Common/Public/UbaStorage.h b/Engine/Source/Programs/UnrealBuildAccelerator/Common/Public/UbaStorage.h
--- a/Engine/Source/Programs/UnrealBuildAccelerator/Common/Public/UbaStorage.h
+++ b/Engine/Source/Programs/UnrealBuildAccelerator/Common/Public/UbaStorage.h
@@ -195,7 +195,6 @@ namespace uba
 			u64 size = 0;
 			u64 lastWritten = 0;
 			bool verified = false;
-			u64 lastInvalidationTime = 0; // Remove when we have figured out a bug
 		};
 		UnorderedMap<StringKey, FileEntry> m_fileTableLookup;
 
