[UBA]
* Added timer scope around compressed obj files preload

diff --git a/Engine/Source/Programs/UnrealBuildAccelerator/Core/Public/UbaProcessStats.h b/Engine/Source/Programs/UnrealBuildAccelerator/Core/Public/UbaProcessStats.h
--- a/Engine/Source/Programs/UnrealBuildAccelerator/Core/Public/UbaProcessStats.h
+++ b/Engine/Source/Programs/UnrealBuildAccelerator/Core/Public/UbaProcessStats.h
@@ -40,6 +40,9 @@ namespace uba
 		UBA_PROCESS_STAT(writeFiles, 19) \
 		UBA_PROCESS_STAT(queryCache, 24) \
 		UBA_PROCESS_STAT(waitDecompress, 30) \
+		UBA_PROCESS_STAT(preparseObjFiles, 30) \
+		UBA_PROCESS_STAT(fileTable, 30) \
+		UBA_PROCESS_STAT(dirTable, 30) \
 
 
 		#define UBA_PROCESS_STAT(T, ver) Timer T;
@@ -119,6 +122,7 @@ namespace uba
 			UBA_KERNEL_STAT(ExtendedTimer, getFileTime, 0) \
 			UBA_KERNEL_STAT(ExtendedTimer, closeHandle, 0) \
 			UBA_KERNEL_STAT(ExtendedTimer, traverseDir, 27) \
+			UBA_KERNEL_STAT(ExtendedTimer, virtualAlloc, 30) \
 
 		#define UBA_KERNEL_STAT(type, var, ver) type var;
 		UBA_KERNEL_STATS
diff --git a/Engine/Source/Programs/UnrealBuildAccelerator/Detours/Private/Windows/UbaDetoursFunctionsWin.cpp b/Engine/Source/Programs/UnrealBuildAccelerator/Detours/Private/Windows/UbaDetoursFunctionsWin.cpp
--- a/Engine/Source/Programs/UnrealBuildAccelerator/Detours/Private/Windows/UbaDetoursFunctionsWin.cpp
+++ b/Engine/Source/Programs/UnrealBuildAccelerator/Detours/Private/Windows/UbaDetoursFunctionsWin.cpp
@@ -878,9 +878,12 @@ void Init(const DetoursPayload& payload, u64 startTime)
 	{
 		TimerScope ts(g_kernelStats.mapViewOfFile);
 		mappedFileTableMem = (u8*)True_MapViewOfFile(mappedFileTableHandle, FILE_MAP_READ, 0, 0, 0);
+		UBA_ASSERT(mappedFileTableMem);
+	}
+	{
+		TimerScope ts2(g_stats.fileTable);
+		g_mappedFileTable.Init(mappedFileTableMem, mappedFileTableCount, mappedFileTableSize);
 	}
-	UBA_ASSERT(mappedFileTableMem);
-	g_mappedFileTable.Init(mappedFileTableMem, mappedFileTableCount, mappedFileTableSize);
 
 	if (!True_DuplicateHandle(g_hostProcess, directoryTableHandle, GetCurrentProcess(), &directoryTableHandle, 0, FALSE, DUPLICATE_SAME_ACCESS))
 		UBA_ASSERTF(false, L"Failed to duplicate directorytable handle (%u)", GetLastError());
@@ -889,17 +892,23 @@ void Init(const DetoursPayload& payload, u64 startTime)
 	{
 		TimerScope ts(g_kernelStats.mapViewOfFile);
 		directoryTableMem = (u8*)True_MapViewOfFile(directoryTableHandle, FILE_MAP_READ, 0, 0, 0);
+		UBA_ASSERT(directoryTableMem);
+	}
+	{
+		TimerScope ts2(g_stats.dirTable);
+		g_directoryTable.Init(directoryTableMem, directoryTableCount, directoryTableSize);
 	}
-	UBA_ASSERT(directoryTableMem);
-	g_directoryTable.Init(directoryTableMem, directoryTableCount, directoryTableSize);
+
+	g_stats.attach.time += GetTime() - startTime;
+	g_stats.attach.count = 1;
 
 	if (payload.storeObjFilesCompressed && g_rules->ShouldDecompressFiles(StringView()))
+	{
+		TimerScope ts(g_stats.preparseObjFiles);
 		g_objFilesPreloader.Start(cmdLine);
+	}
 	else if (g_rulesIndex == 1 || g_rulesIndex == 7 || g_rulesIndex == 11 || g_rulesIndex == 14)
 		PrepopulatePchIncludedFiles(cmdLine, g_rulesIndex);
-
-	g_stats.attach.time += GetTime() - startTime;
-	g_stats.attach.count = 1;
 }
 
 void Deinit(u64 startTime)
