[UBA]
* Fixed so visualizer can visualize traces coming from UbaCli without needing restart
* Added some more logging to cache server maintenance
* Fixed scroll wheel zoom in visualizer
* Fixed so network server DisconnectClients does not put server in a bad state

diff --git a/Engine/Source/Programs/UnrealBuildAccelerator/Cli/Private/UbaCli.cpp b/Engine/Source/Programs/UnrealBuildAccelerator/Cli/Private/UbaCli.cpp
--- a/Engine/Source/Programs/UnrealBuildAccelerator/Cli/Private/UbaCli.cpp
+++ b/Engine/Source/Programs/UnrealBuildAccelerator/Cli/Private/UbaCli.cpp
@@ -619,6 +619,9 @@ namespace uba
 				return -1;
 			}
 
+			if (!storageServer.LoadCasTable(true))
+				return -1;
+
 			if (!cacheCommand.empty())
 			{
 				LoggerWithWriter consoleLogger(g_consoleLogWriter);
@@ -656,8 +659,9 @@ namespace uba
 
 		if (isRemote)
 		{
-			if (!storageServer.LoadCasTable(true))
-				return -1;
+			if (!storageServer.m_casTableLoaded)
+				if (!storageServer.LoadCasTable(true))
+					return -1;
 			if (!networkServer.StartListen(networkBackend, port, listenIp.data))
 				return -1;
 		}
@@ -665,7 +669,7 @@ namespace uba
 
 		auto stopListen = MakeGuard([&]() { networkBackend.StopListen(); });
 
-		auto RunLocal = [&](const TString& app, const TString& arg, bool enableDetour, bool trackInputs = false)
+		auto RunLocal = [&](const TString& app, const TString& arg, bool enableDetour)
 		{
 			u64 start = GetTime();
 			ProcessStartInfo pinfo;
@@ -673,17 +677,31 @@ namespace uba
 			pinfo.application = app.c_str();
 			pinfo.arguments = arg.c_str();
 			pinfo.workingDir = workDir.data;
+
+			u32 bucketId = 1337;
+			if (cacheClient)
+			{
+				bool hit = false;
+				cacheClient->FetchFromCache(hit, RootPaths(), bucketId, pinfo);
+				if (hit)
+				{
+					logger.Info(TC("%s run took %s [cached]"), (enableDetour ? TC("Boxed") : TC("Native")), TimeToText(GetTime() - start).str);
+					return true;
+				}
+			}
+
+
 			pinfo.logFile = logFile.data;
 			pinfo.logLineUserData = &logger;
 			if (enableStdOut)
 				pinfo.logLineFunc = [](void* userData, const tchar* line, u32 length, LogEntryType type) { ((Logger*)userData)->Log(type, line, length); };
-			pinfo.trackInputs = trackInputs;
+			if (populateCache)
+				pinfo.trackInputs = true;
 			logger.Info(TC("Running %s %s"), app.c_str(), arg.c_str());
 			ProcessHandle process = sessionServer.RunProcess(pinfo, false, enableDetour);
 			if (process.GetExitCode() != 0)
 				return logger.Error(TC("Error exit code: %u"), process.GetExitCode());
-			u64 time = GetTime() - start;
-			logger.Info(TC("%s run took %s"), (enableDetour ? TC("Boxed") : TC("Native")), TimeToText(time).str);
+			logger.Info(TC("%s run took %s"), (enableDetour ? TC("Boxed") : TC("Native")), TimeToText(GetTime() - start).str);
 
 			if (populateCache)
 			{
@@ -716,11 +734,10 @@ namespace uba
 			return true;
 		};
 
-		auto RunWithClient = [&](const Function<bool()>& func)
+		auto RunWithClient = [&](const Function<bool()>& func, u32 clientCount)
 			{
 				Vector<Client> clients;
-				auto slg = MakeGuard([&]() { networkBackend.StopListen(); });
-				clients.resize(maxProcessCount == 1 ? 1 : 4);
+				clients.resize(clientCount);
 				u32 clientIndex = 0;
 				for (auto& c : clients)
 				{
@@ -733,7 +750,7 @@ namespace uba
 
 		auto RunAgent = [&](const TString& app, const TString& arg)
 		{
-			return RunWithClient([&]() { return RunRemote(app, arg); });
+			return RunWithClient([&]() { return RunRemote(app, arg); }, 1);
 		};
 
 		CoordinatorWrapper coordinator;
@@ -828,7 +845,7 @@ namespace uba
 				};
 
 			if (commandType == CommandType_Agent)
-				return RunWithClient([&]() { return RunQueue(); });
+				return RunWithClient([&]() { return RunQueue(); }, maxProcessCount == 1 ? 1 : 4);
 			else
 				return RunQueue();
 		};
@@ -848,7 +865,7 @@ namespace uba
 
 			// TODO: This is very horde specific.. maybe all these parameters should be a string or something
 			cinfo.pool = coordinatorPool.c_str();
-			cinfo.maxCoreCount = 400;
+			cinfo.maxCoreCount = coordinatorMaxCoreCount;
 			cinfo.logging = true;
 			if (!coordinator.Create(logger, coordinatorName.c_str(), cinfo, networkBackend, networkServer))
 				return false;
@@ -883,6 +900,10 @@ namespace uba
 			}
 			if (!success)
 				return -1;
+
+			if (false)
+				networkServer.DisconnectClients();
+
 		}
 
 		logger.BeginScope();
diff --git a/Engine/Source/Programs/UnrealBuildAccelerator/Common/Private/UbaCacheServer.cpp b/Engine/Source/Programs/UnrealBuildAccelerator/Common/Private/UbaCacheServer.cpp
--- a/Engine/Source/Programs/UnrealBuildAccelerator/Common/Private/UbaCacheServer.cpp
+++ b/Engine/Source/Programs/UnrealBuildAccelerator/Common/Private/UbaCacheServer.cpp
@@ -269,6 +269,7 @@ namespace uba
 				}
 				else
 				{
+					written += tempBufferPos;
 					success &= file.Write(tempBuffer, tempBufferPos);
 					tempBufferPos = 0;
 				}
@@ -286,6 +287,7 @@ namespace uba
 		bool Close()
 		{
 			success &= file.Write(tempBuffer, tempBufferPos);
+			written += tempBufferPos;
 
 			if (!success)
 				return false;
@@ -303,6 +305,7 @@ namespace uba
 		bool success = true;
 		u8* tempBuffer = nullptr;
 		u64 tempBufferPos = 0;
+		u64 written = 0;
 		TString fileName;
 		TString tempFileName;
 		FileAccessor file;
@@ -356,7 +359,7 @@ namespace uba
 
 		bucket.lastSavedTime = GetSystemTimeAsFileTime() - m_creationTime;
 
-		m_logger.Detail(TC("    Bucket %u saved (%s)"), bucket.index, TimeToText(GetTime() - saveStart).str);
+		m_logger.Detail(TC("    Bucket %u saved - %s (%s)"), bucket.index, BytesToText(file.written).str, TimeToText(GetTime() - saveStart).str);
 		return true;
 	}
 
@@ -507,6 +510,7 @@ namespace uba
 
 		u64 now = GetSystemTimeAsFileTime();
 		u64 oldest = 0;
+		u64 longestUnused = 0;
 
 		u64 lastUseTimeLimit = 0;
 		if (m_expirationTimeSeconds && GetFileTimeAsSeconds(now - m_creationTime) > m_expirationTimeSeconds)
@@ -533,6 +537,7 @@ namespace uba
 			bool checkInputsForDeletes = m_checkInputsForDeletedCas && !deletedCasFiles.empty();
 
 			oldest = 0;
+			longestUnused = 0;
 			totalEntryCount = 0;
 
 			m_server.ParallelFor(workerCountToUseForBuckets, m_buckets, [&](auto& it)
@@ -671,6 +676,8 @@ namespace uba
 						SCOPED_WRITE_LOCK(existingCasLock, l);
 						if (!oldest || entry.creationTime < oldest)
 							oldest = entry.creationTime;
+						if (!longestUnused || entry.lastUsedTime < longestUnused)
+							longestUnused = entry.lastUsedTime;
 
 						for (u64* v : touchedCas)
 							++(*v);
@@ -893,8 +900,9 @@ namespace uba
 		}
 
 		u64 oldestTime = oldest ? GetFileTimeAsTime(now - (m_creationTime + oldest)) : 0;
+		u64 longestUnusedTime = longestUnused ? GetFileTimeAsTime(now - (m_creationTime + longestUnused)) : 0;
 		u64 duration = GetTime() - startTime;
-		m_logger.Info(TC("Maintenance done! (%s) CasFiles: %llu (%s) Entries: %llu OldestEntry: %s"), TimeToText(duration).str, totalCasCount - deletedCasCount, BytesToText(totalCasSize).str, totalEntryCount.load(), TimeToText(oldestTime, true).str);
+		m_logger.Info(TC("Maintenance done! (%s) CasFiles: %llu (%s) Entries: %llu Oldest: %s LongestUnused: %s"), TimeToText(duration).str, totalCasCount - deletedCasCount, BytesToText(totalCasSize).str, totalEntryCount.load(), TimeToText(oldestTime, true).str, TimeToText(longestUnusedTime, true).str);
 		
 		m_longestMaintenance = Max(m_longestMaintenance, duration);
 
diff --git a/Engine/Source/Programs/UnrealBuildAccelerator/Common/Private/UbaNetworkServer.cpp b/Engine/Source/Programs/UnrealBuildAccelerator/Common/Private/UbaNetworkServer.cpp
--- a/Engine/Source/Programs/UnrealBuildAccelerator/Common/Private/UbaNetworkServer.cpp
+++ b/Engine/Source/Programs/UnrealBuildAccelerator/Common/Private/UbaNetworkServer.cpp
@@ -661,6 +661,8 @@ namespace uba
 
 		SCOPED_WRITE_LOCK(m_connectionsLock, lock);
 		m_connections.clear();
+
+		m_workersEnabled = true;
 	}
 
 	bool NetworkServer::AddClient(NetworkBackend& backend, const tchar* ip, u16 port, const u8* cryptoKey128)
@@ -710,8 +712,9 @@ namespace uba
 		if (!m_maxActiveConnections)
 			return;
 
+		m_maxCreatedWorkerCount = Max(m_createdWorkerCount, m_maxCreatedWorkerCount);
 		StringBuffer<> workers;
-		workers.Appendf(TC("%u/%u"), m_createdWorkerCount, m_maxWorkerCount);
+		workers.Appendf(TC("%u/%u"), m_maxCreatedWorkerCount, m_maxWorkerCount);
 
 		logger.Info(TC("  ----- Uba server stats summary ------"));
 		logger.Info(TC("  MaxActiveConnections           %6u"), m_maxActiveConnections);
@@ -1067,6 +1070,8 @@ namespace uba
 			delete temp;
 		}
 		m_firstAvailableWorker = nullptr;
+		m_maxCreatedWorkerCount = Max(m_createdWorkerCount, m_maxCreatedWorkerCount);
+		m_createdWorkerCount = 0;
 	}
 
 	void NetworkServer::RemoveDisconnectedConnections()
diff --git a/Engine/Source/Programs/UnrealBuildAccelerator/Common/Private/UbaTrace.cpp b/Engine/Source/Programs/UnrealBuildAccelerator/Common/Private/UbaTrace.cpp
--- a/Engine/Source/Programs/UnrealBuildAccelerator/Common/Private/UbaTrace.cpp
+++ b/Engine/Source/Programs/UnrealBuildAccelerator/Common/Private/UbaTrace.cpp
@@ -84,8 +84,10 @@ namespace uba
 
 
 		if (namedTrace && m_channel.Init())
+		{
+			m_namedTrace = namedTrace;
 			m_channel.Write(namedTrace);
-
+		}
 		return true;
 	}
 
@@ -95,6 +97,9 @@ namespace uba
 			return true;
 		auto g = MakeGuard([this]() { FreeMemory(); });
 
+		if (!m_namedTrace.empty())
+			m_channel.Write(TC(""), m_namedTrace.c_str());
+
 		{
 			WriterScope writer(*this);
 			if (!writer.IsValid())
@@ -451,10 +456,13 @@ namespace uba
 		#endif
 	}
 
-	bool TraceChannel::Write(const tchar* traceName)
+	bool TraceChannel::Write(const tchar* traceName, const tchar* ifMatching)
 	{
 		#if PLATFORM_WINDOWS
 		WaitForSingleObject((HANDLE)m_mutex, INFINITE);
+		if (ifMatching)
+			if (!Equals((tchar*)m_mem, ifMatching))
+				return true;
 		TStrcpy_s((tchar*)m_mem, 256, traceName);
 		ReleaseMutex((HANDLE)m_mutex);
 		#endif
diff --git a/Engine/Source/Programs/UnrealBuildAccelerator/Common/Public/UbaNetworkServer.h b/Engine/Source/Programs/UnrealBuildAccelerator/Common/Public/UbaNetworkServer.h
--- a/Engine/Source/Programs/UnrealBuildAccelerator/Common/Public/UbaNetworkServer.h
+++ b/Engine/Source/Programs/UnrealBuildAccelerator/Common/Public/UbaNetworkServer.h
@@ -137,6 +137,7 @@ namespace uba
 		u32 m_sendSize = 0;
 		u32 m_receiveTimeoutMs = 0;
 		u32 m_createdWorkerCount = 0;
+		u32 m_maxCreatedWorkerCount = 0;
 		bool m_workersEnabled = true;
 
 		ReaderWriterLock m_addConnectionsLock;
diff --git a/Engine/Source/Programs/UnrealBuildAccelerator/Common/Public/UbaTrace.h b/Engine/Source/Programs/UnrealBuildAccelerator/Common/Public/UbaTrace.h
--- a/Engine/Source/Programs/UnrealBuildAccelerator/Common/Public/UbaTrace.h
+++ b/Engine/Source/Programs/UnrealBuildAccelerator/Common/Public/UbaTrace.h
@@ -19,7 +19,7 @@ namespace uba
 		~TraceChannel();
 
 		bool Init(const tchar* channelName = TC("Default"));
-		bool Write(const tchar* traceName);
+		bool Write(const tchar* traceName, const tchar* ifMatching = nullptr);
 		bool Read(StringBufferBase& outTraceName);
 
 		Logger& m_logger;
@@ -114,6 +114,7 @@ namespace uba
 		u32 AddString(const tchar* string);
 
 		LoggerWithWriter m_logger;
+		TString m_namedTrace;
 		TraceChannel m_channel;
 		ReaderWriterLock m_memoryLock;
 		FileMappingHandle m_memoryHandle;
diff --git a/Engine/Source/Programs/UnrealBuildAccelerator/Visualizer/Private/UbaTraceReader.cpp b/Engine/Source/Programs/UnrealBuildAccelerator/Visualizer/Private/UbaTraceReader.cpp
--- a/Engine/Source/Programs/UnrealBuildAccelerator/Visualizer/Private/UbaTraceReader.cpp
+++ b/Engine/Source/Programs/UnrealBuildAccelerator/Visualizer/Private/UbaTraceReader.cpp
@@ -16,12 +16,7 @@ namespace uba
 
 	TraceReader::~TraceReader()
 	{
-		if (m_hostProcess)
-			CloseHandle(m_hostProcess);
-		if (m_memoryBegin)
-			UnmapViewOfFile(m_memoryBegin, 0, TC("TraceReader"));
-		if (m_memoryHandle.IsValid())
-			CloseFileMapping(m_memoryHandle);
+		Unmap();
 	}
 
 #if PLATFORM_WINDOWS
@@ -213,6 +208,8 @@ namespace uba
 		bool res = ReadMemory(out, true);
 		if (m_hostProcess  && WaitForSingleObject(m_hostProcess, 0) != WAIT_TIMEOUT)
 			StopAllActive(out, GetTime() - m_startTime);
+		if (!res)
+			Unmap();
 		return res;
 	}
 
@@ -817,7 +814,7 @@ namespace uba
 			if (out.version < 15)
 				time = reader.Read7BitEncoded();
 			StopAllActive(out, time);
-			break;
+			return false;
 		}
 		case TraceType_BeginWork:
 		{
@@ -974,6 +971,19 @@ namespace uba
 		m_sessionIndexToSession.clear();
 	}
 
+	void TraceReader::Unmap()
+	{
+		if (m_hostProcess)
+			CloseHandle(m_hostProcess);
+		m_hostProcess = 0;
+		if (m_memoryBegin)
+			UnmapViewOfFile(m_memoryBegin, 0, TC("TraceReader"));
+		m_memoryBegin = nullptr;
+		if (m_memoryHandle.IsValid())
+			CloseFileMapping(m_memoryHandle);
+		m_memoryHandle = {};
+	}
+
 	bool TraceReader::SaveAs(const tchar* fileName)
 	{
 		if (!m_memoryBegin)
diff --git a/Engine/Source/Programs/UnrealBuildAccelerator/Visualizer/Private/UbaTraceReader.h b/Engine/Source/Programs/UnrealBuildAccelerator/Visualizer/Private/UbaTraceReader.h
--- a/Engine/Source/Programs/UnrealBuildAccelerator/Visualizer/Private/UbaTraceReader.h
+++ b/Engine/Source/Programs/UnrealBuildAccelerator/Visualizer/Private/UbaTraceReader.h
@@ -169,6 +169,7 @@ namespace uba
 		bool ReadTrace(TraceView& out, BinaryReader& reader, u64 maxTime);
 		void StopAllActive(TraceView& out, u64 stopTime);
 		void Reset();
+		void Unmap();
 
 		bool SaveAs(const tchar* fileName);
 
diff --git a/Engine/Source/Programs/UnrealBuildAccelerator/Visualizer/Private/UbaVisualizer.cpp b/Engine/Source/Programs/UnrealBuildAccelerator/Visualizer/Private/UbaVisualizer.cpp
--- a/Engine/Source/Programs/UnrealBuildAccelerator/Visualizer/Private/UbaVisualizer.cpp
+++ b/Engine/Source/Programs/UnrealBuildAccelerator/Visualizer/Private/UbaVisualizer.cpp
@@ -172,14 +172,20 @@ namespace uba
 			return false;
 		}
 
+		m_listenTimeout.Create(false);
+
 		m_listenChannel.Append(channelName);
 		m_looping = true;
 		m_autoScroll = false;
-		m_thread.Start([this]() { ThreadLoop(); return 0;});
+		if (!StartHwndThread())
+			return true;
 
-		while (!m_hwnd)
-			if (m_thread.Wait(10))
-				return true;
+		{
+			StringBuffer<> title;
+			GetTitlePrefix(title);
+			title.Appendf(L"Listening for new sessions on channel '%s'", m_listenChannel.data);
+			SetWindowTextW(m_hwnd, title.data);
+		}
 
 		StringBuffer<256> traceName;
 		while (m_hwnd)
@@ -191,12 +197,18 @@ namespace uba
 				return false;
 			}
 
-			if (traceName.count && !traceName.Equals(m_newTraceName.data))
+			if (traceName.count)
 			{
-				m_newTraceName.Clear().Append(traceName);
-				PostMessage(m_hwnd, WM_NEWTRACE, 0, 0);
+				if (!traceName.Equals(m_newTraceName.data))
+				{
+					m_newTraceName.Clear().Append(traceName);
+					PostMessage(m_hwnd, WM_NEWTRACE, 0, 0);
+				}
 			}
-			Sleep(1000);
+			else
+				m_newTraceName.Clear();
+
+			m_listenTimeout.IsSet(1000);
 		}
 
 		return true;
@@ -204,11 +216,11 @@ namespace uba
 
 	bool Visualizer::ShowUsingNamedTrace(const wchar_t* namedTrace)
 	{
-		if (!m_trace.StartReadNamed(m_traceView, namedTrace))
-			return false;
-		m_namedTrace.Append(namedTrace);
 		m_looping = true;
-		m_thread.Start([this]() { ThreadLoop(); return 0;});
+		if (!StartHwndThread())
+			return true;
+		m_newTraceName.Append(namedTrace);
+		PostMessage(m_hwnd, WM_NEWTRACE, 0, 0);
 		return true;
 	}
 
@@ -217,11 +229,8 @@ namespace uba
 		auto destroyClient = MakeGuard([&]() { delete m_client; m_client = nullptr; });
 		m_looping = true;
 		m_autoScroll = false;
-		m_thread.Start([this]() { ThreadLoop(); return 0; });
-
-		while (!m_hwnd)
-			if (m_thread.Wait(10))
-				return true;
+		if (!StartHwndThread())
+			return true;
 
 		wchar_t dots[] = TC("....");
 		u32 dotsCounter = 0;
@@ -262,17 +271,23 @@ namespace uba
 	{
 		m_looping = true;
 		m_autoScroll = false;
-		m_thread.Start([this]() { ThreadLoop(); return 0;});
-
-		while (!m_hwnd)
-			if (m_thread.Wait(10))
-				return true;
+		if (!StartHwndThread())
+			return true;
 		m_fileName.Append(fileName);
 		m_replay = replay;
 		PostMessage(m_hwnd, WM_NEWTRACE, 0, 0);
 		return true;
 	}
 
+	bool Visualizer::StartHwndThread()
+	{
+		m_thread.Start([this]() { ThreadLoop(); return 0;});
+		while (!m_hwnd)
+			if (m_thread.Wait(10))
+				return false;
+		return true;
+	}
+
 	bool Visualizer::HasWindow()
 	{
 		return m_looping == true;
@@ -460,14 +475,7 @@ namespace uba
 
 		StringBuffer<> title;
 		GetTitlePrefix(title);
-		if (!m_namedTrace.IsEmpty())
-			title.Append(m_namedTrace);
-		else if (!m_fileName.IsEmpty())
-			title.Append(m_fileName);
-		else if (m_listenChannel.count)
-			title.Appendf(L"Listening for new sessions on channel '%s'", m_listenChannel.data);
-		else
-			title.Append(L"Socket");
+		title.Append(L"Initializing...");
 
 		m_hwnd = CreateWindowEx(0, windowClassName, title.data, windowStyle, winPosX, winPosY, winWidth, winHeight, NULL, NULL, hInstance, this);
 		SetWindowLongPtr(m_hwnd, GWLP_USERDATA, (LONG_PTR)this);
@@ -517,6 +525,7 @@ namespace uba
 					UnregisterClass(windowClassName, hInstance);
 					m_hwnd = 0;
 					m_looping = false;
+					m_listenTimeout.Set();
 					break;
 				}
 			}
@@ -2339,8 +2348,7 @@ namespace uba
 			}
 			else if (!m_fileName.IsEmpty())
 			{
-				if (!m_trace.ReadFile(m_traceView, m_fileName.data, m_replay != 0))
-					return false;
+				m_trace.ReadFile(m_traceView, m_fileName.data, m_replay != 0);
 				RedrawWindow(hWnd, NULL, NULL, RDW_INVALIDATE | RDW_UPDATENOW);
 				title.Append(m_fileName);
 				m_traceView.finished = m_replay == 0;
@@ -2405,7 +2413,10 @@ namespace uba
 			{
 				u64 timeOffset = (GetTime() - m_startTime - m_pauseTime) * m_replay;
 				if (!m_namedTrace.IsEmpty())
-					m_trace.UpdateReadNamed(m_traceView, changed);
+				{
+					if (!m_trace.UpdateReadNamed(m_traceView, changed))
+						m_listenTimeout.Set();
+				}
 				else if (!m_fileName.IsEmpty() && m_replay)
 					m_trace.UpdateReadFile(m_traceView, timeOffset, changed);
 				else if (m_client)
@@ -2448,7 +2459,7 @@ namespace uba
 			POINT cursorPos = {};
 			GetCursorPos(&cursorPos);
 			ScreenToClient(m_hwnd, &cursorPos);
-			const float scrollAnchorOffsetX = float(cursorPos.x);
+			const float scrollAnchorOffsetX = float(cursorPos.x) - ProgressRectLeft;
 
 			SHORT controlState = GetAsyncKeyState(VK_CONTROL);
 			if (controlState & (1<<15))
diff --git a/Engine/Source/Programs/UnrealBuildAccelerator/Visualizer/Private/UbaVisualizer.h b/Engine/Source/Programs/UnrealBuildAccelerator/Visualizer/Private/UbaVisualizer.h
--- a/Engine/Source/Programs/UnrealBuildAccelerator/Visualizer/Private/UbaVisualizer.h
+++ b/Engine/Source/Programs/UnrealBuildAccelerator/Visualizer/Private/UbaVisualizer.h
@@ -25,6 +25,7 @@ namespace uba
 		HWND GetHwnd();
 
 	private:
+		bool StartHwndThread();
 		bool Unselect();
 		void Reset();
 		void PaintClient(const Function<void(HDC hdc, HDC memDC, RECT& clientRect)>& paintFunc);
@@ -125,6 +126,7 @@ namespace uba
 		
 		StringBuffer<256>m_listenChannel;
 		StringBuffer<256> m_newTraceName;
+		Event m_listenTimeout;
 
 		int m_contentWidth = 0;
 		int m_contentHeight = 0;
