[UBA]
* Fixed openssl crypto paths.. was not taking block size into account for AES

diff --git a/Engine/Source/Programs/UnrealBuildAccelerator/Common/Private/UbaCrypto.cpp b/Engine/Source/Programs/UnrealBuildAccelerator/Common/Private/UbaCrypto.cpp
--- a/Engine/Source/Programs/UnrealBuildAccelerator/Common/Private/UbaCrypto.cpp
+++ b/Engine/Source/Programs/UnrealBuildAccelerator/Common/Private/UbaCrypto.cpp
@@ -2,6 +2,7 @@
 
 #include "UbaCrypto.h"
 #include "UbaLogger.h"
+#include "UbaMemory.h"
 #include "UbaPlatform.h"
 #include "UbaSynchronization.h"
 
@@ -139,6 +140,22 @@ namespace uba
 		}
 		return true;
 	}
+#elif UBA_CRYPTO_TYPE == 2
+	bool OpenSslEncryptDecrypt(CryptoKey key, u8* data, u32 size, int enc)
+	{
+		auto& keyData = *(KeyData*)uintptr_t(key);
+		u8 iv[AES_BLOCK_SIZE];
+		memcpy(iv, (u8*)uintptr_t(key), AES_BLOCK_SIZE);
+
+		u32 alignedSize = size & 0xfffffff0;
+		if (alignedSize)
+			AES_cbc_encrypt(data, data, alignedSize, &(enc ? keyData.encryptKey : keyData.decryptKey), iv, enc);
+
+		u8* data2 = data + alignedSize;
+		for (u32 i=0, left=u32(size-alignedSize);i!=left; ++i)
+			data2[i] = data2[i] ^ iv[i];
+		return true;
+	}
 #endif
 
 	bool Crypto::Encrypt(Logger& logger, CryptoKey key, u8* data, u32 size)
@@ -146,11 +163,7 @@ namespace uba
 #if UBA_CRYPTO_TYPE == 1
 		return BCryptEncryptDecrypt(logger, true, key, data, size);
 #elif UBA_CRYPTO_TYPE == 2
-		auto& keyData = *(KeyData*)uintptr_t(key);
-		unsigned char iv[AES_BLOCK_SIZE];
-		memset(iv, 0x00, AES_BLOCK_SIZE);
-		AES_cbc_encrypt(data, data, size, &keyData.encryptKey, iv, AES_ENCRYPT);
-		return true;
+		return OpenSslEncryptDecrypt(key, data, size, AES_ENCRYPT);
 #else
 		return false;
 #endif
@@ -161,14 +174,19 @@ namespace uba
 #if UBA_CRYPTO_TYPE == 1
 		return BCryptEncryptDecrypt(logger, false, key, data, size);
 #elif UBA_CRYPTO_TYPE == 2
-		auto& keyData = *(KeyData*)uintptr_t(key);
-		unsigned char iv[AES_BLOCK_SIZE];
-		memset(iv, 0x00, AES_BLOCK_SIZE);
-		AES_cbc_encrypt(data, data, size, &keyData.decryptKey, iv, AES_DECRYPT);
-		return true;
+		return OpenSslEncryptDecrypt(key, data, size, AES_DECRYPT);
 #else
 		return false;
 #endif
 	}
+	
+	bool CryptoFromString(u8* out, u32 outSize, const tchar* str)
+	{
+		if (!str || !*str || outSize != 16 || TStrlen(str) != 32)
+			return false;
+		((u64*)out)[0] = StringToValue(str, 16);
+		((u64*)out)[1] = StringToValue(str + 16, 16);
+		return true;
+	}
 
 }
\ No newline at end of file
diff --git a/Engine/Source/Programs/UnrealBuildAccelerator/Common/Public/UbaCrypto.h b/Engine/Source/Programs/UnrealBuildAccelerator/Common/Public/UbaCrypto.h
--- a/Engine/Source/Programs/UnrealBuildAccelerator/Common/Public/UbaCrypto.h
+++ b/Engine/Source/Programs/UnrealBuildAccelerator/Common/Public/UbaCrypto.h
@@ -22,5 +22,5 @@ namespace uba
 		static bool Decrypt(Logger& logger, CryptoKey key, u8* data, u32 size);
 	};
 
-
+	bool CryptoFromString(u8* out, u32 outSize, const tchar* str);
 }
diff --git a/Engine/Source/Programs/UnrealBuildAccelerator/Test/Private/UbaTestCrypto.h b/Engine/Source/Programs/UnrealBuildAccelerator/Test/Private/UbaTestCrypto.h
--- a/Engine/Source/Programs/UnrealBuildAccelerator/Test/Private/UbaTestCrypto.h
+++ b/Engine/Source/Programs/UnrealBuildAccelerator/Test/Private/UbaTestCrypto.h
@@ -11,34 +11,42 @@ namespace uba
 	{
 		u64 key128[] = { 0x1234567812345678llu, 0x1234567812345678llu };
 
-		CryptoKey key = Crypto::CreateKey(logger, (const u8*)key128);
+		CryptoKey encryptKey = Crypto::CreateKey(logger, (const u8*)key128);
 
-		u8 encryptedData[sizeof(EncryptionHandshakeString)];
-		memcpy(encryptedData, EncryptionHandshakeString, sizeof(EncryptionHandshakeString));
-
-		for (u32 i=0; i!=3; ++i)
+		for (u32 dataSize=1;dataSize!=135; ++dataSize)
 		{
-			if (!Crypto::Encrypt(logger, key, encryptedData, sizeof(encryptedData)))
-				return false;
-
-			if (memcmp(encryptedData, EncryptionHandshakeString, sizeof(EncryptionHandshakeString)) == 0)
-				return false;
+			CryptoKey decryptKey = encryptKey;
 
-			if (!Crypto::Decrypt(logger, key, encryptedData, sizeof(encryptedData)))
-				return false;
+			u8* originalData = new u8[dataSize];
+			for (u32 i=0;i!=dataSize; ++i)
+				originalData[i] = u8(rand() % 256);
 
-			if (memcmp(encryptedData, EncryptionHandshakeString, sizeof(EncryptionHandshakeString)) != 0)
-				return false;
+			u8* encryptedData = new u8[dataSize];
+			memcpy(encryptedData, originalData, dataSize);
 
-			if (i == 1)
+			for (u32 i=0; i!=3; ++i)
 			{
-				CryptoKey newKey = Crypto::DuplicateKey(logger, key);
-				Crypto::DestroyKey(key);
-				key = newKey;
+				if (!Crypto::Encrypt(logger, encryptKey, encryptedData, dataSize))
+					return false;
+
+				if (dataSize > 16 && memcmp(encryptedData, originalData, dataSize) == 0)
+					return false;
+
+				if (!Crypto::Decrypt(logger, decryptKey, encryptedData, dataSize))
+					return false;
+
+				if (memcmp(encryptedData, originalData, dataSize) != 0)
+					return false;
+
+				if (i == 1)
+					decryptKey = Crypto::DuplicateKey(logger, encryptKey);
 			}
+			delete[] encryptedData;
+			delete[] originalData;
+			Crypto::DestroyKey(decryptKey);
 		}
 
-		Crypto::DestroyKey(key);
+		Crypto::DestroyKey(encryptKey);
 		return true;
 	}
 }
\ No newline at end of file
