[UBA]
* Moved code that writes .isNotAWS file to earlier to prevent long startup hangs

diff --git a/Engine/Source/Programs/UnrealBuildAccelerator/Common/Public/UbaAWS.h b/Engine/Source/Programs/UnrealBuildAccelerator/Common/Public/UbaAWS.h
--- a/Engine/Source/Programs/UnrealBuildAccelerator/Common/Public/UbaAWS.h
+++ b/Engine/Source/Programs/UnrealBuildAccelerator/Common/Public/UbaAWS.h
@@ -37,7 +37,7 @@ namespace uba
 			StringBuffer<128> token;
 			token.Append(TC("X-aws-ec2-metadata-token: "));
 			if (!http.Query(logger, "PUT", token, statusCode, g_imdsHost, g_imdsToken, "X-aws-ec2-metadata-token-ttl-seconds: 21600\r\n"))
-				return false;
+				return WriteIsNotAws(logger, rootDir);
 			token.Append("\r\n");
 
 			#if PLATFORM_WINDOWS
@@ -51,7 +51,7 @@ namespace uba
 
 			StringBuffer<128> instanceId;
 			if (!http.Query(logger, "GET", instanceId, statusCode, g_imdsHost, g_imdsInstanceId, m_tokenString.c_str()))
-				return WriteIsNotAws(logger, rootDir);
+				return false;
 
 			outExtraInfo.Append(TC(", AWS: ")).Append(instanceId);
 
