[UBA]
* Fixed so mode is propagated when closing file for write that failed to use rename
* Changed default file mode to mode  include S_IRGRP and S_IROTH

diff --git a/Engine/Source/Programs/UnrealBuildAccelerator/Common/Private/UbaFile.cpp b/Engine/Source/Programs/UnrealBuildAccelerator/Common/Private/UbaFile.cpp
--- a/Engine/Source/Programs/UnrealBuildAccelerator/Common/Private/UbaFile.cpp
+++ b/Engine/Source/Programs/UnrealBuildAccelerator/Common/Private/UbaFile.cpp
@@ -703,7 +703,7 @@ namespace uba
 #if PLATFORM_WINDOWS
 		return FILE_ATTRIBUTE_NORMAL;
 #else
-		return S_IRUSR | S_IWUSR | (execute ? S_IXUSR : 0);
+		return S_IRUSR | S_IWUSR | (execute ? S_IXUSR : 0) | S_IRGRP | S_IROTH;
 #endif
 	}
 
diff --git a/Engine/Source/Programs/UnrealBuildAccelerator/Common/Private/UbaFileAccessor.cpp b/Engine/Source/Programs/UnrealBuildAccelerator/Common/Private/UbaFileAccessor.cpp
--- a/Engine/Source/Programs/UnrealBuildAccelerator/Common/Private/UbaFileAccessor.cpp
+++ b/Engine/Source/Programs/UnrealBuildAccelerator/Common/Private/UbaFileAccessor.cpp
@@ -50,6 +50,7 @@ namespace uba
 	{
 		UBA_ASSERT(flagsAndAttributes != 0);
 		m_size = fileSize;
+		m_flagsAndAttributes = flagsAndAttributes;
 
 		const tchar* realFileName = m_fileName;
 
@@ -129,8 +130,6 @@ namespace uba
 	{
 		allowRead = true; // It is not possible to have write only access to file mappings it seems
 
-		m_size = size;
-
 		UBA_ASSERT(flagsAndAttributes != 0);
 		if (!CreateWrite(allowRead, flagsAndAttributes, size, tempPath))
 			return false;
@@ -405,7 +404,7 @@ namespace uba
 							return m_logger.Error(TC("Failed to rename temporary file %s to %s (%s)"), realFileName, m_fileName, strerror(errno));
 
 						// Need to copy, can't rename over devices
-						int targetFd = open(m_fileName, O_CREAT | O_TRUNC | O_WRONLY | O_CLOEXEC);
+						int targetFd = open(m_fileName, O_CREAT | O_TRUNC | O_WRONLY | O_CLOEXEC, m_flagsAndAttributes);
 						auto g = MakeGuard([targetFd]() { close(targetFd); });
 						if (targetFd == -1)
 							return m_logger.Error(TC("Failed to create file %s for move from temporary file %s (%s)"), m_fileName, realFileName, strerror(errno));
diff --git a/Engine/Source/Programs/UnrealBuildAccelerator/Common/Public/UbaFileAccessor.h b/Engine/Source/Programs/UnrealBuildAccelerator/Common/Public/UbaFileAccessor.h
--- a/Engine/Source/Programs/UnrealBuildAccelerator/Common/Public/UbaFileAccessor.h
+++ b/Engine/Source/Programs/UnrealBuildAccelerator/Common/Public/UbaFileAccessor.h
@@ -40,6 +40,7 @@ namespace uba
 		FileMappingHandle m_mappingHandle;
 		u64 m_size = 0;
 		u8* m_data = nullptr;
+		u32 m_flagsAndAttributes = 0;
 		bool m_isWrite = false;
 
 		#if !PLATFORM_WINDOWS
