[UBA]
* Added uba::Parse to parse u64 from string

diff --git a/Engine/Source/Programs/UnrealBuildAccelerator/Core/Private/UbaStringBuffer.cpp b/Engine/Source/Programs/UnrealBuildAccelerator/Core/Private/UbaStringBuffer.cpp
--- a/Engine/Source/Programs/UnrealBuildAccelerator/Core/Private/UbaStringBuffer.cpp
+++ b/Engine/Source/Programs/UnrealBuildAccelerator/Core/Private/UbaStringBuffer.cpp
@@ -132,6 +132,19 @@ namespace uba
 		Replace(str, NonPathSeparator, PathSeparator);
 	}
 
+	bool Parse(u64& out, const tchar* str, u64 strLen)
+	{
+		if (!strLen)
+			return false;
+
+		#if PLATFORM_WINDOWS
+		out = wcstoull(str, nullptr, 10);
+		#else
+		out = strtoull(str, nullptr, 10);
+		#endif
+		return out != 0 || Equals(str, TC("0"));
+	}
+
 	StringBufferBase& StringBufferBase::Append(const tchar* str)
 	{
 		return Append(str, u32(TStrlen(str)));
@@ -325,15 +338,7 @@ namespace uba
 
 	bool StringBufferBase::Parse(u64& out)
 	{
-		if (!count)
-			return false;
-
-		#if PLATFORM_WINDOWS
-		out = wcstoull(data, nullptr, 10);
-		#else
-		out = strtoull(data, nullptr, 10);
-		#endif
-		return out != 0 || Equals(TC("0"));
+		return uba::Parse(out, data, count);
 	}
 
 	bool StringBufferBase::Parse(u32& out)
diff --git a/Engine/Source/Programs/UnrealBuildAccelerator/Core/Public/UbaStringBuffer.h b/Engine/Source/Programs/UnrealBuildAccelerator/Core/Public/UbaStringBuffer.h
--- a/Engine/Source/Programs/UnrealBuildAccelerator/Core/Public/UbaStringBuffer.h
+++ b/Engine/Source/Programs/UnrealBuildAccelerator/Core/Public/UbaStringBuffer.h
@@ -15,6 +15,7 @@ namespace uba
 	bool Equals(const tchar* str1, const tchar* str2, u64 count, bool ignoreCase = true);
 	void Replace(tchar* str, tchar from, tchar to);
 	void FixPathSeparators(tchar* str);
+	bool Parse(u64& out, const tchar* str, u64 strLen);
 	inline void ToLower(tchar* str) { while (tchar c = *str) { if (c >= 'A' && c <= 'Z') *str = c - 'A' + 'a'; ++str; } }
 	inline tchar ToLower(tchar c) { return (c >= 'A' && c <= 'Z') ? (c - 'A' + 'a') : c;}
 	inline tchar ToUpper(tchar c) { return (c >= 'a' && c <= 'z') ? (c - 'a' + 'A') : c;}
