[UBA]
* Added so NetworkServer can have multiple active crypto keys that expire at certain time.
* Added http server to cache service that can be used to register crypto keys. Test with "curl http://localhost:80/addcrypto?0123456701234567,1000" (key will be valid for 1000 seconds)

diff --git a/Engine/Source/Programs/UnrealBuildAccelerator/Agent/Private/UbaAgent.cpp b/Engine/Source/Programs/UnrealBuildAccelerator/Agent/Private/UbaAgent.cpp
--- a/Engine/Source/Programs/UnrealBuildAccelerator/Agent/Private/UbaAgent.cpp
+++ b/Engine/Source/Programs/UnrealBuildAccelerator/Agent/Private/UbaAgent.cpp
@@ -103,7 +103,7 @@ namespace uba
 		logger.Info(TC("  -killrandom             Kills random process and exit session"));
 		logger.Info(TC("  -memwait=<percent>      The amount of memory needed to spawn a process. Set this to 100 to disable. Defaults to 80%%"));
 		logger.Info(TC("  -memkill=<percent>      The amount of memory needed before processes starts to be killed. Set this to 100 to disable. Defaults to 90%%"));
-		logger.Info(TC("  -crypto=<key>           16 bytes crypto key used for secure network transfer"));
+		logger.Info(TC("  -crypto=<key>           32 character (16 bytes) crypto key used for secure network transfer"));
 		logger.Info(TC("  -populateCas=<dir>      Prepopulate cas database with files in dir. If files needed exists on machine this can be an optimization"));
 		#if PLATFORM_MAC
 		logger.Info(TC("  -populateCasFromXcodeVersion=<version>   Prepopulate cas database with files from local xcode installation that matches the version."));
diff --git a/Engine/Source/Programs/UnrealBuildAccelerator/CacheService/Private/UbaCacheService.cpp b/Engine/Source/Programs/UnrealBuildAccelerator/CacheService/Private/UbaCacheService.cpp
--- a/Engine/Source/Programs/UnrealBuildAccelerator/CacheService/Private/UbaCacheService.cpp
+++ b/Engine/Source/Programs/UnrealBuildAccelerator/CacheService/Private/UbaCacheService.cpp
@@ -2,6 +2,7 @@
 
 #include "UbaCacheServer.h"
 #include "UbaFile.h"
+#include "UbaHttpServer.h"
 #include "UbaNetworkBackendTcp.h"
 #include "UbaNetworkServer.h"
 #include "UbaPlatform.h"
@@ -47,6 +48,7 @@ namespace uba
 		logger.Info(TC("  -port=[<host>:]<port>   The ip/name and port (default: %u) to listen for clients on"), DefaultCachePort);
 		logger.Info(TC("  -capacity=<gigaby>      Capacity of local store. Defaults to %u gigabytes"), DefaultCapacityGb);
 		logger.Info(TC("  -expiration=<seconds>   Time until unused cache entries get deleted. Defaults to %s (%u seconds)"), TimeToText(MsToTime(DefaultExpiration*1000)).str, DefaultExpiration);
+		logger.Info(TC("  -http=<port>            If set, a http server will be started and listen on <port>"));
 		logger.Info(TC(""));
 		return -1;
 	}
@@ -115,6 +117,7 @@ namespace uba
 		StringBuffer<256> workDir;
 		StringBuffer<128> listenIp;
 		u16 port = DefaultCachePort;
+		u16 httpPort = 0;
 		bool quiet = false;
 		bool storeCompressed = true;
 		u32 expirationTimeSeconds = DefaultExpiration;
@@ -166,6 +169,11 @@ namespace uba
 				if (!value.Parse(expirationTimeSeconds))
 					return PrintHelp(TC("Invalid value for -expire"));
 			}
+			else if (name.Equals(TC("-http")))
+			{
+				if (!value.Parse(httpPort))
+					httpPort = 80;
+			}
 			else if (name.Equals(TC("-?")))
 			{
 				return PrintHelp(TC(""));
@@ -245,6 +253,32 @@ namespace uba
 		if (!cacheServer.RunMaintenance(true, ShouldExit))
 			return -1;
 
+		HttpServer httpServer(logWriter, networkBackend);
+
+		if (httpPort)
+		{
+			httpServer.AddCommandHandler([&](const tchar* command, tchar* arguments)
+				{
+					if (!Equals(command, TC("addcrypto")))
+						return "Unknown command ('addcrypto' only available)";
+					u64 expirationSeconds = 60;
+					if (tchar* comma  = TStrchr(arguments, ','))
+					{
+						*comma = 0;
+						if (!Parse(expirationSeconds, comma+1, TStrlen(comma+1)))
+							return "Failed to parse expiration seconds";
+					}
+
+					u8 crypto128Data[16];
+					if (!CryptoFromString(crypto128Data, 16, arguments))
+						return "Failed to read crypto argument";
+					u64 expirationTime = GetTime() + MsToTime(expirationSeconds*1000);
+					networkServer.RegisterCryptoKey(crypto128Data, expirationTime);
+					return (const char*)nullptr;
+				});
+			httpServer.StartListen(httpPort);
+		}
+
 		{
 			auto stopListen = MakeGuard([&]() { networkBackend.StopListen(); });
 			auto stopServer = MakeGuard([&]() { networkServer.DisconnectClients(); });
diff --git a/Engine/Source/Programs/UnrealBuildAccelerator/CacheService/Private/UbaHttpServer.cpp b/Engine/Source/Programs/UnrealBuildAccelerator/CacheService/Private/UbaHttpServer.cpp
new file mode 100644
--- /dev/null
+++ b/Engine/Source/Programs/UnrealBuildAccelerator/CacheService/Private/UbaHttpServer.cpp
@@ -0,0 +1,93 @@
+// Copyright Epic Games, Inc. All Rights Reserved.
+
+#include "UbaHttpServer.h"
+#include "UbaStringBuffer.h"
+
+namespace uba
+{
+	struct HttpServer::Connection
+	{
+		HttpServer& server;
+		void* connection;
+	};
+
+	HttpServer::HttpServer(LogWriter& logWriter, NetworkBackend& backend, const tchar* name)
+	:	m_logger(logWriter, name)
+	,	m_backend(backend)
+	{
+	}
+
+	bool HttpServer::ReceiveHeader(void* context, const Guid& connectionUid, u8* headerData, void*& outBodyContext, u8*& outBodyData, u32& outBodySize)
+	{
+		outBodySize = 32*1024;
+		outBodyData = new u8[outBodySize];
+		return true;
+	}
+
+	bool HttpServer::ReceiveBody(void* context, bool recvError, u8* headerData, void* bodyContext, u8* bodyData, u32 bodySize)
+	{
+		auto& c = *(Connection*)context;
+
+		char response[1024];
+		auto sendAndExit = MakeGuard([&]()
+			{
+				NetworkBackend::SendContext sendContext; c.server.m_backend.Send(c.server.m_logger, c.connection, response, u32(strlen(response)), sendContext);
+				delete[] bodyData;
+			});
+
+		auto sendFail = [&](const char* failReason)
+			{
+				snprintf(response, 1024, "HTTP/1.1 404 Not Found\r\n""Content-Type: text/plain\r\n""\r\n""404 %s", failReason);
+				return false;
+			};
+
+		auto request = (const char*)bodyData;
+		if (strncmp(request, "GET /", 5) != 0)
+			return sendFail("Only support GET");
+
+		request += 5;
+		const char* commandEnd = strchr(request, '?');
+		
+		if (!commandEnd)
+			return sendFail("Command not found (must end with '?')");
+
+		StringBuffer<128> command;
+		command.Append(request, u32(commandEnd - request));
+
+		request = commandEnd + 1;
+
+		const char* argumentsEnd = strchr(request, ' ');
+		if (!argumentsEnd)
+			return sendFail("Arguments end not found");
+
+		if (argumentsEnd - request >= 256)
+			return sendFail("arguments too long");
+
+		StringBuffer<256> arguments;
+		arguments.Append(request, u32(argumentsEnd - request));
+
+		if (const char* result = c.server.m_handler(command.data, arguments.data))
+			return sendFail(result);
+
+		snprintf(response, 1024, "HTTP/1.1 200 OK\r\n""Content-Type: text/plain\r\n""\r\n");
+		return false;
+	}
+
+	bool HttpServer::StartListen(u16 port)
+	{
+		m_backend.StartListen(m_logger, port, nullptr, [this](void* connection, const sockaddr& remoteSocketAddr)
+		{
+			m_backend.SetAllowLessThanBodySize(connection, true);
+			auto c = new Connection{*this, connection};
+			m_backend.SetDisconnectCallback(connection, c, [](void* context, const Guid& connectionUid, void* connection) { delete (Connection*)context; });
+			m_backend.SetRecvCallbacks(connection, c, 0, ReceiveHeader, ReceiveBody, TC("Receive"));
+			return true;
+		});
+		return true;
+	}
+
+	void HttpServer::AddCommandHandler(const CommandHandler& handler)
+	{
+		m_handler = handler;
+	}
+}
\ No newline at end of file
diff --git a/Engine/Source/Programs/UnrealBuildAccelerator/CacheService/Private/UbaHttpServer.h b/Engine/Source/Programs/UnrealBuildAccelerator/CacheService/Private/UbaHttpServer.h
new file mode 100644
--- /dev/null
+++ b/Engine/Source/Programs/UnrealBuildAccelerator/CacheService/Private/UbaHttpServer.h
@@ -0,0 +1,27 @@
+// Copyright Epic Games, Inc. All Rights Reserved.
+
+#include "UbaLogger.h"
+#include "UbaNetworkBackend.h"
+
+namespace uba
+{
+	class HttpServer
+	{
+	public:
+		HttpServer(LogWriter& logWriter, NetworkBackend& backend, const tchar* name = TC("UbaHttpServer"));
+
+		using CommandHandler = Function<const char*(const tchar* command, tchar* arguments)>;
+		void AddCommandHandler(const CommandHandler& handler);
+
+		bool StartListen(u16 port);
+
+	private:
+		static bool ReceiveHeader(void* context, const Guid& connectionUid, u8* headerData, void*& outBodyContext, u8*& outBodyData, u32& outBodySize);
+		static bool ReceiveBody(void* context, bool recvError, u8* headerData, void* bodyContext, u8* bodyData, u32 bodySize);
+
+		struct Connection;
+		MutableLogger m_logger;
+		NetworkBackend& m_backend;
+		CommandHandler m_handler;
+	};
+}
\ No newline at end of file
diff --git a/Engine/Source/Programs/UnrealBuildAccelerator/Cli/Private/UbaCli.cpp b/Engine/Source/Programs/UnrealBuildAccelerator/Cli/Private/UbaCli.cpp
--- a/Engine/Source/Programs/UnrealBuildAccelerator/Cli/Private/UbaCli.cpp
+++ b/Engine/Source/Programs/UnrealBuildAccelerator/Cli/Private/UbaCli.cpp
@@ -80,6 +80,7 @@ namespace uba
 		logger.Info(TC("   -storeraw               Disable compression of storage. This will use more storage and might improve performance"));
 		logger.Info(TC("   -maxcpu=<number>        Max number of processes that can be started. Defaults to \"%u\" on this machine"), DefaultProcessorCount);
 		logger.Info(TC("   -visualizer             Spawn a visualizer that visualizes progress"));
+		logger.Info(TC("   -crypto=<32chars>       Will enable crypto on network client/server"));
 		logger.Info(TC("   -coordinator=<name>     Load a UbaCoordinator<name>.dll to instantiate a coordinator to get helpers"));
 		logger.Info(TC("   -cache=<host>[:<port>]  Connect to cache server. Will fetch from cache unless -populatecache is set"));
 		logger.Info(TC("   -populatecache          Populate cache server if connected to one"));
@@ -153,6 +154,7 @@ namespace uba
 		StringBuffer<256> workDir;
 		StringBuffer<128> listenIp;
 		StringBuffer<128> cacheHost;
+		TString crypto;
 		TString coordinatorName;
 		TString coordinatorPool;
 		u32 coordinatorMaxCoreCount = 400;
@@ -243,6 +245,12 @@ namespace uba
 			{
 				launchVisualizer = true;
 			}
+			else if (name.Equals(TC("-crypto")))
+			{
+				if (value.IsEmpty())
+					value.Append(TC("0123456789abcdef0123456789abcdef"));
+				crypto = value.data;
+			}
 			else if (name.Equals(TC("-coordinator")))
 			{
 				if (value.IsEmpty())
@@ -573,6 +581,16 @@ namespace uba
 		if (!ctorSuccess)
 			return -1;
 
+		if (!crypto.empty())
+		{
+			u8 crypto128Data[16];
+			if (!CryptoFromString(crypto128Data, 16, crypto.c_str()))
+				return logger.Error(TC("Failed to parse crypto key %s"), crypto.c_str());
+			networkServer.RegisterCryptoKey(crypto128Data);
+			logger.Info(TC("Using crypto key %s for connections"), crypto.c_str());
+		}
+
+
 		bool isRemote = commandType == CommandType_Remote || commandType == CommandType_Agent;
 		bool useScheduler = EndsWith(application.c_str(), application.size(), TC(".yaml"));
 
diff --git a/Engine/Source/Programs/UnrealBuildAccelerator/Common/Private/UbaNetworkBackendMemory.cpp b/Engine/Source/Programs/UnrealBuildAccelerator/Common/Private/UbaNetworkBackendMemory.cpp
--- a/Engine/Source/Programs/UnrealBuildAccelerator/Common/Private/UbaNetworkBackendMemory.cpp
+++ b/Engine/Source/Programs/UnrealBuildAccelerator/Common/Private/UbaNetworkBackendMemory.cpp
@@ -110,6 +110,11 @@ namespace uba
 		peer.disconnectContext = context;
 	}
 
+	void NetworkBackendMemory::SetAllowLessThanBodySize(void* connection, bool allow)
+	{
+		// Not implemented
+	}
+
 	bool NetworkBackendMemory::StartListen(Logger& logger, u16 port, const tchar* ip, const ListenConnectedFunc& connectedFunc)
 	{
 		SCOPED_WRITE_LOCK(m_connectedFuncLock, l);
diff --git a/Engine/Source/Programs/UnrealBuildAccelerator/Common/Private/UbaNetworkBackendTcp.cpp b/Engine/Source/Programs/UnrealBuildAccelerator/Common/Private/UbaNetworkBackendTcp.cpp
--- a/Engine/Source/Programs/UnrealBuildAccelerator/Common/Private/UbaNetworkBackendTcp.cpp
+++ b/Engine/Source/Programs/UnrealBuildAccelerator/Common/Private/UbaNetworkBackendTcp.cpp
@@ -84,6 +84,8 @@ namespace uba
 
 		Thread recvThread;
 
+		bool allowLess = false;
+
 		Connection(const Connection&) = delete;
 		void operator=(const Connection&) = delete;
 	};
@@ -93,7 +95,7 @@ namespace uba
 	bool DisableNagle(Logger& logger, SOCKET socket);
 	bool EnableFastLoopback(Logger& logger, SOCKET socket);
 	bool SendSocket(Logger& logger, SOCKET socket, const void* b, u64 bufferLen);
-	bool RecvSocket(Logger& logger, SOCKET socket, void* b, u32 bufferLen, u32 timeoutMs, const Guid& connection, const tchar* hint1, const tchar* hint2, bool isFirstCall);
+	bool RecvSocket(Logger& logger, SOCKET socket, void* b, u32& bufferLen, u32 timeoutMs, const Guid& connection, const tchar* hint1, const tchar* hint2, bool isFirstCall, bool allowLess);
 
 	bool NetworkBackendTcp::EnsureInitialized(Logger& logger)
 	{
@@ -228,6 +230,13 @@ namespace uba
 		conn.disconnectContext = context;
 	}
 
+	void NetworkBackendTcp::SetAllowLessThanBodySize(void* connection, bool allow)
+	{
+		auto& conn = *(Connection*)connection;
+		ScopedCriticalSection lock(conn.shutdownLock);
+		conn.allowLess = allow;
+	}
+
 	bool NetworkBackendTcp::StartListen(Logger& logger, u16 port, const tchar* ip, const ListenConnectedFunc& connectedFunc)
 	{
 		if (!EnsureInitialized(logger))
@@ -444,7 +453,7 @@ namespace uba
 				u32 bodySize = 0;
 
 				u8 headerData[MaxHeaderSize];
-				if (!RecvSocket(logger, connection.socket, headerData, connection.headerSize, connection.recvTimeoutMs, connection.uid, connection.recvHint, TC(""), isFirst))
+				if (!RecvSocket(logger, connection.socket, headerData, connection.headerSize, connection.recvTimeoutMs, connection.uid, connection.recvHint, TC(""), isFirst, false))
 					break;
 				isFirst = false;
 
@@ -462,7 +471,7 @@ namespace uba
 				if (!bodySize)
 					continue;
 
-				bool success = RecvSocket(logger, connection.socket, bodyData, bodySize, connection.recvTimeoutMs, connection.uid, connection.recvHint, TC("Body"), false);
+				bool success = RecvSocket(logger, connection.socket, bodyData, bodySize, connection.recvTimeoutMs, connection.uid, connection.recvHint, TC("Body"), false, connection.allowLess);
 
 				m_totalRecv += bodySize;
 
@@ -772,7 +781,7 @@ namespace uba
 		return true;
 	}
 
-	bool RecvSocket(Logger& logger, SOCKET socket, void* b, u32 bufferLen, u32 timeoutMs, const Guid& connection, const tchar* hint1, const tchar* hint2, bool isFirstCall)
+	bool RecvSocket(Logger& logger, SOCKET socket, void* b, u32& bufferLen, u32 timeoutMs, const Guid& connection, const tchar* hint1, const tchar* hint2, bool isFirstCall, bool allowLess)
 	{
 		u8* buffer = (u8*)b;
 		u32 recvLeft = bufferLen;
@@ -827,6 +836,12 @@ namespace uba
 			}
 			recvLeft -= (u32)read;
 			buffer += read;
+
+			if (allowLess)
+			{
+				bufferLen = read;
+				break;
+			}
 		}
 		return true;
 	}
diff --git a/Engine/Source/Programs/UnrealBuildAccelerator/Common/Private/UbaNetworkServer.cpp b/Engine/Source/Programs/UnrealBuildAccelerator/Common/Private/UbaNetworkServer.cpp
--- a/Engine/Source/Programs/UnrealBuildAccelerator/Common/Private/UbaNetworkServer.cpp
+++ b/Engine/Source/Programs/UnrealBuildAccelerator/Common/Private/UbaNetworkServer.cpp
@@ -89,7 +89,7 @@ namespace uba
 	class NetworkServer::Connection
 	{
 	public:
-		Connection(NetworkServer& server, NetworkBackend& backend, void* backendConnection, const sockaddr& remoteSockAddr, CryptoKey cryptoKey, u32 id)
+		Connection(NetworkServer& server, NetworkBackend& backend, void* backendConnection, const sockaddr& remoteSockAddr, bool requiresCrypto, CryptoKey cryptoKey, u32 id)
 		:	m_server(server)
 		,	m_backend(backend)
 		,	m_remoteSockAddr(remoteSockAddr)
@@ -117,8 +117,11 @@ namespace uba
 
 			m_backend.SetRecvTimeout(m_backendConnection, m_server.m_receiveTimeoutMs);
 
-			if (m_cryptoKey)
+			if (requiresCrypto)
+			{
+				m_backend.SetAllowLessThanBodySize(m_backendConnection, true);
 				m_backend.SetRecvCallbacks(m_backendConnection, this, 0, ReceiveHandshakeHeader, ReceiveHandshakeBody, TC("ReceiveHandshake"));
+			}
 			else
 				m_backend.SetRecvCallbacks(m_backendConnection, this, 4, ReceiveVersion, nullptr, TC("ReceiveVersion"));
 		}
@@ -182,12 +185,53 @@ namespace uba
 			u8* handshakeData = bodyData;
 			auto g = MakeGuard([handshakeData]() { delete[] handshakeData; });
 
-			if (!Crypto::Decrypt(conn.m_server.m_logger, conn.m_cryptoKey, handshakeData, sizeof(EncryptionHandshakeString)))
-				return false;
+			auto& logger = conn.m_server.m_logger;
 
-			if (memcmp(handshakeData, EncryptionHandshakeString, sizeof(EncryptionHandshakeString)) != 0)
-				return conn.m_server.m_logger.Error(TC("Crypto mismatch..."));
+			if (bodySize != sizeof(EncryptionHandshakeString))
+				return logger.Error(TC("Crypto mismatch..."));
+
+			auto TestHandshake = [&](CryptoKey key)
+			{
+				u8 temp[sizeof(EncryptionHandshakeString)];
+				memcpy(temp, handshakeData, sizeof(temp));
+				if (!Crypto::Decrypt(logger, key, temp, sizeof(EncryptionHandshakeString)))
+					return false;
+				return memcmp(temp, EncryptionHandshakeString, sizeof(EncryptionHandshakeString)) == 0;
+			};
+
+			if (conn.m_cryptoKey != InvalidCryptoKey)
+			{
+				if (!TestHandshake(conn.m_cryptoKey))
+					return logger.Error(TC("Crypto mismatch..."));
+			}
+			else
+			{
+				SCOPED_WRITE_LOCK(conn.m_server.m_cryptoKeysLock, lock);
+				auto& keys = conn.m_server.m_cryptoKeys;
+				u64 time = GetTime();
+				for (auto it=keys.begin(); it!=keys.end();)
+				{
+					auto& entry = *it;
+					if (entry.expirationTime < time)
+					{
+						it = keys.erase(it);
+						continue;
+					}
+					++it;
+
+					CryptoKey key = Crypto::DuplicateKey(logger, entry.key);
+					auto keyGuard = MakeGuard([&]() { Crypto::DestroyKey(key); });
+					if (!TestHandshake(key))
+						continue;
+					keyGuard.Cancel();
+					conn.m_cryptoKey = key;
+					break;
+				}
+				if (conn.m_cryptoKey == InvalidCryptoKey)
+					return logger.Error(TC("Crypto mismatch..."));
+			}
 
+			conn.m_backend.SetAllowLessThanBodySize(conn.m_backendConnection, false);
 			conn.m_backend.SetRecvCallbacks(conn.m_backendConnection, &conn, 4, ReceiveVersion, nullptr, TC("ReceiveVersion"));
 
 			return true;
@@ -596,27 +640,15 @@ namespace uba
 	{
 		UBA_ASSERT(m_connections.empty());
 		FlushWorkers();
+		for (auto& entry : m_cryptoKeys)
+			Crypto::DestroyKey(entry.key);
 	}
 
-	bool NetworkServer::StartListen(NetworkBackend& backend, u16 port, const tchar* ip, const u8* cryptoKey128)
+	bool NetworkServer::StartListen(NetworkBackend& backend, u16 port, const tchar* ip, bool requiresCrypto)
 	{
-		if (cryptoKey128)
-		{
-			m_listenCrypto = Crypto::CreateKey(m_logger, cryptoKey128);
-			if (!m_listenCrypto)
-				return false;
-		}
-
-		return backend.StartListen(m_logger, port, ip, [&](void* connection, const sockaddr& remoteSockAddr)
+		return backend.StartListen(m_logger, port, ip, [this, &backend, requiresCrypto](void* connection, const sockaddr& remoteSockAddr)
 			{
-				CryptoKey cryptoKey = InvalidCryptoKey;
-				if (m_listenCrypto)
-				{
-					cryptoKey = Crypto::DuplicateKey(m_logger, m_listenCrypto);
-					if (!cryptoKey)
-						return false;
-				}
-				return AddConnection(backend, connection, remoteSockAddr, cryptoKey);
+				return AddConnection(backend, connection, remoteSockAddr, requiresCrypto, InvalidCryptoKey);
 			});
 	}
 	
@@ -665,6 +697,16 @@ namespace uba
 		m_workersEnabled = true;
 	}
 
+	bool NetworkServer::RegisterCryptoKey(const u8* cryptoKey128, u64 expirationTime)
+	{
+		CryptoKey key = Crypto::CreateKey(m_logger, cryptoKey128);
+		if (key == InvalidCryptoKey)
+			return false;
+		SCOPED_WRITE_LOCK(m_cryptoKeysLock, lock);
+		m_cryptoKeys.emplace_back(key, expirationTime);
+		return true;
+	}
+
 	bool NetworkServer::AddClient(NetworkBackend& backend, const tchar* ip, u16 port, const u8* cryptoKey128)
 	{
 		SCOPED_WRITE_LOCK(m_addConnectionsLock, lock);
@@ -695,7 +737,7 @@ namespace uba
 				// TODO: Should this retry?
 				success = backend.Connect(m_logger, ip2.c_str(), [this, &backend, cryptoKey](void* connection, const sockaddr& remoteSocketAddr, bool* timedOut)
 					{
-						return AddConnection(backend, connection, remoteSocketAddr, cryptoKey);
+						return AddConnection(backend, connection, remoteSocketAddr, cryptoKey != InvalidCryptoKey, cryptoKey);
 					}, port, nullptr);
 				if (!success)
 					Crypto::DestroyKey(cryptoKey);
@@ -1124,7 +1166,7 @@ namespace uba
 										if (cryptoKey == InvalidCryptoKey)
 											return false;
 									}
-									return AddConnection(conn.m_backend, connection, remoteSocketAddr, cryptoKey);
+									return AddConnection(conn.m_backend, connection, remoteSocketAddr, cryptoKey != InvalidCryptoKey, cryptoKey);
 								}, nullptr);
 							return 0;
 						});
@@ -1140,7 +1182,7 @@ namespace uba
 		return false;
 	}
 
-	bool NetworkServer::AddConnection(NetworkBackend& backend, void* backendConnection, const sockaddr& remoteSocketAddr, CryptoKey cryptoKey)
+	bool NetworkServer::AddConnection(NetworkBackend& backend, void* backendConnection, const sockaddr& remoteSocketAddr, bool requiresCrypto, CryptoKey cryptoKey)
 	{
 		SCOPED_WRITE_LOCK(m_connectionsLock, lock);
 
@@ -1154,7 +1196,7 @@ namespace uba
 			return false;
 		}
 
-		m_connections.emplace_back(*this, backend, backendConnection, remoteSocketAddr, cryptoKey, m_connectionIdCounter++);
+		m_connections.emplace_back(*this, backend, backendConnection, remoteSocketAddr, requiresCrypto, cryptoKey, m_connectionIdCounter++);
 		m_maxActiveConnections = Max(m_maxActiveConnections, u32(m_connections.size()));
 		return true;
 	}
diff --git a/Engine/Source/Programs/UnrealBuildAccelerator/Common/Public/UbaNetworkBackend.h b/Engine/Source/Programs/UnrealBuildAccelerator/Common/Public/UbaNetworkBackend.h
--- a/Engine/Source/Programs/UnrealBuildAccelerator/Common/Public/UbaNetworkBackend.h
+++ b/Engine/Source/Programs/UnrealBuildAccelerator/Common/Public/UbaNetworkBackend.h
@@ -50,6 +50,7 @@ namespace uba
 		// Disconnect callback. This is called as soon as connection is interrupted from send, recv or shutdown.
 		using DisconnectCallback = void(void* context, const Guid& connectionUid, void* connection);
 		virtual void SetDisconnectCallback(void* connection, void* context, DisconnectCallback* callback) = 0;
+		virtual void SetAllowLessThanBodySize(void* connection, bool allow) = 0;
 
 		// Start listen on port/ip.
 		using ListenConnectedFunc = Function<bool(void* connection, const sockaddr& remoteSocketAddr)>;
diff --git a/Engine/Source/Programs/UnrealBuildAccelerator/Common/Public/UbaNetworkBackendMemory.h b/Engine/Source/Programs/UnrealBuildAccelerator/Common/Public/UbaNetworkBackendMemory.h
--- a/Engine/Source/Programs/UnrealBuildAccelerator/Common/Public/UbaNetworkBackendMemory.h
+++ b/Engine/Source/Programs/UnrealBuildAccelerator/Common/Public/UbaNetworkBackendMemory.h
@@ -19,6 +19,7 @@ namespace uba
 		virtual void SetRecvCallbacks(void* connection, void* context, u32 headerSize, RecvHeaderCallback* h, RecvBodyCallback* b, const tchar* recvHint) override;
 		virtual void SetRecvTimeout(void* connection, u32 timeoutMs) override;
 		virtual void SetDisconnectCallback(void* connection, void* context, DisconnectCallback* callback) override;
+		virtual void SetAllowLessThanBodySize(void* connection, bool allow) override;
 
 		virtual bool StartListen(Logger& logger, u16 port, const tchar* ip, const ListenConnectedFunc& connectedFunc) override;
 		virtual void StopListen() override;
diff --git a/Engine/Source/Programs/UnrealBuildAccelerator/Common/Public/UbaNetworkBackendTcp.h b/Engine/Source/Programs/UnrealBuildAccelerator/Common/Public/UbaNetworkBackendTcp.h
--- a/Engine/Source/Programs/UnrealBuildAccelerator/Common/Public/UbaNetworkBackendTcp.h
+++ b/Engine/Source/Programs/UnrealBuildAccelerator/Common/Public/UbaNetworkBackendTcp.h
@@ -19,6 +19,7 @@ namespace uba
 		virtual void SetRecvCallbacks(void* connection, void* context, u32 headerSize, RecvHeaderCallback* h, RecvBodyCallback* b, const tchar* recvHint) override;
 		virtual void SetRecvTimeout(void* connection, u32 timeoutMs) override;
 		virtual void SetDisconnectCallback(void* connection, void* context, DisconnectCallback* callback) override;
+		virtual void SetAllowLessThanBodySize(void* connection, bool allow) override;
 
 		virtual bool StartListen(Logger& logger, u16 port, const tchar* ip, const ListenConnectedFunc& connectedFunc) override;
 		virtual void StopListen() override;
diff --git a/Engine/Source/Programs/UnrealBuildAccelerator/Common/Public/UbaNetworkServer.h b/Engine/Source/Programs/UnrealBuildAccelerator/Common/Public/UbaNetworkServer.h
--- a/Engine/Source/Programs/UnrealBuildAccelerator/Common/Public/UbaNetworkServer.h
+++ b/Engine/Source/Programs/UnrealBuildAccelerator/Common/Public/UbaNetworkServer.h
@@ -47,10 +47,12 @@ namespace uba
 		NetworkServer(bool& outCtorSuccess, const NetworkServerCreateInfo& info = {}, const tchar* name = TC("UbaServer"));
 		virtual ~NetworkServer();
 
-		bool StartListen(NetworkBackend& backend, u16 port = DefaultPort, const tchar* ip = nullptr, const u8* cryptoKey128 = nullptr); // Start listen for new connections/clients
+		bool StartListen(NetworkBackend& backend, u16 port = DefaultPort, const tchar* ip = nullptr, bool requiresCrypto = false); // Start listen for new connections/clients
 		void DisallowNewClients();	// Disallow new clients to connect but old clients can still create more connections
 		void DisconnectClients();				// Stops all listen and disconnect all active connections
 
+		bool RegisterCryptoKey(const u8* cryptoKey128, u64 expirationTime = ~u64(0));
+
 		bool AddClient(NetworkBackend& backend, const tchar* ip, u16 port = DefaultPort, const u8* cryptoKey128 = nullptr); // Adds a client that server will create one or more connections to (note this will return before we know if it was a success or not)
 
 		void PrintSummary(Logger& logger);
@@ -106,13 +108,20 @@ namespace uba
 		void FlushWorkers();
 
 		bool HandleSystemMessage(const ConnectionInfo& connectionInfo, u8 messageType, BinaryReader& reader, BinaryWriter& writer);
-		bool AddConnection(NetworkBackend& backend, void* backendConnection, const sockaddr& remoteSocketAddr, CryptoKey cryptoKey);
+		bool AddConnection(NetworkBackend& backend, void* backendConnection, const sockaddr& remoteSocketAddr, bool requiresCrypto, CryptoKey cryptoKey);
 
 		void RemoveDisconnectedConnections();
 
 		MutableLogger m_logger;
 
-		CryptoKey m_listenCrypto = InvalidCryptoKey;
+		struct CryptoEntry
+		{
+			CryptoKey key;
+			u64 expirationTime;
+		};
+		ReaderWriterLock m_cryptoKeysLock;
+		List<CryptoEntry> m_cryptoKeys;
+
 		Guid m_uid;
 		bool m_allowNewClients = true;
 
diff --git a/Engine/Source/Programs/UnrealBuildAccelerator/Host/Private/UbaExports.cpp b/Engine/Source/Programs/UnrealBuildAccelerator/Host/Private/UbaExports.cpp
--- a/Engine/Source/Programs/UnrealBuildAccelerator/Host/Private/UbaExports.cpp
+++ b/Engine/Source/Programs/UnrealBuildAccelerator/Host/Private/UbaExports.cpp
@@ -174,17 +174,21 @@ extern "C"
 	bool NetworkServer_StartListen(uba::NetworkServer* server, int port, const uba::tchar* ip, const uba::tchar* crypto)
 	{
 		using namespace uba;
-		u8 crypto128Data[16];
-		u8* crypto128 = nullptr;
+
+		auto s = (NetworkServerWithBackend*)server;
+
+		bool requiresCrypto = false;
 		if (crypto && *crypto)
 		{
-			((u64*)crypto128Data)[0] = StringToValue(crypto, 16);
-			((u64*)crypto128Data)[1] = StringToValue(crypto + 16, 16);
-			crypto128 = crypto128Data;
+			//server->GetLogger().Error(TC("CRYPTO: %s"), crypto);
+			u8 crypto128Data[16];
+			if (!CryptoFromString(crypto128Data, 16, crypto))
+				return server->GetLogger().Error(TC("Failed to parse crypto key %s"), crypto);
+			s->RegisterCryptoKey(crypto128Data);
+			requiresCrypto = true;
 		}
 
-		auto s = (NetworkServerWithBackend*)server;
-		return s->StartListen(*s->backend, u16(port), ip, crypto128);
+		return s->StartListen(*s->backend, u16(port), ip, requiresCrypto);
 	}
 
 	void NetworkServer_Stop(uba::NetworkServer* server)
@@ -201,12 +205,8 @@ extern "C"
 		using namespace uba;
 		u8 crypto128Data[16];
 		u8* crypto128 = nullptr;
-		if (crypto && *crypto)
-		{
-			((u64*)crypto128Data)[0] = StringToValue(crypto, 16);
-			((u64*)crypto128Data)[1] = StringToValue(crypto + 16, 16);
+		if (CryptoFromString(crypto128Data, 16, crypto))
 			crypto128 = crypto128Data;
-		}
 
 		auto s = (NetworkServerWithBackend*)server;
 		return s->AddClient(*s->backend, ip, u16(port), crypto128);
@@ -530,15 +530,28 @@ extern "C"
 		scheduler->GetStats(outQueued, outActiveLocal, outActiveRemote, outFinished);
 	}
 
-	uba::CacheClient* CacheClient_Create(uba::SessionServer* session, bool reportMissReason)
+	uba::CacheClient* CacheClient_Create(uba::SessionServer* session, bool reportMissReason, const uba::tchar* crypto)
 	{
 		using namespace uba;
 		LogWriter& writer = session->GetLogWriter();
 		StorageImpl& storage = (StorageImpl&)session->GetStorage();
 		auto& server = (NetworkServerWithBackend&)session->GetServer();
 
+		u8 crypto128Data[16];
+		u8* crypto128 = nullptr;
+		if (crypto && *crypto)
+		{
+			if (!CryptoFromString(crypto128Data, 16, crypto))
+			{
+				LoggerWithWriter(writer, TC("UbaCacheClient")).Error(TC("Failed to parse crypto key %s"), crypto);
+				return nullptr;
+			}
+			crypto128 = crypto128Data;
+		}
+
 		NetworkClientCreateInfo ncci(writer);
 		ncci.receiveTimeoutSeconds = 60;
+		ncci.cryptoKey128 = crypto128;
 		bool ctorSuccess = false;
 		auto networkClient = new NetworkClientWithBackend(ctorSuccess, ncci, server.backend);
 		if (!ctorSuccess)
diff --git a/Engine/Source/Programs/UnrealBuildAccelerator/Host/Public/UbaExports.h b/Engine/Source/Programs/UnrealBuildAccelerator/Host/Public/UbaExports.h
--- a/Engine/Source/Programs/UnrealBuildAccelerator/Host/Public/UbaExports.h
+++ b/Engine/Source/Programs/UnrealBuildAccelerator/Host/Public/UbaExports.h
@@ -123,7 +123,7 @@ extern "C"
 	UBA_API void Scheduler_GetStats(uba::Scheduler* scheduler, uba::u32& outQueued, uba::u32& outActiveLocal, uba::u32& outActiveRemote, uba::u32& outFinished);
 
 	// Cache
-	UBA_API uba::CacheClient* CacheClient_Create(uba::SessionServer* session, bool reportMissReason = false);
+	UBA_API uba::CacheClient* CacheClient_Create(uba::SessionServer* session, bool reportMissReason = false, const uba::tchar* crypto = nullptr);
 	UBA_API bool CacheClient_Connect(uba::CacheClient* cacheClient, const uba::tchar* host, int port);
 	UBA_API bool CacheClient_WriteToCache(uba::CacheClient* cacheClient, uba::RootPaths* rootPaths, uba::u32 bucket, const uba::ProcessHandle* process, const uba::u8* inputs, uba::u32 inputsSize, const uba::u8* outputs, uba::u32 outputsSize);
 	UBA_API bool CacheClient_FetchFromCache(uba::CacheClient* cacheClient, uba::RootPaths* rootPaths, uba::u32 bucket, const uba::ProcessStartInfo& info);
