[UBA]
* Fixed unity compile error

diff --git a/Engine/Source/Programs/UnrealBuildAccelerator/CacheService/Private/UbaHttpServer.h b/Engine/Source/Programs/UnrealBuildAccelerator/CacheService/Private/UbaHttpServer.h
--- a/Engine/Source/Programs/UnrealBuildAccelerator/CacheService/Private/UbaHttpServer.h
+++ b/Engine/Source/Programs/UnrealBuildAccelerator/CacheService/Private/UbaHttpServer.h
@@ -1,5 +1,7 @@
 // Copyright Epic Games, Inc. All Rights Reserved.
 
+#pragma once
+
 #include "UbaLogger.h"
 #include "UbaNetworkBackend.h"
 
