[UBA]
* Fixed macos compile error

diff --git a/Engine/Source/Programs/UnrealBuildAccelerator/Common/Private/UbaNetworkServer.cpp b/Engine/Source/Programs/UnrealBuildAccelerator/Common/Private/UbaNetworkServer.cpp
--- a/Engine/Source/Programs/UnrealBuildAccelerator/Common/Private/UbaNetworkServer.cpp
+++ b/Engine/Source/Programs/UnrealBuildAccelerator/Common/Private/UbaNetworkServer.cpp
@@ -703,7 +703,7 @@ namespace uba
 		if (key == InvalidCryptoKey)
 			return false;
 		SCOPED_WRITE_LOCK(m_cryptoKeysLock, lock);
-		m_cryptoKeys.emplace_back(key, expirationTime);
+		m_cryptoKeys.push_back(CryptoEntry{key, expirationTime});
 		return true;
 	}
 
