[UBA]
* Fixed bug where ReleaseMutex was not called in all exit paths after WaitForSingleObject

diff --git a/Engine/Source/Programs/UnrealBuildAccelerator/Common/Private/UbaTrace.cpp b/Engine/Source/Programs/UnrealBuildAccelerator/Common/Private/UbaTrace.cpp
--- a/Engine/Source/Programs/UnrealBuildAccelerator/Common/Private/UbaTrace.cpp
+++ b/Engine/Source/Programs/UnrealBuildAccelerator/Common/Private/UbaTrace.cpp
@@ -460,11 +460,11 @@ namespace uba
 	{
 		#if PLATFORM_WINDOWS
 		WaitForSingleObject((HANDLE)m_mutex, INFINITE);
+		auto g = MakeGuard([this]() { ReleaseMutex((HANDLE)m_mutex); });
 		if (ifMatching)
 			if (!Equals((tchar*)m_mem, ifMatching))
 				return true;
 		TStrcpy_s((tchar*)m_mem, 256, traceName);
-		ReleaseMutex((HANDLE)m_mutex);
 		#endif
 		return true;
 	}
