[UBA]
* Added support for having multiple tilde in same path
* Added IsAbsolutePath function

diff --git a/Engine/Source/Programs/UnrealBuildAccelerator/Common/Private/UbaFile.cpp b/Engine/Source/Programs/UnrealBuildAccelerator/Common/Private/UbaFile.cpp
--- a/Engine/Source/Programs/UnrealBuildAccelerator/Common/Private/UbaFile.cpp
+++ b/Engine/Source/Programs/UnrealBuildAccelerator/Common/Private/UbaFile.cpp
@@ -76,7 +76,7 @@ namespace uba
 	#define MAKE_LONG_FILENAME(fileName) \
 		UBA_ASSERT(TStrlen(fileName) < MaxPath); \
 		StringBuffer<MaxPath> STRING_JOIN(longName, __LINE__); \
-		if (fileName && fileName[0] && fileName[1] == ':') \
+		if (IsAbsolutePath(fileName)) \
 		{ \
 			auto& lsb = STRING_JOIN(longName, __LINE__); \
 			lsb.Append(TC("\\\\?\\")); \
diff --git a/Engine/Source/Programs/UnrealBuildAccelerator/Core/Public/UbaDirectoryTable.h b/Engine/Source/Programs/UnrealBuildAccelerator/Core/Public/UbaDirectoryTable.h
--- a/Engine/Source/Programs/UnrealBuildAccelerator/Core/Public/UbaDirectoryTable.h
+++ b/Engine/Source/Programs/UnrealBuildAccelerator/Core/Public/UbaDirectoryTable.h
@@ -248,7 +248,7 @@ namespace uba
 
 		void GetFinalPath(StringBufferBase& out, const tchar* path)
 		{
-			UBA_ASSERT(path[1] == ':');
+			UBA_ASSERT(IsAbsolutePath(path));
 
 			Directory* directory = nullptr;
 			const tchar* prevSlash = TStrchr(path+3, PathSeparator);
diff --git a/Engine/Source/Programs/UnrealBuildAccelerator/Core/Public/UbaPathUtils.h b/Engine/Source/Programs/UnrealBuildAccelerator/Core/Public/UbaPathUtils.h
--- a/Engine/Source/Programs/UnrealBuildAccelerator/Core/Public/UbaPathUtils.h
+++ b/Engine/Source/Programs/UnrealBuildAccelerator/Core/Public/UbaPathUtils.h
@@ -28,13 +28,28 @@ namespace uba
 			++read;
 
 		tchar fullName[1024];
-		if (const tchar* tilde = TStrchr(fileName, '~'))
+		while (true)
 		{
+			tchar temp[1024];
+			if (read == fullName)
+			{
+				TStrcpy_s(temp, sizeof_array(temp), read);
+				read = temp;
+			}
+
+			const tchar* tilde = TStrchr(read, '~');
+			if (!tilde)
+				break;
+
+			u64 fileNameLen = TStrlen(read);
+			if (fileNameLen >= 1024)
+				return false;
+
 			// Since this might be a in memory file we can't use the actual full name, so let's find first slash after '~'
 			const tchar* backslash = TStrchr(tilde, '\\');
-			TStrcpy_s(fullName, sizeof_array(fullName), fileName);
+			TStrcpy_s(fullName, sizeof_array(fullName), read);
 			if (backslash)
-				fullName[backslash - fileName] = 0;
+				fullName[backslash - read] = 0;
 			u32 len = Local_GetLongPathNameW(fullName, fullName, sizeof_array(fullName));
 			if (!len)
 			{
@@ -304,4 +319,6 @@ namespace uba
 		FixPath2(fileName, workingDir, workingDirCharLen, buffer.data + count, buffer.capacity - count, &buffer.count);
 		buffer.count += count;
 	}
-}
\ No newline at end of file
+
+	inline bool IsAbsolutePath(const tchar* path) { return path ? (IsWindows ? (path[0] && path[1] == ':') : path[0] == '/') : false; }
+}
