[UBA]
* Added support for strings and comments in config files
* Added tests
* Fixed for uba to work integrated in other programs (proper dll exports)
* Added process id to process info in visualizer

diff --git a/Engine/Source/Programs/UnrealBuildAccelerator/Common/Private/UbaConfig.cpp b/Engine/Source/Programs/UnrealBuildAccelerator/Common/Private/UbaConfig.cpp
--- a/Engine/Source/Programs/UnrealBuildAccelerator/Common/Private/UbaConfig.cpp
+++ b/Engine/Source/Programs/UnrealBuildAccelerator/Common/Private/UbaConfig.cpp
@@ -5,15 +5,18 @@
 
 namespace uba
 {
-	static ConfigTable& EmptyConfigTable = *new ConfigTable();
-
-	bool ConfigTable::GetValueAsString(StringBufferBase& out, const tchar* key) const
+	bool ConfigTable::GetValueAsString(const tchar*& out, const tchar* key) const
 	{
-		return false;
+		auto findIt = m_values.find(key);
+		if (findIt == m_values.end())
+			return m_parent ? m_parent->GetValueAsString(out, key) : false;
+		out = findIt->second.c_str();
+		return true;
 	}
 
 	bool ConfigTable::GetValueAsU32(u32& out, const tchar* key) const
 	{
+		UBA_ASSERT(false);
 		return false;
 	}
 
@@ -21,7 +24,7 @@ namespace uba
 	{
 		auto findIt = m_values.find(key);
 		if (findIt == m_values.end())
-			return false;
+			return m_parent ? m_parent->GetValueAsBool(out, key) : false;
 		const tchar* value = findIt->second.c_str();
 		if (Equals(value, TC("true")) || Equals(value, TC("1")))
 		{
@@ -41,9 +44,27 @@ namespace uba
 	{
 		m_isLoaded = true;
 
+#if PLATFORM_WINDOWS
+		StringBuffer<> tempPath;
+		if (configFile[1] != ':')
+		{
+			if (GetCurrentDirectoryW(tempPath) && FileExists(logger, tempPath.EnsureEndsWithSlash().Append(configFile).data))
+			{
+				configFile = tempPath.data;
+			}
+			else if (GetDirectoryOfCurrentModule(logger, tempPath.Clear()) && FileExists(logger, tempPath.EnsureEndsWithSlash().Append(configFile).data))
+			{
+				configFile = tempPath.data;
+			}
+			else
+				return false;
+		}
+#endif
+
 		FileAccessor fa(logger, configFile);
 		if (!fa.OpenMemoryRead(0, false))
 			return false;
+		logger.Info(TC("  Loading config from %s"), configFile);
 		return LoadFromText(logger, (const char*)fa.GetData(), fa.GetSize());
 	}
 
@@ -77,21 +98,24 @@ namespace uba
 				return 0;
 			};
 
-		auto consumeLine = [&](StringBufferBase& out) -> char
+		auto consumeLine = [&](StringBufferBase& out, char untilChar) -> char
 			{
+				bool append = true;
 				while (i != e)
 				{
 					if (*i == '\n')
 						return *i;
-					if (*i != '\r')
+
+					append &= (*i != untilChar && *i != '\r');
+
+					if (append)
 						out.Append(*i);
 					++i;
 				}
 				return 0;
 			};
 
-		ConfigTable globalTable;
-		ConfigTable* activeTable = &globalTable;
+		ConfigTable* activeTable = &m_globalTable;
 		while (true)
 		{
 			char token = consumeEmpty();
@@ -101,6 +125,11 @@ namespace uba
 			{
 				++i;
 			}
+			else if (token == '#')
+			{
+				StringBuffer<1024> comment;
+				consumeLine(comment, 0);
+			}
 			else if (token == '[')
 			{
 				++i;
@@ -116,6 +145,7 @@ namespace uba
 					return logger.Error(TC("Unexpected token %c after group %s"), tableName.data);
 				++i;
 				activeTable = &m_tables.try_emplace(tableName.data).first->second;
+				activeTable->m_parent = &m_globalTable;
 			}
 			else
 			{
@@ -126,19 +156,22 @@ namespace uba
 					return logger.Error(TC("Unexpected equals sign after key name %s"), key);
 				++i;
 				token = consumeEmpty();
+
+				StringBuffer<1024> value;
 				if (token == '\"')
 				{
-					return logger.Error(TC("Strings not supported yet %s"), key);
+					++i;
+					token = consumeLine(value, '\"');
 				}
 				else
 				{
-					StringBuffer<128> value;
 					token = consumeIdentifier(value);
-					activeTable->m_values[key.data] = value.data;
-					if (token == 0)
-						break;
-					++i;
 				}
+
+				activeTable->m_values[key.data] = value.data;
+				if (token == 0)
+					break;
+				++i;
 			}
 		}
 
@@ -154,7 +187,7 @@ namespace uba
 	{
 		auto findIt = m_tables.find(name);
 		if (findIt == m_tables.end())
-			return EmptyConfigTable;
+			return m_globalTable;
 		return findIt->second;
 	}
-}
\ No newline at end of file
+}
diff --git a/Engine/Source/Programs/UnrealBuildAccelerator/Common/Private/UbaLogger.cpp b/Engine/Source/Programs/UnrealBuildAccelerator/Common/Private/UbaLogger.cpp
--- a/Engine/Source/Programs/UnrealBuildAccelerator/Common/Private/UbaLogger.cpp
+++ b/Engine/Source/Programs/UnrealBuildAccelerator/Common/Private/UbaLogger.cpp
@@ -184,7 +184,7 @@ namespace uba
 		u32 m_defaultAttributes = 0;
 #endif
 	};
-	LogWriter& g_consoleLogWriter = *new ConsoleLogWriter(); // Leak to prevent asan annoyances during shutdown when asserts happen
+	UBA_API LogWriter& g_consoleLogWriter = *new ConsoleLogWriter(); // Leak to prevent asan annoyances during shutdown when asserts happen
 	thread_local u32 t_consoleLogScopeCount = 0;
 
 	class NullLogWriter : public LogWriter
@@ -194,7 +194,7 @@ namespace uba
 		virtual void EndScope() override {}
 		virtual void Log(LogEntryType type, const tchar* str, u32 strLen, const tchar* prefix = nullptr, u32 prefixLen = 0) override {}
 	} g_nullLogWriterImpl;
-	LogWriter& g_nullLogWriter = g_nullLogWriterImpl;
+	UBA_API LogWriter& g_nullLogWriter = g_nullLogWriterImpl;
 
 
 	ConsoleLogWriter::ConsoleLogWriter()
diff --git a/Engine/Source/Programs/UnrealBuildAccelerator/Common/Private/UbaNetworkServer.cpp b/Engine/Source/Programs/UnrealBuildAccelerator/Common/Private/UbaNetworkServer.cpp
--- a/Engine/Source/Programs/UnrealBuildAccelerator/Common/Private/UbaNetworkServer.cpp
+++ b/Engine/Source/Programs/UnrealBuildAccelerator/Common/Private/UbaNetworkServer.cpp
@@ -7,6 +7,11 @@
 
 namespace uba
 {
+	void NetworkServerCreateInfo::Apply(Config& config)
+	{
+	}
+
+
 	struct NetworkServer::WorkerContext
 	{
 		WorkerContext(NetworkServer& s) : server(s), workAvailable(false)
diff --git a/Engine/Source/Programs/UnrealBuildAccelerator/Common/Private/UbaStorage.cpp b/Engine/Source/Programs/UnrealBuildAccelerator/Common/Private/UbaStorage.cpp
--- a/Engine/Source/Programs/UnrealBuildAccelerator/Common/Private/UbaStorage.cpp
+++ b/Engine/Source/Programs/UnrealBuildAccelerator/Common/Private/UbaStorage.cpp
@@ -3,6 +3,7 @@
 #include "UbaStorage.h"
 #include "UbaBottleneck.h"
 #include "UbaCompressedObjFileHeader.h"
+#include "UbaConfig.h"
 #include "UbaFileAccessor.h"
 #include "UbaBinaryReaderWriter.h"
 #include "UbaDirectoryIterator.h"
@@ -52,6 +53,12 @@ namespace uba
 			free(slot);
 	}
 
+	void StorageCreateInfo::Apply(Config& config)
+	{
+		const ConfigTable& table = config.GetTable(TC("Storage"));
+		table.GetValueAsString(rootDir, TC("RootDir"));
+	}
+
 	void StorageImpl::CasEntryAccessed(CasEntry& casEntry)
 	{
 		bool hasMapping;
diff --git a/Engine/Source/Programs/UnrealBuildAccelerator/Common/Private/UbaStorageServer.cpp b/Engine/Source/Programs/UnrealBuildAccelerator/Common/Private/UbaStorageServer.cpp
--- a/Engine/Source/Programs/UnrealBuildAccelerator/Common/Private/UbaStorageServer.cpp
+++ b/Engine/Source/Programs/UnrealBuildAccelerator/Common/Private/UbaStorageServer.cpp
@@ -7,6 +7,11 @@
 
 namespace uba
 {
+	void StorageServerCreateInfo::Apply(Config& config)
+	{
+		StorageCreateInfo::Apply(config);
+	}
+
 	StorageServer::StorageServer(const StorageServerCreateInfo& info)
 	:	StorageImpl(info, TC("UbaStorageServer"))
 	,	m_server(info.server)
diff --git a/Engine/Source/Programs/UnrealBuildAccelerator/Common/Public/UbaConfig.h b/Engine/Source/Programs/UnrealBuildAccelerator/Common/Public/UbaConfig.h
--- a/Engine/Source/Programs/UnrealBuildAccelerator/Common/Public/UbaConfig.h
+++ b/Engine/Source/Programs/UnrealBuildAccelerator/Common/Public/UbaConfig.h
@@ -6,18 +6,24 @@
 
 namespace uba
 {
+	class Config;
 	class Logger;
 
+
 	class ConfigTable
 	{
 	public:
-		bool GetValueAsString(StringBufferBase& out, const tchar* key) const;
+		bool GetValueAsString(const tchar*& out, const tchar* key) const;
 		bool GetValueAsU32(u32& out, const tchar* key) const;
 		bool GetValueAsBool(bool& out, const tchar* key) const;
 
+	private:
+		ConfigTable* m_parent = nullptr;
 		UnorderedMap<TString, TString> m_values;
+		friend Config;
 	};
 
+
 	class Config : public ConfigTable
 	{
 	public:
@@ -28,6 +34,7 @@ namespace uba
 		const ConfigTable& GetTable(const tchar* name) const;
 
 		bool m_isLoaded = false;
+		ConfigTable m_globalTable;
 		UnorderedMap<TString, ConfigTable> m_tables;
 	};
 }
diff --git a/Engine/Source/Programs/UnrealBuildAccelerator/Common/Public/UbaLogWriter.h b/Engine/Source/Programs/UnrealBuildAccelerator/Common/Public/UbaLogWriter.h
--- a/Engine/Source/Programs/UnrealBuildAccelerator/Common/Public/UbaLogWriter.h
+++ b/Engine/Source/Programs/UnrealBuildAccelerator/Common/Public/UbaLogWriter.h
@@ -24,6 +24,6 @@ namespace uba
 		virtual void Log(LogEntryType type, const tchar* str, u32 strLen, const tchar* prefix = nullptr, u32 prefixLen = 0) = 0;
 	};
 
-	extern LogWriter& g_consoleLogWriter;
-	extern LogWriter& g_nullLogWriter;
+	UBA_API extern LogWriter& g_consoleLogWriter;
+	UBA_API extern LogWriter& g_nullLogWriter;
 }
diff --git a/Engine/Source/Programs/UnrealBuildAccelerator/Common/Public/UbaNetworkServer.h b/Engine/Source/Programs/UnrealBuildAccelerator/Common/Public/UbaNetworkServer.h
--- a/Engine/Source/Programs/UnrealBuildAccelerator/Common/Public/UbaNetworkServer.h
+++ b/Engine/Source/Programs/UnrealBuildAccelerator/Common/Public/UbaNetworkServer.h
@@ -13,6 +13,7 @@
 
 namespace uba
 {
+	class Config;
 	struct BinaryReader;
 	struct BinaryWriter;
 
@@ -28,6 +29,9 @@ namespace uba
 	struct NetworkServerCreateInfo
 	{
 		NetworkServerCreateInfo(LogWriter& w = g_consoleLogWriter) : logWriter(w) {}
+
+		void Apply(Config& config);
+
 		LogWriter& logWriter;
 		u32 workerCount = 0; // Zero means it will use the number of logical cores as worker count
 		u32 sendSize = SendDefaultSize;
@@ -49,7 +53,7 @@ namespace uba
 
 		bool StartListen(NetworkBackend& backend, u16 port = DefaultPort, const tchar* ip = nullptr, bool requiresCrypto = false); // Start listen for new connections/clients
 		void DisallowNewClients();	// Disallow new clients to connect but old clients can still create more connections
-		void DisconnectClients();				// Stops all listen and disconnect all active connections
+		void DisconnectClients();	// Disconnect all active connections
 
 		bool RegisterCryptoKey(const u8* cryptoKey128, u64 expirationTime = ~u64(0));
 
diff --git a/Engine/Source/Programs/UnrealBuildAccelerator/Common/Public/UbaStorage.h b/Engine/Source/Programs/UnrealBuildAccelerator/Common/Public/UbaStorage.h
--- a/Engine/Source/Programs/UnrealBuildAccelerator/Common/Public/UbaStorage.h
+++ b/Engine/Source/Programs/UnrealBuildAccelerator/Common/Public/UbaStorage.h
@@ -15,6 +15,7 @@
 
 namespace uba
 {
+	class Config;
 	class FileAccessor;
 	class Trace;
 	class WorkManager;
@@ -75,6 +76,9 @@ namespace uba
 	struct StorageCreateInfo
 	{
 		StorageCreateInfo(const tchar* rootDir_, LogWriter& w) : writer(w), rootDir(rootDir_) {}
+
+		void Apply(Config& config);
+
 		LogWriter& writer;
 		const tchar* rootDir;
 		u64 casCapacityBytes = 20llu * 1024 * 1024 * 1024;
diff --git a/Engine/Source/Programs/UnrealBuildAccelerator/Common/Public/UbaStorageServer.h b/Engine/Source/Programs/UnrealBuildAccelerator/Common/Public/UbaStorageServer.h
--- a/Engine/Source/Programs/UnrealBuildAccelerator/Common/Public/UbaStorageServer.h
+++ b/Engine/Source/Programs/UnrealBuildAccelerator/Common/Public/UbaStorageServer.h
@@ -10,6 +10,9 @@ namespace uba
 	struct StorageServerCreateInfo : StorageCreateInfo
 	{
 		StorageServerCreateInfo(NetworkServer& s, const tchar* rootDir_, LogWriter& writer) : StorageCreateInfo(rootDir_, writer), server(s) { workManager = &server; }
+
+		void Apply(Config& config);
+
 		NetworkServer& server;
 		bool allowFallback = true; // Will fallback to file system to recalculate cas if cas content does not exist anymore
 		bool writeRecievedCasFilesToDisk = false;
diff --git a/Engine/Source/Programs/UnrealBuildAccelerator/Core/Private/UbaApplicationRules.cpp b/Engine/Source/Programs/UnrealBuildAccelerator/Core/Private/UbaApplicationRules.cpp
--- a/Engine/Source/Programs/UnrealBuildAccelerator/Core/Private/UbaApplicationRules.cpp
+++ b/Engine/Source/Programs/UnrealBuildAccelerator/Core/Private/UbaApplicationRules.cpp
@@ -36,7 +36,7 @@ namespace uba
 		}
 		virtual bool KeepInMemory(const StringView& fileName, const tchar* systemTemp) const override
 		{
-			if (fileName.Contains(TC("\\vctip_")) != 0)
+			if (fileName.Contains(TC("\\vctip_")))
 				return true;
 			if (fileName.Contains(systemTemp))
 				return true;
diff --git a/Engine/Source/Programs/UnrealBuildAccelerator/Core/Public/UbaBase.h b/Engine/Source/Programs/UnrealBuildAccelerator/Core/Public/UbaBase.h
--- a/Engine/Source/Programs/UnrealBuildAccelerator/Core/Public/UbaBase.h
+++ b/Engine/Source/Programs/UnrealBuildAccelerator/Core/Public/UbaBase.h
@@ -20,12 +20,12 @@ namespace uba
 	template<class T> const T& Max(const T& a, const T& b) { return (b > a) ? b : a; }
 
 	#if PLATFORM_WINDOWS
-	inline constexpr bool IsWindows = true;
+	constexpr bool IsWindows = true;
 	using tchar = wchar_t;
 	#define TC(x) L##x
 	#define PERCENT_HS L"%hs"
 	#else
-	inline constexpr bool IsWindows = false;
+	constexpr bool IsWindows = false;
 	using tchar = char;
 	#define TC(x) x
 	#define PERCENT_HS "%s"
diff --git a/Engine/Source/Programs/UnrealBuildAccelerator/Core/Public/UbaStringBuffer.h b/Engine/Source/Programs/UnrealBuildAccelerator/Core/Public/UbaStringBuffer.h
--- a/Engine/Source/Programs/UnrealBuildAccelerator/Core/Public/UbaStringBuffer.h
+++ b/Engine/Source/Programs/UnrealBuildAccelerator/Core/Public/UbaStringBuffer.h
@@ -51,7 +51,7 @@ namespace uba
 		bool StartsWith(const tchar* str, bool ignoreCase = true) const { return uba::StartsWith(data, str, ignoreCase); }
 		bool EndsWith(const tchar* value, bool ignoreCase = true) const { return uba::EndsWith(data, count, value, ignoreCase); }
 		bool Contains(tchar c) const;
-		bool Contains(const tchar* str, bool ignoreCase = true) const { return uba::Contains(data, str, ignoreCase); }
+		bool Contains(const tchar* str, bool ignoreCase = true, const tchar** pos = nullptr) const { return uba::Contains(data, str, ignoreCase, pos); }
 		bool Equals(const tchar* str, bool ignoreCase = true) const { return uba::Equals(data, str, ignoreCase); }
 		const tchar* First(tchar c, u64 offset = 0) const;
 		const tchar* Last(tchar c, u64 offset = 0) const;
diff --git a/Engine/Source/Programs/UnrealBuildAccelerator/Test/Private/UbaTestConfig.h b/Engine/Source/Programs/UnrealBuildAccelerator/Test/Private/UbaTestConfig.h
--- a/Engine/Source/Programs/UnrealBuildAccelerator/Test/Private/UbaTestConfig.h
+++ b/Engine/Source/Programs/UnrealBuildAccelerator/Test/Private/UbaTestConfig.h
@@ -9,8 +9,10 @@ namespace uba
 	bool TestConfig(Logger& logger, const StringBufferBase& rootDir)
 	{
 		static const char* configText =
+			"RootDir = \"e:\\foo\"\r\n"
 			"[CacheClient]\r\n"
 			"UseDirectoryPreparsing = true\r\n"
+			"# Comment = true\r\n"
 			"";
 
 		Config config;
@@ -23,6 +25,13 @@ namespace uba
 			return false;
 		if (test != true)
 			return false;
+		const tchar* str = nullptr;
+		if (!table.GetValueAsString(str, TC("RootDir")))
+			return false;
+		if (TStrcmp(str, TC("e:\\foo")) != 0)
+			return false;
+		if (table.GetValueAsBool(test, TC("Comment")))
+			return false;
 		return true;
 	}
 }
\ No newline at end of file
diff --git a/Engine/Source/Programs/UnrealBuildAccelerator/Visualizer/Private/UbaVisualizer.cpp b/Engine/Source/Programs/UnrealBuildAccelerator/Visualizer/Private/UbaVisualizer.cpp
--- a/Engine/Source/Programs/UnrealBuildAccelerator/Visualizer/Private/UbaVisualizer.cpp
+++ b/Engine/Source/Programs/UnrealBuildAccelerator/Visualizer/Private/UbaVisualizer.cpp
@@ -1381,6 +1381,7 @@ namespace uba
 
 			logger.AddSpace(2);
 			logger.Info(L"  %ls", process.description.c_str());
+			logger.Info(L"  ProcessId: %6u", process.id);
 			logger.Info(L"  Start:    %7ls", TimeToText(process.start, true).str);
 			logger.Info(L"  Duration: %7ls", TimeToText(duration, true).str);
 			if (hasExited && process.exitCode != 0)
@@ -2177,6 +2178,7 @@ namespace uba
 	{
 		bool hasExited = process.stop != ~u64(0);
 		out.Info(L"  %ls", process.description.c_str());
+		out.Info(L"  ProcessId: %u", process.id);
 		out.Info(L"  Start:     %ls", TimeToText(process.start, true).str);
 		if (hasExited)
 			out.Info(L"  Duration:  %ls", TimeToText(process.stop - process.start, true).str);
