[UBA]
* Added error code to NetworkMessage to be able to diagnose "failed to connect" errors in cache client

diff --git a/Engine/Source/Programs/UnrealBuildAccelerator/Common/Private/UbaNetworkClient.cpp b/Engine/Source/Programs/UnrealBuildAccelerator/Common/Private/UbaNetworkClient.cpp
--- a/Engine/Source/Programs/UnrealBuildAccelerator/Common/Private/UbaNetworkClient.cpp
+++ b/Engine/Source/Programs/UnrealBuildAccelerator/Common/Private/UbaNetworkClient.cpp
@@ -241,7 +241,7 @@ namespace uba
 
 		if (messageSize == ErrorSize)
 		{
-			msg->m_error = true;
+			msg->m_error = 1;
 			msg->Done();
 			return true;
 		}
@@ -271,7 +271,7 @@ namespace uba
 	{
 		auto& msg = *(NetworkMessage*)bodyContext;
 		if (recvError)
-			msg.m_error = true;
+			msg.m_error = 2;
 		msg.Done();
 		return true;
 	}
@@ -428,7 +428,7 @@ namespace uba
 		{
 			if (m && m->m_connection == &connection)
 			{
-				m->m_error = true;
+				m->m_error = 3;
 				m->Done(false);
 			}
 			++messageId;
@@ -440,7 +440,10 @@ namespace uba
 		SCOPED_READ_LOCK(m_connectionsLock, connectionLock);
 		SCOPED_WRITE_LOCK(m_connectionsItLock, connectionItLock);
 		if (m_connectionsIt == m_connections.end())
+		{
+			message.m_error = 6;
 			return false;
+		}
 
 		Connection& connection = *m_connectionsIt;
 		++m_connectionsIt;
@@ -466,7 +469,10 @@ namespace uba
 				if (m_availableMessageIds.empty())
 				{
 					if (!connection.connected)
+					{
+						message.m_error = 7;
 						return false;
+					}
 
 					if (m_activeMessageIdMax == 65534)
 					{
@@ -518,6 +524,7 @@ namespace uba
 			TimerScope ts(m_encryptTimer);
 			if (!Crypto::Encrypt(m_logger, m_cryptoKey, data + SendHeaderSize, bodySize))
 			{
+				message.m_error = 8;
 				OnDisconnected(connection);
 				return false;
 			}
@@ -530,6 +537,7 @@ namespace uba
 			TimerScope ts(m_sendTimer);
 			if (!connection.backend->Send(m_logger, connection.backendConnection, data, sendSize, message.m_sendContext))
 			{
+				message.m_error = 9;
 				OnDisconnected(connection);
 				return false;
 			}
@@ -544,13 +552,13 @@ namespace uba
 			if (!gotResponse.IsSet(timeoutMs))
 			{
 				m_logger.Error(TC("Timed out after 10 minutes waiting for message response from server."));
-				message.m_error = true;
+				message.m_error = 4;
 			}
 			else if (m_cryptoKey && !message.m_error && message.m_responseSize)
 			{
 				TimerScope ts(m_decryptTimer);
 				if (!Crypto::Decrypt(m_logger, m_cryptoKey, (u8*)message.m_response, message.m_responseSize))
-					message.m_error = true;
+					message.m_error = 5;
 			}
 		}
 		return !message.m_error;
@@ -636,7 +644,10 @@ namespace uba
 			UBA_ASSERT(!response.GetPosition());
 			TimerScope ts(m_client->m_decryptTimer);
 			if (!Crypto::Decrypt(m_client->m_logger, m_client->m_cryptoKey, (u8*)m_response, m_responseSize))
+			{
+				m_error = 10;
 				return false;
+			}
 		}
 		response.SetSize(response.GetPosition() + m_responseSize);
 		return true;
@@ -666,6 +677,6 @@ namespace uba
 			returnId();
 		}
 		if (hasId)
-			m_doneFunc(m_error, m_doneUserData);
+			m_doneFunc(m_error != 0, m_doneUserData);
 	}
 }
diff --git a/Engine/Source/Programs/UnrealBuildAccelerator/Common/Public/UbaNetworkMessage.h b/Engine/Source/Programs/UnrealBuildAccelerator/Common/Public/UbaNetworkMessage.h
--- a/Engine/Source/Programs/UnrealBuildAccelerator/Common/Public/UbaNetworkMessage.h
+++ b/Engine/Source/Programs/UnrealBuildAccelerator/Common/Public/UbaNetworkMessage.h
@@ -23,6 +23,8 @@ namespace uba
 		bool SendAsync(BinaryReader& response, DoneFunc* func, void* userData); // Try to minimize work in GetResponseFunc since it is running on receiving thread
 		bool ProcessAsyncResults(BinaryReader& response); // Note, this must be called after GetResponseFunc has been called
 
+		u32 GetError() const { return m_error; }
+
 	private:
 		void Done(bool shouldLock = true);
 
@@ -35,7 +37,7 @@ namespace uba
 		u32 m_responseSize = 0;
 		u32 m_responseCapacity = 0;
 		u16 m_id = 0;
-		Atomic<bool> m_error;
+		Atomic<u32> m_error;
 		NetworkClient::Connection* m_connection = nullptr;
 		friend NetworkClient;
 		NetworkMessage(const NetworkMessage&) = delete;
