[UBA]
* Added GetDescription to ProcessStartInfo
* Added info to session description if short filenames are enabled

diff --git a/Engine/Source/Programs/UnrealBuildAccelerator/Common/Private/UbaSession.cpp b/Engine/Source/Programs/UnrealBuildAccelerator/Common/Private/UbaSession.cpp
--- a/Engine/Source/Programs/UnrealBuildAccelerator/Common/Private/UbaSession.cpp
+++ b/Engine/Source/Programs/UnrealBuildAccelerator/Common/Private/UbaSession.cpp
@@ -51,6 +51,18 @@ namespace uba
 	ProcessStartInfo::~ProcessStartInfo() = default;
 	ProcessStartInfo::ProcessStartInfo(const ProcessStartInfo&) = default;
 
+	const tchar* ProcessStartInfo::GetDescription() const
+	{
+		if (description && *description)
+			return description;
+		const tchar* d = application;
+		if (const tchar* lps = TStrrchr(d, PathSeparator))
+			d = lps + 1;
+		if (const tchar* lps2 = TStrrchr(d, NonPathSeparator))
+			d = lps2 + 1;
+		return d;
+	}
+
 	ProcessHandle::ProcessHandle()
 	:	m_process(nullptr)
 	{
@@ -921,7 +933,6 @@ namespace uba
 		if (isSystem && IsKnownSystemFile(applicationName))
 			return true;
 
-
 		out.push_back({ library, temp3.data, attr, isSystem });
 
 		bool result = true;
@@ -2362,7 +2373,8 @@ namespace uba
 		{
 			DWORD value = 0;
 			DWORD valueSize = 4;
-			LSTATUS res = RegGetValueW(HKEY_LOCAL_MACHINE, TC("SYSTEM\\CurrentControlSet\\Control\\FileSystem"), TC("NtfsDisableLastAccessUpdate"), RRF_RT_REG_DWORD, NULL, &value, &valueSize);
+			const tchar* fsKey = TC("SYSTEM\\CurrentControlSet\\Control\\FileSystem");
+			LSTATUS res = RegGetValueW(HKEY_LOCAL_MACHINE, fsKey, TC("NtfsDisableLastAccessUpdate"), RRF_RT_REG_DWORD, NULL, &value, &valueSize);
 			if (res != ERROR_SUCCESS)
 			{
 				m_logger.Detail(TC("Failed to retreive ntfs registry key (%i)"), res);
@@ -2373,6 +2385,21 @@ namespace uba
 				if (lastAccessSettingsValue == 0 || lastAccessSettingsValue == 2)
 					out.Append(TC(" NtfsLastAccessEnabled"));
 			}
+			value = 0;
+			res = RegGetValueW(HKEY_LOCAL_MACHINE, fsKey, TC("NtfsDisable8dot3NameCreation"), RRF_RT_REG_DWORD, NULL, &value, &valueSize);
+			if (res == ERROR_SUCCESS)
+				if (value == 0)
+					out.Append(TC(" NtfsShortNamesEnabled"));
+		}
+		else
+		{
+			StringBuffer<> testDir;
+			testDir.Append(m_rootDir).Append(TC("UbaTestShortNames"));
+			::RemoveDirectory(testDir.data);
+			wchar_t shortName[1024];
+			if (::CreateDirectoryW(testDir.data, NULL))
+				if (GetShortPathName(testDir.data, shortName, 1024) != 0 && !Contains(shortName, TC("UbaTestShortNames")))
+					out.Append(TC(" NtfsShortNamesEnabled"));
 		}
 		#endif
 
diff --git a/Engine/Source/Programs/UnrealBuildAccelerator/Common/Public/UbaProcessStartInfo.h b/Engine/Source/Programs/UnrealBuildAccelerator/Common/Public/UbaProcessStartInfo.h
--- a/Engine/Source/Programs/UnrealBuildAccelerator/Common/Public/UbaProcessStartInfo.h
+++ b/Engine/Source/Programs/UnrealBuildAccelerator/Common/Public/UbaProcessStartInfo.h
@@ -15,6 +15,8 @@ namespace uba
 		UBA_API ~ProcessStartInfo();
 		UBA_API ProcessStartInfo(const ProcessStartInfo&);
 
+		UBA_API const tchar* GetDescription() const;	// Use this instead since it will use application name if description is not set
+
 		const tchar* application = TC("");		// Application name, cl.exe etc. Use full path
 		const tchar* arguments = TC("");		// Arguments. Should not include application name
 		const tchar* workingDir = TC("");		// Working directory. Use full path
