[UBA]
* CacheClient - Added error information when failing to connect to server
* Changed description -> GetDescription()

diff --git a/Engine/Source/Programs/UnrealBuildAccelerator/Common/Private/UbaCacheClient.cpp b/Engine/Source/Programs/UnrealBuildAccelerator/Common/Private/UbaCacheClient.cpp
--- a/Engine/Source/Programs/UnrealBuildAccelerator/Common/Private/UbaCacheClient.cpp
+++ b/Engine/Source/Programs/UnrealBuildAccelerator/Common/Private/UbaCacheClient.cpp
@@ -88,9 +88,10 @@ namespace uba
 					NetworkMessage msg(m_client, CacheServiceId, CacheMessageType_Connect, writer);
 					writer.WriteU32(CacheNetworkVersion);
 					StackBinaryReader<1024> reader;
+					u64 sendTime = GetTime();
 					if (!msg.Send(reader))
 					{
-						m_logger.Info(TC("Failed to connect to cache server. Version mismatch?"));
+						m_logger.Info(TC("Failed to send connect message to cache server (%u). Version mismatch? (%s)"), msg.GetError(), TimeToText(GetTime() - sendTime).str);
 						return;
 					}
 					bool success = reader.ReadBool();
@@ -136,7 +137,7 @@ namespace uba
 		if (cmdKey == CasKeyZero)
 		{
 			#if UBA_LOG_WRITE_CACHE_INFO
-			m_logger.Info(TC("WRITECACHE FAIL: %s"), info.description);
+			m_logger.Info(TC("WRITECACHE FAIL: %s"), info.GetDescription());
 			#endif
 			return false;
 		}
@@ -178,7 +179,7 @@ namespace uba
 			
 			if (path.count < 2)
 			{
-				m_logger.Info(TC("Got messed up path from caller to WriteToCache: %s (%s)"), path.data, info.description);
+				m_logger.Info(TC("Got messed up path from caller to WriteToCache: %s (%s)"), path.data, info.GetDescription());
 				success = false;
 			}
 
@@ -206,7 +207,7 @@ namespace uba
 			}
 			else if (path[path.count-1] == ':')
 			{
-				m_logger.Info(TC("GOT UNKNOWN RELATIVE PATH: %s (%s)"), path.data, info.description);
+				m_logger.Info(TC("GOT UNKNOWN RELATIVE PATH: %s (%s)"), path.data, info.GetDescription());
 				success = false;
 				continue;
 			}
@@ -217,7 +218,7 @@ namespace uba
 				auto root = rootPaths.FindRoot(path);
 				if (!root)
 				{
-					m_logger.Info(TC("FILE WITHOUT ROOT: %s (%s)"), path.data, info.description);
+					m_logger.Info(TC("FILE WITHOUT ROOT: %s (%s)"), path.data, info.GetDescription());
 					success = false;
 					continue;
 				}
@@ -281,7 +282,7 @@ namespace uba
 					
 					if (shouldValidate && FileExists(m_logger, path.data))
 					{
-						m_logger.Warning(TC("CasDb claims file %s does not exist but it does! Will not populate cache for %s"), path.data, info.description); 
+						m_logger.Warning(TC("CasDb claims file %s does not exist but it does! Will not populate cache for %s"), path.data, info.GetDescription()); 
 						return false;
 					}
 
@@ -296,7 +297,7 @@ namespace uba
 					FileAccessor fa(m_logger, path.data);
 					if (!fa.OpenMemoryRead())
 					{
-						m_logger.Warning(TC("CasDb claims file %s does exist but can't open it. Will not populate cache for %s"), path.data, info.description); 
+						m_logger.Warning(TC("CasDb claims file %s does exist but can't open it. Will not populate cache for %s"), path.data, info.GetDescription()); 
 						return false;
 					}
 
@@ -321,7 +322,7 @@ namespace uba
 
 						auto ToString = [](bool b) { return b ? TC("true") : TC("false"); };
 						m_logger.Warning(TC("CasDb claims file %s has caskey %s but recalculating it gives us %s (FileEntry: %llu/%llu/%s, Real: %llu/%llu). Will not populate cache for %s"),
-							path.data, CasKeyString(oldKey).str, CasKeyString(newKey).str, fileEntry.size, fileEntry.lastWritten, ToString(fileEntry.verified), fileSize, fileInfo.lastWriteTime, info.description);
+							path.data, CasKeyString(oldKey).str, CasKeyString(newKey).str, fileEntry.size, fileEntry.lastWritten, ToString(fileEntry.verified), fileSize, fileInfo.lastWriteTime, info.GetDescription());
 						return false;
 					}
 				}
@@ -335,7 +336,7 @@ namespace uba
 			return false;
 
 		if (outputsStringToCasKey.empty())
-			m_logger.Warning(TC("NO OUTPUTS FROM process %s"), info.description); 
+			m_logger.Warning(TC("NO OUTPUTS FROM process %s"), info.GetDescription()); 
 
 		// Make sure server has enough of the path table to be able to resolve offsets from cache entry
 		if (!SendPathTable(bucket, requiredPathTableSize))
@@ -352,7 +353,7 @@ namespace uba
 
 		#if UBA_LOG_WRITE_CACHE_INFO
 		m_logger.BeginScope();
-		m_logger.Info(TC("WRITECACHE: %s -> %u %s"), info.description, bucketId, CasKeyString(cmdKey).str);
+		m_logger.Info(TC("WRITECACHE: %s -> %u %s"), info.GetDescription(), bucketId, CasKeyString(cmdKey).str);
 		#if UBA_LOG_WRITE_CACHE_INFO == 2
 		for (auto& kv : inputsStringToCasKey)
 		{
@@ -404,7 +405,7 @@ namespace uba
 		u8 memory[SendMaxSize];
 
 		u32 fetchId = m_session.CreateProcessId();
-		m_session.GetTrace().CacheBeginFetch(fetchId, info.description);
+		m_session.GetTrace().CacheBeginFetch(fetchId, info.GetDescription());
 		bool success = false;
 		auto tg = MakeGuard([&]()
 			{
@@ -444,7 +445,7 @@ namespace uba
 		auto mg = MakeGuard([&]()
 			{
 				if (!success || UBA_LOG_FETCH_CACHE_INFO == 2)
-					m_logger.Info(TC("FETCHCACHE %s: %s -> %u %s (%u)"), success ? TC("SUCC") : TC("FAIL"), info.description, bucketId, CasKeyString(cmdKey).str, entryCount);
+					m_logger.Info(TC("FETCHCACHE %s: %s -> %u %s (%u)"), success ? TC("SUCC") : TC("FAIL"), info.GetDescription(), bucketId, CasKeyString(cmdKey).str, entryCount);
 			});
 		#endif
 
@@ -482,7 +483,7 @@ namespace uba
 
 				if (!GetLocalPathAndCasKey(bucket, rootPaths, path, cacheCasKey, bucket.serverCasKeyTable, bucket.serverPathTable, casKeyOffset))
 					return false;
-				UBA_ASSERTF(IsCompressed(cacheCasKey), TC("Cache entry for %s has uncompressed cache key for path %s (%s)"), info.description, path.data, CasKeyString(cacheCasKey).str);
+				UBA_ASSERTF(IsCompressed(cacheCasKey), TC("Cache entry for %s has uncompressed cache key for path %s (%s)"), info.GetDescription(), path.data, CasKeyString(cacheCasKey).str);
 
 				if (IsNormalized(cacheCasKey)) // Need to normalize caskey for these files since they contain absolute paths
 				{
@@ -677,7 +678,7 @@ namespace uba
 				}
 
 #if UBA_OLD_TEST
-				UBA_ASSERTF(testIsMatch, TC("%s"), info.description);
+				UBA_ASSERTF(testIsMatch, TC("%s"), info.GetDescription());
 				UBA_ASSERTF(matchingId == entryId, TC("%u vs %u"), matchingId, entryId);
 #endif
 
@@ -725,7 +726,7 @@ namespace uba
 						MemoryBlock normalizedBlock(4*1024*1024);
 						bool destinationIsCompressed = false;
 						if (!fetcher.RetrieveFile(logger, m_client, casKey, path.data, destinationIsCompressed, &normalizedBlock))
-							return logger.Error(TC("Failed to download cache output for %s"), info.description);
+							return logger.Error(TC("Failed to download cache output for %s"), info.GetDescription());
 
 						MemoryBlock localBlock(4*1024*1024);
 
@@ -744,7 +745,7 @@ namespace uba
 							u8 rootIndex = fileStart[rootOffset] - RootPaths::RootStartByte;
 							const TString& root = rootPaths.GetRoot(rootIndex);
 							if (root.empty())
-								return logger.Error(TC("Cache entry uses root path index %u which is not set for this startupinfo (%s)"), rootIndex, info.description);
+								return logger.Error(TC("Cache entry uses root path index %u which is not set for this startupinfo (%s)"), rootIndex, info.GetDescription());
 
 							#if PLATFORM_WINDOWS
 							StringBuffer<> pathTemp;
@@ -771,13 +772,13 @@ namespace uba
 						if (useFileMapping)
 						{
 							if (!destFile.CreateMemoryWrite(false, DefaultAttributes(), localBlock.writtenSize))
-								return logger.Error(TC("Failed to create file for cache output %s for %s"), path.data, info.description);
+								return logger.Error(TC("Failed to create file for cache output %s for %s"), path.data, info.GetDescription());
 							MapMemoryCopy(destFile.GetData(), localBlock.memory, localBlock.writtenSize);
 						}
 						else
 						{
 							if (!destFile.CreateWrite())
-								return logger.Error(TC("Failed to create file for cache output %s for %s"), path.data, info.description);
+								return logger.Error(TC("Failed to create file for cache output %s for %s"), path.data, info.GetDescription());
 							if (!destFile.Write(localBlock.memory, localBlock.writtenSize))
 								return false;
 						}
@@ -792,7 +793,7 @@ namespace uba
 						DowngradedLogger logger(m_logger.m_writer, TC("UbaCacheClientDownload"));
 						bool destinationIsCompressed = IsFileCompressed(info, path);
 						if (!fetcher.RetrieveFile(logger, m_client, casKey, path.data, destinationIsCompressed))
-							return logger.Error(TC("Failed to download cache output %s for %s"), path.data, info.description);
+							return logger.Error(TC("Failed to download cache output %s for %s"), path.data, info.GetDescription());
 					}
 
 					cacheStats.fetchBytesRaw += fetcher.sizeOnDisk;
@@ -810,11 +811,11 @@ namespace uba
 		}
 
 #if UBA_OLD_TEST
-		UBA_ASSERTF(!testIsMatch, TC("%s"), info.description);
+		UBA_ASSERTF(!testIsMatch, TC("%s"), info.GetDescription());
 #endif
 
 		for (auto& miss : misses)
-			m_logger.Info(TC("Cache miss on %s because of mismatch of %s (entry: %u, local: %s cache: %s)"), info.description, miss.path.data(), miss.entryIndex, CasKeyString(miss.local).str, CasKeyString(miss.cache).str);
+			m_logger.Info(TC("Cache miss on %s because of mismatch of %s (entry: %u, local: %s cache: %s)"), info.GetDescription(), miss.path.data(), miss.entryIndex, CasKeyString(miss.local).str, CasKeyString(miss.cache).str);
 
 		return false;
 	}
@@ -1169,7 +1170,7 @@ namespace uba
 
 		// Add arguments list to key
 		auto hashString = [&](const tchar* str, u64 strLen, u32 rootPos) { hasher.Update(str, strLen*sizeof(tchar)); };
-		if (!rootPaths.NormalizeString(m_logger, info.arguments, TStrlen(info.arguments), hashString, TC("CmdKey "), info.description))
+		if (!rootPaths.NormalizeString(m_logger, info.arguments, TStrlen(info.arguments), hashString, TC("CmdKey "), info.GetDescription()))
 			return CasKeyZero;
 
 		// Add content of rsp file to key (This will cost a bit of perf since we need to normalize.. should this be part of key?)
