[UBA]
* Config - Fixed reading of values (non-bool, non-string)

diff --git a/Engine/Source/Programs/UnrealBuildAccelerator/Common/Private/UbaConfig.cpp b/Engine/Source/Programs/UnrealBuildAccelerator/Common/Private/UbaConfig.cpp
--- a/Engine/Source/Programs/UnrealBuildAccelerator/Common/Private/UbaConfig.cpp
+++ b/Engine/Source/Programs/UnrealBuildAccelerator/Common/Private/UbaConfig.cpp
@@ -16,8 +16,11 @@ namespace uba
 
 	bool ConfigTable::GetValueAsU32(u32& out, const tchar* key) const
 	{
-		UBA_ASSERT(false);
-		return false;
+		auto findIt = m_values.find(key);
+		if (findIt == m_values.end())
+			return m_parent ? m_parent->GetValueAsU32(out, key) : false;
+		StringBuffer<> buf(findIt->second);
+		return buf.Parse(out);
 	}
 
 	bool ConfigTable::GetValueAsBool(bool& out, const tchar* key) const
@@ -165,7 +168,7 @@ namespace uba
 				}
 				else
 				{
-					token = consumeIdentifier(value);
+					token = consumeLine(value, ' ');
 				}
 
 				activeTable->m_values[key.data] = value.data;
