[UBA]
* Fixed bug in coordinator wrapper where binary path was not properly concatenated

diff --git a/Engine/Source/Programs/UnrealBuildAccelerator/Common/Public/UbaCoordinatorWrapper.h b/Engine/Source/Programs/UnrealBuildAccelerator/Common/Public/UbaCoordinatorWrapper.h
--- a/Engine/Source/Programs/UnrealBuildAccelerator/Common/Public/UbaCoordinatorWrapper.h
+++ b/Engine/Source/Programs/UnrealBuildAccelerator/Common/Public/UbaCoordinatorWrapper.h
@@ -29,7 +29,9 @@ namespace uba
 			if (!*coordinatorType)
 				return false;
 
-			StringBuffer<128> coordinatorBin;
+			StringBuffer<128> coordinatorBin(info.binariesDir);
+			coordinatorBin.EnsureEndsWithSlash();
+
 			#if PLATFORM_WINDOWS
 			coordinatorBin.Append(TC("UbaCoordinator")).Append(coordinatorType).Append(TC(".dll"));
 			#else
@@ -49,10 +51,6 @@ namespace uba
 			if (!m_destroyCoordinator)
 				return logger.Error(TC("Failed to find UbaDestroyCoordinator function inside %s (%s)"), coordinatorBin.data, LastErrorToText().data);
 
-			StringBuffer<512> binariesDir;
-			if (!GetDirectoryOfCurrentModule(logger, binariesDir))
-				return false;
-
 			m_coordinator = createCoordinator(info);
 			if (!m_coordinator)
 				return false;
