[UBA]
* UbaCli - Added support for cmd.exe style command lines with weird combos of quotes

diff --git a/Engine/Source/Programs/UnrealBuildAccelerator/Cli/Private/UbaCli.cpp b/Engine/Source/Programs/UnrealBuildAccelerator/Cli/Private/UbaCli.cpp
--- a/Engine/Source/Programs/UnrealBuildAccelerator/Cli/Private/UbaCli.cpp
+++ b/Engine/Source/Programs/UnrealBuildAccelerator/Cli/Private/UbaCli.cpp
@@ -197,33 +197,46 @@ namespace uba
 		for (int i=1; i!=argc; ++i)
 		{
 			StringBuffer<> name;
-			StringBuffer<> value;
+			StringBuffer<32*1024> value;
+			const tchar* arg = argv[i];
 
-			if (const tchar* equals = TStrchr(argv[i],'='))
+			if (const tchar* equals = TStrchr(arg,'='))
 			{
-				name.Append(argv[i], equals - argv[i]);
+				name.Append(arg, equals - arg);
 				value.Append(equals+1);
 			}
 			else
 			{
-				name.Append(argv[i]);
+				name.Append(arg);
 			}
 
 			if (!application.empty())
 			{
 				if (!arguments.empty())
 					arguments += ' ';
-				bool hasSpace = TStrchr(argv[i], ' ');
+				TString argTemp;
+				bool hasSpace = TStrchr(arg, ' ');
 				if (hasSpace)
+				{
+					argTemp = arg;
+					size_t index = 0;
+					while (true) {
+							index = argTemp.find('\"', index);
+							if (index == std::string::npos) break;
+							argTemp.replace(index, 1, TC("\\\""));
+							index += 2;
+					}
+					arg = argTemp.c_str();
 					arguments += TC("\"");
-				arguments += argv[i];
+				}
+				arguments += arg;
 				if (hasSpace)
 					arguments += TC("\"");
 				continue;
 			}
 			if (commandType != CommandType_NotSet)
 			{
-				application = argv[i];
+				application = arg;
 			}
 			else if (name.Equals(TC("local")))
 			{
@@ -510,8 +523,7 @@ namespace uba
 			if (application.empty())
 				return PrintHelp(TC("No executable provided"));
 
-			bool isAbsolute = IsWindows ? application[1] == ':' : application[0] == '/';
-			if (!isAbsolute)
+			if (!IsAbsolutePath(application))
 			{
 				StringBuffer<> fullApplicationName;
 				if (!SearchPathForFile(logger, fullApplicationName, application.c_str(), currentDir.data))
