[UBA]
* Improved logging when finding paths without roots for cache

diff --git a/Engine/Source/Programs/UnrealBuildAccelerator/Common/Public/UbaRootPaths.h b/Engine/Source/Programs/UnrealBuildAccelerator/Common/Public/UbaRootPaths.h
--- a/Engine/Source/Programs/UnrealBuildAccelerator/Common/Public/UbaRootPaths.h
+++ b/Engine/Source/Programs/UnrealBuildAccelerator/Common/Public/UbaRootPaths.h
@@ -82,6 +82,10 @@ namespace uba
 			auto root = FindRoot(path);
 			if (!root)
 			{
+				if (auto lastQuote = path.Last('\"'))
+					path.Resize(lastQuote - path.data);
+				if (auto lineEnd = path.Last('\r'))
+					path.Resize(lineEnd - path.data);
 				logger.Info(TC("PATH WITHOUT ROOT: %s (inside %s%s at offset %u)"), path.data, hint, hint2, destPos);
 				return false;
 			}
