[UBA]
* Fixed so file mapping table is updated in Rpc_UpdateTables

diff --git a/Engine/Source/Programs/UnrealBuildAccelerator/Detours/Private/UbaDetoursFileMappingTable.cpp b/Engine/Source/Programs/UnrealBuildAccelerator/Detours/Private/UbaDetoursFileMappingTable.cpp
--- a/Engine/Source/Programs/UnrealBuildAccelerator/Detours/Private/UbaDetoursFileMappingTable.cpp
+++ b/Engine/Source/Programs/UnrealBuildAccelerator/Detours/Private/UbaDetoursFileMappingTable.cpp
@@ -160,8 +160,10 @@ namespace uba
 		writer.Flush();
 		BinaryReader reader;
 		u32 directoryTableSize = reader.ReadU32();
+		u32 fileMappingTableSize = reader.ReadU32();
 		pcs.Leave();
 		g_directoryTable.ParseDirectoryTable(directoryTableSize);
+		g_mappedFileTable.Parse(fileMappingTableSize);
 		DEBUG_LOG_PIPE(L"UpdateTables", L"");
 	}
 
@@ -218,8 +220,7 @@ namespace uba
 	void Rpc_GetFullFileName(const tchar*& path, u64& pathLen, StringBufferBase& tempBuf, bool useVirtualName)
 	{
 		StringKey fileNameKey;
-		bool isAbsolute = IsWindows ? (pathLen > 1 && path[1] == ':') : (pathLen > 0 && path[0] == '/');
-		if (isAbsolute)
+		if (IsAbsolutePath(path))
 		{
 			FixPath(tempBuf, path);
 			if (CaseInsensitiveFs)
