[UBA]
* Fixed so default constructed StringView can still be used to do EndsWith and Contains

diff --git a/Engine/Source/Programs/UnrealBuildAccelerator/Core/Public/UbaStringBuffer.h b/Engine/Source/Programs/UnrealBuildAccelerator/Core/Public/UbaStringBuffer.h
--- a/Engine/Source/Programs/UnrealBuildAccelerator/Core/Public/UbaStringBuffer.h
+++ b/Engine/Source/Programs/UnrealBuildAccelerator/Core/Public/UbaStringBuffer.h
@@ -94,7 +94,7 @@ namespace uba
 
 	struct StringView
 	{
-		StringView() : data(nullptr), count(0) {}
+		StringView() : data(TC("")), count(0) {}
 		StringView(const tchar* d, u32 c) : data(d), count(c) {}
 		StringView(const StringBufferBase& sb) : data(sb.data), count(sb.count) {}
 		StringView(const TString& str) : data(str.data()), count(u32(str.size())) {}
