[UBA]
* Added support for GetLongPathName and RemoveDirectory for remote helpers

diff --git a/Engine/Source/Programs/UnrealBuildAccelerator/Common/Private/UbaProcess.cpp b/Engine/Source/Programs/UnrealBuildAccelerator/Common/Private/UbaProcess.cpp
--- a/Engine/Source/Programs/UnrealBuildAccelerator/Common/Private/UbaProcess.cpp
+++ b/Engine/Source/Programs/UnrealBuildAccelerator/Common/Private/UbaProcess.cpp
@@ -253,7 +253,10 @@ namespace uba
 					if (data[0] != 'M' || data[1] != 'Z')
 						is64Bit = false;
 					else
-						is64Bit = *(u32*)(data + 0x3c) == 0x50450000;
+					{
+						u32 offset = *(u32*)(data + 0x3c);
+						is64Bit = *(u32*)(data + offset) == 0x00004550;
+					}
 				}
 
 				if (!is64Bit)
@@ -710,6 +713,17 @@ namespace uba
 		return true;
 	}
 
+	bool ProcessImpl::HandleGetLongPathName(BinaryReader& reader, BinaryWriter& writer)
+	{
+		GetLongPathNameMessage msg { *this };
+		reader.ReadString(msg.fileName);
+		GetLongPathNameResponse response;
+		m_messageSuccess = m_session.GetLongPathName(response, msg) && m_messageSuccess;
+		writer.WriteU32(response.errorCode);
+		writer.WriteString(response.fileName);
+		return true;
+	}
+
 	bool ProcessImpl::HandleCloseFile(BinaryReader& reader, BinaryWriter& writer)
 	{
 		CloseFileMessage msg { *this };
@@ -796,6 +810,20 @@ namespace uba
 		m_messageSuccess = m_session.CreateDirectory(response, msg) && m_messageSuccess;
 		writer.WriteBool(response.result);
 		writer.WriteU32(response.errorCode);
+		writer.WriteU32(response.directoryTableSize);
+		return true;
+	}
+
+	bool ProcessImpl::HandleRemoveDirectory(BinaryReader& reader, BinaryWriter& writer)
+	{
+		RemoveDirectoryMessage msg;
+		msg.nameKey = reader.ReadStringKey();
+		reader.ReadString(msg.name);
+		RemoveDirectoryResponse response;
+		m_messageSuccess = m_session.RemoveDirectory(response, msg) && m_messageSuccess;
+		writer.WriteBool(response.result);
+		writer.WriteU32(response.errorCode);
+		writer.WriteU32(response.directoryTableSize);
 		return true;
 	}
 
@@ -813,8 +841,8 @@ namespace uba
 
 	bool ProcessImpl::HandleUpdateTables(BinaryReader& reader, BinaryWriter& writer)
 	{
-		u32 dirSize = m_session.GetDirectoryTableSize();
-		writer.WriteU32(dirSize);
+		writer.WriteU32(m_session.GetDirectoryTableSize());
+		writer.WriteU32(m_session.GetFileMappingSize());
 		return true;
 	}
 
diff --git a/Engine/Source/Programs/UnrealBuildAccelerator/Common/Private/UbaSession.cpp b/Engine/Source/Programs/UnrealBuildAccelerator/Common/Private/UbaSession.cpp
--- a/Engine/Source/Programs/UnrealBuildAccelerator/Common/Private/UbaSession.cpp
+++ b/Engine/Source/Programs/UnrealBuildAccelerator/Common/Private/UbaSession.cpp
@@ -645,20 +645,19 @@ namespace uba
 		return true;
 	}
 
-	bool GetDirKey(StringKey& outDirKey, StringBufferBase& outDirName, const tchar*& outLastSlash, const tchar* fileName)
+	bool GetDirKey(StringKey& outDirKey, StringBufferBase& outDirName, const tchar*& outLastSlash, const StringView& fileName)
 	{
-		outLastSlash = TStrrchr(fileName, PathSeparator);
-		UBA_ASSERTF(outLastSlash, TC("Can't get dir key for path %s"), fileName);
+		outLastSlash = TStrrchr(fileName.data, PathSeparator);
+		UBA_ASSERTF(outLastSlash, TC("Can't get dir key for path %s"), fileName.data);
 		if (!outLastSlash)
 			return false;
 
-		u64 dirLen = u64(outLastSlash - fileName);
-		outDirName.Append(fileName, dirLen);
+		u64 dirLen = u64(outLastSlash - fileName.data);
+		outDirName.Append(fileName.data, dirLen);
 		outDirKey = CaseInsensitiveFs ? ToStringKeyLower(outDirName) : ToStringKey(outDirName);
 		return true;
 	}
 
-
 	bool Session::RegisterCreateFileForWrite(StringKey fileNameKey, const StringView& fileName, bool registerRealFile, u64 fileSize, u64 lastWriteTime, bool invalidateStorage)
 	{
 		// Remote is not updating its own directory table
@@ -670,7 +669,7 @@ namespace uba
 		StringKey dirKey;
 		const tchar* lastSlash;
 		StringBuffer<> dirName;
-		if (!GetDirKey(dirKey, dirName, lastSlash, fileName.data))
+		if (!GetDirKey(dirKey, dirName, lastSlash, fileName))
 			return true;
 			
 		#if 0//_DEBUG  // Bring this back, turned off right now because a few lines above the call to this method we add a mapping
@@ -825,7 +824,7 @@ namespace uba
 		return true;
 	}
 
-	u32 Session::RegisterDeleteFile(StringKey fileNameKey, const tchar* fileName)
+	u32 Session::RegisterDeleteFile(StringKey fileNameKey, const StringView& fileName)
 	{
 		// Remote is not updating its own directory table
 		if (m_runningRemote)
@@ -1284,7 +1283,7 @@ namespace uba
 		StringKey dirKey;
 		StringBuffer<> dirNameForHash;
 		const tchar* baseFileName;
-		GetDirKey(dirKey, dirNameForHash, baseFileName, fixedFilePath.data);
+		GetDirKey(dirKey, dirNameForHash, baseFileName, fixedFilePath);
 
 		if (CaseInsensitiveFs)
 			dirNameForHash.MakeLower();
@@ -1312,7 +1311,7 @@ namespace uba
 		UBA_ASSERT(!m_runningRemote);
 		StringBuffer<> fixedFilePath;
 		auto key = GetKeyAndFixedName(fixedFilePath, filePath);
-		RegisterDeleteFile(key, fixedFilePath.data);
+		RegisterDeleteFile(key, fixedFilePath);
 	}
 
 	void Session::RegisterCustomService(CustomServiceFunction&& function)
@@ -1786,7 +1785,7 @@ namespace uba
 		
 			if (fileName.EndsWith(TC(".dll")) || fileName.EndsWith(TC(".exe")))
 			{
-				UBA_ASSERTF(fileName[1] == ':', TC("Got bad filename from process %s"), fileName.data);
+				UBA_ASSERTF(IsAbsolutePath(fileName.data), TC("Got bad filename from process %s"), fileName.data);
 				AddFileMapping(fileNameKey, fileName.data, TC("#"));
 				out.fileName.Append(TC("#"));
 				return true;
@@ -1976,11 +1975,11 @@ namespace uba
 
 		if (!msg.newName.IsEmpty())
 		{
-			RegisterDeleteFile(file.nameKey, file.name.c_str());
+			RegisterDeleteFile(file.nameKey, file.name);
 			RegisterCreateFileForWrite(msg.newNameKey, msg.newName, registerRealFile, fileSize, lastWriteTime);
 		}
 		else if (msg.deleteOnClose)
-			RegisterDeleteFile(file.nameKey, file.name.c_str());
+			RegisterDeleteFile(file.nameKey, file.name);
 		else
 			RegisterCreateFileForWrite(file.nameKey, file.name, registerRealFile, fileSize, lastWriteTime);
 
@@ -2005,7 +2004,7 @@ namespace uba
 
 		out.result = uba::DeleteFileW(msg.fileName.data);
 		out.errorCode = GetLastError();
-		out.directoryTableSize = RegisterDeleteFile(msg.fileNameKey, msg.fileName.data);
+		out.directoryTableSize = RegisterDeleteFile(msg.fileNameKey, msg.fileName);
 		return true;
 	}
 
@@ -2031,7 +2030,7 @@ namespace uba
 		if (!out.result)
 			out.errorCode = GetLastError();
 		RegisterCreateFileForWrite(msg.toKey, msg.toName, true);
-		out.directoryTableSize = RegisterDeleteFile(msg.fromKey, msg.fromName.data);
+		out.directoryTableSize = RegisterDeleteFile(msg.fromKey, msg.fromName);
 		return true;
 	}
 
@@ -2055,15 +2054,47 @@ namespace uba
 	{
 		out.result = uba::CreateDirectoryW(msg.name.data);
 		if (out.result)
+		{
+			StringKey dirKey;
+			const tchar* lastSlash;
+			StringBuffer<> dirName;
+			if (!GetDirKey(dirKey, dirName, lastSlash, msg.name))
+				return true;
+			// Both these functions need to be called. otherwise we can get created directories that does not end up in directory table
 			RegisterCreateFileForWrite(msg.nameKey, msg.name, true);
-		out.errorCode = GetLastError();
+			u32 offset;
+			WriteDirectoryEntries(dirKey, dirName.data, offset);
+		}
+		else
+		{
+			out.errorCode = GetLastError();
+		}
+
+		out.directoryTableSize = GetDirectoryTableSize();
+		return true;
+	}
+
+	bool Session::RemoveDirectory(RemoveDirectoryResponse& out, const RemoveDirectoryMessage& msg)
+	{
+		out.result = uba::RemoveDirectoryW(msg.name.data);
+		if (out.result)
+			RegisterDeleteFile(msg.nameKey, msg.name);
+		else
+			out.errorCode = GetLastError();
+		out.directoryTableSize = GetDirectoryTableSize();
 		return true;
 	}
 
 	bool Session::GetFullFileName(GetFullFileNameResponse& out, const GetFullFileNameMessage& msg)
 	{
 		UBA_ASSERTF(false, TC("SHOULD NOT HAPPEN (only remote).. %s"), msg.fileName.data);
-		return SearchPathForFile(m_logger, out.fileName, msg.fileName.data, msg.process.m_virtualApplicationDir.c_str());
+		return false;
+	}
+
+	bool Session::GetLongPathName(GetLongPathNameResponse& out, const GetLongPathNameMessage& msg)
+	{
+		UBA_ASSERTF(false, TC("SHOULD NOT HAPPEN (only remote).. %s"), msg.fileName.data);
+		return false;
 	}
 
 	bool Session::GetListDirectoryInfo(ListDirectoryResponse& out, tchar* dirName, const StringKey& dirKey)
diff --git a/Engine/Source/Programs/UnrealBuildAccelerator/Common/Private/UbaSessionClient.cpp b/Engine/Source/Programs/UnrealBuildAccelerator/Common/Private/UbaSessionClient.cpp
--- a/Engine/Source/Programs/UnrealBuildAccelerator/Common/Private/UbaSessionClient.cpp
+++ b/Engine/Source/Programs/UnrealBuildAccelerator/Common/Private/UbaSessionClient.cpp
@@ -81,8 +81,8 @@ namespace uba
 
 	void SessionClient::Stop()
 	{
-		CancelAllProcessesAndWait();
 		m_loop = false;
+		CancelAllProcessesAndWait();
 		m_waitToSendEvent.Set();
 		m_loopThread.Wait();
 	}
@@ -182,7 +182,7 @@ namespace uba
 		CasKey casKey;
 		u32 fileAttributes = DefaultAttributes(); // TODO: This is wrong.. need to retrieve from server if this is executable or not
 
-		bool isAbsolute = IsWindows ? (fileName.count > 1 && fileName[1] == ':') : (fileName.count > 0 && fileName[0] == '/');
+		bool isAbsolute = IsAbsolutePath(fileName.data);
 		if (isAbsolute)
 		{
 			UBA_ASSERT(fileNameKey != StringKeyZero);
@@ -792,6 +792,19 @@ namespace uba
 		return true;
 	}
 
+	bool SessionClient::RemoveDirectory(RemoveDirectoryResponse& out, const RemoveDirectoryMessage& msg)
+	{
+		StackBinaryWriter<1024> writer;
+		NetworkMessage networkMsg(m_client, ServiceId, SessionMessageType_RemoveDirectory, writer);
+		writer.WriteString(msg.name);
+		StackBinaryReader<1024> reader;
+		if (!networkMsg.Send(reader, Stats().deleteFileMsg)) // Wrong message
+			return false;
+		out.result = reader.ReadBool();
+		out.errorCode = reader.ReadU32();
+		return true;
+	}
+
 	bool SessionClient::GetFullFileName(GetFullFileNameResponse& out, const GetFullFileNameMessage& msg)
 	{
 		// TODO: There is a potential risk here where two different applications asks for the full name of a file
@@ -825,6 +838,19 @@ namespace uba
 		return true;
 	}
 
+	bool SessionClient::GetLongPathName(GetLongPathNameResponse& out, const GetLongPathNameMessage& msg)
+	{
+		StackBinaryWriter<1024> writer;
+		NetworkMessage networkMsg(m_client, ServiceId, SessionMessageType_GetLongPathName, writer);
+		writer.WriteString(msg.fileName);
+		StackBinaryReader<1024> reader;
+		if (!networkMsg.Send(reader, Stats().getLongNameMsg))
+			return false;
+		out.errorCode = reader.ReadU32();
+		reader.ReadString(out.fileName);
+		return true;
+	}
+
 	bool SessionClient::GetListDirectoryInfo(ListDirectoryResponse& out, tchar* dirName, const StringKey& dirKey)
 	{
 		StackBinaryWriter<1024> writer;
diff --git a/Engine/Source/Programs/UnrealBuildAccelerator/Common/Private/UbaSessionServer.cpp b/Engine/Source/Programs/UnrealBuildAccelerator/Common/Private/UbaSessionServer.cpp
--- a/Engine/Source/Programs/UnrealBuildAccelerator/Common/Private/UbaSessionServer.cpp
+++ b/Engine/Source/Programs/UnrealBuildAccelerator/Common/Private/UbaSessionServer.cpp
@@ -798,6 +798,20 @@ namespace uba
 				writer.WriteU64(serverTime);
 				return true;
 			}
+			case SessionMessageType_GetLongPathName:
+			{
+				#if PLATFORM_WINDOWS
+				StringBuffer<> shortPath;
+				reader.ReadString(shortPath);
+				StringBuffer<> longPath;
+				longPath.count = ::GetLongPathNameW(shortPath.data, longPath.data, longPath.capacity);
+				writer.WriteU32(GetLastError());
+				writer.WriteString(longPath);
+				return true;
+				#else
+				return false;
+				#endif
+			}
 			case SessionMessageType_SendFileToServer:
 			{
 				u32 processId = reader.ReadU32();
@@ -904,7 +918,7 @@ namespace uba
 				bool result = uba::DeleteFileW(fileName.data);
 				u32 errorCode = GetLastError();
 				if (result)
-					RegisterDeleteFile(fileNameKey, fileName.data);
+					RegisterDeleteFile(fileNameKey, fileName);
 				writer.WriteBool(result);
 				writer.WriteU32(errorCode);
 				return true;
@@ -938,6 +952,18 @@ namespace uba
 				return true;
 			}
 
+			case SessionMessageType_RemoveDirectory:
+			{
+				RemoveDirectoryMessage msg;
+				reader.ReadString(msg.name);
+				RemoveDirectoryResponse response;
+				if (!Session::RemoveDirectory(response, msg))
+					return false;
+				writer.WriteBool(response.result);
+				writer.WriteU32(response.errorCode);
+				return true;
+			}
+
 			case SessionMessageType_ListDirectory:
 			{
 				u32 sessionId = reader.ReadU32();
@@ -1512,7 +1538,7 @@ namespace uba
 			{
 				for (auto& i : m_activeRemoteProcesses)
 				{
-					m_logger.Info(TC("ACTIVE PROCESS: %s"), i.GetStartInfo().description);
+					m_logger.Info(TC("ACTIVE PROCESS: %s"), i.GetStartInfo().GetDescription());
 				}
 				break;
 			}
@@ -1718,7 +1744,7 @@ namespace uba
 		auto findIt = m_processes.find(processId);
 		if (findIt == m_processes.end())
 			return str.Appendf(TC("<Process with id %u not found>"), processId).data;
-		return str.Appendf(TC("%s"), findIt->second.GetStartInfo().description).data;
+		return str.Appendf(TC("%s"), findIt->second.GetStartInfo().GetDescription()).data;
 	}
 
 	bool SessionServer::PrepareProcess(const ProcessStartInfo& startInfo, bool isChild, StringBufferBase& outRealApplication, const tchar*& outRealWorkingDir)
diff --git a/Engine/Source/Programs/UnrealBuildAccelerator/Common/Public/UbaNetwork.h b/Engine/Source/Programs/UnrealBuildAccelerator/Common/Public/UbaNetwork.h
--- a/Engine/Source/Programs/UnrealBuildAccelerator/Common/Public/UbaNetwork.h
+++ b/Engine/Source/Programs/UnrealBuildAccelerator/Common/Public/UbaNetwork.h
@@ -15,7 +15,7 @@ namespace uba
 	
 	static constexpr u32 SystemNetworkVersion = 1339;
 	static constexpr u32 StorageNetworkVersion = 4;
-	static constexpr u32 SessionNetworkVersion = 32;
+	static constexpr u32 SessionNetworkVersion = 33;
 	static constexpr u32 CacheNetworkVersion = 4;
 
 	static constexpr u32 CachePathTableMaxSize = 32*1024*1024;
@@ -64,10 +64,12 @@ namespace uba
 		UBA_SESSION_MESSAGE(EnsureBinaryFile) \
 		UBA_SESSION_MESSAGE(GetApplication) \
 		UBA_SESSION_MESSAGE(GetFileFromServer) \
+		UBA_SESSION_MESSAGE(GetLongPathName) \
 		UBA_SESSION_MESSAGE(SendFileToServer) \
 		UBA_SESSION_MESSAGE(DeleteFile) \
 		UBA_SESSION_MESSAGE(CopyFile) \
 		UBA_SESSION_MESSAGE(CreateDirectory) \
+		UBA_SESSION_MESSAGE(RemoveDirectory) \
 		UBA_SESSION_MESSAGE(ListDirectory) \
 		UBA_SESSION_MESSAGE(GetDirectoriesFromServer) \
 		UBA_SESSION_MESSAGE(GetNameToHashFromServer) \
diff --git a/Engine/Source/Programs/UnrealBuildAccelerator/Common/Public/UbaProcess.h b/Engine/Source/Programs/UnrealBuildAccelerator/Common/Public/UbaProcess.h
--- a/Engine/Source/Programs/UnrealBuildAccelerator/Common/Public/UbaProcess.h
+++ b/Engine/Source/Programs/UnrealBuildAccelerator/Common/Public/UbaProcess.h
@@ -74,12 +74,14 @@ namespace uba
 		bool HandleInit(BinaryReader& reader, BinaryWriter& writer);
 		bool HandleCreateFile(BinaryReader& reader, BinaryWriter& writer);
 		bool HandleGetFullFileName(BinaryReader& reader, BinaryWriter& writer);
+		bool HandleGetLongPathName(BinaryReader& reader, BinaryWriter& writer);
 		bool HandleCloseFile(BinaryReader& reader, BinaryWriter& writer);
 		bool HandleDeleteFile(BinaryReader& reader, BinaryWriter& writer);
 		bool HandleCopyFile(BinaryReader& reader, BinaryWriter& writer);
 		bool HandleMoveFile(BinaryReader& reader, BinaryWriter& writer);
 		bool HandleChmod(BinaryReader& reader, BinaryWriter& writer);
 		bool HandleCreateDirectory(BinaryReader& reader, BinaryWriter& writer);
+		bool HandleRemoveDirectory(BinaryReader& reader, BinaryWriter& writer);
 		bool HandleListDirectory(BinaryReader& reader, BinaryWriter& writer);
 		bool HandleUpdateTables(BinaryReader& reader, BinaryWriter& writer);
 		bool HandleCreateProcess(BinaryReader& reader, BinaryWriter& writer);
diff --git a/Engine/Source/Programs/UnrealBuildAccelerator/Common/Public/UbaSession.h b/Engine/Source/Programs/UnrealBuildAccelerator/Common/Public/UbaSession.h
--- a/Engine/Source/Programs/UnrealBuildAccelerator/Common/Public/UbaSession.h
+++ b/Engine/Source/Programs/UnrealBuildAccelerator/Common/Public/UbaSession.h
@@ -39,8 +39,12 @@ namespace uba
 	struct ChmodMessage;
 	struct GetFullFileNameMessage;
 	struct GetFullFileNameResponse;
+	struct GetLongPathNameMessage;
+	struct GetLongPathNameResponse;
 	struct CreateDirectoryMessage;
 	struct CreateDirectoryResponse;
+	struct RemoveDirectoryMessage;
+	struct RemoveDirectoryResponse;
 	struct ListDirectoryMessage;
 	struct ListDirectoryResponse;
 	struct WrittenFile;
@@ -111,7 +115,7 @@ namespace uba
 		bool CreateMemoryMapFromView(MemoryMap& out, StringKey fileNameKey, const tchar* fileName, const CasKey& casKey, u64 alignment);
 
 		bool RegisterCreateFileForWrite(StringKey fileNameKey, const StringView& fileName, bool registerRealFile, u64 fileSize = 0, u64 lastWriteTime = 0, bool invalidateStorage = true);
-		u32 RegisterDeleteFile(StringKey fileNameKey, const tchar* fileName);
+		u32 RegisterDeleteFile(StringKey fileNameKey, const StringView& fileName);
 
 		virtual bool PrepareProcess(const ProcessStartInfo& startInfo, bool isChild, StringBufferBase& outRealApplication, const tchar*& outRealWorkingDir);
 		virtual void* GetProcessEnvironmentVariables();
@@ -125,7 +129,9 @@ namespace uba
 		virtual bool MoveFile(MoveFileResponse& out, const MoveFileMessage& msg);
 		virtual bool Chmod(ChmodResponse& out, const ChmodMessage& msg);
 		virtual bool CreateDirectory(CreateDirectoryResponse& out, const CreateDirectoryMessage& msg);
+		virtual bool RemoveDirectory(RemoveDirectoryResponse& out, const RemoveDirectoryMessage& msg);
 		virtual bool GetFullFileName(GetFullFileNameResponse& out, const GetFullFileNameMessage& msg);
+		virtual bool GetLongPathName(GetLongPathNameResponse& out, const GetLongPathNameMessage& msg);
 		virtual bool GetListDirectoryInfo(ListDirectoryResponse& out, tchar* dirName, const StringKey& dirKey);
 		virtual bool WriteFilesToDisk(ProcessImpl& process, WrittenFile** files, u32 fileCount);
 		virtual bool AllocFailed(Process& process, const tchar* allocType, u32 error);
@@ -393,6 +399,18 @@ namespace uba
 		u32 mappedFileTableSize = 0;
 	};
 
+	struct GetLongPathNameMessage
+	{
+		ProcessImpl& process;
+		StringBuffer<> fileName;
+	};
+
+	struct GetLongPathNameResponse
+	{
+		StringBuffer<> fileName;
+		u32 errorCode = ~0u;
+	};
+
 	struct CreateDirectoryMessage
 	{
 		StringKey nameKey;
@@ -401,8 +419,22 @@ namespace uba
 
 	struct CreateDirectoryResponse
 	{
-		bool result;
-		u32 errorCode;
+		bool result = false;
+		u32 errorCode = 0;
+		u32 directoryTableSize = 0;
+	};
+
+	struct RemoveDirectoryMessage
+	{
+		StringKey nameKey;
+		StringBuffer<> name;
+	};
+
+	struct RemoveDirectoryResponse
+	{
+		bool result = false;
+		u32 errorCode = 0;
+		u32 directoryTableSize = 0;
 	};
 
 	struct ListDirectoryMessage
diff --git a/Engine/Source/Programs/UnrealBuildAccelerator/Common/Public/UbaSessionClient.h b/Engine/Source/Programs/UnrealBuildAccelerator/Common/Public/UbaSessionClient.h
--- a/Engine/Source/Programs/UnrealBuildAccelerator/Common/Public/UbaSessionClient.h
+++ b/Engine/Source/Programs/UnrealBuildAccelerator/Common/Public/UbaSessionClient.h
@@ -55,7 +55,9 @@ namespace uba
 		virtual bool MoveFile(MoveFileResponse& out, const MoveFileMessage& msg) override;
 		virtual bool Chmod(ChmodResponse& out, const ChmodMessage& msg) override;
 		virtual bool CreateDirectory(CreateDirectoryResponse& out, const CreateDirectoryMessage& msg) override;
+		virtual bool RemoveDirectory(RemoveDirectoryResponse& out, const RemoveDirectoryMessage& msg) override;
 		virtual bool GetFullFileName(GetFullFileNameResponse& out, const GetFullFileNameMessage& msg) override;
+		virtual bool GetLongPathName(GetLongPathNameResponse& out, const GetLongPathNameMessage& msg) override;
 		virtual bool GetListDirectoryInfo(ListDirectoryResponse& out, tchar* dirName, const StringKey& dirKey) override;
 		virtual bool WriteFilesToDisk(ProcessImpl& process, WrittenFile** files, u32 fileCount) override;
 		virtual bool AllocFailed(Process& process, const tchar* allocType, u32 error) override;
diff --git a/Engine/Source/Programs/UnrealBuildAccelerator/Common/Public/UbaStats.h b/Engine/Source/Programs/UnrealBuildAccelerator/Common/Public/UbaStats.h
--- a/Engine/Source/Programs/UnrealBuildAccelerator/Common/Public/UbaStats.h
+++ b/Engine/Source/Programs/UnrealBuildAccelerator/Common/Public/UbaStats.h
@@ -72,6 +72,7 @@ namespace uba
 		UBA_SESSION_STAT(Timer, waitGetFileMsg, 10) \
 		UBA_SESSION_STAT(Timer, createMmapFromFile, 12) \
 		UBA_SESSION_STAT(Timer, waitMmapFromFile, 12) \
+		UBA_SESSION_STAT(Timer, getLongNameMsg, 31) \
 
     struct SessionStats
 	{
diff --git a/Engine/Source/Programs/UnrealBuildAccelerator/Core/Public/UbaProcessStats.h b/Engine/Source/Programs/UnrealBuildAccelerator/Core/Public/UbaProcessStats.h
--- a/Engine/Source/Programs/UnrealBuildAccelerator/Core/Public/UbaProcessStats.h
+++ b/Engine/Source/Programs/UnrealBuildAccelerator/Core/Public/UbaProcessStats.h
@@ -43,6 +43,7 @@ namespace uba
 		UBA_PROCESS_STAT(preparseObjFiles, 30) \
 		UBA_PROCESS_STAT(fileTable, 30) \
 		UBA_PROCESS_STAT(dirTable, 30) \
+		UBA_PROCESS_STAT(longPathName, 31) \
 
 
 		#define UBA_PROCESS_STAT(T, ver) Timer T;
diff --git a/Engine/Source/Programs/UnrealBuildAccelerator/Core/Public/UbaProtocol.h b/Engine/Source/Programs/UnrealBuildAccelerator/Core/Public/UbaProtocol.h
--- a/Engine/Source/Programs/UnrealBuildAccelerator/Core/Public/UbaProtocol.h
+++ b/Engine/Source/Programs/UnrealBuildAccelerator/Core/Public/UbaProtocol.h
@@ -10,12 +10,14 @@ namespace uba
 		UBA_PROCESS_MESSAGE(Init) \
 		UBA_PROCESS_MESSAGE(CreateFile) \
 		UBA_PROCESS_MESSAGE(GetFullFileName) \
+		UBA_PROCESS_MESSAGE(GetLongPathName) \
 		UBA_PROCESS_MESSAGE(CloseFile) \
 		UBA_PROCESS_MESSAGE(DeleteFile) \
 		UBA_PROCESS_MESSAGE(CopyFile) \
 		UBA_PROCESS_MESSAGE(MoveFile) \
 		UBA_PROCESS_MESSAGE(Chmod) \
 		UBA_PROCESS_MESSAGE(CreateDirectory) \
+		UBA_PROCESS_MESSAGE(RemoveDirectory) \
 		UBA_PROCESS_MESSAGE(ListDirectory) \
 		UBA_PROCESS_MESSAGE(UpdateTables) \
 		UBA_PROCESS_MESSAGE(CreateProcess) \
diff --git a/Engine/Source/Programs/UnrealBuildAccelerator/Detours/Private/Windows/UbaDetoursFunctionsKernelBase.inl b/Engine/Source/Programs/UnrealBuildAccelerator/Detours/Private/Windows/UbaDetoursFunctionsKernelBase.inl
--- a/Engine/Source/Programs/UnrealBuildAccelerator/Detours/Private/Windows/UbaDetoursFunctionsKernelBase.inl
+++ b/Engine/Source/Programs/UnrealBuildAccelerator/Detours/Private/Windows/UbaDetoursFunctionsKernelBase.inl
@@ -9,20 +9,70 @@ DWORD Local_GetLongPathNameW(LPCWSTR lpszShortPath, LPWSTR lpszLongPath, DWORD c
 	if (findIt != g_longPathNameCache.end())
 	{
 		const wchar_t* longPath = findIt->second;
-		u64 len = wcslen(longPath);
-		UBA_ASSERT(cchBuffer > len);
+		u32 len = u32(wcslen(longPath));
+		if (len == 0)
+		{
+			SetLastError(ERROR_FILE_NOT_FOUND);
+			return 0;
+		}
+		SetLastError(ERROR_SUCCESS);
+		if (cchBuffer <= len)
+			return len + 1;
+
 		memcpy(lpszLongPath, longPath, (len + 1) * 2);
-		return u32(len);
+		return len + 1;
 	}
 
-	wchar_t* newShortPath = g_memoryBlock.Strdup(lpszShortPath);
-	DEBUG_LOG_TRUE(L"GetLongPathNameW", L"(Detour disabled under this call to handle ~) (%ls)", lpszShortPath);
+	wchar_t* newLongPath = nullptr;
+	DWORD res = 0;
 
-	SuppressDetourScope _;
-	DWORD res = True_GetLongPathNameW(lpszShortPath, lpszLongPath, cchBuffer);
-	if (res == 0)
-		return res;
-	g_longPathNameCache.insert({ newShortPath, g_memoryBlock.Strdup(lpszLongPath) });
+	if (g_runningRemote)
+	{
+		u32 errorCode = 0;
+		StringBuffer<> longName;
+		{
+			TimerScope ts(g_stats.longPathName);
+			SCOPED_WRITE_LOCK(g_communicationLock, pcs);
+			BinaryWriter writer;
+			writer.WriteByte(MessageType_GetLongPathName);
+			writer.WriteString(lpszShortPath);
+			writer.Flush();
+			BinaryReader reader;
+			errorCode = reader.ReadU32();
+			reader.ReadString(longName);
+		}
+
+		newLongPath = g_memoryBlock.Strdup(longName.data);
+
+		if (longName.count == 0)
+		{
+			// Error
+		}
+		if (cchBuffer > longName.count)
+		{
+			memcpy(lpszLongPath, longName.data, (longName.count+1)*sizeof(wchar_t));
+			res = longName.count;
+		}
+		else
+		{
+			res = longName.count + 1;
+		}
+
+		SetLastError(errorCode);
+	}
+	else
+	{
+		DEBUG_LOG_TRUE(L"GetLongPathNameW", L"(Detour disabled under this call to handle ~) (%ls)", lpszShortPath);
+
+		SuppressDetourScope _;
+		res = True_GetLongPathNameW(lpszShortPath, lpszLongPath, cchBuffer);
+		if (res == 0)
+			return res;
+		newLongPath = g_memoryBlock.Strdup(lpszLongPath);
+	}
+
+	wchar_t* newShortPath = g_memoryBlock.Strdup(lpszShortPath);
+	g_longPathNameCache.insert({ newShortPath, newLongPath });
 	return res;
 }
 
@@ -102,7 +152,6 @@ BOOL Detoured_SetCurrentDirectoryW(LPCWSTR lpPathName)
 	return True_SetCurrentDirectoryW(lpPathName);
 }
 
-
 BOOL Detoured_DuplicateHandle(HANDLE hSourceProcessHandle, HANDLE hSourceHandle, HANDLE hTargetProcessHandle, LPHANDLE lpTargetHandle, DWORD dwDesiredAccess, BOOL bInheritHandle, DWORD dwOptions)
 {
 	DETOURED_CALL(DuplicateHandle);
@@ -181,6 +230,7 @@ BOOL Detoured_CreateDirectoryW(LPCWSTR lpPathName, LPSECURITY_ATTRIBUTES lpSecur
 		return res;
 	}
 
+	u32 directoryTableSize;
 	BOOL res;
 	u32 errorCode = 0;
 	StringKey pathNameKey = ToStringKeyLower(pathName);
@@ -196,8 +246,11 @@ BOOL Detoured_CreateDirectoryW(LPCWSTR lpPathName, LPSECURITY_ATTRIBUTES lpSecur
 		BinaryReader reader;
 		res = reader.ReadBool();
 		errorCode = reader.ReadU32();
+		directoryTableSize = reader.ReadU32();
 	}
 
+	g_directoryTable.ParseDirectoryTable(directoryTableSize);
+
 	SetLastError(errorCode);
 	DEBUG_LOG_DETOURED(L"CreateDirectoryW", L"%ls -> %ls (%u)", lpPathName, ToString(res), errorCode);
 	return res;
@@ -209,20 +262,38 @@ BOOL Detoured_RemoveDirectoryW(LPCWSTR lpPathName)
 
 	StringBuffer<> pathName;
 	FixPath(pathName, lpPathName);
-	BOOL res;
-	if (!g_runningRemote || pathName.StartsWith(g_systemTemp.data))
+
+	if (pathName.StartsWith(g_systemTemp.data))
 	{
 		SuppressCreateFileDetourScope s; // TODO: Revisit this.. will not work remotely
-		res = True_RemoveDirectoryW(lpPathName);
+		BOOL res = True_RemoveDirectoryW(lpPathName);
+		DEBUG_LOG_TRUE(L"RemoveDirectoryW", L"%ls -> %ls", lpPathName, ToString(res));
+		return res;
 	}
-	else
+
+	u32 directoryTableSize;
+	BOOL res;
+	u32 errorCode = 0;
+	StringKey pathNameKey = ToStringKeyLower(pathName);
+
 	{
-		UBA_ASSERTF(!g_runningRemote, L"RemoveDirectory is not implemented for remote (removing %s)", lpPathName);
-		SetLastError(ERROR_PATH_NOT_FOUND);
-		res = false;
+		TimerScope ts(g_stats.deleteFile);
+		SCOPED_WRITE_LOCK(g_communicationLock, pcs);
+		BinaryWriter writer;
+		writer.WriteByte(MessageType_RemoveDirectory);
+		writer.WriteStringKey(pathNameKey);
+		writer.WriteString(pathName);
+		writer.Flush();
+		BinaryReader reader;
+		res = reader.ReadBool();
+		errorCode = reader.ReadU32();
+		directoryTableSize = reader.ReadU32();
 	}
-	DEBUG_LOG_TRUE(L"RemoveDirectoryW", L"%ls -> %ls", lpPathName, ToString(res));
 
+	g_directoryTable.ParseDirectoryTable(directoryTableSize);
+
+	SetLastError(errorCode);
+	DEBUG_LOG_DETOURED(L"RemoveDirectoryW", L"%ls -> %ls (%u)", lpPathName, ToString(res), errorCode);
 	return res;
 }
 
@@ -1330,7 +1401,7 @@ BOOL Detoured_DeleteFileW(LPCWSTR lpFileName)
 		pcs.Leave();
 		DEBUG_LOG_PIPE(L"DeleteFile", L"%ls", lpFileName);
 	}
-	DEBUG_LOG_DETOURED(L"DeleteFileW", L"(%ls) -> %ls", lpFileName, ToString(result));
+	DEBUG_LOG_DETOURED(L"DeleteFileW", L"(%ls) -> %ls (%u)", lpFileName, ToString(result), errorCode);
 
 	g_directoryTable.ParseDirectoryTable(directoryTableSize);
 	g_mappedFileTable.SetDeleted(fileNameKey, lpFileName, true);
@@ -1419,7 +1490,7 @@ bool Shared_MoveFile(LPCWSTR lpExistingFileName, LPCWSTR lpNewFileName, DWORD dw
 		DEBUG_LOG_PIPE(L"MoveFile", L"%ls to %ls", lpExistingFileName, lpNewFileName);
 	}
 
-	DEBUG_LOG_DETOURED(L"MoveFileExW", L"(PIPE) (%ls to %ls) -> %ls", lpExistingFileName, lpNewFileName, ToString(result));
+	DEBUG_LOG_DETOURED(L"MoveFileExW", L"(PIPE) (%ls to %ls) -> %ls (%u)", lpExistingFileName, lpNewFileName, ToString(result), errorCode);
 
 	g_directoryTable.ParseDirectoryTable(directoryTableSize);
 	g_mappedFileTable.SetDeleted(sourceKey, source.data, true);
@@ -1597,7 +1668,7 @@ __forceinline HANDLE Shared_FindFirstFileExW(LPCWSTR lpFileName, FINDEX_INFO_LEV
 
 
 
-	auto listHandle = new ListDirectoryHandle{ hash.key, insres.first->second };
+	auto listHandle = new ListDirectoryHandle{ hash.key, dir };
 
 	if (!*fileName)
 		listHandle->it = -2;
@@ -2043,13 +2114,13 @@ BOOL Detoured_SetFileInformationByHandle(HANDLE hFile, FILE_INFO_BY_HANDLE_CLASS
 			return true;
 
 		DEBUG_LOG_TRUE(L"SetFileInformationByHandle", L"%llu (FileDispositionInfo)", uintptr_t(hFile));
+		return True_SetFileInformationByHandle(hFile, FileInformationClass, lpFileInformation, dwBufferSize); // In here to be tabbed in log
 	}
 	else
 	{
 		DEBUG_LOG_TRUE(L"SetFileInformationByHandle", L"%llu (%u)", uintptr_t(hFile), FileInformationClass);
+		return True_SetFileInformationByHandle(hFile, FileInformationClass, lpFileInformation, dwBufferSize); // In here to be tabbed in log
 	}
-
-	return True_SetFileInformationByHandle(hFile, FileInformationClass, lpFileInformation, dwBufferSize);
 }
 
 HANDLE Detoured_CreateFileMappingW(HANDLE hFile, LPSECURITY_ATTRIBUTES lpFileMappingAttributes, DWORD flProtect, DWORD dwMaximumSizeHigh, DWORD dwMaximumSizeLow, LPCWSTR lpName)
@@ -2964,8 +3035,7 @@ DWORD Detoured_GetFileAttributesA(LPCSTR lpFileName)
 BOOL Detoured_GetFileAttributesExA(LPCSTR lpFileName, GET_FILEEX_INFO_LEVELS fInfoLevelId, LPVOID lpFileInformation)
 {
 	DETOURED_CALL(GetFileAttributesExA);
-	DEBUG_LOG_TRUE(L"GetFileAttributesExA", L"");
-	UBA_ASSERT(!g_runningRemote);
+	DEBUG_LOG_TRUE(L"GetFileAttributesExA", L""); // Calls ExW on both windows and wine
 	return True_GetFileAttributesExA(lpFileName, fInfoLevelId, lpFileInformation);
 }
 
@@ -3256,12 +3326,12 @@ LPWCH Detoured_GetEnvironmentStringsW()
 	DEBUG_LOG_TRUE(L"GetEnvironmentStringsW", L"");
 	auto res = True_GetEnvironmentStringsW();
 
-	auto it = res;
-	while (*it)
-	{
-		DEBUG_LOG(L"		VAR: %ls", it);
-		it += wcslen(it) + 1;
-	}
+	//auto it = res;
+	//while (*it)
+	//{
+	//	DEBUG_LOG(L"		VAR: %ls", it);
+	//	it += wcslen(it) + 1;
+	//}
 
 	return res;
 }
