[UBA]
* Fixed so CreateProcessA is detoured and calling CreateProcessW
* Added more detours only enabled in debug  (for debugging purposes)

diff --git a/Engine/Source/Programs/UnrealBuildAccelerator/Detours/Private/Windows/UbaDetoursFunctionsKernelBase.inl b/Engine/Source/Programs/UnrealBuildAccelerator/Detours/Private/Windows/UbaDetoursFunctionsKernelBase.inl
--- a/Engine/Source/Programs/UnrealBuildAccelerator/Detours/Private/Windows/UbaDetoursFunctionsKernelBase.inl
+++ b/Engine/Source/Programs/UnrealBuildAccelerator/Detours/Private/Windows/UbaDetoursFunctionsKernelBase.inl
@@ -2786,6 +2786,39 @@ BOOL Detoured_CreateProcessW(LPCWSTR lpApplicationName, LPWSTR lpCommandLine, LP
 	return TRUE;
 }
 
+BOOL Detoured_CreateProcessA(LPCSTR lpApplicationName, LPSTR lpCommandLine, LPSECURITY_ATTRIBUTES lpProcessAttributes, LPSECURITY_ATTRIBUTES lpThreadAttributes, BOOL bInheritHandles,
+	DWORD dwCreationFlags, LPVOID lpEnvironment, LPCSTR lpCurrentDirectory, LPSTARTUPINFOA lpStartupInfo, LPPROCESS_INFORMATION lpProcessInformation)
+{
+	wchar_t* lpApplicationNameW = nullptr;
+	TString lpApplicationNameTemp;
+	if (lpApplicationName)
+	{
+		lpApplicationNameTemp = TString(lpApplicationName, lpApplicationName + strlen(lpApplicationName));
+		lpApplicationNameW = lpApplicationNameTemp.data();
+	}
+	wchar_t* lpCommandLineW = nullptr;
+	TString lpCommandLineTemp;
+	if (lpCommandLine)
+	{
+		lpCommandLineTemp = TString(lpCommandLine, lpCommandLine + strlen(lpCommandLine));
+		lpCommandLineW = lpCommandLineTemp.data();
+	}
+	wchar_t* lpCurrentDirectoryW = nullptr;
+	TString lpCurrentDirectoryTemp;
+	if (lpCurrentDirectory)
+	{
+		lpCurrentDirectoryTemp = TString(lpCurrentDirectory, lpCurrentDirectory + strlen(lpCurrentDirectory));
+		lpCurrentDirectoryW = lpCurrentDirectoryTemp.data();
+	}
+
+	UBA_ASSERT(!lpStartupInfo->lpReserved);
+	UBA_ASSERT(!lpStartupInfo->lpDesktop);
+	UBA_ASSERT(!lpStartupInfo->lpTitle);
+
+	STARTUPINFOW lpStartupInfoW = *(LPSTARTUPINFOW)lpStartupInfo;
+	return Detoured_CreateProcessW(lpApplicationNameW, lpCommandLineW, lpProcessAttributes, lpThreadAttributes, bInheritHandles, dwCreationFlags, lpEnvironment, lpCurrentDirectoryW, &lpStartupInfoW, lpProcessInformation);
+}
+
 void Detoured_ExitProcess(UINT uExitCode)
 {
 	// Can't log this one
@@ -3160,14 +3193,6 @@ LPTOP_LEVEL_EXCEPTION_FILTER Detoured_SetUnhandledExceptionFilter(LPTOP_LEVEL_EX
 	return True_SetUnhandledExceptionFilter(lpTopLevelExceptionFilter);
 }
 
-BOOL Detoured_CreateProcessA(LPCSTR lpApplicationName, LPSTR lpCommandLine, LPSECURITY_ATTRIBUTES lpProcessAttributes, LPSECURITY_ATTRIBUTES lpThreadAttributes, BOOL bInheritHandles,
-	DWORD dwCreationFlags, LPVOID lpEnvironment, LPCSTR lpCurrentDirectory, LPSTARTUPINFOA lpStartupInfo, LPPROCESS_INFORMATION lpProcessInformation)
-{
-	DETOURED_CALL(CreateProcessA);
-	DEBUG_LOG_TRUE(L"CreateProcessA", L"(%hs)", lpCommandLine ? lpCommandLine : "");
-	return True_CreateProcessA(lpApplicationName, lpCommandLine, lpProcessAttributes, lpThreadAttributes, bInheritHandles, dwCreationFlags, lpEnvironment, lpCurrentDirectory, lpStartupInfo, lpProcessInformation);
-}
-
 BOOL Detoured_FlushInstructionCache(HANDLE hProcess, LPCVOID lpBaseAddress, SIZE_T dwSize)
 {
 	DETOURED_CALL(FlushInstructionCache);
@@ -3280,6 +3305,12 @@ BOOL Detoured_DeleteFileA(LPCSTR lpFileName)
 	return True_DeleteFileA(lpFileName);
 }
 
+BOOL Detoured_SetCurrentDirectoryA(LPCSTR lpPathName)
+{
+	DETOURED_CALL(SetCurrentDirectoryA);
+	DEBUG_LOG_TRUE(L"SetCurrentDirectoryA", L"%hs", lpPathName);
+	return True_SetCurrentDirectoryA(lpPathName);
+}
 BOOLEAN Detoured_CreateSymbolicLinkW(LPCWSTR lpSymlinkFileName, LPCWSTR lpTargetFileName, DWORD dwFlags)
 {
 	UBA_ASSERT(!g_runningRemote);
@@ -3480,6 +3511,13 @@ HANDLE Detoured_CreateNamedPipeW(LPCWSTR lpName, DWORD dwOpenMode, DWORD dwPipeM
 	return h;
 }
 
+BOOL Detoured_CallNamedPipeW(LPCWSTR lpNamedPipeName, LPVOID lpInBuffer, DWORD nInBufferSize, LPVOID lpOutBuffer, DWORD nOutBufferSize, LPDWORD lpBytesRead, DWORD nTimeOut)
+{
+	DETOURED_CALL(CreateNamedPipeW);
+	DEBUG_LOG_TRUE(L"CallNamedPipeW", L"%ls %u %u", lpNamedPipeName, nInBufferSize, nOutBufferSize);
+	return True_CallNamedPipeW(lpNamedPipeName, lpInBuffer, nInBufferSize, lpOutBuffer, nOutBufferSize, lpBytesRead, nTimeOut);
+}
+
 BOOL Detoured_PeekNamedPipe(HANDLE hNamedPipe, LPVOID lpBuffer, DWORD nBufferSize, LPDWORD lpBytesRead, LPDWORD lpTotalBytesAvail, LPDWORD lpBytesLeftThisMessage)
 {
 	UBA_ASSERT(!isDetouredHandle(hNamedPipe));
diff --git a/Engine/Source/Programs/UnrealBuildAccelerator/Detours/Private/Windows/UbaDetoursFunctionsWin.h b/Engine/Source/Programs/UnrealBuildAccelerator/Detours/Private/Windows/UbaDetoursFunctionsWin.h
--- a/Engine/Source/Programs/UnrealBuildAccelerator/Detours/Private/Windows/UbaDetoursFunctionsWin.h
+++ b/Engine/Source/Programs/UnrealBuildAccelerator/Detours/Private/Windows/UbaDetoursFunctionsWin.h
@@ -80,6 +80,7 @@
 	DETOURED_FUNCTION(UnmapViewOfFileEx) \
 	DETOURED_FUNCTION(GetFinalPathNameByHandleW) \
 	DETOURED_FUNCTION(CreateProcessW) \
+	DETOURED_FUNCTION(CreateProcessA) \
 	DETOURED_FUNCTION(TerminateProcess) \
 	DETOURED_FUNCTION(SearchPathW) \
 	DETOURED_FUNCTION(LoadLibraryExW) \
@@ -225,7 +226,6 @@
 	DETOURED_FUNCTION(GetModuleBaseNameA) \
 	DETOURED_FUNCTION(GetModuleBaseNameW) \
 	DETOURED_FUNCTION(SetUnhandledExceptionFilter) \
-	DETOURED_FUNCTION(CreateProcessA) \
 	DETOURED_FUNCTION(FlushInstructionCache) \
 	DETOURED_FUNCTION(CreateFile2) \
 	DETOURED_FUNCTION(CreateFileTransactedW) \
@@ -237,6 +237,7 @@
 	DETOURED_FUNCTION(ReplaceFileW) \
 	DETOURED_FUNCTION(CreateHardLinkA) \
 	DETOURED_FUNCTION(DeleteFileA) \
+	DETOURED_FUNCTION(SetCurrentDirectoryA) \
 	DETOURED_FUNCTION(CreateSymbolicLinkW) \
 	DETOURED_FUNCTION(CreateSymbolicLinkA) \
 	DETOURED_FUNCTION(SetEnvironmentVariableW) \
@@ -254,6 +255,7 @@
 	DETOURED_FUNCTION(CreatePipe) \
 	DETOURED_FUNCTION(SetHandleInformation) \
 	DETOURED_FUNCTION(CreateNamedPipeW) \
+	DETOURED_FUNCTION(CallNamedPipeW ) \
 	DETOURED_FUNCTION(PeekNamedPipe) \
 	DETOURED_FUNCTION(GetKernelObjectSecurity) \
 	DETOURED_FUNCTION(ImpersonateNamedPipeClient) \
