[UBA]
* Improved GetVolumeInformationW
* Changed so GetFullPathNameW is always going through detour
* Some cleanup in CreateProcessW

diff --git a/Engine/Source/Programs/UnrealBuildAccelerator/Detours/Private/Windows/UbaDetoursFunctionsKernelBase.inl b/Engine/Source/Programs/UnrealBuildAccelerator/Detours/Private/Windows/UbaDetoursFunctionsKernelBase.inl
--- a/Engine/Source/Programs/UnrealBuildAccelerator/Detours/Private/Windows/UbaDetoursFunctionsKernelBase.inl
+++ b/Engine/Source/Programs/UnrealBuildAccelerator/Detours/Private/Windows/UbaDetoursFunctionsKernelBase.inl
@@ -523,12 +523,15 @@ BOOL Detoured_GetVolumeInformationW(LPCWSTR lpRootPathName, LPWSTR lpVolumeNameB
 	DETOURED_CALL(GetVolumeInformationW);
 	if (g_runningRemote)
 	{
-		UBA_ASSERT(!lpVolumeNameBuffer);
-		UBA_ASSERT(!lpVolumeSerialNumber);
+		if (lpVolumeSerialNumber)
+			*lpVolumeSerialNumber = lpRootPathName[0]; // Let's see if this works, LOL
+
+		//UBA_ASSERT(!lpVolumeNameBuffer);
 		UBA_ASSERT(!lpMaximumComponentLength);
 		UBA_ASSERT(!lpFileSystemFlags);
 
-		wcscpy_s(lpFileSystemNameBuffer, nFileSystemNameSize, L"NTFS"); // TODO: Not everyone has NTFS?
+		if (nFileSystemNameSize)
+			wcscpy_s(lpFileSystemNameBuffer, nFileSystemNameSize, L"NTFS"); // TODO: Not everyone has NTFS?
 		SetLastError(ERROR_SUCCESS);
 		DEBUG_LOG_DETOURED(L"GetVolumeInformationW", L"%ls", lpRootPathName);
 		return true;
@@ -1118,24 +1121,18 @@ DWORD Detoured_GetLongPathNameW(LPCWSTR lpszShortPath, LPWSTR lpszLongPath, DWOR
 DWORD Detoured_GetFullPathNameW(LPCWSTR lpFileName, DWORD nBufferLength, LPWSTR lpBuffer, LPWSTR* lpFilePart)
 {
 	DETOURED_CALL(GetFullPathNameW);
-	if (g_runningRemote)
-	{
-		StringBuffer<> temp;
-		FixPath(temp, lpFileName);
-		u64 requiredSize = temp.count + 1;
-		if (nBufferLength < requiredSize)
-			return DWORD(requiredSize);
-		memcpy(lpBuffer, temp.data, requiredSize * 2);
-		if (lpFilePart)
-			*lpFilePart = wcsrchr(lpBuffer, '\\') + 1;
-		auto res = DWORD(temp.count);
-		DEBUG_LOG_DETOURED(L"GetFullPathNameW", L"%ls -> %u", temp.data, res);
-		SetLastError(ERROR_SUCCESS);
-		return res;
-	}
 
-	auto res = True_GetFullPathNameW(lpFileName, nBufferLength, lpBuffer, lpFilePart);
-	DEBUG_LOG_TRUE(L"GetFullPathNameW", L"%ls -> %u", lpFileName, res);
+	StringBuffer<> temp;
+	FixPath(temp, lpFileName);
+	u64 requiredSize = temp.count + 1;
+	if (nBufferLength < requiredSize)
+		return DWORD(requiredSize);
+	memcpy(lpBuffer, temp.data, requiredSize * 2);
+	if (lpFilePart)
+		*lpFilePart = wcsrchr(lpBuffer, '\\') + 1;
+	auto res = DWORD(temp.count);
+	DEBUG_LOG_DETOURED(L"GetFullPathNameW", L"%ls TO %ls -> %u", lpFileName, temp.data, res);
+	SetLastError(ERROR_SUCCESS);
 	return res;
 }
 
@@ -2656,16 +2653,19 @@ BOOL Detoured_GetConsoleTitleW(LPTSTR lpConsoleTitle, DWORD nSize)
 BOOL Detoured_CreateProcessW(LPCWSTR lpApplicationName, LPWSTR lpCommandLine, LPSECURITY_ATTRIBUTES lpProcessAttributes, LPSECURITY_ATTRIBUTES lpThreadAttributes, BOOL bInheritHandles,
 	DWORD dwCreationFlags, LPVOID lpEnvironment, LPCWSTR lpCurrentDirectory, LPSTARTUPINFOW lpStartupInfo, LPPROCESS_INFORMATION lpProcessInformation)
 {
+	const tchar* originalCmd = lpCommandLine ? lpCommandLine : TC("");
+
 	DETOURED_CALL(CreateProcessW);
-	DEBUG_LOG_DETOURED(L"CreateProcessW", L"%ls %ls CreationFlags: 0x%x StartupFlags: 0x%u Stdin: %llu", lpApplicationName, lpCommandLine ? lpCommandLine : L"", dwCreationFlags, lpStartupInfo->dwFlags, u64(lpStartupInfo->hStdInput));
+	DEBUG_LOG_DETOURED(L"CreateProcessW", L"%ls %ls CreationFlags: 0x%x StartupFlags: 0x%u Stdin: %llu WorkDir: %s", lpApplicationName, originalCmd, dwCreationFlags, lpStartupInfo->dwFlags, u64(lpStartupInfo->hStdInput), (lpCurrentDirectory ? lpCurrentDirectory : L""));
 
-	if ((!lpApplicationName || !*lpApplicationName) && (!lpCommandLine || !*lpCommandLine))
+	if ((!lpApplicationName || !*lpApplicationName) && !*originalCmd)
 	{
 		SetLastError(ERROR_FILE_NOT_FOUND);
 		return FALSE;
 	}
 
-	if (lpCommandLine && (Contains(lpCommandLine, L"winedbg") || Contains(lpCommandLine, L"werfault.exe") || Contains(lpCommandLine, L"vsjitdebugger.exe") || Contains(lpCommandLine, L"crashpad_handler.exe")))
+	// Debug binaries started when process crash... we don't want to detour these.
+	if (Contains(originalCmd, L"winedbg") || Contains(originalCmd, L"werfault.exe") || Contains(originalCmd, L"vsjitdebugger.exe") || Contains(originalCmd, L"crashpad_handler.exe"))
 	{
 		if (g_runningRemote)
 		{
@@ -2695,7 +2695,7 @@ BOOL Detoured_CreateProcessW(LPCWSTR lpApplicationName, LPWSTR lpCommandLine, LP
 		BinaryWriter writer;
 		writer.WriteByte(MessageType_CreateProcess);
 		writer.WriteString(application.data);
-		writer.WriteString(lpCommandLine ? lpCommandLine : L"");
+		writer.WriteString(originalCmd);
 		writer.WriteString(lpCurrentDirectory ? lpCurrentDirectory : g_virtualWorkingDir.data);
 		writer.Flush();
 		BinaryReader reader;
@@ -2710,7 +2710,7 @@ BOOL Detoured_CreateProcessW(LPCWSTR lpApplicationName, LPWSTR lpCommandLine, LP
 
 		currentDir = reader.ReadString();
 		commandLine = reader.ReadString();
-		DEBUG_LOG_PIPE(L"CreateProcess", L"%ls %ls", application.data, lpCommandLine ? lpCommandLine : L"");
+		DEBUG_LOG_PIPE(L"CreateProcess", L"%ls %ls", application.data, originalCmd);
 	}
 
 	LPCSTR dlls[] = { dll };
@@ -2768,7 +2768,7 @@ BOOL Detoured_CreateProcessW(LPCWSTR lpApplicationName, LPWSTR lpCommandLine, LP
 		writer.WriteU32(lpProcessInformation->dwProcessId);
 		writer.WriteU64(u64(lpProcessInformation->hThread));
 		writer.Flush();
-		DEBUG_LOG_PIPE(L"StartProcess", L"%ls %ls", lpApplicationName, lpCommandLine ? lpCommandLine : L"");
+		DEBUG_LOG_PIPE(L"StartProcess", L"%ls %ls", lpApplicationName, originalCmd);
 	}
 
 	HANDLE trueHandle = lpProcessInformation->hProcess;
