[UBA]
* Added option to keep files in memory
* Added option to allow using output files (transferred via handle to session process)
* Added option to suppress logging

diff --git a/Engine/Source/Programs/UnrealBuildAccelerator/Common/Private/UbaProcess.cpp b/Engine/Source/Programs/UnrealBuildAccelerator/Common/Private/UbaProcess.cpp
--- a/Engine/Source/Programs/UnrealBuildAccelerator/Common/Private/UbaProcess.cpp
+++ b/Engine/Source/Programs/UnrealBuildAccelerator/Common/Private/UbaProcess.cpp
@@ -100,8 +100,6 @@ namespace uba
 	,	m_tempFilesLock(parent ? parent->m_tempFilesLock : *new ReaderWriterLock())
 	,	m_tempFiles(parent ? parent->m_tempFiles : *new UnorderedMap<StringKey, WrittenFile>)
 	{
-
-		CreateGuid(m_processGuid);
 		m_cancelEvent.Create(true, true);
 		m_writeEvent.Create(false, true);
 		m_readEvent.Create(false, true);
@@ -1495,7 +1493,7 @@ namespace uba
 			}
 
 			DetoursPayload payload;
-			payload.processGuid = m_processGuid;
+			payload.processId = m_id;
 			payload.hostProcess = hostProcess;
 			payload.cancelEvent = m_cancelEvent.GetHandle();
 			payload.writeEvent = m_writeEvent.GetHandle();
diff --git a/Engine/Source/Programs/UnrealBuildAccelerator/Common/Public/UbaProcess.h b/Engine/Source/Programs/UnrealBuildAccelerator/Common/Public/UbaProcess.h
--- a/Engine/Source/Programs/UnrealBuildAccelerator/Common/Public/UbaProcess.h
+++ b/Engine/Source/Programs/UnrealBuildAccelerator/Common/Public/UbaProcess.h
@@ -120,7 +120,6 @@ namespace uba
 		ProcessImpl* m_parentProcess;
 		ReaderWriterLock m_initLock;
 		u32 m_id;
-		Guid m_processGuid;
 		FileMappingAllocator::Allocation m_comMemory;
 
 	#if PLATFORM_WINDOWS
diff --git a/Engine/Source/Programs/UnrealBuildAccelerator/Core/Public/UbaDetoursPayload.h b/Engine/Source/Programs/UnrealBuildAccelerator/Core/Public/UbaDetoursPayload.h
--- a/Engine/Source/Programs/UnrealBuildAccelerator/Core/Public/UbaDetoursPayload.h
+++ b/Engine/Source/Programs/UnrealBuildAccelerator/Core/Public/UbaDetoursPayload.h
@@ -9,7 +9,7 @@ namespace uba
 {
 	struct __declspec(uuid("{a1041c70-bf7c-4e94-8f8b-bacf0f31ba9a}")) DetoursPayload
 	{
-		Guid processGuid = {};
+		u32 processId = 0;
 		HANDLE hostProcess = 0;
 		HANDLE cancelEvent = 0;
 		HANDLE writeEvent = 0;
@@ -23,6 +23,9 @@ namespace uba
 		bool useCustomAllocator = true;
 		bool isRunningWine = false;
 		bool isChild = false;
+		bool allowKeepFilesInMemory = IsWindows;
+		bool allowOutputFiles = IsWindows;
+		bool suppressLogging = false;
 		bool storeObjFilesCompressed = false;
 		int uiLanguage = 0;
 		StringBuffer<256> logFile;
diff --git a/Engine/Source/Programs/UnrealBuildAccelerator/Detours/Private/UbaDetoursShared.cpp b/Engine/Source/Programs/UnrealBuildAccelerator/Detours/Private/UbaDetoursShared.cpp
--- a/Engine/Source/Programs/UnrealBuildAccelerator/Detours/Private/UbaDetoursShared.cpp
+++ b/Engine/Source/Programs/UnrealBuildAccelerator/Detours/Private/UbaDetoursShared.cpp
@@ -26,6 +26,9 @@ namespace uba
 	ApplicationRules* g_rules;
 	bool g_runningRemote;
 	bool g_isChild;
+	bool g_allowKeepFilesInMemory = IsWindows;
+	bool g_allowOutputFiles = IsWindows;
+	bool g_suppressLogging = false;
 
 	void InitSharedVariables()
 	{
@@ -190,7 +193,7 @@ namespace uba
 	template<typename CharType>
 	void Shared_WriteConsoleT(const CharType* chars, u32 charCount, bool isError)
 	{
-		if (!g_echoOn)
+		if (!g_echoOn || g_suppressLogging)
 			return;
 
 		SCOPED_WRITE_LOCK(g_consoleStringCs, lock);
@@ -304,15 +307,33 @@ namespace uba
 
 				if (dirTableOffset == ~u32(0))
 				{
-					if (g_runningRemote) // This could be a written file not reported to server yet
+					// This could be a newly written file but process has not fetched latest directory table
+					SCOPED_READ_LOCK(g_mappedFileTable.m_lookupLock, lock);
+					auto findIt = g_mappedFileTable.m_lookup.find(fileNameKey);
+					if (findIt != g_mappedFileTable.m_lookup.end() && !findIt->second.deleted)
 					{
-						SCOPED_READ_LOCK(g_mappedFileTable.m_lookupLock, lock);
-						auto findIt = g_mappedFileTable.m_lookup.find(fileNameKey);
-						if (findIt != g_mappedFileTable.m_lookup.end() && !findIt->second.deleted)
+						outAttr.exists = true;
+						outAttr.lastError = ErrorSuccess;
+						outAttr.useCache = false;
+
+						if (g_runningRemote)
 						{
-							outAttr.useCache = false;
-							return findIt->second.name;
+							FileInfo& info = findIt->second;
+							outAttr.useCache = true;
+
+							// TODO: This is missing lots of information..
+#if PLATFORM_WINDOWS
+							LARGE_INTEGER li = ToLargeInteger(info.size);
+							outAttr.data.dwFileAttributes = FILE_ATTRIBUTE_NORMAL;
+							outAttr.data.nFileSizeLow = li.LowPart;
+							outAttr.data.nFileSizeHigh = li.HighPart;
+#else
+							outAttr.data.st_mode = (mode_t)(S_IRUSR | S_IWUSR | S_IRGRP | S_IROTH);
+							outAttr.data.st_size = info.size;
+#endif
 						}
+
+						return findIt->second.name;
 					}
 
 					outAttr.useCache = true;
diff --git a/Engine/Source/Programs/UnrealBuildAccelerator/Detours/Private/UbaDetoursShared.h b/Engine/Source/Programs/UnrealBuildAccelerator/Detours/Private/UbaDetoursShared.h
--- a/Engine/Source/Programs/UnrealBuildAccelerator/Detours/Private/UbaDetoursShared.h
+++ b/Engine/Source/Programs/UnrealBuildAccelerator/Detours/Private/UbaDetoursShared.h
@@ -88,12 +88,14 @@ namespace uba
 	inline constexpr bool g_allowFileMappingDetour = true;
 	inline constexpr bool g_allowFindFileDetour = true;
 	inline constexpr bool g_allowListDirectoryHandle = true;
-	inline constexpr bool g_allowKeepFilesInMemory = IsWindows;
 
 	extern u32 g_rulesIndex;
 	extern ApplicationRules* g_rules;
 	extern bool g_runningRemote;
 	extern bool g_isChild;
+	extern bool g_allowKeepFilesInMemory;
+	extern bool g_allowOutputFiles;
+	extern bool g_suppressLogging;
 
 	#if PLATFORM_WINDOWS
 	constexpr u32 ErrorSuccess = 0;
diff --git a/Engine/Source/Programs/UnrealBuildAccelerator/Detours/Private/Windows/UbaDetoursFunctionsKernelBase.inl b/Engine/Source/Programs/UnrealBuildAccelerator/Detours/Private/Windows/UbaDetoursFunctionsKernelBase.inl
--- a/Engine/Source/Programs/UnrealBuildAccelerator/Detours/Private/Windows/UbaDetoursFunctionsKernelBase.inl
+++ b/Engine/Source/Programs/UnrealBuildAccelerator/Detours/Private/Windows/UbaDetoursFunctionsKernelBase.inl
@@ -597,7 +597,7 @@ DWORD Detoured_GetSecurityInfo(HANDLE handle, SE_OBJECT_TYPE ObjectType, SECURIT
 
 void WriteStdFile(LPCVOID lpBuffer, DWORD nNumberOfBytesToWrite, bool isError)
 {
-	if (!g_echoOn)
+	if (!g_echoOn || g_suppressLogging)
 		return;
 
 	SCOPED_WRITE_LOCK(g_stdFileLock, lock);
diff --git a/Engine/Source/Programs/UnrealBuildAccelerator/Detours/Private/Windows/UbaDetoursFunctionsUcrtBase.inl b/Engine/Source/Programs/UnrealBuildAccelerator/Detours/Private/Windows/UbaDetoursFunctionsUcrtBase.inl
--- a/Engine/Source/Programs/UnrealBuildAccelerator/Detours/Private/Windows/UbaDetoursFunctionsUcrtBase.inl
+++ b/Engine/Source/Programs/UnrealBuildAccelerator/Detours/Private/Windows/UbaDetoursFunctionsUcrtBase.inl
@@ -160,17 +160,17 @@ errno_t Detoured__waccess_s(const wchar_t* path, int mode)
 	if (!CanDetour(path))
 	{
 		auto res = True__waccess_s(path, mode);
-		DEBUG_LOG_TRUE(L"_waccess_s", L"%ls %i -> %ls", path, mode, WaccessResultToString(res));
+		DEBUG_LOG_TRUE(L"_waccess_s", L"(NODETOUR) %ls %i -> %ls", path, mode, WaccessResultToString(res));
 		return res;
 	}
 
 	FileAttributes attr;
 	const wchar_t* realName = Shared_GetFileAttributes(attr, path);
 
-	if (!attr.useCache)
+	if (!attr.useCache && !g_runningRemote)
 	{
 		auto res = True__waccess_s(realName, mode);
-		DEBUG_LOG_TRUE(L"_waccess_s", L"%ls %i -> %ls", path, mode, WaccessResultToString(res));
+		DEBUG_LOG_TRUE(L"_waccess_s", L"(NOCACHE) %ls %i -> %ls", path, mode, WaccessResultToString(res));
 		return res;
 	}
 
@@ -235,6 +235,9 @@ int Detoured__write(int fd, const void* buffer, unsigned int count)
 	DETOURED_CALL(_write);
 	if (fd == StdOutFd && g_isDetachedProcess)
 	{
+		if (g_suppressLogging)
+			return count;
+
 		const char* str = (const char*)buffer;
 		const char* end = str + count;
 		const char* line = str;
diff --git a/Engine/Source/Programs/UnrealBuildAccelerator/Detours/Private/Windows/UbaDetoursFunctionsWin.cpp b/Engine/Source/Programs/UnrealBuildAccelerator/Detours/Private/Windows/UbaDetoursFunctionsWin.cpp
--- a/Engine/Source/Programs/UnrealBuildAccelerator/Detours/Private/Windows/UbaDetoursFunctionsWin.cpp
+++ b/Engine/Source/Programs/UnrealBuildAccelerator/Detours/Private/Windows/UbaDetoursFunctionsWin.cpp
@@ -326,7 +326,7 @@ void UbaAssert(const wchar_t* text, const char* file, u32 line, const char* expr
 	int ret = MessageBoxW(GetConsoleWindow(), b.data, title.data, MB_ABORTRETRYIGNORE|MB_SYSTEMMODAL);
 	if (ret == IDABORT)
 		ExitProcess(terminateCode);
-	else if (ret == IDRETRY)
+	else if (ret == IDRETRY && IsDebuggerPresent())
 		DebugBreak();
 	#else
 	if (allowTerminate)
@@ -367,7 +367,7 @@ const wchar_t* HandleToName(HANDLE handle)
 
 bool NeedsSharedMemory(const wchar_t* file) { return g_allowKeepFilesInMemory && g_rules->NeedsSharedMemory(file); }
 u64 FileTypeMaxSize(const StringBufferBase& file, bool isSystemOrTempFile) { return g_rules->FileTypeMaxSize(file, isSystemOrTempFile); }
-bool IsOutputFile(const StringView& fileName, bool isWrite, bool isDeleteOnClose = false) { return (isWrite || isDeleteOnClose) && g_allowKeepFilesInMemory && g_rules->IsOutputFile(fileName); }
+bool IsOutputFile(const StringView& fileName, bool isWrite, bool isDeleteOnClose = false) { return (isWrite || isDeleteOnClose) && g_allowOutputFiles && g_rules->IsOutputFile(fileName); }
 
 
 bool EnsureMapped(DetouredHandle& handle, DWORD dwFileOffsetHigh = 0, DWORD dwFileOffsetLow = 0, SIZE_T numberOfBytesToMap = 0, void* baseAddress = nullptr)
@@ -675,6 +675,9 @@ void PreInit(const DetoursPayload& payload)
 	g_useMiMalloc = payload.useCustomAllocator;
 	g_runningRemote = payload.runningRemote;
 	g_isChild = payload.isChild;
+	g_allowKeepFilesInMemory = payload.allowKeepFilesInMemory;
+	g_allowOutputFiles = g_allowKeepFilesInMemory && payload.allowOutputFiles;
+	g_suppressLogging = payload.suppressLogging;
 	g_isDetachedProcess = g_rules->AllowDetach();
 	g_isRunningWine = payload.isRunningWine;
 	g_uiLanguage = payload.uiLanguage;
@@ -836,7 +839,7 @@ void Init(const DetoursPayload& payload, u64 startTime)
 		if (const wchar_t* lastBackslash = g_virtualApplication.Last('\\'))
 			g_virtualApplicationDir.Append(g_virtualApplication.data, (lastBackslash + 1 - g_virtualApplication.data));
 		else
-			FatalError(4444, L"What the heck: %s", g_virtualApplication.data);
+			FatalError(4444, L"What the heck: %s (%s)", g_virtualApplication.data, applicationBuffer.data);
 	}
 
 	const wchar_t* cmdLine = True_GetCommandLineW();
@@ -864,6 +867,7 @@ void Init(const DetoursPayload& payload, u64 startTime)
 			temp[sizeof_array(temp)-1] = 0;
 			cmdLine = temp;
 		}
+		DEBUG_LOG(L"ProcessId: %u", payload.processId);
 		DEBUG_LOG(L"Cmdline: %ls", cmdLine);
 		DEBUG_LOG(L"WorkingDir: %ls", g_virtualWorkingDir.data);
 		DEBUG_LOG(L"ExeDir: %ls", g_virtualApplicationDir.data);
