[UBA]
* Fixed missing include

diff --git a/Engine/Source/Programs/UnrealBuildAccelerator/Cli/Private/UbaCli.cpp b/Engine/Source/Programs/UnrealBuildAccelerator/Cli/Private/UbaCli.cpp
--- a/Engine/Source/Programs/UnrealBuildAccelerator/Cli/Private/UbaCli.cpp
+++ b/Engine/Source/Programs/UnrealBuildAccelerator/Cli/Private/UbaCli.cpp
@@ -6,6 +6,7 @@
 #include "UbaCoordinatorWrapper.h"
 #include "UbaFileAccessor.h"
 #include "UbaNetworkBackendTcp.h"
+#include "UbaPathUtils.h"
 #include "UbaPlatform.h"
 #include "UbaProtocol.h"
 #include "UbaRootPaths.h"
