[UBA]
* Forgot to bump trace version in earlier change
* Minor fix to make code compile with older compilers

diff --git a/Engine/Source/Programs/UnrealBuildAccelerator/Common/Public/UbaProcessHandle.h b/Engine/Source/Programs/UnrealBuildAccelerator/Common/Public/UbaProcessHandle.h
--- a/Engine/Source/Programs/UnrealBuildAccelerator/Common/Public/UbaProcessHandle.h
+++ b/Engine/Source/Programs/UnrealBuildAccelerator/Common/Public/UbaProcessHandle.h
@@ -12,7 +12,7 @@ namespace uba
 	struct ProcessStartInfo;
 
 	// Exit code returned by ProcessHandle.GetExitCode() if build was cancelled
-	inline constexpr u32 ProcessCancelExitCode = 99999;
+	constexpr u32 ProcessCancelExitCode = 99999;
 
 
 	struct ProcessLogLine
diff --git a/Engine/Source/Programs/UnrealBuildAccelerator/Common/Public/UbaTrace.h b/Engine/Source/Programs/UnrealBuildAccelerator/Common/Public/UbaTrace.h
--- a/Engine/Source/Programs/UnrealBuildAccelerator/Common/Public/UbaTrace.h
+++ b/Engine/Source/Programs/UnrealBuildAccelerator/Common/Public/UbaTrace.h
@@ -66,7 +66,7 @@ namespace uba
 	using Color = u32;
 	inline Color toColor(u8 r, u8 g, u8 b) { return (r << 16) + (g << 8) + b; }
 
-	static constexpr u32 TraceVersion = 30;
+	static constexpr u32 TraceVersion = 31;
 	static constexpr u32 TraceReadCompatibilityVersion = 6;
 
 	class Trace : public WorkTracker
