[Uba]
* Changed Config GetTable to return a pointer that can be null if table does not exist

diff --git a/Engine/Source/Programs/UnrealBuildAccelerator/Agent/Private/UbaAgent.cpp b/Engine/Source/Programs/UnrealBuildAccelerator/Agent/Private/UbaAgent.cpp
--- a/Engine/Source/Programs/UnrealBuildAccelerator/Agent/Private/UbaAgent.cpp
+++ b/Engine/Source/Programs/UnrealBuildAccelerator/Agent/Private/UbaAgent.cpp
@@ -1,6 +1,7 @@
 // Copyright Epic Games, Inc. All Rights Reserved.
 
 #include "UbaAWS.h"
+#include "UbaConfig.h"
 #include "UbaDirectoryIterator.h"
 #include "UbaNetworkBackendMemory.h"
 #include "UbaNetworkBackendQuic.h"
@@ -80,6 +81,7 @@ namespace uba
 		logger.Info(TC("  -mulcpu=<number>        This value multiplies with number of cpu to figure out max cpu. Defaults to 1.0"));
 		logger.Info(TC("  -maxcon=<number>        Max number of connections that can be started by agent. Defaults to \"%u\" (amount up to max will depend on ping)"), DefaultMaxConnectionCount);
 		logger.Info(TC("  -capacity=<gigaby>      Capacity of local store. Defaults to %u gigabytes"), DefaultCapacityGb);
+		logger.Info(TC("  -config=<file>          Config file that contains options for various systems"));
 		logger.Info(TC("  -quic                   Use Quic instead of tcp backend."));
 		logger.Info(TC("  -name=<name>            The identifier of this agent. Defaults to \"%s\" on this machine"), DefaultAgentName);
 		logger.Info(TC("  -stats[=<threshold>]    Print stats for each process if higher than threshold"));
@@ -407,6 +409,7 @@ namespace uba
 		StringBuffer<256> named;
 		StringBuffer<512> relaunchPath;
 		StringBuffer<256> eventFile;
+		StringBuffer<256> configFile;
 		TString command;
 		u16 port = DefaultPort;
 		u16 proxyPort = DefaultStorageProxyPort;
@@ -489,6 +492,14 @@ namespace uba
 				if (!value.Parse(storageCapacityGb))
 					return PrintHelp(TC("Invalid value for -capacity"));
 			}
+			else if (name.Equals(TC("-config")))
+			{
+				if (value.IsEmpty())
+					return PrintHelp(TC("-dir needs a value"));
+				if (int res = ExpandEnvironmentVariables(value))
+					return res;
+				configFile.Append(value);
+			}
 			else if (name.Equals(TC("-stats")))
 			{
 				if (value.IsEmpty())
@@ -771,6 +782,11 @@ namespace uba
 		dbgStr = TC(" (DEBUG)");
 		#endif
 		logger.Info(TC("UbaAgent v%s%s (Cpu: %u, MaxCon: %u, Dir: \"%s\", StoreCapacity: %uGb%s)"), Version, dbgStr, maxProcessCount, maxConnectionCount, g_rootDir.data, storageCapacityGb, extraInfo.data);
+		
+		Config config;
+		if (!configFile.IsEmpty())
+			config.LoadFromFile(logger, configFile.data);
+				
 		if (!eventFile.IsEmpty())
 			logger.Info(TC("  Will poll for external events in file %s"), eventFile.data);
 		logger.Info(TC(""));
@@ -790,6 +806,7 @@ namespace uba
 		{
 			// Create a uba storage quickly just to fix non-graceful shutdowns
 			StorageCreateInfo info(g_rootDir.data, logWriter);
+			info.Apply(config);
 			info.rootDir = g_rootDir.data;
 			info.casCapacityBytes = storageCapacity;
 			info.storeCompressed = storeCompressed;
@@ -930,7 +947,7 @@ namespace uba
 			}
 			else
 			{
-				const TString& desc = process.GetStartInfo().description;
+				const TString& desc = process.GetStartInfo().GetDescription();
 				StringBuffer<> name;
 				if (!desc.empty())
 					name.Append(desc);
@@ -986,6 +1003,7 @@ namespace uba
 			auto backendGuard = MakeGuard([networkBackend]() { delete networkBackend; });
 
 			NetworkClientCreateInfo ncci(logWriter);
+			//ncci.Apply(config);
 			ncci.sendSize = sendSize;
 			ncci.receiveTimeoutSeconds = receiveTimeoutSeconds;
 			if (hasCrypto)
@@ -1128,6 +1146,8 @@ namespace uba
 				};
 
 			StorageClientCreateInfo storageInfo(*client, g_rootDir.data);
+			storageInfo.Apply(config);
+			storageInfo.rootDir = g_rootDir.data;
 			storageInfo.casCapacityBytes = storageCapacity;
 			storageInfo.storeCompressed = storeCompressed;
 			storageInfo.sendCompressed = sendCompressed;
@@ -1160,6 +1180,7 @@ namespace uba
 				});
 
 			SessionClientCreateInfo info(*storageClient, *client, logWriter);
+			info.Apply(config);
 			info.maxProcessCount = maxProcessCount;
 			info.dedicated = poll;
 			info.maxIdleSeconds = maxIdleSeconds;
diff --git a/Engine/Source/Programs/UnrealBuildAccelerator/Common/Private/UbaCacheClient.cpp b/Engine/Source/Programs/UnrealBuildAccelerator/Common/Private/UbaCacheClient.cpp
--- a/Engine/Source/Programs/UnrealBuildAccelerator/Common/Private/UbaCacheClient.cpp
+++ b/Engine/Source/Programs/UnrealBuildAccelerator/Common/Private/UbaCacheClient.cpp
@@ -24,7 +24,10 @@ namespace uba
 {
 	void CacheClientCreateInfo::Apply(Config& config)
 	{
-		const ConfigTable& table = config.GetTable(TC("CacheClient"));
+		const ConfigTable* tablePtr = config.GetTable(TC("CacheClient"));
+		if (!tablePtr)
+			return;
+		const ConfigTable& table = *tablePtr;
 		table.GetValueAsBool(useDirectoryPreparsing, TC("UseDirectoryPreparsing"));
 		table.GetValueAsBool(validateCacheWritesInput, TC("ValidateCacheWritesInput"));
 		table.GetValueAsBool(validateCacheWritesOutput, TC("ValidateCacheWritesOutput"));
diff --git a/Engine/Source/Programs/UnrealBuildAccelerator/Common/Private/UbaConfig.cpp b/Engine/Source/Programs/UnrealBuildAccelerator/Common/Private/UbaConfig.cpp
--- a/Engine/Source/Programs/UnrealBuildAccelerator/Common/Private/UbaConfig.cpp
+++ b/Engine/Source/Programs/UnrealBuildAccelerator/Common/Private/UbaConfig.cpp
@@ -43,6 +43,14 @@ namespace uba
 		return false;
 	}
 
+	const ConfigTable* ConfigTable::GetTable(const tchar* name) const
+	{
+		auto findIt = m_tables.find(name);
+		if (findIt == m_tables.end())
+			return nullptr;
+		return &findIt->second;
+	}
+
 	bool Config::LoadFromFile(Logger& logger, const tchar* configFile)
 	{
 		m_isLoaded = true;
@@ -93,7 +101,9 @@ namespace uba
 			{
 				while (i != e)
 				{
-					if (!((*i >= 'a' && *i <= 'z') || (*i >= 'A' && *i <= 'Z')))
+					tchar c = *i;
+					bool validChar = (c >= 'a' && c <= 'z') || (c >= 'A' && c <= 'Z') || c == '_' || c == '-';
+					if (!validChar)
 						return *i;
 					out.Append(*i);
 					++i;
@@ -118,7 +128,7 @@ namespace uba
 				return 0;
 			};
 
-		ConfigTable* activeTable = &m_globalTable;
+		ConfigTable* activeTable = this;
 		while (true)
 		{
 			char token = consumeEmpty();
@@ -148,7 +158,7 @@ namespace uba
 					return logger.Error(TC("Unexpected token %c after group %s"), tableName.data);
 				++i;
 				activeTable = &m_tables.try_emplace(tableName.data).first->second;
-				activeTable->m_parent = &m_globalTable;
+				activeTable->m_parent = this;
 			}
 			else
 			{
@@ -185,12 +195,4 @@ namespace uba
 	{
 		return m_isLoaded;
 	}
-
-	const ConfigTable& Config::GetTable(const tchar* name) const
-	{
-		auto findIt = m_tables.find(name);
-		if (findIt == m_tables.end())
-			return m_globalTable;
-		return findIt->second;
-	}
 }
diff --git a/Engine/Source/Programs/UnrealBuildAccelerator/Common/Private/UbaScheduler.cpp b/Engine/Source/Programs/UnrealBuildAccelerator/Common/Private/UbaScheduler.cpp
--- a/Engine/Source/Programs/UnrealBuildAccelerator/Common/Private/UbaScheduler.cpp
+++ b/Engine/Source/Programs/UnrealBuildAccelerator/Common/Private/UbaScheduler.cpp
@@ -3,6 +3,7 @@
 #include "UbaScheduler.h"
 #include "UbaApplicationRules.h"
 #include "UbaCacheClient.h"
+#include "UbaConfig.h"
 #include "UbaNetworkServer.h"
 #include "UbaProcess.h"
 #include "UbaProcessStartInfoHolder.h"
@@ -76,6 +77,16 @@ namespace uba
 		Vector<ProcessLogLine> logLines;
 	};
 
+	void SchedulerCreateInfo::Apply(Config& config)
+	{
+		if (const ConfigTable* table = config.GetTable(TC("Scheduler")))
+		{
+			table->GetValueAsBool(enableProcessReuse, TC("EnableProcessReuse"));
+			table->GetValueAsBool(forceRemote, TC("ForceRemote"));
+			table->GetValueAsBool(forceNative, TC("ForceNative"));
+		}
+	}
+
 
 	Scheduler::Scheduler(const SchedulerCreateInfo& info)
 	:	m_session(info.session)
@@ -84,6 +95,7 @@ namespace uba
 	,	m_enableProcessReuse(info.enableProcessReuse)
 	,	m_forceRemote(info.forceRemote)
 	,	m_forceNative(info.forceNative)
+	,	m_processConfigs(info.processConfigs)
 	,	m_cacheClient(info.cacheClient)
 	,	m_writeToCache(info.writeToCache && info.cacheClient)
 	{
@@ -166,6 +178,22 @@ namespace uba
 		entry.status = useCache ? ProcessStatus_QueuedForCache : ProcessStatus_QueuedForRun;
 		entry.canDetour = info.canDetour;
 		entry.canExecuteRemotely = info.canExecuteRemotely && info.canDetour;
+
+		if (m_processConfigs)
+		{
+			auto name = info2->application.c_str();
+			if (auto lastSeparator = TStrrchr(name, PathSeparator))
+				name = lastSeparator + 1;
+			StringBuffer<128> lower(name);
+			lower.MakeLower();
+			lower.Replace('.', '_');
+			if (const ConfigTable* processConfig = m_processConfigs->GetTable(lower.data))
+			{
+				processConfig->GetValueAsBool(entry.canExecuteRemotely, TC("CanExecuteRemotely"));
+				processConfig->GetValueAsBool(entry.canDetour, TC("CanDetour"));
+			}
+		}
+
 		lock.Leave();
 
 		UpdateQueueCounter(1);
diff --git a/Engine/Source/Programs/UnrealBuildAccelerator/Common/Private/UbaSession.cpp b/Engine/Source/Programs/UnrealBuildAccelerator/Common/Private/UbaSession.cpp
--- a/Engine/Source/Programs/UnrealBuildAccelerator/Common/Private/UbaSession.cpp
+++ b/Engine/Source/Programs/UnrealBuildAccelerator/Common/Private/UbaSession.cpp
@@ -3,6 +3,7 @@
 #include "UbaSession.h"
 #include "UbaBottleneck.h"
 #include "UbaCompressedObjFileHeader.h"
+#include "UbaConfig.h"
 #include "UbaFileAccessor.h"
 #include "UbaObjectFile.h"
 #include "UbaProcess.h"
@@ -194,6 +195,28 @@ namespace uba
 		process->AddRef();
 	}
 
+	void SessionCreateInfo::Apply(Config& config)
+	{
+		const ConfigTable* tablePtr = config.GetTable(TC("Session"));
+		if (!tablePtr)
+			return;
+		const ConfigTable& table = *tablePtr;
+		table.GetValueAsString(rootDir, TC("RootDir"));
+		table.GetValueAsString(traceName, TC("TraceName"));
+		table.GetValueAsString(traceOutputFile, TC("TraceOutputFile"));
+		table.GetValueAsString(extraInfo, TC("ExtraInfo"));
+		table.GetValueAsBool(logToFile, TC("LogToFile"));
+		table.GetValueAsBool(useUniqueId, TC("UseUniqueId"));
+		table.GetValueAsBool(disableCustomAllocator, TC("DisableCustomAllocator"));
+		table.GetValueAsBool(launchVisualizer, TC("LaunchVisualizer"));
+		table.GetValueAsBool(allowMemoryMaps, TC("AllowMemoryMaps"));
+		table.GetValueAsBool(shouldWriteToDisk, TC("ShouldWriteToDisk"));
+		table.GetValueAsBool(traceEnabled, TC("TraceEnabled"));
+		table.GetValueAsBool(detailedTrace, TC("DetailedTrace"));
+		table.GetValueAsBool(storeObjFilesCompressed, TC("StoreObjFilesCompressed"));
+		table.GetValueAsBool(extractObjFilesSymbols, TC("ExtractObjFilesSymbols"));
+	}
+
 	void Session::AddEnvironmentVariableNoLock(const tchar* key, const tchar* value)
 	{
 		m_environmentVariables.insert(m_environmentVariables.end(), key, key + TStrlen(key));
diff --git a/Engine/Source/Programs/UnrealBuildAccelerator/Common/Private/UbaSessionServer.cpp b/Engine/Source/Programs/UnrealBuildAccelerator/Common/Private/UbaSessionServer.cpp
--- a/Engine/Source/Programs/UnrealBuildAccelerator/Common/Private/UbaSessionServer.cpp
+++ b/Engine/Source/Programs/UnrealBuildAccelerator/Common/Private/UbaSessionServer.cpp
@@ -2,6 +2,7 @@
 
 #include "UbaSessionServer.h"	
 #include "UbaApplicationRules.h"
+#include "UbaConfig.h"
 #include "UbaNetworkServer.h"
 #include "UbaProcess.h"
 #include "UbaProcessStartInfoHolder.h"
@@ -90,6 +91,15 @@ namespace uba
 		u32 m_knownInputsCount = 0;
 	};
 
+	void SessionServerCreateInfo::Apply(Config& config)
+	{
+		SessionCreateInfo::Apply(config);
+
+		if (const ConfigTable* table = config.GetTable(TC("Session")))
+		{
+			table->GetValueAsBool(remoteLogEnabled, TC("RemoteLogEnabled"));
+		}
+	}
 
 	SessionServer::SessionServer(const SessionServerCreateInfo& info, const u8* environment, u32 environmentSize)
 	:	Session(info, TC("UbaSessionServer"), false, &info.server)
diff --git a/Engine/Source/Programs/UnrealBuildAccelerator/Common/Private/UbaStorage.cpp b/Engine/Source/Programs/UnrealBuildAccelerator/Common/Private/UbaStorage.cpp
--- a/Engine/Source/Programs/UnrealBuildAccelerator/Common/Private/UbaStorage.cpp
+++ b/Engine/Source/Programs/UnrealBuildAccelerator/Common/Private/UbaStorage.cpp
@@ -55,7 +55,10 @@ namespace uba
 
 	void StorageCreateInfo::Apply(Config& config)
 	{
-		const ConfigTable& table = config.GetTable(TC("Storage"));
+		const ConfigTable* tablePtr = config.GetTable(TC("Storage"));
+		if (!tablePtr)
+			return;
+		const ConfigTable& table = *tablePtr;
 		table.GetValueAsString(rootDir, TC("RootDir"));
 	}
 
diff --git a/Engine/Source/Programs/UnrealBuildAccelerator/Common/Public/UbaConfig.h b/Engine/Source/Programs/UnrealBuildAccelerator/Common/Public/UbaConfig.h
--- a/Engine/Source/Programs/UnrealBuildAccelerator/Common/Public/UbaConfig.h
+++ b/Engine/Source/Programs/UnrealBuildAccelerator/Common/Public/UbaConfig.h
@@ -17,9 +17,12 @@ namespace uba
 		bool GetValueAsU32(u32& out, const tchar* key) const;
 		bool GetValueAsBool(bool& out, const tchar* key) const;
 
+		const ConfigTable* GetTable(const tchar* name) const;
+
 	private:
 		ConfigTable* m_parent = nullptr;
 		UnorderedMap<TString, TString> m_values;
+		UnorderedMap<TString, ConfigTable> m_tables;
 		friend Config;
 	};
 
@@ -31,10 +34,7 @@ namespace uba
 		bool LoadFromText(Logger& logger, const char* text, u64 textLen);
 		bool IsLoaded() const;
 
-		const ConfigTable& GetTable(const tchar* name) const;
 
 		bool m_isLoaded = false;
-		ConfigTable m_globalTable;
-		UnorderedMap<TString, ConfigTable> m_tables;
 	};
 }
diff --git a/Engine/Source/Programs/UnrealBuildAccelerator/Common/Public/UbaScheduler.h b/Engine/Source/Programs/UnrealBuildAccelerator/Common/Public/UbaScheduler.h
--- a/Engine/Source/Programs/UnrealBuildAccelerator/Common/Public/UbaScheduler.h
+++ b/Engine/Source/Programs/UnrealBuildAccelerator/Common/Public/UbaScheduler.h
@@ -9,6 +9,8 @@
 namespace uba
 {
 	class CacheClient;
+	class Config;
+	class ConfigTable;
 	class Process;
 	class RootPaths;
 	class SessionServer;
@@ -19,6 +21,8 @@ namespace uba
 	{
 		SchedulerCreateInfo(SessionServer& s) : session(s) {}
 
+		void Apply(Config& config);
+
 		SessionServer& session;
 		CacheClient* cacheClient = nullptr; // Set cache client for scheduler to use when building
 		u32 maxLocalProcessors = ~0u; // Max local processors to use. ~0u means it will use all processors
@@ -26,6 +30,7 @@ namespace uba
 		bool forceRemote = false; // Force all processes that can run remotely to run remotely.
 		bool forceNative = false; // Force all processes to run native (not detoured)
 		bool writeToCache = false; // Set to true in combination with setting cacheClient to populate cache
+		ConfigTable* processConfigs = nullptr;
 	};
 
 	struct EnqueueProcessInfo
@@ -64,6 +69,8 @@ namespace uba
 
 		void SetProcessFinishedCallback(const Function<void(const ProcessHandle&)>& processFinished); // Set callback 
 
+		SessionServer& GetSession() { return m_session; }
+		
 	private:
 		struct ExitProcessInfo;
 		struct ProcessStartInfo2;
@@ -101,8 +108,8 @@ namespace uba
 			u32* dependencies;
 			u32 dependencyCount;
 			ProcessStatus status;
-			u8 canDetour;
-			u8 canExecuteRemotely;
+			bool canDetour;
+			bool canExecuteRemotely;
 		};
 
 		ReaderWriterLock m_processEntriesLock;
@@ -117,6 +124,7 @@ namespace uba
 		bool m_enableProcessReuse;
 		bool m_forceRemote;
 		bool m_forceNative;
+		ConfigTable* m_processConfigs = nullptr;
 
 		float m_activeLocalProcessWeight = 0.0f;
 		u32 m_activeCacheQueries = 0;
diff --git a/Engine/Source/Programs/UnrealBuildAccelerator/Common/Public/UbaSessionCreateInfo.h b/Engine/Source/Programs/UnrealBuildAccelerator/Common/Public/UbaSessionCreateInfo.h
--- a/Engine/Source/Programs/UnrealBuildAccelerator/Common/Public/UbaSessionCreateInfo.h
+++ b/Engine/Source/Programs/UnrealBuildAccelerator/Common/Public/UbaSessionCreateInfo.h
@@ -6,12 +6,15 @@
 
 namespace uba
 {
+	class Config;
 	class Storage;
 
 	struct SessionCreateInfo
 	{
 		SessionCreateInfo(Storage& s, LogWriter& w = g_consoleLogWriter) : storage(s), logWriter(w) {}
 
+		void Apply(Config& config);
+
 		Storage& storage;
 		LogWriter& logWriter;
 		const tchar* rootDir = nullptr;			// Root dir for logs, binaries, temp files
diff --git a/Engine/Source/Programs/UnrealBuildAccelerator/Common/Public/UbaSessionServerCreateInfo.h b/Engine/Source/Programs/UnrealBuildAccelerator/Common/Public/UbaSessionServerCreateInfo.h
--- a/Engine/Source/Programs/UnrealBuildAccelerator/Common/Public/UbaSessionServerCreateInfo.h
+++ b/Engine/Source/Programs/UnrealBuildAccelerator/Common/Public/UbaSessionServerCreateInfo.h
@@ -12,6 +12,8 @@ namespace uba
 	{
 		SessionServerCreateInfo(Storage& s, NetworkServer& c, LogWriter& writer = g_consoleLogWriter) : SessionCreateInfo(s, writer), server(c) {}
 
+		void Apply(Config& config);
+
 		NetworkServer& server;
 		u8 memWaitLoadPercent = 85; // When memory usage goes above this percent, no new processes will be spawned until back below
 		u8 memKillLoadPercent = 95; // When memory usage goes above this percent, newest processes will be killed to bring it back below
diff --git a/Engine/Source/Programs/UnrealBuildAccelerator/Test/Private/UbaTestConfig.h b/Engine/Source/Programs/UnrealBuildAccelerator/Test/Private/UbaTestConfig.h
--- a/Engine/Source/Programs/UnrealBuildAccelerator/Test/Private/UbaTestConfig.h
+++ b/Engine/Source/Programs/UnrealBuildAccelerator/Test/Private/UbaTestConfig.h
@@ -19,7 +19,10 @@ namespace uba
 		if (!config.LoadFromText(logger, configText, strlen(configText)))
 			return false;
 
-		const ConfigTable& table = config.GetTable(TC("CacheClient"));
+		const ConfigTable* tablePtr = config.GetTable(TC("CacheClient"));
+		if (!tablePtr)
+			return false;
+		const ConfigTable& table = *tablePtr;
 		bool test = false;
 		if (!table.GetValueAsBool(test, TC("UseDirectoryPreparsing")))
 			return false;
