[UBA]
* Added option to be able to trace child processes in visualizer
* Made WriteDirectoryEntries offset output parameter optional
* Fixed so GetKeyAndFixedName can handle relative paths

diff --git a/Engine/Source/Programs/UnrealBuildAccelerator/Common/Private/UbaProcess.cpp b/Engine/Source/Programs/UnrealBuildAccelerator/Common/Private/UbaProcess.cpp
--- a/Engine/Source/Programs/UnrealBuildAccelerator/Common/Private/UbaProcess.cpp
+++ b/Engine/Source/Programs/UnrealBuildAccelerator/Common/Private/UbaProcess.cpp
@@ -1255,7 +1255,7 @@ namespace uba
 		temp.Append(m_startInfo.logFile);
 		if (TStrcmp(temp.data + temp.count - 4, TC(".log")) == 0)
 			temp.Resize(temp.count - 4);
-		temp.Appendf(TC("_CHILD%u.log"), u32(m_childProcesses.size()));
+		temp.Appendf(TC("_CHILD%03u.log"), u32(m_childProcesses.size()));
 		return temp.data;
 	}
 
diff --git a/Engine/Source/Programs/UnrealBuildAccelerator/Common/Private/UbaSession.cpp b/Engine/Source/Programs/UnrealBuildAccelerator/Common/Private/UbaSession.cpp
--- a/Engine/Source/Programs/UnrealBuildAccelerator/Common/Private/UbaSession.cpp
+++ b/Engine/Source/Programs/UnrealBuildAccelerator/Common/Private/UbaSession.cpp
@@ -213,6 +213,7 @@ namespace uba
 		table.GetValueAsBool(shouldWriteToDisk, TC("ShouldWriteToDisk"));
 		table.GetValueAsBool(traceEnabled, TC("TraceEnabled"));
 		table.GetValueAsBool(detailedTrace, TC("DetailedTrace"));
+		table.GetValueAsBool(traceChildProcesses, TC("TraceChildProcesses"));
 		table.GetValueAsBool(storeObjFilesCompressed, TC("StoreObjFilesCompressed"));
 		table.GetValueAsBool(extractObjFilesSymbols, TC("ExtractObjFilesSymbols"));
 	}
@@ -422,10 +423,13 @@ namespace uba
 		WriteDirectoryEntriesRecursive(parentKey, dirPath, parentOffset);
 	}
 
-	u32 Session::WriteDirectoryEntries(const StringKey& dirKey, tchar* dirPath, u32& outTableOffset)
+	u32 Session::WriteDirectoryEntries(const StringKey& dirKey, tchar* dirPath, u32* outTableOffset)
 	{
 		auto& dirTable = m_directoryTable;
-		WriteDirectoryEntriesRecursive(dirKey, dirPath, outTableOffset);
+		u32 temp;
+		if (!outTableOffset)
+			outTableOffset = &temp;
+		WriteDirectoryEntriesRecursive(dirKey, dirPath, *outTableOffset);
 		SCOPED_READ_LOCK(dirTable.m_memoryLock, memoryLock);
 		return dirTable.m_memorySize;
 	}
@@ -714,8 +718,7 @@ namespace uba
 		// When not writing to disk we need to populate lookup before adding non-written files.. otherwise they will be lost once lookup is actually populated
 		if (!shouldWriteToDisk)
 		{
-			u32 offset;
-			u32 res = WriteDirectoryEntries(dirKey, dirName.data, offset);
+			u32 res = WriteDirectoryEntries(dirKey, dirName.data);
 			UBA_ASSERT(res); (void)res;
 		}
 
@@ -741,9 +744,8 @@ namespace uba
 		// If adding a file, clearly it does exist.. so let's reparse it.
 		if (dir.parseOffset == 2)
 		{
-			u32 offset;
 			dirLock.Leave();
-			u32 res = WriteDirectoryEntries(dirKey, dirName.data, offset);
+			u32 res = WriteDirectoryEntries(dirKey, dirName.data);
 			UBA_ASSERT(res); (void)res;
 			dirLock.Enter();
 		}
@@ -907,7 +909,7 @@ namespace uba
 		}
 
 		#if UBA_DEBUG_LOGGER
-		g_debugLogger.Info(TC("TRACKDEL    %s (Key: %s)\n"), fileName, KeyToString(fileNameKey).data);
+		g_debugLogger.Info(TC("TRACKDEL    %s (Key: %s)\n"), fileName.data, KeyToString(fileNameKey).data);
 		#endif
 
 		SCOPED_WRITE_LOCK(dirTable.m_memoryLock, memoryLock);
@@ -1009,6 +1011,7 @@ namespace uba
 		m_storeObjFilesCompressed = info.storeObjFilesCompressed;
 
 		m_detailedTrace = info.detailedTrace;
+		m_traceChildProcesses = info.traceChildProcesses;
 		m_logToFile = info.logToFile;
 		if (info.extraInfo)
 			m_extraInfo = info.extraInfo;
@@ -1301,7 +1304,13 @@ namespace uba
 
 	StringKey GetKeyAndFixedName(StringBuffer<>& fixedFilePath, const tchar* filePath)
 	{
-		FixPath2(filePath, nullptr, 0, fixedFilePath.data, fixedFilePath.capacity, &fixedFilePath.count);
+		StringBuffer<> workingDir;
+		if (!IsAbsolutePath(filePath))
+		{
+			GetCurrentDirectoryW(workingDir);
+			workingDir.EnsureEndsWithSlash();
+		}
+		FixPath2(filePath, workingDir.data, workingDir.count, fixedFilePath.data, fixedFilePath.capacity, &fixedFilePath.count);
 
 		StringKey dirKey;
 		StringBuffer<> dirNameForHash;
@@ -1432,8 +1441,8 @@ namespace uba
 	{
 		u32 processId = process.GetId();
 
-		if (!process.IsChild())
-			m_trace.ProcessAdded(sessionId, processId, process.GetStartInfo().description);
+		if (!process.IsChild() || m_traceChildProcesses)
+			m_trace.ProcessAdded(sessionId, processId, process.GetStartInfo().GetDescription());
 
 		SCOPED_WRITE_LOCK(m_processesLock, lock);
 		bool success = m_processes.try_emplace(processId, ProcessHandle(&process)).second;
@@ -1450,7 +1459,7 @@ namespace uba
 
 		u32 id = process.GetId();
 
-		if (!process.IsChild())
+		if (!process.IsChild() || m_traceChildProcesses)
 		{
 			StackBinaryWriter<1024> writer;
 			process.m_processStats.Write(writer);
@@ -2085,8 +2094,7 @@ namespace uba
 				return true;
 			// Both these functions need to be called. otherwise we can get created directories that does not end up in directory table
 			RegisterCreateFileForWrite(msg.nameKey, msg.name, true);
-			u32 offset;
-			WriteDirectoryEntries(dirKey, dirName.data, offset);
+			WriteDirectoryEntries(dirKey, dirName.data);
 		}
 		else
 		{
@@ -2123,7 +2131,7 @@ namespace uba
 	bool Session::GetListDirectoryInfo(ListDirectoryResponse& out, tchar* dirName, const StringKey& dirKey)
 	{
 		u32 tableOffset;
-		u32 tableSize = WriteDirectoryEntries(dirKey, dirName, tableOffset);
+		u32 tableSize = WriteDirectoryEntries(dirKey, dirName, &tableOffset);
 		out.tableOffset = tableOffset;
 		out.tableSize = tableSize;
 		return true;
@@ -2818,6 +2826,6 @@ namespace uba
 
 		static Atomic<u32> counter;
 		if (addCounterSuffix)
-			out.Append('_').AppendValue(counter++);
+			out.Appendf(TC("_%03u"), counter++);
 	}
 }
diff --git a/Engine/Source/Programs/UnrealBuildAccelerator/Common/Public/UbaSession.h b/Engine/Source/Programs/UnrealBuildAccelerator/Common/Public/UbaSession.h
--- a/Engine/Source/Programs/UnrealBuildAccelerator/Common/Public/UbaSession.h
+++ b/Engine/Source/Programs/UnrealBuildAccelerator/Common/Public/UbaSession.h
@@ -107,7 +107,7 @@ namespace uba
 		bool IsRarelyReadAfterWritten(ProcessImpl& process, const StringView& fileName) const;
 		bool IsKnownSystemFile(const tchar* applicationName);
 		bool ShouldWriteToDisk(const StringView& fileName);
-		u32 WriteDirectoryEntries(const StringKey& dirKey, tchar* dirPath, u32& outTableOffset);
+		u32 WriteDirectoryEntries(const StringKey& dirKey, tchar* dirPath, u32* outTableOffset = nullptr);
 		u32 AddFileMapping(StringKey fileNameKey, const tchar* fileName, const tchar* newFileName, u64 fileSize = InvalidValue);
 		
 		struct MemoryMap { StringBuffer<128> name; u64 size = 0; };
@@ -180,6 +180,7 @@ namespace uba
 		bool m_allowMemoryMaps;
 		bool m_shouldWriteToDisk;
 		bool m_detailedTrace;
+		bool m_traceChildProcesses;
 		bool m_logToFile;
 		bool m_storeObjFilesCompressed;
 
diff --git a/Engine/Source/Programs/UnrealBuildAccelerator/Common/Public/UbaSessionCreateInfo.h b/Engine/Source/Programs/UnrealBuildAccelerator/Common/Public/UbaSessionCreateInfo.h
--- a/Engine/Source/Programs/UnrealBuildAccelerator/Common/Public/UbaSessionCreateInfo.h
+++ b/Engine/Source/Programs/UnrealBuildAccelerator/Common/Public/UbaSessionCreateInfo.h
@@ -29,6 +29,7 @@ namespace uba
 		bool shouldWriteToDisk = true;			// Set to false to skip writing output files to disk
 		bool traceEnabled = false;				// Set to true to always create in-memory trace data. Is not needed if traceName, traceOutputFile or launchVisualizer is set
 		bool detailedTrace = false;				// Enable detailed trace to include jobs, individual file I/O etc in trace dump
+		bool traceChildProcesses = false;		// Trace and visualize child processes so they can be seen in visualizer
 		bool storeObjFilesCompressed = false;	// Compiler will write .obj files compressed to disk and linker will decompress when reading
 		bool extractObjFilesSymbols = false;	// Will extract import/export symbols and clear export table from obj file.
 		u64 deleteSessionsOlderThanSeconds = 12 * 60 * 60; // Delete session folders older than 12 hours by default . Set to 0 to not delete or 1 to delete all
