[UBA]
* Fixed another compile error in UbaCli

diff --git a/Engine/Source/Programs/UnrealBuildAccelerator/Cli/Private/UbaCli.cpp b/Engine/Source/Programs/UnrealBuildAccelerator/Cli/Private/UbaCli.cpp
--- a/Engine/Source/Programs/UnrealBuildAccelerator/Cli/Private/UbaCli.cpp
+++ b/Engine/Source/Programs/UnrealBuildAccelerator/Cli/Private/UbaCli.cpp
@@ -524,7 +524,7 @@ namespace uba
 			if (application.empty())
 				return PrintHelp(TC("No executable provided"));
 
-			if (!IsAbsolutePath(application))
+			if (!IsAbsolutePath(application.c_str()))
 			{
 				StringBuffer<> fullApplicationName;
 				if (!SearchPathForFile(logger, fullApplicationName, application.c_str(), currentDir.data))
