[UBA]
* Added tests for Create/Remove directory
* Misc small cleanups

diff --git a/Engine/Source/Programs/UnrealBuildAccelerator/Common/Private/UbaProcess.cpp b/Engine/Source/Programs/UnrealBuildAccelerator/Common/Private/UbaProcess.cpp
--- a/Engine/Source/Programs/UnrealBuildAccelerator/Common/Private/UbaProcess.cpp
+++ b/Engine/Source/Programs/UnrealBuildAccelerator/Common/Private/UbaProcess.cpp
@@ -264,7 +264,7 @@ namespace uba
 			}
 
 			if (err.IsEmpty())
-				err.Appendf(TC("ERROR: Process %llu %s (%s) not active but did not get exit message. Received %u messages (GetExitCodeProcess returned %u)"), u64(m_nativeProcessHandle), m_description.c_str(), m_realApplication.c_str(), m_messageCount, exitCode);
+				err.Appendf(TC("ERROR: Process %llu %s (%s) not active but did not get exit message. Received %u messages (GetExitCodeProcess returned %u)"), u64(m_nativeProcessHandle), m_startInfo.GetDescription(), m_realApplication.c_str(), m_messageCount, exitCode);
 			LogLine(false, err.data, LogEntryType_Error);
 			m_nativeProcessExitCode = UBA_EXIT_CODE(666);
 		}
@@ -329,7 +329,7 @@ namespace uba
 					break;
 					
 				StringBuffer<> err;
-				err.Appendf(TC("Process %u (%s) %s by signal %i. Received %u messages. Execution time: %s."), m_nativeProcessId, m_description.c_str(), codeType, signalInfo.si_status, m_messageCount, TimeToText(GetTime() - m_startTime).str);
+				err.Appendf(TC("Process %u (%s) %s by signal %i. Received %u messages. Execution time: %s."), m_nativeProcessId, m_startInfo.GetDescription(), codeType, signalInfo.si_status, m_messageCount, TimeToText(GetTime() - m_startTime).str);
 				LogLine(false, err.data, LogEntryType_Error);
 				m_nativeProcessExitCode = UBA_EXIT_CODE(666); // We do exit code 666 to trigger non-uba retry on the outside
 				return false;
@@ -348,7 +348,7 @@ namespace uba
 			}
 
 			StringBuffer<> err;
-			err.Appendf(TC("ERROR: Process %u (%s) not active but did not get exit message. Received %u messages. Signal code: %i. Exit value or signal: %i. Execution time: %s."), m_nativeProcessId, m_description.c_str(), m_messageCount, signalInfo.si_code, signalInfo.si_status, TimeToText(GetTime() - m_startTime).str);
+			err.Appendf(TC("ERROR: Process %u (%s) not active but did not get exit message. Received %u messages. Signal code: %i. Exit value or signal: %i. Execution time: %s."), m_nativeProcessId, m_startInfo.GetDescription(), m_messageCount, signalInfo.si_code, signalInfo.si_status, TimeToText(GetTime() - m_startTime).str);
 			LogLine(false, err.data, LogEntryType_Error);
 			m_nativeProcessExitCode = UBA_EXIT_CODE(666);
 		}
@@ -953,9 +953,11 @@ namespace uba
 		if (nativeProcessHandle)
 		{
 			DuplicateHandle((HANDLE)m_nativeProcessHandle, nativeProcessHandle, GetCurrentProcess(), (HANDLE*)&process.m_nativeProcessHandle, 0, false, DUPLICATE_SAME_ACCESS);
-			UBA_ASSERT(process.m_nativeProcessHandle && process.m_nativeProcessHandle != InvalidProcHandle);
+			if (!process.m_nativeProcessHandle || process.m_nativeProcessHandle == InvalidProcHandle)
+				return m_session.m_logger.Error(TC("Failed to duplicate handle for child process"));
 			DuplicateHandle((HANDLE)m_nativeProcessHandle, nativeThreadHandle, GetCurrentProcess(), &process.m_nativeThreadHandle, 0, false, DUPLICATE_SAME_ACCESS);
-			UBA_ASSERT(process.m_nativeThreadHandle && process.m_nativeThreadHandle != INVALID_HANDLE_VALUE);
+			if (!process.m_nativeThreadHandle || process.m_nativeThreadHandle == INVALID_HANDLE_VALUE)
+				return m_session.m_logger.Error(TC("Failed to duplicate handle for child thread"));
 			process.m_nativeProcessId = nativeProcessId;
 		}
 #else
@@ -1517,9 +1519,9 @@ namespace uba
 				payload.logFile.Append(m_startInfo.logFile);
 			}
 
-			if (!DetourCopyPayloadToProcessEx((HANDLE)m_nativeProcessHandle, DetoursPayloadGuid, &payload, sizeof(payload)))
+			if (!DetourCopyPayloadToProcess((HANDLE)m_nativeProcessHandle, DetoursPayloadGuid, &payload, sizeof(payload)))
 			{
-				logger.Error(TC("Failed to copy payload to process"));//% ls."), commandLine.c_str());
+				logger.Error(TC("Failed to copy payload to process (%s)"), LastErrorToText().data);
 				return UBA_EXIT_CODE(9);
 			}
 		}
@@ -1840,11 +1842,11 @@ namespace uba
 						hadTimeout = true;
 						const tchar* gotMessage = m_gotExitMessage ? TC("Got") : TC("Did not get");
 						const tchar* isCancelledNewCheck = IsCancelled() ? TC("true") : TC("false");
-						logger.Info(TC("WaitForSingleObject timed out after 120 seconds waiting for process %s to exit (Exit code %u, %s ExitMessage and wrote %u files. Cancelled: %s. Runtime: %s). Will terminate and wait again"), m_startInfo.description, m_nativeProcessExitCode, gotMessage, u32(m_writtenFiles.size()), isCancelledNewCheck, TimeToText(GetTime() - m_startTime).str);
+						logger.Info(TC("WaitForSingleObject timed out after 120 seconds waiting for process %s to exit (Exit code %u, %s ExitMessage and wrote %u files. Cancelled: %s. Runtime: %s). Will terminate and wait again"), m_startInfo.GetDescription(), m_nativeProcessExitCode, gotMessage, u32(m_writtenFiles.size()), isCancelledNewCheck, TimeToText(GetTime() - m_startTime).str);
 						TerminateProcess((HANDLE)handle, m_nativeProcessExitCode);
 						continue;
 					}
-					logger.Error(TC("WaitForSingleObject failed while waiting for process %s to exit even after terminating it (%s)"), m_startInfo.description, LastErrorToText().data);
+					logger.Error(TC("WaitForSingleObject failed while waiting for process %s to exit even after terminating it (%s)"), m_startInfo.GetDescription(), LastErrorToText().data);
 				}
 				else if (res == WAIT_FAILED)
 					logger.Error(TC("WaitForSingleObject failed while waiting for process to exit (%s)"), LastErrorToText().data);
diff --git a/Engine/Source/Programs/UnrealBuildAccelerator/Common/Private/UbaSessionClient.cpp b/Engine/Source/Programs/UnrealBuildAccelerator/Common/Private/UbaSessionClient.cpp
--- a/Engine/Source/Programs/UnrealBuildAccelerator/Common/Private/UbaSessionClient.cpp
+++ b/Engine/Source/Programs/UnrealBuildAccelerator/Common/Private/UbaSessionClient.cpp
@@ -1686,7 +1686,7 @@ namespace uba
 						if (session.m_killRandomIndex != ~0u && session.m_killRandomCounter++ == session.m_killRandomIndex)
 						{
 							session.m_loop = false;
-							session.m_logger.Info(TC("Killed random process (%s)"), process.m_description.c_str());
+							session.m_logger.Info(TC("Killed random process (%s)"), process.m_startInfo.GetDescription());
 							return;
 						}
 
diff --git a/Engine/Source/Programs/UnrealBuildAccelerator/Test/Private/UbaTestStress.h b/Engine/Source/Programs/UnrealBuildAccelerator/Test/Private/UbaTestStress.h
--- a/Engine/Source/Programs/UnrealBuildAccelerator/Test/Private/UbaTestStress.h
+++ b/Engine/Source/Programs/UnrealBuildAccelerator/Test/Private/UbaTestStress.h
@@ -126,7 +126,7 @@ namespace uba
 
 					if (u32 exitCode = ph.GetExitCode())
 					{
-						rec->logger.Error(TC("Process %s failed with exit code %u"), ph.GetStartInfo().description, exitCode);
+						rec->logger.Error(TC("Process %s failed with exit code %u"), ph.GetStartInfo().GetDescription(), exitCode);
 						rec->error = true;
 					}
 					else
diff --git a/Engine/Source/Programs/UnrealBuildAccelerator/TestApp/Private/Windows/UbaTestAppWin.cpp b/Engine/Source/Programs/UnrealBuildAccelerator/TestApp/Private/Windows/UbaTestAppWin.cpp
--- a/Engine/Source/Programs/UnrealBuildAccelerator/TestApp/Private/Windows/UbaTestAppWin.cpp
+++ b/Engine/Source/Programs/UnrealBuildAccelerator/TestApp/Private/Windows/UbaTestAppWin.cpp
@@ -135,6 +135,20 @@ int wmain(int argc, wchar_t* argv[])
 
 			CopyFile(L"FileW2", L"FileWF", false);
 		}
+
+		{
+			if (!CreateDirectoryW(L"DirA", NULL))
+				return LogError(L"Failed to create directory");
+
+			if (GetFileAttributesW(L"DirA") == 0)
+				return LogError(L"Failed to get attributes of directory");
+
+			if (!RemoveDirectoryW(L"DirA"))
+				return LogError(L"Failed to remove directory");
+
+			if (GetFileAttributesW(L"DirA") != INVALID_FILE_ATTRIBUTES)
+				return LogError(L"Found attributes of deleted directory");
+		}
 	}
 	else if (wcscmp(argv[1], L"-reuse") == 0)
 	{
