[UBA]
* Exposed more options out through creation structs and configs

diff --git a/Engine/Source/Programs/UnrealBuildAccelerator/Common/Private/UbaProcess.cpp b/Engine/Source/Programs/UnrealBuildAccelerator/Common/Private/UbaProcess.cpp
--- a/Engine/Source/Programs/UnrealBuildAccelerator/Common/Private/UbaProcess.cpp
+++ b/Engine/Source/Programs/UnrealBuildAccelerator/Common/Private/UbaProcess.cpp
@@ -1490,7 +1490,7 @@ namespace uba
 			if (!DuplicateHandle(currentProcess, currentProcess, (HANDLE)m_nativeProcessHandle, &hostProcess, 0, FALSE, DUPLICATE_SAME_ACCESS))
 			{
 				if (!IsCancelled())
-					logger.Error(TC("Failed to duplicate host process handle for process"));//% ls."), commandLine.c_str());
+					logger.Error(TC("Failed to duplicate host process handle for process (%s)"), LastErrorToText().data);
 				return UBA_EXIT_CODE(8);
 			}
 
@@ -1505,6 +1505,9 @@ namespace uba
 			payload.rulesIndex = m_startInfo.rules->index;
 			payload.version = ProcessMessageVersion;
 			payload.runningRemote = runningRemote;
+			payload.allowKeepFilesInMemory = m_session.m_allowKeepFilesInMemory;
+			payload.allowOutputFiles = m_session.m_allowOutputFiles;
+			payload.suppressLogging = m_session.m_suppressLogging;
 			payload.isChild = m_parentProcess != nullptr;
 			payload.trackInputs = m_startInfo.trackInputs;
 			payload.useCustomAllocator = m_startInfo.useCustomAllocator && m_startInfo.rules->AllowMiMalloc();
diff --git a/Engine/Source/Programs/UnrealBuildAccelerator/Common/Private/UbaSession.cpp b/Engine/Source/Programs/UnrealBuildAccelerator/Common/Private/UbaSession.cpp
--- a/Engine/Source/Programs/UnrealBuildAccelerator/Common/Private/UbaSession.cpp
+++ b/Engine/Source/Programs/UnrealBuildAccelerator/Common/Private/UbaSession.cpp
@@ -210,6 +210,10 @@ namespace uba
 		table.GetValueAsBool(disableCustomAllocator, TC("DisableCustomAllocator"));
 		table.GetValueAsBool(launchVisualizer, TC("LaunchVisualizer"));
 		table.GetValueAsBool(allowMemoryMaps, TC("AllowMemoryMaps"));
+		table.GetValueAsBool(allowKeepFilesInMemory, TC("AllowKeepFilesInMemory"));
+		table.GetValueAsBool(allowOutputFiles, TC("AllowOutputFiles"));
+		table.GetValueAsBool(allowSpecialApplications, TC("AllowSpecialApplications"));
+		table.GetValueAsBool(suppressLogging, TC("SuppressLogging"));
 		table.GetValueAsBool(shouldWriteToDisk, TC("ShouldWriteToDisk"));
 		table.GetValueAsBool(traceEnabled, TC("TraceEnabled"));
 		table.GetValueAsBool(detailedTrace, TC("DetailedTrace"));
@@ -1001,6 +1005,10 @@ namespace uba
 		m_runningRemote = runningRemote;
 		m_disableCustomAllocator = info.disableCustomAllocator;
 		m_allowMemoryMaps = info.allowMemoryMaps;
+		m_allowKeepFilesInMemory = info.allowKeepFilesInMemory;
+		m_allowOutputFiles = info.allowOutputFiles;
+		m_allowSpecialApplications = info.allowSpecialApplications;
+		m_suppressLogging = info.suppressLogging;
 		if (!info.allowMemoryMaps)
 			m_keepOutputFileMemoryMapsThreshold = 0;
 		else
@@ -1943,7 +1951,7 @@ namespace uba
 			auto insres = msg.process.m_writtenFiles.try_emplace(name);
 			WrittenFile& writtenFile = insres.first->second;
 
-			UBA_ASSERT(writtenFile.owner == nullptr || writtenFile.owner == &msg.process);
+			UBA_ASSERTF(writtenFile.owner == nullptr || writtenFile.owner == &msg.process || !m_allowOutputFiles, TC("File %s changed owner.. should not happen"), name);
 			writtenFile.owner = &msg.process;
 			writtenFile.attributes = msg.attributes;
 
diff --git a/Engine/Source/Programs/UnrealBuildAccelerator/Common/Public/UbaSession.h b/Engine/Source/Programs/UnrealBuildAccelerator/Common/Public/UbaSession.h
--- a/Engine/Source/Programs/UnrealBuildAccelerator/Common/Public/UbaSession.h
+++ b/Engine/Source/Programs/UnrealBuildAccelerator/Common/Public/UbaSession.h
@@ -178,6 +178,10 @@ namespace uba
 		bool m_runningRemote;
 		bool m_disableCustomAllocator;
 		bool m_allowMemoryMaps;
+		bool m_allowKeepFilesInMemory;
+		bool m_allowOutputFiles;
+		bool m_allowSpecialApplications;
+		bool m_suppressLogging;
 		bool m_shouldWriteToDisk;
 		bool m_detailedTrace;
 		bool m_traceChildProcesses;
diff --git a/Engine/Source/Programs/UnrealBuildAccelerator/Common/Public/UbaSessionCreateInfo.h b/Engine/Source/Programs/UnrealBuildAccelerator/Common/Public/UbaSessionCreateInfo.h
--- a/Engine/Source/Programs/UnrealBuildAccelerator/Common/Public/UbaSessionCreateInfo.h
+++ b/Engine/Source/Programs/UnrealBuildAccelerator/Common/Public/UbaSessionCreateInfo.h
@@ -26,6 +26,10 @@ namespace uba
 		bool disableCustomAllocator = false;	// Disable detouring of allocator inside processes.
 		bool launchVisualizer = false;			// Launch a UbaVisualizer process (this automatically enable trace)
 		bool allowMemoryMaps = IsWindows;		// Use memory maps where possible. Session creates memory maps of files that processes use
+		bool allowKeepFilesInMemory = IsWindows;// Allow detoured process to keep output files in memory and send them to session through file mapping handle
+		bool allowOutputFiles = IsWindows;		// Allow (selected) output files to be written to disk after process has ended.
+		bool allowSpecialApplications = true;	// Allow uba to expand cmd.exe and call known commands instead of running additional process
+		bool suppressLogging = false;			// Suppress all logging produced by detoured processes. Will be made in detoured process to improve performance when needed
 		bool shouldWriteToDisk = true;			// Set to false to skip writing output files to disk
 		bool traceEnabled = false;				// Set to true to always create in-memory trace data. Is not needed if traceName, traceOutputFile or launchVisualizer is set
 		bool detailedTrace = false;				// Enable detailed trace to include jobs, individual file I/O etc in trace dump
