[UBA]
* Added support for starting processes suspended
* Added support for processes being able to start processes that are not child processes (services like mspdbsrv.exe etc)
* Moved wait-for-parent code to its own function

diff --git a/Engine/Source/Programs/UnrealBuildAccelerator/Common/Private/UbaProcess.cpp b/Engine/Source/Programs/UnrealBuildAccelerator/Common/Private/UbaProcess.cpp
--- a/Engine/Source/Programs/UnrealBuildAccelerator/Common/Private/UbaProcess.cpp
+++ b/Engine/Source/Programs/UnrealBuildAccelerator/Common/Private/UbaProcess.cpp
@@ -465,14 +465,12 @@ namespace uba
 				#endif
 			}
 
-			u64 exitStartTime = GetTime();
+			m_processStats.exitTime = GetTime();
 
 			bool cancelled = IsCancelled();
 			if (exitCode == 0)
 				exitCode = InternalExitProcess(cancelled);
 
-			m_processStats.exitTime = GetTime() - exitStartTime;
-
 			if (exitCode == 0 && !m_messageSuccess)
 				exitCode = UBA_EXIT_CODE(1);
 
@@ -558,6 +556,9 @@ namespace uba
 
 		UBA_ASSERT(m_refCount);
 
+		if (m_processStats.exitTime)
+			m_processStats.exitTime = GetTime() - m_processStats.exitTime;
+
 		// Must be done last to make sure shutdown is not racing
 		m_session.ProcessExited(*this, m_processStats.wallTime);
 	}
@@ -858,6 +859,8 @@ namespace uba
 		reader.ReadString(currentDir);
 		if (currentDir.IsEmpty())
 			currentDir.Append(m_startInfo.workingDir);
+		bool startSuspended = reader.ReadBool();
+		bool isChild = reader.ReadBool();
 
 		const tchar* commandLine = nullptr;
 
@@ -893,17 +896,23 @@ namespace uba
 		ProcessStartInfo info;
 		info.application = application.data;
 		info.arguments = commandLine;
-		info.description = application.GetFileName();
 		info.workingDir = currentDir.data;
 		info.logFile = InternalGetChildLogFile(temp);
 		info.priorityClass = m_startInfo.priorityClass;
 		info.outputStatsThresholdMs = m_startInfo.outputStatsThresholdMs;
 		info.logLineUserData = this;
 		info.logLineFunc = [](void* userData, const tchar* line, u32 length, LogEntryType type) { ((ProcessImpl*)userData)->LogLine(false, TString(line, length), type); };
+		info.startSuspended = startSuspended;
+
+		ProcessImpl* parent = isChild ? this : nullptr;
+		ProcessHandle h = m_session.InternalRunProcess(info, true, parent, true);
 
-		ProcessHandle h = m_session.InternalRunProcess(info, true, this, true);
-		m_childProcesses.push_back(h);
-		u32 childProcessId = u32(m_childProcesses.size());
+		u32 childProcessId = ~0u;
+		if (isChild)
+		{
+			m_childProcesses.push_back(h);
+			childProcessId = u32(m_childProcesses.size());
+		}
 
 		auto& process = *(ProcessImpl*)h.m_process;
 		process.m_echoOn = m_echoOn;
@@ -917,6 +926,8 @@ namespace uba
 		writer.WriteBytes(detoursLib, detoursLibLen);
 
 		writer.WriteString(m_realWorkingDir);
+
+		commandLine = process.m_startInfo.arguments;
 		#if PLATFORM_WINDOWS
 		TString realCommandLine = TC("\"") + process.m_realApplication + TC("\" ") + commandLine;
 		writer.WriteString(realCommandLine);
@@ -938,10 +949,12 @@ namespace uba
 		auto& process = *(ProcessImpl*)m_childProcesses[processId - 1].m_process;
 		bool result = reader.ReadBool();
 		u32 lastError = reader.ReadU32();
+
+		auto setWaitForParent = MakeGuard([&process]() { process.m_waitForParent.Set(); });
+
 		if (!result)
 		{
 			m_session.m_logger.Logf(LogEntryType_Info, TC("Detoured process failed to start child process - %s. %s (Working dir: %s)"), LastErrorToText(lastError).data, process.m_realApplication.c_str(), process.m_realWorkingDir);
-			process.m_waitForParent.Set();
 			return true;
 		}
 
@@ -967,7 +980,17 @@ namespace uba
 		process.m_nativeProcessHandle = (ProcHandle)nativeProcessHandle;
 		process.m_nativeProcessId = nativeProcessId;
 #endif
-		process.m_waitForParent.Set();
+
+		setWaitForParent.Execute();
+		
+		#if PLATFORM_WINDOWS
+		// TOOD: This is ugly, should use event or something instead.. right now we make sure to wait and not return until we know payload has been uploaded etc
+		// It is a very uncommon usecase that processes starts suspended.. some process in ninja/cmake does it when building clang/llvm
+		if (process.m_startInfo.startSuspended)
+			while (process.m_nativeThreadHandle)
+				Sleep(1);
+		#endif
+
 		return true;
 	}
 	bool ProcessImpl::HandleExitChildProcess(BinaryReader& reader, BinaryWriter& writer)
@@ -1467,22 +1490,13 @@ namespace uba
 		}
 		else
 		{
-			u64 startTime = GetTime();
-			while (!m_waitForParent.IsSet(500))
-			{
-				if (IsCancelled())
-					break;
-				if (TimeToMs(GetTime() - startTime) > 120 * 1000) // 
-				{
-					startTime = GetTime();
-					logger.Error(TC("Waiting for parent process in createprocess has now taken more than 120 seconds."));
-				}
-			}
-
+			WaitForParent();
 			if (m_nativeProcessHandle == InvalidProcHandle) // Failed to create the child process
 				return UBA_EXIT_CODE(7);
 		}
 
+		auto closeThreadHandle = MakeGuard([&]() { CloseHandle(m_nativeThreadHandle); m_nativeThreadHandle = 0; });
+
 		if (m_detourEnabled)
 		{
 			HANDLE hostProcess;
@@ -1548,14 +1562,13 @@ namespace uba
 
 		m_processStats.startupTime = GetTime() - m_startTime;
 
-		if (ResumeThread(m_nativeThreadHandle) == -1)
+		if (!m_startInfo.startSuspended && ResumeThread(m_nativeThreadHandle) == -1)
 		{
 			logger.Error(TC("Failed to resume thread for"));//% ls. (% ls)", commandLine.c_str(), LastErrorToText().data);
 			return UBA_EXIT_CODE(11);
 		}
 
-		CloseHandle(m_nativeThreadHandle);
-		m_nativeThreadHandle = 0;
+		closeThreadHandle.Execute();
 
 		if (!m_detourEnabled)
 		{
@@ -1775,17 +1788,7 @@ namespace uba
 		else
 		{
 			//logger.Info("Waiting for parent");
-			u64 startTime = GetTime();
-			while (!m_waitForParent.IsSet(500))
-			{
-				if (IsCancelled())
-					break;
-				if (TimeToMs(GetTime() - startTime) > 120 * 1000) // 
-				{
-					startTime = GetTime();
-					logger.Error(TC("Waiting for parent process (%s) has now taken more than 120 seconds. (%s)"), m_parentProcess->m_startInfo.description, m_startInfo.description);
-				}
-			}
+			WaitForParent();
 
 			//logger.Info("DONE waiting on parent");
 
@@ -1806,19 +1809,7 @@ namespace uba
 			return ~0u;
 
 		if (m_parentProcess)
-		{
-			u64 startTime = GetTime();
-			while (!m_waitForParent.IsSet(500))
-			{
-				if (IsCancelled())
-					break;
-				if (TimeToMs(GetTime() - startTime) > 120 * 1000) // 
-				{
-					startTime = GetTime();
-					logger.Error(TC("Waiting for parent process (%s) while exiting has now taken more than 120 seconds."), m_parentProcess->m_startInfo.description);
-				}
-			}
-		}
+			WaitForParent();
 		m_nativeProcessHandle = InvalidProcHandle;
 
 #if PLATFORM_WINDOWS
@@ -1976,4 +1967,19 @@ namespace uba
 				CloseFileMapping(pair.second.mappingHandle);
 		m_tempFiles.clear();
 	}
+
+	void ProcessImpl::WaitForParent()
+	{
+		u64 startTime = GetTime();
+		while (!m_waitForParent.IsSet(500))
+		{
+			if (IsCancelled())
+				break;
+			if (TimeToMs(GetTime() - startTime) > 120 * 1000) // 
+			{
+				startTime = GetTime();
+				m_session.m_logger.Error(TC("Waiting for parent process in createprocess has now taken more than 120 seconds."));
+			}
+		}
+	}
 }
diff --git a/Engine/Source/Programs/UnrealBuildAccelerator/Common/Public/UbaProcess.h b/Engine/Source/Programs/UnrealBuildAccelerator/Common/Public/UbaProcess.h
--- a/Engine/Source/Programs/UnrealBuildAccelerator/Common/Public/UbaProcess.h
+++ b/Engine/Source/Programs/UnrealBuildAccelerator/Common/Public/UbaProcess.h
@@ -114,6 +114,7 @@ namespace uba
 		u32 InternalCreateProcess(bool runningRemote, void* environment, FileMappingHandle communicationHandle, u64 communicationOffset);
 		u32	InternalExitProcess(bool cancel);
 		void ClearTempFiles();
+		void WaitForParent();
 
 		ProcessStartInfo m_startInfo;
 		Session& m_session;
diff --git a/Engine/Source/Programs/UnrealBuildAccelerator/Common/Public/UbaProcessStartInfo.h b/Engine/Source/Programs/UnrealBuildAccelerator/Common/Public/UbaProcessStartInfo.h
--- a/Engine/Source/Programs/UnrealBuildAccelerator/Common/Public/UbaProcessStartInfo.h
+++ b/Engine/Source/Programs/UnrealBuildAccelerator/Common/Public/UbaProcessStartInfo.h
@@ -27,6 +27,7 @@ namespace uba
 		bool trackInputs = false;				// Track all files read. Can read result in ProcessHandle.GetTrackedInputs()
 		bool useCustomAllocator = true;			// Disable detouring of allocator inside processes. If Session.disableCustomAllocator is false this will be overridden
 		bool writeOutputFilesOnFail = false;	// If set to true, output files will be written to disk regardless if process succeeds or not
+		bool startSuspended = false;			// Start process suspended.. a bit internal atm and only supported on windows
 
 		using LogLineCallback = void(void* userData, const tchar* line, u32 length, LogEntryType type);
 		LogLineCallback* logLineFunc = nullptr;	// Callback for when log entries happens
diff --git a/Engine/Source/Programs/UnrealBuildAccelerator/Detours/Private/PosixOS/UbaDetoursFunctionsPosix.cpp b/Engine/Source/Programs/UnrealBuildAccelerator/Detours/Private/PosixOS/UbaDetoursFunctionsPosix.cpp
--- a/Engine/Source/Programs/UnrealBuildAccelerator/Detours/Private/PosixOS/UbaDetoursFunctionsPosix.cpp
+++ b/Engine/Source/Programs/UnrealBuildAccelerator/Detours/Private/PosixOS/UbaDetoursFunctionsPosix.cpp
@@ -1452,6 +1452,8 @@ UBA_EXPORT int UBA_WRAPPER(posix_spawn)(pid_t* pid, const char* path, const posi
 		writer.WriteString(path); // Application
 		writer.WriteString(cmdLine); // Full command line
 		writer.WriteString(pwd); // Current dir
+		writer.WriteBool(false); // Start suspended
+		writer.WriteBool(true); // Is child
 		writer.Flush();
 
 		BinaryReader reader;
diff --git a/Engine/Source/Programs/UnrealBuildAccelerator/Detours/Private/Windows/UbaDetoursFunctionsKernelBase.inl b/Engine/Source/Programs/UnrealBuildAccelerator/Detours/Private/Windows/UbaDetoursFunctionsKernelBase.inl
--- a/Engine/Source/Programs/UnrealBuildAccelerator/Detours/Private/Windows/UbaDetoursFunctionsKernelBase.inl
+++ b/Engine/Source/Programs/UnrealBuildAccelerator/Detours/Private/Windows/UbaDetoursFunctionsKernelBase.inl
@@ -2681,10 +2681,14 @@ BOOL Detoured_CreateProcessW(LPCWSTR lpApplicationName, LPWSTR lpCommandLine, LP
 		}
 	}
 
+	bool isChild = true;
+
 	StringBuffer<> application;
 	if (lpApplicationName)
 		FixPath(application, lpApplicationName);
 
+	bool startSuspended = (dwCreationFlags & CREATE_SUSPENDED) != 0;
+
 	TString commandLine;
 	TString currentDir;
 	u32 processId = 0;
@@ -2697,6 +2701,8 @@ BOOL Detoured_CreateProcessW(LPCWSTR lpApplicationName, LPWSTR lpCommandLine, LP
 		writer.WriteString(application.data);
 		writer.WriteString(originalCmd);
 		writer.WriteString(lpCurrentDirectory ? lpCurrentDirectory : g_virtualWorkingDir.data);
+		writer.WriteBool(startSuspended);
+		writer.WriteBool(isChild);
 		writer.Flush();
 		BinaryReader reader;
 		processId = reader.ReadU32();
@@ -2727,7 +2733,6 @@ BOOL Detoured_CreateProcessW(LPCWSTR lpApplicationName, LPWSTR lpCommandLine, LP
 	else
 		dwCreationFlags |= CREATE_NO_WINDOW;
 
-	UBA_ASSERT((dwCreationFlags & CREATE_SUSPENDED) == 0);
 	dwCreationFlags |= CREATE_SUSPENDED;
 	BOOL res = true;
 	u32 lastError = ERROR_SUCCESS;
@@ -2756,6 +2761,7 @@ BOOL Detoured_CreateProcessW(LPCWSTR lpApplicationName, LPWSTR lpCommandLine, LP
 	--t_disallowDetour;
 	UBA_ASSERTF(res, L"Failed to spawn process %ls (Error code: %u)", commandLine.c_str(), lastError);
 
+	if (isChild)
 	{
 		TimerScope ts(g_stats.createProcess);
 		SCOPED_WRITE_LOCK(g_communicationLock, pcs);
