[UBA]
* Changed so ProcessStartupInfoHolder inherits from ProcessStartupInfo

diff --git a/Engine/Source/Programs/UnrealBuildAccelerator/Common/Private/UbaProcess.cpp b/Engine/Source/Programs/UnrealBuildAccelerator/Common/Private/UbaProcess.cpp
--- a/Engine/Source/Programs/UnrealBuildAccelerator/Common/Private/UbaProcess.cpp
+++ b/Engine/Source/Programs/UnrealBuildAccelerator/Common/Private/UbaProcess.cpp
@@ -142,7 +142,7 @@ namespace uba
 		}
 	}
 
-	void ProcessImpl::Start(const ProcessStartInfo& startInfo, TString&& realApplication, const tchar* realWorkingDir, bool runningRemote, void* environment, bool async, bool enableDetour)
+	bool ProcessImpl::Start(const ProcessStartInfo& startInfo, bool runningRemote, void* environment, bool async, bool enableDetour)
 	{
 		m_detourEnabled = enableDetour;
 
@@ -150,32 +150,21 @@ namespace uba
 
 		m_startInfo = startInfo;
 
-		m_description = startInfo.description;
-		m_startInfo.description = m_description.c_str();
+		FixPathSeparators(m_startInfo.applicationStr.data());
+		FixPathSeparators(m_startInfo.workingDirStr.data());
+		FixPathSeparators(m_startInfo.logFileStr.data());
 
-		m_virtualApplication = startInfo.application;
-		FixPathSeparators(m_virtualApplication.data());
-		m_startInfo.application = m_virtualApplication.c_str();
+		m_startInfo.Expand();
 
-		m_arguments = startInfo.arguments;
+		m_extractExports = Contains(m_startInfo.arguments, TC("/extractexports"), true);
 
-		m_extractExports = Contains(m_arguments.c_str(), TC("/extractexports"), true);
+		StringBuffer<> realApplication(m_startInfo.application);
+		const tchar* realWorkingDir = m_startInfo.workingDir;
 
-		m_startInfo.arguments = m_arguments.c_str();
-
-		m_virtualWorkingDir = startInfo.workingDir;
-		FixPathSeparators(m_virtualWorkingDir.data());
-		m_startInfo.workingDir = m_virtualWorkingDir.c_str();
-
-		m_logFile = startInfo.logFile;
-		FixPathSeparators(m_logFile.data());
-		m_startInfo.logFile = m_logFile.c_str();
-
-		size_t nameIndex = m_virtualApplication.find_last_of(PathSeparator);
-		if (nameIndex != -1)
-			m_virtualApplicationDir = m_virtualApplication.substr(0, nameIndex);
+		if (!m_session.PrepareProcess(m_startInfo, m_parentProcess != nullptr, realApplication, realWorkingDir))
+			return false;
 
-		m_realApplication = std::move(realApplication);
+		m_realApplication = realApplication.data;
 		m_realWorkingDir = realWorkingDir;
 		if (realWorkingDir == startInfo.workingDir)
 			m_realWorkingDir = m_startInfo.workingDir;
@@ -199,6 +188,7 @@ namespace uba
 			ThreadRun(runningRemote, environment);
 			ThreadExit();
 		}
+		return true;
 	}
 
 	bool ProcessImpl::IsActive()
@@ -566,24 +556,16 @@ namespace uba
 	bool ProcessImpl::HandleSpecialApplication()
 	{
 	#if PLATFORM_WINDOWS
-		// TODO: This is super hacky... but we don't want to spawn the cmd.exe just to copy a file since the overhead can be half a second
-		// "C:\WINDOWS\system32\cmd.exe" /C "copy /Y "E:\dev\fn\Engine\Source\Runtime\RenderCore\RenderCore.natvis" "E:\dev\fn\Engine\Intermediate\Build\Win64\x64\UnrealPak\Development\RenderCore\RenderCore.natvis" 1>nul"
-		if (!Contains(m_startInfo.application, TC("cmd.exe")))
-			return false;
-		//m_session.m_logger.Info(TC("GOT HERE: %s"), m_startInfo.arguments);
-		if (!StartsWith(m_startInfo.arguments, TC("/C \"copy /Y \"")))
+		if (!Equals(m_startInfo.application, TC("ubacopy")))
 			return false;
-		const tchar* fromFileBegin = m_startInfo.arguments + 13;
+		const tchar* fromFileBegin = m_startInfo.arguments;
 		const tchar* fromFileEnd = TStrchr(fromFileBegin, '\"');
-		if (!fromFileEnd)
-			return false;
+		UBA_ASSERT(fromFileEnd);
 		const tchar* toFileBegin = TStrchr(fromFileEnd + 1, '\"');
-		if (!toFileBegin)
-			return false;
+		UBA_ASSERT(toFileBegin);
 		++toFileBegin;
 		const tchar* toFileEnd = TStrchr(toFileBegin, '\"');
-		if (!toFileEnd)
-			return false;
+		UBA_ASSERT(toFileEnd);
 
 		m_processStats.wallTime = GetTime() - m_startTime;
 
@@ -877,7 +859,15 @@ namespace uba
 		}
 		else
 		{
-			const tchar* secondParamStart = fullCommandLine.First(' ', 1);
+			u32 searchOffset = 1;
+			if (applicationStr.empty() && IsWindows)
+			{
+				const tchar* extPos = nullptr;
+				fullCommandLine.Contains(TC(".exe"), true, &extPos);
+				UBA_ASSERT(extPos);
+				searchOffset = u32(extPos - fullCommandLine.data);
+			}
+			const tchar* secondParamStart = fullCommandLine.First(' ', searchOffset);
 			if (!secondParamStart)
 				commandLine = TC("");
 			else
@@ -886,15 +876,14 @@ namespace uba
 				applicationStr.assign(fullCommandLine.data, u64(secondParamStart - fullCommandLine.data));
 		}
 
-		StringBuffer<512> application;
-		FixPath(applicationStr.c_str(), nullptr, 0, application);
+		FixPathSeparators(applicationStr.data());
 
 		while (*commandLine == ' ')
 			++commandLine;
 
 		StringBuffer<> temp;
 		ProcessStartInfo info;
-		info.application = application.data;
+		info.application = applicationStr.c_str();
 		info.arguments = commandLine;
 		info.workingDir = currentDir.data;
 		info.logFile = InternalGetChildLogFile(temp);
@@ -1010,13 +999,13 @@ namespace uba
 
 	bool ProcessImpl::HandleCreateTempFile(BinaryReader& reader, BinaryWriter& writer)
 	{
-		CreateTempFile(reader, m_nativeProcessHandle, m_virtualApplication.c_str());
+		CreateTempFile(reader, m_nativeProcessHandle, m_startInfo.application);
 		return true;
 	}
 
 	bool ProcessImpl::HandleOpenTempFile(BinaryReader& reader, BinaryWriter& writer)
 	{
-		OpenTempFile(reader, writer, m_virtualApplication.c_str());
+		OpenTempFile(reader, writer, m_startInfo.application);
 		return true;
 	}
 
@@ -1150,13 +1139,13 @@ namespace uba
 		if (!newProcess)
 			return true;
 
-		m_arguments = nextProcess.arguments;
-		m_description = nextProcess.description;
-		m_logFile = nextProcess.logFile;
+		m_startInfo.argumentsStr = nextProcess.arguments;
+		m_startInfo.descriptionStr = nextProcess.description;
+		m_startInfo.logFileStr = nextProcess.logFile;
 
-		m_startInfo.arguments = m_arguments.c_str();
-		m_startInfo.description = m_description.c_str();
-		m_startInfo.logFile = m_logFile.c_str();
+		m_startInfo.arguments = m_startInfo.argumentsStr.c_str();
+		m_startInfo.description = m_startInfo.descriptionStr.c_str();
+		m_startInfo.logFile = m_startInfo.logFileStr.c_str();
 
 		m_childProcesses.clear();
 		m_logLines.clear();
@@ -1605,7 +1594,7 @@ namespace uba
 			}
 			Vector<const char*> arguments2;
 			arguments2.reserve(arguments.size() + 2);
-			arguments2.push_back(m_virtualApplication.data());
+			arguments2.push_back(m_startInfo.application);
 			for (auto& s : arguments)
 				arguments2.push_back(s.data());
 			arguments2.push_back(nullptr);
diff --git a/Engine/Source/Programs/UnrealBuildAccelerator/Common/Private/UbaProcessStartInfoHolder.cpp b/Engine/Source/Programs/UnrealBuildAccelerator/Common/Private/UbaProcessStartInfoHolder.cpp
new file mode 100644
--- /dev/null
+++ b/Engine/Source/Programs/UnrealBuildAccelerator/Common/Private/UbaProcessStartInfoHolder.cpp
@@ -0,0 +1,127 @@
+// Copyright Epic Games, Inc. All Rights Reserved.
+
+#include "UbaProcessStartInfoHolder.h"
+
+namespace uba
+{
+	bool ProcessStartInfoHolder::Expand()
+	{
+	#if PLATFORM_WINDOWS
+		// Special handling to avoid calling cmd.exe if not needed
+		if (!Contains(application, TC("cmd.exe")))
+			return false;
+
+		const tchar* argsBegin = argumentsStr.c_str();
+		const tchar* argsEnd = argsBegin + argumentsStr.size();
+
+		// Check if application is repeated as first argument, in that case consume it
+		const tchar* firstArgBegin = argsBegin;
+		const tchar* firstArgEnd = nullptr;
+		if (*firstArgBegin == '\"')
+		{
+			++firstArgBegin;
+			firstArgEnd = TStrchr(firstArgBegin, '\"');
+		}
+		else
+			firstArgEnd = TStrchr(firstArgBegin, ' ');
+		if (!firstArgEnd)
+			return false;
+		StringBuffer<32*1024> commands;
+		commands.Append(firstArgBegin, firstArgEnd - firstArgBegin);
+		if (commands.Contains(application))
+			argsBegin = firstArgEnd;
+		while (argsBegin && *argsBegin == ' ')
+			++argsBegin;
+		commands.Clear();
+
+
+		// Parse switches... only supported is /C right now
+		if (!StartsWith(argsBegin, TC("/C ")))
+			return false;
+		argsBegin += 3;
+		while (argsBegin && *argsBegin == ' ')
+			++argsBegin;
+
+		if (*argsBegin == '/') // Unknown switch, don't try to expand cmd
+			return false;
+
+		if (*argsBegin == '\"')
+		{
+			++argsBegin;
+			--argsEnd;
+		}
+		while (argsBegin && *argsBegin == ' ')
+			++argsBegin;
+
+		commands.Append(argsBegin, argsEnd - argsBegin);
+
+		// It could be that there is just a chain of commands to set working dir, in that case we strip out cmd.exe
+		const tchar* andPos = nullptr;
+		if (commands.Contains(TC(" && "), true, &andPos))
+		{
+			if (Contains(andPos + 4, TC(" && "))) // If more than one && we don't try to expand cmd.exe
+				return false;
+			if (!commands.StartsWith(TC("cd /D"))) // First command is not cd, don't try to expand cmd.exe
+				return false;
+			const tchar* workDirStart = commands.data + 6;
+			workingDirStr.assign(workDirStart, andPos - workDirStart);
+
+			StringBuffer<> fixed;
+			FixPath(workingDirStr.data(), nullptr, 0, fixed);
+			workingDirStr = fixed.data;
+			workingDir = workingDirStr.c_str();
+
+			const tchar* commandLine = andPos + 4;
+			while (*commandLine && *commandLine == ' ')
+				++commandLine;
+			const tchar* applicationBegin = commandLine;
+			const tchar* applicationEnd = nullptr;
+			if (*applicationBegin == '\"')
+			{
+				++applicationBegin;
+				applicationEnd = TStrchr(applicationBegin, '\"');
+			}
+			else
+				applicationEnd = TStrchr(applicationBegin, ' ');
+			if (!applicationEnd)
+				applicationEnd = applicationBegin + TStrlen(applicationBegin);
+			applicationStr.assign(applicationBegin, applicationEnd - applicationBegin);
+			FixPath(applicationStr.data(), nullptr, 0, fixed.Clear());
+			applicationStr = fixed.data;
+			application = applicationStr.c_str();
+			argsBegin = *applicationEnd == '\"' ? applicationEnd + 1 : applicationEnd;
+			while (*argsBegin && *argsBegin == ' ')
+				++argsBegin;
+			argumentsStr = argsBegin;
+			arguments = argumentsStr.c_str();
+			return true;
+		}
+		else
+		{
+			// TODO: This is super hacky... but we don't want to spawn the cmd.exe just to copy a file since the overhead can be half a second
+			// "C:\WINDOWS\system32\cmd.exe" /C "copy /Y "E:\dev\fn\Engine\Source\Runtime\RenderCore\RenderCore.natvis" "E:\dev\fn\Engine\Intermediate\Build\Win64\x64\UnrealPak\Development\RenderCore\RenderCore.natvis" 1>nul"
+			if (!commands.StartsWith(TC("copy /Y \"")))
+				return false;
+
+			const tchar* fromFileBegin = commands.data + 9;
+			const tchar* fromFileEnd = TStrchr(fromFileBegin, '\"');
+			if (!fromFileEnd)
+				return false;
+			const tchar* toFileBegin = TStrchr(fromFileEnd + 1, '\"');
+			if (!toFileBegin)
+				return false;
+			++toFileBegin;
+			const tchar* toFileEnd = TStrchr(toFileBegin, '\"');
+			if (!toFileEnd)
+				return false;
+
+			applicationStr = TC("ubacopy");
+			application = applicationStr.c_str();
+			argumentsStr = fromFileBegin;
+			arguments = argumentsStr.c_str();
+			return true;
+		}
+	#endif
+		return false;
+	}
+}
diff --git a/Engine/Source/Programs/UnrealBuildAccelerator/Common/Private/UbaProcessStartInfoHolder.h b/Engine/Source/Programs/UnrealBuildAccelerator/Common/Private/UbaProcessStartInfoHolder.h
--- a/Engine/Source/Programs/UnrealBuildAccelerator/Common/Private/UbaProcessStartInfoHolder.h
+++ b/Engine/Source/Programs/UnrealBuildAccelerator/Common/Private/UbaProcessStartInfoHolder.h
@@ -22,71 +22,69 @@ namespace uba
 		out.Append(buffer);
 	}
 
-	struct ProcessStartInfoHolder
+	struct ProcessStartInfoHolder : public ProcessStartInfo
 	{
-		ProcessStartInfoHolder() {}
+		ProcessStartInfoHolder() = default;
 
 		ProcessStartInfoHolder(const ProcessStartInfo& si)
 		{
-			startInfo = si;
+			*(ProcessStartInfo*)this = si;
 
 			StringBuffer<512> temp;
 			FixFileName(temp, si.workingDir, nullptr);
 			temp.EnsureEndsWithSlash();
-			workingDir = temp.data;
-			startInfo.workingDir = workingDir.c_str();
+			workingDirStr = temp.data;
+			workingDir = workingDirStr.c_str();
 
-			temp.EnsureEndsWithSlash();
-			StringBuffer<512> temp2;
-			FixFileName(temp2, si.application, temp.data);
-			application = temp2.data;
-			startInfo.application = application.c_str();
+			applicationStr = si.application;
+			application = applicationStr.c_str();
 
-			arguments = si.arguments;
-			startInfo.arguments = arguments.c_str();
+			argumentsStr = si.arguments;
+			arguments = argumentsStr.c_str();
 
-			description = si.description;
-			startInfo.description = description.c_str();
+			if (si.description)
+				descriptionStr = si.description;
+			description = descriptionStr.c_str();
 
-			logFile = si.logFile;
-			startInfo.logFile = logFile.c_str();
+			logFileStr = si.logFile;
+			logFile = logFileStr.c_str();
 		}
 
 		void Write(BinaryWriter& writer)
 		{
-			writer.WriteString(description);
-			writer.WriteString(application);
-			writer.WriteString(arguments);
-			writer.WriteString(workingDir);
-			writer.WriteString(logFile);
+			writer.WriteString(descriptionStr);
+			writer.WriteString(applicationStr);
+			writer.WriteString(argumentsStr);
+			writer.WriteString(workingDirStr);
+			writer.WriteString(logFileStr);
 			writer.WriteU32(*(u32*)&weight);
-			writer.WriteBool(startInfo.trackInputs);
-			writer.WriteBool(startInfo.writeOutputFilesOnFail);
-			writer.WriteU64(startInfo.outputStatsThresholdMs);
+			writer.WriteBool(trackInputs);
+			writer.WriteBool(writeOutputFilesOnFail);
+			writer.WriteU64(outputStatsThresholdMs);
 		}
 
 		void Read(BinaryReader& reader)
 		{
-			description = reader.ReadString();
-			application = reader.ReadString();
-			arguments = reader.ReadString();
-			workingDir = reader.ReadString();
-			logFile = reader.ReadString();
+			descriptionStr = reader.ReadString();
+			applicationStr = reader.ReadString();
+			argumentsStr = reader.ReadString();
+			workingDirStr = reader.ReadString();
+			logFileStr = reader.ReadString();
 
-			Replace(application.data(), '/', PathSeparator); // TODO: Is this needed?
+			Replace(applicationStr.data(), '/', PathSeparator); // TODO: Is this needed?
 
 			u32 weight32 = reader.ReadU32();
 			weight = *(float*)&weight32;
 			
-			startInfo.trackInputs = reader.ReadBool();
-			startInfo.writeOutputFilesOnFail = reader.ReadBool();
-			startInfo.outputStatsThresholdMs = reader.ReadU64();
-
-			startInfo.description = description.c_str();
-			startInfo.application = application.c_str();
-			startInfo.arguments = arguments.c_str();
-			startInfo.workingDir = workingDir.c_str();
-			startInfo.logFile = logFile.c_str();
+			trackInputs = reader.ReadBool();
+			writeOutputFilesOnFail = reader.ReadBool();
+			outputStatsThresholdMs = reader.ReadU64();
+
+			description = descriptionStr.c_str();
+			application = applicationStr.c_str();
+			arguments = argumentsStr.c_str();
+			workingDir = workingDirStr.c_str();
+			logFile = logFileStr.c_str();
 		}
 
 		ProcessStartInfoHolder(const ProcessStartInfoHolder& o)
@@ -96,30 +94,29 @@ namespace uba
 
 		void operator=(const ProcessStartInfoHolder& o)
 		{
-			startInfo = o.startInfo;
-
-			workingDir = o.workingDir;
-			startInfo.workingDir = workingDir.c_str();
-			application = o.application;
-			startInfo.application = application.c_str();
-			arguments = o.arguments;
-			startInfo.arguments = arguments.c_str();
-			description = o.description;
-			startInfo.description = description.c_str();
-			logFile = o.logFile;
-			startInfo.logFile = logFile.c_str();
+			*(ProcessStartInfo*)this = o;
+
+			workingDirStr = o.workingDirStr;
+			workingDir = workingDirStr.c_str();
+			applicationStr = o.applicationStr;
+			application = applicationStr.c_str();
+			argumentsStr = o.argumentsStr;
+			arguments = argumentsStr.c_str();
+			descriptionStr = o.descriptionStr;
+			description = descriptionStr.c_str();
+			logFileStr = o.logFileStr;
+			logFile = logFileStr.c_str();
 
 			weight = o.weight;
 		}
 
-		ProcessStartInfo startInfo;
+		bool Expand();
 
-		TString description;
-		TString application;
-		TString arguments;
-		TString workingDir;
-		TString logFile;
+		TString descriptionStr;
+		TString applicationStr;
+		TString argumentsStr;
+		TString workingDirStr;
+		TString logFileStr;
 		float weight = 1.0f;
 	};
-
 }
\ No newline at end of file
diff --git a/Engine/Source/Programs/UnrealBuildAccelerator/Common/Private/UbaScheduler.cpp b/Engine/Source/Programs/UnrealBuildAccelerator/Common/Private/UbaScheduler.cpp
--- a/Engine/Source/Programs/UnrealBuildAccelerator/Common/Private/UbaScheduler.cpp
+++ b/Engine/Source/Programs/UnrealBuildAccelerator/Common/Private/UbaScheduler.cpp
@@ -47,33 +47,33 @@ namespace uba
 	class SkippedProcess : public Process
 	{
 	public:
-		SkippedProcess(const ProcessStartInfo& i) : holder(i) {}
+		SkippedProcess(const ProcessStartInfo& i) : startInfo(i) {}
 		virtual u32 GetExitCode() override { return ProcessCancelExitCode; }
 		virtual bool HasExited() override { return true; }
 		virtual bool WaitForExit(u32 millisecondsTimeout) override{ return true; }
-		virtual const ProcessStartInfo& GetStartInfo() const override { return holder.startInfo; }
+		virtual const ProcessStartInfo& GetStartInfo() const override { return startInfo; }
 		virtual const Vector<ProcessLogLine>& GetLogLines() const override { static Vector<ProcessLogLine> v{ProcessLogLine{TC("Skipped"), LogEntryType_Warning}}; return v; }
 		virtual const Vector<u8>& GetTrackedInputs() const override { static Vector<u8> v; return v;}
 		virtual const Vector<u8>& GetTrackedOutputs() const override { static Vector<u8> v; return v;}
 		virtual bool IsRemote() const override { return false; }
 		virtual ProcessExecutionType GetExecutionType() const override { return ProcessExecutionType_Native; }
-		ProcessStartInfoHolder holder;
+		ProcessStartInfoHolder startInfo;
 	};
 
 	class CachedProcess : public Process
 	{
 	public:
-		CachedProcess(const ProcessStartInfo& i) : holder(i) {}
+		CachedProcess(const ProcessStartInfo& i) : startInfo(i) {}
 		virtual u32 GetExitCode() override { return 0; }
 		virtual bool HasExited() override { return true; }
 		virtual bool WaitForExit(u32 millisecondsTimeout) override{ return true; }
-		virtual const ProcessStartInfo& GetStartInfo() const override { return holder.startInfo; }
+		virtual const ProcessStartInfo& GetStartInfo() const override { return startInfo; }
 		virtual const Vector<ProcessLogLine>& GetLogLines() const override { return logLines; }
 		virtual const Vector<u8>& GetTrackedInputs() const override { static Vector<u8> v; return v;}
 		virtual const Vector<u8>& GetTrackedOutputs() const override { static Vector<u8> v; return v;}
 		virtual bool IsRemote() const override { return false; }
 		virtual ProcessExecutionType GetExecutionType() const override { return ProcessExecutionType_FromCache; }
-		ProcessStartInfoHolder holder;
+		ProcessStartInfoHolder startInfo;
 		Vector<ProcessLogLine> logLines;
 	};
 
@@ -162,10 +162,11 @@ namespace uba
 		}
 
 		auto info2 = new ProcessStartInfo2(info.info, ki, info.knownInputsCount);
+		info2->Expand();
 		info2->weight = info.weight;
 
-		const ApplicationRules* rules = m_session.GetRules(info2->startInfo);
-		info2->startInfo.rules = rules;
+		const ApplicationRules* rules = m_session.GetRules(*info2);
+		info2->rules = rules;
 
 		bool useCache = m_cacheClient && !m_writeToCache && rules->IsCacheable();
 
@@ -181,7 +182,7 @@ namespace uba
 
 		if (m_processConfigs)
 		{
-			auto name = info2->application.c_str();
+			auto name = info2->application;
 			if (auto lastSeparator = TStrrchr(name, PathSeparator))
 				name = lastSeparator + 1;
 			StringBuffer<128> lower(name);
@@ -415,7 +416,7 @@ namespace uba
 			exitInfo->isLocal = isLocal;
 			exitInfo->processIndex = indexToRun;
 
-			ProcessStartInfo si = info->startInfo;
+			ProcessStartInfo si = *info;
 			si.userData = exitInfo;
 			si.exitedFunc = [](void* userData, const ProcessHandle& handle)
 				{
@@ -431,7 +432,7 @@ namespace uba
 				// TODO: This should not use work manager since it is mostly waiting on network
 				m_session.GetServer().AddWork([this, exitInfo]()
 					{
-						ProcessStartInfo& si = exitInfo->startInfo->startInfo;
+						ProcessStartInfo& si = *exitInfo->startInfo;
 						u64 startTime = GetTime();
 
 						bool cacheHit = false;
@@ -504,7 +505,7 @@ namespace uba
 			ei->startInfo = newInfo;
 			ei->processIndex = indexToRun;
 
-			auto& si = newInfo->startInfo;
+			auto& si = *newInfo;
 			UBA_ASSERT(Equals(currentStartInfo.application, si.application));
 			outNextProcess.arguments = si.arguments;
 			outNextProcess.workingDir = si.workingDir;
@@ -520,8 +521,8 @@ namespace uba
 		ph.m_process = &process;
 
 		auto si = info.startInfo;
-		if (auto func = si->startInfo.exitedFunc)
-			func(si->startInfo.userData, ph);
+		if (auto func = si->exitedFunc)
+			func(si->userData, ph);
 
 		SCOPED_WRITE_LOCK(m_processEntriesLock, lock);
 		auto& entry = m_processEntries[info.processIndex];
@@ -555,9 +556,9 @@ namespace uba
 
 	void Scheduler::SkipProcess(ProcessStartInfo2& info)
 	{
-		ProcessHandle ph(new SkippedProcess(info.startInfo));
-		if (auto func = info.startInfo.exitedFunc)
-			func(info.startInfo.userData, ph);
+		ProcessHandle ph(new SkippedProcess(info));
+		if (auto func = info.exitedFunc)
+			func(info.userData, ph);
 		FinishProcess(ph);
 	}
 
diff --git a/Engine/Source/Programs/UnrealBuildAccelerator/Common/Private/UbaSession.cpp b/Engine/Source/Programs/UnrealBuildAccelerator/Common/Private/UbaSession.cpp
--- a/Engine/Source/Programs/UnrealBuildAccelerator/Common/Private/UbaSession.cpp
+++ b/Engine/Source/Programs/UnrealBuildAccelerator/Common/Private/UbaSession.cpp
@@ -1231,18 +1231,11 @@ namespace uba
 	void Session::ValidateStartInfo(const ProcessStartInfo& startInfo)
 	{
 		UBA_ASSERTF(startInfo.workingDir && *startInfo.workingDir, TC("Working dir must be set when spawning process"));
-		UBA_ASSERTF(!TStrchr(startInfo.application, '~'), TC("Application path must use long name (%s)"), startInfo.application);
 		UBA_ASSERTF(!TStrchr(startInfo.workingDir, '~'), TC("WorkingDir path must use long name (%s)"), startInfo.workingDir);
 	}
 
 	ProcessHandle Session::InternalRunProcess(const ProcessStartInfo& startInfo, bool async, ProcessImpl* parent, bool enableDetour)
 	{
-		StringBuffer<> realApplication(startInfo.application);
-		const tchar* realWorkingDir = startInfo.workingDir;
-
-		if (!PrepareProcess(startInfo, parent != nullptr, realApplication, realWorkingDir))
-			return {};
-
 		auto& si = const_cast<ProcessStartInfo&>(startInfo);
 		si.useCustomAllocator &= !m_disableCustomAllocator;
 		const tchar* originalLogFile = si.logFile;
@@ -1271,7 +1264,8 @@ namespace uba
 		u32 id = CreateProcessId();
 		auto process = new ProcessImpl(*this, id, parent);
 		ProcessHandle h(process);
-		process->Start(startInfo, realApplication.data, realWorkingDir, m_runningRemote, env, async, enableDetour);
+		if (!process->Start(startInfo, m_runningRemote, env, async, enableDetour))
+			return {};
 
 		si.logFile = originalLogFile;
 		return h;
@@ -1669,8 +1663,14 @@ namespace uba
 		return fileName.EndsWith(TC(".h"));
 	}
 
-	bool Session::PrepareProcess(const ProcessStartInfo& startInfo, bool isChild, StringBufferBase& outRealApplication, const tchar*& outRealWorkingDir)
+	bool Session::PrepareProcess(ProcessStartInfoHolder& startInfo, bool isChild, StringBufferBase& outRealApplication, const tchar*& outRealWorkingDir)
 	{
+		if (StartsWith(startInfo.application, TC("ubacopy")))
+			return true;
+		if (!SearchPathForFile(m_logger, outRealApplication.Clear(), startInfo.application, startInfo.workingDir))
+			return false;
+		startInfo.applicationStr = outRealApplication.data;
+		startInfo.application = startInfo.applicationStr.c_str();
 		return true;
 	}
 
diff --git a/Engine/Source/Programs/UnrealBuildAccelerator/Common/Private/UbaSessionClient.cpp b/Engine/Source/Programs/UnrealBuildAccelerator/Common/Private/UbaSessionClient.cpp
--- a/Engine/Source/Programs/UnrealBuildAccelerator/Common/Private/UbaSessionClient.cpp
+++ b/Engine/Source/Programs/UnrealBuildAccelerator/Common/Private/UbaSessionClient.cpp
@@ -15,6 +15,15 @@
 
 namespace uba
 {
+	struct SessionClient::ModuleInfo
+	{
+		ModuleInfo(const tchar* n, const CasKey& c, u32 a) : name(n), casKey(c), attributes(a), done(true) {}
+		TString name;
+		CasKey casKey;
+		u32 attributes;
+		Event done;
+	};
+
 	SessionClient::SessionClient(const SessionClientCreateInfo& info)
 	: Session(info, TC("UbaSessionClient"), true, &info.client)
 	,	m_client(info.client)
@@ -234,21 +243,94 @@ namespace uba
 		return WriteBinFile(out, destFile.data, newKey, keyStr, fileAttributes);
 	}
 
-	bool SessionClient::PrepareProcess(const ProcessStartInfo& startInfo, bool isChild, StringBufferBase& outRealApplication, const tchar*& outRealWorkingDir)
+	bool SessionClient::PrepareProcess(ProcessStartInfoHolder& startInfo, bool isChild, StringBufferBase& outRealApplication, const tchar*& outRealWorkingDir)
 	{
-		outRealApplication.Clear();
 		outRealWorkingDir = m_processWorkingDir.data;
-		return EnsureApplicationEnvironment(outRealApplication, 0, startInfo.application);
-	}
+		if (StartsWith(startInfo.application, TC("ubacopy")))
+			return true;
+		outRealApplication.Clear();
 
-	struct SessionClient::ModuleInfo
-	{
-		ModuleInfo(const tchar* n, const CasKey& c, u32 a) : name(n), casKey(c), attributes(a), done(true) {}
-		TString name;
-		CasKey casKey;
-		u32 attributes;
-		Event done;
-	};
+		const tchar* application = startInfo.application;
+		UBA_ASSERT(application && *application);
+		bool isAbsolute = IsAbsolutePath(startInfo.application);
+
+		SCOPED_WRITE_LOCK(m_handledApplicationEnvironmentsLock, environmentslock);
+		auto insres = m_handledApplicationEnvironments.try_emplace(application);
+		environmentslock.Leave();
+
+		ApplicationEnvironment& appEnv = insres.first->second;
+		SCOPED_WRITE_LOCK(appEnv.lock, lock);
+
+		if (!appEnv.realApplication.empty())
+		{
+			outRealApplication.Append(appEnv.realApplication);
+			if (!isAbsolute)
+			{
+				startInfo.applicationStr = appEnv.virtualApplication;
+				startInfo.application = startInfo.applicationStr.c_str();
+			}
+			return true;
+
+		}
+
+		List<ModuleInfo> modules;
+		if (!ReadModules(modules, 0, application))
+			return false;
+
+		StringBuffer<MaxPath> applicationDir;
+		applicationDir.AppendDir(application);
+		KeyToString keyStr(ToStringKeyLower(applicationDir));
+
+		Atomic<bool> success = true;
+		Atomic<u32> handledCount;
+
+		for (auto& m : modules)
+		{
+			m_client.AddWork([&handledCount, &m, this, &success, &keyStr]()
+				{
+					++handledCount;
+					auto g = MakeGuard([&]() { m.done.Set(); });
+					CasKey newCasKey;
+					bool storeUncompressed = true;
+					u64 fileSize;
+					const tchar* moduleName = m.name.c_str();
+					if (!RetrieveCasFile(newCasKey, fileSize, m.casKey, moduleName, storeUncompressed))
+					{
+						m_logger.Error(TC("Casfile not found for %s (%s)"), moduleName, CasKeyString(m.casKey).str);
+						success = false;
+						return;
+					}
+					if (const tchar* lastSeparator = TStrrchr(moduleName, PathSeparator))
+						moduleName = lastSeparator + 1;
+					StringBuffer<MaxPath> temp;
+					if (!WriteBinFile(temp, moduleName, newCasKey, keyStr, m.attributes))
+						success = false;
+				}, 1, TC("EnsureApp"));
+		}
+
+		while (handledCount < modules.size())
+			m_client.DoWork();
+
+		// Wait for all to be done
+		for (auto& m : modules)
+			if (!m.done.IsSet(10 * 60 * 1000)) // 10 minutes is a very long time
+				return m_logger.Error(TC("Timed out while waiting for application cas files to be downloaded"));
+
+		if (!success)
+			return false;
+		
+		outRealApplication.Append(m_sessionBinDir).Append(keyStr).Append(PathSeparator).AppendFileName(application);
+		appEnv.realApplication = outRealApplication.data;
+
+		if (!isAbsolute)
+		{
+			appEnv.virtualApplication = modules.front().name;
+			startInfo.applicationStr = appEnv.virtualApplication;
+			startInfo.application = startInfo.applicationStr.c_str();
+		}
+
+		return true;
+	}
 
 	bool SessionClient::ReadModules(List<ModuleInfo>& outModules, u32 processId, const tchar* application)
 	{
@@ -282,7 +364,7 @@ namespace uba
 			{
 				StringBuffer<> localSystemModule;
 				localSystemModule.Append(m_systemPath).Append(moduleFile.data + serverSystemPathLen);
-				if (FileExists(m_logger, localSystemModule.data))
+				if (FileExists(m_logger, localSystemModule.data) && !localSystemModule.EndsWith(TC(".exe")))
 					continue;
 				moduleFile.Clear().Append(localSystemModule);
 			}
@@ -292,67 +374,6 @@ namespace uba
 		return true;
 	}
 
-	bool SessionClient::EnsureApplicationEnvironment(StringBufferBase& out, u32 processId, const tchar* application)
-	{
-		StringBuffer<MaxPath> applicationDir;
-		applicationDir.AppendDir(application);
-		KeyToString keyStr(ToStringKeyLower(applicationDir));
-
-		UBA_ASSERT(application && *application);
-		SCOPED_WRITE_LOCK(m_handledApplicationEnvironmentsLock, lock);
-		auto insres = m_handledApplicationEnvironments.insert(application);
-		if (insres.second)
-		{
-			auto failGuard = MakeGuard([&]() { m_handledApplicationEnvironments.erase(insres.first); });
-
-			List<ModuleInfo> modules;
-
-			if (!ReadModules(modules, processId, application))
-				return false;
-
-			Atomic<bool> success = true;
-			Atomic<u32> handledCount;
-			for (auto& m : modules)
-			{
-				m_client.AddWork([&handledCount, &m, this, &success, &keyStr]()
-					{
-						++handledCount;
-						auto g = MakeGuard([&]() { m.done.Set(); });
-						CasKey newCasKey;
-						bool storeUncompressed = true;
-						u64 fileSize;
-						const tchar* moduleName = m.name.c_str();
-						if (!RetrieveCasFile(newCasKey, fileSize, m.casKey, moduleName, storeUncompressed))
-						{
-							m_logger.Error(TC("Casfile not found for %s (%s)"), moduleName, CasKeyString(m.casKey).str);
-							success = false;
-							return;
-						}
-						if (const tchar* lastSeparator = TStrrchr(moduleName, PathSeparator))
-							moduleName = lastSeparator + 1;
-						StringBuffer<MaxPath> temp;
-						if (!WriteBinFile(temp, moduleName, newCasKey, keyStr, m.attributes))
-							success = false;
-					}, 1, TC("EnsureApp"));
-			}
-
-			while (handledCount < modules.size())
-				m_client.DoWork();
-
-			// Wait for all to be done
-			for (auto& m : modules)
-				if (!m.done.IsSet(10 * 60 * 1000)) // 10 minutes is a very long time
-					return m_logger.Error(TC("Timed out while waiting for application cas files to be downloaded"));
-
-			if (!success)
-				return false;
-
-			failGuard.Cancel();
-		}
-		out.Append(m_sessionBinDir).Append(keyStr).Append(PathSeparator).AppendFileName(application);
-		return true;
-	}
-
 	void* SessionClient::GetProcessEnvironmentVariables()
 	{
 		UBA_ASSERT(!m_environmentVariables.empty());
@@ -824,8 +845,9 @@ namespace uba
 		}
 		rec.handled = true;
 
-		auto& dir = msg.process.m_virtualApplicationDir;
-		if (!EnsureBinaryFile(out.fileName, out.virtualFileName, msg.process.m_id, msg.fileName, msg.fileNameKey, dir.c_str()))
+		StringBuffer<> dir;
+		dir.AppendDir(msg.process.m_startInfo.application);
+		if (!EnsureBinaryFile(out.fileName, out.virtualFileName, msg.process.m_id, msg.fileName, msg.fileNameKey, dir.data))
 			return false;
 
 		StringKey fileNameKey = msg.fileNameKey;
@@ -1604,9 +1626,8 @@ namespace uba
 					waitTimeoutMs = 200;
 				}
 
-				for (InternalProcessStartInfo& info : startInfos)
+				for (InternalProcessStartInfo& startInfo : startInfos)
 				{
-					auto& startInfo = info.startInfo;
 					startInfo.uiLanguage = int(m_uiLanguage);
 					startInfo.priorityClass = m_defaultPriorityClass;
 					startInfo.useCustomAllocator = !m_disableCustomAllocator;
@@ -1620,23 +1641,14 @@ namespace uba
 						startInfo.logFile = logFile.data;
 					}
 
-					StringBuffer<> realApplication;
-					if (!EnsureApplicationEnvironment(realApplication, info.processId, startInfo.application))
-					{
-						m_logger.Error(TC("Failed to ensure application environment for %s"), startInfo.application);
-						SendReturnProcess(info.processId, TC("Failed to ensure application environment"));
-						m_loop = false;
-						break;
-					}
-
 					void* env = GetProcessEnvironmentVariables();
 
-					auto process = new ProcessImpl(*this, info.processId, nullptr);
+					auto process = new ProcessImpl(*this, startInfo.processId, nullptr);
 
 					activeProcesses.emplace_back(process);
 					ProcessRec* rec = &activeProcesses.back();
 
-					rec->weight = info.weight;
+					rec->weight = startInfo.weight;
 
 					{
 						SCOPED_WRITE_LOCK(activeWeightLock, lock);
@@ -1739,7 +1751,7 @@ namespace uba
 							session.m_processFinished(&process);
 					};
 
-					process->Start(startInfo, realApplication.data, m_processWorkingDir.data, true, env, true, true);
+					process->Start(startInfo, true, env, true, true);
 				}
 
 				RemoveInactiveProcesses();
diff --git a/Engine/Source/Programs/UnrealBuildAccelerator/Common/Private/UbaSessionServer.cpp b/Engine/Source/Programs/UnrealBuildAccelerator/Common/Private/UbaSessionServer.cpp
--- a/Engine/Source/Programs/UnrealBuildAccelerator/Common/Private/UbaSessionServer.cpp
+++ b/Engine/Source/Programs/UnrealBuildAccelerator/Common/Private/UbaSessionServer.cpp
@@ -10,16 +10,16 @@
 
 namespace uba
 {
-	class SessionServer::RemoteProcess : public Process, public ProcessStartInfoHolder
+	class SessionServer::RemoteProcess final : public Process
 	{
 	public:
 		RemoteProcess(SessionServer* server, const ProcessStartInfo& si, u32 processId, float weight_)
-		:	ProcessStartInfoHolder(si)
-		,	m_server(server)
+		:	m_server(server)
+		,	m_startInfo(si)
 		,	m_processId(processId)
 		,	m_done(true)
 		{
-			weight = weight_;
+			m_startInfo.weight = weight_;
 		}
 
 		~RemoteProcess()
@@ -27,7 +27,7 @@ namespace uba
 			delete[] m_knownInputs;
 		}
 
-		virtual const ProcessStartInfo& GetStartInfo() const override { return startInfo; }
+		virtual const ProcessStartInfo& GetStartInfo() const override { return m_startInfo; }
 		virtual u32 GetId() override { return m_processId; }
 		virtual u32 GetExitCode() override { UBA_ASSERT(m_done.IsSet(0)); return m_exitCode; }
 		virtual bool HasExited() override { return m_done.IsSet(0); }
@@ -62,16 +62,17 @@ namespace uba
 		void CallProcessExit(ProcessHandle& h)
 		{
 			SCOPED_WRITE_LOCK(m_exitedLock, lock);
-			if (!startInfo.exitedFunc)
+			if (!m_startInfo.exitedFunc)
 				return;
-			auto exitedFunc = startInfo.exitedFunc;
-			auto userData = startInfo.userData;
-			startInfo.exitedFunc = nullptr;
-			startInfo.userData = nullptr;
+			auto exitedFunc = m_startInfo.exitedFunc;
+			auto userData = m_startInfo.userData;
+			m_startInfo.exitedFunc = nullptr;
+			m_startInfo.userData = nullptr;
 			exitedFunc(userData, h);
 		}
 
 		SessionServer* m_server;
+		ProcessStartInfoHolder m_startInfo;
 		ReaderWriterLock m_exitedLock;
 		u32 m_processId;
 		u32 m_exitCode = ~0u;
@@ -268,6 +269,8 @@ namespace uba
 
 	ProcessHandle SessionServer::RunProcessRemote(const ProcessStartInfo& startInfo, float weight, const void* knownInputs, u32 knownInputsCount)
 	{
+		UBA_ASSERT(!startInfo.startSuspended);
+
 		FlushDeadProcesses();
 		ValidateStartInfo(startInfo);
 		u32 processId = CreateProcessId();
@@ -278,7 +281,7 @@ namespace uba
 			auto keys = remoteProcess->m_knownInputs = new RemoteProcess::KnownInput[knownInputsCount];
 
 			u32 keysIndex = 0;
-			const TString& workingDir = remoteProcess->workingDir;
+			const TString& workingDir = remoteProcess->m_startInfo.workingDirStr;
 			for (auto kiIt = (const tchar*)knownInputs; *kiIt; kiIt += TStrlen(kiIt) + 1)
 			{
 				StringBuffer<> fileName;
@@ -314,7 +317,7 @@ namespace uba
 			remoteProcess->m_knownInputsCount = keysIndex;
 		}
 
-		remoteProcess->startInfo.rules = GetRules(remoteProcess->startInfo);
+		remoteProcess->m_startInfo.rules = GetRules(remoteProcess->m_startInfo);
 
 		ProcessHandle h(remoteProcess); // Keep ref count up even if process is removed by callbacks etc.
 
@@ -875,7 +878,7 @@ namespace uba
 							if (findIt != m_processes.end())
 							{
 								auto& process = *(RemoteProcess*)findIt->second.m_process;
-								writeCompressed = process.startInfo.rules->StoreFileCompressed(destination);
+								writeCompressed = process.m_startInfo.rules->StoreFileCompressed(destination);
 							}
 						}
 						success = m_storage.CopyOrLink(casKey, destination.data, attributes, writeCompressed);
@@ -907,7 +910,7 @@ namespace uba
 						ProcessHandle h(findIt->second);
 						lock.Leave();
 						auto& process = *(RemoteProcess*)h.m_process;
-						if (process.startInfo.trackInputs)
+						if (process.m_startInfo.trackInputs)
 						{
 							u64 bytes = GetStringWriteSize(destination.data, destination.count);
 							u64 prevSize = process.m_trackedOutputs.size();
@@ -1045,7 +1048,7 @@ namespace uba
 
 					ProcessAdded(*process, sessionId);
 					writer.WriteU32(process->m_processId);
-					process->Write(writer);
+					process->m_startInfo.Write(writer);
 
 					for (auto kiIt = process->m_knownInputs, kiEnd = kiIt + process->m_knownInputsCount; kiIt!=kiEnd; ++kiIt)
 						if (session.sentKeys.insert(kiIt->key).second)
@@ -1056,7 +1059,7 @@ namespace uba
 					if (writer.GetCapacityLeft() < 5000) // Arbitrary number to cover all parameters above
 						break;
 
-					weightLeft -= process->weight;
+					weightLeft -= process->m_startInfo.weight;
 				}
 				fillLock.Leave();
 
@@ -1178,9 +1181,9 @@ namespace uba
 					process.m_logLines.push_back({ std::move(text), type });
 				}
 
-				if (auto func = process.startInfo.logLineFunc)
+				if (auto func = process.m_startInfo.logLineFunc)
 					for (auto& line : process.m_logLines)
-						func(process.startInfo.logLineUserData, line.text.c_str(), u32(line.text.size()), line.type);
+						func(process.m_startInfo.logLineUserData, line.text.c_str(), u32(line.text.size()), line.type);
 
 				u32 id = process.m_processId;
 				Vector<ProcessLogLine> emptyLines;
@@ -1757,8 +1760,11 @@ namespace uba
 		return str.Appendf(TC("%s"), findIt->second.GetStartInfo().GetDescription()).data;
 	}
 
-	bool SessionServer::PrepareProcess(const ProcessStartInfo& startInfo, bool isChild, StringBufferBase& outRealApplication, const tchar*& outRealWorkingDir)
+	bool SessionServer::PrepareProcess(ProcessStartInfoHolder& startInfo, bool isChild, StringBufferBase& outRealApplication, const tchar*& outRealWorkingDir)
 	{
+		if (!Session::PrepareProcess(startInfo, isChild, outRealApplication, outRealWorkingDir))
+			return false;
+
 		if (!m_memTotal || !m_allowWaitOnMem || isChild)
 			return true;
 
diff --git a/Engine/Source/Programs/UnrealBuildAccelerator/Common/Public/UbaProcess.h b/Engine/Source/Programs/UnrealBuildAccelerator/Common/Public/UbaProcess.h
--- a/Engine/Source/Programs/UnrealBuildAccelerator/Common/Public/UbaProcess.h
+++ b/Engine/Source/Programs/UnrealBuildAccelerator/Common/Public/UbaProcess.h
@@ -2,6 +2,7 @@
 
 #pragma once
 
+#include "../Private/UbaProcessStartInfoHolder.h" // TODO Move this file
 #include "UbaSession.h"
 #include "UbaStats.h"
 
@@ -59,7 +60,7 @@ namespace uba
 
 		struct PipeReader;
 
-		void Start(const ProcessStartInfo& startInfo, TString&& realApplication, const tchar* realWorkingDir, bool runningRemote, void* environment, bool async, bool enableDetour);
+		bool Start(const ProcessStartInfo& startInfo, bool runningRemote, void* environment, bool async, bool enableDetour);
 
 		bool IsActive();
 		bool IsCancelled();
@@ -116,7 +117,7 @@ namespace uba
 		void ClearTempFiles();
 		void WaitForParent();
 
-		ProcessStartInfo m_startInfo;
+		ProcessStartInfoHolder m_startInfo;
 		Session& m_session;
 		ProcessImpl* m_parentProcess;
 		ReaderWriterLock m_initLock;
@@ -158,12 +159,6 @@ namespace uba
 		bool m_detourEnabled = true;
 		TString m_realApplication;
 		const tchar* m_realWorkingDir = nullptr;
-		TString m_virtualApplication;
-		TString m_virtualApplicationDir;
-		TString m_virtualWorkingDir;
-		TString m_arguments;
-		TString m_description;
-		TString m_logFile;
 		u64 m_startTime = 0;
 		Event m_waitForParent;
 		ReaderWriterLock m_logLinesLock;
diff --git a/Engine/Source/Programs/UnrealBuildAccelerator/Common/Public/UbaSession.h b/Engine/Source/Programs/UnrealBuildAccelerator/Common/Public/UbaSession.h
--- a/Engine/Source/Programs/UnrealBuildAccelerator/Common/Public/UbaSession.h
+++ b/Engine/Source/Programs/UnrealBuildAccelerator/Common/Public/UbaSession.h
@@ -22,6 +22,7 @@ namespace uba
 	class Storage;
 	class WorkManager;
 	struct ProcessStartInfo;
+	struct ProcessStartInfoHolder;
 	struct ProcessStats;
 	struct InitMessage;
 	struct InitResponse;
@@ -117,7 +118,7 @@ namespace uba
 		bool RegisterCreateFileForWrite(StringKey fileNameKey, const StringView& fileName, bool registerRealFile, u64 fileSize = 0, u64 lastWriteTime = 0, bool invalidateStorage = true);
 		u32 RegisterDeleteFile(StringKey fileNameKey, const StringView& fileName);
 
-		virtual bool PrepareProcess(const ProcessStartInfo& startInfo, bool isChild, StringBufferBase& outRealApplication, const tchar*& outRealWorkingDir);
+		virtual bool PrepareProcess(ProcessStartInfoHolder& startInfo, bool isChild, StringBufferBase& outRealApplication, const tchar*& outRealWorkingDir);
 		virtual void* GetProcessEnvironmentVariables();
 		virtual void PrintSessionStats(Logger& logger);
 
diff --git a/Engine/Source/Programs/UnrealBuildAccelerator/Common/Public/UbaSessionClient.h b/Engine/Source/Programs/UnrealBuildAccelerator/Common/Public/UbaSessionClient.h
--- a/Engine/Source/Programs/UnrealBuildAccelerator/Common/Public/UbaSessionClient.h
+++ b/Engine/Source/Programs/UnrealBuildAccelerator/Common/Public/UbaSessionClient.h
@@ -47,7 +47,7 @@ namespace uba
 	private:
 		bool RetrieveCasFile(CasKey& outNewKey, u64& outSize, const CasKey& casKey, const tchar* hint, bool storeUncompressed, bool allowProxy = true);
 
-		virtual bool PrepareProcess(const ProcessStartInfo& startInfo, bool isChild, StringBufferBase& outRealApplication, const tchar*& outRealWorkingDir) override;
+		virtual bool PrepareProcess(ProcessStartInfoHolder& startInfo, bool isChild, StringBufferBase& outRealApplication, const tchar*& outRealWorkingDir) override;
 		virtual void* GetProcessEnvironmentVariables() override;
 		virtual bool CreateFile(CreateFileResponse& out, const CreateFileMessage& msg) override;
 		virtual bool DeleteFile(DeleteFileResponse& out, const DeleteFileMessage& msg) override;
@@ -75,7 +75,6 @@ namespace uba
 
 		bool GetCasKeyForFile(CasKey& out, u32 processId, const StringBufferBase& fileName, const StringKey& fileNameKey);
 		bool ReadModules(List<ModuleInfo>& outModules, u32 processId, const tchar* application);
-		bool EnsureApplicationEnvironment(StringBufferBase& out, u32 processId, const tchar* application);
 		bool EnsureBinaryFile(StringBufferBase& out, StringBufferBase& outVirtual, u32 processId, const StringBufferBase& fileName, const StringKey& fileNameKey, const tchar* applicationDir);
 		bool WriteBinFile(StringBufferBase& out, const tchar* binaryName, const CasKey& casKey, const KeyToString& applicationDir, u32 fileAttributes);
 		bool SendFiles(ProcessImpl& process, Timer& sendFiles);
@@ -125,8 +124,9 @@ namespace uba
 		Atomic<u64> m_terminationTime;
 		Atomic<u32> m_maxProcessCount;
 
+		struct ApplicationEnvironment { ReaderWriterLock lock; TString virtualApplication; TString realApplication; };
 		ReaderWriterLock m_handledApplicationEnvironmentsLock;
-		UnorderedSet<TString> m_handledApplicationEnvironments;
+		UnorderedMap<TString, ApplicationEnvironment> m_handledApplicationEnvironments;
 
 		ReaderWriterLock m_binFileLock;
 		UnorderedMap<TString, CasKey> m_writtenBinFiles;
diff --git a/Engine/Source/Programs/UnrealBuildAccelerator/Common/Public/UbaSessionServer.h b/Engine/Source/Programs/UnrealBuildAccelerator/Common/Public/UbaSessionServer.h
--- a/Engine/Source/Programs/UnrealBuildAccelerator/Common/Public/UbaSessionServer.h
+++ b/Engine/Source/Programs/UnrealBuildAccelerator/Common/Public/UbaSessionServer.h
@@ -82,7 +82,7 @@ namespace uba
 
 		TString GetProcessDescription(u32 processId);
 
-		virtual bool PrepareProcess(const ProcessStartInfo& startInfo, bool isChild, StringBufferBase& outRealApplication, const tchar*& outRealWorkingDir) override final;
+		virtual bool PrepareProcess(ProcessStartInfoHolder& startInfo, bool isChild, StringBufferBase& outRealApplication, const tchar*& outRealWorkingDir) override final;
 		virtual bool CreateFile(CreateFileResponse& out, const CreateFileMessage& msg) override final;
 		virtual void FileEntryAdded(StringKey fileNameKey, u64 lastWritten, u64 size) override final;
 		virtual void PrintSessionStats(Logger& logger) override final;
