[UBA]
* Added more code to be able to diagnose rare issue on the farm

diff --git a/Engine/Source/Programs/UnrealBuildAccelerator/Common/Private/UbaSessionClient.cpp b/Engine/Source/Programs/UnrealBuildAccelerator/Common/Private/UbaSessionClient.cpp
--- a/Engine/Source/Programs/UnrealBuildAccelerator/Common/Private/UbaSessionClient.cpp
+++ b/Engine/Source/Programs/UnrealBuildAccelerator/Common/Private/UbaSessionClient.cpp
@@ -173,6 +173,7 @@ namespace uba
 			writer.WriteU32(processId);
 			writer.WriteString(fileName);
 			writer.WriteStringKey(fileNameKey);
+			writer.WriteBool(false);
 
 			StackBinaryReader<128> reader;
 			if (!msg.Send(reader, Stats().getFileMsg))
@@ -2070,35 +2071,56 @@ namespace uba
 				CasKey key;
 				if (GetCasKeyForFile(key, process.m_id, path, fileNameKey))
 				{
-					logStr.Clear().Appendf(TC("File %s caskey is %s."), path.data, CasKeyString(key).str);
-
-					StringBuffer<512> casKeyFile;
-					if (m_storage.GetCasFileName(casKeyFile, key))
+					if (key == CasKeyZero)
 					{
-						logStr.Appendf(TC(" CasKeyFile: %s"), casKeyFile.data);
-						u64 size = 0;
-						u32 attributes = 0;
-						bool exists = FileExists(m_logger, casKeyFile.data, &size, &attributes);
-						logStr.Appendf(TC(" Exists: %s"), ToString(exists));
-						if (exists)
+						StackBinaryWriter<1024> writer;
+						NetworkMessage msg(m_client, ServiceId, SessionMessageType_GetFileFromServer, writer);
+						writer.WriteU32(0);
+						writer.WriteString(path.data);
+						writer.WriteStringKey(fileNameKey);
+						writer.WriteBool(true);
+
+						StackBinaryReader<128> reader;
+						if (!msg.Send(reader, Stats().getFileMsg))
+							return;
+
+						key = reader.ReadCasKey();
+						if (key == CasKeyZero)
 						{
-							logStr.Appendf(TC(" Size: %llu Attr: %u"), size, attributes);
+							logStr.Clear().Appendf(TC("File %s has caskey zero. Server claims it does not exist!."), path.data);
+						}
+					}
+					if (key != CasKeyZero)
+					{
+						logStr.Clear().Appendf(TC("File %s caskey is %s."), path.data, CasKeyString(key).str);
 
-							FileHandle fileHandle = uba::CreateFileW(casKeyFile.data, GENERIC_READ, FILE_SHARE_READ, OPEN_EXISTING, DefaultAttributes());
-							if (fileHandle == InvalidFileHandle)
-							{
-								logStr.Appendf(TC(" Failed to open file %s (%s)"), casKeyFile.data, LastErrorToText().data);
-							}
-							else
+						StringBuffer<512> casKeyFile;
+						if (m_storage.GetCasFileName(casKeyFile, key))
+						{
+							logStr.Appendf(TC(" CasKeyFile: %s"), casKeyFile.data);
+							u64 size = 0;
+							u32 attributes = 0;
+							bool exists = FileExists(m_logger, casKeyFile.data, &size, &attributes);
+							logStr.Appendf(TC(" Exists: %s"), ToString(exists));
+							if (exists)
 							{
-								logStr.Appendf(TC(" CreateFile for read successful"));
-								uba::CloseFile(casKeyFile.data, fileHandle);
+								logStr.Appendf(TC(" Size: %llu Attr: %u"), size, attributes);
+
+								FileHandle fileHandle = uba::CreateFileW(casKeyFile.data, GENERIC_READ, FILE_SHARE_READ, OPEN_EXISTING, DefaultAttributes());
+								if (fileHandle == InvalidFileHandle)
+								{
+									logStr.Appendf(TC(" Failed to open file %s (%s)"), casKeyFile.data, LastErrorToText().data);
+								}
+								else
+								{
+									logStr.Appendf(TC(" CreateFile for read successful"));
+									uba::CloseFile(casKeyFile.data, fileHandle);
+								}
 							}
 						}
+						else
+							logStr.Appendf(TC(" Failed to get cas filename for cas key"));
 					}
-					else
-						logStr.Appendf(TC(" Failed to get cas filename for cas key"));
-
 				}
 				else
 					logStr.Clear().Appendf(TC("File %s caskey not found"), path.data);
diff --git a/Engine/Source/Programs/UnrealBuildAccelerator/Common/Private/UbaSessionServer.cpp b/Engine/Source/Programs/UnrealBuildAccelerator/Common/Private/UbaSessionServer.cpp
--- a/Engine/Source/Programs/UnrealBuildAccelerator/Common/Private/UbaSessionServer.cpp
+++ b/Engine/Source/Programs/UnrealBuildAccelerator/Common/Private/UbaSessionServer.cpp
@@ -771,6 +771,8 @@ namespace uba
 				StringBuffer<> fileName;
 				reader.ReadString(fileName);
 				StringKey fileNameKey = reader.ReadStringKey();
+				bool isValidation = reader.ReadBool();
+
 				CasKey casKey;
 				if (!StoreCasFile(casKey, fileNameKey, fileName.data))
 					return false;
@@ -783,6 +785,8 @@ namespace uba
 						// Not finding a file is a valid path. Some applications try with a path and if fails try another path
 						//m_logger.Error(TC("Failed to create cas for %s (not found)"), fileName.data);
 						writer.WriteCasKey(casKey);
+						if (isValidation)
+							m_logger.Info(TC("GetFileAttributesW for %s returned %u (%s)"), fileName.data, attr, LastErrorToText().data);
 						return true;
 					}
 
