[UBA]
* Added so path environment variable is networked over to helper. Needed when programs do their own search to find executable from name

diff --git a/Engine/Source/Programs/UnrealBuildAccelerator/Common/Private/UbaSessionClient.cpp b/Engine/Source/Programs/UnrealBuildAccelerator/Common/Private/UbaSessionClient.cpp
--- a/Engine/Source/Programs/UnrealBuildAccelerator/Common/Private/UbaSessionClient.cpp
+++ b/Engine/Source/Programs/UnrealBuildAccelerator/Common/Private/UbaSessionClient.cpp
@@ -1197,7 +1197,6 @@ namespace uba
 		}
 
 		#if PLATFORM_WINDOWS
-		AddEnvironmentVariableNoLock(TC("Path"), TC("c:\\noenvironment"));
 		AddEnvironmentVariableNoLock(TC("TEMP"), m_tempPath.data);
 		AddEnvironmentVariableNoLock(TC("TMP"), m_tempPath.data);
 		#else
diff --git a/Engine/Source/Programs/UnrealBuildAccelerator/Common/Private/UbaSessionServer.cpp b/Engine/Source/Programs/UnrealBuildAccelerator/Common/Private/UbaSessionServer.cpp
--- a/Engine/Source/Programs/UnrealBuildAccelerator/Common/Private/UbaSessionServer.cpp
+++ b/Engine/Source/Programs/UnrealBuildAccelerator/Common/Private/UbaSessionServer.cpp
@@ -168,7 +168,6 @@ namespace uba
 		}
 
 		#if PLATFORM_WINDOWS
-		m_localEnvironmentVariables.insert(TC("Path"));
 		m_localEnvironmentVariables.insert(TC("TMP"));
 		m_localEnvironmentVariables.insert(TC("TEMP"));
 		#else
