[UBA]
* Fixed bug in waccess when running remotely.. it could return the result that a file exists even though it didn't

diff --git a/Engine/Source/Programs/UnrealBuildAccelerator/Detours/Private/Windows/UbaDetoursFunctionsUcrtBase.inl b/Engine/Source/Programs/UnrealBuildAccelerator/Detours/Private/Windows/UbaDetoursFunctionsUcrtBase.inl
--- a/Engine/Source/Programs/UnrealBuildAccelerator/Detours/Private/Windows/UbaDetoursFunctionsUcrtBase.inl
+++ b/Engine/Source/Programs/UnrealBuildAccelerator/Detours/Private/Windows/UbaDetoursFunctionsUcrtBase.inl
@@ -167,7 +167,7 @@ errno_t Detoured__waccess_s(const wchar_t* path, int mode)
 	FileAttributes attr;
 	const wchar_t* realName = Shared_GetFileAttributes(attr, path);
 
-	if (!attr.useCache && !g_runningRemote)
+	if (!attr.useCache)
 	{
 		auto res = True__waccess_s(realName, mode);
 		DEBUG_LOG_TRUE(L"_waccess_s", L"(NOCACHE) %ls %i -> %ls", path, mode, WaccessResultToString(res));
