[UBA]
* Fix for when absolute path was sent in to SearchPathForFile

diff --git a/Engine/Source/Programs/UnrealBuildAccelerator/Common/Private/UbaFile.cpp b/Engine/Source/Programs/UnrealBuildAccelerator/Common/Private/UbaFile.cpp
--- a/Engine/Source/Programs/UnrealBuildAccelerator/Common/Private/UbaFile.cpp
+++ b/Engine/Source/Programs/UnrealBuildAccelerator/Common/Private/UbaFile.cpp
@@ -373,6 +373,8 @@ namespace uba
 
 	bool SearchPathForFile(Logger& logger, StringBufferBase& out, const tchar* file, const tchar* applicationDir)
 	{
+		UBA_ASSERT(!IsAbsolutePath(file));
+
 		StringBuffer<> fullPath;
 		fullPath.Append(applicationDir);
 		fullPath.EnsureEndsWithSlash();
@@ -405,8 +407,7 @@ namespace uba
 
 			*it = 0;
 
-			fullPath.Resize(0);
-			fullPath.Append(lastStart);
+			fullPath.Clear().Append(lastStart);
 			if (*lastStart)
 				fullPath.EnsureEndsWithSlash();
 			fullPath.Append(file);
diff --git a/Engine/Source/Programs/UnrealBuildAccelerator/Common/Private/UbaSession.cpp b/Engine/Source/Programs/UnrealBuildAccelerator/Common/Private/UbaSession.cpp
--- a/Engine/Source/Programs/UnrealBuildAccelerator/Common/Private/UbaSession.cpp
+++ b/Engine/Source/Programs/UnrealBuildAccelerator/Common/Private/UbaSession.cpp
@@ -1667,6 +1667,8 @@ namespace uba
 	{
 		if (StartsWith(startInfo.application, TC("ubacopy")))
 			return true;
+		if (IsAbsolutePath(startInfo.application))
+			return true;
 		if (!SearchPathForFile(m_logger, outRealApplication.Clear(), startInfo.application, startInfo.workingDir))
 			return false;
 		startInfo.applicationStr = outRealApplication.data;
