[UBA]
* Fixed file attributes bug for posix platforms
* Added so unit test output detoured process log when failing

diff --git a/Engine/Source/Programs/UnrealBuildAccelerator/Detours/Private/UbaDetoursShared.cpp b/Engine/Source/Programs/UnrealBuildAccelerator/Detours/Private/UbaDetoursShared.cpp
--- a/Engine/Source/Programs/UnrealBuildAccelerator/Detours/Private/UbaDetoursShared.cpp
+++ b/Engine/Source/Programs/UnrealBuildAccelerator/Detours/Private/UbaDetoursShared.cpp
@@ -328,7 +328,7 @@ namespace uba
 							outAttr.data.nFileSizeLow = li.LowPart;
 							outAttr.data.nFileSizeHigh = li.HighPart;
 #else
-							outAttr.data.st_mode = (mode_t)(S_IRUSR | S_IWUSR | S_IRGRP | S_IROTH);
+							outAttr.data.st_mode = (mode_t)(S_IFREG | S_IRUSR | S_IWUSR | S_IRGRP | S_IROTH);
 							outAttr.data.st_size = info.size;
 #endif
 						}
diff --git a/Engine/Source/Programs/UnrealBuildAccelerator/Test/Private/UbaTestSession.h b/Engine/Source/Programs/UnrealBuildAccelerator/Test/Private/UbaTestSession.h
--- a/Engine/Source/Programs/UnrealBuildAccelerator/Test/Private/UbaTestSession.h
+++ b/Engine/Source/Programs/UnrealBuildAccelerator/Test/Private/UbaTestSession.h
@@ -146,7 +146,11 @@ namespace uba
 		u32 exitCode = process.GetExitCode();
 
 		if (exitCode != 0)
+		{
+			for (auto& logLine : process.GetLogLines())
+				logger.Error(logLine.text.c_str());
 			return logger.Error(TC("UbaTestApp returned exit code %u"), exitCode);
+		}
 
 		{
 			StringBuffer<MaxPath> fileW2;
diff --git a/Engine/Source/Programs/UnrealBuildAccelerator/TestApp/Private/PosixOS/UbaTestAppPosix.cpp b/Engine/Source/Programs/UnrealBuildAccelerator/TestApp/Private/PosixOS/UbaTestAppPosix.cpp
--- a/Engine/Source/Programs/UnrealBuildAccelerator/TestApp/Private/PosixOS/UbaTestAppPosix.cpp
+++ b/Engine/Source/Programs/UnrealBuildAccelerator/TestApp/Private/PosixOS/UbaTestAppPosix.cpp
@@ -119,7 +119,7 @@ int main(int argc, char* argv[])
 		if (stat("FileW", &attrW2) == -1)
 			return LogError("stat for FileW failed");
 		if (S_ISREG(attrW2.st_mode) == 0)
-			return LogError("stat for FileW did not return normal file");
+			return LogError("stat for FileW did not return normal file (%u)", attrW2.st_mode);
 
 		if (rename("FileW", "FileW2") == -1)
 			return LogError("rename for FileW to FileW2 failed");
