[UBA]
* Added so remote helpers register output files to file mapping in case there are child processes that wants to read the output of parent processes

diff --git a/Engine/Source/Programs/UnrealBuildAccelerator/Common/Private/UbaSession.cpp b/Engine/Source/Programs/UnrealBuildAccelerator/Common/Private/UbaSession.cpp
--- a/Engine/Source/Programs/UnrealBuildAccelerator/Common/Private/UbaSession.cpp
+++ b/Engine/Source/Programs/UnrealBuildAccelerator/Common/Private/UbaSession.cpp
@@ -2013,6 +2013,13 @@ namespace uba
 			if (msg.process.m_extractExports && msg.process.m_startInfo.rules->ShouldExtractSymbols(file.name))
 				if (!ExtractSymbolsFromObjectFile(msg, name, fileSize))
 					return false;
+
+			// It might be that child processes need this file so we need to add the mapping
+			if (m_runningRemote && registerRealFile) // Note, if file is in file mapping it never touches the disk on remotes. no point adding it
+			{
+				FileExists(m_logger, writtenFile.name.c_str(), &fileSize); // Need to get file size from disk
+				AddFileMapping(writtenFile.key, name, writtenFile.name.c_str(), fileSize);
+			}
 		}
 
 		if (!msg.newName.IsEmpty())
