[Uba]
* Moved UbaProcessStartInfoHolder to public

diff --git a/Engine/Source/Programs/UnrealBuildAccelerator/Common/Private/UbaProcessStartInfoHolder.h b/Engine/Source/Programs/UnrealBuildAccelerator/Common/Private/UbaProcessStartInfoHolder.h
deleted file mode 100644
diff --git a/Engine/Source/Programs/UnrealBuildAccelerator/Common/Public/UbaProcess.h b/Engine/Source/Programs/UnrealBuildAccelerator/Common/Public/UbaProcess.h
--- a/Engine/Source/Programs/UnrealBuildAccelerator/Common/Public/UbaProcess.h
+++ b/Engine/Source/Programs/UnrealBuildAccelerator/Common/Public/UbaProcess.h
@@ -2,7 +2,7 @@
 
 #pragma once
 
-#include "../Private/UbaProcessStartInfoHolder.h" // TODO Move this file
+#include "UbaProcessStartInfoHolder.h"
 #include "UbaSession.h"
 #include "UbaStats.h"
 
diff --git a/Engine/Source/Programs/UnrealBuildAccelerator/Common/Public/UbaProcessStartInfoHolder.h b/Engine/Source/Programs/UnrealBuildAccelerator/Common/Public/UbaProcessStartInfoHolder.h
new file mode 100644
--- /dev/null
+++ b/Engine/Source/Programs/UnrealBuildAccelerator/Common/Public/UbaProcessStartInfoHolder.h
@@ -0,0 +1,122 @@
+// Copyright Epic Games, Inc. All Rights Reserved.
+
+#pragma once
+
+#define Local_GetLongPathNameW uba::GetLongPathNameW
+
+#include "UbaBinaryReaderWriter.h"
+#include "UbaFile.h"
+#include "UbaPathUtils.h"
+#include "UbaProcessStartInfo.h"
+
+namespace uba
+{
+	inline void FixFileName(StringBufferBase& out, const tchar* fileName, const tchar* workingDir)
+	{
+		tchar buffer[1024];
+		u32 charLen;
+		u64 workingDirLen = 0;
+		if (workingDir)
+			workingDirLen = TStrlen(workingDir);
+		FixPath2(fileName, workingDir, workingDirLen, buffer, sizeof_array(buffer), &charLen);
+		out.Append(buffer);
+	}
+
+	struct ProcessStartInfoHolder : public ProcessStartInfo
+	{
+		ProcessStartInfoHolder() = default;
+
+		ProcessStartInfoHolder(const ProcessStartInfo& si)
+		{
+			*(ProcessStartInfo*)this = si;
+
+			StringBuffer<512> temp;
+			FixFileName(temp, si.workingDir, nullptr);
+			temp.EnsureEndsWithSlash();
+			workingDirStr = temp.data;
+			workingDir = workingDirStr.c_str();
+
+			applicationStr = si.application;
+			application = applicationStr.c_str();
+
+			argumentsStr = si.arguments;
+			arguments = argumentsStr.c_str();
+
+			if (si.description)
+				descriptionStr = si.description;
+			description = descriptionStr.c_str();
+
+			logFileStr = si.logFile;
+			logFile = logFileStr.c_str();
+		}
+
+		void Write(BinaryWriter& writer)
+		{
+			writer.WriteString(descriptionStr);
+			writer.WriteString(applicationStr);
+			writer.WriteString(argumentsStr);
+			writer.WriteString(workingDirStr);
+			writer.WriteString(logFileStr);
+			writer.WriteU32(*(u32*)&weight);
+			writer.WriteBool(trackInputs);
+			writer.WriteBool(writeOutputFilesOnFail);
+			writer.WriteU64(outputStatsThresholdMs);
+		}
+
+		void Read(BinaryReader& reader)
+		{
+			descriptionStr = reader.ReadString();
+			applicationStr = reader.ReadString();
+			argumentsStr = reader.ReadString();
+			workingDirStr = reader.ReadString();
+			logFileStr = reader.ReadString();
+
+			Replace(applicationStr.data(), '/', PathSeparator); // TODO: Is this needed?
+
+			u32 weight32 = reader.ReadU32();
+			weight = *(float*)&weight32;
+			
+			trackInputs = reader.ReadBool();
+			writeOutputFilesOnFail = reader.ReadBool();
+			outputStatsThresholdMs = reader.ReadU64();
+
+			description = descriptionStr.c_str();
+			application = applicationStr.c_str();
+			arguments = argumentsStr.c_str();
+			workingDir = workingDirStr.c_str();
+			logFile = logFileStr.c_str();
+		}
+
+		ProcessStartInfoHolder(const ProcessStartInfoHolder& o)
+		{
+			*this = o;
+		}
+
+		void operator=(const ProcessStartInfoHolder& o)
+		{
+			*(ProcessStartInfo*)this = o;
+
+			workingDirStr = o.workingDirStr;
+			workingDir = workingDirStr.c_str();
+			applicationStr = o.applicationStr;
+			application = applicationStr.c_str();
+			argumentsStr = o.argumentsStr;
+			arguments = argumentsStr.c_str();
+			descriptionStr = o.descriptionStr;
+			description = descriptionStr.c_str();
+			logFileStr = o.logFileStr;
+			logFile = logFileStr.c_str();
+
+			weight = o.weight;
+		}
+
+		bool Expand();
+
+		TString descriptionStr;
+		TString applicationStr;
+		TString argumentsStr;
+		TString workingDirStr;
+		TString logFileStr;
+		float weight = 1.0f;
+	};
+}
\ No newline at end of file
