[UBA]
* Fixed validation code for reusing processes.. needs to match the fixed up path

diff --git a/Engine/Source/Programs/UnrealBuildAccelerator/Common/Private/UbaScheduler.cpp b/Engine/Source/Programs/UnrealBuildAccelerator/Common/Private/UbaScheduler.cpp
--- a/Engine/Source/Programs/UnrealBuildAccelerator/Common/Private/UbaScheduler.cpp
+++ b/Engine/Source/Programs/UnrealBuildAccelerator/Common/Private/UbaScheduler.cpp
@@ -506,11 +506,19 @@ namespace uba
 			ei->processIndex = indexToRun;
 
 			auto& si = *newInfo;
-			UBA_ASSERT(Equals(currentStartInfo.application, si.application));
 			outNextProcess.arguments = si.arguments;
 			outNextProcess.workingDir = si.workingDir;
 			outNextProcess.description = si.description;
 			outNextProcess.logFile = si.logFile;
+
+			#if UBA_DEBUG
+			StringBuffer<> temp;
+			if (IsAbsolutePath(si.application))
+				FixPath(si.application, nullptr, 0, temp);
+			else
+				SearchPathForFile(m_session.GetLogger(), temp, si.application, si.workingDir);
+			UBA_ASSERTF(Equals(currentStartInfo.application, temp.data), TC("%s vs %s"), currentStartInfo.application, temp.data);
+			#endif
 			return true;
 		}
 	}
