[UBA]
* Fixed validation code in scheduler for process reuse
* Moved all code inside ServerSession::HandleMessage into one function per message to more easily be able to read callstacks
* Added more validation around missing file errors.. now server logs out lots of information too

diff --git a/Engine/Source/Programs/UnrealBuildAccelerator/Common/Private/UbaScheduler.cpp b/Engine/Source/Programs/UnrealBuildAccelerator/Common/Private/UbaScheduler.cpp
--- a/Engine/Source/Programs/UnrealBuildAccelerator/Common/Private/UbaScheduler.cpp
+++ b/Engine/Source/Programs/UnrealBuildAccelerator/Common/Private/UbaScheduler.cpp
@@ -512,13 +512,20 @@ namespace uba
 			outNextProcess.logFile = si.logFile;
 
 			#if UBA_DEBUG
-			StringBuffer<> temp;
-			if (IsAbsolutePath(si.application))
-				FixPath(si.application, nullptr, 0, temp);
-			else
-				SearchPathForFile(m_session.GetLogger(), temp, si.application, si.workingDir);
-			UBA_ASSERTF(Equals(currentStartInfo.application, temp.data), TC("%s vs %s"), currentStartInfo.application, temp.data);
+			auto PrepPath = [this](StringBufferBase& out, const ProcessStartInfo& psi)
+				{
+					if (IsAbsolutePath(psi.application))
+						FixPath(psi.application, nullptr, 0, out);
+					else
+						SearchPathForFile(m_session.GetLogger(), out, psi.application, psi.workingDir);
+				};
+			StringBuffer<> temp1;
+			StringBuffer<> temp2;
+			PrepPath(temp1, currentStartInfo);
+			PrepPath(temp2, si);
+			UBA_ASSERTF(temp1.Equals(temp2.data), TC("%s vs %s"), temp1.data, temp2.data);
 			#endif
+			
 			return true;
 		}
 	}
diff --git a/Engine/Source/Programs/UnrealBuildAccelerator/Common/Private/UbaSessionClient.cpp b/Engine/Source/Programs/UnrealBuildAccelerator/Common/Private/UbaSessionClient.cpp
--- a/Engine/Source/Programs/UnrealBuildAccelerator/Common/Private/UbaSessionClient.cpp
+++ b/Engine/Source/Programs/UnrealBuildAccelerator/Common/Private/UbaSessionClient.cpp
@@ -173,7 +173,6 @@ namespace uba
 			writer.WriteU32(processId);
 			writer.WriteString(fileName);
 			writer.WriteStringKey(fileNameKey);
-			writer.WriteBool(false);
 
 			StackBinaryReader<128> reader;
 			if (!msg.Send(reader, Stats().getFileMsg))
@@ -1988,7 +1987,8 @@ namespace uba
 	bool SessionClient::LogLine(ProcessImpl& process, const tchar* line, LogEntryType logType)
 	{
 		// TODO: Remove this once we have figured out a bug that seems to exist for remote execution
-		#if PLATFORM_WINDOWS
+#if PLATFORM_WINDOWS
+
 		auto rules = process.m_startInfo.rules;
 		if (!rules)
 			return true;
@@ -2015,9 +2015,6 @@ namespace uba
 		if (!fileEnd)
 			return true;
 
-		StringBuffer<> searchString;
-		searchString.Append(fileBegin, fileEnd - fileBegin).Replace('/', PathSeparator);
-
 		MemoryBlock memoryBlock;
 		DirectoryTable dirTable(&memoryBlock);
 		{
@@ -2026,8 +2023,26 @@ namespace uba
 			dirTable.Init(m_directoryTable.m_memory, 0, m_directoryTable.m_memorySize);
 		}
 
+		StringBuffer<> errorPath;
+		errorPath.Append(fileBegin, fileEnd - fileBegin).Replace('/', PathSeparator);
+
+		{
+			StackBinaryWriter<1024> writer;
+			NetworkMessage msg(m_client, ServiceId, SessionMessageType_DebugFileNotFoundError, writer);
+			writer.WriteString(errorPath);
+			writer.WriteString(process.m_startInfo.workingDir);
+			msg.Send();
+		}
+
+		StringView searchString = errorPath;
+		if (searchString.data[0] == '.' && searchString.data[1] == '.')
+		{
+			searchString.data += 3;
+			searchString.count -= 3;
+		}
+
 		u32 foundCount = 0;
-		dirTable.TraverseAllFilesNoLock([&](const DirectoryTable::EntryInformation& info, const StringBufferBase& path)
+		dirTable.TraverseAllFilesNoLock([&](const DirectoryTable::EntryInformation& info, const StringBufferBase& path, u32 dirOffset)
 			{
 				if (!path.EndsWith(searchString.data))
 					return;
@@ -2038,7 +2053,7 @@ namespace uba
 
 				++foundCount;
 				StringBuffer<> logStr;
-				logStr.Appendf(TC("File %s found in directory table while searching for matches for %s (size %llu attr %u)"), path.data, searchString.data, info.size, info.attributes);
+				logStr.Appendf(TC("File %s found in directory table at offset %u of %u while searching for matches for %s (File size %llu attr %u)"), path.data, dirOffset, dirTable.m_memorySize, searchString.data, info.size, info.attributes);
 				process.LogLine(false, logStr.data, logType);
 
 				StringKey fileNameKey = ToStringKey(path);
@@ -2070,69 +2085,53 @@ namespace uba
 				CasKey key;
 				if (GetCasKeyForFile(key, process.m_id, path, fileNameKey))
 				{
-					if (key == CasKeyZero)
-					{
-						StackBinaryWriter<1024> writer;
-						NetworkMessage msg(m_client, ServiceId, SessionMessageType_GetFileFromServer, writer);
-						writer.WriteU32(0);
-						writer.WriteString(path.data);
-						writer.WriteStringKey(fileNameKey);
-						writer.WriteBool(true);
-
-						StackBinaryReader<128> reader;
-						if (!msg.Send(reader, Stats().getFileMsg))
-							return;
+					logStr.Clear().Appendf(TC("File %s caskey is %s."), path.data, CasKeyString(key).str);
 
-						key = reader.ReadCasKey();
-						if (key == CasKeyZero)
-						{
-							logStr.Clear().Appendf(TC("File %s has caskey zero. Server claims it does not exist!."), path.data);
-						}
-					}
-					if (key != CasKeyZero)
+					StringBuffer<512> casKeyFile;
+					if (m_storage.GetCasFileName(casKeyFile, key))
 					{
-						logStr.Clear().Appendf(TC("File %s caskey is %s."), path.data, CasKeyString(key).str);
-
-						StringBuffer<512> casKeyFile;
-						if (m_storage.GetCasFileName(casKeyFile, key))
+						logStr.Appendf(TC(" CasKeyFile: %s"), casKeyFile.data);
+						u64 size = 0;
+						u32 attributes = 0;
+						bool exists = FileExists(m_logger, casKeyFile.data, &size, &attributes);
+						logStr.Appendf(TC(" Exists: %s"), ToString(exists));
+						if (exists)
 						{
-							logStr.Appendf(TC(" CasKeyFile: %s"), casKeyFile.data);
-							u64 size = 0;
-							u32 attributes = 0;
-							bool exists = FileExists(m_logger, casKeyFile.data, &size, &attributes);
-							logStr.Appendf(TC(" Exists: %s"), ToString(exists));
-							if (exists)
+							logStr.Appendf(TC(" Size: %llu Attr: %u"), size, attributes);
+
+							FileHandle fileHandle = uba::CreateFileW(casKeyFile.data, GENERIC_READ, FILE_SHARE_READ, OPEN_EXISTING, DefaultAttributes());
+							if (fileHandle == InvalidFileHandle)
 							{
-								logStr.Appendf(TC(" Size: %llu Attr: %u"), size, attributes);
-
-								FileHandle fileHandle = uba::CreateFileW(casKeyFile.data, GENERIC_READ, FILE_SHARE_READ, OPEN_EXISTING, DefaultAttributes());
-								if (fileHandle == InvalidFileHandle)
-								{
-									logStr.Appendf(TC(" Failed to open file %s (%s)"), casKeyFile.data, LastErrorToText().data);
-								}
-								else
-								{
-									logStr.Appendf(TC(" CreateFile for read successful"));
-									uba::CloseFile(casKeyFile.data, fileHandle);
-								}
+								logStr.Appendf(TC(" Failed to open file %s (%s)"), casKeyFile.data, LastErrorToText().data);
+							}
+							else
+							{
+								logStr.Appendf(TC(" CreateFile for read successful"));
+								uba::CloseFile(casKeyFile.data, fileHandle);
 							}
 						}
-						else
-							logStr.Appendf(TC(" Failed to get cas filename for cas key"));
 					}
+					else
+						logStr.Appendf(TC(" Failed to get cas filename for cas key"));
 				}
 				else
 					logStr.Clear().Appendf(TC("File %s caskey not found"), path.data);
 				process.LogLine(false, logStr.data, logType);
 			});
+
 		if (!foundCount)
 		{
 			StringBuffer<> logStr;
 			logStr.Appendf(TC("No matching entry found in directory table while searching for matches for %s. DirTable size: %u"), searchString.data, GetDirectoryTableSize());
 			process.LogLine(false, logStr.data, logType);
+			if (errorPath.StartsWith(TC("..\\Intermediate")))
+			{
+				auto workDir = process.m_startInfo.workingDir;
+				StringBuffer<> fullPath;
+				FixPath(errorPath.data, workDir, TStrlen(workDir), fullPath);
+			}
 		}
-
-		#endif
+#endif
 		return true;
 	}
 
diff --git a/Engine/Source/Programs/UnrealBuildAccelerator/Common/Private/UbaSessionServer.cpp b/Engine/Source/Programs/UnrealBuildAccelerator/Common/Private/UbaSessionServer.cpp
--- a/Engine/Source/Programs/UnrealBuildAccelerator/Common/Private/UbaSessionServer.cpp
+++ b/Engine/Source/Programs/UnrealBuildAccelerator/Common/Private/UbaSessionServer.cpp
@@ -576,882 +576,900 @@ namespace uba
 	{
 		switch (messageType)
 		{
-			case SessionMessageType_Connect:
-			{
-				StringBuffer<128> name;
-				reader.ReadString(name);
-				u32 clientVersion = reader.ReadU32();
-
-				m_logger.Detail(TC("Client session %s connected (Id: %u, Uid: %s)"), name.data, connectionInfo.GetId(), GuidToString(connectionInfo.GetUid()).str);
+			#define UBA_SESSION_MESSAGE(x) case SessionMessageType_##x: return Handle##x(connectionInfo, reader, writer);
+			UBA_SESSION_MESSAGES
+			#undef UBA_SESSION_MESSAGE
+		}
 
-				CasKey clientKeys[2];
-				clientKeys[0] = reader.ReadCasKey();
-				clientKeys[1] = reader.ReadCasKey();
+		UBA_ASSERT(false);
+		return false;
+	}
 
-				bool binAsVersion = clientKeys[0] != CasKeyZero;
-				{
-					SCOPED_WRITE_LOCK(m_binKeysLock, lock);
-					if (m_detoursBinaryKey == CasKeyZero || (binAsVersion && m_agentBinaryKey == CasKeyZero))
-					{
-						StringBuffer<> dir;
-						if (!GetDirectoryOfCurrentModule(m_logger, dir))
-							return false;
-						u64 dirCount = dir.count;
-						bool deferCreation = true;
-						bool fileIsCompressed = false;
-						if (binAsVersion && m_agentBinaryKey == CasKeyZero)
-						{
-							UBA_ASSERT(IsWindows);
-							dir.Append(PathSeparator).Append(UBA_AGENT_EXECUTABLE);
-							m_storage.StoreCasFile(m_agentBinaryKey, dir.data, CasKeyZero, deferCreation, fileIsCompressed);
-						}
-						dir.Resize(dirCount).Append(PathSeparator).Append(UBA_DETOURS_LIBRARY);
-						if (!m_storage.StoreCasFile(m_detoursBinaryKey, dir.data, CasKeyZero, deferCreation, fileIsCompressed))
-							return m_logger.Error(TC("Failed to create cas for %s"), dir.data);
-						UBA_ASSERT(m_detoursBinaryKey != CasKeyZero);
-					}
-				}
+	bool SessionServer::HandleConnect(const ConnectionInfo& connectionInfo, BinaryReader& reader, BinaryWriter& writer)
+	{
+		StringBuffer<128> name;
+		reader.ReadString(name);
+		u32 clientVersion = reader.ReadU32();
 
-				bool isValidVersion = clientVersion == SessionNetworkVersion;
-				if (binAsVersion)
-					isValidVersion = clientKeys[0] == m_agentBinaryKey && clientKeys[1] == m_detoursBinaryKey;
+		m_logger.Detail(TC("Client session %s connected (Id: %u, Uid: %s)"), name.data, connectionInfo.GetId(), GuidToString(connectionInfo.GetUid()).str);
 
-				writer.WriteBool(isValidVersion);
+		CasKey clientKeys[2];
+		clientKeys[0] = reader.ReadCasKey();
+		clientKeys[1] = reader.ReadCasKey();
 
-				if (!isValidVersion)
+		bool binAsVersion = clientKeys[0] != CasKeyZero;
+		{
+			SCOPED_WRITE_LOCK(m_binKeysLock, lock);
+			if (m_detoursBinaryKey == CasKeyZero || (binAsVersion && m_agentBinaryKey == CasKeyZero))
+			{
+				StringBuffer<> dir;
+				if (!GetDirectoryOfCurrentModule(m_logger, dir))
+					return false;
+				u64 dirCount = dir.count;
+				bool deferCreation = true;
+				bool fileIsCompressed = false;
+				if (binAsVersion && m_agentBinaryKey == CasKeyZero)
 				{
-					StringBuffer<> response;
-					if (!isValidVersion)
-					{
-						m_logger.Warning(TC("Version mismatch. Server is on version %u while client is on %u. Disconnecting %s"), SessionNetworkVersion, clientVersion, name.data);
-						response.Appendf(TC("Version mismatch. Server is on version %u while client is on %u. Disconnecting..."), SessionNetworkVersion, clientVersion);
-					}
-					else
-					{
-						m_logger.Warning(TC("UbaAgent binaries mismatch. Disconnecting %s"), name.data);
-						response.Appendf(TC("UbaAgent binaries mismatch. Disconnecting..."));
-					}
-					writer.WriteString(response);
-					writer.WriteCasKey(m_agentBinaryKey);
-					writer.WriteCasKey(m_detoursBinaryKey);
-					return true;
+					UBA_ASSERT(IsWindows);
+					dir.Append(PathSeparator).Append(UBA_AGENT_EXECUTABLE);
+					m_storage.StoreCasFile(m_agentBinaryKey, dir.data, CasKeyZero, deferCreation, fileIsCompressed);
 				}
+				dir.Resize(dirCount).Append(PathSeparator).Append(UBA_DETOURS_LIBRARY);
+				if (!m_storage.StoreCasFile(m_detoursBinaryKey, dir.data, CasKeyZero, deferCreation, fileIsCompressed))
+					return m_logger.Error(TC("Failed to create cas for %s"), dir.data);
+				UBA_ASSERT(m_detoursBinaryKey != CasKeyZero);
+			}
+		}
 
-				u32 processSlotCount = reader.ReadU32();
-				bool dedicated = reader.ReadBool();
-
-				StringBuffer<256> info;
-				reader.ReadString(info);
-
-				// I have no explanation for this. On linux we get a shutdown crash when running through UBT if session is allocated with normal new
-				// For now we will work around it by using aligned_alloc which seems to be working on all platforms
-				auto& session = *new (aligned_alloc(alignof(ClientSession), sizeof(ClientSession))) ClientSession();
-				ScopedCriticalSection lock(m_remoteProcessAndSessionLock);
-				m_clientSessions.push_back(&session);
-				u32 sessionId = u32(m_clientSessions.size());
-				session.name = name.data;
-				session.id = connectionInfo.GetId();
-				session.processSlotCount = processSlotCount;
-				session.dedicated = dedicated;
-				m_availableRemoteSlotCount += processSlotCount;
-				++m_connectionCount;
-
-				if (!InitializeNameToHashTable())
-					return false;
-
-				writer.WriteCasKey(m_detoursBinaryKey);
-				writer.WriteBool(m_resetCas);
-				writer.WriteU32(sessionId);
-				writer.WriteU32(m_uiLanguage);
-				writer.WriteBool(m_storeObjFilesCompressed);
-				writer.WriteBool(m_detailedTrace);
-				writer.WriteBool(m_remoteLogEnabled);
-				writer.WriteBool(m_remoteTraceEnabled);
-				WriteRemoteEnvironmentVariables(writer);
+		bool isValidVersion = clientVersion == SessionNetworkVersion;
+		if (binAsVersion)
+			isValidVersion = clientKeys[0] == m_agentBinaryKey && clientKeys[1] == m_detoursBinaryKey;
 
-				m_trace.SessionAdded(sessionId, connectionInfo.GetId(), name.data, info.data); // Must be inside lock for TraceSessionUpdate() to not include
+		writer.WriteBool(isValidVersion);
 
-				lock.Leave();
-
-				return true;
+		if (!isValidVersion)
+		{
+			StringBuffer<> response;
+			if (!isValidVersion)
+			{
+				m_logger.Warning(TC("Version mismatch. Server is on version %u while client is on %u. Disconnecting %s"), SessionNetworkVersion, clientVersion, name.data);
+				response.Appendf(TC("Version mismatch. Server is on version %u while client is on %u. Disconnecting..."), SessionNetworkVersion, clientVersion);
 			}
-			case SessionMessageType_EnsureBinaryFile:
+			else
 			{
-				u32 processId = reader.ReadU32(); (void)processId;
-				StringBuffer<> fileName;
-				reader.ReadString(fileName);
-				StringKey fileNameKey = reader.ReadStringKey();
-				StringBuffer<> applicationDir;
-				reader.ReadString(applicationDir);
-
-				StringBuffer<> lookupStr;
-				lookupStr.Append(fileName).Append(applicationDir).Append('#');
-				lookupStr.MakeLower();
-				StringKey lookupKey = ToStringKeyNoCheck(lookupStr.data, lookupStr.count);
-
-				SCOPED_WRITE_LOCK(m_applicationDataLock, lock);
-				auto insres = m_applicationData.try_emplace(lookupKey);
-				ApplicationData& data = insres.first->second;
-				lock.Leave();
+				m_logger.Warning(TC("UbaAgent binaries mismatch. Disconnecting %s"), name.data);
+				response.Appendf(TC("UbaAgent binaries mismatch. Disconnecting..."));
+			}
+			writer.WriteString(response);
+			writer.WriteCasKey(m_agentBinaryKey);
+			writer.WriteCasKey(m_detoursBinaryKey);
+			return true;
+		}
 
-				SCOPED_WRITE_LOCK(data.lock, lock2);
-				if (!data.bytes.empty())
-				{
-					writer.WriteBytes(data.bytes.data(), data.bytes.size());
-					return true;
-				}
+		u32 processSlotCount = reader.ReadU32();
+		bool dedicated = reader.ReadBool();
 
-				CasKey casKey = CasKeyZero;
-				StringBuffer<> absoluteFile;
-				if (SearchPathForFile(m_logger, absoluteFile, fileName.data, applicationDir.data))
-					if (!absoluteFile.StartsWith(m_systemPath.data) || !IsKnownSystemFile(absoluteFile.data))
-					{
-						fileNameKey = ToStringKeyLower(absoluteFile);
-						if (!StoreCasFile(casKey, fileNameKey, absoluteFile.data))
-							return false;
-					}
+		StringBuffer<256> info;
+		reader.ReadString(info);
 
-				u64 startPos = writer.GetPosition();
-				writer.WriteCasKey(casKey);
-				writer.WriteString(absoluteFile);
+		// I have no explanation for this. On linux we get a shutdown crash when running through UBT if session is allocated with normal new
+		// For now we will work around it by using aligned_alloc which seems to be working on all platforms
+		auto& session = *new (aligned_alloc(alignof(ClientSession), sizeof(ClientSession))) ClientSession();
+		ScopedCriticalSection lock(m_remoteProcessAndSessionLock);
+		m_clientSessions.push_back(&session);
+		u32 sessionId = u32(m_clientSessions.size());
+		session.name = name.data;
+		session.id = connectionInfo.GetId();
+		session.processSlotCount = processSlotCount;
+		session.dedicated = dedicated;
+		m_availableRemoteSlotCount += processSlotCount;
+		++m_connectionCount;
+
+		if (!InitializeNameToHashTable())
+			return false;
 
-				u64 bytesSize = writer.GetPosition() - startPos;
-				data.bytes.resize(bytesSize);
-				memcpy(data.bytes.data(), writer.GetData() + startPos, bytesSize);
+		writer.WriteCasKey(m_detoursBinaryKey);
+		writer.WriteBool(m_resetCas);
+		writer.WriteU32(sessionId);
+		writer.WriteU32(m_uiLanguage);
+		writer.WriteBool(m_storeObjFilesCompressed);
+		writer.WriteBool(m_detailedTrace);
+		writer.WriteBool(m_remoteLogEnabled);
+		writer.WriteBool(m_remoteTraceEnabled);
+		WriteRemoteEnvironmentVariables(writer);
 
-				return true;
-			}
-			case SessionMessageType_GetApplication:
-			{
-				u32 processId = reader.ReadU32(); (void)processId;
-				StringBuffer<> applicationName;
-				reader.ReadString(applicationName);
-				StringKey applicationKey = ToStringKeyLower(applicationName);
-
-				SCOPED_WRITE_LOCK(m_applicationDataLock, lock);
-				auto insres = m_applicationData.try_emplace(applicationKey);
-				ApplicationData& data = insres.first->second;
-				lock.Leave();
-				
-				SCOPED_WRITE_LOCK(data.lock, lock2);
-				if (!data.bytes.empty())
-				{
-					writer.WriteBytes(data.bytes.data(), data.bytes.size());
-					return true;
-				}
+		m_trace.SessionAdded(sessionId, connectionInfo.GetId(), name.data, info.data); // Must be inside lock for TraceSessionUpdate() to not include
 
-				u64 startPos = writer.GetPosition();
-				Vector<BinaryModule> modules;
-				if (!GetBinaryModules(modules, applicationName.data))
-					return false;
+		lock.Leave();
+		return true;
+	}
 
-				writer.WriteU32(m_systemPath.count);
-				writer.WriteU32(u32(modules.size()));
-				for (BinaryModule& m : modules)
-				{
-					CasKey casKey;
-					if (!StoreCasFile(casKey, StringKeyZero, m.path.c_str()))
-						return false;
-					writer.WriteString(m.path);
-					writer.WriteU32(m.fileAttributes);
-					writer.WriteBool(m.isSystem);
-					writer.WriteCasKey(casKey);
-				}
+	bool SessionServer::HandleEnsureBinaryFile(const ConnectionInfo& connectionInfo, BinaryReader& reader, BinaryWriter& writer)
+	{
+		u32 processId = reader.ReadU32(); (void)processId;
+		StringBuffer<> fileName;
+		reader.ReadString(fileName);
+		StringKey fileNameKey = reader.ReadStringKey();
+		StringBuffer<> applicationDir;
+		reader.ReadString(applicationDir);
+
+		StringBuffer<> lookupStr;
+		lookupStr.Append(fileName).Append(applicationDir).Append('#');
+		lookupStr.MakeLower();
+		StringKey lookupKey = ToStringKeyNoCheck(lookupStr.data, lookupStr.count);
+
+		SCOPED_WRITE_LOCK(m_applicationDataLock, lock);
+		auto insres = m_applicationData.try_emplace(lookupKey);
+		ApplicationData& data = insres.first->second;
+		lock.Leave();
 
-				u64 bytesSize = writer.GetPosition() - startPos;
-				data.bytes.resize(bytesSize);
-				memcpy(data.bytes.data(), writer.GetData() + startPos, bytesSize);
+		SCOPED_WRITE_LOCK(data.lock, lock2);
+		if (!data.bytes.empty())
+		{
+			writer.WriteBytes(data.bytes.data(), data.bytes.size());
+			return true;
+		}
 
-				return true;
-			}
-			case SessionMessageType_GetFileFromServer:
+		CasKey casKey = CasKeyZero;
+		StringBuffer<> absoluteFile;
+		if (SearchPathForFile(m_logger, absoluteFile, fileName.data, applicationDir.data))
+			if (!absoluteFile.StartsWith(m_systemPath.data) || !IsKnownSystemFile(absoluteFile.data))
 			{
-				u32 processId = reader.ReadU32(); (void)processId;
-				StringBuffer<> fileName;
-				reader.ReadString(fileName);
-				StringKey fileNameKey = reader.ReadStringKey();
-				bool isValidation = reader.ReadBool();
-
-				CasKey casKey;
-				if (!StoreCasFile(casKey, fileNameKey, fileName.data))
+				fileNameKey = ToStringKeyLower(absoluteFile);
+				if (!StoreCasFile(casKey, fileNameKey, absoluteFile.data))
 					return false;
-				if (casKey == CasKeyZero)
-				{
-					// TODO: Should this instead use DirectoryTable? (it is currently not properly populated for lookups)
-					u32 attr = GetFileAttributesW(fileName.data);
-					if (attr == INVALID_FILE_ATTRIBUTES || !IsDirectory(attr))
-					{
-						// Not finding a file is a valid path. Some applications try with a path and if fails try another path
-						//m_logger.Error(TC("Failed to create cas for %s (not found)"), fileName.data);
-						writer.WriteCasKey(casKey);
-						if (isValidation)
-							m_logger.Info(TC("GetFileAttributesW for %s returned %u (%s)"), fileName.data, attr, LastErrorToText().data);
-						return true;
-					}
+			}
 
-					casKey = CasKeyIsDirectory;
-				}
+		u64 startPos = writer.GetPosition();
+		writer.WriteCasKey(casKey);
+		writer.WriteString(absoluteFile);
 
-				u64 serverTime;
-				if (m_nameToHashInitialized && casKey != CasKeyIsDirectory)
-				{
-					SCOPED_WRITE_LOCK(m_nameToHashLookupLock, lock);
-					serverTime = GetTime();
-					CasKey& lookupCasKey = m_nameToHashLookup[fileNameKey];
-					if (lookupCasKey != casKey)
-					{
-						lookupCasKey = casKey;
-						BinaryWriter w(m_nameToHashTableMem.memory, m_nameToHashTableMem.writtenSize, NameToHashMemSize);
-						m_nameToHashTableMem.AllocateNoLock(sizeof(StringKey) + sizeof(CasKey), 1, TC("NameToHashTable"));
-						w.WriteStringKey(fileNameKey);
-						w.WriteCasKey(casKey);
-					}
-				}
-				else
-					serverTime = GetTime();
+		u64 bytesSize = writer.GetPosition() - startPos;
+		data.bytes.resize(bytesSize);
+		memcpy(data.bytes.data(), writer.GetData() + startPos, bytesSize);
 
-				writer.WriteCasKey(casKey);
-				writer.WriteU64(serverTime);
-				return true;
-			}
-			case SessionMessageType_GetLongPathName:
-			{
-				#if PLATFORM_WINDOWS
-				StringBuffer<> shortPath;
-				reader.ReadString(shortPath);
-				StringBuffer<> longPath;
-				longPath.count = ::GetLongPathNameW(shortPath.data, longPath.data, longPath.capacity);
-				writer.WriteU32(GetLastError());
-				writer.WriteString(longPath);
-				return true;
-				#else
-				return false;
-				#endif
-			}
-			case SessionMessageType_SendFileToServer:
-			{
-				u32 processId = reader.ReadU32();
-				StringBuffer<> destination;
-				reader.ReadString(destination);
-				StringKey destinationKey = reader.ReadStringKey();
-				u32 attributes = reader.ReadU32();
-				UBA_ASSERT(attributes);
-				CasKey casKey = reader.ReadCasKey();
-				Storage::RetrieveResult res;
-				bool success = m_storage.RetrieveCasFile(res, casKey, destination.data);
-				casKey = res.casKey;
-				if (!success)
-				{
-					auto logType = connectionInfo.ShouldDisconnect() ? LogEntryType_Info : LogEntryType_Warning;
-					m_logger.Logf(logType, TC("Failed to retrieve cas for %s from client (Needed to write %s)"), CasKeyString(casKey).str, destination.data);
-				}
+		return true;
+	}
 
-				bool shouldWriteToDisk = ShouldWriteToDisk(destination);
-				if (success)
-				{
-					if (destination.StartsWith(TC("<log>")))
-					{
-						StringBuffer<> logPath;
-						logPath.Append(m_sessionLogDir).Append(destination.data + 5);
-						if (!m_storage.CopyOrLink(casKey, logPath.data, attributes))
-							m_logger.Error(TC("Failed to copy cas from %s to %s"), CasKeyString(casKey).str, logPath.data);
-						else if (!m_storage.DropCasFile(casKey, false, logPath.data))
-							m_logger.Error(TC("Failed to drop cas %s"), CasKeyString(casKey).str);
-						writer.WriteBool(true);
-						return true;
-					}
+	bool SessionServer::HandleGetApplication(const ConnectionInfo& connectionInfo, BinaryReader& reader, BinaryWriter& writer)
+	{
+		u32 processId = reader.ReadU32(); (void)processId;
+		StringBuffer<> applicationName;
+		reader.ReadString(applicationName);
+		StringKey applicationKey = ToStringKeyLower(applicationName);
+
+		SCOPED_WRITE_LOCK(m_applicationDataLock, lock);
+		auto insres = m_applicationData.try_emplace(applicationKey);
+		ApplicationData& data = insres.first->second;
+		lock.Leave();
+				
+		SCOPED_WRITE_LOCK(data.lock, lock2);
+		if (!data.bytes.empty())
+		{
+			writer.WriteBytes(data.bytes.data(), data.bytes.size());
+			return true;
+		}
 
-					if (destination.StartsWith(TC("<uba>")))
-					{
-						StringBuffer<> ubaPath;
-						ubaPath.Append(m_sessionLogDir).AppendValue(connectionInfo.GetId()).Append(TC(".uba"));
-						m_storage.CopyOrLink(casKey, ubaPath.data, attributes);
-						m_storage.DropCasFile(casKey, false, ubaPath.data);
-						writer.WriteBool(true);
-						return true;
-					}
+		u64 startPos = writer.GetPosition();
+		Vector<BinaryModule> modules;
+		if (!GetBinaryModules(modules, applicationName.data))
+			return false;
 
-					if (shouldWriteToDisk)
-					{
-						bool writeCompressed = false;
-						if (m_storeObjFilesCompressed)
-						{
-							SCOPED_READ_LOCK(m_processesLock, lock);
-							auto findIt = m_processes.find(processId);
-							if (findIt != m_processes.end())
-							{
-								auto& process = *(RemoteProcess*)findIt->second.m_process;
-								writeCompressed = process.m_startInfo.rules->StoreFileCompressed(destination);
-							}
-						}
-						success = m_storage.CopyOrLink(casKey, destination.data, attributes, writeCompressed);
-						if (!success)
-							m_logger.Error(TC("Failed to copy cas from %s to %s (%s)"), CasKeyString(casKey).str, destination.data, GetProcessDescription(processId).c_str());
-					}
-					else
-					{
-						success = m_storage.FakeCopy(casKey, destination.data);
-						if (!success)
-							m_logger.Error(TC("Failed to fake copy cas from %s to %s (%s)"), CasKeyString(casKey).str, destination.data, GetProcessDescription(processId).c_str());
-						SCOPED_WRITE_LOCK(m_receivedFilesLock, lock);
-						m_receivedFiles.try_emplace(destinationKey, casKey);
-					}
-				}
-				writer.WriteBool(success); 
+		writer.WriteU32(m_systemPath.count);
+		writer.WriteU32(u32(modules.size()));
+		for (BinaryModule& m : modules)
+		{
+			CasKey casKey;
+			if (!StoreCasFile(casKey, StringKeyZero, m.path.c_str()))
+				return false;
+			writer.WriteString(m.path);
+			writer.WriteU32(m.fileAttributes);
+			writer.WriteBool(m.isSystem);
+			writer.WriteCasKey(casKey);
+		}
 
-				if (success)
-				{
-					m_storage.DropCasFile(casKey, false, destination.data);
-					bool invalidateStorage = false; // No need, already handled in m_storage.CopyOrLink
-					RegisterCreateFileForWrite(StringKeyZero, destination, shouldWriteToDisk, 0, 0, invalidateStorage);
+		u64 bytesSize = writer.GetPosition() - startPos;
+		data.bytes.resize(bytesSize);
+		memcpy(data.bytes.data(), writer.GetData() + startPos, bytesSize);
 
+		return true;
+	}
 
-					SCOPED_WRITE_LOCK(m_processesLock, lock);
-					auto findIt = m_processes.find(processId);
-					if (findIt != m_processes.end())
-					{
-						ProcessHandle h(findIt->second);
-						lock.Leave();
-						auto& process = *(RemoteProcess*)h.m_process;
-						if (process.m_startInfo.trackInputs)
-						{
-							u64 bytes = GetStringWriteSize(destination.data, destination.count);
-							u64 prevSize = process.m_trackedOutputs.size();
-							process.m_trackedOutputs.resize(prevSize + bytes);
-							BinaryWriter w2(process.m_trackedOutputs.data(), prevSize, prevSize + bytes);
-							w2.WriteString(destination);
-						}
-					}
-				}
-				return true;
-			}
-			
-			case SessionMessageType_DeleteFile:
+	bool SessionServer::HandleGetFileFromServer(const ConnectionInfo& connectionInfo, BinaryReader& reader, BinaryWriter& writer)
+	{
+		u32 processId = reader.ReadU32(); (void)processId;
+		StringBuffer<> fileName;
+		reader.ReadString(fileName);
+		StringKey fileNameKey = reader.ReadStringKey();
+
+		CasKey casKey;
+		if (!StoreCasFile(casKey, fileNameKey, fileName.data))
+			return false;
+		if (casKey == CasKeyZero)
+		{
+			// TODO: Should this instead use DirectoryTable? (it is currently not properly populated for lookups)
+			u32 attr = GetFileAttributesW(fileName.data);
+			if (attr == INVALID_FILE_ATTRIBUTES || !IsDirectory(attr))
 			{
-				StringKey fileNameKey = reader.ReadStringKey();
-				StringBuffer<> fileName;
-				reader.ReadString(fileName);
-				bool result = uba::DeleteFileW(fileName.data);
-				u32 errorCode = GetLastError();
-				if (result)
-					RegisterDeleteFile(fileNameKey, fileName);
-				writer.WriteBool(result);
-				writer.WriteU32(errorCode);
+				// Not finding a file is a valid path. Some applications try with a path and if fails try another path
+				//m_logger.Error(TC("Failed to create cas for %s (not found)"), fileName.data);
+				writer.WriteCasKey(casKey);
 				return true;
 			}
 
-			case SessionMessageType_CopyFile:
+			casKey = CasKeyIsDirectory;
+		}
+
+		u64 serverTime;
+		if (m_nameToHashInitialized && casKey != CasKeyIsDirectory)
+		{
+			SCOPED_WRITE_LOCK(m_nameToHashLookupLock, lock);
+			serverTime = GetTime();
+			CasKey& lookupCasKey = m_nameToHashLookup[fileNameKey];
+			if (lookupCasKey != casKey)
 			{
-				StringKey fromNameKey = reader.ReadStringKey(); (void)fromNameKey;
-				StringBuffer<> fromName;
-				reader.ReadString(fromName);
-				StringKey toNameKey = reader.ReadStringKey();
-				StringBuffer<> toName;
-				reader.ReadString(toName);
-				bool result = uba::CopyFileW(fromName.data, toName.data, false);
-				u32 errorCode = GetLastError();
-				if (result)
-					RegisterCreateFileForWrite(toNameKey, toName, true);
-				writer.WriteU32(errorCode);
-				return true;
+				lookupCasKey = casKey;
+				BinaryWriter w(m_nameToHashTableMem.memory, m_nameToHashTableMem.writtenSize, NameToHashMemSize);
+				m_nameToHashTableMem.AllocateNoLock(sizeof(StringKey) + sizeof(CasKey), 1, TC("NameToHashTable"));
+				w.WriteStringKey(fileNameKey);
+				w.WriteCasKey(casKey);
 			}
+		}
+		else
+			serverTime = GetTime();
+
+		writer.WriteCasKey(casKey);
+		writer.WriteU64(serverTime);
+		return true;
+	}
+
+	bool SessionServer::HandleGetLongPathName(const ConnectionInfo& connectionInfo, BinaryReader& reader, BinaryWriter& writer)
+	{
+		#if PLATFORM_WINDOWS
+		StringBuffer<> shortPath;
+		reader.ReadString(shortPath);
+		StringBuffer<> longPath;
+		longPath.count = ::GetLongPathNameW(shortPath.data, longPath.data, longPath.capacity);
+		writer.WriteU32(GetLastError());
+		writer.WriteString(longPath);
+		return true;
+		#else
+		return false;
+		#endif
+	}
+
+	bool SessionServer::HandleSendFileToServer(const ConnectionInfo& connectionInfo, BinaryReader& reader, BinaryWriter& writer)
+	{
+		u32 processId = reader.ReadU32();
+		StringBuffer<> destination;
+		reader.ReadString(destination);
+		StringKey destinationKey = reader.ReadStringKey();
+		u32 attributes = reader.ReadU32();
+		UBA_ASSERT(attributes);
+		CasKey casKey = reader.ReadCasKey();
+		Storage::RetrieveResult res;
+		bool success = m_storage.RetrieveCasFile(res, casKey, destination.data);
+		casKey = res.casKey;
+		if (!success)
+		{
+			auto logType = connectionInfo.ShouldDisconnect() ? LogEntryType_Info : LogEntryType_Warning;
+			m_logger.Logf(logType, TC("Failed to retrieve cas for %s from client (Needed to write %s)"), CasKeyString(casKey).str, destination.data);
+		}
 
-			case SessionMessageType_CreateDirectory:
+		bool shouldWriteToDisk = ShouldWriteToDisk(destination);
+		if (success)
+		{
+			if (destination.StartsWith(TC("<log>")))
 			{
-				CreateDirectoryMessage msg;
-				reader.ReadString(msg.name);
-				CreateDirectoryResponse response;
-				if (!Session::CreateDirectory(response, msg))
-					return false;
-				writer.WriteBool(response.result);
-				writer.WriteU32(response.errorCode);
+				StringBuffer<> logPath;
+				logPath.Append(m_sessionLogDir).Append(destination.data + 5);
+				if (!m_storage.CopyOrLink(casKey, logPath.data, attributes))
+					m_logger.Error(TC("Failed to copy cas from %s to %s"), CasKeyString(casKey).str, logPath.data);
+				else if (!m_storage.DropCasFile(casKey, false, logPath.data))
+					m_logger.Error(TC("Failed to drop cas %s"), CasKeyString(casKey).str);
+				writer.WriteBool(true);
 				return true;
 			}
 
-			case SessionMessageType_RemoveDirectory:
+			if (destination.StartsWith(TC("<uba>")))
 			{
-				RemoveDirectoryMessage msg;
-				reader.ReadString(msg.name);
-				RemoveDirectoryResponse response;
-				if (!Session::RemoveDirectory(response, msg))
-					return false;
-				writer.WriteBool(response.result);
-				writer.WriteU32(response.errorCode);
+				StringBuffer<> ubaPath;
+				ubaPath.Append(m_sessionLogDir).AppendValue(connectionInfo.GetId()).Append(TC(".uba"));
+				m_storage.CopyOrLink(casKey, ubaPath.data, attributes);
+				m_storage.DropCasFile(casKey, false, ubaPath.data);
+				writer.WriteBool(true);
 				return true;
 			}
 
-			case SessionMessageType_ListDirectory:
+			if (shouldWriteToDisk)
 			{
-				u32 sessionId = reader.ReadU32();
-				u32 sessionIndex = sessionId - 1;
-				ScopedCriticalSection lock(m_remoteProcessAndSessionLock);
-				UBA_ASSERT(sessionIndex < m_clientSessions.size());
-				ClientSession& session = *m_clientSessions[sessionIndex];
-				lock.Leave();
-
-				StringBuffer<> dirName;
-				reader.ReadString(dirName);
-				StringKey dirKey = reader.ReadStringKey();
-				ListDirectoryResponse out;
-				GetListDirectoryInfo(out, dirName.data, dirKey);
-				writer.WriteU32(out.tableOffset);
-				WriteDirectoryTable(session, reader, writer);
-				return true;
+				bool writeCompressed = false;
+				if (m_storeObjFilesCompressed)
+				{
+					SCOPED_READ_LOCK(m_processesLock, lock);
+					auto findIt = m_processes.find(processId);
+					if (findIt != m_processes.end())
+					{
+						auto& process = *(RemoteProcess*)findIt->second.m_process;
+						writeCompressed = process.m_startInfo.rules->StoreFileCompressed(destination);
+					}
+				}
+				success = m_storage.CopyOrLink(casKey, destination.data, attributes, writeCompressed);
+				if (!success)
+					m_logger.Error(TC("Failed to copy cas from %s to %s (%s)"), CasKeyString(casKey).str, destination.data, GetProcessDescription(processId).c_str());
 			}
-			case SessionMessageType_GetDirectoriesFromServer:
+			else
 			{
-				u32 sessionId = reader.ReadU32();
-				u32 sessionIndex = sessionId - 1;
-				ScopedCriticalSection lock(m_remoteProcessAndSessionLock);
-				UBA_ASSERT(sessionIndex < m_clientSessions.size());
-				ClientSession& session = *m_clientSessions[sessionIndex];
-				lock.Leave();
-				WriteDirectoryTable(session, reader, writer);
-				return true;
+				success = m_storage.FakeCopy(casKey, destination.data);
+				if (!success)
+					m_logger.Error(TC("Failed to fake copy cas from %s to %s (%s)"), CasKeyString(casKey).str, destination.data, GetProcessDescription(processId).c_str());
+				SCOPED_WRITE_LOCK(m_receivedFilesLock, lock);
+				m_receivedFiles.try_emplace(destinationKey, casKey);
 			}
-			case SessionMessageType_GetNameToHashFromServer:
-			{
-				u32 requestedSize = reader.ReadU32();
+		}
+		writer.WriteBool(success); 
 
-				SCOPED_READ_LOCK(m_nameToHashLookupLock, lock);
-				if (requestedSize == ~0u)
+		if (success)
+		{
+			m_storage.DropCasFile(casKey, false, destination.data);
+			bool invalidateStorage = false; // No need, already handled in m_storage.CopyOrLink
+			RegisterCreateFileForWrite(StringKeyZero, destination, shouldWriteToDisk, 0, 0, invalidateStorage);
+
+
+			SCOPED_WRITE_LOCK(m_processesLock, lock);
+			auto findIt = m_processes.find(processId);
+			if (findIt != m_processes.end())
+			{
+				ProcessHandle h(findIt->second);
+				lock.Leave();
+				auto& process = *(RemoteProcess*)h.m_process;
+				if (process.m_startInfo.trackInputs)
 				{
-					requestedSize = u32(m_nameToHashTableMem.writtenSize);
-					writer.WriteU32(requestedSize);
+					u64 bytes = GetStringWriteSize(destination.data, destination.count);
+					u64 prevSize = process.m_trackedOutputs.size();
+					process.m_trackedOutputs.resize(prevSize + bytes);
+					BinaryWriter w2(process.m_trackedOutputs.data(), prevSize, prevSize + bytes);
+					w2.WriteString(destination);
 				}
-				writer.WriteU64(GetTime());
-				lock.Leave();
-
-				WriteNameToHashTable(reader, writer, requestedSize);
-				return true;
 			}
-			case SessionMessageType_ProcessAvailable:
-			{
-				u32 sessionId = reader.ReadU32();
-				u32 sessionIndex = sessionId - 1;
+		}
+		return true;
+	}
 
-				ScopedCriticalSection sessionsLock(m_remoteProcessAndSessionLock);
-				UBA_ASSERT(sessionIndex < m_clientSessions.size());
-				ClientSession& session = *m_clientSessions[sessionIndex];
-				sessionsLock.Leave();
+	bool SessionServer::HandleDeleteFile(const ConnectionInfo& connectionInfo, BinaryReader& reader, BinaryWriter& writer)
+	{
+		StringKey fileNameKey = reader.ReadStringKey();
+		StringBuffer<> fileName;
+		reader.ReadString(fileName);
+		bool result = uba::DeleteFileW(fileName.data);
+		u32 errorCode = GetLastError();
+		if (result)
+			RegisterDeleteFile(fileNameKey, fileName);
+		writer.WriteBool(result);
+		writer.WriteU32(errorCode);
+		return true;
+	}
 
-				u32 weight32 = reader.ReadU32();
-				float availableWeight = *(float*)&weight32;
+	bool SessionServer::HandleCopyFile(const ConnectionInfo& connectionInfo, BinaryReader& reader, BinaryWriter& writer)
+	{
+		StringKey fromNameKey = reader.ReadStringKey(); (void)fromNameKey;
+		StringBuffer<> fromName;
+		reader.ReadString(fromName);
+		StringKey toNameKey = reader.ReadStringKey();
+		StringBuffer<> toName;
+		reader.ReadString(toName);
+		bool result = uba::CopyFileW(fromName.data, toName.data, false);
+		u32 errorCode = GetLastError();
+		if (result)
+			RegisterCreateFileForWrite(toNameKey, toName, true);
+		writer.WriteU32(errorCode);
+		return true;
+	}
 
-				Vector<RemoteProcess::KnownInput*> knownInputsToSend;
+	bool SessionServer::HandleCreateDirectory(const ConnectionInfo& connectionInfo, BinaryReader& reader, BinaryWriter& writer)
+	{
+		CreateDirectoryMessage msg;
+		reader.ReadString(msg.name);
+		CreateDirectoryResponse response;
+		if (!Session::CreateDirectory(response, msg))
+			return false;
+		writer.WriteBool(response.result);
+		writer.WriteU32(response.errorCode);
+		return true;
+	}
 
-				float weightLeft = availableWeight;
-				u32 addCount = 0;
-				SCOPED_WRITE_LOCK(m_fillUpOneAtTheTimeLock, fillLock); // This is a lock to group files better (all clients connect at the same time);
-				while (weightLeft > 0)
-				{
-					RemoteProcess* process = DequeueProcess(sessionId, connectionInfo.GetId());
-					if (process == nullptr)
-						break;
+	bool SessionServer::HandleRemoveDirectory(const ConnectionInfo& connectionInfo, BinaryReader& reader, BinaryWriter& writer)
+	{
+		RemoveDirectoryMessage msg;
+		reader.ReadString(msg.name);
+		RemoveDirectoryResponse response;
+		if (!Session::RemoveDirectory(response, msg))
+			return false;
+		writer.WriteBool(response.result);
+		writer.WriteU32(response.errorCode);
+		return true;
+	}
 
-					ProcessAdded(*process, sessionId);
-					writer.WriteU32(process->m_processId);
-					process->m_startInfo.Write(writer);
+	bool SessionServer::HandleListDirectory(const ConnectionInfo& connectionInfo, BinaryReader& reader, BinaryWriter& writer)
+	{
+		u32 sessionId = reader.ReadU32();
+		u32 sessionIndex = sessionId - 1;
+		ScopedCriticalSection lock(m_remoteProcessAndSessionLock);
+		UBA_ASSERT(sessionIndex < m_clientSessions.size());
+		ClientSession& session = *m_clientSessions[sessionIndex];
+		lock.Leave();
 
-					for (auto kiIt = process->m_knownInputs, kiEnd = kiIt + process->m_knownInputsCount; kiIt!=kiEnd; ++kiIt)
-						if (session.sentKeys.insert(kiIt->key).second)
-							knownInputsToSend.push_back(kiIt);
+		StringBuffer<> dirName;
+		reader.ReadString(dirName);
+		StringKey dirKey = reader.ReadStringKey();
+		ListDirectoryResponse out;
+		GetListDirectoryInfo(out, dirName.data, dirKey);
+		writer.WriteU32(out.tableOffset);
+		WriteDirectoryTable(session, reader, writer);
+		return true;
+	}
 
-					++addCount;
+	bool SessionServer::HandleGetDirectoriesFromServer(const ConnectionInfo& connectionInfo, BinaryReader& reader, BinaryWriter& writer)
+	{
+		u32 sessionId = reader.ReadU32();
+		u32 sessionIndex = sessionId - 1;
+		ScopedCriticalSection lock(m_remoteProcessAndSessionLock);
+		UBA_ASSERT(sessionIndex < m_clientSessions.size());
+		ClientSession& session = *m_clientSessions[sessionIndex];
+		lock.Leave();
+		WriteDirectoryTable(session, reader, writer);
+		return true;
+	}
 
-					if (writer.GetCapacityLeft() < 5000) // Arbitrary number to cover all parameters above
-						break;
+	bool SessionServer::HandleGetNameToHashFromServer(const ConnectionInfo& connectionInfo, BinaryReader& reader, BinaryWriter& writer)
+	{
+		u32 requestedSize = reader.ReadU32();
 
-					weightLeft -= process->m_startInfo.weight;
-				}
-				fillLock.Leave();
+		SCOPED_READ_LOCK(m_nameToHashLookupLock, lock);
+		if (requestedSize == ~0u)
+		{
+			requestedSize = u32(m_nameToHashTableMem.writtenSize);
+			writer.WriteU32(requestedSize);
+		}
+		writer.WriteU64(GetTime());
+		lock.Leave();
 
-				u32 neededDirectoryTableSize = GetDirectoryTableSize();
-				u32 neededHashTableSize;
-				{
-					SCOPED_READ_LOCK(m_nameToHashLookupLock, l);
-					neededHashTableSize = u32(m_nameToHashTableMem.writtenSize);
-				}
+		WriteNameToHashTable(reader, writer, requestedSize);
+		return true;
+	}
+
+	bool SessionServer::HandleProcessAvailable(const ConnectionInfo& connectionInfo, BinaryReader& reader, BinaryWriter& writer)
+	{
+		u32 sessionId = reader.ReadU32();
+		u32 sessionIndex = sessionId - 1;
+
+		ScopedCriticalSection sessionsLock(m_remoteProcessAndSessionLock);
+		UBA_ASSERT(sessionIndex < m_clientSessions.size());
+		ClientSession& session = *m_clientSessions[sessionIndex];
+		sessionsLock.Leave();
+
+		u32 weight32 = reader.ReadU32();
+		float availableWeight = *(float*)&weight32;
+
+		Vector<RemoteProcess::KnownInput*> knownInputsToSend;
+
+		float weightLeft = availableWeight;
+		u32 addCount = 0;
+		SCOPED_WRITE_LOCK(m_fillUpOneAtTheTimeLock, fillLock); // This is a lock to group files better (all clients connect at the same time);
+		while (weightLeft > 0)
+		{
+			RemoteProcess* process = DequeueProcess(sessionId, connectionInfo.GetId());
+			if (process == nullptr)
+				break;
+
+			ProcessAdded(*process, sessionId);
+			writer.WriteU32(process->m_processId);
+			process->m_startInfo.Write(writer);
+
+			for (auto kiIt = process->m_knownInputs, kiEnd = kiIt + process->m_knownInputsCount; kiIt!=kiEnd; ++kiIt)
+				if (session.sentKeys.insert(kiIt->key).second)
+					knownInputsToSend.push_back(kiIt);
+
+			++addCount;
+
+			if (writer.GetCapacityLeft() < 5000) // Arbitrary number to cover all parameters above
+				break;
+
+			weightLeft -= process->m_startInfo.weight;
+		}
+		fillLock.Leave();
+
+		u32 neededDirectoryTableSize = GetDirectoryTableSize();
+		u32 neededHashTableSize;
+		{
+			SCOPED_READ_LOCK(m_nameToHashLookupLock, l);
+			neededHashTableSize = u32(m_nameToHashTableMem.writtenSize);
+		}
 
-				sessionsLock.Enter();
-				//if (addCount)
-				//	m_logger.Debug(TC("Gave %u processes to %s using up %.1f weight out of %.1f available"), addCount, session.name.c_str(), availableWeight - weightLeft, availableWeight);
+		sessionsLock.Enter();
+		//if (addCount)
+		//	m_logger.Debug(TC("Gave %u processes to %s using up %.1f weight out of %.1f available"), addCount, session.name.c_str(), availableWeight - weightLeft, availableWeight);
+
+		bool remoteExecutionEnabled = m_remoteExecutionEnabled || !m_queuedRemoteProcesses.empty();
+		if (!remoteExecutionEnabled)
+		{
+			if (session.enabled)
+				m_availableRemoteSlotCount -= session.processSlotCount - session.usedSlotCount;
+			session.enabled = false;
+			m_logger.Detail(TC("Disable remote execution on %s because remote execution has been disabled and queue is empty (will finish %u processes)"), session.name.c_str(), session.usedSlotCount);
+		}
 
-				bool remoteExecutionEnabled = m_remoteExecutionEnabled || !m_queuedRemoteProcesses.empty();
-				if (!remoteExecutionEnabled)
+		// If this client session has 0 active processes and m_maxRemoteProcessCount < total available compute - client session, then we can disconnect this client
+		if (remoteExecutionEnabled && !addCount && m_maxRemoteProcessCount != ~0u)
+		{
+			if (!session.dedicated && !session.usedSlotCount)
+			{
+				if (m_maxRemoteProcessCount < m_availableRemoteSlotCount - session.processSlotCount)
 				{
 					if (session.enabled)
 						m_availableRemoteSlotCount -= session.processSlotCount - session.usedSlotCount;
 					session.enabled = false;
-					m_logger.Detail(TC("Disable remote execution on %s because remote execution has been disabled and queue is empty (will finish %u processes)"), session.name.c_str(), session.usedSlotCount);
-				}
-
-				// If this client session has 0 active processes and m_maxRemoteProcessCount < total available compute - client session, then we can disconnect this client
-				if (remoteExecutionEnabled && !addCount && m_maxRemoteProcessCount != ~0u)
-				{
-					if (!session.dedicated && !session.usedSlotCount)
-					{
-						if (m_maxRemoteProcessCount < m_availableRemoteSlotCount - session.processSlotCount)
-						{
-							if (session.enabled)
-								m_availableRemoteSlotCount -= session.processSlotCount - session.usedSlotCount;
-							session.enabled = false;
-							remoteExecutionEnabled = false;
-							m_logger.Info(TC("Disable remote execution on %s because host session has enough help (%u left and %u remote slots)"), session.name.c_str(), m_maxRemoteProcessCount.load(), m_availableRemoteSlotCount);
-						}
-					}
+					remoteExecutionEnabled = false;
+					m_logger.Info(TC("Disable remote execution on %s because host session has enough help (%u left and %u remote slots)"), session.name.c_str(), m_maxRemoteProcessCount.load(), m_availableRemoteSlotCount);
 				}
-				sessionsLock.Leave();
+			}
+		}
+		sessionsLock.Leave();
 
-				writer.WriteU32(remoteExecutionEnabled ? SessionProcessAvailableResponse_None : SessionProcessAvailableResponse_RemoteExecutionDisabled);
+		writer.WriteU32(remoteExecutionEnabled ? SessionProcessAvailableResponse_None : SessionProcessAvailableResponse_RemoteExecutionDisabled);
 
 				
-				// Write in the needed dir and hash table offset to be up-to-date (to potentially avoid additional messages from client
-				writer.WriteU32(neededDirectoryTableSize);
-				writer.WriteU32(neededHashTableSize);
-
-				// Send caskeys of known inputs so client can start retrieving them straight away
-				u32 kiCapacity = u32(writer.GetCapacityLeft() - sizeof(u32)) / sizeof(RemoteProcess::KnownInput);
-				u32 toSendCount = Min(kiCapacity, u32(knownInputsToSend.size()));
-				writer.WriteU32(toSendCount);
-				for (auto kv : knownInputsToSend)
-				{
-					if (!toSendCount--)
-						break;
-					writer.WriteCasKey(kv->key);
-					writer.WriteU32(kv->mappingAlignment);
-				}
+		// Write in the needed dir and hash table offset to be up-to-date (to potentially avoid additional messages from client
+		writer.WriteU32(neededDirectoryTableSize);
+		writer.WriteU32(neededHashTableSize);
+
+		// Send caskeys of known inputs so client can start retrieving them straight away
+		u32 kiCapacity = u32(writer.GetCapacityLeft() - sizeof(u32)) / sizeof(RemoteProcess::KnownInput);
+		u32 toSendCount = Min(kiCapacity, u32(knownInputsToSend.size()));
+		writer.WriteU32(toSendCount);
+		for (auto kv : knownInputsToSend)
+		{
+			if (!toSendCount--)
+				break;
+			writer.WriteCasKey(kv->key);
+			writer.WriteU32(kv->mappingAlignment);
+		}
+		return true;
+	}
 
-				return true;
-			}
-			case SessionMessageType_ProcessInputs:
-			{
-				u32 processId = u32(reader.Read7BitEncoded());
-				SCOPED_WRITE_LOCK(m_processesLock, lock);
-				auto findIt = m_processes.find(processId);
-				if (findIt == m_processes.end())
-					return m_logger.Error(TC("Failed to find process for id %u when receiving custom message"), processId);
-				ProcessHandle h(findIt->second);
-				lock.Leave();
-				auto& process = *(RemoteProcess*)h.m_process;
-				auto& inputs = process.m_trackedInputs;
-				u64 size = inputs.size();
-				if (u64 addCapacity = reader.Read7BitEncoded())
-					inputs.reserve(size + addCapacity);
-				u64 toRead = reader.GetLeft();
-				inputs.resize(size + toRead);
-				reader.ReadBytes(inputs.data() + size, toRead);
-				return true;
-			}
-			case SessionMessageType_ProcessFinished:
-			{
-				u32 processId = reader.ReadU32();
+	bool SessionServer::HandleProcessInputs(const ConnectionInfo& connectionInfo, BinaryReader& reader, BinaryWriter& writer)
+	{
+		u32 processId = u32(reader.Read7BitEncoded());
+		SCOPED_WRITE_LOCK(m_processesLock, lock);
+		auto findIt = m_processes.find(processId);
+		if (findIt == m_processes.end())
+			return m_logger.Error(TC("Failed to find process for id %u when receiving custom message"), processId);
+		ProcessHandle h(findIt->second);
+		lock.Leave();
+		auto& process = *(RemoteProcess*)h.m_process;
+		auto& inputs = process.m_trackedInputs;
+		u64 size = inputs.size();
+		if (u64 addCapacity = reader.Read7BitEncoded())
+			inputs.reserve(size + addCapacity);
+		u64 toRead = reader.GetLeft();
+		inputs.resize(size + toRead);
+		reader.ReadBytes(inputs.data() + size, toRead);
+		return true;
+	}
 
-				ProcessHandle h = ProcessRemoved(processId);
-				if (!h.m_process)
-				{
-					m_logger.Warning(TC("Client finished process with id %u that is not found on server"), processId);
-					return true;
-				}
-				auto& process = *(RemoteProcess*)h.m_process;
+	bool SessionServer::HandleProcessFinished(const ConnectionInfo& connectionInfo, BinaryReader& reader, BinaryWriter& writer)
+	{
+		u32 processId = reader.ReadU32();
 
-				ScopedCriticalSection cs2(m_remoteProcessAndSessionLock);
-				if (!m_activeRemoteProcesses.erase(&process))
-				{
-					cs2.Leave();
-					m_logger.Warning(TC("Got finished process but process was not in active remote processes. Was there a disconnect happening directly after but executed before?"));
-					return true;
-				}
-				u32 sessionIndex = process.m_sessionId - 1;
-				UBA_ASSERT(sessionIndex < m_clientSessions.size());
-				auto& session = *m_clientSessions[sessionIndex];
-				++m_finishedRemoteProcessCount;
-				--session.usedSlotCount;
-				if (session.enabled)
-					++m_availableRemoteSlotCount;
-				process.m_clientId = ~0u;
-				cs2.Leave();
+		ProcessHandle h = ProcessRemoved(processId);
+		if (!h.m_process)
+		{
+			m_logger.Warning(TC("Client finished process with id %u that is not found on server"), processId);
+			return true;
+		}
+		auto& process = *(RemoteProcess*)h.m_process;
 
-				u32 exitCode = reader.ReadU32();
-				u32 logLineCount = reader.ReadU32();
+		ScopedCriticalSection cs2(m_remoteProcessAndSessionLock);
+		if (!m_activeRemoteProcesses.erase(&process))
+		{
+			cs2.Leave();
+			m_logger.Warning(TC("Got finished process but process was not in active remote processes. Was there a disconnect happening directly after but executed before?"));
+			return true;
+		}
+		u32 sessionIndex = process.m_sessionId - 1;
+		UBA_ASSERT(sessionIndex < m_clientSessions.size());
+		auto& session = *m_clientSessions[sessionIndex];
+		++m_finishedRemoteProcessCount;
+		--session.usedSlotCount;
+		if (session.enabled)
+			++m_availableRemoteSlotCount;
+		process.m_clientId = ~0u;
+		cs2.Leave();
+
+		u32 exitCode = reader.ReadU32();
+		u32 logLineCount = reader.ReadU32();
+
+		process.m_exitCode = exitCode;
+		process.m_logLines.reserve(logLineCount);
+		while (logLineCount-- != 0)
+		{
+			TString text = reader.ReadString();
+			LogEntryType type = LogEntryType(reader.ReadByte());
+			process.m_logLines.push_back({ std::move(text), type });
+		}
 
-				process.m_exitCode = exitCode;
-				process.m_logLines.reserve(logLineCount);
-				while (logLineCount-- != 0)
-				{
-					TString text = reader.ReadString();
-					LogEntryType type = LogEntryType(reader.ReadByte());
-					process.m_logLines.push_back({ std::move(text), type });
-				}
+		if (auto func = process.m_startInfo.logLineFunc)
+			for (auto& line : process.m_logLines)
+				func(process.m_startInfo.logLineUserData, line.text.c_str(), u32(line.text.size()), line.type);
 
-				if (auto func = process.m_startInfo.logLineFunc)
-					for (auto& line : process.m_logLines)
-						func(process.m_startInfo.logLineUserData, line.text.c_str(), u32(line.text.size()), line.type);
+		u32 id = process.m_processId;
+		Vector<ProcessLogLine> emptyLines;
+		auto& logLines = (exitCode != 0 || m_detailedTrace) ? process.m_logLines : emptyLines;
+		m_trace.ProcessExited(id, exitCode, reader.GetPositionData(), reader.GetLeft(), logLines);
 
-				u32 id = process.m_processId;
-				Vector<ProcessLogLine> emptyLines;
-				auto& logLines = (exitCode != 0 || m_detailedTrace) ? process.m_logLines : emptyLines;
-				m_trace.ProcessExited(id, exitCode, reader.GetPositionData(), reader.GetLeft(), logLines);
+		ProcessStats processStats;
+		processStats.Read(reader, ~0u);
+
+		//SessionStats sessionStats;
+		//sessionStats.Read(reader);
+		//StorageStats storageStats;
+		//storageStats.Read(reader);
+
+		process.m_processorTime = processStats.cpuTime;
+		process.m_wallTime = processStats.wallTime;
+		process.m_server = nullptr;
+		process.m_done.Set();
+		process.CallProcessExit(h);
+		return true;
+	}
 
-				ProcessStats processStats;
-				processStats.Read(reader, ~0u);
+	bool SessionServer::HandleProcessReturned(const ConnectionInfo& connectionInfo, BinaryReader& reader, BinaryWriter& writer)
+	{
+		u32 processId = reader.ReadU32();
+		StringBuffer<> reason;
+		reader.ReadString(reason);
 
-				//SessionStats sessionStats;
-				//sessionStats.Read(reader);
-				//StorageStats storageStats;
-				//storageStats.Read(reader);
+		ProcessHandle h = ProcessRemoved(processId);
+		RemoteProcess* process = (RemoteProcess*)h.m_process;
+		if (!process)
+		{
+			m_logger.Warning(TC("Client %s returned process %u that is not found on server (%s)"), GuidToString(connectionInfo.GetUid()).str, processId, reason.data);
+			return true;
+		}
 
-				process.m_processorTime = processStats.cpuTime;
-				process.m_wallTime = processStats.wallTime;
-				process.m_server = nullptr;
-				process.m_done.Set();
-				process.CallProcessExit(h);
-				return true;
-			}
+		ScopedCriticalSection cs2(m_remoteProcessAndSessionLock);
+		if (!m_activeRemoteProcesses.erase(process))
+		{
+			cs2.Leave();
+			m_logger.Warning(TC("Got returned process %u from client %s but process was not in active remote processes. Was there a disconnect happening directly after but executed before?"), processId, GuidToString(connectionInfo.GetUid()).str);
+			return true;
+		}
+		u32 sessionIndex = process->m_sessionId - 1;
+		UBA_ASSERT(sessionIndex < m_clientSessions.size());
+		auto& session = *m_clientSessions[sessionIndex];
+		--session.usedSlotCount;
+		if (session.enabled)
+			++m_availableRemoteSlotCount;
 
-			case SessionMessageType_ProcessReturned:
-			{
-				u32 processId = reader.ReadU32();
-				StringBuffer<> reason;
-				reader.ReadString(reason);
+		m_logger.Detail(TC("Client %s returned process to queue (%s)"), session.name.c_str(), reason.data);
+		++m_returnedRemoteProcessCount;
 
-				ProcessHandle h = ProcessRemoved(processId);
-				RemoteProcess* process = (RemoteProcess*)h.m_process;
-				if (!process)
-				{
-					m_logger.Warning(TC("Client %s returned process %u that is not found on server (%s)"), GuidToString(connectionInfo.GetUid()).str, processId, reason.data);
-					return true;
-				}
+		process->m_executingHost.clear();
+		process->m_clientId = ~0u;
+		process->m_sessionId = 0;
 
-				ScopedCriticalSection cs2(m_remoteProcessAndSessionLock);
-				if (!m_activeRemoteProcesses.erase(process))
-				{
-					cs2.Leave();
-					m_logger.Warning(TC("Got returned process %u from client %s but process was not in active remote processes. Was there a disconnect happening directly after but executed before?"), processId, GuidToString(connectionInfo.GetUid()).str);
-					return true;
-				}
-				u32 sessionIndex = process->m_sessionId - 1;
-				UBA_ASSERT(sessionIndex < m_clientSessions.size());
-				auto& session = *m_clientSessions[sessionIndex];
-				--session.usedSlotCount;
-				if (session.enabled)
-					++m_availableRemoteSlotCount;
+		m_trace.ProcessReturned(process->m_processId);
+		m_queuedRemoteProcesses.push_front(h);
 
-				m_logger.Detail(TC("Client %s returned process to queue (%s)"), session.name.c_str(), reason.data);
-				++m_returnedRemoteProcessCount;
+		if (m_remoteProcessReturnedEvent)
+			m_remoteProcessReturnedEvent(*process);
+		return true;
+	}
 
-				process->m_executingHost.clear();
-				process->m_clientId = ~0u;
-				process->m_sessionId = 0;
+	bool SessionServer::HandleVirtualAllocFailed(const ConnectionInfo& connectionInfo, BinaryReader& reader, BinaryWriter& writer)
+	{
+		m_logger.Error(TC("VIRTUAL ALLOC FAILING ON REMOTE MACHINE %s !"), GuidToString(connectionInfo.GetUid()).str);
+		return true;
+	}
 
-				m_trace.ProcessReturned(process->m_processId);
-				m_queuedRemoteProcesses.push_front(h);
+	bool SessionServer::HandleGetTraceInformation(const ConnectionInfo& connectionInfo, BinaryReader& reader, BinaryWriter& writer)
+	{
+		u32 remotePos = reader.ReadU32();
+		SCOPED_READ_LOCK(m_trace.m_memoryLock, lock);
+		u32 localPos;
+		{
+			SCOPED_READ_LOCK(m_trace.m_memoryLock, l);
+			localPos = u32(m_trace.m_memoryPos);
+		}
 
-				if (m_remoteProcessReturnedEvent)
-					m_remoteProcessReturnedEvent(*process);
-				return true;
-			}
+		writer.WriteU32(localPos);
+		u32 toWrite = Min(localPos - remotePos, u32(writer.GetCapacityLeft()));
+		writer.WriteBytes(m_trace.m_memoryBegin + remotePos, toWrite);
+		return true;
+	}
 
-			case SessionMessageType_VirtualAllocFailed:
-			{
-				m_logger.Error(TC("VIRTUAL ALLOC FAILING ON REMOTE MACHINE %s !"), GuidToString(connectionInfo.GetUid()).str);
-				return true;
-			}
-			case SessionMessageType_GetTraceInformation:
-			{
-				u32 remotePos = reader.ReadU32();
-				SCOPED_READ_LOCK(m_trace.m_memoryLock, lock);
-				u32 localPos;
-				{
-					SCOPED_READ_LOCK(m_trace.m_memoryLock, l);
-					localPos = u32(m_trace.m_memoryPos);
-				}
+	bool SessionServer::HandlePing(const ConnectionInfo& connectionInfo, BinaryReader& reader, BinaryWriter& writer)
+	{
+		u32 sessionId = reader.ReadU32();
+		u64 lastPing = reader.ReadU64();
+		u64 memAvail = reader.ReadU64();
+		u64 memTotal = reader.ReadU64();
+		u32 cpuLoadValue = reader.ReadU32();
 
-				writer.WriteU32(localPos);
-				u32 toWrite = Min(localPos - remotePos, u32(writer.GetCapacityLeft()));
-				writer.WriteBytes(m_trace.m_memoryBegin + remotePos, toWrite);
-				return true;
-			}
-			case SessionMessageType_Ping:
-			{
-				u32 sessionId = reader.ReadU32();
-				u64 lastPing = reader.ReadU64();
-				u64 memAvail = reader.ReadU64();
-				u64 memTotal = reader.ReadU64();
-				u32 cpuLoadValue = reader.ReadU32();
-
-				u64 pingTime = GetTime();
-				u32 sessionIndex = sessionId - 1;
-				ScopedCriticalSection lock(m_remoteProcessAndSessionLock);
-				UBA_ASSERT(sessionIndex < m_clientSessions.size());
-				auto& session = *m_clientSessions[sessionIndex];
-				session.pingTime = pingTime;
-				session.lastPing = lastPing;
-				session.memAvail = memAvail;
-				session.memTotal = memTotal;
-				session.cpuLoad = *(float*)&cpuLoadValue;
-				writer.WriteBool(session.abort);
+		u64 pingTime = GetTime();
+		u32 sessionIndex = sessionId - 1;
+		ScopedCriticalSection lock(m_remoteProcessAndSessionLock);
+		UBA_ASSERT(sessionIndex < m_clientSessions.size());
+		auto& session = *m_clientSessions[sessionIndex];
+		session.pingTime = pingTime;
+		session.lastPing = lastPing;
+		session.memAvail = memAvail;
+		session.memTotal = memTotal;
+		session.cpuLoad = *(float*)&cpuLoadValue;
+		writer.WriteBool(session.abort);
 
-				return true;
-			}
-			case SessionMessageType_Notification:
-			{
-				u32 sessionId = reader.ReadU32();
-				StringBuffer<1024> str;
-				reader.ReadString(str);
-				m_trace.SessionNotification(sessionId, str.data);
-				return true;
-			}
-			case SessionMessageType_GetNextProcess:
-			{
-				u32 processId = reader.ReadU32();
-				u32 prevExitCode = reader.ReadU32();
-				SCOPED_WRITE_LOCK(m_processesLock, lock);
-				auto findIt = m_processes.find(processId);
-				if (findIt == m_processes.end())
-					return m_logger.Error(TC("Failed to find process for id %u when receiving custom message"), processId);
-				ProcessHandle h(findIt->second);
-				lock.Leave();
+		return true;
+	}
 
-				auto& remoteProcess = *(RemoteProcess*)h.m_process;
-				SCOPED_WRITE_LOCK(remoteProcess.m_exitedLock, exitedLock);
-				NextProcessInfo nextProcess;
-				bool newProcess;
-				remoteProcess.m_exitCode = prevExitCode;
-				remoteProcess.m_done.Set();
-				bool success = GetNextProcess(remoteProcess, newProcess, nextProcess, prevExitCode, reader);
-				remoteProcess.m_exitCode = ~0u;
-				remoteProcess.m_done.Reset();
-				if (!success)
-					return false;
+	bool SessionServer::HandleNotification(const ConnectionInfo& connectionInfo, BinaryReader& reader, BinaryWriter& writer)
+	{
+		u32 sessionId = reader.ReadU32();
+		StringBuffer<1024> str;
+		reader.ReadString(str);
+		m_trace.SessionNotification(sessionId, str.data);
+		return true;
+	}
 
-				writer.WriteBool(newProcess);
-				if (newProcess)
-				{
-					writer.WriteString(nextProcess.arguments);
-					writer.WriteString(nextProcess.workingDir);
-					writer.WriteString(nextProcess.description);
-					writer.WriteString(nextProcess.logFile);
-				}
-				return true;
-			}
-			case SessionMessageType_Custom:
-			{
-				u32 processId = reader.ReadU32();
-				SCOPED_WRITE_LOCK(m_processesLock, lock);
-				auto findIt = m_processes.find(processId);
-				if (findIt == m_processes.end())
-					return m_logger.Error(TC("Failed to find process for id %u when receiving custom message"), processId);
-				ProcessHandle h(findIt->second);
-				lock.Leave();
+	bool SessionServer::HandleGetNextProcess(const ConnectionInfo& connectionInfo, BinaryReader& reader, BinaryWriter& writer)
+	{
+		u32 processId = reader.ReadU32();
+		u32 prevExitCode = reader.ReadU32();
+		SCOPED_WRITE_LOCK(m_processesLock, lock);
+		auto findIt = m_processes.find(processId);
+		if (findIt == m_processes.end())
+			return m_logger.Error(TC("Failed to find process for id %u when receiving custom message"), processId);
+		ProcessHandle h(findIt->second);
+		lock.Leave();
 
-				auto& remoteProcess = *(RemoteProcess*)h.m_process;
-				SCOPED_WRITE_LOCK(remoteProcess.m_exitedLock, exitedLock);
-				CustomMessage(remoteProcess, reader, writer);
-				return true;
-			}
-			case SessionMessageType_UpdateEnvironment:
+		auto& remoteProcess = *(RemoteProcess*)h.m_process;
+		SCOPED_WRITE_LOCK(remoteProcess.m_exitedLock, exitedLock);
+		NextProcessInfo nextProcess;
+		bool newProcess;
+		remoteProcess.m_exitCode = prevExitCode;
+		remoteProcess.m_done.Set();
+		bool success = GetNextProcess(remoteProcess, newProcess, nextProcess, prevExitCode, reader);
+		remoteProcess.m_exitCode = ~0u;
+		remoteProcess.m_done.Reset();
+		if (!success)
+			return false;
+
+		writer.WriteBool(newProcess);
+		if (newProcess)
+		{
+			writer.WriteString(nextProcess.arguments);
+			writer.WriteString(nextProcess.workingDir);
+			writer.WriteString(nextProcess.description);
+			writer.WriteString(nextProcess.logFile);
+		}
+		return true;
+	}
+
+	bool SessionServer::HandleCustom(const ConnectionInfo& connectionInfo, BinaryReader& reader, BinaryWriter& writer)
+	{
+		u32 processId = reader.ReadU32();
+		SCOPED_WRITE_LOCK(m_processesLock, lock);
+		auto findIt = m_processes.find(processId);
+		if (findIt == m_processes.end())
+			return m_logger.Error(TC("Failed to find process for id %u when receiving custom message"), processId);
+		ProcessHandle h(findIt->second);
+		lock.Leave();
+
+		auto& remoteProcess = *(RemoteProcess*)h.m_process;
+		SCOPED_WRITE_LOCK(remoteProcess.m_exitedLock, exitedLock);
+		CustomMessage(remoteProcess, reader, writer);
+		return true;
+	}
+
+	bool SessionServer::HandleUpdateEnvironment(const ConnectionInfo& connectionInfo, BinaryReader& reader, BinaryWriter& writer)
+	{
+		u32 processId = reader.ReadU32();
+		StringBuffer<> reason;
+		reader.ReadString(reason);
+		m_trace.ProcessEnvironmentUpdated(processId, reason.data, reader.GetPositionData(), reader.GetLeft());
+		return true;
+	}
+
+	bool SessionServer::HandleSummary(const ConnectionInfo& connectionInfo, BinaryReader& reader, BinaryWriter& writer)
+	{
+		u32 sessionId = reader.ReadU32();
+		m_trace.SessionSummary(sessionId, reader.GetPositionData(), reader.GetLeft());
+		return true;
+	}
+
+	bool SessionServer::HandleCommand(const ConnectionInfo& connectionInfo, BinaryReader& reader, BinaryWriter& writer)
+	{
+		StringBuffer<128> command;
+		reader.ReadString(command);
+
+		auto WriteString = [&](const tchar* str, LogEntryType type = LogEntryType_Info) { writer.WriteByte(type); writer.WriteString(str); };
+
+		if (command.Equals(TC("status")))
+		{
+			u32 totalUsed = 0;
+			u32 totalSlots = 0;
+			ScopedCriticalSection queueLock(m_remoteProcessAndSessionLock);
+			u64 time = GetTime();
+			for (auto& s : m_clientSessions)
 			{
-				u32 processId = reader.ReadU32();
-				StringBuffer<> reason;
-				reader.ReadString(reason);
-				m_trace.ProcessEnvironmentUpdated(processId, reason.data, reader.GetPositionData(), reader.GetLeft());
-				return true;
+				if (!s->enabled)
+					continue;
+				WriteString(StringBuffer<>().Appendf(TC("Session %u (%s)"), s->id, s->name.c_str()).data);
+				WriteString(StringBuffer<>().Appendf(TC("   Process slots used %u/%u"), s->usedSlotCount, s->processSlotCount).data);
+				if (s->pingTime)
+					WriteString(StringBuffer<>().Appendf(TC("   Last ping %s ago"), TimeToText(time - s->pingTime).str).data);
+				totalUsed += s->usedSlotCount;
+				totalSlots += s->processSlotCount;
 			}
-			case SessionMessageType_Summary:
+			WriteString(StringBuffer<>().Appendf(TC("Total remote slots used %u/%u"), totalUsed, totalSlots).data);
+		}
+		else if (command.StartsWith(TC("abort")))
+		{
+			bool abortWithProxy = command.Equals(TC("abortproxy"));
+			bool abortUseProxy = command.Equals(TC("abortnonproxy"));
+			if (!abortWithProxy && !abortUseProxy)
 			{
-				u32 sessionId = reader.ReadU32();
-				m_trace.SessionSummary(sessionId, reader.GetPositionData(), reader.GetLeft());
-				return true;
+				abortWithProxy = true;
+				abortUseProxy = true;
 			}
-			case SessionMessageType_Command:
-			{
-				StringBuffer<128> command;
-				reader.ReadString(command);
-
-				auto WriteString = [&](const tchar* str, LogEntryType type = LogEntryType_Info) { writer.WriteByte(type); writer.WriteString(str); };
-
-				if (command.Equals(TC("status")))
-				{
-					u32 totalUsed = 0;
-					u32 totalSlots = 0;
-					ScopedCriticalSection queueLock(m_remoteProcessAndSessionLock);
-					u64 time = GetTime();
-					for (auto& s : m_clientSessions)
-					{
-						if (!s->enabled)
-							continue;
-						WriteString(StringBuffer<>().Appendf(TC("Session %u (%s)"), s->id, s->name.c_str()).data);
-						WriteString(StringBuffer<>().Appendf(TC("   Process slots used %u/%u"), s->usedSlotCount, s->processSlotCount).data);
-						if (s->pingTime)
-							WriteString(StringBuffer<>().Appendf(TC("   Last ping %s ago"), TimeToText(time - s->pingTime).str).data);
-						totalUsed += s->usedSlotCount;
-						totalSlots += s->processSlotCount;
-					}
-					WriteString(StringBuffer<>().Appendf(TC("Total remote slots used %u/%u"), totalUsed, totalSlots).data);
-				}
-				else if (command.StartsWith(TC("abort")))
-				{
-					bool abortWithProxy = command.Equals(TC("abortproxy"));
-					bool abortUseProxy = command.Equals(TC("abortnonproxy"));
-					if (!abortWithProxy && !abortUseProxy)
-					{
-						abortWithProxy = true;
-						abortUseProxy = true;
-					}
-					ScopedCriticalSection queueLock(m_remoteProcessAndSessionLock);
-					u32 abortCount = 0;
-					for (auto& s : m_clientSessions)
-					{
-						if (!s->enabled || s->abort)
-							continue;
-						bool hasProxy = m_storage.HasProxy(s->id);
-						if (abortWithProxy && hasProxy)
-							s->abort = true;
-						else if (abortUseProxy && !hasProxy)
-							s->abort = true;
-						if (s->abort)
-							++abortCount;
-					}
-					WriteString(StringBuffer<>().Appendf(TC("Aborting: %u remote sessions"), abortCount).data);
-				}
-				else if (command.Equals(TC("disableremote")))
-				{
-					DisableRemoteExecution();
-					WriteString(StringBuffer<>().Appendf(TC("Remote execution is disabled")).data);
-				}
-				else
-				{
-					WriteString(StringBuffer<>().Appendf(TC("Unknown command: %s"), command.data).data, LogEntryType_Error);
-				}
-				writer.WriteByte(255);
-				return true;
-			}
-			case SessionMessageType_SHGetKnownFolderPath:
+			ScopedCriticalSection queueLock(m_remoteProcessAndSessionLock);
+			u32 abortCount = 0;
+			for (auto& s : m_clientSessions)
 			{
-#if PLATFORM_WINDOWS
-				GUID kfid;
-				reader.ReadBytes(&kfid, sizeof(GUID));
-				u32 flags = reader.ReadU32();
-				PWSTR str;
-
-				static HMODULE moduleHandle = LoadLibrary(L"Shell32.dll");
-				using SHGetKnownFolderPathFunc = HRESULT(const GUID& rfid, DWORD dwFlags, HANDLE hToken, PWSTR* ppszPath);
-				static SHGetKnownFolderPathFunc* SHGetKnownFolderPath = (SHGetKnownFolderPathFunc*)GetProcAddress(moduleHandle, "SHGetKnownFolderPath");;
-				HRESULT res = SHGetKnownFolderPath(kfid, flags, NULL, &str);
-				writer.WriteU32(res);
-				if (res == S_OK)
-				{
-					writer.WriteString(str);
-					CoTaskMemFree(str);
-				}
-#endif
-				return true;
+				if (!s->enabled || s->abort)
+					continue;
+				bool hasProxy = m_storage.HasProxy(s->id);
+				if (abortWithProxy && hasProxy)
+					s->abort = true;
+				else if (abortUseProxy && !hasProxy)
+					s->abort = true;
+				if (s->abort)
+					++abortCount;
 			}
+			WriteString(StringBuffer<>().Appendf(TC("Aborting: %u remote sessions"), abortCount).data);
+		}
+		else if (command.Equals(TC("disableremote")))
+		{
+			DisableRemoteExecution();
+			WriteString(StringBuffer<>().Appendf(TC("Remote execution is disabled")).data);
 		}
+		else
+		{
+			WriteString(StringBuffer<>().Appendf(TC("Unknown command: %s"), command.data).data, LogEntryType_Error);
+		}
+		writer.WriteByte(255);
+		return true;
+	}
 
-		UBA_ASSERT(false);
-		return false;
+	bool SessionServer::HandleSHGetKnownFolderPath(const ConnectionInfo& connectionInfo, BinaryReader& reader, BinaryWriter& writer)
+	{
+#if PLATFORM_WINDOWS
+		GUID kfid;
+		reader.ReadBytes(&kfid, sizeof(GUID));
+		u32 flags = reader.ReadU32();
+		PWSTR str;
+
+		static HMODULE moduleHandle = LoadLibrary(L"Shell32.dll");
+		using SHGetKnownFolderPathFunc = HRESULT(const GUID& rfid, DWORD dwFlags, HANDLE hToken, PWSTR* ppszPath);
+		static SHGetKnownFolderPathFunc* SHGetKnownFolderPath = (SHGetKnownFolderPathFunc*)GetProcAddress(moduleHandle, "SHGetKnownFolderPath");;
+		HRESULT res = SHGetKnownFolderPath(kfid, flags, NULL, &str);
+		writer.WriteU32(res);
+		if (res == S_OK)
+		{
+			writer.WriteString(str);
+			CoTaskMemFree(str);
+		}
+#endif
+		return true;
 	}
 
 	bool SessionServer::StoreCasFile(CasKey& out, const StringKey& fileNameKey, const tchar* fileName)
@@ -1956,4 +1974,105 @@ namespace uba
 
 		return true;
 	}
+
+	bool SessionServer::HandleDebugFileNotFoundError(const ConnectionInfo& connectionInfo, BinaryReader& reader, BinaryWriter& writer)
+	{
+#if PLATFORM_WINDOWS
+		StringBuffer<> errorPath;
+		reader.ReadString(errorPath);
+		StringBuffer<> workDir;
+		reader.ReadString(workDir);
+
+		StringView searchString = errorPath;
+		if (searchString.data[0] == '.' && searchString.data[1] == '.')
+		{
+			searchString.data += 3;
+			searchString.count -= 3;
+		}
+
+		auto LogLine = [&](const StringView& text) { m_logger.Log(LogEntryType_Warning, text.data, text.count); };
+
+		// Make a copy of dir table since we can't populate it
+		MemoryBlock block(64*1024*1024);
+		DirectoryTable dirTable(&block);
+		u8* dirMem;
+		u32 dirMemSize;
+		{
+			SCOPED_READ_LOCK(m_directoryTable.m_memoryLock, lock);
+			dirMem = m_directoryTable.m_memory;
+			dirMemSize = m_directoryTable.m_memorySize;
+		}
+
+		dirTable.Init(dirMem, 0, dirMemSize);
+
+		u32 foundCount = 0;
+		dirTable.TraverseAllFilesNoLock([&](const DirectoryTable::EntryInformation& info, const StringBufferBase& path, u32 dirOffset)
+			{
+				if (!path.EndsWith(searchString.data))
+					return;
+				if (path[path.count - searchString.count - 1] != PathSeparator)
+					return;
+
+				auto ToString = [](bool b) { return b ? TC("true") : TC("false"); };
+
+				++foundCount;
+				StringBuffer<> logStr;
+				logStr.Appendf(TC("File %s found in directory table at offset %u of %u while searching for matches for %s (File size %llu attr %u)"), path.data, dirOffset, dirTable.m_memorySize, searchString.data, info.size, info.attributes);
+				LogLine(logStr);
+
+				StringKey fileNameKey = ToStringKey(path);
+				{
+					SCOPED_READ_LOCK(m_fileMappingTableLookupLock, mlock);
+					auto findIt = m_fileMappingTableLookup.find(fileNameKey);
+					if (findIt != m_fileMappingTableLookup.end())
+					{
+						auto& entry = findIt->second;
+						SCOPED_READ_LOCK(entry.lock, entryCs);
+						logStr.Clear().Appendf(TC("File %s found in mapping table table."), path.data);
+						if (entry.handled)
+						{
+							StringBuffer<128> mappingName;
+							if (entry.mapping.IsValid())
+								Storage::GetMappingString(mappingName, entry.mapping, entry.mappingOffset);
+							else
+								mappingName.Append(TC("Not valid"));
+							logStr.Appendf(TC(" Success: %s Size: %u IsDir: %s Mapping name: %s Mapping offset: %u"), ToString(entry.success), entry.size, ToString(entry.isDir), mappingName.data, entry.mappingOffset);
+						}
+						else
+						{
+							logStr.Appendf(TC(" Entry not handled"));
+						}
+					}
+					else
+						logStr.Clear().Appendf(TC("File %s not found in mapping table table."), path.data);
+					LogLine(logStr);
+				}
+				{
+					SCOPED_READ_LOCK(m_nameToHashLookupLock, hlock);
+					auto findIt = m_nameToHashLookup.find(fileNameKey);
+					if (findIt != m_nameToHashLookup.end())
+						logStr.Clear().Appendf(TC("File %s found in name-to-hash lookup. CasKey is %s"), path.data, CasKeyString(findIt->second).str);
+					else
+						logStr.Clear().Appendf(TC("File %s not found in name-to-hash lookup"), path.data);
+					LogLine(logStr);
+				}
+			});
+
+		if (!foundCount)
+		{
+			StringBuffer<> logStr;
+			logStr.Appendf(TC("No matching entry found in directory table while searching for matches for %s. DirTable size: %u"), searchString.data, GetDirectoryTableSize());
+			LogLine(logStr);
+			/*
+			if (errorPath.StartsWith(TC("..\\Intermediate")))
+			{
+				StringBuffer<> fullPath;
+				FixPath(errorPath.data, workDir.data, workDir.count, fullPath);
+				GetFileAttributes
+			}
+			*/
+		}
+#endif
+		return true;
+	}
 }
diff --git a/Engine/Source/Programs/UnrealBuildAccelerator/Common/Public/UbaNetwork.h b/Engine/Source/Programs/UnrealBuildAccelerator/Common/Public/UbaNetwork.h
--- a/Engine/Source/Programs/UnrealBuildAccelerator/Common/Public/UbaNetwork.h
+++ b/Engine/Source/Programs/UnrealBuildAccelerator/Common/Public/UbaNetwork.h
@@ -15,7 +15,7 @@ namespace uba
 	
 	static constexpr u32 SystemNetworkVersion = 1339;
 	static constexpr u32 StorageNetworkVersion = 4;
-	static constexpr u32 SessionNetworkVersion = 33;
+	static constexpr u32 SessionNetworkVersion = 34;
 	static constexpr u32 CacheNetworkVersion = 4;
 
 	static constexpr u32 CachePathTableMaxSize = 32*1024*1024;
@@ -87,6 +87,7 @@ namespace uba
 		UBA_SESSION_MESSAGE(Summary) \
 		UBA_SESSION_MESSAGE(Command) \
 		UBA_SESSION_MESSAGE(SHGetKnownFolderPath) \
+		UBA_SESSION_MESSAGE(DebugFileNotFoundError) \
 
 	enum SessionMessageType : u8 
 	{
diff --git a/Engine/Source/Programs/UnrealBuildAccelerator/Common/Public/UbaSessionServer.h b/Engine/Source/Programs/UnrealBuildAccelerator/Common/Public/UbaSessionServer.h
--- a/Engine/Source/Programs/UnrealBuildAccelerator/Common/Public/UbaSessionServer.h
+++ b/Engine/Source/Programs/UnrealBuildAccelerator/Common/Public/UbaSessionServer.h
@@ -67,6 +67,9 @@ namespace uba
 
 		void OnDisconnected(const Guid& clientUid, u32 clientId);
 		bool HandleMessage(const ConnectionInfo& connectionInfo, u8 messageType, BinaryReader& reader, BinaryWriter& writer);
+		#define UBA_SESSION_MESSAGE(x) bool Handle##x(const ConnectionInfo& connectionInfo, BinaryReader& reader, BinaryWriter& writer);
+		UBA_SESSION_MESSAGES
+		#undef UBA_SESSION_MESSAGE
 
 		bool StoreCasFile(CasKey& out, const StringKey& fileNameKey, const tchar* fileName);
 		bool WriteDirectoryTable(ClientSession& session, BinaryReader& reader, BinaryWriter& writer);
@@ -91,6 +94,7 @@ namespace uba
 		void WriteRemoteEnvironmentVariables(BinaryWriter& writer);
 		bool InitializeNameToHashTable();
 
+
 		NetworkServer& m_server;
 		u32 m_uiLanguage;
 		Atomic<u32> m_maxRemoteProcessCount;
diff --git a/Engine/Source/Programs/UnrealBuildAccelerator/Core/Public/UbaDirectoryTable.h b/Engine/Source/Programs/UnrealBuildAccelerator/Core/Public/UbaDirectoryTable.h
--- a/Engine/Source/Programs/UnrealBuildAccelerator/Core/Public/UbaDirectoryTable.h
+++ b/Engine/Source/Programs/UnrealBuildAccelerator/Core/Public/UbaDirectoryTable.h
@@ -328,11 +328,12 @@ namespace uba
 				DirectoryTable::EntryInformation info;
 				StringBuffer<> fileName(path);
 				fileName.Append(PathSeparator);
-				GetEntryInformation(info, fileKv.second, fileName.data + fileName.count, fileName.capacity - fileName.count);
+				u32 fileOffset = fileKv.second;
+				GetEntryInformation(info, fileOffset, fileName.data + fileName.count, fileName.capacity - fileName.count);
 				fileName.count = TStrlen(fileName.data);
 				if (CaseInsensitiveFs)
 					fileName.MakeLower();
-				func(info, fileName);
+				func(info, fileName, fileOffset);
 				TraverseFilesRecursiveNoLock(fileName, func);
 			}
 		}
@@ -350,6 +351,8 @@ namespace uba
 #endif
 
 		DirectoryTable(MemoryBlock* block) : m_memoryBlock(block), m_lookup(block) {}
+		DirectoryTable(const DirectoryTable&) = delete;
+		void operator=(const DirectoryTable&) = delete;
 
 		MemoryBlock* m_memoryBlock;
 		ReaderWriterLock m_lookupLock;
diff --git a/Engine/Source/Programs/UnrealBuildAccelerator/Detours/Private/PosixOS/UbaDetoursFunctionsPosix.cpp b/Engine/Source/Programs/UnrealBuildAccelerator/Detours/Private/PosixOS/UbaDetoursFunctionsPosix.cpp
--- a/Engine/Source/Programs/UnrealBuildAccelerator/Detours/Private/PosixOS/UbaDetoursFunctionsPosix.cpp
+++ b/Engine/Source/Programs/UnrealBuildAccelerator/Detours/Private/PosixOS/UbaDetoursFunctionsPosix.cpp
@@ -2041,6 +2041,7 @@ extern "C"
 
 		if (newProcess)
 		{
+			g_kernelStats = {};
 			g_stats = {};
 
 			//#if UBA_DEBUG_LOG_ENABLED
diff --git a/Engine/Source/Programs/UnrealBuildAccelerator/Detours/Private/Windows/UbaDetoursFunctionsWin.cpp b/Engine/Source/Programs/UnrealBuildAccelerator/Detours/Private/Windows/UbaDetoursFunctionsWin.cpp
--- a/Engine/Source/Programs/UnrealBuildAccelerator/Detours/Private/Windows/UbaDetoursFunctionsWin.cpp
+++ b/Engine/Source/Programs/UnrealBuildAccelerator/Detours/Private/Windows/UbaDetoursFunctionsWin.cpp
@@ -1040,6 +1040,7 @@ extern "C"
 
 		if (newProcess)
 		{
+			g_kernelStats = {};
 			g_stats = {};
 
 			#if UBA_DEBUG_LOG_ENABLED
