[UBA]
* Fixed so bad file path entries are not saved to casdb file
* Added so allow proxy can be set from config files
* Added assert for if bad files have entered the cas entry directory on disk

diff --git a/Engine/Source/Programs/UnrealBuildAccelerator/Common/Private/UbaSession.cpp b/Engine/Source/Programs/UnrealBuildAccelerator/Common/Private/UbaSession.cpp
--- a/Engine/Source/Programs/UnrealBuildAccelerator/Common/Private/UbaSession.cpp
+++ b/Engine/Source/Programs/UnrealBuildAccelerator/Common/Private/UbaSession.cpp
@@ -1434,6 +1434,12 @@ namespace uba
 	{
 		return m_tempPath.data;
 	}
+	
+	const tchar* Session::GetRootDir()
+	{
+		return m_rootDir.data;
+	}
+
 	u32 Session::CreateProcessId()
 	{
 		return ++m_processIdCounter;
diff --git a/Engine/Source/Programs/UnrealBuildAccelerator/Common/Private/UbaStorage.cpp b/Engine/Source/Programs/UnrealBuildAccelerator/Common/Private/UbaStorage.cpp
--- a/Engine/Source/Programs/UnrealBuildAccelerator/Common/Private/UbaStorage.cpp
+++ b/Engine/Source/Programs/UnrealBuildAccelerator/Common/Private/UbaStorage.cpp
@@ -1354,20 +1354,26 @@ namespace uba
 		for (u32 i=0; i!=fileTableSize; ++i)
 		{
 			StringKey fileNameKey = reader.ReadStringKey();
-			auto insres = m_fileTableLookup.try_emplace(fileNameKey);
-			FileEntry& fileEntry = insres.first->second;
-			fileEntry.verified = false;
 			if (reader.GetPosition() + 24 > fileSize)
 			{
 				m_fileTableLookup.clear();
 				m_logger.Warning(TC("CasTable file %s is corrupt"), fileName.data);
 				return false;
 			}
-			fileEntry.size = reader.ReadU64();
-			fileEntry.lastWritten = reader.ReadU64();
-			CasKey key = reader.ReadCasKey();
-			if (key != CasKeyZero)
-				fileEntry.casKey = AsCompressed(key, m_storeCompressed);
+			u64 size = reader.ReadU64();
+			u64 lastWritten = reader.ReadU64();
+			CasKey casKey = reader.ReadCasKey();
+			if (casKey != CasKeyZero)
+				casKey = AsCompressed(casKey, m_storeCompressed);
+			else if (size || lastWritten)
+				continue; // This should not happen
+
+			auto insres = m_fileTableLookup.try_emplace(fileNameKey);
+			FileEntry& fileEntry = insres.first->second;
+			fileEntry.verified = false;
+			fileEntry.size = size;
+			fileEntry.lastWritten = lastWritten;
+			fileEntry.casKey = casKey;
 		}
 
 
@@ -1512,100 +1518,91 @@ namespace uba
 			SCOPED_READ_LOCK(m_casLookupLock, casLookupLock);
 			SCOPED_READ_LOCK(m_accessLock, accessLock);
 
-			u8 buffer[1024];
-			{
-				BinaryWriter writer(buffer);
-				writer.WriteU32(CasTableVersion);
-				writer.WriteU32(u32(m_fileTableLookup.size()));
-				writer.WriteU32(u32(m_casLookup.size()));
-				if (!tempFile.Write(buffer, writer.GetPosition()))
-					return false;
-			}
+			Vector<u8> buffer;
 
-			{
-				constexpr u64 entrySize = sizeof(StringKey) + sizeof(u64) * 2 + sizeof(CasKey);
-				u64 fileTableWriteSize = m_fileTableLookup.size() * entrySize;
-				Vector<u8> fileTableData;
-				fileTableData.resize(fileTableWriteSize);
-				BinaryWriter writer(fileTableData.data(), 0, fileTableData.size());
-				for (auto& pair : m_fileTableLookup)
-				{
-					FileEntry& fileEntry = pair.second;
-					SCOPED_READ_LOCK(fileEntry.lock, entryLock);
-					writer.WriteStringKey(pair.first);
-					writer.WriteU64(fileEntry.size);
-					writer.WriteU64(fileEntry.lastWritten);
-					writer.WriteCasKey(fileEntry.casKey);
-				}
-				if (!tempFile.Write(fileTableData.data(), writer.GetPosition()))
-					return false;
-			}
+			u64 casLookupEntrySize = sizeof(CasKey) + sizeof(u64);
 
-#if UBA_USE_SPARSEFILE
+			u64 headerSize = sizeof(u32)*3;
+			u64 fileTableMaxWriteSize = m_fileTableLookup.size() * (sizeof(StringKey) + sizeof(u64) * 2 + sizeof(CasKey));
+			u64 casLookupMaxWriteSize = m_casLookup.size() * casLookupEntrySize + sizeof(CasKey);
+
+			buffer.reserve(headerSize + fileTableMaxWriteSize + casLookupMaxWriteSize);
+			BinaryWriter writer(buffer.data(), 0, buffer.capacity());
+
+			// Header
+			writer.WriteU32(CasTableVersion);
+			auto fileTableSizePtr = (u32*)writer.AllocWrite(4);
+			auto casLookupSizePtr = (u32*)writer.AllocWrite(4);
+
+			// File table
+			u32 fileTableSize = 0;
+			for (auto& pair : m_fileTableLookup)
 			{
-				BinaryWriter writer(buffer);
-				for (u32 i=0; i!=CasDbDataFileCount; ++i)
-					writer.Write7BitEncoded(m_casDataBuffer.GetPersistentSize(i));
-				WriteFile(m_logger, tempFileName.data, fileHandle, buffer, writer.GetPosition());
+				FileEntry& fileEntry = pair.second;
+				SCOPED_READ_LOCK(fileEntry.lock, entryLock);
+				if (fileEntry.casKey == CasKeyZero)
+					continue;
+				writer.WriteStringKey(pair.first);
+				writer.WriteU64(fileEntry.size);
+				writer.WriteU64(fileEntry.lastWritten);
+				writer.WriteCasKey(fileEntry.casKey);
+				++fileTableSize;
 			}
-#endif
+			*fileTableSizePtr = fileTableSize;
 
+			// Cas table
+			u32 casTableSize = 0;
+			CasEntry* last = nullptr;
+			for (CasEntry* it = m_newestAccessed; it; it = it->nextAccessed)
 			{
-				constexpr u64 entrySize = sizeof(CasKey) + sizeof(u64); // TODO: Wrong for sparse file
-				u64 casLookupWriteSize = m_casLookup.size() * entrySize;
-				Vector<u8> casLookupBuffer;
-				casLookupBuffer.resize(casLookupWriteSize + sizeof(CasKey)); // Add terminator
-				BinaryWriter writer(casLookupBuffer.data(), 0, casLookupBuffer.size());
-
-				CasEntry* last = nullptr;
-				for (CasEntry* it = m_newestAccessed; it; it = it->nextAccessed)
+				last = it;
+				CasEntry& casEntry = *it;
+				if (casEntry.verified && !casEntry.exists)
+					continue;
+				if (casEntry.dropped)
 				{
-					last = it;
-					CasEntry& casEntry = *it;
-					if (casEntry.verified && !casEntry.exists)
-						continue;
-					if (casEntry.dropped)
-					{
 #if !UBA_USE_SPARSEFILE
-						StringBuffer<512> casFileName;
-						if (!StorageImpl::GetCasFileName(casFileName, casEntry.key))
-							continue;
-						DeleteFileW(casFileName.data);
+					StringBuffer<512> casFileName;
+					if (!StorageImpl::GetCasFileName(casFileName, casEntry.key))
+						continue;
+					DeleteFileW(casFileName.data);
 #else
-						// TODO!  UBA_ASSERT(false);
+					// TODO!  UBA_ASSERT(false);
 #endif
-						continue;
-					}
+					continue;
+				}
 
 #if UBA_USE_SPARSEFILE
-					auto findIt = handleToIndex.find(casEntry.mappingHandle);
-					if (findIt == handleToIndex.end())
-					{
-						if (m_deferredCasCreationLookup.find(casEntry.key) != m_deferredCasCreationLookup.end())
-							continue;
-						m_logger.Error(TC("Can't find cas database file with mappingHandle %llu"), uintptr_t(casEntry.mappingHandle));
+				auto findIt = handleToIndex.find(casEntry.mappingHandle);
+				if (findIt == handleToIndex.end())
+				{
+					if (m_deferredCasCreationLookup.find(casEntry.key) != m_deferredCasCreationLookup.end())
 						continue;
-					}
+					m_logger.Error(TC("Can't find cas database file with mappingHandle %llu"), uintptr_t(casEntry.mappingHandle));
+					continue;
+				}
 #endif
 
-					if (writer.GetCapacityLeft() < entrySize + sizeof(CasKey))
-						return m_logger.Error(TC("This should not happen, somehow there are more valid entries in access list than lookup. (Lookup has %llu entries)"), m_casLookup.size());
+				if (writer.GetCapacityLeft() < casLookupEntrySize + sizeof(CasKey))
+					return m_logger.Error(TC("This should not happen, somehow there are more valid entries in access list than lookup. (Lookup has %llu entries)"), m_casLookup.size());
 
-					UBA_ASSERT(casEntry.key != CasKeyZero);
-					writer.WriteCasKey(casEntry.key);
-					writer.WriteU64(casEntry.size);
+				UBA_ASSERT(casEntry.key != CasKeyZero);
+				writer.WriteCasKey(casEntry.key);
+				writer.WriteU64(casEntry.size);
+				++casTableSize;
 
 #if UBA_USE_SPARSEFILE
-					writer.WriteU32(findIt->second);
-					writer.WriteU64(casEntry.mappingOffset);
-					writer.WriteU64(casEntry.mappingSize);
+				writer.WriteU32(findIt->second);
+				writer.WriteU64(casEntry.mappingOffset);
+				writer.WriteU64(casEntry.mappingSize);
 #endif
-				}
-				writer.WriteCasKey(CasKeyZero);
-				if (!tempFile.Write(casLookupBuffer.data(), writer.GetPosition()))
-					return false;
-				UBA_ASSERT(m_oldestAccessed == last); (void)last;
 			}
+			writer.WriteCasKey(CasKeyZero);
+			*casLookupSizePtr = casTableSize;
+
+			if (!tempFile.Write(buffer.data(), writer.GetPosition()))
+				return false;
+			UBA_ASSERT(m_oldestAccessed == last); (void)last;
 			if (!tempFile.Close())
 				return false;
 		}
@@ -1777,7 +1774,7 @@ namespace uba
 							if (casKey != fe.casKey)
 							{
 								++errorCount;
-								m_logger.Error(TC("CasKey mismatch for %s even though size and lastwritten were the same. Corrupt path table!"), p.data());
+								m_logger.Error(TC("CasKey mismatch for %s even though size and lastwritten were the same. Corrupt path table! (Correct: %s. Wrong: %s)"), p.data(), CasKeyString(casKey).str, CasKeyString(fe.casKey).str);
 							}
 
 						}, 1, TC(""));
diff --git a/Engine/Source/Programs/UnrealBuildAccelerator/Common/Private/UbaStorageClient.cpp b/Engine/Source/Programs/UnrealBuildAccelerator/Common/Private/UbaStorageClient.cpp
--- a/Engine/Source/Programs/UnrealBuildAccelerator/Common/Private/UbaStorageClient.cpp
+++ b/Engine/Source/Programs/UnrealBuildAccelerator/Common/Private/UbaStorageClient.cpp
@@ -1,6 +1,7 @@
 // Copyright Epic Games, Inc. All Rights Reserved.
 
 #include "UbaStorageClient.h"
+#include "UbaConfig.h"
 #include "UbaDirectoryIterator.h"
 #include "UbaFileAccessor.h"
 #include "UbaNetworkClient.h"
@@ -10,10 +11,23 @@
 
 namespace uba
 {
+	void StorageClientCreateInfo::Apply(Config& config)
+	{
+		StorageCreateInfo::Apply(config);
+
+		const ConfigTable* tablePtr = config.GetTable(TC("Storage"));
+		if (!tablePtr)
+			return;
+		const ConfigTable& table = *tablePtr;
+		table.GetValueAsBool(sendCompressed, TC("SendCompressed"));
+		table.GetValueAsBool(allowProxy, TC("AllowProxy"));
+	}
+
 	StorageClient::StorageClient(const StorageClientCreateInfo& info)
 	:	StorageImpl(info, TC("UbaStorageClient"))
 	,	m_client(info.client)
 	,	m_sendCompressed(info.sendCompressed)
+	,	m_allowProxy(info.allowProxy)
 	,	m_zone(info.zone)
 	,	m_getProxyBackendCallback(info.getProxyBackendCallback)
 	,	m_getProxyBackendUserData(info.getProxyBackendUserData)
@@ -247,7 +261,7 @@ namespace uba
 			ProxyClient* proxy = nullptr;
 
 			bool wantsProxy = false;
-			if (allowProxy)
+			if (allowProxy && m_allowProxy)
 			{
 				SCOPED_WRITE_LOCK(m_proxyClientLock, proxyLock);
 				if (m_proxyClient)
diff --git a/Engine/Source/Programs/UnrealBuildAccelerator/Common/Public/UbaSession.h b/Engine/Source/Programs/UnrealBuildAccelerator/Common/Public/UbaSession.h
--- a/Engine/Source/Programs/UnrealBuildAccelerator/Common/Public/UbaSession.h
+++ b/Engine/Source/Programs/UnrealBuildAccelerator/Common/Public/UbaSession.h
@@ -76,6 +76,7 @@ namespace uba
 		Trace& GetTrace(); // Trace written to be session
 		const ApplicationRules* GetRules(const ProcessStartInfo& si); // Get application rules used for process
 		const tchar* GetTempPath(); // Path for temp files used for current session
+		const tchar* GetRootDir(); // Root dir for session files
 		bool ShouldStoreObjFilesCompressed() { return m_storeObjFilesCompressed; }
 
 		virtual ~Session();
diff --git a/Engine/Source/Programs/UnrealBuildAccelerator/Common/Public/UbaStorageClient.h b/Engine/Source/Programs/UnrealBuildAccelerator/Common/Public/UbaStorageClient.h
--- a/Engine/Source/Programs/UnrealBuildAccelerator/Common/Public/UbaStorageClient.h
+++ b/Engine/Source/Programs/UnrealBuildAccelerator/Common/Public/UbaStorageClient.h
@@ -15,10 +15,14 @@ namespace uba
 	struct StorageClientCreateInfo : StorageCreateInfo
 	{
 		StorageClientCreateInfo(NetworkClient& c, const tchar* rootDir_) : StorageCreateInfo(rootDir_, c.GetLogWriter()), client(c) {}
+
+		void Apply(Config& config);
+
 		NetworkClient& client;
 		const tchar* zone = TC("");
 		u16 proxyPort = DefaultStorageProxyPort;
 		bool sendCompressed = true;
+		bool allowProxy = true;
 		
 		GetProxyBackendCallback* getProxyBackendCallback = nullptr;
 		void* getProxyBackendUserData = nullptr;
@@ -61,6 +65,7 @@ namespace uba
 
 		NetworkClient& m_client;
 		bool m_sendCompressed;
+		bool m_allowProxy;
 
 		Guid m_storageServerUid;
 
diff --git a/Engine/Source/Programs/UnrealBuildAccelerator/Core/Public/UbaHash.h b/Engine/Source/Programs/UnrealBuildAccelerator/Core/Public/UbaHash.h
--- a/Engine/Source/Programs/UnrealBuildAccelerator/Core/Public/UbaHash.h
+++ b/Engine/Source/Programs/UnrealBuildAccelerator/Core/Public/UbaHash.h
@@ -139,6 +139,7 @@ namespace uba
 
 	inline CasKey CasKeyFromString(const tchar* str)
 	{
+		UBA_ASSERT(TStrlen(str) == 40);
 		CasKey key;
 		u8* data = (u8*)&key;
 		const tchar* pos = str;
