UbaVisualizer: Add -help arg

#rnx

diff --git a/Engine/Source/Programs/UnrealBuildAccelerator/Visualizer/Private/UbaVisualizerApp.cpp b/Engine/Source/Programs/UnrealBuildAccelerator/Visualizer/Private/UbaVisualizerApp.cpp
--- a/Engine/Source/Programs/UnrealBuildAccelerator/Visualizer/Private/UbaVisualizerApp.cpp
+++ b/Engine/Source/Programs/UnrealBuildAccelerator/Visualizer/Private/UbaVisualizerApp.cpp
@@ -6,10 +6,10 @@
 
 using namespace uba;
 
-int PrintHelp(const tchar* message)
+static int PrintHelp(const tchar* message = nullptr)
 {
 	StringBuffer<64*1024> s;
-	if (*message)
+	if (message && *message)
 		s.Appendf(TC("%s\r\n\r\n"), message);
 
 	s.Appendf(TC("\r\n"));
@@ -94,8 +94,12 @@ int WINAPI wWinMain(_In_ HINSTANCE hInstance, _In_opt_ HINSTANCE hPrevInstance,
 		{
 			name.Append(argv[i]);
 		}
-		
-		if (name.Equals(TC("-host")))
+
+		if (name.Equals(TC("-help")))
+		{
+			return PrintHelp();
+		}
+		else if (name.Equals(TC("-host")))
 		{
 			if (value.IsEmpty())
 				return PrintHelp(TC("-host needs a value"));
