UbaVisualizer: Remove Save Trace option if a trace file was loaded

#rnx

diff --git a/Engine/Source/Programs/UnrealBuildAccelerator/Visualizer/Private/UbaVisualizer.cpp b/Engine/Source/Programs/UnrealBuildAccelerator/Visualizer/Private/UbaVisualizer.cpp
--- a/Engine/Source/Programs/UnrealBuildAccelerator/Visualizer/Private/UbaVisualizer.cpp
+++ b/Engine/Source/Programs/UnrealBuildAccelerator/Visualizer/Private/UbaVisualizer.cpp
@@ -2682,7 +2682,9 @@ namespace uba
 				}
 			}
 
-			AppendMenuW(hMenu, MF_STRING, Popup_SaveAs, L"&Save Trace");
+			if (m_fileName.IsEmpty()) {
+				AppendMenuW(hMenu, MF_STRING, Popup_SaveAs, L"&Save Trace");
+			}
 			AppendMenuW(hMenu, MF_SEPARATOR, 0, NULL);
 			AppendMenuW(hMenu, MF_STRING, Popup_Quit, L"&Quit");
 			m_showPopup = true;
