[UBA]
* Moved m_isConnected event Set call to inside write lock for connected functions to make sure connected functions are only called once
* Added delete of copy ctor and assignment

diff --git a/Engine/Source/Programs/UnrealBuildAccelerator/Common/Private/UbaNetworkClient.cpp b/Engine/Source/Programs/UnrealBuildAccelerator/Common/Private/UbaNetworkClient.cpp
--- a/Engine/Source/Programs/UnrealBuildAccelerator/Common/Private/UbaNetworkClient.cpp
+++ b/Engine/Source/Programs/UnrealBuildAccelerator/Common/Private/UbaNetworkClient.cpp
@@ -183,9 +183,9 @@ namespace uba
 		SCOPED_WRITE_LOCK(m_onConnectedFunctionsLock, lock);
 		for (auto& f : m_onConnectedFunctions)
 			f();
+		m_isConnected.Set();
 		lock.Leave();
 
-		m_isConnected.Set();
 		return true;
 	}
 
@@ -358,7 +358,7 @@ namespace uba
 	{
 		SCOPED_WRITE_LOCK(m_onConnectedFunctionsLock, lock);
 		m_onConnectedFunctions.push_back(function);
-		if (m_connectionCount.load() == 0)
+		if (!m_isConnected.IsSet(0))
 			return;
 		lock.Leave();
 		function();
diff --git a/Engine/Source/Programs/UnrealBuildAccelerator/Common/Public/UbaNetworkClient.h b/Engine/Source/Programs/UnrealBuildAccelerator/Common/Public/UbaNetworkClient.h
--- a/Engine/Source/Programs/UnrealBuildAccelerator/Common/Public/UbaNetworkClient.h
+++ b/Engine/Source/Programs/UnrealBuildAccelerator/Common/Public/UbaNetworkClient.h
@@ -127,5 +127,8 @@ namespace uba
 		Timer m_decryptTimer;
 
 		friend NetworkMessage;
+
+		NetworkClient(const NetworkClient&) = delete;
+		NetworkClient& operator=(const NetworkClient&) = delete;
 	};
 }
