[UBA]
* Added handling of \Global??\ paths (used by FilterConnectCommunicationPort and others)
* Added more logging on some failures

diff --git a/Engine/Source/Programs/UnrealBuildAccelerator/Detours/Private/Windows/UbaDetoursFunctionsNtDll.inl b/Engine/Source/Programs/UnrealBuildAccelerator/Detours/Private/Windows/UbaDetoursFunctionsNtDll.inl
--- a/Engine/Source/Programs/UnrealBuildAccelerator/Detours/Private/Windows/UbaDetoursFunctionsNtDll.inl
+++ b/Engine/Source/Programs/UnrealBuildAccelerator/Detours/Private/Windows/UbaDetoursFunctionsNtDll.inl
@@ -590,7 +590,7 @@ NTSTATUS NTAPI Shared_NtCreateFile(bool IsCreateFunc, PHANDLE hFileHandle, ACCES
 	HANDLE rootDir = ObjectAttributes->RootDirectory;
 	{
 		const wchar_t* buf = ObjectAttributes->ObjectName->Buffer;
-		u64 bufBytes = ObjectAttributes->ObjectName->Length;
+		u32 bufChars = ObjectAttributes->ObjectName->Length / 2;
 
 		if (suppressCreateFileDetour)
 		{
@@ -599,7 +599,7 @@ NTSTATUS NTAPI Shared_NtCreateFile(bool IsCreateFunc, PHANDLE hFileHandle, ACCES
 		{
 			suppressCreateFileDetour = true;
 		}
-		else if (wcsncmp(buf, L"\\Device", 7) == 0)
+		else if ((bufChars >= 7u && wcsncmp(buf, L"\\Device", 7u)) == 0 || (bufChars >= 10u && wcsncmp(buf, L"\\Global??\\", 10u)) == 0) // \Global is for FilterConnectCommunicationPort and friends
 		{
 			suppressCreateFileDetour = true;
 		}
@@ -627,14 +627,14 @@ NTSTATUS NTAPI Shared_NtCreateFile(bool IsCreateFunc, PHANDLE hFileHandle, ACCES
 				if (isDetouredHandle(ObjectAttributes->RootDirectory))
 				{
 					auto& dh = asDetouredHandle(ObjectAttributes->RootDirectory);
-					fileName.Append(dh.fileObject->fileInfo->originalName).EnsureEndsWithSlash().Append(buf, bufBytes / 2);
+					fileName.Append(dh.fileObject->fileInfo->originalName).EnsureEndsWithSlash().Append(buf, bufChars);
 					rootDir = dh.trueHandle;
 					ObjectAttributes->RootDirectory = nullptr;
 				}
 				else if (isListDirectoryHandle(ObjectAttributes->RootDirectory))
 				{
 					auto& lh = asListDirectoryHandle(ObjectAttributes->RootDirectory);
-					fileName.Append(lh.originalName).EnsureEndsWithSlash().Append(buf, bufBytes / 2);
+					fileName.Append(lh.originalName).EnsureEndsWithSlash().Append(buf, bufChars);
 					ObjectAttributes->RootDirectory = nullptr;
 				}
 				else
@@ -646,7 +646,7 @@ NTSTATUS NTAPI Shared_NtCreateFile(bool IsCreateFunc, PHANDLE hFileHandle, ACCES
 			}
 			else
 			{
-				fileName.Append(buf, bufBytes / 2);
+				fileName.Append(buf, bufChars);
 				if (fileName.StartsWith(L"\\DosDevices")) // Something used in msbuild.. 
 					suppressCreateFileDetour = true;
 			}
@@ -1018,7 +1018,7 @@ NTSTATUS NTAPI Shared_NtCreateFile(bool IsCreateFunc, PHANDLE hFileHandle, ACCES
 		if (!handleStrEnd)
 		{
 			UBA_ASSERT(handleStrEnd);
-			return STATUS_OBJECT_NAME_NOT_FOUND;
+			return STATUS_UNSUCCESSFUL;
 		}
 		HANDLE mappingHandle = (HANDLE)StringToValue(handleStr, handleStrEnd - handleStr);
 		const wchar_t* mappingOffsetStr = handleStrEnd + 1;
@@ -1026,8 +1026,12 @@ NTSTATUS NTAPI Shared_NtCreateFile(bool IsCreateFunc, PHANDLE hFileHandle, ACCES
 		info.trueFileMapOffset = mappingOffset;
 
 		info.isFileMap = true;
-		True_DuplicateHandle(g_hostProcess, mappingHandle, GetCurrentProcess(), &info.trueFileMapHandle, 0, FALSE, DUPLICATE_SAME_ACCESS);
-		UBA_ASSERTF(info.trueFileMapHandle, L"Can't duplicate handle 0x%llx (%ls) for file %ls", uintptr_t(mappingHandle), lpFileName, info.originalName);
+		if (!True_DuplicateHandle(g_hostProcess, mappingHandle, GetCurrentProcess(), &info.trueFileMapHandle, 0, FALSE, DUPLICATE_SAME_ACCESS))
+		{
+			Rpc_WriteLogf(L"Can't duplicate handle 0x%llx (%ls) for file %ls", uintptr_t(mappingHandle), lpFileName, info.originalName);
+			UBA_ASSERTF(info.trueFileMapHandle, L"Can't duplicate handle 0x%llx (%ls) for file %ls", uintptr_t(mappingHandle), lpFileName, info.originalName);
+			return STATUS_UNSUCCESSFUL;
+		}
 		DetouredHandle* dh = new DetouredHandle(HandleType_File);
 		UBA_ASSERT(info.size != InvalidValue);
 		dh->dirTableOffset = dirTableOffset;
@@ -1096,8 +1100,8 @@ NTSTATUS NTAPI Shared_NtCreateFile(bool IsCreateFunc, PHANDLE hFileHandle, ACCES
 				else if (CreateDisposition == FILE_OPEN)
 				{
 					*hFileHandle = INVALID_HANDLE_VALUE;
-					DEBUG_LOG_DETOURED(funcName, L"NOTFOUND (%ls) -> Error", fileName.data);
-					return STATUS_OBJECT_NAME_NOT_FOUND;
+					DEBUG_LOG_DETOURED(funcName, L"ALREADYEXISTS (%ls) -> Error", fileName.data);
+					return STATUS_OBJECT_NAME_EXISTS;//STATUS_OBJECT_NAME_NOT_FOUND;
 				}
 
 				bool isLocal = !IsOutputFile(fileName, isWrite, isDeleteOnClose);
@@ -1179,7 +1183,7 @@ NTSTATUS NTAPI Shared_NtCreateFile(bool IsCreateFunc, PHANDLE hFileHandle, ACCES
 			info.lastDesiredAccess = lastDesiredAccess;
 			Rpc_UpdateCloseHandle(L"", closeId, false, L"", 0, 0, false);
 		}
-		DEBUG_LOG_TRUE(funcName, L"%ls (%ls) (%ls) -> %ls", isWriteStr, lpFileName, (fileName.data != lpFileName ? fileName.data : L""), ToString(res));
+		DEBUG_LOG_TRUE(funcName, L"%ls (%ls) (%ls) -> %ls (0x%x)", isWriteStr, lpFileName, (fileName.data != lpFileName ? fileName.data : L""), ToString(res), res);
 		return res;
 	}
 	if (!canDetour)
