[UBA]
* Fixed typo in previous submit where code handles \Global??\ paths

diff --git a/Engine/Source/Programs/UnrealBuildAccelerator/Detours/Private/Windows/UbaDetoursFunctionsNtDll.inl b/Engine/Source/Programs/UnrealBuildAccelerator/Detours/Private/Windows/UbaDetoursFunctionsNtDll.inl
--- a/Engine/Source/Programs/UnrealBuildAccelerator/Detours/Private/Windows/UbaDetoursFunctionsNtDll.inl
+++ b/Engine/Source/Programs/UnrealBuildAccelerator/Detours/Private/Windows/UbaDetoursFunctionsNtDll.inl
@@ -599,7 +599,7 @@ NTSTATUS NTAPI Shared_NtCreateFile(bool IsCreateFunc, PHANDLE hFileHandle, ACCES
 		{
 			suppressCreateFileDetour = true;
 		}
-		else if ((bufChars >= 7u && wcsncmp(buf, L"\\Device", 7u)) == 0 || (bufChars >= 10u && wcsncmp(buf, L"\\Global??\\", 10u)) == 0) // \Global is for FilterConnectCommunicationPort and friends
+		else if ((bufChars >= 7u && wcsncmp(buf, L"\\Device", 7u) == 0) || (bufChars >= 10u && wcsncmp(buf, L"\\Global??\\", 10u) == 0)) // \Global is for FilterConnectCommunicationPort and friends
 		{
 			suppressCreateFileDetour = true;
 		}
