[UBA]
* Added more details to logging when uba cache client fails to connect

diff --git a/Engine/Source/Programs/UnrealBuildAccelerator/Common/Private/UbaNetworkClient.cpp b/Engine/Source/Programs/UnrealBuildAccelerator/Common/Private/UbaNetworkClient.cpp
--- a/Engine/Source/Programs/UnrealBuildAccelerator/Common/Private/UbaNetworkClient.cpp
+++ b/Engine/Source/Programs/UnrealBuildAccelerator/Common/Private/UbaNetworkClient.cpp
@@ -441,7 +441,12 @@ namespace uba
 		SCOPED_WRITE_LOCK(m_connectionsItLock, connectionItLock);
 		if (m_connectionsIt == m_connections.end())
 		{
-			message.m_error = 6;
+			if (m_isDisconnecting)
+				message.m_error = 11;
+			else if (!m_connections.empty())
+				message.m_error = 12; // should never happen
+			else
+				message.m_error = 6;
 			return false;
 		}
 
