[UBA]
* Started adding support for mspdbsrv.exe for remote execution. Disabled for now

diff --git a/Engine/Source/Programs/UnrealBuildAccelerator/Common/Private/UbaProcess.cpp b/Engine/Source/Programs/UnrealBuildAccelerator/Common/Private/UbaProcess.cpp
--- a/Engine/Source/Programs/UnrealBuildAccelerator/Common/Private/UbaProcess.cpp
+++ b/Engine/Source/Programs/UnrealBuildAccelerator/Common/Private/UbaProcess.cpp
@@ -1180,6 +1180,12 @@ namespace uba
 		return true;
 	}
 
+	bool ProcessImpl::HandleRpcCommunication(BinaryReader& reader, BinaryWriter& writer)
+	{
+		//m_session.RpcCommunication(*this, reader, writer);
+		return true;
+	}
+
 	bool ProcessImpl::CreateTempFile(BinaryReader& reader, ProcHandle nativeProcessHandle, const tchar* application)
 	{
 		StringKey key = reader.ReadStringKey();
diff --git a/Engine/Source/Programs/UnrealBuildAccelerator/Common/Public/UbaProcess.h b/Engine/Source/Programs/UnrealBuildAccelerator/Common/Public/UbaProcess.h
--- a/Engine/Source/Programs/UnrealBuildAccelerator/Common/Public/UbaProcess.h
+++ b/Engine/Source/Programs/UnrealBuildAccelerator/Common/Public/UbaProcess.h
@@ -100,6 +100,7 @@ namespace uba
 		bool HandleGetNextProcess(BinaryReader& reader, BinaryWriter& writer);
 		bool HandleCustom(BinaryReader& reader, BinaryWriter& writer);
 		bool HandleSHGetKnownFolderPath(BinaryReader& reader, BinaryWriter& writer);
+		bool HandleRpcCommunication(BinaryReader& reader, BinaryWriter& writer);
 
 		void LogLine(bool printInSession, TString&& line, LogEntryType logType);
 		bool CreateTempFile(BinaryReader& reader, ProcHandle nativeProcessHandle, const tchar* application);
diff --git a/Engine/Source/Programs/UnrealBuildAccelerator/Core/Public/UbaProtocol.h b/Engine/Source/Programs/UnrealBuildAccelerator/Core/Public/UbaProtocol.h
--- a/Engine/Source/Programs/UnrealBuildAccelerator/Core/Public/UbaProtocol.h
+++ b/Engine/Source/Programs/UnrealBuildAccelerator/Core/Public/UbaProtocol.h
@@ -35,6 +35,7 @@ namespace uba
 		UBA_PROCESS_MESSAGE(GetNextProcess) \
 		UBA_PROCESS_MESSAGE(Custom) \
 		UBA_PROCESS_MESSAGE(SHGetKnownFolderPath) \
+		UBA_PROCESS_MESSAGE(RpcCommunication) \
 
 	enum MessageType : u8
 	{
diff --git a/Engine/Source/Programs/UnrealBuildAccelerator/Detours/Lib/BuildAsm.bat b/Engine/Source/Programs/UnrealBuildAccelerator/Detours/Lib/BuildAsm.bat
new file mode 100755
--- /dev/null
+++ b/Engine/Source/Programs/UnrealBuildAccelerator/Detours/Lib/BuildAsm.bat
@@ -0,0 +1 @@
+ml64.exe /FoUbaAsmX64.obj /c /Zf UbaAsmX64.asm
diff --git a/Engine/Source/Programs/UnrealBuildAccelerator/Detours/Private/Windows/UbaDetoursFunctionsNtDll.inl b/Engine/Source/Programs/UnrealBuildAccelerator/Detours/Private/Windows/UbaDetoursFunctionsNtDll.inl
--- a/Engine/Source/Programs/UnrealBuildAccelerator/Detours/Private/Windows/UbaDetoursFunctionsNtDll.inl
+++ b/Engine/Source/Programs/UnrealBuildAccelerator/Detours/Private/Windows/UbaDetoursFunctionsNtDll.inl
@@ -1514,6 +1514,39 @@ NTSTATUS NTAPI Detoured_NtReadFile(HANDLE FileHandle, HANDLE Event, PIO_APC_ROUT
 	return True_NtReadFile(FileHandle, Event, ApcRoutine, ApcContext, IoStatusBlock, Buffer, Length, ByteOffset, Key);
 }
 
+NTSTATUS NTAPI Detoured_NtAlpcCreatePort(PHANDLE PortHandle, POBJECT_ATTRIBUTES ObjectAttributes, PALPC_PORT_ATTRIBUTES PortAttributes)
+{
+	//DEBUG_LOG_TRUE(L"NtAlpcCreatePort", L"");
+	return True_NtAlpcCreatePort(PortHandle, ObjectAttributes, PortAttributes);
+}
+
+NTSTATUS NTAPI Detoured_NtAlpcConnectPort(PHANDLE PortHandle, PUNICODE_STRING PortName, POBJECT_ATTRIBUTES ObjectAttributes, PALPC_PORT_ATTRIBUTES PortAttributes, DWORD ConnectionFlags, PSID RequiredServerSid, PPORT_MESSAGE ConnectionMessage, PSIZE_T ConnectMessageSize, PALPC_MESSAGE_ATTRIBUTES OutMessageAttributes, PALPC_MESSAGE_ATTRIBUTES InMessageAttributes, PLARGE_INTEGER Timeout)
+{
+	//DEBUG_LOG_TRUE(L"NtAlpcConnectPort", L"");
+	return True_NtAlpcConnectPort(PortHandle, PortName, ObjectAttributes, PortAttributes, ConnectionFlags, RequiredServerSid, ConnectionMessage, ConnectMessageSize, OutMessageAttributes, InMessageAttributes, Timeout);
+}
+
+NTSTATUS NTAPI Detoured_NtAlpcCreatePortSection(HANDLE PortHandle, ULONG Flags, HANDLE SectionHandle, SIZE_T SectionSize, PHANDLE AlpcSectionHandle, PSIZE_T ActualSectionSize)
+{
+	//DEBUG_LOG_TRUE(L"NtAlpcCreatePortSection", L"");
+	return True_NtAlpcCreatePortSection(PortHandle, Flags, SectionHandle, SectionSize, AlpcSectionHandle, ActualSectionSize);
+}
+
+NTSTATUS NTAPI Detoured_NtAlpcSendWaitReceivePort(HANDLE PortHandle, DWORD Flags, PPORT_MESSAGE SendMessage_, PALPC_MESSAGE_ATTRIBUTES SendMessageAttributes, PPORT_MESSAGE ReceiveMessage, PSIZE_T BufferLength, PALPC_MESSAGE_ATTRIBUTES ReceiveMessageAttributes, PLARGE_INTEGER Timeout)
+{
+	//u64 size = 0;
+	//if (SendMessage_)
+	//	size = SendMessage_->u1.s1.DataLength;
+	//DEBUG_LOG_TRUE(L"NtAlpcSendWaitReceivePort", L"%llu", size);
+	return True_NtAlpcSendWaitReceivePort(PortHandle, Flags, SendMessage_, SendMessageAttributes, ReceiveMessage, BufferLength, ReceiveMessageAttributes, Timeout);
+}
+
+NTSTATUS NTAPI Detoured_NtAlpcDisconnectPort(HANDLE PortHandle, ULONG Flags)
+{
+	//DEBUG_LOG_TRUE(L"NtAlpcDisconnectPort", L"");
+	return True_NtAlpcDisconnectPort(PortHandle, Flags);
+}
+
 NTSTATUS NTAPI Detoured_ZwQueryDirectoryFile(HANDLE FileHandle, HANDLE Event, PIO_APC_ROUTINE ApcRoutine, PVOID ApcContext, PIO_STATUS_BLOCK IoStatusBlock, PVOID FileInformation, ULONG Length,
 	FILE_INFORMATION_CLASS FileInformationClass, BOOLEAN ReturnSingleEntry, PUNICODE_STRING FileName, BOOLEAN RestartScan)
 {
diff --git a/Engine/Source/Programs/UnrealBuildAccelerator/Detours/Private/Windows/UbaDetoursFunctionsRpcrt4.inl b/Engine/Source/Programs/UnrealBuildAccelerator/Detours/Private/Windows/UbaDetoursFunctionsRpcrt4.inl
new file mode 100644
--- /dev/null
+++ b/Engine/Source/Programs/UnrealBuildAccelerator/Detours/Private/Windows/UbaDetoursFunctionsRpcrt4.inl
@@ -0,0 +1,143 @@
+// Copyright Epic Games, Inc. All Rights Reserved.
+
+#if UBA_SUPPORT_MSPDBSRV
+
+RPC_STATUS Detoured_RpcStringBindingComposeW(RPC_WSTR ObjUuid, RPC_WSTR Protseq, RPC_WSTR NetworkAddr, RPC_WSTR Endpoint, RPC_WSTR Options, RPC_WSTR* StringBinding)
+{
+	RPC_STATUS res = True_RpcStringBindingComposeW(ObjUuid, Protseq, NetworkAddr, Endpoint, Options, StringBinding);
+
+	DEBUG_LOG_TRUE(L"RpcStringBindingComposeW", L"%s %s %s To %s", Protseq, NetworkAddr, Endpoint, *(wchar_t**)StringBinding);
+	return res;
+}
+
+UnorderedSet<void*> g_rpcBindings;
+
+RPC_STATUS Detoured_RpcBindingFromStringBindingW(RPC_WSTR StringBinding, RPC_BINDING_HANDLE *Binding)
+{
+	/*
+	if (g_runningRemote && Contains((wchar_t*)StringBinding, L"mspdb"))
+	{
+		SCOPED_WRITE_LOCK(g_communicationLock, pcs);
+		BinaryWriter writer;
+		writer.WriteByte(MessageType_RpcCommunication);
+		writer.WriteByte(1);
+		writer.Flush();
+		BinaryReader reader;
+		*Binding = new char();
+		g_rpcBindings.insert(*Binding);
+		return 0;
+	}
+	*/
+
+	auto res = True_RpcBindingFromStringBindingW(StringBinding, Binding);
+
+	g_rpcBindings.insert(*Binding);
+
+	DEBUG_LOG_TRUE(L"RpcBindingFromStringBindingW", L"%s TO %p", StringBinding, *Binding);
+	return res;
+}
+
+RPC_STATUS Detoured_RpcBindingSetAuthInfoExW(RPC_BINDING_HANDLE Binding, RPC_WSTR ServerPrincName, unsigned long AuthnLevel, unsigned long AuthnSvc, RPC_AUTH_IDENTITY_HANDLE AuthIdentity, unsigned long AuthzSvc, RPC_SECURITY_QOS *SecurityQOS)
+{
+	//if (g_rpcBindings.find(Binding) != g_rpcBindings.end())
+	//{
+	//	return 0;
+	//}
+
+	auto res = True_RpcBindingSetAuthInfoExW(Binding, ServerPrincName, AuthnLevel, AuthnSvc, AuthIdentity, AuthzSvc, SecurityQOS);
+	DEBUG_LOG_TRUE(L"RpcBindingSetAuthInfoExW", L"%p", Binding);
+	return res;
+}
+
+extern "C" 
+{
+	void* True2_NdrClientCall2 = nullptr;
+
+	ULONG_PTR __cdecl Detoured_NdrClientCall2(PMIDL_STUB_DESC pStubDescriptor, PFORMAT_STRING  pFormat, ...);
+
+	ULONG_PTR Local_NdrClientCall2(PMIDL_STUB_DESC pStubDescriptor, PFORMAT_STRING pFormat, void** pStack)
+	{
+		void* returnAddress = *(pStack - 3);(void)returnAddress;
+
+		//const WCHAR* CallingModule = FindModuleByAddress(returnAddress);
+		//if (!CallingModule)
+		//    CallingModule = L"unknown";
+		RPC_BINDING_HANDLE* Binding = nullptr;
+
+		unsigned char handleType = 255;
+		if (pFormat) {
+			handleType = pFormat[0];
+			if (!handleType) { // explicit handle
+				// pStack[0] explicit handle
+				// pStack[1] Arg0
+				// pStack[2] Arg1
+				// ...
+				// pStack[n+1] ArgN
+				Binding = (RPC_BINDING_HANDLE*)&pStack[0]; // first argument in the var arg list is the explicit handle
+			}
+			else { // implicit handle
+				// pStack[0] Arg0
+				// pStack[1] Arg1
+				// ...
+				// pStack[n] ArgN
+				Binding = (RPC_BINDING_HANDLE*)&pStubDescriptor->IMPLICIT_HANDLE_INFO.pPrimitiveHandle;
+			}
+		}
+
+		if (Binding && *Binding == nullptr)
+			Binding = nullptr;
+
+		if (PRPC_CLIENT_INTERFACE rpcInterface = (PRPC_CLIENT_INTERFACE)pStubDescriptor->RpcInterfaceInformation)
+		{
+			//WCHAR interfaceID[48];
+			//Sbie_StringFromGUID(&rpcInterface->InterfaceId.SyntaxGUID, interfaceID);
+
+			//Sbie_snwprintf(text, 512, L"Calling %s UUID = %s}, %d.%d, BindingHandle = 0x%X (%X), caller = '%s'", Function, interfaceID,
+			//    rpcInterface->InterfaceId.SyntaxVersion.MajorVersion, rpcInterface->InterfaceId.SyntaxVersion.MinorVersion, Binding ? *Binding : nullptr, handleType, CallingModule);
+		}
+		else
+		{
+			//Sbie_snwprintf(text, 512, L"Calling %s caller = '%s'", Function, CallingModule);
+		}
+
+		if (g_rpcBindings.find(*Binding) != g_rpcBindings.end())
+			printf("");
+		//SbieApi_MonitorPutMsg(MONITOR_RPC | MONITOR_TRACE, text);
+
+		DEBUG_LOG_TRUE(L"NdrClientCall2", L"%p", *Binding);
+		return 0;
+	}
+
+}
+/*
+ULONG_PTR __cdecl Detoured_NdrClientCall2(PMIDL_STUB_DESC pStubDescriptor, PFORMAT_STRING  pFormat, ...)
+{
+	RPC_BINDING_HANDLE* Binding = nullptr;
+
+	unsigned char handleType = 255;
+	if (pFormat) {
+		handleType = pFormat[0];
+		if (!handleType) { // explicit handle
+			// pStack[0] explicit handle
+			// pStack[1] Arg0
+			// pStack[2] Arg1
+			// ...
+			// pStack[n+1] ArgN
+			//Binding = (RPC_BINDING_HANDLE*)&pStack[0]; // first argument in the var arg list is the explicit handle
+		}
+		else { // implicit handle
+			// pStack[0] Arg0
+			// pStack[1] Arg1
+			// ...
+			// pStack[n] ArgN
+			Binding = (RPC_BINDING_HANDLE*)&pStubDescriptor->IMPLICIT_HANDLE_INFO.pPrimitiveHandle;
+		}
+	}
+
+	if (Binding && *Binding == nullptr)
+		Binding = nullptr;
+	DEBUG_LOG_TRUE(L"NdrClientCall2", L"%p", *Binding);
+	return 0;
+}
+*/
+#endif
diff --git a/Engine/Source/Programs/UnrealBuildAccelerator/Detours/Private/Windows/UbaDetoursFunctionsWin.cpp b/Engine/Source/Programs/UnrealBuildAccelerator/Detours/Private/Windows/UbaDetoursFunctionsWin.cpp
--- a/Engine/Source/Programs/UnrealBuildAccelerator/Detours/Private/Windows/UbaDetoursFunctionsWin.cpp
+++ b/Engine/Source/Programs/UnrealBuildAccelerator/Detours/Private/Windows/UbaDetoursFunctionsWin.cpp
@@ -474,6 +474,7 @@ thread_local const wchar_t* t_createFileFileName;
 #include "UbaDetoursFunctionsImagehlp.inl"
 #include "UbaDetoursFunctionsDbgHelp.inl"
 #include "UbaDetoursFunctionsShell32.inl"
+#include "UbaDetoursFunctionsRpcrt4.inl"
 
 extern u32 g_consoleStringIndex;
 
@@ -566,7 +567,14 @@ int DetourAttachFunctions(bool runningRemote)
 		DETOURED_FUNCTIONS_SHLWAPI
 	}
 
-	#undef DETOURED_FUNCTION
+	#if UBA_SUPPORT_MSPDBSRV
+	if (HMODULE moduleHandle = GetModuleHandleW(L"rpcrt4.dll"))
+	{
+		DETOURED_FUNCTIONS_RPCRT4
+	}
+	#endif
+
+#undef DETOURED_FUNCTION
 
 	// Can't attach to these when running through debugger with some vs extensions (Microsoft child process debugging)
 #if UBA_DEBUG
@@ -596,6 +604,10 @@ int DetourAttachFunctions(bool runningRemote)
 		ExitProcess(1343);
 	}
 
+	#if UBA_SUPPORT_MSPDBSRV
+	True2_NdrClientCall2 = True_NdrClientCall2;
+	#endif
+
 	return 0;
 }
 
@@ -935,7 +947,7 @@ void Deinit(u64 startTime)
 	#endif
 
 	DWORD exitCode = STILL_ACTIVE;
-	if (!GetExitCodeProcess(GetCurrentProcess(), &exitCode))
+	if (!True_GetExitCodeProcess(GetCurrentProcess(), &exitCode))
 		exitCode = STILL_ACTIVE;
 
 	if (!g_exitMessageSent)
diff --git a/Engine/Source/Programs/UnrealBuildAccelerator/Detours/Private/Windows/UbaDetoursFunctionsWin.h b/Engine/Source/Programs/UnrealBuildAccelerator/Detours/Private/Windows/UbaDetoursFunctionsWin.h
--- a/Engine/Source/Programs/UnrealBuildAccelerator/Detours/Private/Windows/UbaDetoursFunctionsWin.h
+++ b/Engine/Source/Programs/UnrealBuildAccelerator/Detours/Private/Windows/UbaDetoursFunctionsWin.h
@@ -26,6 +26,7 @@
 	DETOURED_FUNCTIONS_NTDLL \
 	DETOURED_FUNCTIONS_SHLWAPI \
 	DETOURED_FUNCTIONS_UCRTBASE \
+	DETOURED_FUNCTIONS_RPCRT4 \
 
 #define DETOURED_FUNCTIONS_KERNELBASE \
 	DETOURED_FUNCTION(GetCommandLineW) \
@@ -154,6 +155,17 @@
 	DETOURED_FUNCTION(_wsplitpath_s) \
 	DETOURED_FUNCTIONS_UCRTBASE_DEBUG \
 
+#if UBA_SUPPORT_MSPDBSRV
+#define DETOURED_FUNCTIONS_RPCRT4 \
+	DETOURED_FUNCTION(RpcStringBindingComposeW) \
+	DETOURED_FUNCTION(RpcBindingSetAuthInfoExW) \
+	DETOURED_FUNCTION(RpcBindingFromStringBindingW) \
+	DETOURED_FUNCTION(NdrClientCall2) \
+
+#else
+#define DETOURED_FUNCTIONS_RPCRT4
+#endif
+
 #if UBA_USE_MIMALLOC
 #define DETOURED_FUNCTIONS_MEMORY \
 	DETOURED_FUNCTION(malloc) \
@@ -293,6 +305,11 @@
 	DETOURED_FUNCTION(NtQueryFullAttributesFile) \
 	DETOURED_FUNCTION(NtFlushBuffersFileEx) \
 	DETOURED_FUNCTION(NtReadFile) \
+	DETOURED_FUNCTION(NtAlpcCreatePort) \
+	DETOURED_FUNCTION(NtAlpcConnectPort) \
+	DETOURED_FUNCTION(NtAlpcCreatePortSection) \
+	DETOURED_FUNCTION(NtAlpcSendWaitReceivePort) \
+	DETOURED_FUNCTION(NtAlpcDisconnectPort) \
 	DETOURED_FUNCTION(ZwSetInformationFile) \
 	DETOURED_FUNCTION(ZwQueryDirectoryFile) \
 	//DETOURED_FUNCTION(ZwCreateFile) \
@@ -334,6 +351,9 @@
 #endif
 
 extern "C" {
+	using PALPC_PORT_ATTRIBUTES = void*;
+	using PALPC_MESSAGE_ATTRIBUTES = void*;
+	using PPORT_MESSAGE = void*;
 	enum FS_INFORMATION_CLASS {};
 	NTSTATUS NTAPI NtQueryVolumeInformationFile(HANDLE FileHandle, PIO_STATUS_BLOCK IoStatusBlock, PVOID FsInformation, ULONG Length, FS_INFORMATION_CLASS FsInformationClass);
 	NTSTATUS NTAPI NtQueryFullAttributesFile(POBJECT_ATTRIBUTES ObjectAttributes, PVOID Attributes);
@@ -347,6 +367,11 @@ extern "C" {
 	NTSTATUS NTAPI NtSetInformationObject(HANDLE ObjectHandle, OBJECT_INFORMATION_CLASS ObjectInformationClass, PVOID ObjectInformation, ULONG Length);
 	NTSTATUS NTAPI NtCreateSection(PHANDLE SectionHandle, ACCESS_MASK DesiredAccess, POBJECT_ATTRIBUTES ObjectAttributes, PLARGE_INTEGER MaximumSize, ULONG SectionPageProtection, ULONG AllocationAttributes, HANDLE FileHandle);
 	NTSTATUS NTAPI NtCreateIoCompletion(PHANDLE IoCompletionHandle, ACCESS_MASK DesiredAccess, POBJECT_ATTRIBUTES ObjectAttributes, ULONG Count);
+	NTSTATUS NTAPI NtAlpcCreatePort(PHANDLE PortHandle, POBJECT_ATTRIBUTES ObjectAttributes, PALPC_PORT_ATTRIBUTES PortAttributes);
+	NTSTATUS NTAPI NtAlpcConnectPort(PHANDLE PortHandle, PUNICODE_STRING PortName, POBJECT_ATTRIBUTES ObjectAttributes, PALPC_PORT_ATTRIBUTES PortAttributes, DWORD ConnectionFlags, PSID RequiredServerSid, PPORT_MESSAGE ConnectionMessage, PSIZE_T ConnectMessageSize, PALPC_MESSAGE_ATTRIBUTES OutMessageAttributes, PALPC_MESSAGE_ATTRIBUTES InMessageAttributes, PLARGE_INTEGER Timeout);
+	NTSTATUS NTAPI NtAlpcCreatePortSection(HANDLE PortHandle, ULONG Flags, HANDLE SectionHandle, SIZE_T SectionSize, PHANDLE AlpcSectionHandle, PSIZE_T ActualSectionSize);
+	NTSTATUS NTAPI NtAlpcSendWaitReceivePort(HANDLE PortHandle, DWORD Flags, PPORT_MESSAGE SendMessage_, PALPC_MESSAGE_ATTRIBUTES SendMessageAttributes, PPORT_MESSAGE ReceiveMessage, PSIZE_T BufferLength, PALPC_MESSAGE_ATTRIBUTES ReceiveMessageAttributes, PLARGE_INTEGER Timeout);
+	NTSTATUS NTAPI NtAlpcDisconnectPort(HANDLE PortHandle, ULONG Flags);
 	NTSTATUS NTAPI ZwCreateFile(PHANDLE FileHandle, ACCESS_MASK DesiredAccess, POBJECT_ATTRIBUTES ObjectAttributes, PIO_STATUS_BLOCK IoStatusBlock, PLARGE_INTEGER AllocationSize, ULONG FileAttributes, ULONG ShareAccess, ULONG CreateDisposition, ULONG CreateOptions, PVOID EaBuffer, ULONG EaLength);
 	NTSTATUS NTAPI ZwOpenFile(PHANDLE FileHandle, ACCESS_MASK DesiredAccess, POBJECT_ATTRIBUTES ObjectAttributes, PIO_STATUS_BLOCK IoStatusBlock, ULONG ShareAccess, ULONG OpenOptions);
 	NTSTATUS NTAPI ZwQueryDirectoryFile(HANDLE FileHandle, HANDLE Event, PIO_APC_ROUTINE ApcRoutine, PVOID ApcContext, PIO_STATUS_BLOCK IoStatusBlock, PVOID FileInformation, ULONG Length, FILE_INFORMATION_CLASS FileInformationClass, BOOLEAN ReturnSingleEntry, PUNICODE_STRING FileName, BOOLEAN RestartScan);
diff --git a/Engine/Source/Programs/UnrealBuildAccelerator/Detours/UbaDetours.Build.cs b/Engine/Source/Programs/UnrealBuildAccelerator/Detours/UbaDetours.Build.cs
--- a/Engine/Source/Programs/UnrealBuildAccelerator/Detours/UbaDetours.Build.cs
+++ b/Engine/Source/Programs/UnrealBuildAccelerator/Detours/UbaDetours.Build.cs
@@ -1,5 +1,6 @@
 // Copyright Epic Games, Inc. All Rights Reserved.
 
+using System.IO;
 using UnrealBuildTool;
 
 [SupportedPlatforms(UnrealPlatformClass.Desktop)]
@@ -26,6 +27,19 @@ public class UbaDetours : ModuleRules
 				"ntdll.lib",
 				"onecore.lib"
 			});
+
+			// This is to handle mspdbsrv.exe ... not supporting arm atm
+			bool supportMspdbSrv = false; // Target.Architecture == UnrealArch.X64 && Target.Configuration == UnrealTargetConfiguration.Debug
+			
+			if (supportMspdbSrv)
+			{
+				PrivateDefinitions.Add("UBA_SUPPORT_MSPDBSRV=1");
+				PublicAdditionalLibraries.Add(Path.Combine(ModuleDirectory, "Lib", "UbaAsmX64.obj"));
+			}
+			else
+			{
+				PrivateDefinitions.Add("UBA_SUPPORT_MSPDBSRV=0");
+			}
 		}
 		else if (Target.Platform.IsInGroup(UnrealPlatformGroup.Linux))
 		{
