[UBA]
* Added so GetExitCodeProcess updates file/mapping tables when exit code process "succeeds" (as in process has exited)
* Added special handling of mspdbsrv.exe, vctip.exe and git.exe in CreateProcess. git.exe should not be there but is needed for now when building clang. mspdbsrv.exe is wip
* Increased buffer size for commandline in wspawnl
* Changed so first arg is the application name when calling CreateProcess from wspawnl

diff --git a/Engine/Source/Programs/UnrealBuildAccelerator/Detours/Private/Windows/UbaDetoursFunctionsKernelBase.inl b/Engine/Source/Programs/UnrealBuildAccelerator/Detours/Private/Windows/UbaDetoursFunctionsKernelBase.inl
--- a/Engine/Source/Programs/UnrealBuildAccelerator/Detours/Private/Windows/UbaDetoursFunctionsKernelBase.inl
+++ b/Engine/Source/Programs/UnrealBuildAccelerator/Detours/Private/Windows/UbaDetoursFunctionsKernelBase.inl
@@ -2682,6 +2682,16 @@ BOOL Detoured_CreateProcessW(LPCWSTR lpApplicationName, LPWSTR lpCommandLine, LP
 	}
 
 	bool isChild = true;
+	// We don't care about tracking mspdbsrv or vctip.. they are services just spawned by this process
+	if (Contains(originalCmd, L"mspdbsrv.exe") || Contains(originalCmd, L"vctip.exe") || Contains(originalCmd, L"git.exe"))
+	{
+		if (!g_runningRemote)
+		{
+			SuppressDetourScope _;
+			return True_CreateProcessW(lpApplicationName, lpCommandLine, lpProcessAttributes, lpThreadAttributes, bInheritHandles, dwCreationFlags, lpEnvironment, lpCurrentDirectory, lpStartupInfo, lpProcessInformation);
+		}
+		isChild = false;
+	}
 
 	StringBuffer<> application;
 	if (lpApplicationName)
@@ -2864,10 +2874,16 @@ BOOL Detoured_TerminateProcess(HANDLE hProcess, UINT uExitCode)
 BOOL Detoured_GetExitCodeProcess(HANDLE hProcess, LPDWORD lpExitCode)
 {
 	DETOURED_CALL(GetExitCodeProcess);
+	HANDLE trueHandle = hProcess;
 	if (isDetouredHandle(hProcess))
-		hProcess = asDetouredHandle(hProcess).trueHandle;
-	BOOL res = True_GetExitCodeProcess(hProcess, lpExitCode);
-	DEBUG_LOG_DETOURED(L"GetExitCodeProcess", L"%llu Exit code: %u -> %ls", uintptr_t(hProcess), *lpExitCode, ToString(res));
+		trueHandle = asDetouredHandle(hProcess).trueHandle;
+	BOOL res = True_GetExitCodeProcess(trueHandle, lpExitCode);
+
+	DEBUG_LOG_DETOURED(L"GetExitCodeProcess", L"%llu Exit code: %u -> %ls", uintptr_t(trueHandle), *lpExitCode, ToString(res));
+
+	if (res != STILL_ACTIVE)
+		Rpc_UpdateTables();
+
 	return res;
 }
 
diff --git a/Engine/Source/Programs/UnrealBuildAccelerator/Detours/Private/Windows/UbaDetoursFunctionsUcrtBase.inl b/Engine/Source/Programs/UnrealBuildAccelerator/Detours/Private/Windows/UbaDetoursFunctionsUcrtBase.inl
--- a/Engine/Source/Programs/UnrealBuildAccelerator/Detours/Private/Windows/UbaDetoursFunctionsUcrtBase.inl
+++ b/Engine/Source/Programs/UnrealBuildAccelerator/Detours/Private/Windows/UbaDetoursFunctionsUcrtBase.inl
@@ -279,7 +279,7 @@ int Detoured_fputs(const char* str, FILE* stream)
 int Detoured__wspawnl(int mode, const wchar_t* cmdname, const wchar_t* arg0, const wchar_t* arg1, ...)
 {
 	DETOURED_CALL(_wspawnl);
-	StringBuffer<> cmdLine;
+	StringBuffer<32*1024> cmdLine; // Need long commandline for cmake/ninja
 
 	if (cmdname == nullptr)
 		cmdname = arg0;
@@ -293,7 +293,8 @@ int Detoured__wspawnl(int mode, const wchar_t* cmdname, const wchar_t* arg0, con
 		cmdname = cmdNameTemp.data;
 	}
 
-	cmdLine.Append(L"dummy "); // Just because CreateProcess expects first arg to be name of application
+	// TODO: Not sure what the rule is here.. very confusing with CreateProcess having application in two places
+	cmdLine.Append('\"').Append(cmdname).Append('\"');
 
 	//wcscpy_s(cmdLine, 1024, arg0);
 	//wcscat_s(cmdLine, 1024, L" ");
