[UBA]
* Fixed old comment in visualizer

diff --git a/Engine/Source/Programs/UnrealBuildAccelerator/Visualizer/Private/UbaVisualizer.cpp b/Engine/Source/Programs/UnrealBuildAccelerator/Visualizer/Private/UbaVisualizer.cpp
--- a/Engine/Source/Programs/UnrealBuildAccelerator/Visualizer/Private/UbaVisualizer.cpp
+++ b/Engine/Source/Programs/UnrealBuildAccelerator/Visualizer/Private/UbaVisualizer.cpp
@@ -1566,7 +1566,7 @@ namespace uba
 				L"cpu/mem stats",
 				L"process bars",
 				L"timeline",
-				L"detailed data (use -BoxDetailedTrace for even more)",
+				L"detailed data (use -UbaDetailedTrace for even more)",
 				L"workers (threads on host taking care of requests from helpers)",
 				L"cursor (vertical line)",
 			};
