[UBA]
* Added SetEnvironmentVariableW wrapper for win and implementation for posix

diff --git a/Engine/Source/Programs/UnrealBuildAccelerator/Core/Private/UbaPlatform.cpp b/Engine/Source/Programs/UnrealBuildAccelerator/Core/Private/UbaPlatform.cpp
--- a/Engine/Source/Programs/UnrealBuildAccelerator/Core/Private/UbaPlatform.cpp
+++ b/Engine/Source/Programs/UnrealBuildAccelerator/Core/Private/UbaPlatform.cpp
@@ -304,6 +304,15 @@ namespace uba
 		#endif
 	}
 
+	bool SetEnvironmentVariableW(const tchar* name, const tchar* value)
+	{
+		#if PLATFORM_WINDOWS
+		return ::SetEnvironmentVariableW(name, value);
+		#else
+		return setenv(name, value, 1) == 0;
+		#endif
+	}
+
 	u32 ExpandEnvironmentStringsW(const tchar* lpSrc, tchar* lpDst, u32 nSize)
 	{
 		#if PLATFORM_WINDOWS
diff --git a/Engine/Source/Programs/UnrealBuildAccelerator/Core/Public/UbaPlatform.h b/Engine/Source/Programs/UnrealBuildAccelerator/Core/Public/UbaPlatform.h
--- a/Engine/Source/Programs/UnrealBuildAccelerator/Core/Public/UbaPlatform.h
+++ b/Engine/Source/Programs/UnrealBuildAccelerator/Core/Public/UbaPlatform.h
@@ -97,6 +97,7 @@ namespace uba
 	ProcHandle GetCurrentProcessHandle();
 
 	u32 GetEnvironmentVariableW(const tchar* name, tchar* buffer, u32 nSize);
+	bool SetEnvironmentVariableW(const tchar* name, const tchar* value);
 	u32 ExpandEnvironmentStringsW(const tchar* lpSrc, tchar* lpDst, u32 nSize);
 	u32 GetLogicalProcessorCount();
 	u32 GetProcessorGroupCount();
