[UBA]
* Added mscoree.dll to known system files

diff --git a/Engine/Source/Programs/UnrealBuildAccelerator/Core/Public/UbaWinBinDependencyParser.h b/Engine/Source/Programs/UnrealBuildAccelerator/Core/Public/UbaWinBinDependencyParser.h
--- a/Engine/Source/Programs/UnrealBuildAccelerator/Core/Public/UbaWinBinDependencyParser.h
+++ b/Engine/Source/Programs/UnrealBuildAccelerator/Core/Public/UbaWinBinDependencyParser.h
@@ -37,6 +37,7 @@ namespace uba
 		L"iphlpapi.dll",
 		L"kernel32.dll",
 		L"kernelbase.dll",
+		L"mscoree.dll",
 		L"msvcp_win.dll",
 		L"msvcrt.dll",
 		L"mswsock.dll",
