[UBA]
* Added code to replace /FS with /Z7 to embed pdb information in obj file instead of writing to separate file when compiling remotely (code to handle mspdbsrv remotely is WIP and might never happen if this work well enough)

diff --git a/Engine/Source/Programs/UnrealBuildAccelerator/Common/Private/UbaProcess.cpp b/Engine/Source/Programs/UnrealBuildAccelerator/Common/Private/UbaProcess.cpp
--- a/Engine/Source/Programs/UnrealBuildAccelerator/Common/Private/UbaProcess.cpp
+++ b/Engine/Source/Programs/UnrealBuildAccelerator/Common/Private/UbaProcess.cpp
@@ -174,6 +174,17 @@ namespace uba
 
 		UBA_ASSERT(m_startInfo.rules);
 
+		// If running remote we can't use mspdbsrv (not supported yet).. so instead embed information in .obj file
+		// TODO: This should be placed somewhere else it feels like.
+		#if PLATFORM_WINDOWS
+		if (runningRemote && (m_startInfo.rules->index == 1 || m_startInfo.rules->index == 2))
+		{
+			tchar* pos = nullptr;
+			if (Contains(m_startInfo.argumentsStr.data(), L"/FS ", true, (const tchar**)&pos))
+				memcpy(pos, L"/Z7", 6);
+		}
+		#endif
+
 		m_session.ProcessAdded(*this, 0);
 
 		if (async)
