[UBA]
* Added debug log for GetLongPathNameW
* Added GetCommandLineA in debug with assert if running remote

diff --git a/Engine/Source/Programs/UnrealBuildAccelerator/Detours/Private/Windows/UbaDetoursFunctionsKernelBase.inl b/Engine/Source/Programs/UnrealBuildAccelerator/Detours/Private/Windows/UbaDetoursFunctionsKernelBase.inl
--- a/Engine/Source/Programs/UnrealBuildAccelerator/Detours/Private/Windows/UbaDetoursFunctionsKernelBase.inl
+++ b/Engine/Source/Programs/UnrealBuildAccelerator/Detours/Private/Windows/UbaDetoursFunctionsKernelBase.inl
@@ -59,6 +59,8 @@ DWORD Local_GetLongPathNameW(LPCWSTR lpszShortPath, LPWSTR lpszLongPath, DWORD c
 		}
 
 		SetLastError(errorCode);
+
+		DEBUG_LOG_DETOURED(L"GetLongPathNameW", L"%ls", lpszShortPath);
 	}
 	else
 	{
@@ -79,14 +81,14 @@ DWORD Local_GetLongPathNameW(LPCWSTR lpszShortPath, LPWSTR lpszLongPath, DWORD c
 LPWSTR Detoured_GetCommandLineW()
 {
 	DETOURED_CALL(GetCommandLineW);
-	if (!g_runningRemote)
+	if (!g_virtualCommandLineW)
 	{
 		LPWSTR str = True_GetCommandLineW();
 		DEBUG_LOG_TRUE(L"GetCommandLineW", L"");// str);
 		return str;
 	}
-	DEBUG_LOG_DETOURED(L"GetCommandLineW", L"");// g_virtualCommandLine);
-	return g_virtualCommandLine;
+	DEBUG_LOG_DETOURED(L"GetCommandLineW", L"");
+	return g_virtualCommandLineW;
 }
 
 DWORD Detoured_GetCurrentDirectoryW(DWORD nBufferLength, LPWSTR lpBuffer)
@@ -3634,6 +3636,19 @@ BOOL Detoured_IsProcessorFeaturePresent(DWORD ProcessorFeature)
 //	return True_CreateFileMappingNumaW(hFile, lpFileMappingAttributes, flProtect, dwMaximumSizeHigh, dwMaximumSizeLow, lpName, nndPreferred);
 //}
 
+LPSTR Detoured_GetCommandLineA()
+{
+	DETOURED_CALL(GetCommandLineA);
+	if (!g_virtualCommandLineA)
+	{
+		auto str = True_GetCommandLineA();
+		DEBUG_LOG_TRUE(L"GetCommandLineA", L"");// str);
+		return str;
+	}
+	DEBUG_LOG_DETOURED(L"GetCommandLineA", L"");
+	return g_virtualCommandLineA;
+}
+
 BOOL Detoured_FreeLibrary(HMODULE hModule)
 {
 	DETOURED_CALL(FreeLibrary);
diff --git a/Engine/Source/Programs/UnrealBuildAccelerator/Detours/Private/Windows/UbaDetoursFunctionsUcrtBase.inl b/Engine/Source/Programs/UnrealBuildAccelerator/Detours/Private/Windows/UbaDetoursFunctionsUcrtBase.inl
--- a/Engine/Source/Programs/UnrealBuildAccelerator/Detours/Private/Windows/UbaDetoursFunctionsUcrtBase.inl
+++ b/Engine/Source/Programs/UnrealBuildAccelerator/Detours/Private/Windows/UbaDetoursFunctionsUcrtBase.inl
@@ -411,4 +411,11 @@ errno_t Detoured__wmakepath_s(wchar_t* path, size_t sizeInWords, const wchar_t*
 	DEBUG_LOG_TRUE(L"_wmakepath_s", L"%ls %ls %ls %ls %ls", path, drive, dir, fname, ext);
 	return res;
 }
+char* Detoured__getcwd(char* buffer, int maxlen)
+{
+	DETOURED_CALL(_getcwd);
+	UBA_ASSERT(!g_virtualCommandLineA);
+	return True__getcwd(buffer, maxlen);
+}
+
 #endif // defined(DETOURED_INCLUDE_DEBUG)
diff --git a/Engine/Source/Programs/UnrealBuildAccelerator/Detours/Private/Windows/UbaDetoursFunctionsWin.cpp b/Engine/Source/Programs/UnrealBuildAccelerator/Detours/Private/Windows/UbaDetoursFunctionsWin.cpp
--- a/Engine/Source/Programs/UnrealBuildAccelerator/Detours/Private/Windows/UbaDetoursFunctionsWin.cpp
+++ b/Engine/Source/Programs/UnrealBuildAccelerator/Detours/Private/Windows/UbaDetoursFunctionsWin.cpp
@@ -112,8 +112,8 @@ int g_uiLanguage;
 
 StringBuffer<256> g_exeDir;
 
-const wchar_t* g_commandLine;
-wchar_t* g_virtualCommandLine;
+wchar_t* g_virtualCommandLineW;
+char* g_virtualCommandLineA;
 
 constexpr u32 TrackInputsMemCapacity = 512 * 1024;
 u8* g_trackInputsMem;
@@ -857,20 +857,28 @@ void Init(const DetoursPayload& payload, u64 startTime)
 	const wchar_t* cmdLine = True_GetCommandLineW();
 
 	const wchar_t* exePos;
-	if (Contains(cmdLine, g_exeDir.data, true, &exePos))
+	if (g_runningRemote && Contains(cmdLine, g_exeDir.data, true, &exePos))
 	{
 		StringBuffer<> buf;
 		buf.Append(cmdLine, exePos - cmdLine);
 		buf.Append(g_virtualApplicationDir);
 		TString realCmdLine(buf.data);
 		realCmdLine += (cmdLine + g_exeDir.count);
-		g_virtualCommandLine = g_memoryBlock.Strdup(realCmdLine.c_str());
+		g_virtualCommandLineW = g_memoryBlock.Strdup(realCmdLine.c_str());
+		}
+	//else
+	//	g_virtualCommandLineW = g_memoryBlock.Strdup(cmdLine);
+
+	if (g_virtualCommandLineW)
+	{
+		g_virtualCommandLineA = (char*)g_memoryBlock.Allocate(realCmdLine.size() + 1, 1, L"");
+		size_t res;
+		wcstombs_s(&res, g_virtualCommandLineA, realCmdLine.size() + 1, g_virtualCommandLineW, realCmdLine.size());
 	}
 
 	#if UBA_DEBUG_LOG_ENABLED
 	if (isLogging())
 	{
-		g_commandLine = cmdLine;
 		u64 cmdLineLen = wcslen(cmdLine);
 		wchar_t temp[LogBufSize - 10];
 		if (cmdLineLen > sizeof_array(temp))
diff --git a/Engine/Source/Programs/UnrealBuildAccelerator/Detours/Private/Windows/UbaDetoursFunctionsWin.h b/Engine/Source/Programs/UnrealBuildAccelerator/Detours/Private/Windows/UbaDetoursFunctionsWin.h
--- a/Engine/Source/Programs/UnrealBuildAccelerator/Detours/Private/Windows/UbaDetoursFunctionsWin.h
+++ b/Engine/Source/Programs/UnrealBuildAccelerator/Detours/Private/Windows/UbaDetoursFunctionsWin.h
@@ -18,6 +18,9 @@
 #define DETOURED_INCLUDE_DEBUG
 #endif
 
+#if defined(DETOURED_INCLUDE_DEBUG)
+#include <direct.h>
+#endif
 
 
 #define DETOURED_FUNCTIONS \
@@ -203,6 +206,7 @@
 #if defined(DETOURED_INCLUDE_DEBUG)
 
 #define DETOURED_FUNCTIONS_KERNELBASE_DEBUG \
+	DETOURED_FUNCTION(GetCommandLineA) \
 	DETOURED_FUNCTION(FreeLibrary) \
 	DETOURED_FUNCTION(RegOpenKeyExW) \
 	DETOURED_FUNCTION(SetLastError) \
@@ -331,6 +335,7 @@
 	DETOURED_FUNCTION(getenv) \
 	DETOURED_FUNCTION(getenv_s) \
 	DETOURED_FUNCTION(_wmakepath_s) \
+	DETOURED_FUNCTION(_getcwd) \
 	//DETOURED_FUNCTION(_wsopen_s) \
 	//DETOURED_FUNCTION(_fileno) \
 
