[UBA]
* Fixed naming of log files when produced by remote sessions. Now suffix is using process id instead of counter number
* Removed m_id from Session (not used)
* Added so Scheduler MaxLocalProcessors can be controlled from config files

diff --git a/Engine/Source/Programs/UnrealBuildAccelerator/Agent/Private/UbaAgent.cpp b/Engine/Source/Programs/UnrealBuildAccelerator/Agent/Private/UbaAgent.cpp
--- a/Engine/Source/Programs/UnrealBuildAccelerator/Agent/Private/UbaAgent.cpp
+++ b/Engine/Source/Programs/UnrealBuildAccelerator/Agent/Private/UbaAgent.cpp
@@ -952,7 +952,7 @@ namespace uba
 				if (!desc.empty())
 					name.Append(desc);
 				else
-					GetNameFromArguments(name, process.GetStartInfo().arguments, false);
+					GenerateNameForProcess(name, process.GetStartInfo().arguments, 0);
 				LogEntryType entryType = LogEntryType_Info;
 				if (errorCode)
 				{
diff --git a/Engine/Source/Programs/UnrealBuildAccelerator/Common/Private/UbaScheduler.cpp b/Engine/Source/Programs/UnrealBuildAccelerator/Common/Private/UbaScheduler.cpp
--- a/Engine/Source/Programs/UnrealBuildAccelerator/Common/Private/UbaScheduler.cpp
+++ b/Engine/Source/Programs/UnrealBuildAccelerator/Common/Private/UbaScheduler.cpp
@@ -84,6 +84,7 @@ namespace uba
 			table->GetValueAsBool(enableProcessReuse, TC("EnableProcessReuse"));
 			table->GetValueAsBool(forceRemote, TC("ForceRemote"));
 			table->GetValueAsBool(forceNative, TC("ForceNative"));
+			table->GetValueAsU32(maxLocalProcessors, TC("MaxLocalProcessors"));
 		}
 	}
 
@@ -626,7 +627,8 @@ namespace uba
 				StringBuffer<> logFile;
 				if (true)
 				{
-					GetNameFromArguments(logFile, si.arguments, true);
+					static u32 processId = 1; // TODO: This should be done in a better way.. or not at all?
+					GenerateNameForProcess(logFile, si.arguments, ++processId);
 					logFile.Append(TC(".log"));
 					si.logFile = logFile.data;
 				};
diff --git a/Engine/Source/Programs/UnrealBuildAccelerator/Common/Private/UbaSession.cpp b/Engine/Source/Programs/UnrealBuildAccelerator/Common/Private/UbaSession.cpp
--- a/Engine/Source/Programs/UnrealBuildAccelerator/Common/Private/UbaSession.cpp
+++ b/Engine/Source/Programs/UnrealBuildAccelerator/Common/Private/UbaSession.cpp
@@ -1000,7 +1000,6 @@ namespace uba
 		UBA_ASSERTF(info.rootDir && *info.rootDir, TC("No root dir set when creating session"));
 		m_rootDir.count = GetFullPathNameW(info.rootDir, m_rootDir.capacity, m_rootDir.data, NULL);
 		m_rootDir.Replace('/', PathSeparator).EnsureEndsWithSlash();
-		m_uid = u32(HashString()(m_rootDir.data));
 
 		m_runningRemote = runningRemote;
 		m_disableCustomAllocator = info.disableCustomAllocator;
@@ -1240,6 +1239,8 @@ namespace uba
 		si.useCustomAllocator &= !m_disableCustomAllocator;
 		const tchar* originalLogFile = si.logFile;
 		
+		u32 processId = CreateProcessId();
+
 		StringBuffer<> logFile;
 		if (si.logFile && *si.logFile)
 		{
@@ -1252,7 +1253,7 @@ namespace uba
 		else if (m_logToFile)
 		{
 			logFile.Append(m_sessionLogDir);
-			GetNameFromArguments(logFile, startInfo.arguments, true);
+			GenerateNameForProcess(logFile, startInfo.arguments, processId);
 			logFile.Append(TC(".log"));
 			si.logFile = logFile.data;
 		}
@@ -1261,8 +1262,7 @@ namespace uba
 			si.rules = GetRules(si);
 
 		void* env = GetProcessEnvironmentVariables();
-		u32 id = CreateProcessId();
-		auto process = new ProcessImpl(*this, id, parent);
+		auto process = new ProcessImpl(*this, processId, parent);
 		ProcessHandle h(process);
 		if (!process->Start(startInfo, m_runningRemote, env, async, enableDetour))
 			return {};
@@ -2816,7 +2816,7 @@ namespace uba
 	{
 	}
 
-	void GetNameFromArguments(StringBufferBase& out, const tchar* arguments, bool addCounterSuffix)
+	void GenerateNameForProcess(StringBufferBase& out, const tchar* arguments, u32 counterSuffix)
 	{
 		const tchar* start = arguments;
 		const tchar* it = arguments;
@@ -2847,8 +2847,7 @@ namespace uba
 		if (out.IsEmpty())
 			out.Append(TC("NoGoodName"));
 
-		static Atomic<u32> counter;
-		if (addCounterSuffix)
-			out.Appendf(TC("_%03u"), counter++);
+		if (counterSuffix)
+			out.Appendf(TC("_%03u"), counterSuffix);
 	}
 }
diff --git a/Engine/Source/Programs/UnrealBuildAccelerator/Common/Private/UbaSessionClient.cpp b/Engine/Source/Programs/UnrealBuildAccelerator/Common/Private/UbaSessionClient.cpp
--- a/Engine/Source/Programs/UnrealBuildAccelerator/Common/Private/UbaSessionClient.cpp
+++ b/Engine/Source/Programs/UnrealBuildAccelerator/Common/Private/UbaSessionClient.cpp
@@ -1452,7 +1452,7 @@ namespace uba
 			SendLogFileToServer(*(ProcessImpl*)child.m_process);
 	}
 
-	void SessionClient::GetLogFileName(StringBufferBase& out, const tchar* logFile, const tchar* arguments)
+	void SessionClient::GetLogFileName(StringBufferBase& out, const tchar* logFile, const tchar* arguments, u32 processId)
 	{
 		out.Append(m_sessionLogDir.data);
 		if (logFile && *logFile)
@@ -1463,7 +1463,7 @@ namespace uba
 		}
 		else
 		{
-			GetNameFromArguments(out, arguments, true);
+			GenerateNameForProcess(out, arguments, processId);
 			out.Append(TC(".log"));
 		}
 	}
@@ -1636,7 +1636,7 @@ namespace uba
 					StringBuffer<> logFile;
 					if (m_logToFile)
 					{
-						GetLogFileName(logFile, startInfo.logFile, startInfo.arguments);
+						GetLogFileName(logFile, startInfo.logFile, startInfo.arguments, startInfo.processId);
 						startInfo.logFile = logFile.data;
 					}
 
@@ -1898,7 +1898,7 @@ namespace uba
 			if (m_logToFile)
 			{
 				StringBuffer<512> logFile;
-				GetLogFileName(logFile, outNextProcess.logFile.c_str(), outNextProcess.arguments.c_str());
+				GetLogFileName(logFile, outNextProcess.logFile.c_str(), outNextProcess.arguments.c_str(), process.GetId());
 				outNextProcess.logFile = logFile.data;
 			}
 		}
diff --git a/Engine/Source/Programs/UnrealBuildAccelerator/Common/Public/UbaSession.h b/Engine/Source/Programs/UnrealBuildAccelerator/Common/Public/UbaSession.h
--- a/Engine/Source/Programs/UnrealBuildAccelerator/Common/Public/UbaSession.h
+++ b/Engine/Source/Programs/UnrealBuildAccelerator/Common/Public/UbaSession.h
@@ -192,7 +192,6 @@ namespace uba
 
 		u64 m_keepOutputFileMemoryMapsThreshold;
 
-		u32 m_uid;
 		Atomic<u32> m_processIdCounter;
 
 		MemoryBlock m_directoryTableMemory;
@@ -265,7 +264,7 @@ namespace uba
 		friend class ProcessImpl;
 	};
 
-	void GetNameFromArguments(StringBufferBase& out, const tchar* arguments, bool addCounterSuffix);
+	void GenerateNameForProcess(StringBufferBase& out, const tchar* arguments, u32 counterSuffix);
 
 	using FileAccess = u8;
 
diff --git a/Engine/Source/Programs/UnrealBuildAccelerator/Common/Public/UbaSessionClient.h b/Engine/Source/Programs/UnrealBuildAccelerator/Common/Public/UbaSessionClient.h
--- a/Engine/Source/Programs/UnrealBuildAccelerator/Common/Public/UbaSessionClient.h
+++ b/Engine/Source/Programs/UnrealBuildAccelerator/Common/Public/UbaSessionClient.h
@@ -91,7 +91,7 @@ namespace uba
 		bool SendProcessFinished(ProcessImpl& process, u32 exitCode);
 		void SendPing(u64 memAvail, u64 memTotal);
 		void SendLogFileToServer(ProcessImpl& pi);
-		void GetLogFileName(StringBufferBase& out, const tchar* logFile, const tchar* arguments);
+		void GetLogFileName(StringBufferBase& out, const tchar* logFile, const tchar* arguments, u32 processId);
 		u32 CountLogLines(ProcessImpl& process);
 		void WriteLogLines(BinaryWriter& writer, ProcessImpl& process);
 
diff --git a/Engine/Source/Programs/UnrealBuildAccelerator/TestApp/Private/Windows/UbaTestAppWin.cpp b/Engine/Source/Programs/UnrealBuildAccelerator/TestApp/Private/Windows/UbaTestAppWin.cpp
--- a/Engine/Source/Programs/UnrealBuildAccelerator/TestApp/Private/Windows/UbaTestAppWin.cpp
+++ b/Engine/Source/Programs/UnrealBuildAccelerator/TestApp/Private/Windows/UbaTestAppWin.cpp
@@ -131,9 +131,11 @@ int wmain(int argc, wchar_t* argv[])
 			if (fh == INVALID_HANDLE_VALUE)
 				return LogError(L"Failed to create file File");
 			CloseHandle(fh);
-			MoveFile(L"FileW", L"FileW2");
+			if (!MoveFile(L"FileW", L"FileW2"))
+				return LogError(L"Failed to move file from FileW to FileW2");
 
-			CopyFile(L"FileW2", L"FileWF", false);
+			if (!CopyFile(L"FileW2", L"FileWF", false))
+				return LogError(L"Failed to copy file from FileW2 to FileWF");
 		}
 
 		{
