[UBA]
* Changed so code getting processor frequency uses registry instead of CallNtPowerInformation since we want a number that is the same every time and CallNtPowerInformation + MaxMhz does not do that

diff --git a/Engine/Source/Programs/UnrealBuildAccelerator/Common/Private/UbaSession.cpp b/Engine/Source/Programs/UnrealBuildAccelerator/Common/Private/UbaSession.cpp
--- a/Engine/Source/Programs/UnrealBuildAccelerator/Common/Private/UbaSession.cpp
+++ b/Engine/Source/Programs/UnrealBuildAccelerator/Common/Private/UbaSession.cpp
@@ -2430,10 +2430,22 @@ namespace uba
 		#if PLATFORM_WINDOWS
 		GetPhysicallyInstalledSystemMemory(&totalMemoryInKilobytes);
 
-		Vector<PROCESSOR_POWER_INFORMATION> procInfos;
-		procInfos.resize(cpuCount);
-		if (CallNtPowerInformation(ProcessorInformation, NULL, 0, procInfos.data(), cpuCount*sizeof(PROCESSOR_POWER_INFORMATION)) == STATUS_SUCCESS)
-			hzStr.Appendf(TC(" @ %.1fGHz"), float(procInfos[0].MaxMhz) / 1000.0f);
+		{
+			u32 maxMHz = 0;
+			DWORD valueSize = 4;
+			const tchar* key = TC("HARDWARE\\DESCRIPTION\\System\\CentralProcessor\\0");
+			LSTATUS res = RegGetValueW(HKEY_LOCAL_MACHINE, key, TC("~MHz"), RRF_RT_REG_DWORD, NULL, &maxMHz, &valueSize);
+			if (res != ERROR_SUCCESS)
+			{
+				// This will not always be the same and since we use the system info as part of key for client uniqueness it is annoying to get multiple sessions for same instance
+				Vector<PROCESSOR_POWER_INFORMATION> procInfos;
+				procInfos.resize(cpuCount);
+				if (CallNtPowerInformation(ProcessorInformation, NULL, 0, procInfos.data(), cpuCount*sizeof(PROCESSOR_POWER_INFORMATION)) == STATUS_SUCCESS)
+					maxMHz = procInfos[0].MaxMhz;
+			}
+			hzStr.Appendf(TC(" @ %.1fGHz"), float(maxMHz) / 1000.0f);
+		}
+
 		#else
 		u64 throwAway;
 		GetMemoryInfo(throwAway, totalMemoryInKilobytes);
diff --git a/Engine/Source/Programs/UnrealBuildAccelerator/Visualizer/Private/UbaTraceReader.cpp b/Engine/Source/Programs/UnrealBuildAccelerator/Visualizer/Private/UbaTraceReader.cpp
--- a/Engine/Source/Programs/UnrealBuildAccelerator/Visualizer/Private/UbaTraceReader.cpp
+++ b/Engine/Source/Programs/UnrealBuildAccelerator/Visualizer/Private/UbaTraceReader.cpp
@@ -303,20 +303,19 @@ namespace uba
 			for (u32 i = 0; i != out.sessions.size(); ++i)
 			{
 				auto& oldSession = out.sessions[i];
-				if (oldSession.name == sessionName.data)
-				{
-					if (oldSession.disconnectTime == ~u64(0))
-						break;
-					oldSession.isReset = true;
-					oldSession.disconnectTime = ~u64(0);
-					oldSession.proxyName.clear();
-					oldSession.proxyCreated = false;
-					oldSession.notification.clear();
-					//oldSession.fetchedFiles.clear();
-					//oldSession.storedFiles.clear();
-					virtualSessionIndex = i;
+				if (oldSession.name != sessionName.data)
+					continue;
+				if (oldSession.disconnectTime == ~u64(0))
 					break;
-				}
+				oldSession.isReset = true;
+				oldSession.disconnectTime = ~u64(0);
+				oldSession.proxyName.clear();
+				oldSession.proxyCreated = false;
+				oldSession.notification.clear();
+				//oldSession.fetchedFiles.clear();
+				//oldSession.storedFiles.clear();
+				virtualSessionIndex = i;
+				break;
 			}
 
 
