[UBA]
* UbaCoordinatorHorde now sets DISABLE_CWD_CHANGES to work with applications that needs current directory to not be set
* Added more initialization code in UbaCoordinatorHorde to more gracefully initialize ue core code
* Added code to coordinator wrapper to use scheduler queued process count in order to figure out how many remote cores are desired

diff --git a/Engine/Source/Developer/UbaCoordinatorHorde/Private/UbaCoordinatorHorde.cpp b/Engine/Source/Developer/UbaCoordinatorHorde/Private/UbaCoordinatorHorde.cpp
--- a/Engine/Source/Developer/UbaCoordinatorHorde/Private/UbaCoordinatorHorde.cpp
+++ b/Engine/Source/Developer/UbaCoordinatorHorde/Private/UbaCoordinatorHorde.cpp
@@ -4,6 +4,7 @@
 
 #if defined(UBA_COORDINATOR_HORDE_DLL)
 
+#include "DesktopPlatformModule.h"
 #include "HAL/Platform.h"
 #include "HAL/PlatformApplicationMisc.h"
 #include "HttpManager.h"
@@ -12,6 +13,7 @@
 #include "Misc/ConfigCacheIni.h"
 #include "Misc/OutputDeviceConsole.h"
 #include "Misc/OutputDeviceRedirector.h"
+#include "SocketSubsystem.h"
 #include "UbaCoordinator.h"
 #include "UbaHordeAgentManager.h"
 
@@ -67,6 +69,14 @@ extern "C"
 			GLog->AddOutputDevice(GLogConsole);
 		}
 
+		GGameThreadId = FPlatformTLS::GetCurrentThreadId();
+		GIsGameThreadIdInitialized = true;
+
+		// Since we are not setting CWD we need to manually call these systems from this thread (if not called from game thread LoadModule uses CWD which is not set)
+		ISocketSubsystem::Get(PLATFORM_SOCKETSUBSYSTEM);
+		FHttpModule::Get();
+		FDesktopPlatformModule::TryGet();
+
 		auto coordinator = new CoordinatorHorde(info.workDir, info.binariesDir);
 		auto& m = coordinator->m_manager;
 		m.SetPool(info.pool);
diff --git a/Engine/Source/Programs/UnrealBuildAccelerator/Common/Private/UbaScheduler.cpp b/Engine/Source/Programs/UnrealBuildAccelerator/Common/Private/UbaScheduler.cpp
--- a/Engine/Source/Programs/UnrealBuildAccelerator/Common/Private/UbaScheduler.cpp
+++ b/Engine/Source/Programs/UnrealBuildAccelerator/Common/Private/UbaScheduler.cpp
@@ -217,6 +217,24 @@ namespace uba
 		m_processFinished = processFinished;
 	}
 
+	u32 Scheduler::GetProcessCountThatCanRunRemotelyNow()
+	{
+		u32 count = 0;
+		SCOPED_READ_LOCK(m_processEntriesLock, lock);
+		for (auto& entry : m_processEntries)
+		{
+			if (!entry.canExecuteRemotely)
+				continue;
+			if (entry.status != ProcessStatus_QueuedForRun)
+				continue;
+			++count;
+		}
+
+		count += m_activeRemoteProcesses;
+
+		return count;
+	}
+
 	void Scheduler::ThreadLoop()
 	{
 		while (m_loop)
diff --git a/Engine/Source/Programs/UnrealBuildAccelerator/Common/Public/UbaCoordinatorWrapper.h b/Engine/Source/Programs/UnrealBuildAccelerator/Common/Public/UbaCoordinatorWrapper.h
--- a/Engine/Source/Programs/UnrealBuildAccelerator/Common/Public/UbaCoordinatorWrapper.h
+++ b/Engine/Source/Programs/UnrealBuildAccelerator/Common/Public/UbaCoordinatorWrapper.h
@@ -5,6 +5,7 @@
 #include "UbaCoordinator.h"
 #include "UbaFile.h"
 #include "UbaNetworkServer.h"
+#include "UbaScheduler.h"
 #include "UbaStringBuffer.h"
 
 #if !PLATFORM_WINDOWS
@@ -22,7 +23,7 @@ namespace uba
 	class CoordinatorWrapper
 	{
 	public:
-		bool Create(Logger& logger, const tchar* coordinatorType, const CoordinatorCreateInfo& info, NetworkBackend& networkBackend, NetworkServer& networkServer)
+		bool Create(Logger& logger, const tchar* coordinatorType, const CoordinatorCreateInfo& info, NetworkBackend& networkBackend, NetworkServer& networkServer, Scheduler* scheduler = nullptr)
 		{
 			UbaCreateCoordinatorFunc* createCoordinator = nullptr;
 
@@ -58,24 +59,30 @@ namespace uba
 			m_loopCoordinator.Create(true);
 			m_networkBackend = &networkBackend;
 			m_networkServer = &networkServer;
+			m_scheduler = scheduler;
 
-			m_coordinatorThread.Start([this, nb = &networkBackend, ns = &networkServer, tc = info.maxCoreCount]()
-			{
-				m_coordinator->SetAddClientCallback([](void* userData, const tchar* ip, u16 port)
-					{
-						auto& cw = *(CoordinatorWrapper*)userData;
-						return cw.m_networkServer->AddClient(*cw.m_networkBackend, ip, port);
-					}, this);
+			m_coordinatorThread.Start([this, mcc = info.maxCoreCount]() { ThreadUpdate(mcc); return 0; });
+
+			return true;
+		}
 
-				do
+		void ThreadUpdate(u32 maxCoreCount)
+		{
+			m_coordinator->SetAddClientCallback([](void* userData, const tchar* ip, u16 port)
 				{
-					m_coordinator->SetTargetCoreCount(tc);
-				}
-				while (!m_loopCoordinator.IsSet(3000));
+					auto& cw = *(CoordinatorWrapper*)userData;
+					return cw.m_networkServer->AddClient(*cw.m_networkBackend, ip, port);
+				}, this);
 
-				return 0;
-			});
-			return true;
+			do
+			{
+				u32 coreCount = maxCoreCount;
+				if (m_scheduler)
+					coreCount = Min(m_scheduler->GetProcessCountThatCanRunRemotelyNow(), maxCoreCount);
+
+				m_coordinator->SetTargetCoreCount(coreCount);
+			}
+			while (!m_loopCoordinator.IsSet(3000));
 		}
 
 		void Destroy()
@@ -91,6 +98,7 @@ namespace uba
 		Coordinator* m_coordinator = nullptr;
 		NetworkBackend* m_networkBackend = nullptr;
 		NetworkServer* m_networkServer = nullptr;
+		Scheduler* m_scheduler = nullptr;
 		UbaDestroyCoordinatorFunc* m_destroyCoordinator = nullptr;
 		Event m_loopCoordinator;
 		Thread m_coordinatorThread;
diff --git a/Engine/Source/Programs/UnrealBuildAccelerator/Common/Public/UbaScheduler.h b/Engine/Source/Programs/UnrealBuildAccelerator/Common/Public/UbaScheduler.h
--- a/Engine/Source/Programs/UnrealBuildAccelerator/Common/Public/UbaScheduler.h
+++ b/Engine/Source/Programs/UnrealBuildAccelerator/Common/Public/UbaScheduler.h
@@ -71,6 +71,8 @@ namespace uba
 
 		SessionServer& GetSession() { return m_session; }
 		
+		u32 GetProcessCountThatCanRunRemotelyNow();
+
 	private:
 		struct ExitProcessInfo;
 		struct ProcessStartInfo2;
diff --git a/Engine/Source/Programs/UnrealBuildAccelerator/UbaCoordinatorHorde.Target.cs b/Engine/Source/Programs/UnrealBuildAccelerator/UbaCoordinatorHorde.Target.cs
--- a/Engine/Source/Programs/UnrealBuildAccelerator/UbaCoordinatorHorde.Target.cs
+++ b/Engine/Source/Programs/UnrealBuildAccelerator/UbaCoordinatorHorde.Target.cs
@@ -66,6 +66,7 @@ public class UbaCoordinatorHordeTarget : TargetRules
 		GlobalDefinitions.Add("ENABLE_STATNAMEDEVENTS=0");
 		GlobalDefinitions.Add("ALLOW_HITCH_DETECTION=0");
 		GlobalDefinitions.Add("UE_USE_MALLOC_FILL_BYTES=0");
+		GlobalDefinitions.Add("DISABLE_CWD_CHANGES=1");
 
 		GlobalDefinitions.Add("UBA_COORDINATOR_HORDE_DLL");
 	}
