[UBA]
* Added config file support to UbaVisualizer, first iteration. If clicking "Save Position/Settings" in context menu config will be saved in %PROGRAMDATA%\Epic\UbaVisualizer\UbaVisualizer.toml

diff --git a/Engine/Source/Programs/UnrealBuildAccelerator/Common/Private/UbaConfig.cpp b/Engine/Source/Programs/UnrealBuildAccelerator/Common/Private/UbaConfig.cpp
--- a/Engine/Source/Programs/UnrealBuildAccelerator/Common/Private/UbaConfig.cpp
+++ b/Engine/Source/Programs/UnrealBuildAccelerator/Common/Private/UbaConfig.cpp
@@ -23,6 +23,15 @@ namespace uba
 		return buf.Parse(out);
 	}
 
+	bool ConfigTable::GetValueAsInt(int& out, const tchar* key) const
+	{
+		auto findIt = m_values.find(key);
+		if (findIt == m_values.end())
+			return m_parent ? m_parent->GetValueAsInt(out, key) : false;
+		out = (int)wcstol(findIt->second.c_str(), 0, 10);
+		return true;
+	}
+
 	bool ConfigTable::GetValueAsBool(bool& out, const tchar* key) const
 	{
 		auto findIt = m_values.find(key);
@@ -51,6 +60,18 @@ namespace uba
 		return &findIt->second;
 	}
 
+	void ConfigTable::AddValue(const tchar* key, int value)
+	{
+		tchar buf[256];
+		TSprintf_s(buf, sizeof_array(buf), TC("%i"), value);
+		m_values[key] = buf;
+	}
+
+	void ConfigTable::AddValue(const tchar* key, bool value)
+	{
+		m_values[key] = value ? TC("true") : TC("false");
+	}
+
 	bool Config::LoadFromFile(Logger& logger, const tchar* configFile)
 	{
 		m_isLoaded = true;
@@ -195,4 +216,23 @@ namespace uba
 	{
 		return m_isLoaded;
 	}
+
+	bool Config::SaveToFile(Logger& logger, const tchar* configFile)
+	{
+		UBA_ASSERT(m_tables.empty()); // TODO
+
+		FileAccessor fa(logger, configFile);
+		if (!fa.CreateWrite())
+			return false;
+
+		for (auto& kv : m_values)
+		{
+			char line[1024];
+			int written = sprintf_s(line, sizeof_array(line), "%S = %S\r\n", kv.first.c_str(), kv.second.c_str());
+			if (!fa.Write(line, written))
+				return false;
+		}
+		
+		return fa.Close();
+	}
 }
diff --git a/Engine/Source/Programs/UnrealBuildAccelerator/Common/Public/UbaConfig.h b/Engine/Source/Programs/UnrealBuildAccelerator/Common/Public/UbaConfig.h
--- a/Engine/Source/Programs/UnrealBuildAccelerator/Common/Public/UbaConfig.h
+++ b/Engine/Source/Programs/UnrealBuildAccelerator/Common/Public/UbaConfig.h
@@ -15,13 +15,17 @@ namespace uba
 	public:
 		bool GetValueAsString(const tchar*& out, const tchar* key) const;
 		bool GetValueAsU32(u32& out, const tchar* key) const;
+		bool GetValueAsInt(int& out, const tchar* key) const;
 		bool GetValueAsBool(bool& out, const tchar* key) const;
 
 		const ConfigTable* GetTable(const tchar* name) const;
 
+		void AddValue(const tchar* key, int value);
+		void AddValue(const tchar* key, bool value);
+
 	private:
 		ConfigTable* m_parent = nullptr;
-		UnorderedMap<TString, TString> m_values;
+		Map<TString, TString> m_values;
 		UnorderedMap<TString, ConfigTable> m_tables;
 		friend Config;
 	};
@@ -34,6 +38,7 @@ namespace uba
 		bool LoadFromText(Logger& logger, const char* text, u64 textLen);
 		bool IsLoaded() const;
 
+		bool SaveToFile(Logger& logger, const tchar* configFile);
 
 		bool m_isLoaded = false;
 	};
diff --git a/Engine/Source/Programs/UnrealBuildAccelerator/Visualizer/Private/UbaVisualizer.cpp b/Engine/Source/Programs/UnrealBuildAccelerator/Visualizer/Private/UbaVisualizer.cpp
--- a/Engine/Source/Programs/UnrealBuildAccelerator/Visualizer/Private/UbaVisualizer.cpp
+++ b/Engine/Source/Programs/UnrealBuildAccelerator/Visualizer/Private/UbaVisualizer.cpp
@@ -1,31 +1,90 @@
 // Copyright Epic Games, Inc. All Rights Reserved.
 
 #include "UbaVisualizer.h"
+#include "UbaConfig.h"
 #include <algorithm>
 
 #include <uxtheme.h>
 #include <dwmapi.h>
+
 #pragma comment (lib, "UxTheme.lib")
 #pragma comment (lib, "Dwmapi.lib")
 #define WM_NEWTRACE WM_USER+1
 
-enum
-{
-	Popup_CopySessionInfo = 3,
-	Popup_CopyProcessInfo,
-	Popup_CopyProcessLog,
-	Popup_Replay,
-	Popup_Pause,
-	Popup_Play,
-	Popup_JumpToEnd,
-	Popup_ShowText,
-	Popup_ShowCreateWriteColors,
-	Popup_SaveAs,
-	Popup_Quit,
-};
-
 namespace uba
 {
+	enum
+	{
+		Popup_CopySessionInfo = 3,
+		Popup_CopyProcessInfo,
+		Popup_CopyProcessLog,
+		Popup_Replay,
+		Popup_Pause,
+		Popup_Play,
+		Popup_JumpToEnd,
+
+		#define UBA_VISUALIZER_FLAG(name, defaultValue, desc) Popup_##name,
+		UBA_VISUALIZER_FLAGS2
+		#undef UBA_VISUALIZER_FLAG
+
+		Popup_SaveAs,
+		Popup_SaveSettings,
+		Popup_Quit,
+	};
+
+	VisualizerConfig::VisualizerConfig(const tchar* fn) : filename(fn)
+	{
+	}
+
+	bool VisualizerConfig::Load(Logger& logger)
+	{
+		Config config;
+		if (!config.LoadFromFile(logger, filename.c_str()))
+		{
+			DWORD value = 1;
+			DWORD valueSize = sizeof(value);
+			if (RegGetValueW(HKEY_CURRENT_USER, L"Software\\Microsoft\\Windows\\CurrentVersion\\Themes\\Personalize", L"AppsUseLightTheme", RRF_RT_REG_DWORD, NULL, &value, &valueSize) == ERROR_SUCCESS)
+				DarkMode = value == 0;
+
+			return false;
+		}
+		config.GetValueAsInt(x, L"X");
+		config.GetValueAsInt(y, L"Y");
+		config.GetValueAsInt(width, L"Width");
+		config.GetValueAsInt(height, L"Height");
+		#define UBA_VISUALIZER_FLAG(name, defaultValue, desc) config.GetValueAsBool(show##name, L"Show" #name);
+		UBA_VISUALIZER_FLAGS1
+		#undef UBA_VISUALIZER_FLAG
+		#define UBA_VISUALIZER_FLAG(name, defaultValue, desc) config.GetValueAsBool(name, TC(#name));
+		UBA_VISUALIZER_FLAGS2
+		#undef UBA_VISUALIZER_FLAG
+		return true;
+	}
+
+	bool VisualizerConfig::Save(Logger& logger)
+	{
+		Config config;
+		config.AddValue(L"X", x);
+		config.AddValue(L"Y", y);
+		config.AddValue(L"Width", width);
+		config.AddValue(L"Height", height);
+		#define UBA_VISUALIZER_FLAG(name, defaultValue, desc) config.AddValue(L"Show" #name, show##name);
+		UBA_VISUALIZER_FLAGS1
+		#undef UBA_VISUALIZER_FLAG
+		#define UBA_VISUALIZER_FLAG(name, defaultValue, desc) config.AddValue(TC(#name), name);
+		UBA_VISUALIZER_FLAGS2
+		#undef UBA_VISUALIZER_FLAG
+		return config.SaveToFile(logger, filename.c_str());
+	}
+
+	enum VisualizerFlag
+	{
+		#define UBA_VISUALIZER_FLAG(name, defaultValue, desc) VisualizerFlag_##name,
+		UBA_VISUALIZER_FLAGS1
+		#undef UBA_VISUALIZER_FLAG
+		VisualizerFlag_Count
+	};
+
 	class DrawTextLogger : public Logger
 	{
 	public:
@@ -136,14 +195,11 @@ namespace uba
 		TString& m_out;
 	};
 
-	Visualizer::Visualizer(Logger& logger)
+	Visualizer::Visualizer(VisualizerConfig& config, Logger& logger)
 	:	m_logger(logger)
+	,	m_config(config)
 	,	m_trace(logger)
 	{
-		for (bool& visible : m_visibleComponents)
-			visible = true;
-		m_visibleComponents[ComponentType_DetailedData] = false;
-		m_visibleComponents[ComponentType_Workers] = false;
 	}
 
 	Visualizer::~Visualizer()
@@ -157,12 +213,6 @@ namespace uba
 		delete m_client;
 	}
 
-	void Visualizer::SetTheme(bool dark)
-	{
-		m_isThemeSet = true;
-		m_useDarkMode = dark;
-	}
-
 	bool Visualizer::ShowUsingListener(const wchar_t* channelName)
 	{
 		TraceChannel channel(m_logger);
@@ -347,17 +397,9 @@ namespace uba
 		Unselect();
 	}
 
-	void Visualizer::ThreadLoop()
+	void Visualizer::InitBrushes()
 	{
-		if (!m_isThemeSet)
-		{
-			DWORD value = 1;
-			DWORD valueSize = sizeof(value);
-			if (RegGetValueW(HKEY_CURRENT_USER, L"Software\\Microsoft\\Windows\\CurrentVersion\\Themes\\Personalize", L"AppsUseLightTheme", RRF_RT_REG_DWORD, NULL, &value, &valueSize) == ERROR_SUCCESS)
-				m_useDarkMode = value == 0;
-		}
-
-		if (m_useDarkMode)
+		if (m_config.DarkMode)
 		{
 			m_textColor = RGB(190, 190, 190);
 			m_textWarningColor = RGB(190, 190, 0);
@@ -436,16 +478,21 @@ namespace uba
 		m_recvPen = CreatePen(PS_SOLID, 1, m_recvColor);
 		m_cpuPen = CreatePen(PS_SOLID, 1, m_cpuColor);
 		m_memPen = CreatePen(PS_SOLID, 1, m_memColor);
+	}
+
+	void Visualizer::ThreadLoop()
+	{
+		InitBrushes();
 
 		LOGBRUSH br = { 0 };
 		GetObject(m_backgroundBrush, sizeof(br), &br);
 		m_processUpdatePen = CreatePen(PS_SOLID, 2, RGB(GetRValue(br.lbColor), GetGValue(br.lbColor), GetBValue(br.lbColor)));
 
 		HINSTANCE hInstance = GetModuleHandle(NULL);
-		u32 winPosX = 100;
-		u32 winPosY = 100;
-		u32 winWidth = 1500;
-		u32 winHeight = 1500;
+		u32 winPosX = m_config.x;
+		u32 winPosY = m_config.y;
+		u32 winWidth = m_config.width;
+		u32 winHeight = m_config.height;
 
 		WNDCLASSEX wndClassEx;
 		ZeroMemory(&wndClassEx, sizeof(wndClassEx));
@@ -483,7 +530,7 @@ namespace uba
 		BOOL cloak = TRUE;
 		DwmSetWindowAttribute(m_hwnd, DWMWA_CLOAK, &cloak, sizeof(cloak));
 
-		if (m_useDarkMode)
+		//if (m_config.DarkMode)
 		{
 			SetWindowTheme(m_hwnd, L"DarkMode_Explorer", NULL);
 			SendMessageW(m_hwnd, WM_THEMECHANGED, 0, 0);
@@ -585,6 +632,26 @@ namespace uba
 			RedrawWindow(m_hwnd, NULL, NULL, RDW_INVALIDATE);
 	}
 
+	void Visualizer::SaveSettings()
+	{
+		RECT rect;
+		GetWindowRect(m_hwnd, &rect);
+
+		m_config.x = rect.left;
+		m_config.y = rect.top;
+		m_config.width = rect.right - rect.left;
+		m_config.height = rect.bottom - rect.top;
+		m_config.Save(m_logger);
+	}
+
+	void Visualizer::DirtyBitmaps()
+	{
+		for (auto& session : m_traceView.sessions)
+			for (auto& processor : session.processors)
+				for (auto& process : processor.processes)
+						process.bitmapDirty = true;
+	}
+
 	void Visualizer::PaintClient(const Function<void(HDC hdc, HDC memDC, RECT& clientRect)>& paintFunc)
 	{
 		HDC hdc = GetDC(m_hwnd);
@@ -657,7 +724,7 @@ namespace uba
 		progressRect.left += ProgressRectLeft;
 		progressRect.bottom -= 30;
 
-		bool shouldDrawText = m_showText && m_zoomValue > 0.27f;
+		bool shouldDrawText = m_zoomValue > 0.27f;
 
 		SetBkMode(hdc, TRANSPARENT);
 		SetTextColor(hdc, m_textColor);
@@ -668,7 +735,7 @@ namespace uba
 		SelectObject(textDC, m_font);
 		SelectObject(textDC, GetStockObject(NULL_BRUSH));
 		SetBkMode(textDC, TRANSPARENT);
-		SetBkColor(hdc, m_useDarkMode ? RGB(70, 70, 70) : RGB(180, 180, 180));
+		SetBkColor(hdc, m_config.DarkMode ? RGB(70, 70, 70) : RGB(180, 180, 180));
 
 		//TEXTMETRIC metric;
 		//GetTextMetrics(textDC, &metric);
@@ -776,7 +843,7 @@ namespace uba
 			}
 			posY += SessionStepY;
 
-			bool showGraph = m_visibleComponents[ComponentType_SendRecv] || m_visibleComponents[ComponentType_CpuMem];
+			bool showGraph = m_config.showNetworkStats || m_config.showCpuMemStats;
 			if (showGraph && hasUpdates)
 			{
 				if (posY + GraphHeight >= progressRect.top && posY + GraphHeight - 5 < progressRect.bottom)
@@ -826,7 +893,7 @@ namespace uba
 
 						if (!isFirstUpdate && x > clientRect.left && prevX <= clientRect.right)
 						{
-							if (m_visibleComponents[ComponentType_SendRecv] && updateSend != 0 && updateRecv != 0)
+							if (m_config.showNetworkStats && updateSend != 0 && updateRecv != 0)
 							{
 								SelectObject(hdc, m_sendPen);
 								MoveToEx(hdc, prevX, prevSendY, NULL);
@@ -835,7 +902,7 @@ namespace uba
 								MoveToEx(hdc, prevX, prevRecvY, NULL);
 								LineTo(hdc, x, recvY);
 							}
-							if (m_visibleComponents[ComponentType_CpuMem])
+							if (m_config.showCpuMemStats)
 							{
 								SelectObject(hdc, m_cpuPen);
 								MoveToEx(hdc, prevX, prevCpuY, NULL);
@@ -859,7 +926,7 @@ namespace uba
 				posY += GraphHeight;
 			}
 
-			if (m_visibleComponents[ComponentType_DetailedData])
+			if (m_config.showDetailedData)
 			{
 				auto drawText = [&](const StringBufferBase& text, RECT& rect)
 					{
@@ -874,7 +941,7 @@ namespace uba
 				PaintDetailedStats(posY, progressRect, session, isRemote, playTime, drawText);
 			}
 
-			if (m_visibleComponents[ComponentType_Bars])
+			if (m_config.showProcessBars)
 			{
 				processLocation.processorIndex = 0;
 				for (auto& processor : session.processors)
@@ -947,7 +1014,7 @@ namespace uba
 							++rect.top;
 
 							int processWidth = rect.right - rect.left;
-							if (shouldDrawText && processWidth > 3)
+							if (shouldDrawText && m_config.ShowProcessText && processWidth > 3)
 							{
 								if (!process.bitmap || process.bitmapDirty)
 								{
@@ -980,7 +1047,7 @@ namespace uba
 
 									rect2.left += 3; // Move in text a bit
 									
-									bool dropShadow = m_useDarkMode;
+									bool dropShadow = m_config.DarkMode;
 									if (dropShadow)
 									{
 										SetTextColor(textDC, RGB(5, 60, 5));
@@ -1046,7 +1113,7 @@ namespace uba
 						lastStop = Max(lastStop, processor.processes.rbegin()->stop);
 			}
 
-			if (m_visibleComponents[ComponentType_Workers] && isFirst)
+			if (m_config.showWorkers && isFirst)
 			{
 				u32 trackIndex = 0;
 				for (auto& workTrack : m_traceView.workTracks)
@@ -1155,7 +1222,7 @@ namespace uba
 							++rect.top;
 
 							int processWidth = rect.right - rect.left;
-							if (shouldDrawText && processWidth > 3)
+							if (shouldDrawText && m_config.ShowProcessText && processWidth > 3)
 							{
 								if (!work.bitmap)
 								{
@@ -1228,10 +1295,10 @@ namespace uba
 
 		float timelineSelected = m_timelineSelected;
 
-		if (m_visibleComponents[ComponentType_Timeline])
+		if (m_config.showTimeline && !m_traceView.sessions.empty())
 			PaintTimeline(hdc, clientRect);
 
-		if (m_visibleComponents[ComponentType_Cursor] && m_mouseOverWindow)
+		if (m_config.showCursorLine && m_mouseOverWindow)
 		{
 			float timeScale = (m_horizontalScaleValue * m_zoomValue)*50.0f;
 			float startOffset = -(m_scrollPosX / timeScale);
@@ -1293,13 +1360,14 @@ namespace uba
 			int bottom = SessionStepY - 3;
 			int left = progressRect.right - 16;
 			int right = progressRect.right - 7;
-			for (int i = sizeof_array(m_visibleComponents) - 1; i >= 0; --i)
+			bool* values = &m_config.showNetworkStats;
+			for (int i = VisualizerFlag_Count - 1; i >= 0; --i)
 			{
 				SelectObject(hdc, m_buttonSelected == u32(i) ? m_textPen : m_checkboxPen);
 				SelectObject(hdc, GetStockObject(NULL_BRUSH));
 				Rectangle(hdc, left, top, right, bottom);
 
-				if (m_visibleComponents[i])
+				if (values[i])
 				{
 					MoveToEx(hdc, left + 2, top + 2, NULL);
 					LineTo(hdc, right - 2, bottom - 2);
@@ -1313,7 +1381,7 @@ namespace uba
 
 			top -= 2;
 			left -= 20;
-			if (m_visibleComponents[ComponentType_CpuMem])
+			if (m_config.showCpuMemStats)
 			{
 				SelectObject(hdc, m_font);
 
@@ -1325,7 +1393,7 @@ namespace uba
 				ExtTextOutW(hdc, left, top, 0, NULL, L"MEM", 3, NULL);
 				left -= 25;
 			}
-			if (m_visibleComponents[ComponentType_SendRecv])
+			if (m_config.showNetworkStats)
 			{
 				SetTextColor(hdc, m_sendColor);
 				ExtTextOutW(hdc, left, top, 0, NULL, L"SND", 3, NULL);
@@ -1560,19 +1628,16 @@ namespace uba
 		}
 		else if (m_buttonSelected != ~0u)
 		{
+			bool* values = &m_config.showNetworkStats;
+
 			const wchar_t* tooltip[] =
 			{
-				L"network stats",
-				L"cpu/mem stats",
-				L"process bars",
-				L"timeline",
-				L"detailed data (use -UbaDetailedTrace for even more)",
-				L"workers (threads on host taking care of requests from helpers)",
-				L"cursor (vertical line)",
+				#define UBA_VISUALIZER_FLAG(name, defaultValue, desc) desc,
+				UBA_VISUALIZER_FLAGS1
+				#undef UBA_VISUALIZER_FLAG
 			};
-			static_assert(sizeof(tooltip)/sizeof(tooltip[0]) == ComponentType_Count);
 
-			bool vis = m_visibleComponents[m_buttonSelected];
+			bool vis = values[m_buttonSelected];
 			SelectObject(hdc, m_popupFont);
 			DrawTextLogger logger(m_hwnd, hdc, m_popupFontHeight, m_tooltipBackgroundBrush);
 			logger.Info(L"%ls %ls", vis ? L"Hide" : L"Show", tooltip[m_buttonSelected]);
@@ -1632,7 +1697,7 @@ namespace uba
 
 		u64 writeFilesTime = process.writeFilesTime;
 
-		if (!done || process.exitCode != 0 || !m_showCreateWriteColors || (TimeToMs(writeFilesTime, m_traceView.frequency) < 300 && TimeToMs(process.createFilesTime, m_traceView.frequency) < 300))
+		if (!done || process.exitCode != 0 || !m_config.ShowReadWriteColors || (TimeToMs(writeFilesTime, m_traceView.frequency) < 300 && TimeToMs(process.createFilesTime, m_traceView.frequency) < 300))
 		{
 			if (writingBitmap)
 				rect.right = 256;
@@ -1927,7 +1992,7 @@ namespace uba
 			int bottom = SessionStepY - 3;
 			int left = progressRect.right - 16;
 			int right = progressRect.right - 7;
-			for (int i = sizeof_array(m_visibleComponents) - 1; i >= 0; --i)
+			for (int i = VisualizerFlag_Count - 1; i >= 0; --i)
 			{
 				if (pos.x >= left && pos.x <= right && pos.y >= top && pos.y <= bottom)
 				{
@@ -1980,7 +2045,7 @@ namespace uba
 
 			posY += SessionStepY;
 
-			bool showGraph = m_visibleComponents[ComponentType_SendRecv] || m_visibleComponents[ComponentType_CpuMem];
+			bool showGraph = m_config.showNetworkStats || m_config.showCpuMemStats;
 			if (showGraph && !session.updates.empty())
 			{
 				if (pos.y >= posY && pos.y < posY + GraphHeight)
@@ -2027,7 +2092,7 @@ namespace uba
 				posY += GraphHeight;
 			}
 
-			if (m_visibleComponents[ComponentType_DetailedData])
+			if (m_config.showDetailedData)
 			{
 				auto drawText = [&](const StringBufferBase& text, RECT& rect)
 					{
@@ -2038,7 +2103,7 @@ namespace uba
 				PaintDetailedStats(posY, progressRect, session, i != 0, playTime, drawText);
 			}
 
-			if (m_visibleComponents[ComponentType_Bars])
+			if (m_config.showProcessBars)
 			{
 				u32 processorIndex = 0;
 				for (auto& processor : session.processors)
@@ -2103,7 +2168,7 @@ namespace uba
 						lastStop = Max(lastStop, processor.processes.rbegin()->stop);
 			}
 
-			if (m_visibleComponents[ComponentType_Workers] && isFirst)
+			if (m_config.showWorkers && isFirst)
 			{
 				int trackIndex = 0;
 				for (auto& workTrack : m_traceView.workTracks)
@@ -2157,7 +2222,7 @@ namespace uba
 				}
 			}
 
-			if (m_visibleComponents[ComponentType_Timeline])
+			if (m_config.showTimeline && !m_traceView.sessions.empty())
 			{
 				int timelineTop = Min(posY, int(progressRect.bottom));
 				if (pos.y >= timelineTop && pos.y < timelineTop + 40)
@@ -2249,7 +2314,7 @@ namespace uba
 
 	void Visualizer::UnselectAndRedraw()
 	{
-		if (Unselect() || m_visibleComponents[ComponentType_Cursor])
+		if (Unselect() || m_config.showCursorLine)
 			RedrawWindow(m_hwnd, NULL, NULL, RDW_INVALIDATE);
 	}
 
@@ -2453,52 +2518,68 @@ namespace uba
 				break;
 
 			int delta = GET_WHEEL_DELTA_WPARAM(wParam);// / WHEEL_DELTA;
+			bool controlDown = GetAsyncKeyState(VK_CONTROL) & (1<<15);
 
-			RECT r;
-			GetClientRect(hWnd, &r);
+			if (m_config.ScaleHorizontalWithScrollWheel || controlDown)
+			{
+				RECT r;
+				GetClientRect(hWnd, &r);
 
-			// Use mouse cursor as scroll anchor point
-			POINT cursorPos = {};
-			GetCursorPos(&cursorPos);
-			ScreenToClient(m_hwnd, &cursorPos);
-			const float scrollAnchorOffsetX = float(cursorPos.x) - ProgressRectLeft;
+				// Use mouse cursor as scroll anchor point
+				POINT cursorPos = {};
+				GetCursorPos(&cursorPos);
+				ScreenToClient(m_hwnd, &cursorPos);
+				const float scrollAnchorOffsetX = float(cursorPos.x) - ProgressRectLeft;
 
-			SHORT controlState = GetAsyncKeyState(VK_CONTROL);
-			if (controlState & (1<<15))
-			{
-				float newValue = Max(m_zoomValue + float(delta)*0.0005f, 0.05f);
-				m_scrollPosY = Min(0.0f, float(m_scrollPosY)*newValue/m_zoomValue);//LOWORD(lParam);
-				m_scrollPosX = Min(0.0f, float(m_scrollPosX - scrollAnchorOffsetX)*newValue/m_zoomValue + scrollAnchorOffsetX);//LOWORD(lParam);
+				if (controlDown)
+				{
+					float newValue = Max(m_zoomValue + float(delta)*0.0005f, 0.05f);
+					m_scrollPosY = Min(0.0f, float(m_scrollPosY)*newValue/m_zoomValue);//LOWORD(lParam);
+					m_scrollPosX = Min(0.0f, float(m_scrollPosX - scrollAnchorOffsetX)*newValue/m_zoomValue + scrollAnchorOffsetX);//LOWORD(lParam);
 
-				m_zoomValue = newValue;
-			}
-			else
-			{
-				float newValue = m_horizontalScaleValue + m_horizontalScaleValue*float(delta)*0.0006f;
-				m_scrollPosX = Min(0.0f, float(m_scrollPosX - scrollAnchorOffsetX)*newValue/m_horizontalScaleValue + scrollAnchorOffsetX);//LOWORD(lParam);
-				m_horizontalScaleValue = newValue;
-			}
+					m_zoomValue = newValue;
+				}
+				else
+				{
+					float newValue = m_horizontalScaleValue + m_horizontalScaleValue*float(delta)*0.0006f;
+					m_scrollPosX = Min(0.0f, float(m_scrollPosX - scrollAnchorOffsetX)*newValue/m_horizontalScaleValue + scrollAnchorOffsetX);//LOWORD(lParam);
+					m_horizontalScaleValue = newValue;
+				}
 
-			UpdateAutoscroll();
-			UpdateSelection();
+				UpdateAutoscroll();
+				UpdateSelection();
 
-			int minScroll = r.right - m_contentWidth;
-			m_scrollPosX = Min(0.0f, Max(m_scrollPosX, float(minScroll)));
-			m_scrollPosY = Min(0.0f, Max(m_scrollPosY, float(r.bottom - m_contentHeight)));
+				int minScroll = r.right - m_contentWidth;
+				m_scrollPosX = Min(0.0f, Max(m_scrollPosX, float(minScroll)));
+				m_scrollPosY = Min(0.0f, Max(m_scrollPosY, float(r.bottom - m_contentHeight)));
 
-			//if (!m_traceView.finished && m_scrollPosX <= minScroll)
-			//	m_autoScroll = true;
+				//if (!m_traceView.finished && m_scrollPosX <= minScroll)
+				//	m_autoScroll = true;
 
 		
-			if (m_showCreateWriteColors)
-				for (auto& session : m_traceView.sessions)
-					for (auto& processor : session.processors)
-						for (auto& process : processor.processes)
-							if (TimeToMs(process.writeFilesTime, m_traceView.frequency) >= 300 || TimeToMs(process.createFilesTime, m_traceView.frequency) >= 300)
-								process.bitmapDirty = true;
+				if (m_config.ShowReadWriteColors)
+					for (auto& session : m_traceView.sessions)
+						for (auto& processor : session.processors)
+							for (auto& process : processor.processes)
+								if (TimeToMs(process.writeFilesTime, m_traceView.frequency) >= 300 || TimeToMs(process.createFilesTime, m_traceView.frequency) >= 300)
+									process.bitmapDirty = true;
 
-			UpdateScrollbars(true);
-			RedrawWindow(hWnd, NULL, NULL, RDW_INVALIDATE);
+				UpdateScrollbars(true);
+				RedrawWindow(hWnd, NULL, NULL, RDW_INVALIDATE);
+			}
+			else
+			{
+				RECT r;
+				GetClientRect(hWnd, &r);
+				float oldScrollY = m_scrollPosY;
+				m_scrollPosY = m_scrollPosY + delta;
+				m_scrollPosY = Min(Max(m_scrollPosY, float(r.bottom - m_contentHeight)), 0.0f);
+				if (oldScrollY != m_scrollPosY)
+				{
+					UpdateScrollbars(true);
+					RedrawWindow(hWnd, NULL, NULL, RDW_INVALIDATE);
+				}
+			}
 			break;
 		}
 		case WM_MOUSEMOVE:
@@ -2529,7 +2610,7 @@ namespace uba
 			}
 			else
 			{
-				if (UpdateSelection() || m_visibleComponents[ComponentType_Cursor])
+				if (UpdateSelection() || m_config.showCursorLine)
 					RedrawWindow(hWnd, NULL, NULL, RDW_INVALIDATE);
 				/*
 				else
@@ -2588,7 +2669,8 @@ namespace uba
 		{
 			if (m_buttonSelected != ~0u)
 			{
-				m_visibleComponents[m_buttonSelected] = !m_visibleComponents[m_buttonSelected];
+				bool* values = &m_config.showNetworkStats;
+				values[m_buttonSelected] = !values[m_buttonSelected];
 				HitTestResult res;
 				HitTest(res, { -1, -1 });
 				UpdateScrollbars(true);
@@ -2650,8 +2732,11 @@ namespace uba
 			HMENU hMenu = CreatePopupMenu();
 			ClientToScreen(hWnd, &point);
 
-			AppendMenuW(hMenu, MF_STRING, Popup_ShowText, m_showText ? L"&Hide Process text" : L"&Show Process text");
-			AppendMenuW(hMenu, MF_STRING, Popup_ShowCreateWriteColors, m_showCreateWriteColors ? L"&Hide Create/Write colors" : L"&Show Create/Write colors");
+			#define UBA_VISUALIZER_FLAG(name, defaultValue, desc) \
+				AppendMenuW(hMenu, MF_STRING | (m_config.name ? MF_CHECKED : 0), Popup_##name, desc);
+			UBA_VISUALIZER_FLAGS2
+			#undef UBA_VISUALIZER_FLAG
+
 			AppendMenuW(hMenu, MF_SEPARATOR, 0, NULL);
 
 			if (m_sessionSelectedIndex != ~0u)
@@ -2686,14 +2771,15 @@ namespace uba
 				AppendMenuW(hMenu, MF_STRING, Popup_SaveAs, L"&Save Trace");
 			}
 			AppendMenuW(hMenu, MF_SEPARATOR, 0, NULL);
+			AppendMenuW(hMenu, MF_STRING, Popup_SaveSettings, L"Save Position/Settings");
 			AppendMenuW(hMenu, MF_STRING, Popup_Quit, L"&Quit");
 			m_showPopup = true;
 			switch (TrackPopupMenu(hMenu, TPM_RETURNCMD | TPM_RIGHTBUTTON, point.x, point.y, 0, hWnd, NULL))
 			{
 			case Popup_SaveAs:
 			{
-				OPENFILENAME ofn;       // common dialog box structure
-				TCHAR szFile[260] = { 0 };       // if using TCHAR macros
+				OPENFILENAME ofn;
+				TCHAR szFile[260] = { 0 };
 
 				// Initialize OPENFILENAME
 				ZeroMemory(&ofn, sizeof(ofn));
@@ -2712,19 +2798,31 @@ namespace uba
 					m_trace.SaveAs(ofn.lpstrFile);
 				break;
 			}
-			case Popup_ShowText:
-				m_showText = !m_showText;
+			case Popup_ShowProcessText:
+				m_config.ShowProcessText = !m_config.ShowProcessText;
 				RedrawWindow(hWnd, NULL, NULL, RDW_INVALIDATE | RDW_UPDATENOW);
 				break;
-			case Popup_ShowCreateWriteColors:
-				m_showCreateWriteColors = !m_showCreateWriteColors;
-				for (auto& session : m_traceView.sessions)
-					for (auto& processor : session.processors)
-						for (auto& process : processor.processes)
-								process.bitmapDirty = true;
+			case Popup_ShowReadWriteColors:
+				m_config.ShowReadWriteColors = !m_config.ShowReadWriteColors;
+				DirtyBitmaps();
 				RedrawWindow(hWnd, NULL, NULL, RDW_INVALIDATE | RDW_UPDATENOW);
 				break;
-
+			case Popup_ScaleHorizontalWithScrollWheel:
+				m_config.ScaleHorizontalWithScrollWheel = !m_config.ScaleHorizontalWithScrollWheel;
+				break;
+			case Popup_DarkMode:
+			{
+				m_config.DarkMode = !m_config.DarkMode;
+				DirtyBitmaps();
+				InitBrushes();
+				SetWindowTheme(m_hwnd, m_config.DarkMode ? L"DarkMode_Explorer" : L"Explorer", NULL);
+				SendMessageW(m_hwnd, WM_THEMECHANGED, 0, 0);
+				BOOL useDarkMode = m_config.DarkMode;
+				u32 attribute = 20; // DWMWA_USE_IMMERSIVE_DARK_MODE
+				DwmSetWindowAttribute(m_hwnd, attribute, &useDarkMode, sizeof(useDarkMode));
+				RedrawWindow(hWnd, NULL, NULL, RDW_INVALIDATE | RDW_UPDATENOW);
+				break;
+			}
 			case Popup_Replay:
 				m_replay = 1;
 				m_paused = false;
@@ -2745,6 +2843,10 @@ namespace uba
 				PostMessage(m_hwnd, WM_NEWTRACE, 0, 0);
 				break;
 
+			case Popup_SaveSettings:
+				SaveSettings();
+				break;
+
 			case Popup_Quit: // Quit
 				m_looping = false;
 				break;
diff --git a/Engine/Source/Programs/UnrealBuildAccelerator/Visualizer/Private/UbaVisualizer.h b/Engine/Source/Programs/UnrealBuildAccelerator/Visualizer/Private/UbaVisualizer.h
--- a/Engine/Source/Programs/UnrealBuildAccelerator/Visualizer/Private/UbaVisualizer.h
+++ b/Engine/Source/Programs/UnrealBuildAccelerator/Visualizer/Private/UbaVisualizer.h
@@ -8,14 +8,52 @@
 
 namespace uba
 {
+	class Config;
+
+	#define UBA_VISUALIZER_FLAGS1 \
+		UBA_VISUALIZER_FLAG(NetworkStats, true, L"network stats") \
+		UBA_VISUALIZER_FLAG(CpuMemStats, true, L"cpu/mem stats") \
+		UBA_VISUALIZER_FLAG(ProcessBars, true, L"process bars") \
+		UBA_VISUALIZER_FLAG(Timeline, true, L"timeline") \
+		UBA_VISUALIZER_FLAG(DetailedData, false, L"detailed data (use -UbaDetailedTrace for even more)") \
+		UBA_VISUALIZER_FLAG(Workers, false, L"workers (threads on host taking care of requests from helpers)") \
+		UBA_VISUALIZER_FLAG(CursorLine, false, L"cursor (vertical line)") \
+
+	#define UBA_VISUALIZER_FLAGS2 \
+		UBA_VISUALIZER_FLAG(ShowProcessText, true, L"Show text in process bars") \
+		UBA_VISUALIZER_FLAG(ShowReadWriteColors, true, L"Show colors for read/write times in process bars") \
+		UBA_VISUALIZER_FLAG(ScaleHorizontalWithScrollWheel, false, L"Use scroll wheel to scale horizontally") \
+		UBA_VISUALIZER_FLAG(DarkMode, false, L"Use dark mode to draw visualizer") \
+
+	struct VisualizerConfig
+	{
+		VisualizerConfig(const tchar* filename);
+
+		bool Load(Logger& logger);
+		bool Save(Logger& logger);
+
+		TString filename;
+
+		int x = 100;
+		int y = 100;
+		int width = 1500;
+		int height = 1500;
+
+		#define UBA_VISUALIZER_FLAG(name, defaultValue, desc) bool show##name = defaultValue;
+		UBA_VISUALIZER_FLAGS1
+		#undef UBA_VISUALIZER_FLAG
+
+		#define UBA_VISUALIZER_FLAG(name, defaultValue, desc) bool name = defaultValue;
+		UBA_VISUALIZER_FLAGS2
+		#undef UBA_VISUALIZER_FLAG
+	};
+
 	class Visualizer
 	{
 	public:
-		Visualizer(Logger& logger);
+		Visualizer(VisualizerConfig& config, Logger& logger);
 		~Visualizer();
 
-		void SetTheme(bool dark);
-
 		bool ShowUsingListener(const wchar_t* channelName);
 		bool ShowUsingNamedTrace(const wchar_t* namedTrace);
 		bool ShowUsingSocket(NetworkBackend& backend, const wchar_t* host, u16 port = DefaultPort);
@@ -67,10 +105,13 @@ namespace uba
 		bool UpdateSelection();
 		void UpdateScrollbars(bool redraw);
 		void GetTitlePrefix(StringBufferBase& out);
+		void InitBrushes();
 		void ThreadLoop();
 		void Pause(bool pause);
 		void StartDragToScroll(const POINT& anchor);
 		void StopDragToScroll();
+		void SaveSettings();
+		void DirtyBitmaps();
 
 		StringBuffer<256> m_namedTrace;
 		StringBuffer<256> m_fileName;
@@ -114,12 +155,9 @@ namespace uba
 		HFONT m_font = 0;
 		HFONT m_popupFont = 0;
 		int m_popupFontHeight = 0;
-		bool m_useDarkMode = true;
-		bool m_isThemeSet = false;
-		bool m_showText = true;
-		bool m_showCreateWriteColors = true;
 
 		Logger& m_logger;
+		VisualizerConfig m_config;
 		NetworkClient* m_client = nullptr;
 		TraceReader m_trace;
 		TraceView m_traceView;
@@ -161,19 +199,6 @@ namespace uba
 		bool m_mouseOverWindow = false;
 		bool m_showPopup = false;
 
-		enum ComponentType
-		{
-			ComponentType_SendRecv,
-			ComponentType_CpuMem,
-			ComponentType_Bars,
-			ComponentType_Timeline,
-			ComponentType_DetailedData,
-			ComponentType_Workers,
-			ComponentType_Cursor,
-			ComponentType_Count
-		};
-		bool m_visibleComponents[ComponentType_Count];
-
 		HBITMAP m_cachedBitmap = 0;
 		RECT m_cachedBitmapRect = { INT_MIN, INT_MIN, INT_MIN, INT_MIN };
 
diff --git a/Engine/Source/Programs/UnrealBuildAccelerator/Visualizer/Private/UbaVisualizerApp.cpp b/Engine/Source/Programs/UnrealBuildAccelerator/Visualizer/Private/UbaVisualizerApp.cpp
--- a/Engine/Source/Programs/UnrealBuildAccelerator/Visualizer/Private/UbaVisualizerApp.cpp
+++ b/Engine/Source/Programs/UnrealBuildAccelerator/Visualizer/Private/UbaVisualizerApp.cpp
@@ -1,6 +1,7 @@
 // Copyright Epic Games, Inc. All Rights Reserved.
 
 #include "UbaVisualizer.h"
+#include "UbaConfig.h"
 #include "UbaNetworkBackendTcp.h"
 #include "UbaVersion.h"
 
@@ -68,8 +69,6 @@ int WINAPI wWinMain(_In_ HINSTANCE hInstance, _In_opt_ HINSTANCE hPrevInstance,
 	StringBuffer<> channel;
 	u32 port = DefaultPort;
 	u32 replay = 0;
-	bool useDark = false;
-	bool isThemeSet = false;
 
 	int argc;
 	auto argv = CommandLineToArgvW(GetCommandLine(), &argc);
@@ -135,14 +134,6 @@ int WINAPI wWinMain(_In_ HINSTANCE hInstance, _In_opt_ HINSTANCE hPrevInstance,
 			if (!value.IsEmpty())
 				value.Parse(replay);
 		}
-		else if (name.Equals(TC("-theme")))
-		{
-			isThemeSet = true;
-			if (value.Equals(TC("dark")))
-				useDark = true;
-			else if (!value.Equals(TC("light")))
-				return PrintHelp(TC("Invalid value for -theme. Must be 'light' or 'dark'"));
-		}
 		else
 		{
 			StringBuffer<> msg;
@@ -159,13 +150,16 @@ int WINAPI wWinMain(_In_ HINSTANCE hInstance, _In_opt_ HINSTANCE hPrevInstance,
 	MessageBoxLogWriter logWriter;
 	LoggerWithWriter logger(logWriter);
 
-	NetworkBackendTcp networkBackend(logWriter);
-	Visualizer visualizer(logger);
-	logWriter.m_visualizer = &visualizer;
+	StringBuffer<> configPath;
+	configPath.count = ExpandEnvironmentStringsW(L"%PROGRAMDATA%", configPath.data, configPath.capacity) - 1;
+	configPath.Append(L"\\Epic\\UbaVisualizer\\UbaVisualizer.toml");
 
-	if (isThemeSet)
-		visualizer.SetTheme(useDark);
+	VisualizerConfig visualizerConfig(configPath.data);
+	visualizerConfig.Load(logger);
 
+	NetworkBackendTcp networkBackend(logWriter);
+	Visualizer visualizer(visualizerConfig, logger);
+	logWriter.m_visualizer = &visualizer;
 
 	if (channel.count)
 	{
