[UBA]
* Added new functions:
   Uba_Create - Create all components needed in one call using a config file
   Uba_RunProcess - Run a process using context created with Uba_Create
   Uba_RegisterNewFile - Tell uba external code has created a file
   Uba_Destroy - Destroy context
* Renamed DestroyProcessHandle to ProcessHandle_Destroy to match the other functions

diff --git a/Engine/Source/Programs/UnrealBuildAccelerator/Host/Private/UbaExports.cpp b/Engine/Source/Programs/UnrealBuildAccelerator/Host/Private/UbaExports.cpp
--- a/Engine/Source/Programs/UnrealBuildAccelerator/Host/Private/UbaExports.cpp
+++ b/Engine/Source/Programs/UnrealBuildAccelerator/Host/Private/UbaExports.cpp
@@ -4,6 +4,7 @@
 #include "UbaAWS.h"
 #include "UbaCacheClient.h"
 #include "UbaConfig.h"
+#include "UbaCoordinatorWrapper.h"
 #include "UbaNetworkBackendQuic.h"
 #include "UbaNetworkBackendTcp.h"
 #include "UbaNetworkClient.h"
@@ -60,8 +61,8 @@ namespace uba
 	class NetworkClientWithBackend : public NetworkClient
 	{
 	public:
-		NetworkClientWithBackend(bool& outSuccess, const NetworkClientCreateInfo& info, NetworkBackend* nb)
-		: NetworkClient(outSuccess, info), backend(nb)
+		NetworkClientWithBackend(bool& outSuccess, const NetworkClientCreateInfo& info, NetworkBackend* nb, const tchar* name)
+		: NetworkClient(outSuccess, info, name), backend(nb)
 		{
 		}
 
@@ -151,6 +152,9 @@ extern "C"
 			networkBackend = new NetworkBackendTcp(writer);
 
 		NetworkServerCreateInfo info(writer);
+
+		info.Apply(GetConfig());
+
 		info.workerCount = workerCount;
 		info.sendSize = sendSize;
 		info.receiveTimeoutSeconds = receiveTimeoutSeconds;
@@ -216,16 +220,19 @@ extern "C"
 	{
 		using namespace uba;
 
+		StorageServerCreateInfo info(server, rootDir, writer);
+		info.Apply(GetConfig());
+
 		#if UBA_USE_AWS
 		StringBuffer<> fixedRootDir;
-		fixedRootDir.count = GetFullPathNameW(rootDir, fixedRootDir.capacity, fixedRootDir.data, NULL);
+		fixedRootDir.count = GetFullPathNameW(info.rootDir, fixedRootDir.capacity, fixedRootDir.data, NULL);
 		fixedRootDir.Replace('/', PathSeparator).EnsureEndsWithSlash();
-		rootDir = fixedRootDir.data;
+		info.rootDir = fixedRootDir.data;
 		AWS aws;
 		if (!zone || !*zone)
 		{
 			LoggerWithWriter logger(writer, TC(""));
-			if (aws.QueryAvailabilityZone(logger, rootDir))
+			if (aws.QueryAvailabilityZone(logger, info.rootDir))
 				zone = aws.GetAvailabilityZone();
 		}
 		#endif
@@ -237,7 +244,6 @@ extern "C"
 			zone = zoneTemp.data;
 		}
 
-		StorageServerCreateInfo info(server, rootDir, writer);
 		info.casCapacityBytes = casCapacityBytes;
 		info.storeCompressed = storeCompressed;
 		info.zone = zone;
@@ -264,7 +270,7 @@ extern "C"
 		storageServer->DeleteCasForFile(file);
 	}
 
-	uba::u32 ProcessHandle_GetExitCode(uba::ProcessHandle* handle)
+	uba::u32 ProcessHandle_GetExitCode(const uba::ProcessHandle* handle)
 	{
 		return handle->GetExitCode();
 	}
@@ -298,11 +304,21 @@ extern "C"
 		return uba::TimeToTick(handle->GetTotalWallTime());
 	}
 
+	bool ProcessHandle_WaitForExit(uba::ProcessHandle* handle, uba::u32 millisecondsTimeout)
+	{
+		return handle->WaitForExit(millisecondsTimeout);
+	}
+
 	void ProcessHandle_Cancel(uba::ProcessHandle* handle, bool terminate)
 	{
 		handle->Cancel(terminate);
 	}
 
+	void ProcessHandle_Destroy(uba::ProcessHandle* handle)
+	{
+		delete handle;
+	}
+
 	void DestroyProcessHandle(uba::ProcessHandle* handle)
 	{
 		delete handle;
@@ -316,6 +332,7 @@ extern "C"
 	uba::SessionServerCreateInfo* SessionServerCreateInfo_Create(uba::StorageServer& storage, uba::NetworkServer& client, uba::LogWriter& writer, const uba::tchar* rootDir, const uba::tchar* traceOutputFile, bool disableCustomAllocator, bool launchVisualizer, bool resetCas, bool writeToDisk, bool detailedTrace, bool allowWaitOnMem, bool allowKillOnMem, bool storeObjFilesCompressed)
 	{
 		auto info = new uba::SessionServerCreateInfo(storage, client, writer);
+		info->Apply(uba::GetConfig());
 		info->rootDir = TStrdup(rootDir);
 		info->traceOutputFile = TStrdup(traceOutputFile);
 		info->disableCustomAllocator = disableCustomAllocator;
@@ -490,8 +507,10 @@ extern "C"
 	uba::Scheduler* Scheduler_Create(uba::SessionServer* session, uba::u32 maxLocalProcessors, bool enableProcessReuse)
 	{
 		uba::SchedulerCreateInfo info{*session};
+		info.Apply(uba::GetConfig());
 		info.maxLocalProcessors = maxLocalProcessors;
 		info.enableProcessReuse = enableProcessReuse;
+		info.processConfigs = &uba::GetConfig();
 		return new uba::Scheduler(info);
 	}
 
@@ -500,14 +519,14 @@ extern "C"
 		scheduler->Start();
 	}
 
-	void Scheduler_EnqueueProcess(uba::Scheduler* scheduler, const uba::ProcessStartInfo& info, float weight, const void* knownInputs, uba::u32 knownInputsBytes, uba::u32 knownInputsCount)
+	uba::u32 Scheduler_EnqueueProcess(uba::Scheduler* scheduler, const uba::ProcessStartInfo& info, float weight, const void* knownInputs, uba::u32 knownInputsBytes, uba::u32 knownInputsCount)
 	{
 		uba::EnqueueProcessInfo epi(info);
 		epi.weight = weight;
 		epi.knownInputs = knownInputs;
 		epi.knownInputsBytes = knownInputsBytes;
 		epi.knownInputsCount = knownInputsCount;
-		scheduler->EnqueueProcess(epi);
+		return scheduler->EnqueueProcess(epi);
 	}
 
 	void Scheduler_SetMaxLocalProcessors(uba::Scheduler* scheduler, uba::u32 maxLocalProcessors)
@@ -553,7 +572,7 @@ extern "C"
 		ncci.receiveTimeoutSeconds = 60;
 		ncci.cryptoKey128 = crypto128;
 		bool ctorSuccess = false;
-		auto networkClient = new NetworkClientWithBackend(ctorSuccess, ncci, server.backend);
+		auto networkClient = new NetworkClientWithBackend(ctorSuccess, ncci, server.backend, TC("UbaCache"));
 		if (!ctorSuccess)
 		{
 			delete networkClient;
@@ -615,4 +634,130 @@ extern "C"
 		uba::FindImports(binary, [&](const uba::tchar* importName, bool isKnown) { func(importName, userData); });
 #endif
 	}
+
+	struct UbaInstance
+	{
+		uba::Scheduler* scheduler;
+		uba::TString workDir;
+		uba::CoordinatorWrapper coordinator;
+	};
+
+	void* Uba_Create(const uba::tchar* configFile)
+	{
+		using namespace uba;
+		auto& config = GetConfig(configFile);
+		auto networkServer = (uba::NetworkServerWithBackend*)NetworkServer_Create();
+		auto storageServer = StorageServer_Create(*networkServer, nullptr, 0, true);
+
+		SessionServerCreateInfo ssci((Storage&)*storageServer, *networkServer);
+		ssci.Apply(config);
+		auto sessionServer = SessionServer_Create(ssci);
+
+		uba::SchedulerCreateInfo sci{*sessionServer};
+		sci.Apply(config);
+		sci.processConfigs = &config;
+		auto scheduler = new uba::Scheduler(sci);
+		scheduler->Start();
+
+		bool networkListen = true;
+		if (auto* ubaTable = config.GetTable(TC("Uba")))
+			ubaTable->GetValueAsBool(networkListen, TC("NetworkListen"));
+
+		if (networkListen)
+			NetworkServer_StartListen(networkServer);
+
+		auto ubaInstance = new UbaInstance();
+		ubaInstance->scheduler = scheduler;
+		
+		StringBuffer<> temp;
+		GetCurrentDirectoryW(temp);
+		ubaInstance->workDir = temp.data;
+
+		if (auto coordinatorTable = config.GetTable(TC("Coordinator")))
+		{
+			const tchar* coordinatorName;
+			if (coordinatorTable->GetValueAsString(coordinatorName, TC("Name")))
+			{
+				auto& logger = sessionServer->GetLogger();
+				const tchar* rootDir = nullptr;
+				coordinatorTable->GetValueAsString(rootDir, TC("RootDir"));
+				if (!rootDir)
+					rootDir = sessionServer->GetRootDir();
+				StringBuffer<512> coordinatorWorkDir(rootDir);
+				coordinatorWorkDir.EnsureEndsWithSlash().Append(coordinatorName);
+				StringBuffer<512> binariesDir;
+				if (!GetDirectoryOfCurrentModule(logger, binariesDir))
+					return nullptr;
+
+				CoordinatorCreateInfo cinfo;
+				cinfo.workDir = coordinatorWorkDir.data;
+				cinfo.binariesDir = binariesDir.data;
+
+				coordinatorTable->GetValueAsString(cinfo.pool, TC("Pool"));
+				UBA_ASSERT(cinfo.pool);
+
+				cinfo.maxCoreCount = 500;
+				coordinatorTable->GetValueAsU32(cinfo.maxCoreCount, TC("MaxCoreCount"));
+
+				cinfo.logging = false;
+				coordinatorTable->GetValueAsBool(cinfo.logging, TC("Log"));
+
+				const tchar* uri = nullptr;
+				if (coordinatorTable->GetValueAsString(uri, TC("Uri")))
+					uba::SetEnvironmentVariableW(TC("UE_HORDE_URL"), uri);
+
+				if (!ubaInstance->coordinator.Create(logger, coordinatorName, cinfo, *networkServer->backend, *networkServer, scheduler))
+					return nullptr;
+			}
+		}
+
+		return ubaInstance;
+	}
+
+	uba::u32 Uba_RunProcess(void* uba, const uba::tchar* app, const uba::tchar* args, const uba::tchar* workDir, const uba::tchar* desc, void* userData, ProcessHandle_ExitCallback* exit)
+	{
+		using namespace uba;
+
+		auto& ubaInstance = *(UbaInstance*)uba;
+
+		if (!workDir)
+			workDir = ubaInstance.workDir.data();
+
+		auto scheduler = ubaInstance.scheduler;
+		ProcessStartInfo info;
+		info.application = app;
+		info.arguments = args;
+		info.workingDir = workDir;
+		info.description = desc;
+		info.userData = userData;
+		info.exitedFunc = exit;
+		return Scheduler_EnqueueProcess(scheduler, info, 1.0f, nullptr, 0, 0);
+	}
+
+	void Uba_RegisterNewFile(void* uba, const uba::tchar* file)
+	{
+		using namespace uba;
+		auto& ubaInstance = *(UbaInstance*)uba;
+		ubaInstance.scheduler->GetSession().RegisterNewFile(file);
+	}
+
+	void Uba_Destroy(void* uba)
+	{
+		using namespace uba;
+		auto ubaInstance = (UbaInstance*)uba;
+		auto scheduler = ubaInstance->scheduler;
+		auto sessionServer = &scheduler->GetSession();
+		auto storageServer = (StorageServer*)&sessionServer->GetStorage();
+		auto networkServer = &sessionServer->GetServer();
+
+		NetworkServer_Stop(networkServer);
+		SessionServer_CancelAll(sessionServer);
+
+		delete ubaInstance;
+
+		Scheduler_Destroy(scheduler);
+		SessionServer_Destroy(sessionServer);
+		StorageServer_Destroy(storageServer);
+		NetworkServer_Destroy(networkServer);
+	}
 }
diff --git a/Engine/Source/Programs/UnrealBuildAccelerator/Host/Public/UbaExports.h b/Engine/Source/Programs/UnrealBuildAccelerator/Host/Public/UbaExports.h
--- a/Engine/Source/Programs/UnrealBuildAccelerator/Host/Public/UbaExports.h
+++ b/Engine/Source/Programs/UnrealBuildAccelerator/Host/Public/UbaExports.h
@@ -74,14 +74,15 @@ extern "C"
 	UBA_API void ProcessStartInfo_Destroy(uba::ProcessStartInfo* info);
 
 	// ProcessHandle
-	UBA_API uba::u32 ProcessHandle_GetExitCode(uba::ProcessHandle* handle);
+	UBA_API uba::u32 ProcessHandle_GetExitCode(const uba::ProcessHandle* handle);
 	UBA_API const uba::tchar* ProcessHandle_GetExecutingHost(uba::ProcessHandle* handle);
 	UBA_API const uba::tchar* ProcessHandle_GetLogLine(const uba::ProcessHandle* handle, uba::u32 index);
 	UBA_API uba::u64 ProcessHandle_GetHash(uba::ProcessHandle* handle);
 	UBA_API uba::u64 ProcessHandle_GetTotalProcessorTime(uba::ProcessHandle* handle);
 	UBA_API uba::u64 ProcessHandle_GetTotalWallTime(uba::ProcessHandle* handle);
+	UBA_API bool ProcessHandle_WaitForExit(uba::ProcessHandle* handle, uba::u32 millisecondsTimeout);
 	UBA_API void ProcessHandle_Cancel(uba::ProcessHandle* handle, bool terminate);
-	UBA_API void DestroyProcessHandle(uba::ProcessHandle* handle);
+	UBA_API void ProcessHandle_Destroy(uba::ProcessHandle* handle);
 	UBA_API const uba::ProcessStartInfo* Process_GetStartInfo(uba::Process& process);
 
 	// SessionServer
@@ -92,6 +93,7 @@ extern "C"
 	UBA_API uba::SessionServerCreateInfo* SessionServerCreateInfo_Create(uba::StorageServer& storage, uba::NetworkServer& client, uba::LogWriter& writer, const uba::tchar* rootDir, const uba::tchar* traceOutputFile,
 		bool disableCustomAllocator, bool launchVisualizer, bool resetCas, bool writeToDisk, bool detailedTrace, bool allowWaitOnMem = false, bool allowKillOnMem = false, bool storeObjFilesCompressed = false);
 	UBA_API void SessionServerCreateInfo_Destroy(uba::SessionServerCreateInfo* info);
+
 	UBA_API uba::SessionServer* SessionServer_Create(const uba::SessionServerCreateInfo& info, const uba::u8* environment = nullptr, uba::u32 environmentSize = 0);
 	UBA_API void SessionServer_SetRemoteProcessAvailable(uba::SessionServer* server, SessionServer_RemoteProcessAvailableCallback* available, void* userData);
 	UBA_API void SessionServer_SetRemoteProcessReturned(uba::SessionServer* server, SessionServer_RemoteProcessReturnedCallback* returned, void* userData);
@@ -116,7 +118,7 @@ extern "C"
 	// Scheduler
 	UBA_API uba::Scheduler* Scheduler_Create(uba::SessionServer* session, uba::u32 maxLocalProcessors = ~0u, bool enableProcessReuse = false);
 	UBA_API void Scheduler_Start(uba::Scheduler* scheduler);
-	UBA_API void Scheduler_EnqueueProcess(uba::Scheduler* scheduler, const uba::ProcessStartInfo& info, float weight = 1.0f, const void* knownInputs = nullptr, uba::u32 knownInputsBytes = 0, uba::u32 knownInputsCount = 0);
+	UBA_API uba::u32 Scheduler_EnqueueProcess(uba::Scheduler* scheduler, const uba::ProcessStartInfo& info, float weight = 1.0f, const void* knownInputs = nullptr, uba::u32 knownInputsBytes = 0, uba::u32 knownInputsCount = 0);
 	UBA_API void Scheduler_SetMaxLocalProcessors(uba::Scheduler* scheduler, uba::u32 maxLocalProcessors);
 	UBA_API void Scheduler_Stop(uba::Scheduler* scheduler);
 	UBA_API void Scheduler_Destroy(uba::Scheduler* scheduler);
@@ -136,4 +138,14 @@ extern "C"
 
 	using ImportFunc = void(const uba::tchar* importName, void* userData);
 	UBA_API void Uba_FindImports(const uba::tchar* binary, ImportFunc* func, void* userData);
+
+	// High level interface using config file instead. Uses scheduler under the hood
+	UBA_API void* Uba_Create(const uba::tchar* configFile);
+	UBA_API uba::u32 Uba_RunProcess(void* uba, const uba::tchar* app, const uba::tchar* args, const uba::tchar* workDir, const uba::tchar* desc, void* userData, ProcessHandle_ExitCallback* exit);
+	UBA_API void Uba_RegisterNewFile(void* uba, const uba::tchar* file);
+	UBA_API void Uba_Destroy(void* uba);
+
+
+	// DEPRECATED
+	UBA_API void DestroyProcessHandle(uba::ProcessHandle* handle);
 }
