[UBA]
* Fixed non windows compile errors

diff --git a/Engine/Source/Programs/UnrealBuildAccelerator/Common/Private/UbaConfig.cpp b/Engine/Source/Programs/UnrealBuildAccelerator/Common/Private/UbaConfig.cpp
--- a/Engine/Source/Programs/UnrealBuildAccelerator/Common/Private/UbaConfig.cpp
+++ b/Engine/Source/Programs/UnrealBuildAccelerator/Common/Private/UbaConfig.cpp
@@ -28,7 +28,11 @@ namespace uba
 		auto findIt = m_values.find(key);
 		if (findIt == m_values.end())
 			return m_parent ? m_parent->GetValueAsInt(out, key) : false;
+		#if PLATFORM_WINDOWS
 		out = (int)wcstol(findIt->second.c_str(), 0, 10);
+		#else
+		out = atoi(findIt->second.c_str());
+		#endif
 		return true;
 	}
 
@@ -228,7 +232,11 @@ namespace uba
 		for (auto& kv : m_values)
 		{
 			char line[1024];
+			#if PLATFORM_WINDOWS
 			int written = sprintf_s(line, sizeof_array(line), "%S = %S\r\n", kv.first.c_str(), kv.second.c_str());
+			#else
+			int written = snprintf(line, sizeof_array(line), "%s = %s\r\n", kv.first.c_str(), kv.second.c_str());
+			#endif
 			if (!fa.Write(line, written))
 				return false;
 		}
