[UBA]
* Fixed compile error in detour lib

diff --git a/Engine/Source/Programs/UnrealBuildAccelerator/Detours/Private/Windows/UbaDetoursFunctionsWin.cpp b/Engine/Source/Programs/UnrealBuildAccelerator/Detours/Private/Windows/UbaDetoursFunctionsWin.cpp
--- a/Engine/Source/Programs/UnrealBuildAccelerator/Detours/Private/Windows/UbaDetoursFunctionsWin.cpp
+++ b/Engine/Source/Programs/UnrealBuildAccelerator/Detours/Private/Windows/UbaDetoursFunctionsWin.cpp
@@ -865,15 +865,16 @@ void Init(const DetoursPayload& payload, u64 startTime)
 		TString realCmdLine(buf.data);
 		realCmdLine += (cmdLine + g_exeDir.count);
 		g_virtualCommandLineW = g_memoryBlock.Strdup(realCmdLine.c_str());
-		}
+	}
 	//else
 	//	g_virtualCommandLineW = g_memoryBlock.Strdup(cmdLine);
 
 	if (g_virtualCommandLineW)
 	{
-		g_virtualCommandLineA = (char*)g_memoryBlock.Allocate(realCmdLine.size() + 1, 1, L"");
+		u64 len = wcslen(g_virtualCommandLineW);
+		g_virtualCommandLineA = (char*)g_memoryBlock.Allocate(len + 1, 1, L"");
 		size_t res;
-		wcstombs_s(&res, g_virtualCommandLineA, realCmdLine.size() + 1, g_virtualCommandLineW, realCmdLine.size());
+		wcstombs_s(&res, g_virtualCommandLineA, len + 1, g_virtualCommandLineW, len);
 	}
 
 	#if UBA_DEBUG_LOG_ENABLED
