[UBA]
* Fixed hang because of recursive calls of detouring messages

diff --git a/Engine/Source/Programs/UnrealBuildAccelerator/Detours/Private/Windows/UbaDetoursMainWin.cpp b/Engine/Source/Programs/UnrealBuildAccelerator/Detours/Private/Windows/UbaDetoursMainWin.cpp
--- a/Engine/Source/Programs/UnrealBuildAccelerator/Detours/Private/Windows/UbaDetoursMainWin.cpp
+++ b/Engine/Source/Programs/UnrealBuildAccelerator/Detours/Private/Windows/UbaDetoursMainWin.cpp
@@ -107,19 +107,19 @@ namespace uba
 				TerminateCurrentProcess(1337);
 
 			DWORD exitCode;
-			if (!GetExitCodeProcess(g_hostProcess, &exitCode))
+			if (!True_GetExitCodeProcess(g_hostProcess, &exitCode))
 				TerminateCurrentProcess(1338);
 
 			if (exitCode != STILL_ACTIVE)
 				TerminateCurrentProcess(1339);
 
-			DWORD cancelRes = WaitForSingleObject(g_cancelEvent, 0);
+			DWORD cancelRes = True_WaitForSingleObject(g_cancelEvent, 0);
 			if (cancelRes == WAIT_OBJECT_0)
 				ExitProcess(1339);
 			else if (cancelRes != WAIT_TIMEOUT)
 				TerminateCurrentProcess(1353);
 
-			res = WaitForSingleObject(g_readEvent, 500);
+			res = True_WaitForSingleObject(g_readEvent, 500);
 
 		} while (true);
 	}
