[UBA]
* Fixed so child processes spawned by processes that has "detach" enabled also become detached. This to fix problem with some compilers that are detached spawns linker which have log outputs
* Fixed some logging macros for debugging

diff --git a/Engine/Source/Programs/UnrealBuildAccelerator/Detours/Private/Windows/UbaDetoursFunctionsKernelBase.inl b/Engine/Source/Programs/UnrealBuildAccelerator/Detours/Private/Windows/UbaDetoursFunctionsKernelBase.inl
--- a/Engine/Source/Programs/UnrealBuildAccelerator/Detours/Private/Windows/UbaDetoursFunctionsKernelBase.inl
+++ b/Engine/Source/Programs/UnrealBuildAccelerator/Detours/Private/Windows/UbaDetoursFunctionsKernelBase.inl
@@ -410,7 +410,7 @@ BOOL Detoured_ReadFile(HANDLE hFile, LPVOID lpBuffer, DWORD nNumberOfBytesToRead
 BOOL Detoured_WriteConsoleA(HANDLE hConsoleOutput, const VOID* lpBuffer, DWORD nNumberOfCharsToWrite, LPDWORD lpNumberOfCharsWritten, LPVOID lpReserved)
 {
 	DETOURED_CALL(WriteConsoleA);
-	//DEBUG_LOG_TRUE_AND_DETOURED(L"WriteConsoleA (%hs)", (char*)lpBuffer);
+	//DEBUG_LOG_DETOURED(L"WriteConsoleA", L"(%hs)", (char*)lpBuffer); // Too much spam
 	Shared_WriteConsole((const char*)lpBuffer, nNumberOfCharsToWrite, false);
 	if (lpNumberOfCharsWritten)
 		*lpNumberOfCharsWritten = nNumberOfCharsToWrite;
@@ -420,7 +420,7 @@ BOOL Detoured_WriteConsoleA(HANDLE hConsoleOutput, const VOID* lpBuffer, DWORD n
 BOOL Detoured_WriteConsoleW(HANDLE hConsoleOutput, const VOID* lpBuffer, DWORD nNumberOfCharsToWrite, LPDWORD lpNumberOfCharsWritten, LPVOID lpReserved)
 {
 	DETOURED_CALL(WriteConsoleW);
-	//DEBUG_LOG_DETOURED(L"WriteConsoleW", L""); // Too much spam
+	//DEBUG_LOG_DETOURED(L"WriteConsoleW", L"(%s)", (const wchar_t*)lpBuffer); // Too much spam
 	Shared_WriteConsole((const wchar_t*)lpBuffer, nNumberOfCharsToWrite, false);
 	if (lpNumberOfCharsWritten)
 		*lpNumberOfCharsWritten = nNumberOfCharsToWrite;
@@ -637,6 +637,7 @@ BOOL Detoured_WriteFile(HANDLE hFile, LPCVOID lpBuffer, DWORD nNumberOfBytesToWr
 
 		if (dh.type == HandleType_Std)
 		{
+			//DEBUG_LOG_DETOURED(L"WriteStdFile1", L"%llu", uintptr_t(hFile));
 			WriteStdFile(lpBuffer, nNumberOfBytesToWrite, hFile == g_stdHandle[0]);
 			*lpNumberOfBytesWritten = nNumberOfBytesToWrite;
 			SetLastError(ERROR_SUCCESS);
@@ -663,6 +664,7 @@ BOOL Detoured_WriteFile(HANDLE hFile, LPCVOID lpBuffer, DWORD nNumberOfBytesToWr
 	}
 	else if (hFile == g_stdHandle[1] || hFile == g_stdHandle[0])
 	{
+		//DEBUG_LOG_DETOURED(L"WriteStdFile2", L"%llu", uintptr_t(hFile));
 		WriteStdFile(lpBuffer, nNumberOfBytesToWrite, hFile == g_stdHandle[0]);
 		*lpNumberOfBytesWritten = nNumberOfBytesToWrite;
 		SetLastError(ERROR_SUCCESS);
@@ -954,14 +956,17 @@ DWORD Detoured_GetFileType(HANDLE hFile)
 		DetouredHandle& dh = asDetouredHandle(hFile);
 		SetLastError(ERROR_SUCCESS);
 		if (dh.type == HandleType_Std)
+		{
+			DEBUG_LOG_DETOURED(L"GetFileType", L"%llu (%ls) -> FILE_TYPE_CHAR", uintptr_t(hFile), HandleToName(hFile));
 			return FILE_TYPE_CHAR;
+		}
 		UBA_ASSERT(dh.type == HandleType_File);
-		DEBUG_LOG_DETOURED(L"GetFileType", L"%llu (%ls) -> %u", uintptr_t(hFile), HandleToName(hFile), FILE_TYPE_DISK);
+		DEBUG_LOG_DETOURED(L"GetFileType", L"%llu (%ls) -> FILE_TYPE_DISK", uintptr_t(hFile), HandleToName(hFile));
 		return FILE_TYPE_DISK;
 	}
 	if (isListDirectoryHandle(hFile))
 	{
-		DEBUG_LOG_DETOURED(L"GetFileType", L"%llu (%ls) -> %u", uintptr_t(hFile), HandleToName(hFile), FILE_TYPE_DISK);
+		DEBUG_LOG_DETOURED(L"GetFileType", L"%llu (%ls) -> FILE_TYPE_DISK", uintptr_t(hFile), HandleToName(hFile));
 		SetLastError(ERROR_SUCCESS);
 		return FILE_TYPE_DISK;
 	}
diff --git a/Engine/Source/Programs/UnrealBuildAccelerator/Detours/Private/Windows/UbaDetoursFunctionsWin.cpp b/Engine/Source/Programs/UnrealBuildAccelerator/Detours/Private/Windows/UbaDetoursFunctionsWin.cpp
--- a/Engine/Source/Programs/UnrealBuildAccelerator/Detours/Private/Windows/UbaDetoursFunctionsWin.cpp
+++ b/Engine/Source/Programs/UnrealBuildAccelerator/Detours/Private/Windows/UbaDetoursFunctionsWin.cpp
@@ -784,20 +784,29 @@ void Init(const DetoursPayload& payload, u64 startTime)
 
 	DetourAttachFunctions(g_runningRemote);
 
+	if (!g_isDetachedProcess)
+	{
+		// If GetStdHandle returns 0 it is likely that there is a parent process and that one has detached process set (which means now conhost is created)
+		// .. solve this by detaching this process too
+		HANDLE stdoutHandle = True_GetStdHandle(STD_OUTPUT_HANDLE);
+		if (stdoutHandle == 0)
+		{
+			g_isDetachedProcess = true;
+		}
+		else
+		{
+			HANDLE stderrHandle = True_GetStdHandle(STD_ERROR_HANDLE);
+			g_stdHandle[0] = GetFileType(stderrHandle) == FILE_TYPE_CHAR ? stderrHandle : 0;
+			g_stdHandle[1] = GetFileType(stdoutHandle) == FILE_TYPE_CHAR ? stdoutHandle : 0;
+		}
+	}
+
 	if (g_isDetachedProcess)
 	{
 		g_stdHandle[0] = makeDetouredHandle(new DetouredHandle(HandleType_Std)); // STD_ERR
 		g_stdHandle[1] = makeDetouredHandle(new DetouredHandle(HandleType_Std)); // STD_OUT
 		g_stdHandle[2] = makeDetouredHandle(new DetouredHandle(HandleType_Std)); // STD_IN
 	}
-	else
-	{
-		HANDLE stderrHandle = True_GetStdHandle(STD_ERROR_HANDLE);
-		g_stdHandle[0] = GetFileType(stderrHandle) == FILE_TYPE_CHAR ? stderrHandle : 0;
-		
-		HANDLE stdoutHandle = True_GetStdHandle(STD_OUTPUT_HANDLE);
-		g_stdHandle[1] = GetFileType(stdoutHandle) == FILE_TYPE_CHAR ? stdoutHandle : 0;
-	}
 
 	if (payload.trackInputs)
 		g_trackInputsMem = (u8*)g_memoryBlock.Allocate(TrackInputsMemCapacity, 1, L"TrackInputs");
