[UBA]
* Fixed process reuse stats that was missing child process stats for each reuse step

diff --git a/Engine/Source/Programs/UnrealBuildAccelerator/Common/Private/UbaProcess.cpp b/Engine/Source/Programs/UnrealBuildAccelerator/Common/Private/UbaProcess.cpp
--- a/Engine/Source/Programs/UnrealBuildAccelerator/Common/Private/UbaProcess.cpp
+++ b/Engine/Source/Programs/UnrealBuildAccelerator/Common/Private/UbaProcess.cpp
@@ -1128,6 +1128,7 @@ namespace uba
 					
 		ProcessStats processStats;
 		processStats.Read(reader, TraceVersion);
+		processStats.Add(m_processStats);
 		processStats.startupTime = m_processStats.startupTime;
 		processStats.wallTime = GetTime() - m_startTime;
 		processStats.cpuTime = 0;
@@ -1135,11 +1136,12 @@ namespace uba
 		
 		KernelStats kernelStats;
 		kernelStats.Read(reader, TraceVersion);
-					
+		kernelStats.Add(m_kernelStats);
+
 		processStats.Write(statsWriter);
 		m_sessionStats.Write(statsWriter);
 		m_storageStats.Write(statsWriter);
-		m_kernelStats.Write(statsWriter);
+		kernelStats.Write(statsWriter);
 		BinaryReader statsReader(statsWriter.GetData(), 0, statsWriter.GetPosition());
 
 		bool newProcess = false;
