[UBA]
* Added "active processes" view in visualizer. Easier to get an overview
* Changed so session stats is written into trace file as soon as possible
* Added support for some more types in config file add/write
* Added dedicated path for updateProgress into session in order to minimize trace file and have custom rendering for progress
* Added support for changing font size in visualizer
* Added option to enable autosaving settings in visualizer
* Fixed so it is possible to jump around in active named trace in UbaVisualizer
* Added hyperlink support ro status update so user can click on the text to follow the link

diff --git a/Engine/Source/Programs/UnrealBuildAccelerator/Common/Private/UbaConfig.cpp b/Engine/Source/Programs/UnrealBuildAccelerator/Common/Private/UbaConfig.cpp
--- a/Engine/Source/Programs/UnrealBuildAccelerator/Common/Private/UbaConfig.cpp
+++ b/Engine/Source/Programs/UnrealBuildAccelerator/Common/Private/UbaConfig.cpp
@@ -10,7 +10,16 @@ namespace uba
 		auto findIt = m_values.find(key);
 		if (findIt == m_values.end())
 			return m_parent ? m_parent->GetValueAsString(out, key) : false;
-		out = findIt->second.c_str();
+		out = findIt->second.string.c_str();
+		return true;
+	}
+
+	bool ConfigTable::GetValueAsString(TString& out, const tchar* key) const
+	{
+		const tchar* str;
+		if (!GetValueAsString(str, key))
+			return false;
+		out = str;
 		return true;
 	}
 
@@ -19,7 +28,7 @@ namespace uba
 		auto findIt = m_values.find(key);
 		if (findIt == m_values.end())
 			return m_parent ? m_parent->GetValueAsU32(out, key) : false;
-		StringBuffer<> buf(findIt->second);
+		StringBuffer<> buf(findIt->second.string);
 		return buf.Parse(out);
 	}
 
@@ -29,9 +38,9 @@ namespace uba
 		if (findIt == m_values.end())
 			return m_parent ? m_parent->GetValueAsInt(out, key) : false;
 		#if PLATFORM_WINDOWS
-		out = (int)wcstol(findIt->second.c_str(), 0, 10);
+		out = (int)wcstol(findIt->second.string.c_str(), 0, 10);
 		#else
-		out = atoi(findIt->second.c_str());
+		out = atoi(findIt->second.string.c_str());
 		#endif
 		return true;
 	}
@@ -41,7 +50,7 @@ namespace uba
 		auto findIt = m_values.find(key);
 		if (findIt == m_values.end())
 			return m_parent ? m_parent->GetValueAsBool(out, key) : false;
-		const tchar* value = findIt->second.c_str();
+		const tchar* value = findIt->second.string.c_str();
 		if (Equals(value, TC("true")) || Equals(value, TC("1")))
 		{
 			out = true;
@@ -68,12 +77,24 @@ namespace uba
 	{
 		tchar buf[256];
 		TSprintf_s(buf, sizeof_array(buf), TC("%i"), value);
-		m_values[key] = buf;
+		m_values[key] = Value{ValueType_Value, buf};
+	}
+
+	void ConfigTable::AddValue(const tchar* key, u32 value)
+	{
+		tchar buf[256];
+		TSprintf_s(buf, sizeof_array(buf), TC("%u"), value);
+		m_values[key] = Value{ValueType_Value, buf};
 	}
 
 	void ConfigTable::AddValue(const tchar* key, bool value)
 	{
-		m_values[key] = value ? TC("true") : TC("false");
+		m_values[key] = Value{ValueType_Value, value ? TC("true") : TC("false")};
+	}
+
+	void ConfigTable::AddValue(const tchar* key, const tchar* str)
+	{
+		m_values[key] = Value{ValueType_String, str};
 	}
 
 	bool Config::LoadFromFile(Logger& logger, const tchar* configFile)
@@ -195,18 +216,20 @@ namespace uba
 				++i;
 				token = consumeEmpty();
 
+				ValueType type = ValueType_Value;
 				StringBuffer<1024> value;
 				if (token == '\"')
 				{
 					++i;
 					token = consumeLine(value, '\"');
+					type = ValueType_String;
 				}
 				else
 				{
 					token = consumeLine(value, ' ');
 				}
 
-				activeTable->m_values[key.data] = value.data;
+				activeTable->m_values[key.data] = Value{type, value.data};
 				if (token == 0)
 					break;
 				++i;
@@ -232,10 +255,11 @@ namespace uba
 		for (auto& kv : m_values)
 		{
 			char line[1024];
+			const char* quote = kv.second.type == ValueType_String ? "\"" : "";
 			#if PLATFORM_WINDOWS
-			int written = sprintf_s(line, sizeof_array(line), "%S = %S\r\n", kv.first.c_str(), kv.second.c_str());
+			int written = sprintf_s(line, sizeof_array(line), "%S = %s%S%s\r\n", kv.first.c_str(), quote, kv.second.string.c_str(), quote);
 			#else
-			int written = snprintf(line, sizeof_array(line), "%s = %s\r\n", kv.first.c_str(), kv.second.c_str());
+			int written = snprintf(line, sizeof_array(line), "%s = %s%s%s\r\n", kv.first.c_str(), quote, kv.second.string.c_str(), quote);
 			#endif
 			if (!fa.Write(line, written))
 				return false;
diff --git a/Engine/Source/Programs/UnrealBuildAccelerator/Common/Private/UbaScheduler.cpp b/Engine/Source/Programs/UnrealBuildAccelerator/Common/Private/UbaScheduler.cpp
--- a/Engine/Source/Programs/UnrealBuildAccelerator/Common/Private/UbaScheduler.cpp
+++ b/Engine/Source/Programs/UnrealBuildAccelerator/Common/Private/UbaScheduler.cpp
@@ -599,7 +599,7 @@ namespace uba
 	void Scheduler::UpdateQueueCounter(int offset)
 	{
 		m_queuedProcesses += u32(offset);
-		m_session.UpdateStatus(1, 1, TC("Queue"), 4, StringBuffer<32>().Appendf(TC("%u"), m_queuedProcesses.load()).data, LogEntryType_Info);
+		m_session.UpdateProgress(m_queuedProcesses + m_activeLocalProcesses + m_activeRemoteProcesses + m_finishedProcesses, m_finishedProcesses, 0);
 	}
 
 	void Scheduler::UpdateActiveProcessCounter(bool isLocal, int offset)
@@ -608,7 +608,6 @@ namespace uba
 			m_activeLocalProcesses += u32(offset);
 		else
 			m_activeRemoteProcesses += u32(offset);
-		m_session.UpdateStatus(2, 1, TC("Active"), 4, StringBuffer<32>().Appendf(TC("Local %u Remote %u"), m_activeLocalProcesses.load(), m_activeRemoteProcesses.load()).data, LogEntryType_Info);
 	}
 
 	void Scheduler::FinishProcess(const ProcessHandle& handle)
@@ -616,7 +615,7 @@ namespace uba
 		++m_finishedProcesses;
 		if (m_processFinished)
 			m_processFinished(handle);
-		m_session.UpdateStatus(3, 1, TC("Finished"), 4, StringBuffer<32>().Appendf(TC("%u"), m_finishedProcesses.load()).data, LogEntryType_Info);
+		m_session.UpdateProgress(m_queuedProcesses + m_activeLocalProcesses + m_activeRemoteProcesses + m_finishedProcesses, m_finishedProcesses, 0);
 	}
 
 	bool Scheduler::EnqueueFromFile(const tchar* yamlFilename)
diff --git a/Engine/Source/Programs/UnrealBuildAccelerator/Common/Private/UbaSessionServer.cpp b/Engine/Source/Programs/UnrealBuildAccelerator/Common/Private/UbaSessionServer.cpp
--- a/Engine/Source/Programs/UnrealBuildAccelerator/Common/Private/UbaSessionServer.cpp
+++ b/Engine/Source/Programs/UnrealBuildAccelerator/Common/Private/UbaSessionServer.cpp
@@ -346,6 +346,7 @@ namespace uba
 		if (m_remoteExecutionEnabled)
 			m_logger.Info(TC("Disable remote execution (remote sessions will finish current processes)"));
 		m_remoteExecutionEnabled = false;
+		m_trace.RemoteExecutionDisabled();
 	}
 
 	void SessionServer::SetCustomCasKeyFromTrackedInputs(const tchar* fileName_, const tchar* workingDir_, const u8* trackedInputs, u32 trackedInputsBytes)
@@ -478,9 +479,14 @@ namespace uba
 		m_trace.ProcessExited(id, exitCode, statsWriter.GetData(), statsWriter.GetPosition(), Vector<ProcessLogLine>());
 	}
 
-	void SessionServer::UpdateStatus(u32 statusIndex, u32 statusNameIndent, const tchar* statusName, u32 statusTextIndent, const tchar* statusText, LogEntryType statusType)
+	void SessionServer::UpdateProgress(u32 processesTotal, u32 processesDone, u32 errorCount)
 	{
-		m_trace.StatusUpdate(statusIndex, statusNameIndent, statusName, statusTextIndent, statusText, statusType);
+		m_trace.ProgressUpdate(processesTotal, processesDone, errorCount);
+	}
+
+	void SessionServer::UpdateStatus(u32 statusRow, u32 statusColumn, const tchar* statusText, LogEntryType statusType, const tchar* statusLink)
+	{
+		m_trace.StatusUpdate(statusRow, statusColumn, statusText, statusType, statusLink);
 	}
 
 	NetworkServer& SessionServer::GetServer()
diff --git a/Engine/Source/Programs/UnrealBuildAccelerator/Common/Private/UbaTrace.cpp b/Engine/Source/Programs/UnrealBuildAccelerator/Common/Private/UbaTrace.cpp
--- a/Engine/Source/Programs/UnrealBuildAccelerator/Common/Private/UbaTrace.cpp
+++ b/Engine/Source/Programs/UnrealBuildAccelerator/Common/Private/UbaTrace.cpp
@@ -196,12 +196,12 @@ namespace uba
 		WriterScope writer(*this); \
 		if (!writer.IsValid()) \
 			return; \
-		writer.WriteByte(x); \
+		writer.WriteByte(TraceType_##x); \
 		writer.Write7BitEncoded(GetTime() - m_startTime);
 
 	void Trace::SessionAdded(u32 sessionId, u32 clientId, const tchar* name, const tchar* info)
 	{
-		BEGIN_TRACE_ENTRY(TraceType_SessionAdded);
+		BEGIN_TRACE_ENTRY(SessionAdded);
 		writer.WriteString(name);
 		writer.WriteString(info);
 		writer.Write7BitEncoded(clientId);
@@ -210,7 +210,7 @@ namespace uba
 
 	void Trace::SessionUpdate(u32 sessionId, u32 connectionCount, u64 send, u64 recv, u64 lastPing, u64 memAvail, u64 memTotal, float cpuLoad)
 	{
-		BEGIN_TRACE_ENTRY(TraceType_SessionUpdate);
+		BEGIN_TRACE_ENTRY(SessionUpdate);
 		writer.Write7BitEncoded(sessionId);
 		writer.Write7BitEncoded(connectionCount);
 		writer.Write7BitEncoded(send);
@@ -223,27 +223,27 @@ namespace uba
 
 	void Trace::SessionNotification(u32 sessionId, const tchar* text)
 	{
-		BEGIN_TRACE_ENTRY(TraceType_SessionNotification);
+		BEGIN_TRACE_ENTRY(SessionNotification);
 		writer.WriteU32(sessionId);
 		writer.WriteString(text);
 	}
 
 	void Trace::SessionSummary(u32 sessionId, const u8* data, u64 dataSize)
 	{
-		BEGIN_TRACE_ENTRY(TraceType_SessionSummary);
+		BEGIN_TRACE_ENTRY(SessionSummary);
 		writer.WriteU32(sessionId);
 		writer.WriteBytes(data, dataSize);
 	}
 
 	void Trace::SessionDisconnect(u32 sessionId)
 	{
-		BEGIN_TRACE_ENTRY(TraceType_SessionDisconnect);
+		BEGIN_TRACE_ENTRY(SessionDisconnect);
 		writer.WriteU32(sessionId);
 	}
 
 	void Trace::ProcessAdded(u32 sessionId, u32 processId, const tchar* description)
 	{
-		BEGIN_TRACE_ENTRY(TraceType_ProcessAdded);
+		BEGIN_TRACE_ENTRY(ProcessAdded);
 		writer.WriteU32(sessionId);
 		writer.WriteU32(processId);
 		writer.WriteString(description);
@@ -251,7 +251,7 @@ namespace uba
 
 	void Trace::ProcessEnvironmentUpdated(u32 processId, const tchar* reason, const u8* data, u64 dataSize)
 	{
-		BEGIN_TRACE_ENTRY(TraceType_ProcessEnvironmentUpdated);
+		BEGIN_TRACE_ENTRY(ProcessEnvironmentUpdated);
 		writer.WriteU32(processId);
 		writer.WriteString(reason);
 		writer.WriteBytes(data, dataSize);
@@ -259,7 +259,7 @@ namespace uba
 
 	void Trace::ProcessExited(u32 processId, u32 exitCode, const u8* data, u64 dataSize, const Vector<ProcessLogLine>& logLines)
 	{
-		BEGIN_TRACE_ENTRY(TraceType_ProcessExited);
+		BEGIN_TRACE_ENTRY(ProcessExited);
 		writer.WriteU32(processId);
 		writer.WriteU32(exitCode);
 		writer.WriteBytes(data, dataSize);
@@ -278,20 +278,20 @@ namespace uba
 
 	void Trace::ProcessReturned(u32 processId)
 	{
-		BEGIN_TRACE_ENTRY(TraceType_ProcessReturned);
+		BEGIN_TRACE_ENTRY(ProcessReturned);
 		writer.WriteU32(processId);
 	}
 
 	void Trace::ProxyCreated(u32 clientId, const tchar* proxyName)
 	{
-		BEGIN_TRACE_ENTRY(TraceType_ProxyCreated);
+		BEGIN_TRACE_ENTRY(ProxyCreated);
 		writer.Write7BitEncoded(clientId);
 		writer.WriteString(proxyName);
 	}
 
 	void Trace::ProxyUsed(u32 clientId, const tchar* proxyName)
 	{
-		BEGIN_TRACE_ENTRY(TraceType_ProxyUsed);
+		BEGIN_TRACE_ENTRY(ProxyUsed);
 		writer.Write7BitEncoded(clientId);
 		writer.WriteString(proxyName);
 	}
@@ -301,7 +301,7 @@ namespace uba
 		if (detailed)
 		{
 			u32 stringIndex = AddString(hint);
-			BEGIN_TRACE_ENTRY(TraceType_FileBeginFetch);
+			BEGIN_TRACE_ENTRY(FileBeginFetch);
 			writer.Write7BitEncoded(clientId);
 			writer.WriteCasKey(key);
 			writer.Write7BitEncoded(size);
@@ -309,7 +309,7 @@ namespace uba
 		}
 		else
 		{
-			BEGIN_TRACE_ENTRY(TraceType_FileFetchLight);
+			BEGIN_TRACE_ENTRY(FileFetchLight);
 			writer.Write7BitEncoded(clientId);
 			writer.Write7BitEncoded(size);
 		}
@@ -317,7 +317,7 @@ namespace uba
 
 	void Trace::FileEndFetch(u32 clientId, const CasKey& key)
 	{
-		BEGIN_TRACE_ENTRY(TraceType_FileEndFetch);
+		BEGIN_TRACE_ENTRY(FileEndFetch);
 		writer.Write7BitEncoded(clientId);
 		writer.WriteCasKey(key);
 	}
@@ -327,7 +327,7 @@ namespace uba
 		if (detailed)
 		{
 			u32 stringIndex = AddString(hint);
-			BEGIN_TRACE_ENTRY(TraceType_FileBeginStore);
+			BEGIN_TRACE_ENTRY(FileBeginStore);
 			writer.Write7BitEncoded(clientId);
 			writer.WriteCasKey(key);
 			writer.Write7BitEncoded(size);
@@ -335,7 +335,7 @@ namespace uba
 		}
 		else
 		{
-			BEGIN_TRACE_ENTRY(TraceType_FileStoreLight);
+			BEGIN_TRACE_ENTRY(FileStoreLight);
 			writer.Write7BitEncoded(clientId);
 			writer.Write7BitEncoded(size);
 		}
@@ -343,7 +343,7 @@ namespace uba
 
 	void Trace::FileEndStore(u32 clientId, const CasKey& key)
 	{
-		BEGIN_TRACE_ENTRY(TraceType_FileEndStore);
+		BEGIN_TRACE_ENTRY(FileEndStore);
 		writer.Write7BitEncoded(clientId);
 		writer.WriteCasKey(key);
 	}
@@ -351,38 +351,50 @@ namespace uba
 	void Trace::BeginWork(u32 workIndex, const tchar* desc)
 	{
 		u32 stringIndex = AddString(desc);
-		BEGIN_TRACE_ENTRY(TraceType_BeginWork);
+		BEGIN_TRACE_ENTRY(BeginWork);
 		writer.Write7BitEncoded(workIndex);
 		writer.Write7BitEncoded(stringIndex);
 	}
 
 	void Trace::EndWork(u32 workIndex)
 	{
-		BEGIN_TRACE_ENTRY(TraceType_EndWork);
+		BEGIN_TRACE_ENTRY(EndWork);
 		writer.Write7BitEncoded(workIndex);
 	}
 
-	void Trace::StatusUpdate(u32 statusIndex, u32 statusNameIndent, const tchar* statusName, u32 statusTextIndent, const tchar* statusText, LogEntryType statusType)
+	void Trace::ProgressUpdate(u32 processesTotal, u32 processesDone, u32 errorCount)
 	{
-		BEGIN_TRACE_ENTRY(TraceType_StatusUpdate);
-		writer.Write7BitEncoded(statusIndex);
-		writer.Write7BitEncoded(statusNameIndent);
-		writer.WriteString(statusName);
-		writer.Write7BitEncoded(statusTextIndent);
+		BEGIN_TRACE_ENTRY(ProgressUpdate);
+		writer.Write7BitEncoded(processesTotal);
+		writer.Write7BitEncoded(processesDone);
+		writer.Write7BitEncoded(errorCount);
+	}
+
+	void Trace::StatusUpdate(u32 statusRow, u32 statusColumn, const tchar* statusText, LogEntryType statusType, const tchar* statusLink)
+	{
+		BEGIN_TRACE_ENTRY(StatusUpdate);
+		writer.Write7BitEncoded(statusRow);
+		writer.Write7BitEncoded(statusColumn);
 		writer.WriteString(statusText);
 		writer.WriteByte(statusType);
+		writer.WriteString(statusLink ? statusLink : TC(""));
+	}
+
+	void Trace::RemoteExecutionDisabled()
+	{
+		BEGIN_TRACE_ENTRY(RemoteExecutionDisabled);
 	}
 
 	void Trace::CacheBeginFetch(u32 fetchId, const tchar* description)
 	{
-		BEGIN_TRACE_ENTRY(TraceType_CacheBeginFetch);
+		BEGIN_TRACE_ENTRY(CacheBeginFetch);
 		writer.Write7BitEncoded(fetchId);
 		writer.WriteString(description);
 	}
 
 	void Trace::CacheEndFetch(u32 fetchId, bool success, const u8* data, u64 dataSize)
 	{
-		BEGIN_TRACE_ENTRY(TraceType_CacheEndFetch);
+		BEGIN_TRACE_ENTRY(CacheEndFetch);
 		writer.Write7BitEncoded(fetchId);
 		writer.WriteBool(success);
 		writer.WriteBytes(data, dataSize);
@@ -390,13 +402,13 @@ namespace uba
 
 	void Trace::CacheBeginWrite(u32 processId)
 	{
-		BEGIN_TRACE_ENTRY(TraceType_CacheBeginWrite);
+		BEGIN_TRACE_ENTRY(CacheBeginWrite);
 		writer.Write7BitEncoded(processId);
 	}
 
 	void Trace::CacheEndWrite(u32 processId, bool success, u64 bytesSent)
 	{
-		BEGIN_TRACE_ENTRY(TraceType_CacheEndWrite);
+		BEGIN_TRACE_ENTRY(CacheEndWrite);
 		writer.Write7BitEncoded(processId);
 		writer.WriteBool(success);
 		writer.Write7BitEncoded(bytesSent);
diff --git a/Engine/Source/Programs/UnrealBuildAccelerator/Common/Public/UbaConfig.h b/Engine/Source/Programs/UnrealBuildAccelerator/Common/Public/UbaConfig.h
--- a/Engine/Source/Programs/UnrealBuildAccelerator/Common/Public/UbaConfig.h
+++ b/Engine/Source/Programs/UnrealBuildAccelerator/Common/Public/UbaConfig.h
@@ -14,6 +14,7 @@ namespace uba
 	{
 	public:
 		bool GetValueAsString(const tchar*& out, const tchar* key) const;
+		bool GetValueAsString(TString& out, const tchar* key) const;
 		bool GetValueAsU32(u32& out, const tchar* key) const;
 		bool GetValueAsInt(int& out, const tchar* key) const;
 		bool GetValueAsBool(bool& out, const tchar* key) const;
@@ -21,11 +22,15 @@ namespace uba
 		const ConfigTable* GetTable(const tchar* name) const;
 
 		void AddValue(const tchar* key, int value);
+		void AddValue(const tchar* key, u32 value);
 		void AddValue(const tchar* key, bool value);
+		void AddValue(const tchar* key, const tchar* str);
 
 	private:
 		ConfigTable* m_parent = nullptr;
-		Map<TString, TString> m_values;
+		enum ValueType { ValueType_Value, ValueType_String };
+		struct Value { ValueType type; TString string; };
+		Map<TString, Value> m_values;
 		UnorderedMap<TString, ConfigTable> m_tables;
 		friend Config;
 	};
diff --git a/Engine/Source/Programs/UnrealBuildAccelerator/Common/Public/UbaSessionServer.h b/Engine/Source/Programs/UnrealBuildAccelerator/Common/Public/UbaSessionServer.h
--- a/Engine/Source/Programs/UnrealBuildAccelerator/Common/Public/UbaSessionServer.h
+++ b/Engine/Source/Programs/UnrealBuildAccelerator/Common/Public/UbaSessionServer.h
@@ -56,8 +56,11 @@ namespace uba
 		// End external process.
 		void EndExternalProcess(u32 id, u32 exitCode);
 
+		// Update progress. Will show in visualizer
+		void UpdateProgress(u32 processesTotal, u32 processesDone, u32 errorCount);
+
 		// Add external status information to trace stream. Will show in visualizer
-		void UpdateStatus(u32 statusIndex, u32 statusNameIndent, const tchar* statusName, u32 statusTextIndent, const tchar* statusText, LogEntryType statusType);
+		void UpdateStatus(u32 statusRow, u32 statusColumn, const tchar* statusText, LogEntryType statusType, const tchar* statusLink);
 
 		// Get the network server used by this session
 		NetworkServer& GetServer();
diff --git a/Engine/Source/Programs/UnrealBuildAccelerator/Common/Public/UbaTrace.h b/Engine/Source/Programs/UnrealBuildAccelerator/Common/Public/UbaTrace.h
--- a/Engine/Source/Programs/UnrealBuildAccelerator/Common/Public/UbaTrace.h
+++ b/Engine/Source/Programs/UnrealBuildAccelerator/Common/Public/UbaTrace.h
@@ -55,6 +55,8 @@ namespace uba
 		UBA_TRACE_TYPE(CacheEndFetch) \
 		UBA_TRACE_TYPE(CacheBeginWrite) \
 		UBA_TRACE_TYPE(CacheEndWrite) \
+		UBA_TRACE_TYPE(ProgressUpdate) \
+		UBA_TRACE_TYPE(RemoteExecutionDisabled) \
 
 	enum TraceType : u8
 	{
@@ -66,7 +68,7 @@ namespace uba
 	using Color = u32;
 	inline Color toColor(u8 r, u8 g, u8 b) { return (r << 16) + (g << 8) + b; }
 
-	static constexpr u32 TraceVersion = 31;
+	static constexpr u32 TraceVersion = 32;
 	static constexpr u32 TraceReadCompatibilityVersion = 6;
 
 	class Trace : public WorkTracker
@@ -95,7 +97,9 @@ namespace uba
 		void FileEndStore(u32 clientId, const CasKey& key);
 		void BeginWork(u32 workIndex, const tchar* desc);
 		void EndWork(u32 workIndex);
-		void StatusUpdate(u32 statusIndex, u32 statusNameIndent, const tchar* statusName, u32 statusTextIndent, const tchar* statusText, LogEntryType statusType);
+		void ProgressUpdate(u32 processesTotal, u32 processesDone, u32 errorCount);
+		void StatusUpdate(u32 statusRow, u32 statusColumn, const tchar* statusText, LogEntryType statusType, const tchar* statusLink);
+		void RemoteExecutionDisabled();
 
 		void CacheBeginFetch(u32 fetchId, const tchar* description);
 		void CacheEndFetch(u32 fetchId, bool success, const u8* data, u64 dataSize);
diff --git a/Engine/Source/Programs/UnrealBuildAccelerator/Host/Private/UbaExports.cpp b/Engine/Source/Programs/UnrealBuildAccelerator/Host/Private/UbaExports.cpp
--- a/Engine/Source/Programs/UnrealBuildAccelerator/Host/Private/UbaExports.cpp
+++ b/Engine/Source/Programs/UnrealBuildAccelerator/Host/Private/UbaExports.cpp
@@ -429,9 +429,14 @@ extern "C"
 		server->EndExternalProcess(id, exitCode);
 	}
 
-	void SessionServer_UpdateStatus(uba::SessionServer* server, uba::u32 statusIndex, uba::u32 statusNameIndent, const uba::tchar* statusName, uba::u32 statusTextIndent, const uba::tchar* statusText, uba::LogEntryType statusType)
+	void SessionServer_UpdateProgress(uba::SessionServer* server, uba::u32 processesTotal, uba::u32 processesDone, uba::u32 errorCount)
 	{
-		server->UpdateStatus(statusIndex, statusNameIndent, statusName, statusTextIndent, statusText, statusType);
+		server->UpdateProgress(processesTotal, processesDone, errorCount);
+	}
+
+	void SessionServer_UpdateStatus(uba::SessionServer* server, uba::u32 statusRow, uba::u32 statusColumn, const uba::tchar* statusText, uba::LogEntryType statusType, const uba::tchar* statusLink)
+	{
+		server->UpdateStatus(statusRow, statusColumn, statusText, statusType, statusLink);
 	}
 
 	void SessionServer_RegisterCustomService(uba::SessionServer* server, SessionServer_CustomServiceFunction* function, void* userData)
diff --git a/Engine/Source/Programs/UnrealBuildAccelerator/Host/Public/UbaExports.h b/Engine/Source/Programs/UnrealBuildAccelerator/Host/Public/UbaExports.h
--- a/Engine/Source/Programs/UnrealBuildAccelerator/Host/Public/UbaExports.h
+++ b/Engine/Source/Programs/UnrealBuildAccelerator/Host/Public/UbaExports.h
@@ -111,7 +111,8 @@ extern "C"
 	UBA_API void SessionServer_SetCustomCasKeyFromTrackedInputs(uba::SessionServer* server, uba::ProcessHandle* handle, const uba::tchar* fileName, const uba::tchar* workingDir);
 	UBA_API uba::u32 SessionServer_BeginExternalProcess(uba::SessionServer* server, const uba::tchar* description);
 	UBA_API void SessionServer_EndExternalProcess(uba::SessionServer* server, uba::u32 id, uba::u32 exitCode);
-	UBA_API void SessionServer_UpdateStatus(uba::SessionServer* server, uba::u32 statusIndex, uba::u32 statusNameIndent, const uba::tchar* statusName, uba::u32 statusTextIndent, const uba::tchar* statusText, uba::LogEntryType statusType);
+	UBA_API void SessionServer_UpdateProgress(uba::SessionServer* server, uba::u32 processesTotal, uba::u32 processesDone, uba::u32 errorCount);
+	UBA_API void SessionServer_UpdateStatus(uba::SessionServer* server, uba::u32 statusRow, uba::u32 statusColumn, const uba::tchar* statusText, uba::LogEntryType statusType, const uba::tchar* statusLink);
 	UBA_API void SessionServer_RegisterCustomService(uba::SessionServer* server, SessionServer_CustomServiceFunction* function, void* userData = nullptr);
 	UBA_API void SessionServer_Destroy(uba::SessionServer* server);
 
diff --git a/Engine/Source/Programs/UnrealBuildAccelerator/Visualizer/Private/UbaTraceReader.cpp b/Engine/Source/Programs/UnrealBuildAccelerator/Visualizer/Private/UbaTraceReader.cpp
--- a/Engine/Source/Programs/UnrealBuildAccelerator/Visualizer/Private/UbaTraceReader.cpp
+++ b/Engine/Source/Programs/UnrealBuildAccelerator/Visualizer/Private/UbaTraceReader.cpp
@@ -8,6 +8,30 @@
 
 namespace uba
 {
+	TraceView::Process* TraceView::GetProcess(const ProcessLocation& loc)
+	{
+		return &(sessions[loc.sessionIndex].processors[loc.processorIndex].processes[loc.processIndex]);
+	}
+
+	void TraceView::Clear()
+	{
+		sessions.clear();
+		workTracks.clear();
+		strings.clear();
+		statusMap.clear();
+		cacheWrites.clear();
+		startTime = 0;
+		totalProcessActiveCount = 0;
+		totalProcessExitedCount = 0;
+		activeSessionCount = 0;
+		progressProcessesTotal = 0;
+		progressProcessesDone = 0;
+		progressErrorCount = 0;
+		remoteExecutionDisabled = false;
+		finished = true;
+	};
+
+
 	TraceReader::TraceReader(Logger& logger)
 	:	m_logger(logger)
 	,	m_channel(logger)
@@ -28,7 +52,7 @@ namespace uba
 
 	bool TraceReader::ReadFile(TraceView& out, const tchar* fileName, bool replay)
 	{
-		Reset();
+		Reset(out);
 
 		FileHandle readHandle;
 		if (!OpenFileSequentialRead(m_logger, fileName, readHandle))
@@ -68,8 +92,9 @@ namespace uba
 		else
 			out.frequency = GetFrequency();
 
-		out.startTime = reader.Read7BitEncoded();
+		out.realStartTime = reader.Read7BitEncoded();
 
+		out.startTime = out.realStartTime;
 		if (replay)
 			out.startTime = GetTime();
 		else if (traceSystemStartTimeUs)
@@ -116,7 +141,7 @@ namespace uba
 
 	bool TraceReader::StartReadClient(TraceView& out, NetworkClient& client)
 	{
-		Reset();
+		Reset(out);
 
 		u32 traceMemSize = 128 * 1024 * 1024;
 		m_memoryHandle = uba::CreateMemoryMappingW(m_logger, PAGE_READWRITE, traceMemSize);
@@ -171,49 +196,70 @@ namespace uba
 				break;
 		}
 		outChanged = !m_activeProcesses.empty() || m_memoryPos != m_memoryEnd;
-		return ReadMemory(out, false);
+		return ReadMemory(out, false, ~u64(0));
 	}
 
-	bool TraceReader::StartReadNamed(TraceView& out, const tchar* namedTrace, bool silentFail)
+	bool TraceReader::StartReadNamed(TraceView& out, const tchar* namedTrace, bool silentFail, bool replay)
 	{
-		Reset();
+		Reset(out);
 
-		m_memoryHandle.handle = ::OpenFileMappingW(PAGE_READWRITE, false, namedTrace);
-		if (!m_memoryHandle.IsValid())
+		if (namedTrace && m_namedTrace != namedTrace)
 		{
-			if (!silentFail)
-				m_logger.Error(L"OpenFileMappingW - Failed to open file mapping %ls (%ls)", namedTrace, LastErrorToText().data);
-			return false;
+			m_memoryHandle.handle = ::OpenFileMappingW(PAGE_READWRITE, false, namedTrace);
+			if (!m_memoryHandle.IsValid())
+			{
+				if (!silentFail)
+					m_logger.Error(L"OpenFileMappingW - Failed to open file mapping %ls (%ls)", namedTrace, LastErrorToText().data);
+				return false;
+			}
+			m_memoryBegin = MapViewOfFile(m_memoryHandle, FILE_MAP_READ, 0, 0);
+			if (!m_memoryBegin)
+				return false;
 		}
-		m_memoryBegin = MapViewOfFile(m_memoryHandle, FILE_MAP_READ, 0, 0);
-		if (!m_memoryBegin)
-			return false;
+
 		m_memoryPos = m_memoryBegin;
 		m_memoryEnd = m_memoryBegin;
-		m_hostProcess = 0;
 		out.finished = false;
 		out.sessions.emplace_back();
 		out.sessions.back().name = L"LOCAL";
 		bool changed;
-		return UpdateReadNamed(out, changed);
+		return UpdateReadNamed(out, replay ? 0ull : ~u64(0), changed);
 	}
 
-	bool TraceReader::UpdateReadNamed(TraceView& out, bool& outChanged)
+	bool TraceReader::UpdateReadNamed(TraceView& out, u64 maxTime, bool& outChanged)
 	{
 		outChanged = false;
 		if (!m_memoryBegin)
 			return true;
 		m_memoryEnd = m_memoryBegin + *(u32*)m_memoryBegin;
 		outChanged = !m_activeProcesses.empty() || m_memoryPos != m_memoryEnd;
-		bool res = ReadMemory(out, true);
-		if (m_hostProcess  && WaitForSingleObject(m_hostProcess, 0) != WAIT_TIMEOUT)
-			StopAllActive(out, GetTime() - m_startTime);
-		if (!res)
-			Unmap();
+		bool res = ReadMemory(out, true, maxTime);
+		if (m_hostProcess && WaitForSingleObject(m_hostProcess, 0) != WAIT_TIMEOUT)
+			StopAllActive(out, GetTime() - out.realStartTime);
+
+		if (res || m_namedTrace.empty())
+			return true;
+
+		// Move memory to local mapping in case we want to replay
+		u64 traceMemSize = m_memoryEnd - m_memoryBegin;
+		u8* memoryBegin = nullptr;
+		FileMappingHandle memoryHandle = uba::CreateMemoryMappingW(m_logger, PAGE_READWRITE, traceMemSize);
+		if (memoryHandle.IsValid())
+			if ((memoryBegin = MapViewOfFile(memoryHandle, FILE_MAP_ALL_ACCESS, 0, traceMemSize)) != nullptr)
+				memcpy(memoryBegin, m_memoryBegin, traceMemSize);
+
+		u64 pos = m_memoryPos - m_memoryBegin;
+		u64 end = m_memoryEnd - m_memoryBegin;
+		Unmap();
+		m_memoryHandle = memoryHandle;
+		m_memoryPos = memoryBegin + pos;
+		m_memoryEnd = memoryBegin + end;
+		m_memoryBegin = memoryBegin;
+
 		return res;
 	}
 
-	bool TraceReader::ReadMemory(TraceView& out, bool trackHost)
+	bool TraceReader::ReadMemory(TraceView& out, bool trackHost, u64 maxTime)
 	{
 		if (m_memoryEnd == m_memoryPos)
 			return true;
@@ -232,9 +278,12 @@ namespace uba
 				m_logger.Error(L"Incompatible trace version (%u). Current executable supports version %u to %u.", version, TraceReadCompatibilityVersion, TraceVersion);
 				return false;
 			}
+			bool replay = maxTime != ~u64(0);
+
 			out.version = version;
 			u32 hostProcessId = traceReader.ReadU32();
-			if (trackHost)
+			m_hostProcess = 0;
+			if (trackHost && !replay)
 				m_hostProcess = OpenProcess(SYNCHRONIZE, FALSE, hostProcessId);
 			u64 traceSystemStartTimeUs = 0;
 			if (version >= 18)
@@ -244,19 +293,27 @@ namespace uba
 			else
 				out.frequency = GetFrequency();
 
-			out.startTime = traceReader.Read7BitEncoded();
+			out.realStartTime = traceReader.Read7BitEncoded();
 			if (traceSystemStartTimeUs)
-				out.startTime = GetTime() - UsToTime(GetSystemTimeUs() - traceSystemStartTimeUs);
-			m_startTime = out.startTime;
+				out.realStartTime = GetTime() - UsToTime(GetSystemTimeUs() - traceSystemStartTimeUs);
+
+			out.startTime = out.realStartTime;
+			if (replay)
+				out.startTime = GetTime();
 		}
 
-		while (traceReader.GetPosition() != toRead)
+		u64 lastPos = traceReader.GetPosition();
+		while (lastPos != toRead)
 		{
-			UBA_ASSERT(traceReader.GetPosition() < toRead);
-			if (!ReadTrace(out, traceReader, ~u64(0)))
+			if (!ReadTrace(out, traceReader, maxTime))
 				return false;
+			u64 pos = traceReader.GetPosition();
+			if (pos == lastPos)
+				break;
+			UBA_ASSERT(pos <= toRead);
+			lastPos = pos;
 		}
-		m_memoryPos += toRead;
+		m_memoryPos += lastPos;
 		return true;
 	}
 
@@ -503,7 +560,8 @@ namespace uba
 			const u8* dataStart = reader.GetPositionData();
 			processStats.Read(reader, out.version);
 
-			process.isRemote = sessionIndex != 0;
+			UBA_ASSERT(process.isRemote == (sessionIndex != 0));
+
 			if (process.isRemote)
 			{
 				if (out.version >= 7)
@@ -584,6 +642,7 @@ namespace uba
 			process.exitCode = 0u;
 			process.stop = time;
 			process.bitmapDirty = true;
+			bool isRemote = process.isRemote;
 
 			processes.emplace_back();
 			auto& process2 = processes.back();
@@ -593,6 +652,7 @@ namespace uba
 			process2.start = time;
 			process2.stop = ~u64(0);
 			process2.exitCode = ~0u;
+			process2.isRemote = isRemote;
 
 			++session.processExitedCount;
 			++out.totalProcessExitedCount;
@@ -881,15 +941,40 @@ namespace uba
 				record.stop = time;
 			break;
 		}
+		case TraceType_ProgressUpdate:
+		{
+			out.progressProcessesTotal = u32(reader.Read7BitEncoded());
+			out.progressProcessesDone = u32(reader.Read7BitEncoded());
+			out.progressErrorCount = u32(reader.Read7BitEncoded());
+			break;
+		}
 		case TraceType_StatusUpdate:
 		{
-			u32 statusIndex = u32(reader.Read7BitEncoded());
-			auto& status = out.statusMap[statusIndex];
-			status.nameIndent = u32(reader.Read7BitEncoded());
-			status.name = reader.ReadString();
-			status.textIndent = u32(reader.Read7BitEncoded());
-			status.text = reader.ReadString();
-			status.type = (LogEntryType)reader.ReadByte();
+			if (out.version < 32)
+			{
+				reader.Read7BitEncoded();
+				reader.Read7BitEncoded();
+				reader.ReadString();
+				reader.Read7BitEncoded();
+				reader.ReadString();
+				reader.ReadByte();
+			}
+			else
+			{
+				u64 row = reader.Read7BitEncoded();
+				u64 column = reader.Read7BitEncoded();
+				u64 key = (row << 32) | column;
+				auto& status = out.statusMap[key];
+				status.text = reader.ReadString();
+				status.type = (LogEntryType)reader.ReadByte();
+				status.link = reader.ReadString();
+			}
+
+			break;
+		}
+		case TraceType_RemoteExecutionDisabled:
+		{
+			out.remoteExecutionDisabled = true;
 			break;
 		}
 		case TraceType_String:
@@ -963,8 +1048,9 @@ namespace uba
 		out.finished = true;
 	}
 
-	void TraceReader::Reset()
+	void TraceReader::Reset(TraceView& out)
 	{
+		out.Clear();
 		m_activeProcesses.clear();
 		m_activeWorkRecords.clear();
 		m_sessionIndexToSession.clear();
@@ -981,6 +1067,7 @@ namespace uba
 		if (m_memoryHandle.IsValid())
 			CloseFileMapping(m_memoryHandle);
 		m_memoryHandle = {};
+		m_namedTrace.clear();
 	}
 
 	bool TraceReader::SaveAs(const tchar* fileName)
@@ -1058,6 +1145,7 @@ namespace uba
 		process.start = time;
 		process.stop = ~u64(0);
 		process.exitCode = ~0u;
+		process.isRemote = sessionIndex != 0;
 		return &process;
 	}
 
diff --git a/Engine/Source/Programs/UnrealBuildAccelerator/Visualizer/Private/UbaTraceReader.h b/Engine/Source/Programs/UnrealBuildAccelerator/Visualizer/Private/UbaTraceReader.h
--- a/Engine/Source/Programs/UnrealBuildAccelerator/Visualizer/Private/UbaTraceReader.h
+++ b/Engine/Source/Programs/UnrealBuildAccelerator/Visualizer/Private/UbaTraceReader.h
@@ -74,11 +74,9 @@ namespace uba
 
 		struct StatusUpdate
 		{
-			TString name;
 			TString text;
-			u32 nameIndent;
-			u32 textIndent;
 			LogEntryType type;
+			TString link;
 		};
 
 		struct Session
@@ -127,20 +125,25 @@ namespace uba
 			bool success = false;
 		};
 
-		Process* GetProcess(const ProcessLocation& loc) { return &(sessions[loc.sessionIndex].processors[loc.processorIndex].processes[loc.processIndex]); }
-		void Clear() { sessions.clear(); workTracks.clear(); strings.clear(); statusMap.clear(); cacheWrites.clear(); startTime = 0; finished = true; totalProcessActiveCount = 0; totalProcessExitedCount = 0; activeSessionCount = 0; };
+		Process* GetProcess(const ProcessLocation& loc);
+		void Clear();
 
 		Vector<Session> sessions;
 		Vector<WorkTrack> workTracks;
 		Vector<TString> strings;
-		Map<u32, StatusUpdate> statusMap;
+		Map<u64, StatusUpdate> statusMap;
 		Map<u32, CacheWrite> cacheWrites;
+		u64 realStartTime = 0;
 		u64 startTime = 0;
 		u64 frequency = 0;
 		u32 totalProcessActiveCount = 0;
 		u32 totalProcessExitedCount = 0;
 		u32 activeSessionCount = 0;
 		u32 version = 0;
+		u32 progressProcessesTotal = 0;
+		u32 progressProcessesDone = 0;
+		u32 progressErrorCount = 0;
+		bool remoteExecutionDisabled = false;
 		bool finished = true;
 	};
 
@@ -162,13 +165,13 @@ namespace uba
 		bool UpdateReadClient(TraceView& out, NetworkClient& client, bool& outChanged);
 
 		// Use for local
-		bool StartReadNamed(TraceView& out, const tchar* namedTrace, bool silentFail = false);
-		bool UpdateReadNamed(TraceView& out, bool& outChanged);
+		bool StartReadNamed(TraceView& out, const tchar* namedTrace, bool silentFail = false, bool replay = false);
+		bool UpdateReadNamed(TraceView& out, u64 maxTime, bool& outChanged);
 
-		bool ReadMemory(TraceView& out, bool trackHost);
+		bool ReadMemory(TraceView& out, bool trackHost, u64 maxTime);
 		bool ReadTrace(TraceView& out, BinaryReader& reader, u64 maxTime);
 		void StopAllActive(TraceView& out, u64 stopTime);
-		void Reset();
+		void Reset(TraceView& out);
 		void Unmap();
 
 		bool SaveAs(const tchar* fileName);
@@ -193,10 +196,10 @@ namespace uba
 		TraceChannel m_channel;
 		ReaderWriterLock m_memoryLock;
 		FileMappingHandle m_memoryHandle;
+		TString m_namedTrace;
 		u8* m_memoryBegin = nullptr;
 		u8* m_memoryPos = nullptr;
 		u8* m_memoryEnd = nullptr;
-		u64 m_startTime = ~u64(0);
 		HANDLE m_hostProcess = NULL;
 	};
 }
\ No newline at end of file
diff --git a/Engine/Source/Programs/UnrealBuildAccelerator/Visualizer/Private/UbaVisualizer.cpp b/Engine/Source/Programs/UnrealBuildAccelerator/Visualizer/Private/UbaVisualizer.cpp
--- a/Engine/Source/Programs/UnrealBuildAccelerator/Visualizer/Private/UbaVisualizer.cpp
+++ b/Engine/Source/Programs/UnrealBuildAccelerator/Visualizer/Private/UbaVisualizer.cpp
@@ -27,13 +27,18 @@ namespace uba
 		UBA_VISUALIZER_FLAGS2
 		#undef UBA_VISUALIZER_FLAG
 
+		Popup_IncreaseFontSize,
+		Popup_DecreaseFontSize,
+
 		Popup_SaveAs,
 		Popup_SaveSettings,
+		Popup_OpenSettings,
 		Popup_Quit,
 	};
 
 	VisualizerConfig::VisualizerConfig(const tchar* fn) : filename(fn)
 	{
+		fontName = TC("Arial");
 	}
 
 	bool VisualizerConfig::Load(Logger& logger)
@@ -50,8 +55,11 @@ namespace uba
 		}
 		config.GetValueAsInt(x, L"X");
 		config.GetValueAsInt(y, L"Y");
-		config.GetValueAsInt(width, L"Width");
-		config.GetValueAsInt(height, L"Height");
+		config.GetValueAsU32(width, L"Width");
+		config.GetValueAsU32(height, L"Height");
+		config.GetValueAsU32(fontSize, L"FontSize");
+		config.GetValueAsString(fontName, L"FontName");
+		config.GetValueAsU32(maxActiveVisible, L"MaxActiveVisible");
 		#define UBA_VISUALIZER_FLAG(name, defaultValue, desc) config.GetValueAsBool(show##name, L"Show" #name);
 		UBA_VISUALIZER_FLAGS1
 		#undef UBA_VISUALIZER_FLAG
@@ -68,6 +76,9 @@ namespace uba
 		config.AddValue(L"Y", y);
 		config.AddValue(L"Width", width);
 		config.AddValue(L"Height", height);
+		config.AddValue(L"FontSize", fontSize);
+		config.AddValue(L"FontName", fontName.c_str());
+		config.AddValue(L"MaxActiveVisible", maxActiveVisible);
 		#define UBA_VISUALIZER_FLAG(name, defaultValue, desc) config.AddValue(L"Show" #name, show##name);
 		UBA_VISUALIZER_FLAGS1
 		#undef UBA_VISUALIZER_FLAG
@@ -136,7 +147,7 @@ namespace uba
 			GetClientRect(hwnd, &clientRect);
 
 			if (r.right > clientRect.right)
-				OffsetRect(&r, -width, 0);
+				OffsetRect(&r, -width - 15, 0);
 			if (r.bottom > clientRect.bottom)
 			{
 				OffsetRect(&r, 0, clientRect.bottom - r.bottom);
@@ -252,7 +263,7 @@ namespace uba
 				if (!traceName.Equals(m_newTraceName.data))
 				{
 					m_newTraceName.Clear().Append(traceName);
-					PostMessage(m_hwnd, WM_NEWTRACE, 0, 0);
+					PostNewTrace(0, false);
 				}
 			}
 			else
@@ -270,7 +281,7 @@ namespace uba
 		if (!StartHwndThread())
 			return true;
 		m_newTraceName.Append(namedTrace);
-		PostMessage(m_hwnd, WM_NEWTRACE, 0, 0);
+		PostNewTrace(0, false);
 		return true;
 	}
 
@@ -304,7 +315,7 @@ namespace uba
 			if (!m_client->Connect(backend, host, port))
 				continue;
 
-			PostMessage(m_hwnd, WM_NEWTRACE, 0, 0);
+			PostNewTrace(0, false);
 
 			while (m_hwnd && m_client->IsConnected())
 				Sleep(1000);
@@ -324,8 +335,7 @@ namespace uba
 		if (!StartHwndThread())
 			return true;
 		m_fileName.Append(fileName);
-		m_replay = replay;
-		PostMessage(m_hwnd, WM_NEWTRACE, 0, 0);
+		PostNewTrace(replay, 0);
 		return true;
 	}
 
@@ -359,7 +369,7 @@ namespace uba
 
 	bool Visualizer::Unselect()
 	{
-		if (m_processSelected || m_sessionSelectedIndex != ~0u || m_statsSelected || m_timelineSelected || m_fetchedFilesSelected != ~0u || m_workSelected)
+		if (m_processSelected || m_sessionSelectedIndex != ~0u || m_statsSelected || m_timelineSelected || m_fetchedFilesSelected != ~0u || m_workSelected || !m_hyperLinkSelected.empty())
 		{
 			m_processSelected = false;
 			m_sessionSelectedIndex = ~0u;
@@ -368,6 +378,7 @@ namespace uba
 			m_timelineSelected = 0;
 			m_fetchedFilesSelected = ~0u;
 			m_workSelected = false;
+			m_hyperLinkSelected.clear();
 			return true;
 		}
 		return false;
@@ -383,14 +394,12 @@ namespace uba
 		m_textBitmaps.clear();
 		m_lastBitmap = 0;
 		m_lastBitmapOffset = BitmapCacheHeight;
-		m_traceView.Clear();
 		m_autoScroll = true;
 		m_scrollPosX = 0;
 		m_scrollPosY = 0;
 		//m_zoomValue = 0.75f;
 		//m_horizontalScaleValue = 1.0f;
 
-		//m_replay = 0;
 		m_startTime = GetTime();
 		m_pauseTime = 0;
 
@@ -500,17 +509,13 @@ namespace uba
 		wndClassEx.style = CS_HREDRAW | CS_VREDRAW;
 		wndClassEx.lpfnWndProc = &StaticWinProc;
 		wndClassEx.hIcon = LoadIcon(hInstance, MAKEINTRESOURCE(123));
-		wndClassEx.hCursor = LoadCursor(NULL, IDC_ARROW);
+		wndClassEx.hCursor = 0;
 		wndClassEx.hInstance = hInstance;
 		wndClassEx.hbrBackground = NULL;
 		wndClassEx.lpszClassName = TEXT("UbaVisualizer");
 		ATOM wndClassAtom = RegisterClassEx(&wndClassEx);
 
-		NONCLIENTMETRICS nonClientMetrics;
-		nonClientMetrics.cbSize = sizeof(nonClientMetrics);
-		SystemParametersInfo(SPI_GETNONCLIENTMETRICS, sizeof(nonClientMetrics), &nonClientMetrics, 0);
-		//m_font = (HFONT)CreateFontIndirect(&nonClientMetrics.lfMessageFont);
-		m_font = (HFONT)CreateFontW(-9, 0, 0, 0, FW_NORMAL, false, false, false, ANSI_CHARSET, OUT_DEFAULT_PRECIS, CLIP_DEFAULT_PRECIS, CLEARTYPE_QUALITY, DEFAULT_PITCH, L"Arial");
+		UpdateFont();
 
 		const TCHAR* fontName = TEXT("Consolas");
 		m_popupFont = (HFONT)CreateFontW(-12, 0, 0, 0, FW_NORMAL, false, false, false, ANSI_CHARSET, OUT_DEFAULT_PRECIS, CLIP_DEFAULT_PRECIS, 0, FIXED_PITCH | FF_MODERN, fontName);
@@ -567,7 +572,11 @@ namespace uba
 				if (!m_looping || msg.message == WM_QUIT || msg.message == WM_DESTROY || msg.message == WM_CLOSE)
 				{
 					if (m_hwnd)
+					{
+						if (m_config.AutoSaveSettings)
+							SaveSettings();
 						DestroyWindow(m_hwnd);
+					}
 					DeleteObject(m_font);
 					UnregisterClass(windowClassName, hInstance);
 					m_hwnd = 0;
@@ -610,6 +619,7 @@ namespace uba
 			m_timelineSelected = 0;
 			m_fetchedFilesSelected = ~0u;
 			m_workSelected = false;
+			m_hyperLinkSelected.clear();
 			m_autoScroll = false;
 			m_mouseAnchor = { anchor.x, anchor.y };
 			m_scrollAtAnchorX = m_scrollPosX;
@@ -652,6 +662,38 @@ namespace uba
 						process.bitmapDirty = true;
 	}
 
+	void Visualizer::UpdateFont()
+	{
+		//NONCLIENTMETRICS nonClientMetrics;
+		//nonClientMetrics.cbSize = sizeof(nonClientMetrics);
+		//SystemParametersInfo(SPI_GETNONCLIENTMETRICS, sizeof(nonClientMetrics), &nonClientMetrics, 0);
+		//m_font = (HFONT)CreateFontIndirect(&nonClientMetrics.lfMessageFont);
+		if (m_font)
+			DeleteObject(m_font);
+		if (m_fontUnderlined)
+			DeleteObject(m_fontUnderlined);
+		m_fontHeight = m_config.fontSize;
+		m_font = (HFONT)CreateFontW(4 - m_fontHeight, 0, 0, 0, FW_NORMAL, 0, 0, 0, ANSI_CHARSET, OUT_DEFAULT_PRECIS, CLIP_DEFAULT_PRECIS, CLEARTYPE_QUALITY, DEFAULT_PITCH, m_config.fontName.c_str());
+		m_fontUnderlined = (HFONT)CreateFontW(4 - m_fontHeight, 0, 0, 0, FW_NORMAL, 0, 1, 0, ANSI_CHARSET, OUT_DEFAULT_PRECIS, CLIP_DEFAULT_PRECIS, CLEARTYPE_QUALITY, DEFAULT_PITCH, m_config.fontName.c_str());
+
+		m_rawBoxHeight = m_fontHeight + 2;
+		m_sessionStepY = m_rawBoxHeight + 2;
+	}
+
+	void Visualizer::ChangeFontSize(int offset)
+	{
+		m_config.fontSize += offset;
+		m_config.fontSize = Max(m_config.fontSize, 10u);
+		UpdateFont();
+		DirtyBitmaps();
+		Redraw();
+	}
+
+	void Visualizer::Redraw()
+	{
+		RedrawWindow(m_hwnd, NULL, NULL, RDW_INVALIDATE | RDW_UPDATENOW);
+	}
+
 	void Visualizer::PaintClient(const Function<void(HDC hdc, HDC memDC, RECT& clientRect)>& paintFunc)
 	{
 		HDC hdc = GetDC(m_hwnd);
@@ -680,9 +722,6 @@ namespace uba
 
 	constexpr int ProgressRectLeft = 30;
 	constexpr int GraphHeight = 30;
-	constexpr int RawBoxHeight = 15;
-	constexpr int SessionStepY = RawBoxHeight + 2;
-	constexpr int FontHeight = 13;
 
 	struct SessionRec
 	{
@@ -711,12 +750,14 @@ namespace uba
 	void Visualizer::PaintAll(HDC hdc, const RECT& clientRect)
 	{
 		u64 currentTime = m_paused ? m_pauseStart : GetTime();
-		u64 playTime = currentTime - m_traceView.startTime - m_pauseTime;
+		u64 playTime = 0;
+		if (m_traceView.startTime)
+			playTime = currentTime - m_traceView.startTime - m_pauseTime;
 		if (m_replay)
 			playTime *= m_replay;
 
 		int posY = int(m_scrollPosY);
-		float boxHeight = float(RawBoxHeight)*m_zoomValue;
+		float boxHeight = float(m_rawBoxHeight)*m_zoomValue;
 		int stepY = int(boxHeight) + 2;
 		float scaleX = 50.0f*m_zoomValue*m_horizontalScaleValue;
 
@@ -747,35 +788,126 @@ namespace uba
 
 		u64 lastStop = 0;
 
+		auto drawStatusText = [&](const StringView& text, LogEntryType type, int indent, bool moveY, bool underlined = false)
+			{
+				RECT rect;
+				rect.left = 5 + indent*m_fontHeight;
+				rect.right = clientRect.right;
+				rect.top = posY;
+				rect.bottom = posY + m_fontHeight + 2;
+				SetTextColor(hdc, type == LogEntryType_Info ? m_textColor : (type == LogEntryType_Error ? m_textErrorColor : m_textWarningColor));
+
+				if (underlined)
+					SelectObject(hdc, m_fontUnderlined);
+				ExtTextOutW(hdc, rect.left, posY, ETO_CLIPPED, &rect, text.data, text.count, NULL);
+				if (underlined)
+					SelectObject(hdc, m_font);
+				if (moveY)
+					posY = rect.bottom;
+			};
+
+		if (m_config.showProgress && m_traceView.progressProcessesTotal)
+		{
+			drawStatusText(ToView(L"Progress"), LogEntryType_Info, 1, false);
+
+			float progress = float(m_traceView.progressProcessesDone) / m_traceView.progressProcessesTotal;
+			u32 width = m_fontHeight * 18;
+			RECT rect;
+			rect.left = 3 + 6*m_fontHeight;
+			rect.right = rect.left + width;
+			rect.top = posY;
+			rect.bottom = posY + m_fontHeight;
+			FillRect(hdc, &rect, m_processBrushes[0].inProgress);
+
+			rect.right = rect.left + int(progress*width);//durationMs/100;
+			FillRect(hdc, &rect, m_traceView.progressErrorCount ? m_processBrushes[0].error : m_processBrushes[0].success);
+
+			const tchar* remoteDisabled = m_traceView.remoteExecutionDisabled ? TC(" (Remote spawn disabled)") : TC("");
+			StringBuffer<> str;
+			str.Appendf(L"%u%% (%u/%u) %s", u32(progress*100.0f), m_traceView.progressProcessesDone, m_traceView.progressProcessesTotal, remoteDisabled);
+			drawStatusText(str, LogEntryType_Info, 6, true);
+		}
+
 		if (!m_traceView.statusMap.empty())
 		{
 			posY += 4;
-
-			auto drawStatusText = [&](const TString& text, LogEntryType type, int left, bool moveY)
-				{
-					RECT rect;
-					rect.left = left;
-					rect.right = clientRect.right;
-					rect.top = posY;
-					rect.bottom = posY + FontHeight + 2;
-					SetTextColor(hdc, type == LogEntryType_Info ? m_textColor : (type == LogEntryType_Error ? m_textErrorColor : m_textWarningColor));
-					ExtTextOutW(hdc, left, posY, ETO_CLIPPED, &rect, text.c_str(), u32(text.size()), NULL);
-					if (moveY)
-						posY = rect.bottom;
-				};
-
+			u32 lastRow = ~0u;
+			u32 row = ~0u;
 			for (auto& kv : m_traceView.statusMap)
 			{
 				auto& status = kv.second;
-				if (status.name.empty() && status.text.empty())
+				if (status.text.empty())
 					continue;
-				drawStatusText(status.name, status.type, 5 + status.nameIndent*15, false);
-				drawStatusText(status.text, status.type, 5 + status.textIndent*15, true);
+				row = u32(kv.first >> 32);
+				if (lastRow != ~0u && lastRow != row)
+					posY += m_fontHeight + 2;
+				lastRow = row;
+				u32 column = u32(kv.first & ~0u);
+				drawStatusText(status.text, status.type, column, false, !status.link.empty());
 			}
+			if (row != ~0u)
+				posY += m_fontHeight + 2;
+
 			SetTextColor(hdc, m_textColor);
 			posY += 4;
 		}
 
+		if (m_config.showActiveProcesses && !m_trace.m_activeProcesses.empty())
+		{
+			Map<u64, TraceView::Process*> activeProcesses;
+			u32 remoteCount = 0;
+			for (auto& pair : m_trace.m_activeProcesses)
+			{
+				auto& active = pair.second;
+				TraceView::Process& process = m_trace.GetSession(m_traceView, active.sessionIndex).processors[active.processorIndex].processes[active.processIndex];
+				activeProcesses.try_emplace(~0llu - process.start, &process);
+				if (process.isRemote)
+					++remoteCount;
+			}
+
+			drawStatusText(ToView(L"Active"), LogEntryType_Info, 1, false);
+			StringBuffer<> str;
+			str.Appendf(L"%u", u32(activeProcesses.size()) - remoteCount);
+			if (remoteCount)
+				str.Appendf(L" local %u remote", remoteCount);
+			str.Append(L" process");
+			if (activeProcesses.size() > 1)
+				str.Append(L"es");
+			drawStatusText(str, LogEntryType_Info, 6, true);
+
+			u32 index = 0;
+			for (auto& kv : activeProcesses)
+			{
+				if (index++ == m_config.maxActiveVisible)
+					break;
+				
+				u32 durationMs = u32(TimeToMs(playTime - kv.second->start));
+				RECT rect;
+				rect.left = 3 + 6*m_fontHeight;
+				rect.right = rect.left + durationMs/100;
+				rect.top = posY;
+				rect.bottom = posY + m_fontHeight;
+				FillRect(hdc, &rect, m_processBrushes[0].inProgress);
+
+				str.Clear();
+				auto firstPar = kv.second->description.find_first_of('(');
+				if (firstPar != -1 && *kv.second->description.rbegin() == ')')
+				{
+					str.Append(kv.second->description.c_str() + firstPar + 1).Resize(str.count-1);
+					str.Append(L"  ").Append(kv.second->description.c_str(), firstPar);
+				}
+				else
+					str.Append(kv.second->description);
+
+				if (kv.second->isRemote)
+					str.Append(L" (remote)");
+				if (durationMs > 1000)
+					str.Appendf(L" - %us", durationMs/1000);
+				drawStatusText(str, LogEntryType_Info, 6, true);
+			}
+			posY += 3;
+		}
+
 		TraceView::WorkRecord selectedWork;
 
 		SessionRec sortedSessions[1024];
@@ -784,50 +916,52 @@ namespace uba
 		TraceView::ProcessLocation processLocation { 0, 0, 0 };
 		for (u64 i = 0, e = m_traceView.sessions.size(); i != e; ++i)
 		{
+			bool isFirst = i == 0;
 			auto& session = *sortedSessions[i].session;
 			bool hasUpdates = !session.updates.empty();
-			if (!hasUpdates && session.processors.empty())
+			if (!isFirst && !hasUpdates && session.processors.empty())
 				continue;
 
 			processLocation.sessionIndex = sortedSessions[i].index;
-			bool isFirst = i == 0;
 			if (!isFirst)
 				posY += 3;
 
-			if (posY + stepY >= progressRect.top && posY <= progressRect.bottom)
+			if (m_config.showTitleBars)
 			{
-				SelectObject(hdc, m_separatorPen);
-				MoveToEx(hdc, 0, posY, NULL);
-				LineTo(hdc, clientRect.right, posY);
+				if (posY + stepY >= progressRect.top && posY <= progressRect.bottom)
+				{
+					SelectObject(hdc, m_separatorPen);
+					MoveToEx(hdc, 0, posY, NULL);
+					LineTo(hdc, clientRect.right, posY);
 
-				StringBuffer<> text;
-				text.Append(session.name);
+					StringBuffer<> text;
+					text.Append(session.name);
 
-				if (hasUpdates && session.disconnectTime == ~u64(0))
-				{
-					u64 ping = session.updates.back().ping;
-					u64 memAvail = session.updates.back().memAvail;
-					float cpuLoad = session.updates.back().cpuLoad;
-
-					text.Appendf(L" - Cpu: %.1f%%", cpuLoad * 100.0f);
-					if (memAvail)
-						text.Appendf(L" Mem: %ls/%ls", BytesToText(session.memTotal - memAvail).str, BytesToText(session.memTotal).str);
-					if (ping)
-						text.Appendf(L" Ping: %ls", TimeToText(ping, false, m_traceView.frequency).str);
-					if (!session.notification.empty())
-						text.Append(L" - ").Append(session.notification);
-				}
-				else
-				{
-					text.Append(L" - Disconnected");
-					if (!session.notification.empty())
-						text.Append(L" (").Append(session.notification).Append(')');
-				}
+					if (hasUpdates && session.disconnectTime == ~u64(0))
+					{
+						u64 ping = session.updates.back().ping;
+						u64 memAvail = session.updates.back().memAvail;
+						float cpuLoad = session.updates.back().cpuLoad;
+
+						text.Appendf(L" - Cpu: %.1f%%", cpuLoad * 100.0f);
+						if (memAvail)
+							text.Appendf(L" Mem: %ls/%ls", BytesToText(session.memTotal - memAvail).str, BytesToText(session.memTotal).str);
+						if (ping)
+							text.Appendf(L" Ping: %ls", TimeToText(ping, false, m_traceView.frequency).str);
+						if (!session.notification.empty())
+							text.Append(L" - ").Append(session.notification);
+					}
+					else if (!isFirst)
+					{
+						text.Append(L" - Disconnected");
+						if (!session.notification.empty())
+							text.Append(L" (").Append(session.notification).Append(')');
+					}
 
-				bool selected = m_sessionSelectedIndex == processLocation.sessionIndex;
+					bool selected = m_sessionSelectedIndex == processLocation.sessionIndex;
+
+					int textBottom = Min(posY + m_sessionStepY, int(progressRect.bottom));
 
-				int textBottom = Min(posY + SessionStepY, int(progressRect.bottom));
-				{
 					RECT rect;
 					rect.left = 5;
 					rect.right = clientRect.right;
@@ -840,8 +974,8 @@ namespace uba
 					if (selected)
 						SetBkMode(hdc, TRANSPARENT);
 				}
+				posY += m_sessionStepY;
 			}
-			posY += SessionStepY;
 
 			bool showGraph = m_config.showNetworkStats || m_config.showCpuMemStats;
 			if (showGraph && hasUpdates)
@@ -946,7 +1080,7 @@ namespace uba
 				processLocation.processorIndex = 0;
 				for (auto& processor : session.processors)
 				{
-					if (posY + SessionStepY >= progressRect.top && posY < progressRect.bottom)
+					if (posY + m_sessionStepY >= progressRect.top && posY < progressRect.bottom)
 					{
 						float barHeight = boxHeight;
 						int textOffsetY = 0;
@@ -959,7 +1093,7 @@ namespace uba
 
 						const int textHeight = int(barHeight);
 						const int rectBottom = posY + textHeight;
-						const int offsetY = (textHeight - FontHeight + textOffsetY) / 2;
+						const int offsetY = (textHeight - m_fontHeight + textOffsetY) / 2;
 
 						if (shouldDrawText)
 						{
@@ -1030,7 +1164,7 @@ namespace uba
 										}
 										process.bitmap = m_lastBitmap;
 										process.bitmapOffset = m_lastBitmapOffset;
-										m_lastBitmapOffset += FontHeight;
+										m_lastBitmapOffset += m_fontHeight;
 									}
 									if (lastSelectedBitmap != process.bitmap)
 									{
@@ -1038,8 +1172,8 @@ namespace uba
 										lastSelectedBitmap = process.bitmap;
 									}
 
-									RECT rect2{ 0, int(process.bitmapOffset), 256, int(process.bitmapOffset) + FontHeight };
-									RECT rect3{ 0, int(process.bitmapOffset), processWidth, int(process.bitmapOffset) + FontHeight };
+									RECT rect2{ 0, int(process.bitmapOffset), 256, int(process.bitmapOffset) + m_fontHeight };
+									RECT rect3{ 0, int(process.bitmapOffset), processWidth, int(process.bitmapOffset) + m_fontHeight };
 									if (!done)
 										rect3.right = 256;
 
@@ -1078,7 +1212,7 @@ namespace uba
 									bitmapOffsetY -= bltOffsetY;
 									bltOffsetY = 0;
 								}
-								int height = Min(textHeight, FontHeight);
+								int height = Min(textHeight, m_fontHeight);
 								if (bltOffsetY + height > textHeight)
 									height = textHeight - bltOffsetY;
 
@@ -1118,7 +1252,7 @@ namespace uba
 				u32 trackIndex = 0;
 				for (auto& workTrack : m_traceView.workTracks)
 				{
-					if (posY + SessionStepY >= progressRect.top && posY <= progressRect.bottom)
+					if (posY + m_sessionStepY >= progressRect.top && posY <= progressRect.bottom)
 					{
 						int textOffsetY = 0;
 						float barHeight = boxHeight;
@@ -1131,7 +1265,7 @@ namespace uba
 
 						const int textHeight = int(barHeight);
 						const int rectBottom = posY + textHeight;
-						const int offsetY = (textHeight - FontHeight + textOffsetY) / 2;
+						const int offsetY = (textHeight - m_fontHeight + textOffsetY) / 2;
 
 						if (shouldDrawText)
 						{
@@ -1236,13 +1370,13 @@ namespace uba
 									}
 									SelectObject(textDC, m_lastBitmap);
 
-									RECT rect2{ 0,m_lastBitmapOffset,256, m_lastBitmapOffset + FontHeight };
+									RECT rect2{ 0,m_lastBitmapOffset,256, m_lastBitmapOffset + m_fontHeight };
 
 									FillRect(textDC, &rect2, m_workBrush);
 									ExtTextOutW(textDC, rect2.left, rect2.top, ETO_CLIPPED, &rect2, work.description, int(wcslen(work.description)), NULL);
 									work.bitmap = m_lastBitmap;
 									work.bitmapOffset = m_lastBitmapOffset;
-									m_lastBitmapOffset += FontHeight;
+									m_lastBitmapOffset += m_fontHeight;
 								}
 
 								if (lastSelectedBitmap != work.bitmap)
@@ -1259,7 +1393,7 @@ namespace uba
 									bitmapOffsetY -= bltOffsetY;
 									bltOffsetY = 0;
 								}
-								int height = Min(textHeight, FontHeight);
+								int height = Min(textHeight, m_fontHeight);
 								if (bltOffsetY + height > textHeight)
 									height = textHeight - bltOffsetY;
 
@@ -1290,7 +1424,7 @@ namespace uba
 		DeleteObject(nullBmp);
 		DeleteDC(textDC);
 
-		m_contentWidth = ProgressRectLeft + int(TimeToS(lastStop != ~u64(0) ? lastStop : playTime) * scaleX);
+		m_contentWidth = ProgressRectLeft + int(TimeToS((lastStop != 0 && lastStop != ~u64(0)) ? lastStop : playTime) * scaleX);
 		m_contentHeight = posY - int(m_scrollPosY) + stepY + 14;
 
 		float timelineSelected = m_timelineSelected;
@@ -1356,11 +1490,13 @@ namespace uba
 		}
 
 		{
+			int boxSide = 8;
+			int boxStride = boxSide + 2;
 			int top = 5;
-			int bottom = SessionStepY - 3;
-			int left = progressRect.right - 16;
+			int bottom = top + boxSide;
+			int left = progressRect.right - 7 - boxSide;
 			int right = progressRect.right - 7;
-			bool* values = &m_config.showNetworkStats;
+			bool* values = &m_config.showProgress;
 			for (int i = VisualizerFlag_Count - 1; i >= 0; --i)
 			{
 				SelectObject(hdc, m_buttonSelected == u32(i) ? m_textPen : m_checkboxPen);
@@ -1375,8 +1511,8 @@ namespace uba
 					LineTo(hdc, left + 1, bottom - 2);
 				}
 
-				left -= 14;
-				right -= 14;
+				left -= boxStride;
+				right -= boxStride;
 			}
 
 			top -= 2;
@@ -1628,8 +1764,6 @@ namespace uba
 		}
 		else if (m_buttonSelected != ~0u)
 		{
-			bool* values = &m_config.showNetworkStats;
-
 			const wchar_t* tooltip[] =
 			{
 				#define UBA_VISUALIZER_FLAG(name, defaultValue, desc) desc,
@@ -1637,10 +1771,9 @@ namespace uba
 				#undef UBA_VISUALIZER_FLAG
 			};
 
-			bool vis = values[m_buttonSelected];
 			SelectObject(hdc, m_popupFont);
 			DrawTextLogger logger(m_hwnd, hdc, m_popupFontHeight, m_tooltipBackgroundBrush);
-			logger.Info(L"%ls %ls", vis ? L"Hide" : L"Show", tooltip[m_buttonSelected]);
+			logger.Info(L"%ls %ls", L"Show", tooltip[m_buttonSelected]);
 			logger.DrawAtCursor();
 		}
 		else if (m_fetchedFilesSelected != ~0u)
@@ -1655,12 +1788,12 @@ namespace uba
 				int height = Min(int(clientRect.bottom), int(fetchedFiles.size() * m_popupFontHeight));
 
 				SelectObject(hdc, m_font);
-				DrawTextLogger logger(m_hwnd, hdc, r, FontHeight, m_tooltipBackgroundBrush);
+				DrawTextLogger logger(m_hwnd, hdc, r, m_fontHeight, m_tooltipBackgroundBrush);
 				for (auto& f : fetchedFiles)
 				{
 					if (f.hint == TC("KnownInput"))
 						continue;
-					if (logger.rect.top >= r.bottom - FontHeight)
+					if (logger.rect.top >= r.bottom - m_fontHeight)
 					{
 						if (logger.rect.left + colWidth >= r.right)
 						{
@@ -1763,10 +1896,10 @@ namespace uba
 
 	void Visualizer::PaintTimeline(HDC hdc, const RECT& clientRect)
 	{
-		float boxHeight = float(RawBoxHeight) * m_zoomValue;
+		float boxHeight = float(m_rawBoxHeight) * m_zoomValue;
 		int posY = m_contentHeight - int(boxHeight) - 14;
 		float timeScale = (m_horizontalScaleValue*m_zoomValue)*50.0f;
-		int top = Min(posY, int(clientRect.bottom - SessionStepY - 10));
+		int top = Min(posY, int(clientRect.bottom - m_sessionStepY - 10));
 			
 		float startOffset = ((m_scrollPosX/timeScale) - int(m_scrollPosX/timeScale)) * timeScale;
 		int index = -int(startOffset/timeScale);
@@ -1857,7 +1990,7 @@ namespace uba
 
 	void Visualizer::PaintDetailedStats(int& posY, const RECT& progressRect, TraceView::Session& session, bool isRemote, u64 playTime, const DrawTextFunc& drawTextFunc)
 	{
-		int stepY = FontHeight;
+		int stepY = m_fontHeight;
 		int startPosY = posY;
 		int posX = progressRect.left + 5;
 		RECT textRect;
@@ -1971,7 +2104,9 @@ namespace uba
 	void Visualizer::HitTest(HitTestResult& outResult, const POINT& pos)
 	{
 		u64 currentTime = m_paused ? m_pauseStart : GetTime();
-		u64 playTime = currentTime - m_traceView.startTime - m_pauseTime;
+		u64 playTime = 0;
+		if (m_traceView.startTime)
+			playTime = currentTime - m_traceView.startTime - m_pauseTime;
 		if (m_replay)
 			playTime *= m_replay;
 
@@ -1979,7 +2114,7 @@ namespace uba
 		GetClientRect(m_hwnd, &clientRect);
 
 		int posY = int(m_scrollPosY);
-		float boxHeight = float(RawBoxHeight)*m_zoomValue;
+		float boxHeight = float(m_rawBoxHeight)*m_zoomValue;
 		int stepY = int(boxHeight) + 2;
 		float scaleX = 50.0f*m_zoomValue*m_horizontalScaleValue;
 
@@ -1988,9 +2123,11 @@ namespace uba
 		progressRect.bottom -= 30;
 
 		{
+			int boxSide = 8;
+			int boxStride = boxSide + 2;
 			int top = 5;
-			int bottom = SessionStepY - 3;
-			int left = progressRect.right - 16;
+			int bottom = top + boxSide;
+			int left = progressRect.right - 7 - boxSide;
 			int right = progressRect.right - 7;
 			for (int i = VisualizerFlag_Count - 1; i >= 0; --i)
 			{
@@ -1999,23 +2136,51 @@ namespace uba
 					outResult.buttonSelected = i;
 					return;
 				}
-				left -= 14;
-				right -= 14;
+				left -= boxStride;
+				right -= boxStride;
 			}
 		}
 
 		u64 lastStop = 0;
 
+		if (m_config.showProgress && m_traceView.progressProcessesTotal)
+			posY += m_fontHeight + 2;
+
 		if (!m_traceView.statusMap.empty())
 		{
 			posY += 4;
-			auto drawStatusText = [&]() { posY = posY + FontHeight + 2; };
+			u32 lastRow = ~0u;
+			u32 row = ~0u;
 			for (auto& kv : m_traceView.statusMap)
-				if (!kv.second.name.empty() || !kv.second.text.empty())
-					drawStatusText();
+			{
+				if (kv.second.text.empty())
+					continue;
+
+				if (!kv.second.link.empty())
+				{
+					if (pos.y >= posY && pos.y < posY + m_fontHeight && pos.x > 20 && pos.x < 80) // TODO: x is just hard coded to fit horde for now
+					{
+						outResult.hyperLink = kv.second.link;
+						return;
+					}
+				}
+
+				row = u32(kv.first >> 32);
+				if (lastRow != ~0u && lastRow != row)
+					posY += m_fontHeight + 2;
+				lastRow = row;
+			}
+			if (row != ~0u)
+				posY += m_fontHeight + 2;
 			posY += 4;
 		}
 
+		if (m_config.showActiveProcesses && !m_trace.m_activeProcesses.empty())
+		{
+			u32 lines = 1 + Min(u32(m_trace.m_activeProcesses.size()), m_config.maxActiveVisible);
+			posY += (m_fontHeight + 2)*lines;
+		}
+
 		TraceView::ProcessLocation& outLocation = outResult.processLocation;
 
 		SessionRec sortedSessions[1024];
@@ -2023,27 +2188,29 @@ namespace uba
 
 		for (u64 i = 0, e = m_traceView.sessions.size(); i != e; ++i)
 		{
+			bool isFirst = i == 0;
 			auto& session = *sortedSessions[i].session;
-
 			bool hasUpdates = !session.updates.empty();
-			if (!hasUpdates && session.processors.empty())
+			if (!isFirst && !hasUpdates && session.processors.empty())
 				continue;
 
 			u32 sessionIndex = sortedSessions[i].index;
-			bool isFirst = i == 0;
 			if (!isFirst)
 				posY += 3;
 
-			if (pos.y >= posY && pos.y < posY + SessionStepY)
+			if (m_config.showTitleBars)
 			{
-				if (pos.x < 500)
+				if (pos.y >= posY && pos.y < posY + m_sessionStepY)
 				{
-					outResult.sessionSelectedIndex = sessionIndex;
-					return;
+					if (pos.x < 500)
+					{
+						outResult.sessionSelectedIndex = sessionIndex;
+						return;
+					}
 				}
-			}
 
-			posY += SessionStepY;
+				posY += m_sessionStepY;
+			}
 
 			bool showGraph = m_config.showNetworkStats || m_config.showCpuMemStats;
 			if (showGraph && !session.updates.empty())
@@ -2235,7 +2402,7 @@ namespace uba
 
 		}
 
-		m_contentWidth = ProgressRectLeft + int(TimeToS(lastStop != ~u64(0) ? lastStop : playTime) * scaleX);
+		m_contentWidth = ProgressRectLeft + int(TimeToS((lastStop != 0 && lastStop != ~u64(0)) ? lastStop : playTime) * scaleX);
 		m_contentHeight = posY - int(m_scrollPosY) + stepY + 14;
 	}
 
@@ -2331,8 +2498,9 @@ namespace uba
 		RECT rect;
 		GetClientRect(m_hwnd, &rect);
 		float timeS = TimeToS(playTime);
+		float oldScrollPosX = m_scrollPosX;
 		m_scrollPosX = Min(0.0f, (float)rect.right - timeS*50.0f*m_horizontalScaleValue*m_zoomValue - ProgressRectLeft);
-		return true;
+		return oldScrollPosX != m_scrollPosX;
 	}
 
 	bool Visualizer::UpdateSelection()
@@ -2350,7 +2518,8 @@ namespace uba
 			res.statsSelected == m_statsSelected && memcmp(&res.stats, &m_stats, sizeof(Stats)) == 0 &&
 			res.buttonSelected == m_buttonSelected && res.timelineSelected == m_timelineSelected &&
 			res.fetchedFilesSelected == m_fetchedFilesSelected &&
-			res.workSelected == m_workSelected && res.workTrack == m_workTrack && res.workIndex == m_workIndex)
+			res.workSelected == m_workSelected && res.workTrack == m_workTrack && res.workIndex == m_workIndex &&
+			res.hyperLink == m_hyperLinkSelected)
 			return false;
 		m_processSelected = res.processSelected;
 		m_processSelectedLocation = res.processLocation;
@@ -2363,6 +2532,7 @@ namespace uba
 		m_workSelected = res.workSelected;
 		m_workTrack = res.workTrack;
 		m_workIndex = res.workIndex;
+		m_hyperLinkSelected = res.hyperLink;
 		return true;
 	}
 
@@ -2373,7 +2543,7 @@ namespace uba
 
 		SCROLLINFO si;
 		si.cbSize = sizeof(SCROLLINFO);
-		si.fMask = SIF_ALL;
+		si.fMask = SIF_ALL | SIF_DISABLENOSCROLL;
 		si.nMin = 0;
 		si.nMax = m_contentHeight;
 		si.nPage = rect.bottom;
@@ -2387,14 +2557,20 @@ namespace uba
 		SetScrollInfo(m_hwnd, SB_HORZ, &si, redraw);
 	}
 
+	void Visualizer::PostNewTrace(u32 replay, bool paused)
+	{
+		KillTimer(m_hwnd, 0);
+		PostMessage(m_hwnd, WM_NEWTRACE, replay, paused);
+	}
+
 	LRESULT Visualizer::WinProc(HWND hWnd, UINT Msg, WPARAM wParam, LPARAM lParam)
 	{
 		switch (Msg)
 		{
 		case WM_NEWTRACE:
 		{
-			if (!m_traceView.finished)
-				return 0;
+			m_replay = u32(wParam);
+			m_paused = lParam;
 			Reset();
 			StringBuffer<> title;
 			GetTitlePrefix(title);
@@ -2416,17 +2592,17 @@ namespace uba
 			else if (!m_fileName.IsEmpty())
 			{
 				m_trace.ReadFile(m_traceView, m_fileName.data, m_replay != 0);
-				RedrawWindow(hWnd, NULL, NULL, RDW_INVALIDATE | RDW_UPDATENOW);
+				Redraw();
 				title.Append(m_fileName);
 				m_traceView.finished = m_replay == 0;
 			}
 			else
 			{
-				if (!m_trace.StartReadNamed(m_traceView, m_newTraceName.data, true))
+				if (!m_trace.StartReadNamed(m_traceView, m_newTraceName.data, true, m_replay != 0))
 					return false;
 				m_namedTrace.Clear().Append(m_newTraceName);
-				title.Appendf(L"%s (Listening for new sessions on channel '%s')", m_namedTrace.data, m_listenChannel.data);
 				m_traceView.finished = false;
+				title.Appendf(L"%s (Listening for new sessions on channel '%s')", m_namedTrace.data, m_listenChannel.data);
 			}
 
 			SetWindowTextW(m_hwnd, title.data);
@@ -2479,15 +2655,18 @@ namespace uba
 			if (!m_paused)
 			{
 				u64 timeOffset = (GetTime() - m_startTime - m_pauseTime) * m_replay;
-				if (!m_namedTrace.IsEmpty())
+				if (!m_fileName.IsEmpty())
 				{
-					if (!m_trace.UpdateReadNamed(m_traceView, changed))
-						m_listenTimeout.Set();
+					if (m_replay)
+						m_trace.UpdateReadFile(m_traceView, timeOffset, changed);
 				}
-				else if (!m_fileName.IsEmpty() && m_replay)
-					m_trace.UpdateReadFile(m_traceView, timeOffset, changed);
 				else if (m_client)
 					m_trace.UpdateReadClient(m_traceView, *m_client, changed);
+				else
+				{
+					if (!m_trace.UpdateReadNamed(m_traceView, m_replay ? timeOffset : ~u64(0), changed))
+						m_listenTimeout.Set();
+				}
 			}
 
 			if (m_traceView.finished)
@@ -2669,7 +2848,7 @@ namespace uba
 		{
 			if (m_buttonSelected != ~0u)
 			{
-				bool* values = &m_config.showNetworkStats;
+				bool* values = &m_config.showProgress;
 				values[m_buttonSelected] = !values[m_buttonSelected];
 				HitTestResult res;
 				HitTest(res, { -1, -1 });
@@ -2678,34 +2857,55 @@ namespace uba
 			}
 			else if (m_timelineSelected)
 			{
-				if (!m_fileName.IsEmpty()) // Only works for files right now
+				if (!m_client) // Does not work for network streams
 				{
-					float timelineSelected = m_timelineSelected;
+					float timelineSelected = Max(m_timelineSelected, 0.0f);
 					Reset();
-					if (!m_trace.ReadFile(m_traceView, m_fileName.data, true))
-						return false;
+					if (!m_fileName.IsEmpty())
+					{
+						if (!m_trace.ReadFile(m_traceView, m_fileName.data, true))
+							return false;
+					}
+					else
+					{
+						if (!m_trace.StartReadNamed(m_traceView, nullptr, true, true))
+							return false;
+					}
+
 					bool changed;
 					u64 time = MsToTime(u64(timelineSelected * 1000.0));
+
+					if (m_traceView.realStartTime + time > m_startTime)
+						time = m_startTime - m_traceView.realStartTime;
+
 					m_traceView.finished = false;
-					m_trace.UpdateReadFile(m_traceView, time, changed);
+					if (!m_fileName.IsEmpty())
+						m_trace.UpdateReadFile(m_traceView, time, changed);
+					else
+						m_trace.UpdateReadNamed(m_traceView, time, changed);
+
 					m_pauseStart = m_startTime + time;
 					if (m_paused || !m_replay)
 					{
 						m_pauseTime = 0;
-						m_paused = true;
 						m_replay = 1;
 					}
 					else
 					{
-						m_pauseTime = GetTime() - m_pauseStart;
+						m_pauseTime = m_startTime - m_pauseStart;
+						SetTimer(m_hwnd, 0, 200, NULL);
 					}
 
 					HitTestResult res;
 					HitTest(res, { -1, -1 });
 					UpdateScrollbars(true);
-					RedrawWindow(hWnd, NULL, NULL, RDW_INVALIDATE | RDW_UPDATENOW);
+					Redraw();
 				}
 			}
+			else if (!m_hyperLinkSelected.empty())
+			{
+				ShellExecuteW(NULL, L"open", m_hyperLinkSelected.c_str(), NULL, NULL, SW_SHOW);
+			}
 			else
 			{
 				POINTS p = MAKEPOINTS(lParam);
@@ -2714,6 +2914,16 @@ namespace uba
 			break;
 		}
 
+		case WM_SETCURSOR:
+		{
+			static HCURSOR arrow = LoadCursorW(NULL, IDC_ARROW);
+			static HCURSOR hand = LoadCursorW(NULL, IDC_HAND);
+			if (!m_hyperLinkSelected.empty())
+				SetCursor(hand);
+			else
+				SetCursor(arrow);
+			break;
+		}
 		case WM_LBUTTONUP:
 		{
 			if (!(m_buttonSelected != ~0u || m_timelineSelected))
@@ -2737,6 +2947,9 @@ namespace uba
 			UBA_VISUALIZER_FLAGS2
 			#undef UBA_VISUALIZER_FLAG
 
+			AppendMenuW(hMenu, MF_STRING, Popup_IncreaseFontSize, L"&Increase Font Size");
+			AppendMenuW(hMenu, MF_STRING, Popup_DecreaseFontSize, L"&Decrease Font Size");
+
 			AppendMenuW(hMenu, MF_SEPARATOR, 0, NULL);
 
 			if (m_sessionSelectedIndex != ~0u)
@@ -2753,25 +2966,31 @@ namespace uba
 					AppendMenuW(hMenu, MF_STRING, Popup_CopyProcessLog, L"Copy Process &Log");
 				AppendMenuW(hMenu, MF_SEPARATOR, 0, NULL);
 			}
-			if (m_fileName.data)
+
+			if (!m_traceView.sessions.empty())
 			{
-				if (!m_replay)
-					AppendMenuW(hMenu, MF_STRING, Popup_Replay, L"&Replay Trace");
-				else
+				if (!m_client)
 				{
-					if (m_paused)
-						AppendMenuW(hMenu, MF_STRING, Popup_Play, L"&Play");
+					if (!m_replay || m_traceView.finished)
+						AppendMenuW(hMenu, MF_STRING, Popup_Replay, L"&Replay Trace");
 					else
-						AppendMenuW(hMenu, MF_STRING, Popup_Pause, L"&Pause");
-					AppendMenuW(hMenu, MF_STRING, Popup_JumpToEnd, L"&Jump To End");
+					{
+						if (m_paused)
+							AppendMenuW(hMenu, MF_STRING, Popup_Play, L"&Play");
+						else
+							AppendMenuW(hMenu, MF_STRING, Popup_Pause, L"&Pause");
+						AppendMenuW(hMenu, MF_STRING, Popup_JumpToEnd, L"&Jump To End");
+					}
 				}
-			}
 
-			if (m_fileName.IsEmpty()) {
-				AppendMenuW(hMenu, MF_STRING, Popup_SaveAs, L"&Save Trace");
+				if (m_fileName.IsEmpty())
+					AppendMenuW(hMenu, MF_STRING, Popup_SaveAs, L"&Save Trace");
+
+				AppendMenuW(hMenu, MF_SEPARATOR, 0, NULL);
 			}
-			AppendMenuW(hMenu, MF_SEPARATOR, 0, NULL);
+
 			AppendMenuW(hMenu, MF_STRING, Popup_SaveSettings, L"Save Position/Settings");
+			AppendMenuW(hMenu, MF_STRING, Popup_OpenSettings, L"Open Settings file");
 			AppendMenuW(hMenu, MF_STRING, Popup_Quit, L"&Quit");
 			m_showPopup = true;
 			switch (TrackPopupMenu(hMenu, TPM_RETURNCMD | TPM_RIGHTBUTTON, point.x, point.y, 0, hWnd, NULL))
@@ -2800,12 +3019,12 @@ namespace uba
 			}
 			case Popup_ShowProcessText:
 				m_config.ShowProcessText = !m_config.ShowProcessText;
-				RedrawWindow(hWnd, NULL, NULL, RDW_INVALIDATE | RDW_UPDATENOW);
+				Redraw();
 				break;
 			case Popup_ShowReadWriteColors:
 				m_config.ShowReadWriteColors = !m_config.ShowReadWriteColors;
 				DirtyBitmaps();
-				RedrawWindow(hWnd, NULL, NULL, RDW_INVALIDATE | RDW_UPDATENOW);
+				Redraw();
 				break;
 			case Popup_ScaleHorizontalWithScrollWheel:
 				m_config.ScaleHorizontalWithScrollWheel = !m_config.ScaleHorizontalWithScrollWheel;
@@ -2820,13 +3039,15 @@ namespace uba
 				BOOL useDarkMode = m_config.DarkMode;
 				u32 attribute = 20; // DWMWA_USE_IMMERSIVE_DARK_MODE
 				DwmSetWindowAttribute(m_hwnd, attribute, &useDarkMode, sizeof(useDarkMode));
-				RedrawWindow(hWnd, NULL, NULL, RDW_INVALIDATE | RDW_UPDATENOW);
+				Redraw();
 				break;
 			}
+			case Popup_AutoSaveSettings:
+				m_config.AutoSaveSettings = !m_config.AutoSaveSettings;
+				break;
+
 			case Popup_Replay:
-				m_replay = 1;
-				m_paused = false;
-				PostMessage(m_hwnd, WM_NEWTRACE, 0, 0);
+				PostNewTrace(1, false);
 				break;
 
 			case Popup_Play:
@@ -2839,18 +3060,28 @@ namespace uba
 
 			case Popup_JumpToEnd:
 				m_traceView.finished = true;
-				m_replay = 0;
-				PostMessage(m_hwnd, WM_NEWTRACE, 0, 0);
+				PostNewTrace(0, false);
 				break;
 
 			case Popup_SaveSettings:
 				SaveSettings();
 				break;
 
+			case Popup_OpenSettings:
+				ShellExecuteW(NULL, L"open", m_config.filename.c_str(), NULL, NULL, SW_SHOW);
+				break;
+
 			case Popup_Quit: // Quit
 				m_looping = false;
 				break;
 
+			case Popup_IncreaseFontSize:
+				ChangeFontSize(1);
+				break;
+			case Popup_DecreaseFontSize:
+				ChangeFontSize(-1);
+				break;
+
 			case Popup_CopySessionInfo:
 			{
 				TString str;
@@ -2904,6 +3135,10 @@ namespace uba
 		{
 			if (wParam == VK_SPACE)
 				Pause(!m_paused);
+			if (wParam == VK_ADD)
+				ChangeFontSize(1);
+			if (wParam == VK_SUBTRACT)
+				ChangeFontSize(-1);
 			break;
 		}
 		case WM_VSCROLL:
diff --git a/Engine/Source/Programs/UnrealBuildAccelerator/Visualizer/Private/UbaVisualizer.h b/Engine/Source/Programs/UnrealBuildAccelerator/Visualizer/Private/UbaVisualizer.h
--- a/Engine/Source/Programs/UnrealBuildAccelerator/Visualizer/Private/UbaVisualizer.h
+++ b/Engine/Source/Programs/UnrealBuildAccelerator/Visualizer/Private/UbaVisualizer.h
@@ -11,11 +11,14 @@ namespace uba
 	class Config;
 
 	#define UBA_VISUALIZER_FLAGS1 \
+		UBA_VISUALIZER_FLAG(Progress, true, L"progress") \
+		UBA_VISUALIZER_FLAG(ActiveProcesses, false, L"active processes") \
+		UBA_VISUALIZER_FLAG(TitleBars, true, L"title bars") \
+		UBA_VISUALIZER_FLAG(DetailedData, false, L"detailed data (use -UbaDetailedTrace for even more)") \
 		UBA_VISUALIZER_FLAG(NetworkStats, true, L"network stats") \
 		UBA_VISUALIZER_FLAG(CpuMemStats, true, L"cpu/mem stats") \
 		UBA_VISUALIZER_FLAG(ProcessBars, true, L"process bars") \
 		UBA_VISUALIZER_FLAG(Timeline, true, L"timeline") \
-		UBA_VISUALIZER_FLAG(DetailedData, false, L"detailed data (use -UbaDetailedTrace for even more)") \
 		UBA_VISUALIZER_FLAG(Workers, false, L"workers (threads on host taking care of requests from helpers)") \
 		UBA_VISUALIZER_FLAG(CursorLine, false, L"cursor (vertical line)") \
 
@@ -24,6 +27,7 @@ namespace uba
 		UBA_VISUALIZER_FLAG(ShowReadWriteColors, true, L"Show colors for read/write times in process bars") \
 		UBA_VISUALIZER_FLAG(ScaleHorizontalWithScrollWheel, false, L"Use scroll wheel to scale horizontally") \
 		UBA_VISUALIZER_FLAG(DarkMode, false, L"Use dark mode to draw visualizer") \
+		UBA_VISUALIZER_FLAG(AutoSaveSettings, true, L"Auto save Position/Settings on close") \
 
 	struct VisualizerConfig
 	{
@@ -36,8 +40,11 @@ namespace uba
 
 		int x = 100;
 		int y = 100;
-		int width = 1500;
-		int height = 1500;
+		u32 width = 1500;
+		u32 height = 1500;
+		u32 fontSize = 13;
+		TString fontName;
+		u32 maxActiveVisible = 5;
 
 		#define UBA_VISUALIZER_FLAG(name, defaultValue, desc) bool show##name = defaultValue;
 		UBA_VISUALIZER_FLAGS1
@@ -95,6 +102,7 @@ namespace uba
 			bool workSelected = false;
 			u32 workTrack = ~0u;
 			u32 workIndex = ~0u;
+			TString hyperLink;
 		};
 		void HitTest(HitTestResult& outResult, const POINT& pos);
 
@@ -112,6 +120,9 @@ namespace uba
 		void StopDragToScroll();
 		void SaveSettings();
 		void DirtyBitmaps();
+		void UpdateFont();
+		void ChangeFontSize(int offset);
+		void Redraw();
 
 		StringBuffer<256> m_namedTrace;
 		StringBuffer<256> m_fileName;
@@ -153,6 +164,10 @@ namespace uba
 		HPEN m_processUpdatePen = 0;
 		HPEN m_checkboxPen = 0;
 		HFONT m_font = 0;
+		HFONT m_fontUnderlined = 0;
+		int m_fontHeight = 0;
+		int m_rawBoxHeight = 0;
+		int m_sessionStepY = 0;
 		HFONT m_popupFont = 0;
 		int m_popupFontHeight = 0;
 
@@ -191,6 +206,7 @@ namespace uba
 		u32 m_buttonSelected = ~0u;
 		float m_timelineSelected = 0;
 		u32 m_fetchedFilesSelected = ~0u;
+		TString m_hyperLinkSelected;
 
 		bool m_workSelected = false;
 		u32 m_workTrack = ~0u;
@@ -211,6 +227,7 @@ namespace uba
 
 		Thread m_thread;
 
+		void PostNewTrace(u32 replay, bool paused);
 		LRESULT WinProc(HWND hWnd, UINT Msg, WPARAM wParam, LPARAM lParam);
 		static LRESULT CALLBACK StaticWinProc(HWND hWnd, UINT Msg, WPARAM wParam, LPARAM lParam);
 	};
