[Uba]
* Fixed negative time bug in visualizer for active processes section

diff --git a/Engine/Source/Programs/UnrealBuildAccelerator/Visualizer/Private/UbaVisualizer.cpp b/Engine/Source/Programs/UnrealBuildAccelerator/Visualizer/Private/UbaVisualizer.cpp
--- a/Engine/Source/Programs/UnrealBuildAccelerator/Visualizer/Private/UbaVisualizer.cpp
+++ b/Engine/Source/Programs/UnrealBuildAccelerator/Visualizer/Private/UbaVisualizer.cpp
@@ -881,7 +881,10 @@ namespace uba
 				if (index++ == m_config.maxActiveVisible)
 					break;
 				
-				u32 durationMs = u32(TimeToMs(playTime - kv.second->start));
+				u32 durationMs = 0;
+				if (playTime > kv.second->start)
+					durationMs = u32(TimeToMs(playTime - kv.second->start));
+
 				RECT rect;
 				rect.left = 3 + 6*m_fontHeight;
 				rect.right = rect.left + durationMs/100;
