[UBA]
* Fixed so Config::SaveToFile makes sure directory exists before saving file

diff --git a/Engine/Source/Programs/UnrealBuildAccelerator/Common/Private/UbaConfig.cpp b/Engine/Source/Programs/UnrealBuildAccelerator/Common/Private/UbaConfig.cpp
--- a/Engine/Source/Programs/UnrealBuildAccelerator/Common/Private/UbaConfig.cpp
+++ b/Engine/Source/Programs/UnrealBuildAccelerator/Common/Private/UbaConfig.cpp
@@ -248,6 +248,11 @@ namespace uba
 	{
 		UBA_ASSERT(m_tables.empty()); // TODO
 
+		StringBuffer<> dir;
+		dir.AppendDir(configFile);
+		if (!DirectoryCache().CreateDirectory(logger, dir.data))
+			return false;
+
 		FileAccessor fa(logger, configFile);
 		if (!fa.CreateWrite())
 			return false;
