[UBA]
* Added -config=<file> support to UbaVisualizer
* Added "GetOwnerInfo" which can be used to get id and pid for owner.. right now devenv.exe is the only tracked owner but can add support for more applications
* Fixed so both trace name and default config name is affected by owner info.. this means a UbaVisualizer launched from inside visual studio will use a different config file
* Fixed bugs related to replays of active trace in visualizer

diff --git a/Engine/Source/Programs/UnrealBuildAccelerator/Common/Private/UbaSession.cpp b/Engine/Source/Programs/UnrealBuildAccelerator/Common/Private/UbaSession.cpp
--- a/Engine/Source/Programs/UnrealBuildAccelerator/Common/Private/UbaSession.cpp
+++ b/Engine/Source/Programs/UnrealBuildAccelerator/Common/Private/UbaSession.cpp
@@ -773,7 +773,7 @@ namespace uba
 		{
 			FileInformation info;
 			if (!GetFileInformation(info, m_logger, fileName.data))
-				return m_logger.Error(TC("Failed to get file information for %s while checking file added for write. This should not happen! (%s)"), fileName, LastErrorToText().data);
+				return m_logger.Error(TC("Failed to get file information for %s while checking file added for write. This should not happen! (%s)"), fileName.data, LastErrorToText().data);
 
 			attributes = info.attributes;
 			volumeSerial = info.volumeSerialNumber;
@@ -836,7 +836,7 @@ namespace uba
 		}
 
 		#if UBA_DEBUG_LOGGER
-		g_debugLogger.Info(TC("TRACKADD    %s (Size: %llu, Key: %s, Id: %llu)\n"), fileName, fileSize, KeyToString(fileNameKey).data, fileIndex);
+		g_debugLogger.Info(TC("TRACKADD    %s (Size: %llu, Key: %s, Id: %llu)\n"), fileName.data, fileSize, KeyToString(fileNameKey).data, fileIndex);
 		#endif
 
 
@@ -1121,8 +1121,14 @@ namespace uba
 		if (info.traceName && *info.traceName)
 			traceName.Append(info.traceName);
 		else if (info.launchVisualizer || !m_traceOutputFile.IsEmpty() || info.traceEnabled)
+		{
 			traceName.Append(m_id);
 
+			OwnerInfo ownerInfo = GetOwnerInfo();
+			if (ownerInfo.pid)
+				traceName.Appendf(L"_%s%u", ownerInfo.id, ownerInfo.pid);
+		}
+
 		if (!traceName.IsEmpty())
 			StartTrace(IsWindows ? traceName.data : nullptr); // non-windows named shared memory not implemented (only needed for UbaVisualizer which you can't run on linux either way)
 
diff --git a/Engine/Source/Programs/UnrealBuildAccelerator/Common/Private/UbaTrace.cpp b/Engine/Source/Programs/UnrealBuildAccelerator/Common/Private/UbaTrace.cpp
--- a/Engine/Source/Programs/UnrealBuildAccelerator/Common/Private/UbaTrace.cpp
+++ b/Engine/Source/Programs/UnrealBuildAccelerator/Common/Private/UbaTrace.cpp
@@ -4,6 +4,11 @@
 #include "UbaFileAccessor.h"
 #include "UbaProcessHandle.h"
 
+#if PLATFORM_WINDOWS
+#include <tlhelp32.h>
+#include <psapi.h>
+#endif
+
 namespace uba
 {
 	constexpr u64 TraceMessageMaxSize = 256 * 1024;
@@ -490,4 +495,55 @@ namespace uba
 		#endif
 		return true;
 	}
+
+	static OwnerInfo InternalGetOwnerInfo()
+	{
+		static tchar buffer[MAX_PATH];
+		*buffer = 0;
+
+		OwnerInfo info { buffer, 0 };
+
+		#if PLATFORM_WINDOWS
+		HANDLE h = CreateToolhelp32Snapshot(TH32CS_SNAPPROCESS, 0);
+		PROCESSENTRY32 pe = { 0 };
+		pe.dwSize = sizeof(PROCESSENTRY32);
+		UnorderedMap<u32, u32> pidToParent;
+		if (Process32First(h, &pe))
+		{
+			do
+			{
+				pidToParent[pe.th32ProcessID] = pe.th32ParentProcessID;
+			}
+			while (Process32Next(h, &pe));
+		}
+		CloseHandle(h);
+
+		u32 pid = ::GetCurrentProcessId();
+		while (true)
+		{
+			auto findIt = pidToParent.find(pid);
+			if (findIt == pidToParent.end())
+				break;
+			pid = findIt->second;
+
+			HANDLE Handle = OpenProcess(PROCESS_QUERY_INFORMATION | PROCESS_VM_READ, FALSE, pid);
+			GetModuleFileNameExW(Handle, 0, buffer, MAX_PATH);
+			CloseHandle(Handle);
+			if (!Contains(buffer, L"devenv.exe"))
+				continue;
+			TStrcpy_s(buffer, MAX_PATH, L"vs");
+			info.pid = pid;
+			return info;
+		}
+		*buffer = 0;
+		#endif
+
+		return info;
+	}
+
+	const OwnerInfo& GetOwnerInfo()
+	{
+		static OwnerInfo info = InternalGetOwnerInfo();
+		return info;
+	}
 }
\ No newline at end of file
diff --git a/Engine/Source/Programs/UnrealBuildAccelerator/Common/Public/UbaTrace.h b/Engine/Source/Programs/UnrealBuildAccelerator/Common/Public/UbaTrace.h
--- a/Engine/Source/Programs/UnrealBuildAccelerator/Common/Public/UbaTrace.h
+++ b/Engine/Source/Programs/UnrealBuildAccelerator/Common/Public/UbaTrace.h
@@ -135,4 +135,12 @@ namespace uba
 
 		friend class SessionServer;
 	};
+
+	struct OwnerInfo
+	{
+		const tchar* id;
+		u32 pid;
+	};
+
+	const OwnerInfo& GetOwnerInfo();
 }
\ No newline at end of file
diff --git a/Engine/Source/Programs/UnrealBuildAccelerator/Visualizer/Private/UbaTraceReader.cpp b/Engine/Source/Programs/UnrealBuildAccelerator/Visualizer/Private/UbaTraceReader.cpp
--- a/Engine/Source/Programs/UnrealBuildAccelerator/Visualizer/Private/UbaTraceReader.cpp
+++ b/Engine/Source/Programs/UnrealBuildAccelerator/Visualizer/Private/UbaTraceReader.cpp
@@ -203,7 +203,10 @@ namespace uba
 	{
 		Reset(out);
 
-		if (namedTrace && m_namedTrace != namedTrace)
+		if (!namedTrace)
+			namedTrace = TC("");
+
+		if (*namedTrace && m_namedTrace != namedTrace)
 		{
 			m_memoryHandle.handle = ::OpenFileMappingW(PAGE_READWRITE, false, namedTrace);
 			if (!m_memoryHandle.IsValid())
@@ -217,11 +220,12 @@ namespace uba
 				return false;
 		}
 
+		m_namedTrace = namedTrace;
+
 		m_memoryPos = m_memoryBegin;
 		m_memoryEnd = m_memoryBegin;
 		out.finished = false;
 		out.sessions.emplace_back();
-		out.sessions.back().name = L"LOCAL";
 		bool changed;
 		return UpdateReadNamed(out, replay ? 0ull : ~u64(0), changed);
 	}
diff --git a/Engine/Source/Programs/UnrealBuildAccelerator/Visualizer/Private/UbaVisualizer.cpp b/Engine/Source/Programs/UnrealBuildAccelerator/Visualizer/Private/UbaVisualizer.cpp
--- a/Engine/Source/Programs/UnrealBuildAccelerator/Visualizer/Private/UbaVisualizer.cpp
+++ b/Engine/Source/Programs/UnrealBuildAccelerator/Visualizer/Private/UbaVisualizer.cpp
@@ -260,7 +260,15 @@ namespace uba
 
 			if (traceName.count)
 			{
-				if (!traceName.Equals(m_newTraceName.data))
+				StringBuffer<128> filter;
+				if (!m_config.ShowAllTraces)
+				{
+					OwnerInfo ownerInfo = GetOwnerInfo();
+					if (ownerInfo.pid)
+						filter.Appendf(L"_%s%u", ownerInfo.id, ownerInfo.pid);
+				}
+
+				if (!traceName.Equals(m_newTraceName.data) && traceName.EndsWith(filter.data))
 				{
 					m_newTraceName.Clear().Append(traceName);
 					PostNewTrace(0, false);
@@ -394,9 +402,9 @@ namespace uba
 		m_textBitmaps.clear();
 		m_lastBitmap = 0;
 		m_lastBitmapOffset = BitmapCacheHeight;
-		m_autoScroll = true;
-		m_scrollPosX = 0;
-		m_scrollPosY = 0;
+		//m_autoScroll = true;
+		//m_scrollPosX = 0;
+		//m_scrollPosY = 0;
 		//m_zoomValue = 0.75f;
 		//m_horizontalScaleValue = 1.0f;
 
@@ -749,12 +757,7 @@ namespace uba
 
 	void Visualizer::PaintAll(HDC hdc, const RECT& clientRect)
 	{
-		u64 currentTime = m_paused ? m_pauseStart : GetTime();
-		u64 playTime = 0;
-		if (m_traceView.startTime)
-			playTime = currentTime - m_traceView.startTime - m_pauseTime;
-		if (m_replay)
-			playTime *= m_replay;
+		u64 playTime = GetPlayTime();
 
 		int posY = int(m_scrollPosY);
 		float boxHeight = float(m_rawBoxHeight)*m_zoomValue;
@@ -828,9 +831,8 @@ namespace uba
 			drawStatusText(str, LogEntryType_Info, 6, true);
 		}
 
-		if (!m_traceView.statusMap.empty())
+		if (m_config.showStatus && !m_traceView.statusMap.empty())
 		{
-			posY += 4;
 			u32 lastRow = ~0u;
 			u32 row = ~0u;
 			for (auto& kv : m_traceView.statusMap)
@@ -849,9 +851,9 @@ namespace uba
 				posY += m_fontHeight + 2;
 
 			SetTextColor(hdc, m_textColor);
-			posY += 4;
+			posY += 3;
 		}
-
+		
 		if (m_config.showActiveProcesses && !m_trace.m_activeProcesses.empty())
 		{
 			Map<u64, TraceView::Process*> activeProcesses;
@@ -903,7 +905,9 @@ namespace uba
 					str.Append(kv.second->description);
 
 				if (kv.second->isRemote)
-					str.Append(L" (remote)");
+					str.Append(L" [remote]");
+				else if (kv.second->cacheFetch)
+					str.Append(L" [cache]");
 				if (durationMs > 1000)
 					str.Appendf(L" - %us", durationMs/1000);
 				drawStatusText(str, LogEntryType_Info, 6, true);
@@ -1427,7 +1431,7 @@ namespace uba
 		DeleteObject(nullBmp);
 		DeleteDC(textDC);
 
-		m_contentWidth = ProgressRectLeft + int(TimeToS((lastStop != 0 && lastStop != ~u64(0)) ? lastStop : playTime) * scaleX);
+		m_contentWidth = ProgressRectLeft + Max(0, int(TimeToS((lastStop != 0 && lastStop != ~u64(0)) ? lastStop : playTime) * scaleX));
 		m_contentHeight = posY - int(m_scrollPosY) + stepY + 14;
 
 		float timelineSelected = m_timelineSelected;
@@ -2104,7 +2108,7 @@ namespace uba
 		}
 	}
 
-	void Visualizer::HitTest(HitTestResult& outResult, const POINT& pos)
+	u64 Visualizer::GetPlayTime()
 	{
 		u64 currentTime = m_paused ? m_pauseStart : GetTime();
 		u64 playTime = 0;
@@ -2112,6 +2116,12 @@ namespace uba
 			playTime = currentTime - m_traceView.startTime - m_pauseTime;
 		if (m_replay)
 			playTime *= m_replay;
+		return playTime;
+	}
+
+	void Visualizer::HitTest(HitTestResult& outResult, const POINT& pos)
+	{
+		u64 playTime = GetPlayTime();
 
 		RECT clientRect;
 		GetClientRect(m_hwnd, &clientRect);
@@ -2149,9 +2159,8 @@ namespace uba
 		if (m_config.showProgress && m_traceView.progressProcessesTotal)
 			posY += m_fontHeight + 2;
 
-		if (!m_traceView.statusMap.empty())
+		if (m_config.showStatus && !m_traceView.statusMap.empty())
 		{
-			posY += 4;
 			u32 lastRow = ~0u;
 			u32 row = ~0u;
 			for (auto& kv : m_traceView.statusMap)
@@ -2159,6 +2168,11 @@ namespace uba
 				if (kv.second.text.empty())
 					continue;
 
+				row = u32(kv.first >> 32);
+				if (lastRow != ~0u && lastRow != row)
+					posY += m_fontHeight + 2;
+				lastRow = row;
+
 				if (!kv.second.link.empty())
 				{
 					if (pos.y >= posY && pos.y < posY + m_fontHeight && pos.x > 20 && pos.x < 80) // TODO: x is just hard coded to fit horde for now
@@ -2167,15 +2181,10 @@ namespace uba
 						return;
 					}
 				}
-
-				row = u32(kv.first >> 32);
-				if (lastRow != ~0u && lastRow != row)
-					posY += m_fontHeight + 2;
-				lastRow = row;
 			}
 			if (row != ~0u)
 				posY += m_fontHeight + 2;
-			posY += 4;
+			posY += 3;
 		}
 
 		if (m_config.showActiveProcesses && !m_trace.m_activeProcesses.empty())
@@ -2405,7 +2414,7 @@ namespace uba
 
 		}
 
-		m_contentWidth = ProgressRectLeft + int(TimeToS((lastStop != 0 && lastStop != ~u64(0)) ? lastStop : playTime) * scaleX);
+		m_contentWidth = ProgressRectLeft + Max(0, int(TimeToS((lastStop != 0 && lastStop != ~u64(0)) ? lastStop : playTime) * scaleX));
 		m_contentHeight = posY - int(m_scrollPosY) + stepY + 14;
 	}
 
@@ -2493,10 +2502,7 @@ namespace uba
 		if (!m_autoScroll)
 			return false;
 
-		u64 currentTime = m_paused ? m_pauseStart : GetTime();
-		u64 playTime = currentTime - m_traceView.startTime - m_pauseTime;
-		if (m_replay)
-			playTime *= m_replay;
+		u64 playTime = GetPlayTime();
 
 		RECT rect;
 		GetClientRect(m_hwnd, &rect);
@@ -2574,6 +2580,9 @@ namespace uba
 		{
 			m_replay = u32(wParam);
 			m_paused = lParam;
+			m_autoScroll = true;
+			m_scrollPosX = 0;
+			m_scrollPosY = 0;
 			Reset();
 			StringBuffer<> title;
 			GetTitlePrefix(title);
@@ -2888,19 +2897,27 @@ namespace uba
 						m_trace.UpdateReadNamed(m_traceView, time, changed);
 
 					m_pauseStart = m_startTime + time;
-					if (m_paused || !m_replay)
+					m_pauseTime = m_startTime - m_pauseStart;
+
+					if (!m_paused)
 					{
-						m_pauseTime = 0;
+						m_autoScroll = true;
 						m_replay = 1;
+						SetTimer(m_hwnd, 0, 200, NULL);
 					}
 					else
 					{
-						m_pauseTime = m_startTime - m_pauseStart;
-						SetTimer(m_hwnd, 0, 200, NULL);
+						m_pauseTime = 0;
 					}
 
 					HitTestResult res;
 					HitTest(res, { -1, -1 });
+					
+					RECT r;
+					GetClientRect(hWnd, &r);
+					m_scrollPosX = Min(Max(m_scrollPosX, float(r.right - m_contentWidth)), 0.0f);
+					m_scrollPosY = Min(0.0f, Max(m_scrollPosY, float(r.bottom - m_contentHeight)));
+
 					UpdateScrollbars(true);
 					Redraw();
 				}
@@ -3032,6 +3049,9 @@ namespace uba
 			case Popup_ScaleHorizontalWithScrollWheel:
 				m_config.ScaleHorizontalWithScrollWheel = !m_config.ScaleHorizontalWithScrollWheel;
 				break;
+			case Popup_ShowAllTraces:
+				m_config.ShowAllTraces = !m_config.ShowAllTraces;
+				break;
 			case Popup_DarkMode:
 			{
 				m_config.DarkMode = !m_config.DarkMode;
diff --git a/Engine/Source/Programs/UnrealBuildAccelerator/Visualizer/Private/UbaVisualizer.h b/Engine/Source/Programs/UnrealBuildAccelerator/Visualizer/Private/UbaVisualizer.h
--- a/Engine/Source/Programs/UnrealBuildAccelerator/Visualizer/Private/UbaVisualizer.h
+++ b/Engine/Source/Programs/UnrealBuildAccelerator/Visualizer/Private/UbaVisualizer.h
@@ -12,6 +12,7 @@ namespace uba
 
 	#define UBA_VISUALIZER_FLAGS1 \
 		UBA_VISUALIZER_FLAG(Progress, true, L"progress") \
+		UBA_VISUALIZER_FLAG(Status, true, L"status") \
 		UBA_VISUALIZER_FLAG(ActiveProcesses, false, L"active processes") \
 		UBA_VISUALIZER_FLAG(TitleBars, true, L"title bars") \
 		UBA_VISUALIZER_FLAG(DetailedData, false, L"detailed data (use -UbaDetailedTrace for even more)") \
@@ -28,6 +29,7 @@ namespace uba
 		UBA_VISUALIZER_FLAG(ScaleHorizontalWithScrollWheel, false, L"Use scroll wheel to scale horizontally") \
 		UBA_VISUALIZER_FLAG(DarkMode, false, L"Use dark mode to draw visualizer") \
 		UBA_VISUALIZER_FLAG(AutoSaveSettings, true, L"Auto save Position/Settings on close") \
+		UBA_VISUALIZER_FLAG(ShowAllTraces, true, L"Show all traces started on channel") \
 
 	struct VisualizerConfig
 	{
@@ -104,6 +106,7 @@ namespace uba
 			u32 workIndex = ~0u;
 			TString hyperLink;
 		};
+		u64 GetPlayTime();
 		void HitTest(HitTestResult& outResult, const POINT& pos);
 
 		void WriteProcessStats(Logger& out, TraceView::Process& process);
diff --git a/Engine/Source/Programs/UnrealBuildAccelerator/Visualizer/Private/UbaVisualizerApp.cpp b/Engine/Source/Programs/UnrealBuildAccelerator/Visualizer/Private/UbaVisualizerApp.cpp
--- a/Engine/Source/Programs/UnrealBuildAccelerator/Visualizer/Private/UbaVisualizerApp.cpp
+++ b/Engine/Source/Programs/UnrealBuildAccelerator/Visualizer/Private/UbaVisualizerApp.cpp
@@ -4,6 +4,8 @@
 #include "UbaConfig.h"
 #include "UbaNetworkBackendTcp.h"
 #include "UbaVersion.h"
+#include <tlhelp32.h>
+#include <psapi.h>
 
 using namespace uba;
 
@@ -27,7 +29,7 @@ static int PrintHelp(const tchar* message = nullptr)
 	s.Appendf(TC("  -file=<name>         Name of file to parse\r\n"));
 	s.Appendf(TC("  -listen[=<channel>]  Listen for announcements of new sessions. Defaults to channel '%s'\r\n"), TC("Default"));
 	s.Appendf(TC("  -replay              Visualize the data as if it was running right now\r\n"));
-	s.Appendf(TC("  -theme=<dark/light>  Force dark/light theme\r\n"));
+	s.Appendf(TC("  -config=<file>       Specify config file to use\r\n"));
 	s.Appendf(TC("\r\n"));
 	MessageBox(NULL, s.data, TC("UbaVisualizer"), 0);
 	//wprintf(s.data);
@@ -63,10 +65,45 @@ struct MessageBoxLogWriter : public LogWriter
 
 int WINAPI wWinMain(_In_ HINSTANCE hInstance, _In_opt_ HINSTANCE hPrevInstance, _In_ PWSTR pCmdLine, _In_ int nShowCmd)
 {
+	{
+		HANDLE h = CreateToolhelp32Snapshot(TH32CS_SNAPPROCESS, 0);
+		PROCESSENTRY32 pe = { 0 };
+		pe.dwSize = sizeof(PROCESSENTRY32);
+
+		UnorderedMap<u32, u32> pidToParent;
+		if (Process32First(h, &pe))
+		{
+			do
+			{
+				pidToParent[pe.th32ProcessID] = pe.th32ParentProcessID;
+			} while( Process32Next(h, &pe));
+		}
+
+		CloseHandle(h);
+
+		u32 pid = ::GetCurrentProcessId();
+		while (true)
+		{
+			auto findIt = pidToParent.find(pid);
+			if (findIt == pidToParent.end())
+				break;
+			pid = findIt->second;
+
+			HANDLE Handle = OpenProcess(PROCESS_QUERY_INFORMATION | PROCESS_VM_READ, FALSE, pid);
+			TCHAR Buffer[MAX_PATH];
+			GetModuleFileNameExW(Handle, 0, Buffer, MAX_PATH);
+
+			CloseHandle(Handle);
+		}
+
+	}
+
+
 	StringBuffer<> host; // 192.168.86.49
 	StringBuffer<> named;
 	StringBuffer<> file;
 	StringBuffer<> channel;
+	StringBuffer<> configPath;
 	u32 port = DefaultPort;
 	u32 replay = 0;
 
@@ -134,6 +171,12 @@ int WINAPI wWinMain(_In_ HINSTANCE hInstance, _In_opt_ HINSTANCE hPrevInstance,
 			if (!value.IsEmpty())
 				value.Parse(replay);
 		}
+		else if (name.Equals(TC("-config")))
+		{
+			if (value.IsEmpty())
+				return PrintHelp(TC("-config needs a value"));
+			configPath.Append(value);
+		}
 		else
 		{
 			StringBuffer<> msg;
@@ -150,11 +193,23 @@ int WINAPI wWinMain(_In_ HINSTANCE hInstance, _In_opt_ HINSTANCE hPrevInstance,
 	MessageBoxLogWriter logWriter;
 	LoggerWithWriter logger(logWriter);
 
-	StringBuffer<> configPath;
-	configPath.count = ExpandEnvironmentStringsW(L"%PROGRAMDATA%", configPath.data, configPath.capacity) - 1;
-	configPath.Append(L"\\Epic\\UbaVisualizer\\UbaVisualizer.toml");
+	bool showAllTraces = true;
+	if (!configPath.count)
+	{
+		configPath.count = ExpandEnvironmentStringsW(L"%PROGRAMDATA%", configPath.data, configPath.capacity) - 1;
+		configPath.Append(L"\\Epic\\UbaVisualizer\\UbaVisualizer");
+
+		OwnerInfo ownerInfo = GetOwnerInfo();
+		if (ownerInfo.pid)
+		{
+			configPath.Append('_').Append(ownerInfo.id);
+			showAllTraces = false;
+		}
+		configPath.Append(L".toml");
+	}
 
 	VisualizerConfig visualizerConfig(configPath.data);
+	visualizerConfig.ShowAllTraces = showAllTraces;
 	visualizerConfig.Load(logger);
 
 	NetworkBackendTcp networkBackend(logWriter);
