[UBA]
* Fixed potential race condition related to additional connections and them being added while send call picks up and use the connection before callbacks are set and there is an disconnect

diff --git a/Engine/Source/Programs/UnrealBuildAccelerator/Common/Private/UbaNetworkClient.cpp b/Engine/Source/Programs/UnrealBuildAccelerator/Common/Private/UbaNetworkClient.cpp
--- a/Engine/Source/Programs/UnrealBuildAccelerator/Common/Private/UbaNetworkClient.cpp
+++ b/Engine/Source/Programs/UnrealBuildAccelerator/Common/Private/UbaNetworkClient.cpp
@@ -195,7 +195,7 @@ namespace uba
 	void NetworkClient::DisconnectCallback(void* context, const Guid& connectionUid, void* connection)
 	{
 		auto& c = *(Connection*)context;
-		c.owner.OnDisconnected(c);
+		c.owner.OnDisconnected(c, 1);
 		c.disconnectedEvent.Set();
 	}
 
@@ -209,10 +209,10 @@ namespace uba
 		connection->backendConnection = backendConnection;
 		connection->connected = 1;
 		connection->backend = &backend;
-		{
-			SCOPED_WRITE_LOCK(m_connectionsItLock, l);
-			m_connectionsIt = m_connections.begin();
-		}
+
+		SCOPED_WRITE_LOCK(m_connectionsItLock, l); // Take this lock to make sure callbacks are set before connection is used
+
+		m_logger.Detail(TC("Connected to server... (0x%p)"), backendConnection);
 		lock.Leave();
 
 		backend.SetDisconnectCallback(backendConnection, connection, DisconnectCallback);
@@ -283,7 +283,7 @@ namespace uba
 			SCOPED_READ_LOCK(m_connectionsLock, lock);
 			for (auto& c : m_connections)
 			{
-				OnDisconnected(c);
+				OnDisconnected(c, 0);
 				c.disconnectedEvent.IsSet(~0u);
 			}
 		}
@@ -405,11 +405,11 @@ namespace uba
 		return m_connections.front().backend;
 	}
 
-	void NetworkClient::OnDisconnected(Connection& connection)
+	void NetworkClient::OnDisconnected(Connection& connection, u32 reason)
 	{
 		if (connection.connected.exchange(0) == 1)
 		{
-			m_logger.Detail(TC("Disconnected from server..."));
+			m_logger.Detail(TC("Disconnected from server... (0x%p) (%u)"), connection.backendConnection, reason);
 
 			connection.backend->Shutdown(connection.backendConnection);
 
@@ -422,7 +422,6 @@ namespace uba
 			}
 		}
 
-		u16 messageId = 0;
 		SCOPED_WRITE_LOCK(m_activeMessagesLock, lock);
 		for (auto m : m_activeMessages)
 		{
@@ -431,7 +430,6 @@ namespace uba
 				m->m_error = 3;
 				m->Done(false);
 			}
-			++messageId;
 		}
 	}
 
@@ -530,7 +528,7 @@ namespace uba
 			if (!Crypto::Encrypt(m_logger, m_cryptoKey, data + SendHeaderSize, bodySize))
 			{
 				message.m_error = 8;
-				OnDisconnected(connection);
+				OnDisconnected(connection, 8);
 				return false;
 			}
 		}
@@ -543,7 +541,7 @@ namespace uba
 			if (!connection.backend->Send(m_logger, connection.backendConnection, data, sendSize, message.m_sendContext))
 			{
 				message.m_error = 9;
-				OnDisconnected(connection);
+				OnDisconnected(connection, 9);
 				return false;
 			}
 		}
diff --git a/Engine/Source/Programs/UnrealBuildAccelerator/Common/Public/UbaNetworkClient.h b/Engine/Source/Programs/UnrealBuildAccelerator/Common/Public/UbaNetworkClient.h
--- a/Engine/Source/Programs/UnrealBuildAccelerator/Common/Public/UbaNetworkClient.h
+++ b/Engine/Source/Programs/UnrealBuildAccelerator/Common/Public/UbaNetworkClient.h
@@ -82,7 +82,7 @@ namespace uba
 		static void DisconnectCallback(void* context, const Guid& connectionUid, void* connection);
 		static bool ReceiveResponseHeader(void* context, const Guid& connectionUid, u8* headerData, void*& outBodyContext, u8*& outBodyData, u32& outBodySize);
 		static bool ReceiveResponseBody(void* context, bool recvError, u8* headerData, void* bodyContext, u8* bodyData, u32 bodySize);
-		void OnDisconnected(Connection& connection);
+		void OnDisconnected(Connection& connection, u32 reason);
 		bool Send(NetworkMessage& message, void* response, u32 responseCapacity, bool async);
 		void ReturnMessageId(u16 id);
 		const tchar* SetGetPrefix(const tchar* originalPrefix);
