[UBA]
* Fixed non windows compile errors
* Fixed bug with m_connectionsIt not being set

diff --git a/Engine/Source/Programs/UnrealBuildAccelerator/Common/Private/UbaNetworkClient.cpp b/Engine/Source/Programs/UnrealBuildAccelerator/Common/Private/UbaNetworkClient.cpp
--- a/Engine/Source/Programs/UnrealBuildAccelerator/Common/Private/UbaNetworkClient.cpp
+++ b/Engine/Source/Programs/UnrealBuildAccelerator/Common/Private/UbaNetworkClient.cpp
@@ -211,6 +211,7 @@ namespace uba
 		connection->backend = &backend;
 
 		SCOPED_WRITE_LOCK(m_connectionsItLock, l); // Take this lock to make sure callbacks are set before connection is used
+		m_connectionsIt = --m_connections.end();
 
 		m_logger.Detail(TC("Connected to server... (0x%p)"), backendConnection);
 		lock.Leave();
diff --git a/Engine/Source/Programs/UnrealBuildAccelerator/Common/Private/UbaTrace.cpp b/Engine/Source/Programs/UnrealBuildAccelerator/Common/Private/UbaTrace.cpp
--- a/Engine/Source/Programs/UnrealBuildAccelerator/Common/Private/UbaTrace.cpp
+++ b/Engine/Source/Programs/UnrealBuildAccelerator/Common/Private/UbaTrace.cpp
@@ -498,7 +498,7 @@ namespace uba
 
 	static OwnerInfo InternalGetOwnerInfo()
 	{
-		static tchar buffer[MAX_PATH];
+		static tchar buffer[260];
 		*buffer = 0;
 
 		OwnerInfo info { buffer, 0 };
diff --git a/Engine/Source/Programs/UnrealBuildAccelerator/Visualizer/Private/UbaVisualizerApp.cpp b/Engine/Source/Programs/UnrealBuildAccelerator/Visualizer/Private/UbaVisualizerApp.cpp
--- a/Engine/Source/Programs/UnrealBuildAccelerator/Visualizer/Private/UbaVisualizerApp.cpp
+++ b/Engine/Source/Programs/UnrealBuildAccelerator/Visualizer/Private/UbaVisualizerApp.cpp
@@ -65,40 +65,6 @@ struct MessageBoxLogWriter : public LogWriter
 
 int WINAPI wWinMain(_In_ HINSTANCE hInstance, _In_opt_ HINSTANCE hPrevInstance, _In_ PWSTR pCmdLine, _In_ int nShowCmd)
 {
-	{
-		HANDLE h = CreateToolhelp32Snapshot(TH32CS_SNAPPROCESS, 0);
-		PROCESSENTRY32 pe = { 0 };
-		pe.dwSize = sizeof(PROCESSENTRY32);
-
-		UnorderedMap<u32, u32> pidToParent;
-		if (Process32First(h, &pe))
-		{
-			do
-			{
-				pidToParent[pe.th32ProcessID] = pe.th32ParentProcessID;
-			} while( Process32Next(h, &pe));
-		}
-
-		CloseHandle(h);
-
-		u32 pid = ::GetCurrentProcessId();
-		while (true)
-		{
-			auto findIt = pidToParent.find(pid);
-			if (findIt == pidToParent.end())
-				break;
-			pid = findIt->second;
-
-			HANDLE Handle = OpenProcess(PROCESS_QUERY_INFORMATION | PROCESS_VM_READ, FALSE, pid);
-			TCHAR Buffer[MAX_PATH];
-			GetModuleFileNameExW(Handle, 0, Buffer, MAX_PATH);
-
-			CloseHandle(Handle);
-		}
-
-	}
-
-
 	StringBuffer<> host; // 192.168.86.49
 	StringBuffer<> named;
 	StringBuffer<> file;
