[UBA]
* Fixed another linux compile error

diff --git a/Engine/Source/Programs/UnrealBuildAccelerator/Common/Private/UbaSession.cpp b/Engine/Source/Programs/UnrealBuildAccelerator/Common/Private/UbaSession.cpp
--- a/Engine/Source/Programs/UnrealBuildAccelerator/Common/Private/UbaSession.cpp
+++ b/Engine/Source/Programs/UnrealBuildAccelerator/Common/Private/UbaSession.cpp
@@ -1126,7 +1126,7 @@ namespace uba
 
 			OwnerInfo ownerInfo = GetOwnerInfo();
 			if (ownerInfo.pid)
-				traceName.Appendf(L"_%s%u", ownerInfo.id, ownerInfo.pid);
+				traceName.Appendf(TC("_%s%u"), ownerInfo.id, ownerInfo.pid);
 		}
 
 		if (!traceName.IsEmpty())
