[UBA]
* Fixed bug related to detached process spawning process with custom stdin and populated stdout and stderr with detoured handles

diff --git a/Engine/Source/Programs/UnrealBuildAccelerator/Detours/Private/Windows/UbaDetoursFunctionsKernelBase.inl b/Engine/Source/Programs/UnrealBuildAccelerator/Detours/Private/Windows/UbaDetoursFunctionsKernelBase.inl
--- a/Engine/Source/Programs/UnrealBuildAccelerator/Detours/Private/Windows/UbaDetoursFunctionsKernelBase.inl
+++ b/Engine/Source/Programs/UnrealBuildAccelerator/Detours/Private/Windows/UbaDetoursFunctionsKernelBase.inl
@@ -2738,9 +2738,18 @@ BOOL Detoured_CreateProcessW(LPCWSTR lpApplicationName, LPWSTR lpCommandLine, LP
 
 	LPCSTR dlls[] = { dll };
 
-	UBA_ASSERT(!isDetouredHandle(lpStartupInfo->hStdOutput));
 	UBA_ASSERT(!isDetouredHandle(lpStartupInfo->hStdInput));
-	UBA_ASSERT(!isDetouredHandle(lpStartupInfo->hStdError));
+
+	if (g_isDetachedProcess)
+	{
+		if (lpStartupInfo->hStdError == g_stdHandle[0])
+			lpStartupInfo->hStdError = 0;
+		if (lpStartupInfo->hStdOutput == g_stdHandle[1])
+			lpStartupInfo->hStdOutput = 0;
+	}
+
+	UBA_ASSERT(!isDetouredHandle(lpStartupInfo->hStdOutput));
+	UBA_ASSERTF(!isDetouredHandle(lpStartupInfo->hStdError), L"Got detoured handle for stderror: %llu", lpStartupInfo->hStdError);
 
 	lpStartupInfo->dwFlags |= STARTF_USESHOWWINDOW;
 	lpStartupInfo->wShowWindow = SW_HIDE;
