[UBA]
* Fixed bug where error message box can pop up multiple times
* Fixed so font size scales with processes when zooming

diff --git a/Engine/Source/Programs/UnrealBuildAccelerator/Visualizer/Private/UbaTraceReader.cpp b/Engine/Source/Programs/UnrealBuildAccelerator/Visualizer/Private/UbaTraceReader.cpp
--- a/Engine/Source/Programs/UnrealBuildAccelerator/Visualizer/Private/UbaTraceReader.cpp
+++ b/Engine/Source/Programs/UnrealBuildAccelerator/Visualizer/Private/UbaTraceReader.cpp
@@ -278,10 +278,8 @@ namespace uba
 			u32 traceSize = traceReader.ReadU32(); (void)traceSize;
 			u32 version = traceReader.ReadU32();
 			if (version < TraceReadCompatibilityVersion || version > TraceVersion)
-			{
-				m_logger.Error(L"Incompatible trace version (%u). Current executable supports version %u to %u.", version, TraceReadCompatibilityVersion, TraceVersion);
-				return false;
-			}
+				return m_logger.Error(L"Incompatible trace version (%u). Current executable supports version %u to %u.", version, TraceReadCompatibilityVersion, TraceVersion);
+
 			bool replay = maxTime != ~u64(0);
 
 			out.version = version;
diff --git a/Engine/Source/Programs/UnrealBuildAccelerator/Visualizer/Private/UbaVisualizer.cpp b/Engine/Source/Programs/UnrealBuildAccelerator/Visualizer/Private/UbaVisualizer.cpp
--- a/Engine/Source/Programs/UnrealBuildAccelerator/Visualizer/Private/UbaVisualizer.cpp
+++ b/Engine/Source/Programs/UnrealBuildAccelerator/Visualizer/Private/UbaVisualizer.cpp
@@ -251,6 +251,12 @@ namespace uba
 		StringBuffer<256> traceName;
 		while (m_hwnd)
 		{
+			if (m_locked)
+			{
+				m_listenTimeout.IsSet(1000);
+				continue;
+			}
+
 			traceName.Clear();
 			if (!channel.Read(traceName))
 			{
@@ -366,6 +372,11 @@ namespace uba
 		return m_hwnd;
 	}
 
+	void Visualizer::Lock(bool lock)
+	{
+		m_locked = lock;
+	}
+
 	void Visualizer::GetTitlePrefix(StringBufferBase& out)
 	{
 		out.Append(L"UbaVisualizer");
@@ -525,6 +536,10 @@ namespace uba
 
 		UpdateFont();
 
+		u32 defaultProcessFontHeight = 13;
+		m_zoomValue = defaultProcessFontHeight / 20.0f;
+		UpdateProcessFont();
+
 		const TCHAR* fontName = TEXT("Consolas");
 		m_popupFont = (HFONT)CreateFontW(-12, 0, 0, 0, FW_NORMAL, false, false, false, ANSI_CHARSET, OUT_DEFAULT_PRECIS, CLIP_DEFAULT_PRECIS, 0, FIXED_PITCH | FF_MODERN, fontName);
 		m_popupFontHeight = 14;
@@ -662,12 +677,26 @@ namespace uba
 		m_config.Save(m_logger);
 	}
 
-	void Visualizer::DirtyBitmaps()
+	void Visualizer::DirtyBitmaps(bool full)
 	{
 		for (auto& session : m_traceView.sessions)
 			for (auto& processor : session.processors)
 				for (auto& process : processor.processes)
-						process.bitmapDirty = true;
+				{
+					process.bitmapDirty = true;
+					if (full)
+						process.bitmap = 0;
+				}
+
+		if (!full)
+			return;
+
+		for (HBITMAP bm : m_textBitmaps)
+			DeleteObject(bm);
+		DeleteObject(m_lastBitmap);
+		m_textBitmaps.clear();
+		m_lastBitmapOffset = BitmapCacheHeight;
+		m_lastBitmap = 0;
 	}
 
 	void Visualizer::UpdateFont()
@@ -683,9 +712,23 @@ namespace uba
 		m_fontHeight = m_config.fontSize;
 		m_font = (HFONT)CreateFontW(4 - m_fontHeight, 0, 0, 0, FW_NORMAL, 0, 0, 0, ANSI_CHARSET, OUT_DEFAULT_PRECIS, CLIP_DEFAULT_PRECIS, CLEARTYPE_QUALITY, DEFAULT_PITCH, m_config.fontName.c_str());
 		m_fontUnderlined = (HFONT)CreateFontW(4 - m_fontHeight, 0, 0, 0, FW_NORMAL, 0, 1, 0, ANSI_CHARSET, OUT_DEFAULT_PRECIS, CLIP_DEFAULT_PRECIS, CLEARTYPE_QUALITY, DEFAULT_PITCH, m_config.fontName.c_str());
+		m_sessionStepY = m_fontHeight + 4;
+	}
+
+	void Visualizer::UpdateProcessFont()
+	{
+		int fontHeight = int(m_zoomValue*20.0f);
+
+		//m_processFont = m_font;
+		m_processFontHeight = fontHeight;
+		m_boxHeight = m_processFontHeight + 2;
+		m_progressRectLeft = int(5 + float(m_processFontHeight) * 1.8f);
+
+		if (m_processFont)
+			DeleteObject(m_processFont);
+		m_processFont = (HFONT)CreateFontW(1 - fontHeight, 0, 0, 0, FW_NORMAL, 0, 0, 0, ANSI_CHARSET, OUT_DEFAULT_PRECIS, CLIP_DEFAULT_PRECIS, CLEARTYPE_QUALITY, DEFAULT_PITCH, m_config.fontName.c_str());
 
-		m_rawBoxHeight = m_fontHeight + 2;
-		m_sessionStepY = m_rawBoxHeight + 2;
+		DirtyBitmaps(true);
 	}
 
 	void Visualizer::ChangeFontSize(int offset)
@@ -693,7 +736,7 @@ namespace uba
 		m_config.fontSize += offset;
 		m_config.fontSize = Max(m_config.fontSize, 10u);
 		UpdateFont();
-		DirtyBitmaps();
+		DirtyBitmaps(true);
 		Redraw();
 	}
 
@@ -728,7 +771,6 @@ namespace uba
 		ReleaseDC(m_hwnd, hdc);
 	}
 
-	constexpr int ProgressRectLeft = 30;
 	constexpr int GraphHeight = 30;
 
 	struct SessionRec
@@ -736,12 +778,12 @@ namespace uba
 		TraceView::Session* session;
 		u32 index;
 	};
-	void Populate(SessionRec* recs, TraceView& traceView)
+	void Populate(SessionRec* recs, TraceView& traceView, bool sort)
 	{
 		u32 count = u32(traceView.sessions.size());
 		for (u32 i = 0, e = count; i != e; ++i)
 			recs[i] = { &traceView.sessions[i], i };
-		if (count <= 1)
+		if (count <= 1 || !sort)
 			return;
 		std::sort(recs + 1, recs + traceView.sessions.size(), [](SessionRec& a, SessionRec& b)
 			{
@@ -760,15 +802,19 @@ namespace uba
 		u64 playTime = GetPlayTime();
 
 		int posY = int(m_scrollPosY);
-		float boxHeight = float(m_rawBoxHeight)*m_zoomValue;
+		float boxHeight = float(m_boxHeight);
 		int stepY = int(boxHeight) + 2;
+		int processStepY = int(boxHeight) + 2;
 		float scaleX = 50.0f*m_zoomValue*m_horizontalScaleValue;
 
 		RECT progressRect = clientRect;
-		progressRect.left += ProgressRectLeft;
-		progressRect.bottom -= 30;
+		progressRect.left += m_progressRectLeft;
 
-		bool shouldDrawText = m_zoomValue > 0.27f;
+		if (m_config.showTimeline)
+		{
+			progressRect.bottom -= m_fontHeight + 10;
+		}
+		bool shouldDrawText = m_zoomValue > 0.25f;
 
 		SetBkMode(hdc, TRANSPARENT);
 		SetTextColor(hdc, m_textColor);
@@ -776,7 +822,7 @@ namespace uba
 
 		HDC textDC = CreateCompatibleDC(hdc);
 		SetTextColor(textDC, m_textColor);
-		SelectObject(textDC, m_font);
+		SelectObject(textDC, m_processFont);
 		SelectObject(textDC, GetStockObject(NULL_BRUSH));
 		SetBkMode(textDC, TRANSPARENT);
 		SetBkColor(hdc, m_config.DarkMode ? RGB(70, 70, 70) : RGB(180, 180, 180));
@@ -918,7 +964,7 @@ namespace uba
 		TraceView::WorkRecord selectedWork;
 
 		SessionRec sortedSessions[1024];
-		Populate(sortedSessions, m_traceView);
+		Populate(sortedSessions, m_traceView, m_config.SortActiveRemoteSessions);
 
 		TraceView::ProcessLocation processLocation { 0, 0, 0 };
 		for (u64 i = 0, e = m_traceView.sessions.size(); i != e; ++i)
@@ -1082,6 +1128,8 @@ namespace uba
 				PaintDetailedStats(posY, progressRect, session, isRemote, playTime, drawText);
 			}
 
+			SelectObject(hdc, m_processFont);
+
 			if (m_config.showProcessBars)
 			{
 				processLocation.processorIndex = 0;
@@ -1100,7 +1148,7 @@ namespace uba
 
 						const int textHeight = int(barHeight);
 						const int rectBottom = posY + textHeight;
-						const int offsetY = (textHeight - m_fontHeight + textOffsetY) / 2;
+						const int offsetY = (textHeight - m_processFontHeight + textOffsetY) / 2;
 
 						if (shouldDrawText)
 						{
@@ -1171,7 +1219,7 @@ namespace uba
 										}
 										process.bitmap = m_lastBitmap;
 										process.bitmapOffset = m_lastBitmapOffset;
-										m_lastBitmapOffset += m_fontHeight;
+										m_lastBitmapOffset += m_processFontHeight;
 									}
 									if (lastSelectedBitmap != process.bitmap)
 									{
@@ -1179,8 +1227,8 @@ namespace uba
 										lastSelectedBitmap = process.bitmap;
 									}
 
-									RECT rect2{ 0, int(process.bitmapOffset), 256, int(process.bitmapOffset) + m_fontHeight };
-									RECT rect3{ 0, int(process.bitmapOffset), processWidth, int(process.bitmapOffset) + m_fontHeight };
+									RECT rect2{ 0, int(process.bitmapOffset), 256, int(process.bitmapOffset) + m_processFontHeight };
+									RECT rect3{ 0, int(process.bitmapOffset), processWidth, int(process.bitmapOffset) + m_processFontHeight };
 									if (!done)
 										rect3.right = 256;
 
@@ -1188,18 +1236,21 @@ namespace uba
 
 									rect2.left += 3; // Move in text a bit
 									
+									//SetTextAlign(textDC, TA_LEFT|TA_TOP|TA_NOUPDATECP|VTA_CENTER);
+									int textY = rect2.top - 1;// - int(m_processFontHeight/9);
+
 									bool dropShadow = m_config.DarkMode;
 									if (dropShadow)
 									{
 										SetTextColor(textDC, RGB(5, 60, 5));
 										++rect2.left;
-										++rect2.top;
-										ExtTextOutW(textDC, rect2.left, rect2.top, ETO_CLIPPED, &rect2, process.description.c_str(), int(process.description.size()), NULL);
+										++textY;
+										ExtTextOutW(textDC, rect2.left, textY, ETO_CLIPPED, &rect2, process.description.c_str(), int(process.description.size()), NULL);
 										--rect2.left;
-										--rect2.top;
+										--textY;
 									}
 									SetTextColor(textDC, m_textColor);
-									ExtTextOutW(textDC, rect2.left, rect2.top, ETO_CLIPPED, &rect2, process.description.c_str(), int(process.description.size()), NULL);
+									ExtTextOutW(textDC, rect2.left, textY, ETO_CLIPPED, &rect2, process.description.c_str(), int(process.description.size()), NULL);
 
 									if (!selected)
 										process.bitmapDirty = false;
@@ -1219,7 +1270,7 @@ namespace uba
 									bitmapOffsetY -= bltOffsetY;
 									bltOffsetY = 0;
 								}
-								int height = Min(textHeight, m_fontHeight);
+								int height = Min(textHeight, m_processFontHeight);
 								if (bltOffsetY + height > textHeight)
 									height = textHeight - bltOffsetY;
 
@@ -1244,7 +1295,7 @@ namespace uba
 					lastStop = Max(lastStop, processor.processes.rbegin()->stop);
 
 					++processLocation.processorIndex;
-					posY += stepY;
+					posY += processStepY;
 				}
 			}
 			else
@@ -1272,7 +1323,7 @@ namespace uba
 
 						const int textHeight = int(barHeight);
 						const int rectBottom = posY + textHeight;
-						const int offsetY = (textHeight - m_fontHeight + textOffsetY) / 2;
+						const int offsetY = (textHeight - m_processFontHeight + textOffsetY) / 2;
 
 						if (shouldDrawText)
 						{
@@ -1377,13 +1428,13 @@ namespace uba
 									}
 									SelectObject(textDC, m_lastBitmap);
 
-									RECT rect2{ 0,m_lastBitmapOffset,256, m_lastBitmapOffset + m_fontHeight };
+									RECT rect2{ 0,m_lastBitmapOffset,256, m_lastBitmapOffset + m_processFontHeight };
 
 									FillRect(textDC, &rect2, m_workBrush);
 									ExtTextOutW(textDC, rect2.left, rect2.top, ETO_CLIPPED, &rect2, work.description, int(wcslen(work.description)), NULL);
 									work.bitmap = m_lastBitmap;
 									work.bitmapOffset = m_lastBitmapOffset;
-									m_lastBitmapOffset += m_fontHeight;
+									m_lastBitmapOffset += m_processFontHeight;
 								}
 
 								if (lastSelectedBitmap != work.bitmap)
@@ -1400,7 +1451,7 @@ namespace uba
 									bitmapOffsetY -= bltOffsetY;
 									bltOffsetY = 0;
 								}
-								int height = Min(textHeight, m_fontHeight);
+								int height = Min(textHeight, m_processFontHeight);
 								if (bltOffsetY + height > textHeight)
 									height = textHeight - bltOffsetY;
 
@@ -1425,13 +1476,15 @@ namespace uba
 					posY += stepY;
 				}
 			}
+
+			SelectObject(hdc, m_font);
 		}
 
 		SelectObject(textDC, oldBmp);
 		DeleteObject(nullBmp);
 		DeleteDC(textDC);
 
-		m_contentWidth = ProgressRectLeft + Max(0, int(TimeToS((lastStop != 0 && lastStop != ~u64(0)) ? lastStop : playTime) * scaleX));
+		m_contentWidth = m_progressRectLeft + Max(0, int(TimeToS((lastStop != 0 && lastStop != ~u64(0)) ? lastStop : playTime) * scaleX));
 		m_contentHeight = posY - int(m_scrollPosY) + stepY + 14;
 
 		float timelineSelected = m_timelineSelected;
@@ -1446,7 +1499,7 @@ namespace uba
 			POINT pos;
 			GetCursorPos(&pos);
 			ScreenToClient(m_hwnd, &pos);
-			timelineSelected = startOffset + (pos.x - ProgressRectLeft) / timeScale;
+			timelineSelected = startOffset + (pos.x - m_progressRectLeft) / timeScale;
 		}
 
 		if (timelineSelected)
@@ -1538,6 +1591,8 @@ namespace uba
 			}
 			if (m_config.showNetworkStats)
 			{
+				SelectObject(hdc, m_font);
+
 				SetTextColor(hdc, m_sendColor);
 				ExtTextOutW(hdc, left, top, 0, NULL, L"SND", 3, NULL);
 				left -= 25;
@@ -1903,10 +1958,9 @@ namespace uba
 
 	void Visualizer::PaintTimeline(HDC hdc, const RECT& clientRect)
 	{
-		float boxHeight = float(m_rawBoxHeight) * m_zoomValue;
-		int posY = m_contentHeight - int(boxHeight) - 14;
+		int posY = m_contentHeight - m_fontHeight - 8;
 		float timeScale = (m_horizontalScaleValue*m_zoomValue)*50.0f;
-		int top = Min(posY, int(clientRect.bottom - m_sessionStepY - 10));
+		int top = Min(posY, int(clientRect.bottom - m_fontHeight - 8));
 			
 		float startOffset = ((m_scrollPosX/timeScale) - int(m_scrollPosX/timeScale)) * timeScale;
 		int index = -int(startOffset/timeScale);
@@ -1930,7 +1984,7 @@ namespace uba
 		int lineStepSize = textStepSize / 5;
 
 		RECT progressRect = clientRect;
-		progressRect.left += ProgressRectLeft;
+		progressRect.left += m_progressRectLeft;
 
 		SelectObject(hdc, m_textPen);
 
@@ -1963,7 +2017,7 @@ namespace uba
 				{
 					RECT textRect;
 					textRect.top = top + 8;
-					textRect.bottom = textRect.top + 15;
+					textRect.bottom = textRect.top + m_fontHeight;
 					textRect.right = pos + 20;
 					textRect.left = pos - 20;
 					DrawTextW(hdc, buffer.data, buffer.count, &textRect, DT_SINGLELINE | DT_CENTER);
@@ -2127,12 +2181,12 @@ namespace uba
 		GetClientRect(m_hwnd, &clientRect);
 
 		int posY = int(m_scrollPosY);
-		float boxHeight = float(m_rawBoxHeight)*m_zoomValue;
-		int stepY = int(boxHeight) + 2;
+		float boxHeight = float(m_boxHeight);
+		int processStepY = int(boxHeight) + 2;
 		float scaleX = 50.0f*m_zoomValue*m_horizontalScaleValue;
 
 		RECT progressRect = clientRect;
-		progressRect.left += ProgressRectLeft;
+		progressRect.left += m_progressRectLeft;
 		progressRect.bottom -= 30;
 
 		{
@@ -2196,7 +2250,7 @@ namespace uba
 		TraceView::ProcessLocation& outLocation = outResult.processLocation;
 
 		SessionRec sortedSessions[1024];
-		Populate(sortedSessions, m_traceView);
+		Populate(sortedSessions, m_traceView, m_config.SortActiveRemoteSessions);
 
 		for (u64 i = 0, e = m_traceView.sessions.size(); i != e; ++i)
 		{
@@ -2287,7 +2341,7 @@ namespace uba
 				u32 processorIndex = 0;
 				for (auto& processor : session.processors)
 				{
-					if (pos.y < progressRect.bottom && posY + stepY >= progressRect.top && posY <= progressRect.bottom && pos.y >= posY-1 && pos.y < posY-1 + stepY)
+					if (pos.y < progressRect.bottom && posY + processStepY >= progressRect.top && posY <= progressRect.bottom && pos.y >= posY-1 && pos.y < posY-1 + processStepY)
 					{
 						u32 processIndex = 0;
 						int posX = int(m_scrollPosX) + progressRect.left;
@@ -2336,7 +2390,7 @@ namespace uba
 					if (!processor.processes.empty())
 						lastStop = Max(lastStop, processor.processes.rbegin()->stop);
 
-					posY += stepY;
+					posY += processStepY;
 					++processorIndex;
 				}
 			}
@@ -2352,7 +2406,7 @@ namespace uba
 				int trackIndex = 0;
 				for (auto& workTrack : m_traceView.workTracks)
 				{
-					if (pos.y < progressRect.bottom && posY + stepY >= progressRect.top && posY <= progressRect.bottom && pos.y >= posY-1 && pos.y < posY-1 + stepY)
+					if (pos.y < progressRect.bottom && posY + processStepY >= progressRect.top && posY <= progressRect.bottom && pos.y >= posY-1 && pos.y < posY-1 + processStepY)
 					{
 						u32 workIndex = 0;
 						int posX = int(m_scrollPosX) + progressRect.left;
@@ -2397,7 +2451,7 @@ namespace uba
 						}
 					}
 					++trackIndex;
-					posY += stepY;
+					posY += processStepY;
 				}
 			}
 
@@ -2408,14 +2462,14 @@ namespace uba
 				{
 					float timeScale = (m_horizontalScaleValue * m_zoomValue)*50.0f;
 					float startOffset = -(m_scrollPosX / timeScale);
-					outResult.timelineSelected = startOffset + (pos.x - ProgressRectLeft) / timeScale;
+					outResult.timelineSelected = startOffset + (pos.x - m_progressRectLeft) / timeScale;
 				}
 			}
 
 		}
 
-		m_contentWidth = ProgressRectLeft + Max(0, int(TimeToS((lastStop != 0 && lastStop != ~u64(0)) ? lastStop : playTime) * scaleX));
-		m_contentHeight = posY - int(m_scrollPosY) + stepY + 14;
+		m_contentWidth = m_progressRectLeft + Max(0, int(TimeToS((lastStop != 0 && lastStop != ~u64(0)) ? lastStop : playTime) * scaleX));
+		m_contentHeight = posY - int(m_scrollPosY) + processStepY + 14;
 	}
 
 	void Visualizer::WriteProcessStats(Logger& out, TraceView::Process& process)
@@ -2508,7 +2562,7 @@ namespace uba
 		GetClientRect(m_hwnd, &rect);
 		float timeS = TimeToS(playTime);
 		float oldScrollPosX = m_scrollPosX;
-		m_scrollPosX = Min(0.0f, (float)rect.right - timeS*50.0f*m_horizontalScaleValue*m_zoomValue - ProgressRectLeft);
+		m_scrollPosX = Min(0.0f, (float)rect.right - timeS*50.0f*m_horizontalScaleValue*m_zoomValue - m_progressRectLeft);
 		return oldScrollPosX != m_scrollPosX;
 	}
 
@@ -2708,10 +2762,11 @@ namespace uba
 			if (m_dragToScrollCounter > 0)
 				break;
 
-			int delta = GET_WHEEL_DELTA_WPARAM(wParam);// / WHEEL_DELTA;
+			int delta = GET_WHEEL_DELTA_WPARAM(wParam);
 			bool controlDown = GetAsyncKeyState(VK_CONTROL) & (1<<15);
+			bool shiftDown = GetAsyncKeyState(VK_LSHIFT) & (1<<15);
 
-			if (m_config.ScaleHorizontalWithScrollWheel || controlDown)
+			if (m_config.ScaleHorizontalWithScrollWheel || controlDown || shiftDown)
 			{
 				RECT r;
 				GetClientRect(hWnd, &r);
@@ -2720,23 +2775,41 @@ namespace uba
 				POINT cursorPos = {};
 				GetCursorPos(&cursorPos);
 				ScreenToClient(m_hwnd, &cursorPos);
-				const float scrollAnchorOffsetX = float(cursorPos.x) - ProgressRectLeft;
+
+				float newZoomValue = m_zoomValue;
+				float newScaleValue = m_horizontalScaleValue;
 
 				if (controlDown)
 				{
-					float newValue = Max(m_zoomValue + float(delta)*0.0005f, 0.05f);
-					m_scrollPosY = Min(0.0f, float(m_scrollPosY)*newValue/m_zoomValue);//LOWORD(lParam);
-					m_scrollPosX = Min(0.0f, float(m_scrollPosX - scrollAnchorOffsetX)*newValue/m_zoomValue + scrollAnchorOffsetX);//LOWORD(lParam);
-
-					m_zoomValue = newValue;
+					if (delta < 0)
+					{
+						newZoomValue = (m_processFontHeight - 1) / 20.0f;
+						if (newZoomValue == 0)
+							newZoomValue = m_zoomValue;
+					}
+					else if (delta > 0)
+						newZoomValue = (m_processFontHeight + 1) / 20.0f;
 				}
 				else
+					newScaleValue = m_horizontalScaleValue + m_horizontalScaleValue*float(delta)*0.0006f;
+
+				// TODO: m_progressRectLeft changes with zoom so anchor logic is wrong
+				const float scrollAnchorOffsetX = float(cursorPos.x) - m_progressRectLeft;
+				const float scrollAnchorOffsetY = 0;//float(cursorPos.y)*m_zoomValue;// - m_progressRectLeft;
+
+				m_scrollPosY = Min(0.0f, float(m_scrollPosY - scrollAnchorOffsetY)*(newZoomValue/m_zoomValue) + scrollAnchorOffsetY);
+				m_scrollPosX = Min(0.0f, float(m_scrollPosX - scrollAnchorOffsetX)*(newZoomValue/m_zoomValue)*(newScaleValue/m_horizontalScaleValue) + scrollAnchorOffsetX);//LOWORD(lParam);
+
+				if (m_horizontalScaleValue != newScaleValue)
+					m_horizontalScaleValue = newScaleValue;
+
+				if (m_zoomValue != newZoomValue)
 				{
-					float newValue = m_horizontalScaleValue + m_horizontalScaleValue*float(delta)*0.0006f;
-					m_scrollPosX = Min(0.0f, float(m_scrollPosX - scrollAnchorOffsetX)*newValue/m_horizontalScaleValue + scrollAnchorOffsetX);//LOWORD(lParam);
-					m_horizontalScaleValue = newValue;
+					m_zoomValue = newZoomValue;
+					UpdateProcessFont();
 				}
 
+
 				UpdateAutoscroll();
 				UpdateSelection();
 
@@ -2752,7 +2825,7 @@ namespace uba
 					for (auto& session : m_traceView.sessions)
 						for (auto& processor : session.processors)
 							for (auto& process : processor.processes)
-								if (TimeToMs(process.writeFilesTime, m_traceView.frequency) >= 300 || TimeToMs(process.createFilesTime, m_traceView.frequency) >= 300)
+								//if (TimeToMs(process.writeFilesTime, m_traceView.frequency) >= 300 || TimeToMs(process.createFilesTime, m_traceView.frequency) >= 300)
 									process.bitmapDirty = true;
 
 				UpdateScrollbars(true);
@@ -3043,7 +3116,7 @@ namespace uba
 				break;
 			case Popup_ShowReadWriteColors:
 				m_config.ShowReadWriteColors = !m_config.ShowReadWriteColors;
-				DirtyBitmaps();
+				DirtyBitmaps(false);
 				Redraw();
 				break;
 			case Popup_ScaleHorizontalWithScrollWheel:
@@ -3052,10 +3125,13 @@ namespace uba
 			case Popup_ShowAllTraces:
 				m_config.ShowAllTraces = !m_config.ShowAllTraces;
 				break;
+			case Popup_SortActiveRemoteSessions:
+				m_config.SortActiveRemoteSessions = !m_config.SortActiveRemoteSessions;
+				break;
 			case Popup_DarkMode:
 			{
 				m_config.DarkMode = !m_config.DarkMode;
-				DirtyBitmaps();
+				DirtyBitmaps(false);
 				InitBrushes();
 				SetWindowTheme(m_hwnd, m_config.DarkMode ? L"DarkMode_Explorer" : L"Explorer", NULL);
 				SendMessageW(m_hwnd, WM_THEMECHANGED, 0, 0);
diff --git a/Engine/Source/Programs/UnrealBuildAccelerator/Visualizer/Private/UbaVisualizer.h b/Engine/Source/Programs/UnrealBuildAccelerator/Visualizer/Private/UbaVisualizer.h
--- a/Engine/Source/Programs/UnrealBuildAccelerator/Visualizer/Private/UbaVisualizer.h
+++ b/Engine/Source/Programs/UnrealBuildAccelerator/Visualizer/Private/UbaVisualizer.h
@@ -30,6 +30,7 @@ namespace uba
 		UBA_VISUALIZER_FLAG(DarkMode, false, L"Use dark mode to draw visualizer") \
 		UBA_VISUALIZER_FLAG(AutoSaveSettings, true, L"Auto save Position/Settings on close") \
 		UBA_VISUALIZER_FLAG(ShowAllTraces, true, L"Show all traces started on channel") \
+		UBA_VISUALIZER_FLAG(SortActiveRemoteSessions, true, L"Sort active sessions on top") \
 
 	struct VisualizerConfig
 	{
@@ -70,6 +71,7 @@ namespace uba
 
 		bool HasWindow();
 		HWND GetHwnd();
+		void Lock(bool lock);
 
 	private:
 		bool StartHwndThread();
@@ -122,8 +124,9 @@ namespace uba
 		void StartDragToScroll(const POINT& anchor);
 		void StopDragToScroll();
 		void SaveSettings();
-		void DirtyBitmaps();
+		void DirtyBitmaps(bool full);
 		void UpdateFont();
+		void UpdateProcessFont();
 		void ChangeFontSize(int offset);
 		void Redraw();
 
@@ -168,8 +171,11 @@ namespace uba
 		HPEN m_checkboxPen = 0;
 		HFONT m_font = 0;
 		HFONT m_fontUnderlined = 0;
+		HFONT m_processFont = 0;
+		int m_progressRectLeft = 30;
+		int m_processFontHeight = 0;
 		int m_fontHeight = 0;
-		int m_rawBoxHeight = 0;
+		int m_boxHeight = 0;
 		int m_sessionStepY = 0;
 		HFONT m_popupFont = 0;
 		int m_popupFontHeight = 0;
@@ -191,7 +197,7 @@ namespace uba
 
 		float m_scrollPosX = 0;
 		float m_scrollPosY = 0;
-		float m_zoomValue = 0.75f;
+		float m_zoomValue = 0.5f;
 		float m_horizontalScaleValue = 0.5f;
 		bool m_autoScroll = true;
 		bool m_paused = false;
@@ -217,6 +223,7 @@ namespace uba
 		
 		bool m_mouseOverWindow = false;
 		bool m_showPopup = false;
+		bool m_locked = false;
 
 		HBITMAP m_cachedBitmap = 0;
 		RECT m_cachedBitmapRect = { INT_MIN, INT_MIN, INT_MIN, INT_MIN };
diff --git a/Engine/Source/Programs/UnrealBuildAccelerator/Visualizer/Private/UbaVisualizerApp.cpp b/Engine/Source/Programs/UnrealBuildAccelerator/Visualizer/Private/UbaVisualizerApp.cpp
--- a/Engine/Source/Programs/UnrealBuildAccelerator/Visualizer/Private/UbaVisualizerApp.cpp
+++ b/Engine/Source/Programs/UnrealBuildAccelerator/Visualizer/Private/UbaVisualizerApp.cpp
@@ -51,13 +51,21 @@ struct MessageBoxLogWriter : public LogWriter
 		if (type > LogEntryType_Warning)
 			return;
 
-		HWND hwnd = m_visualizer ? m_visualizer->GetHwnd() : NULL;
-		UINT flags = MB_ICONERROR;
+		HWND hwnd = NULL;
+		if (m_visualizer)
+		{
+			hwnd = m_visualizer->GetHwnd();
+			m_visualizer->Lock(true);
+		}
+
+		UINT flags = type == LogEntryType_Error ? MB_ICONERROR : MB_ICONWARNING;
 		if (!hwnd)
 			flags |= MB_TOPMOST;
 		MessageBox(hwnd, str, TC("UbaVisualizer"), flags);
 		if (type == LogEntryType_Error)
 			ExitProcess(~0u);
+		if (m_visualizer)
+			m_visualizer->Lock(false);
 	}
 	
 	Visualizer* m_visualizer = nullptr;
