[UBA]
* Fixed so GetOwnerInfo searches max 5 processes up for devenv.exe. This will fix some infinite loop issue on the farm where processes can be parents to each other in a circular way

diff --git a/Engine/Source/Programs/UnrealBuildAccelerator/Common/Private/UbaTrace.cpp b/Engine/Source/Programs/UnrealBuildAccelerator/Common/Private/UbaTrace.cpp
--- a/Engine/Source/Programs/UnrealBuildAccelerator/Common/Private/UbaTrace.cpp
+++ b/Engine/Source/Programs/UnrealBuildAccelerator/Common/Private/UbaTrace.cpp
@@ -519,7 +519,8 @@ namespace uba
 		CloseHandle(h);
 
 		u32 pid = ::GetCurrentProcessId();
-		while (true)
+		u32 maxDepth = 5; // It seems there could be cirular child->parent dependencies.. so having max depth will solve that
+		while (maxDepth--)
 		{
 			auto findIt = pidToParent.find(pid);
 			if (findIt == pidToParent.end())
