[UBA]
* Cleaned up GetOwnerInfo and changed so visited pids are erased from lookup which means they will leave when finding circular dependencies

diff --git a/Engine/Source/Programs/UnrealBuildAccelerator/Common/Private/UbaTrace.cpp b/Engine/Source/Programs/UnrealBuildAccelerator/Common/Private/UbaTrace.cpp
--- a/Engine/Source/Programs/UnrealBuildAccelerator/Common/Private/UbaTrace.cpp
+++ b/Engine/Source/Programs/UnrealBuildAccelerator/Common/Private/UbaTrace.cpp
@@ -504,39 +504,46 @@ namespace uba
 		OwnerInfo info { buffer, 0 };
 
 		#if PLATFORM_WINDOWS
-		HANDLE h = CreateToolhelp32Snapshot(TH32CS_SNAPPROCESS, 0);
+		HANDLE snapshotHandle = CreateToolhelp32Snapshot(TH32CS_SNAPPROCESS, 0);
+		if (snapshotHandle == INVALID_HANDLE_VALUE)
+			return info;
+
 		PROCESSENTRY32 pe = { 0 };
 		pe.dwSize = sizeof(PROCESSENTRY32);
 		UnorderedMap<u32, u32> pidToParent;
-		if (Process32First(h, &pe))
+		if (Process32First(snapshotHandle, &pe))
 		{
 			do
 			{
 				pidToParent[pe.th32ProcessID] = pe.th32ParentProcessID;
 			}
-			while (Process32Next(h, &pe));
+			while (Process32Next(snapshotHandle, &pe));
 		}
-		CloseHandle(h);
+		CloseHandle(snapshotHandle);
 
 		u32 pid = ::GetCurrentProcessId();
-		u32 maxDepth = 5; // It seems there could be cirular child->parent dependencies.. so having max depth will solve that
-		while (maxDepth--)
+		while (true)
 		{
 			auto findIt = pidToParent.find(pid);
 			if (findIt == pidToParent.end())
 				break;
 			pid = findIt->second;
+			pidToParent.erase(findIt);
 
-			HANDLE Handle = OpenProcess(PROCESS_QUERY_INFORMATION | PROCESS_VM_READ, FALSE, pid);
-			GetModuleFileNameExW(Handle, 0, buffer, MAX_PATH);
-			CloseHandle(Handle);
-			if (!Contains(buffer, L"devenv.exe"))
+			HANDLE parentHandle = OpenProcess(PROCESS_QUERY_INFORMATION | PROCESS_VM_READ, FALSE, pid);
+			if (parentHandle == NULL)
+				break;
+			tchar moduleName[260];
+			DWORD len = GetModuleFileNameExW(parentHandle, 0, moduleName, MAX_PATH);
+			CloseHandle(parentHandle);
+			if (!len)
+				break;
+			if (!Contains(moduleName, L"devenv.exe"))
 				continue;
 			TStrcpy_s(buffer, MAX_PATH, L"vs");
 			info.pid = pid;
-			return info;
+			break;
 		}
-		*buffer = 0;
 		#endif
 
 		return info;
