[UBA]
* Changed so ubavisualizer copy itself to temp and relaunch from there by default. Can be disabled with "-nocopy"
* Added support for creating ubavisualizer with a parent. Provide -parent=<hwnd-in-hex> when launching UbaVisualizer to make it a child

diff --git a/Engine/Source/Programs/UnrealBuildAccelerator/Core/Public/UbaMemory.h b/Engine/Source/Programs/UnrealBuildAccelerator/Core/Public/UbaMemory.h
--- a/Engine/Source/Programs/UnrealBuildAccelerator/Core/Public/UbaMemory.h
+++ b/Engine/Source/Programs/UnrealBuildAccelerator/Core/Public/UbaMemory.h
@@ -228,6 +228,7 @@ namespace uba
 	inline u8 HexToByte(tchar c) { return (c >= '0' && c <= '9') ? u8(c - '0') : u8(c - 'a' + 10); }
 	constexpr tchar g_hexChars[] = TC("0123456789abcdef");
 
+	// TODO: These are backwards but changing would break cas storage
 	inline u32 ValueToString(tchar* out, int capacity, u64 value)
 	{
 		(void)capacity;
@@ -244,6 +245,7 @@ namespace uba
 		return u32(it - out);
 	}
 
+	// TODO: These are backwards but changing would break cas storage
 	inline u64 StringToValue(const tchar* str, u64 len)
 	{
 		u64 v = 0;
@@ -257,4 +259,18 @@ namespace uba
 
 		return v;
 	}
+
+	inline u64 StringToValue2(const tchar* str, u64 len)
+	{
+		u64 v = 0;
+		const tchar* pos = str;
+		while (*pos)
+		{
+			u8 a = HexToByte(*pos++);
+			u8 b = HexToByte(*pos++);
+			v = u64(v << 8) | u64(a << 4 | b);
+		}
+
+		return v;
+	}
 }
diff --git a/Engine/Source/Programs/UnrealBuildAccelerator/UbaAgent.Target.cs b/Engine/Source/Programs/UnrealBuildAccelerator/UbaAgent.Target.cs
--- a/Engine/Source/Programs/UnrealBuildAccelerator/UbaAgent.Target.cs
+++ b/Engine/Source/Programs/UnrealBuildAccelerator/UbaAgent.Target.cs
@@ -6,7 +6,7 @@ using System.IO;
 [SupportedPlatforms(UnrealPlatformClass.Desktop)]
 public class UbaAgentTarget : TargetRules
 {
-	const string UbaVersion = "Uba_v0.1.3";
+	const string UbaVersion = "Uba_v1.0.0";
 
 	public static void CommonUbaSettings(TargetRules Rules, TargetInfo Target, bool ShouldExport = false)
 	{
diff --git a/Engine/Source/Programs/UnrealBuildAccelerator/Visualizer/Private/UbaVisualizer.cpp b/Engine/Source/Programs/UnrealBuildAccelerator/Visualizer/Private/UbaVisualizer.cpp
--- a/Engine/Source/Programs/UnrealBuildAccelerator/Visualizer/Private/UbaVisualizer.cpp
+++ b/Engine/Source/Programs/UnrealBuildAccelerator/Visualizer/Private/UbaVisualizer.cpp
@@ -548,29 +548,45 @@ namespace uba
 		DWORD windowStyle = WS_CAPTION | WS_SYSMENU | WS_THICKFRAME | WS_MAXIMIZEBOX | WS_MINIMIZEBOX | WS_CLIPCHILDREN | WS_VSCROLL | WS_HSCROLL;
 		const TCHAR* windowClassName = MAKEINTATOM(wndClassAtom);
 
+		DWORD exStyle = 0;
+		if (m_config.parent)
+			windowStyle = WS_POPUP | WS_VSCROLL | WS_HSCROLL;// | WS_VISIBLE;
+
 		StringBuffer<> title;
 		GetTitlePrefix(title);
 		title.Append(L"Initializing...");
 
-		m_hwnd = CreateWindowEx(0, windowClassName, title.data, windowStyle, winPosX, winPosY, winWidth, winHeight, NULL, NULL, hInstance, this);
-		SetWindowLongPtr(m_hwnd, GWLP_USERDATA, (LONG_PTR)this);
+		HWND hwnd = CreateWindowEx(exStyle, windowClassName, title.data, windowStyle, winPosX, winPosY, winWidth, winHeight, NULL, NULL, hInstance, this);
+		SetWindowLongPtr(hwnd, GWLP_USERDATA, (LONG_PTR)this);
 
 		BOOL cloak = TRUE;
-		DwmSetWindowAttribute(m_hwnd, DWMWA_CLOAK, &cloak, sizeof(cloak));
+		DwmSetWindowAttribute(hwnd, DWMWA_CLOAK, &cloak, sizeof(cloak));
 
 		//if (m_config.DarkMode)
 		{
-			SetWindowTheme(m_hwnd, L"DarkMode_Explorer", NULL);
-			SendMessageW(m_hwnd, WM_THEMECHANGED, 0, 0);
+			SetWindowTheme(hwnd, L"DarkMode_Explorer", NULL);
+			SendMessageW(hwnd, WM_THEMECHANGED, 0, 0);
 			BOOL useDarkMode = true;
 			u32 attribute = 20; // DWMWA_USE_IMMERSIVE_DARK_MODE
-			DwmSetWindowAttribute(m_hwnd, attribute, &useDarkMode, sizeof(useDarkMode));
+			DwmSetWindowAttribute(hwnd, attribute, &useDarkMode, sizeof(useDarkMode));
 		}
 
 		HitTestResult res;
 		HitTest(res, { -1, -1 });
 
-		ShowWindow(m_hwnd, SW_SHOW);
+		if (m_config.parent)
+		{
+			m_parentHwnd = (HWND)(uintptr_t)m_config.parent;
+			if (!SetParent(hwnd, m_parentHwnd))
+ 				m_logger.Error(L"SetParent failed using parentHwnd 0x%llx", m_parentHwnd);
+
+			PostMessage(m_parentHwnd, 0x0444, 0, (LPARAM)hwnd);
+		}
+
+		m_hwnd = hwnd;
+
+		if (!m_parentHwnd)
+			ShowWindow(m_hwnd, SW_SHOW);
 		UpdateWindow(m_hwnd);
 		UpdateScrollbars(true);
 
@@ -2626,6 +2642,12 @@ namespace uba
 		PostMessage(m_hwnd, WM_NEWTRACE, replay, paused);
 	}
 
+	void Visualizer::PostQuit()
+	{
+		m_looping = false;
+		PostMessage(m_hwnd, WM_USER+666, 0, 0);
+	}
+
 	LRESULT Visualizer::WinProc(HWND hWnd, UINT Msg, WPARAM wParam, LPARAM lParam)
 	{
 		switch (Msg)
@@ -2681,14 +2703,13 @@ namespace uba
 			// and GetMessage will stay stuck indefinitely.
 			if (wParam == SC_CLOSE)
 			{
-				m_looping = false;
+				PostQuit();
 				return 0;
 			}
 		break;
 
 		case WM_DESTROY:
-			m_looping = false;
-			//PostQuitMessage(0);
+			PostQuit();
 			return 0;
 
 		case WM_ERASEBKGND:
@@ -3171,7 +3192,7 @@ namespace uba
 				break;
 
 			case Popup_Quit: // Quit
-				m_looping = false;
+				PostQuit();
 				break;
 
 			case Popup_IncreaseFontSize:
diff --git a/Engine/Source/Programs/UnrealBuildAccelerator/Visualizer/Private/UbaVisualizer.h b/Engine/Source/Programs/UnrealBuildAccelerator/Visualizer/Private/UbaVisualizer.h
--- a/Engine/Source/Programs/UnrealBuildAccelerator/Visualizer/Private/UbaVisualizer.h
+++ b/Engine/Source/Programs/UnrealBuildAccelerator/Visualizer/Private/UbaVisualizer.h
@@ -56,6 +56,8 @@ namespace uba
 		#define UBA_VISUALIZER_FLAG(name, defaultValue, desc) bool name = defaultValue;
 		UBA_VISUALIZER_FLAGS2
 		#undef UBA_VISUALIZER_FLAG
+
+		u64 parent = 0;
 	};
 
 	class Visualizer
@@ -151,6 +153,7 @@ namespace uba
 
 		Atomic<bool> m_looping;
 		HWND m_hwnd = 0;
+		HWND m_parentHwnd = 0;
 		COLORREF m_textColor = {};
 		COLORREF m_textWarningColor = {};
 		COLORREF m_textErrorColor = {};
@@ -238,6 +241,7 @@ namespace uba
 		Thread m_thread;
 
 		void PostNewTrace(u32 replay, bool paused);
+		void PostQuit();
 		LRESULT WinProc(HWND hWnd, UINT Msg, WPARAM wParam, LPARAM lParam);
 		static LRESULT CALLBACK StaticWinProc(HWND hWnd, UINT Msg, WPARAM wParam, LPARAM lParam);
 	};
diff --git a/Engine/Source/Programs/UnrealBuildAccelerator/Visualizer/Private/UbaVisualizerApp.cpp b/Engine/Source/Programs/UnrealBuildAccelerator/Visualizer/Private/UbaVisualizerApp.cpp
--- a/Engine/Source/Programs/UnrealBuildAccelerator/Visualizer/Private/UbaVisualizerApp.cpp
+++ b/Engine/Source/Programs/UnrealBuildAccelerator/Visualizer/Private/UbaVisualizerApp.cpp
@@ -30,6 +30,8 @@ static int PrintHelp(const tchar* message = nullptr)
 	s.Appendf(TC("  -listen[=<channel>]  Listen for announcements of new sessions. Defaults to channel '%s'\r\n"), TC("Default"));
 	s.Appendf(TC("  -replay              Visualize the data as if it was running right now\r\n"));
 	s.Appendf(TC("  -config=<file>       Specify config file to use\r\n"));
+	s.Appendf(TC("  -parent=<hwnd>       Specify hwnd this window should be a child of\r\n"));
+	s.Appendf(TC("  -nocopy              Will prevent UbaVisualizer.exe from being copied to temp and executed from there\r\n"));
 	s.Appendf(TC("\r\n"));
 	MessageBox(NULL, s.data, TC("UbaVisualizer"), 0);
 	//wprintf(s.data);
@@ -80,6 +82,8 @@ int WINAPI wWinMain(_In_ HINSTANCE hInstance, _In_opt_ HINSTANCE hPrevInstance,
 	StringBuffer<> configPath;
 	u32 port = DefaultPort;
 	u32 replay = 0;
+	u64 parent = 0;
+	bool copyAndLaunch = true;
 
 	int argc;
 	auto argv = CommandLineToArgvW(GetCommandLine(), &argc);
@@ -151,6 +155,26 @@ int WINAPI wWinMain(_In_ HINSTANCE hInstance, _In_opt_ HINSTANCE hPrevInstance,
 				return PrintHelp(TC("-config needs a value"));
 			configPath.Append(value);
 		}
+		else if (name.Equals(TC("-parent")))
+		{
+			if (value.IsEmpty())
+				return PrintHelp(TC("-parent needs a value"));
+			if (value.count > 8)
+				return PrintHelp(TC("-parent has invalid value"));
+			value.MakeLower();
+			if (value.count & 1) // uneven char, add 0 in the front
+			{
+				memmove(value.data + 1, value.data, (value.count+1)*sizeof(tchar));
+				value.data[0] = '0';
+				++value.count;
+			}
+			parent = StringToValue2(value.data, value.count);
+			//value.ParseHex(parent);
+		}
+		else if (name.Equals(TC("-nocopy")))
+		{
+			copyAndLaunch = false;
+		}
 		else
 		{
 			StringBuffer<> msg;
@@ -159,14 +183,83 @@ int WINAPI wWinMain(_In_ HINSTANCE hInstance, _In_opt_ HINSTANCE hPrevInstance,
 		}
 	}
 
+	MessageBoxLogWriter logWriter;
+	LoggerWithWriter logger(logWriter);
+
+	if (copyAndLaunch)
+	{
+		StringBuffer<> tempPath;
+		tempPath.count = GetTempPathW(tempPath.capacity, tempPath.data);
+		if (!tempPath.count)
+		{
+			logger.Error(TC("GetTempPathW failed"));
+			return -1;
+		}
+
+		StringBuffer<> thisExe;
+		thisExe.count = GetModuleFileNameW(NULL, thisExe.data, thisExe.capacity);
+		if (!thisExe.count)
+		{
+			logger.Error(TC("GetModuleFileNameW failed"));
+			return -1;
+		}
+
+		WIN32_FILE_ATTRIBUTE_DATA data;
+		if (!GetFileAttributesExW(thisExe.data, GetFileExInfoStandard, &data))
+		{
+			logger.Error(TC("GetFileAttributesExW failed"));
+			return -1;
+		}
+
+		u64 thisLastWriteTime = *(u64*)&data.ftLastWriteTime;
+
+		StringBuffer<> ubaFileName;
+		for (u32 i=0; i!=10; ++i)
+		{
+			ubaFileName.Append(tempPath).EnsureEndsWithSlash().Appendf(TC("UbaVisualizer%u.exe"), i);
+			if (GetFileAttributesExW(ubaFileName.data, GetFileExInfoStandard, &data))
+			{
+				u64 lastWriteTime = *(u64*)&data.ftLastWriteTime;
+				if (thisLastWriteTime == lastWriteTime)
+					break;
+			}
+			if (CopyFileW(thisExe.data, ubaFileName.data, false))
+				break;
+			ubaFileName.Clear();
+		}
+
+		if (!ubaFileName.count)
+		{
+			logger.Error(TC("Failed to create temporary UbaVisualizer.exe to launch."));
+			return -1;
+		}
+
+		StringBuffer<> args;
+		args.Append(ubaFileName);
+		for (int i = 1; i != argc; ++i)
+			args.Append(' ').Append(argv[i]);
+		args.Append(" -nocopy");
+
+		STARTUPINFOW si;
+		ZeroMemory(&si, sizeof(si));
+		si.cb = sizeof(si);
+		PROCESS_INFORMATION pi;
+		ZeroMemory(&pi, sizeof(pi));
+		if (!CreateProcessW(NULL, args.data, NULL, NULL, false, 0, NULL, NULL, &si, &pi))
+		{
+			logger.Error(TC("Failed to launch process %s"), ubaFileName.data);
+			return -1;
+		}
+		CloseHandle(pi.hProcess);
+		CloseHandle(pi.hThread);
+		return 0;
+	}
+
 	LocalFree(argv);
 
 	if (host.IsEmpty() && named.IsEmpty() && file.IsEmpty() && !channel.count)
 		channel.Append(TC("Default")); // return PrintHelp(TC("No host/named/file provided. Add -host=<host> or -file=<file> or -named=<name>"));
 
-	MessageBoxLogWriter logWriter;
-	LoggerWithWriter logger(logWriter);
-
 	bool showAllTraces = true;
 	if (!configPath.count)
 	{
@@ -183,6 +276,7 @@ int WINAPI wWinMain(_In_ HINSTANCE hInstance, _In_opt_ HINSTANCE hPrevInstance,
 	}
 
 	VisualizerConfig visualizerConfig(configPath.data);
+	visualizerConfig.parent = parent;
 	visualizerConfig.ShowAllTraces = showAllTraces;
 	visualizerConfig.Load(logger);
 
