[UBA]
* Added a bunch of new options to UbaVisualizer. Check out context menu and checkboxes
* Added more solid code for checking if parent exists or not.. should prevent leaving dead UbaVisualizer.exe processes around
* Changed exit code error message to use hex for exit code

diff --git a/Engine/Source/Programs/UnrealBuildAccelerator/Common/Private/UbaProcess.cpp b/Engine/Source/Programs/UnrealBuildAccelerator/Common/Private/UbaProcess.cpp
--- a/Engine/Source/Programs/UnrealBuildAccelerator/Common/Private/UbaProcess.cpp
+++ b/Engine/Source/Programs/UnrealBuildAccelerator/Common/Private/UbaProcess.cpp
@@ -261,11 +261,11 @@ namespace uba
 				if (!is64Bit)
 					err.Appendf(TC("ERROR: Process did not start properly. Doesn't seem to be a 64-bit executable (%s Size: %llu, CasKey: %s)"), m_realApplication.c_str(), fileSize, CasKeyString(key).str);
 				else
-					err.Appendf(TC("ERROR: Process did not start properly. GetExitCodeProcess returned %u (%s Size: %llu, CasKey: %s)"), exitCode, m_realApplication.c_str(), fileSize, CasKeyString(key).str);
+					err.Appendf(TC("ERROR: Process did not start properly. GetExitCodeProcess returned 0x%x (%s Size: %llu, CasKey: %s)"), exitCode, m_realApplication.c_str(), fileSize, CasKeyString(key).str);
 			}
 
 			if (err.IsEmpty())
-				err.Appendf(TC("ERROR: Process %llu %s (%s) not active but did not get exit message. Received %u messages (GetExitCodeProcess returned %u)"), u64(m_nativeProcessHandle), m_startInfo.GetDescription(), m_realApplication.c_str(), m_messageCount, exitCode);
+				err.Appendf(TC("ERROR: Process %llu %s (%s) not active but did not get exit message. Received %u messages (GetExitCodeProcess returned 0x%x)"), u64(m_nativeProcessHandle), m_startInfo.GetDescription(), m_realApplication.c_str(), m_messageCount, exitCode);
 			LogLine(false, err.data, LogEntryType_Error);
 			m_nativeProcessExitCode = UBA_EXIT_CODE(666);
 		}
diff --git a/Engine/Source/Programs/UnrealBuildAccelerator/Visualizer/Private/UbaVisualizer.cpp b/Engine/Source/Programs/UnrealBuildAccelerator/Visualizer/Private/UbaVisualizer.cpp
--- a/Engine/Source/Programs/UnrealBuildAccelerator/Visualizer/Private/UbaVisualizer.cpp
+++ b/Engine/Source/Programs/UnrealBuildAccelerator/Visualizer/Private/UbaVisualizer.cpp
@@ -257,6 +257,9 @@ namespace uba
 				continue;
 			}
 
+			if (m_parentHwnd && !IsWindow(m_parentHwnd))
+				PostQuit();
+
 			traceName.Clear();
 			if (!channel.Read(traceName))
 			{
@@ -945,30 +948,41 @@ namespace uba
 				if (index++ == m_config.maxActiveVisible)
 					break;
 				
+				auto& process = *kv.second;
 				u32 durationMs = 0;
-				if (playTime > kv.second->start)
-					durationMs = u32(TimeToMs(playTime - kv.second->start));
+				if (playTime > process.start)
+					durationMs = u32(TimeToMs(playTime - process.start));
+
+				//int left = 3 + 6*m_fontHeight;
+				//int right = rect.left + durationMs/100;
+				int posX = int(m_scrollPosX) + progressRect.left;
+				u64 stop = process.stop;
+				bool done = stop != ~u64(0);
+				if (!done)
+					stop = playTime;
+				int left = int(posX + TimeToS(process.start) * scaleX);
+				int right = int(posX + TimeToS(stop) * scaleX) - 1;
 
 				RECT rect;
-				rect.left = 3 + 6*m_fontHeight;
-				rect.right = rect.left + durationMs/100;
+				rect.left = left;
+				rect.right = right;
 				rect.top = posY;
 				rect.bottom = posY + m_fontHeight;
 				FillRect(hdc, &rect, m_processBrushes[0].inProgress);
 
 				str.Clear();
-				auto firstPar = kv.second->description.find_first_of('(');
-				if (firstPar != -1 && *kv.second->description.rbegin() == ')')
+				auto firstPar = process.description.find_first_of('(');
+				if (firstPar != -1 && *process.description.rbegin() == ')')
 				{
-					str.Append(kv.second->description.c_str() + firstPar + 1).Resize(str.count-1);
-					str.Append(L"  ").Append(kv.second->description.c_str(), firstPar);
+					str.Append(process.description.c_str() + firstPar + 1).Resize(str.count-1);
+					str.Append(L"  ").Append(process.description.c_str(), firstPar);
 				}
 				else
-					str.Append(kv.second->description);
+					str.Append(process.description);
 
-				if (kv.second->isRemote)
+				if (process.isRemote)
 					str.Append(L" [remote]");
-				else if (kv.second->cacheFetch)
+				else if (process.cacheFetch)
 					str.Append(L" [cache]");
 				if (durationMs > 1000)
 					str.Appendf(L" - %us", durationMs/1000);
@@ -988,8 +1002,14 @@ namespace uba
 			bool isFirst = i == 0;
 			auto& session = *sortedSessions[i].session;
 			bool hasUpdates = !session.updates.empty();
-			if (!isFirst && !hasUpdates && session.processors.empty())
-				continue;
+			if (!isFirst)
+			{
+				if (!hasUpdates && session.processors.empty())
+					continue;
+
+				if (!m_config.showFinishedProcesses && session.disconnectTime != ~u64(0))
+					continue;
+			}
 
 			processLocation.sessionIndex = sortedSessions[i].index;
 			if (!isFirst)
@@ -1151,6 +1171,8 @@ namespace uba
 				processLocation.processorIndex = 0;
 				for (auto& processor : session.processors)
 				{
+					bool drawProcessorIndex = m_config.showFinishedProcesses;
+
 					if (posY + m_sessionStepY >= progressRect.top && posY < progressRect.bottom)
 					{
 						float barHeight = boxHeight;
@@ -1166,35 +1188,25 @@ namespace uba
 						const int rectBottom = posY + textHeight;
 						const int offsetY = (textHeight - m_processFontHeight + textOffsetY) / 2;
 
-						if (shouldDrawText)
-						{
-							RECT rect;
-							rect.left = 5;
-							rect.right = progressRect.left - 5;
-							rect.top = posY;
-							rect.bottom = rectBottom;
-
-							StringBuffer<> buf;
-							buf.AppendValue(u64(processLocation.processorIndex) + 1);
-							ExtTextOutW(hdc, 5, posY + offsetY, ETO_CLIPPED, &rect, buf.data, buf.count, NULL);
-						}
-
 						processLocation.processIndex = 0;
 						int posX = int(m_scrollPosX) + progressRect.left;
 						for (auto& process : processor.processes)
 						{
 							int left = int(posX + TimeToS(process.start) * scaleX);
 
+							auto pig = MakeGuard([&]() { ++processLocation.processIndex; });
+
 							if (left >= progressRect.right)
-							{
-								++processLocation.processIndex;
 								continue;
-							}
 
 							u64 stop = process.stop;
 							bool done = stop != ~u64(0);
 							if (!done)
 								stop = playTime;
+							else if (!m_config.showFinishedProcesses)
+								continue;
+
+							drawProcessorIndex = true;
 
 							RECT rect;
 							rect.left = left;
@@ -1203,10 +1215,7 @@ namespace uba
 							rect.bottom = rectBottom;
 
 							if (rect.right <= progressRect.left)
-							{
-								++processLocation.processIndex;
 								continue;
-							}
 
 							rect.right = Max(int(rect.right), left + 1);
 
@@ -1304,14 +1313,28 @@ namespace uba
 									BitBlt(hdc, rect.left, rect.top + bltOffsetY, width, height, textDC, bitmapOffsetX, bitmapOffsetY, SRCCOPY);
 								}
 							}
-							++processLocation.processIndex;
+						}
+
+						if (drawProcessorIndex)
+						{
+							RECT rect;
+							rect.left = 5;
+							rect.right = progressRect.left - 5;
+							rect.top = posY;
+							rect.bottom = rectBottom;
+
+							StringBuffer<> buf;
+							buf.AppendValue(u64(processLocation.processorIndex) + 1);
+							ExtTextOutW(hdc, 5, posY + offsetY, ETO_CLIPPED, &rect, buf.data, buf.count, NULL);
 						}
 					}
 
 					lastStop = Max(lastStop, processor.processes.rbegin()->stop);
 
 					++processLocation.processorIndex;
-					posY += processStepY;
+
+					if (drawProcessorIndex)
+						posY += processStepY;
 				}
 			}
 			else
@@ -1522,7 +1545,7 @@ namespace uba
 		{
 			int posX = int(m_scrollPosX) + progressRect.left;
 			int left = int(posX + timelineSelected * scaleX);
-			int timelineTop = Min(posY, int(progressRect.bottom)) + 2;
+			int timelineTop = GetTimelineTop(clientRect);
 
 			// TODO: Draw line up
 			MoveToEx(hdc, left, 2, NULL);
@@ -1974,10 +1997,8 @@ namespace uba
 
 	void Visualizer::PaintTimeline(HDC hdc, const RECT& clientRect)
 	{
-		int posY = m_contentHeight - m_fontHeight - 8;
+		int top = GetTimelineTop(clientRect);
 		float timeScale = (m_horizontalScaleValue*m_zoomValue)*50.0f;
-		int top = Min(posY, int(clientRect.bottom - m_fontHeight - 8));
-			
 		float startOffset = ((m_scrollPosX/timeScale) - int(m_scrollPosX/timeScale)) * timeScale;
 		int index = -int(startOffset/timeScale);
 			
@@ -2189,6 +2210,13 @@ namespace uba
 		return playTime;
 	}
 
+	int Visualizer::GetTimelineTop(const RECT& clientRect)
+	{
+		int posY = m_contentHeight - m_fontHeight - 8;
+		int maxY = int(clientRect.bottom - m_fontHeight - 8);
+		return m_config.LockTimelineToBottom ? maxY : Min(posY, maxY);
+	}
+
 	void Visualizer::HitTest(HitTestResult& outResult, const POINT& pos)
 	{
 		u64 playTime = GetPlayTime();
@@ -2273,8 +2301,13 @@ namespace uba
 			bool isFirst = i == 0;
 			auto& session = *sortedSessions[i].session;
 			bool hasUpdates = !session.updates.empty();
-			if (!isFirst && !hasUpdates && session.processors.empty())
-				continue;
+			if (!isFirst)
+			{
+				if (!hasUpdates && session.processors.empty())
+					continue;
+				if (!m_config.showFinishedProcesses && session.disconnectTime != ~u64(0))
+					continue;
+			}
 
 			u32 sessionIndex = sortedSessions[i].index;
 			if (!isFirst)
@@ -2357,6 +2390,8 @@ namespace uba
 				u32 processorIndex = 0;
 				for (auto& processor : session.processors)
 				{
+					bool drawProcessorIndex = m_config.showFinishedProcesses;
+
 					if (pos.y < progressRect.bottom && posY + processStepY >= progressRect.top && posY <= progressRect.bottom && pos.y >= posY-1 && pos.y < posY-1 + processStepY)
 					{
 						u32 processIndex = 0;
@@ -2365,11 +2400,11 @@ namespace uba
 						{
 							int left = int(posX + TimeToS(process.start) * scaleX);
 
+							auto pig = MakeGuard([&]() { ++processIndex; });
+
 							if (left >= progressRect.right)
-							{
-								++processIndex;
 								continue;
-							}
+
 							if (left < progressRect.left)
 								left = progressRect.left;
 
@@ -2377,16 +2412,18 @@ namespace uba
 							bool done = stopTime != ~u64(0);
 							if (!done)
 								stopTime = playTime;
+							else if (!m_config.showFinishedProcesses)
+								continue;
+
+							drawProcessorIndex = true;
 
 							RECT rect;
 							rect.left = left;
 							rect.right = int(posX + TimeToS(stopTime) * scaleX);
 
 							if (rect.right <= progressRect.left)
-							{
-								++processIndex;
 								continue;
-							}
+
 							rect.right = Max(int(rect.right), left + 1);
 							rect.top = posY;
 							rect.bottom = posY + int(float(18) * m_zoomValue);
@@ -2399,14 +2436,15 @@ namespace uba
 								outResult.processSelected = true;
 								return;
 							}
-							++processIndex;
 						}
 					}
 
 					if (!processor.processes.empty())
 						lastStop = Max(lastStop, processor.processes.rbegin()->stop);
 
-					posY += processStepY;
+					if (drawProcessorIndex)
+						posY += processStepY;
+
 					++processorIndex;
 				}
 			}
@@ -2470,22 +2508,21 @@ namespace uba
 					posY += processStepY;
 				}
 			}
-
-			if (m_config.showTimeline && !m_traceView.sessions.empty())
-			{
-				int timelineTop = Min(posY, int(progressRect.bottom));
-				if (pos.y >= timelineTop && pos.y < timelineTop + 40)
-				{
-					float timeScale = (m_horizontalScaleValue * m_zoomValue)*50.0f;
-					float startOffset = -(m_scrollPosX / timeScale);
-					outResult.timelineSelected = startOffset + (pos.x - m_progressRectLeft) / timeScale;
-				}
-			}
-
 		}
 
 		m_contentWidth = m_progressRectLeft + Max(0, int(TimeToS((lastStop != 0 && lastStop != ~u64(0)) ? lastStop : playTime) * scaleX));
 		m_contentHeight = posY - int(m_scrollPosY) + processStepY + 14;
+
+		if (m_config.showTimeline && !m_traceView.sessions.empty())
+		{
+			int timelineTop = GetTimelineTop(clientRect);
+			if (pos.y >= timelineTop && pos.y < timelineTop + 40)
+			{
+				float timeScale = (m_horizontalScaleValue * m_zoomValue)*50.0f;
+				float startOffset = -(m_scrollPosX / timeScale);
+				outResult.timelineSelected = startOffset + (pos.x - m_progressRectLeft) / timeScale;
+			}
+		}
 	}
 
 	void Visualizer::WriteProcessStats(Logger& out, TraceView::Process& process)
@@ -2574,12 +2611,24 @@ namespace uba
 
 		u64 playTime = GetPlayTime();
 
+
 		RECT rect;
 		GetClientRect(m_hwnd, &rect);
 		float timeS = TimeToS(playTime);
-		float oldScrollPosX = m_scrollPosX;
-		m_scrollPosX = Min(0.0f, (float)rect.right - timeS*50.0f*m_horizontalScaleValue*m_zoomValue - m_progressRectLeft);
-		return oldScrollPosX != m_scrollPosX;
+
+		if (m_config.AutoScaleHorizontal)
+		{
+			m_scrollPosX = 0;
+			timeS = Max(timeS, 20.0f/m_zoomValue);
+			m_horizontalScaleValue = float(rect.right - m_progressRectLeft - 2)/(m_zoomValue*timeS*50.0f);
+			return true;
+		}
+		else
+		{
+			float oldScrollPosX = m_scrollPosX;
+			m_scrollPosX = Min(0.0f, (float)rect.right - timeS*50.0f*m_horizontalScaleValue*m_zoomValue - m_progressRectLeft);
+			return oldScrollPosX != m_scrollPosX;
+		}
 	}
 
 	bool Visualizer::UpdateSelection()
@@ -3148,6 +3197,15 @@ namespace uba
 				break;
 			case Popup_SortActiveRemoteSessions:
 				m_config.SortActiveRemoteSessions = !m_config.SortActiveRemoteSessions;
+				Redraw();
+				break;
+			case Popup_AutoScaleHorizontal:
+				m_config.AutoScaleHorizontal = !m_config.AutoScaleHorizontal;
+				Redraw();
+				break;
+			case Popup_LockTimelineToBottom:
+				m_config.LockTimelineToBottom = !m_config.LockTimelineToBottom;
+				Redraw();
 				break;
 			case Popup_DarkMode:
 			{
diff --git a/Engine/Source/Programs/UnrealBuildAccelerator/Visualizer/Private/UbaVisualizer.h b/Engine/Source/Programs/UnrealBuildAccelerator/Visualizer/Private/UbaVisualizer.h
--- a/Engine/Source/Programs/UnrealBuildAccelerator/Visualizer/Private/UbaVisualizer.h
+++ b/Engine/Source/Programs/UnrealBuildAccelerator/Visualizer/Private/UbaVisualizer.h
@@ -14,11 +14,12 @@ namespace uba
 		UBA_VISUALIZER_FLAG(Progress, true, L"progress") \
 		UBA_VISUALIZER_FLAG(Status, true, L"status") \
 		UBA_VISUALIZER_FLAG(ActiveProcesses, false, L"active processes") \
-		UBA_VISUALIZER_FLAG(TitleBars, true, L"title bars") \
+		UBA_VISUALIZER_FLAG(TitleBars, true, L"instance title bars") \
 		UBA_VISUALIZER_FLAG(DetailedData, false, L"detailed data (use -UbaDetailedTrace for even more)") \
 		UBA_VISUALIZER_FLAG(NetworkStats, true, L"network stats") \
 		UBA_VISUALIZER_FLAG(CpuMemStats, true, L"cpu/mem stats") \
 		UBA_VISUALIZER_FLAG(ProcessBars, true, L"process bars") \
+		UBA_VISUALIZER_FLAG(FinishedProcesses, true, L"finished process bars") \
 		UBA_VISUALIZER_FLAG(Timeline, true, L"timeline") \
 		UBA_VISUALIZER_FLAG(Workers, false, L"workers (threads on host taking care of requests from helpers)") \
 		UBA_VISUALIZER_FLAG(CursorLine, false, L"cursor (vertical line)") \
@@ -31,6 +32,8 @@ namespace uba
 		UBA_VISUALIZER_FLAG(AutoSaveSettings, true, L"Auto save Position/Settings on close") \
 		UBA_VISUALIZER_FLAG(ShowAllTraces, true, L"Show all traces started on channel") \
 		UBA_VISUALIZER_FLAG(SortActiveRemoteSessions, true, L"Sort active sessions on top") \
+		UBA_VISUALIZER_FLAG(AutoScaleHorizontal, true, L"Automatically scale horizontal to fit processes") \
+		UBA_VISUALIZER_FLAG(LockTimelineToBottom, true, L"Lock timeline to always paint at bottom") \
 
 	struct VisualizerConfig
 	{
@@ -111,6 +114,7 @@ namespace uba
 			TString hyperLink;
 		};
 		u64 GetPlayTime();
+		int GetTimelineTop(const RECT& clientRect);
 		void HitTest(HitTestResult& outResult, const POINT& pos);
 
 		void WriteProcessStats(Logger& out, TraceView::Process& process);
