[UBA]
* Fixed bug where pid and id got lost when UbaVisualizer was copied and started from different location
* Added so pid and id can be picked up by environment variables UBA_OWNER_ID and UBA_OWNER_PID

diff --git a/Engine/Source/Programs/UnrealBuildAccelerator/Common/Private/UbaTrace.cpp b/Engine/Source/Programs/UnrealBuildAccelerator/Common/Private/UbaTrace.cpp
--- a/Engine/Source/Programs/UnrealBuildAccelerator/Common/Private/UbaTrace.cpp
+++ b/Engine/Source/Programs/UnrealBuildAccelerator/Common/Private/UbaTrace.cpp
@@ -503,6 +503,15 @@ namespace uba
 
 		OwnerInfo info { buffer, 0 };
 
+		StringBuffer<32> ownerPidStr;
+		ownerPidStr.count = GetEnvironmentVariableW(TC("UBA_OWNER_PID"), ownerPidStr.data, ownerPidStr.capacity);
+		if (ownerPidStr.count)
+		{
+			GetEnvironmentVariableW(TC("UBA_OWNER_ID"), buffer, sizeof_array(buffer));
+			ownerPidStr.Parse(info.pid);
+			return info;
+		}
+
 		#if PLATFORM_WINDOWS
 		HANDLE snapshotHandle = CreateToolhelp32Snapshot(TH32CS_SNAPPROCESS, 0);
 		if (snapshotHandle == INVALID_HANDLE_VALUE)
@@ -535,6 +544,7 @@ namespace uba
 				break;
 			tchar moduleName[260];
 			DWORD len = GetModuleFileNameExW(parentHandle, 0, moduleName, MAX_PATH);
+			g_debugLogger.Info(TC("MODULE: %s\n"), moduleName);
 			CloseHandle(parentHandle);
 			if (!len)
 				break;
@@ -554,4 +564,4 @@ namespace uba
 		static OwnerInfo info = InternalGetOwnerInfo();
 		return info;
 	}
-}
\ No newline at end of file
+}
diff --git a/Engine/Source/Programs/UnrealBuildAccelerator/Detours/Private/Windows/UbaDetoursFunctionsKernelBase.inl b/Engine/Source/Programs/UnrealBuildAccelerator/Detours/Private/Windows/UbaDetoursFunctionsKernelBase.inl
--- a/Engine/Source/Programs/UnrealBuildAccelerator/Detours/Private/Windows/UbaDetoursFunctionsKernelBase.inl
+++ b/Engine/Source/Programs/UnrealBuildAccelerator/Detours/Private/Windows/UbaDetoursFunctionsKernelBase.inl
@@ -3395,12 +3395,14 @@ LPWCH Detoured_GetEnvironmentStringsW()
 	DEBUG_LOG_TRUE(L"GetEnvironmentStringsW", L"");
 	auto res = True_GetEnvironmentStringsW();
 
-	//auto it = res;
-	//while (*it)
-	//{
-	//	DEBUG_LOG(L"		VAR: %ls", it);
-	//	it += wcslen(it) + 1;
-	//}
+	#if 0 // Enable to print out environment variables in the log
+	auto it = res;
+	while (*it)
+	{
+		DEBUG_LOG(L"		VAR: %ls", it);
+		it += wcslen(it) + 1;
+	}
+	#endif
 
 	return res;
 }
diff --git a/Engine/Source/Programs/UnrealBuildAccelerator/Visualizer/Private/UbaVisualizerApp.cpp b/Engine/Source/Programs/UnrealBuildAccelerator/Visualizer/Private/UbaVisualizerApp.cpp
--- a/Engine/Source/Programs/UnrealBuildAccelerator/Visualizer/Private/UbaVisualizerApp.cpp
+++ b/Engine/Source/Programs/UnrealBuildAccelerator/Visualizer/Private/UbaVisualizerApp.cpp
@@ -175,6 +175,16 @@ int WINAPI wWinMain(_In_ HINSTANCE hInstance, _In_opt_ HINSTANCE hPrevInstance,
 		{
 			copyAndLaunch = false;
 		}
+		else if (name.Equals(TC("-ownerPid")))
+		{
+			u32 ownerPid;
+			if (value.Parse(ownerPid))
+				const_cast<OwnerInfo&>(GetOwnerInfo()).pid = ownerPid;
+		}
+		else if (name.Equals(TC("-ownerId")))
+		{
+			TStrcpy_s(const_cast<tchar*>(GetOwnerInfo().id), 260, value.data);
+		}
 		else
 		{
 			StringBuffer<> msg;
@@ -240,6 +250,10 @@ int WINAPI wWinMain(_In_ HINSTANCE hInstance, _In_opt_ HINSTANCE hPrevInstance,
 			args.Append(' ').Append(argv[i]);
 		args.Append(" -nocopy");
 
+		OwnerInfo ownerInfo = GetOwnerInfo();
+		if (ownerInfo.pid)
+			args.Appendf(TC(" -ownerPid=%u -ownerId=%s"), ownerInfo.pid, ownerInfo.id);
+
 		STARTUPINFOW si;
 		ZeroMemory(&si, sizeof(si));
 		si.cb = sizeof(si);
