[UBA]
* Removed debug line that shouldn't be there

diff --git a/Engine/Source/Programs/UnrealBuildAccelerator/Common/Private/UbaTrace.cpp b/Engine/Source/Programs/UnrealBuildAccelerator/Common/Private/UbaTrace.cpp
--- a/Engine/Source/Programs/UnrealBuildAccelerator/Common/Private/UbaTrace.cpp
+++ b/Engine/Source/Programs/UnrealBuildAccelerator/Common/Private/UbaTrace.cpp
@@ -544,7 +544,6 @@ namespace uba
 				break;
 			tchar moduleName[260];
 			DWORD len = GetModuleFileNameExW(parentHandle, 0, moduleName, MAX_PATH);
-			g_debugLogger.Info(TC("MODULE: %s\n"), moduleName);
 			CloseHandle(parentHandle);
 			if (!len)
 				break;
