[UBA]
* Fixed so posix opendir is using directory table in order to work properly on remote helpers

diff --git a/Engine/Source/Programs/UnrealBuildAccelerator/Common/Public/UbaDirectoryIterator.h b/Engine/Source/Programs/UnrealBuildAccelerator/Common/Public/UbaDirectoryIterator.h
--- a/Engine/Source/Programs/UnrealBuildAccelerator/Common/Public/UbaDirectoryIterator.h
+++ b/Engine/Source/Programs/UnrealBuildAccelerator/Common/Public/UbaDirectoryIterator.h
@@ -232,9 +232,19 @@ namespace uba
 		if (dirPath2[dirPathLen - 1] != '/')
 			dirPath2[dirPathLen++] = '/';
 
-		struct dirent* pDirent;
-		while ((pDirent = readdir(dir)) != NULL)
+		while (true)
 		{
+			errno = 0;
+			struct dirent* pDirent = readdir(dir);
+			if (!pDirent)
+			{
+				if (errno == 0)
+					break;
+
+				// This is actually an error.. should return false?
+				break;
+			}
+
 			char* fileName = pDirent->d_name;
 			if (fileName[0] == '.' && (fileName[1] == 0 || (fileName[1] == '.' && fileName[2] == 0))) 
 				continue;
diff --git a/Engine/Source/Programs/UnrealBuildAccelerator/Detours/Private/PosixOS/UbaDetoursFunctionsPosix.cpp b/Engine/Source/Programs/UnrealBuildAccelerator/Detours/Private/PosixOS/UbaDetoursFunctionsPosix.cpp
--- a/Engine/Source/Programs/UnrealBuildAccelerator/Detours/Private/PosixOS/UbaDetoursFunctionsPosix.cpp
+++ b/Engine/Source/Programs/UnrealBuildAccelerator/Detours/Private/PosixOS/UbaDetoursFunctionsPosix.cpp
@@ -991,12 +991,64 @@ UBA_EXPORT int UBA_WRAPPER(fclose)(FILE* stream)
 	return res;
 }
 
+struct DirInfo
+{
+	Vector<u32> fileTableOffsets;
+	int it = -1;
+	dirent ent;
+};
+
+bool IsDirInfo(DIR* dir) { return (u64(dir) & 0x1000'0000'0000'0000) != 0; }
+DirInfo* AsDirInfo(DIR* dir) { return (DirInfo*)(uintptr_t(dir) & ~0x1000'0000'0000'0000); }
+
 UBA_EXPORT DIR* UBA_WRAPPER(opendir)(const char* name)
 {
 	UBA_INIT_DETOUR(opendir, name);
 	StringBuffer<> dirName;
 	FixPath(dirName, name);
 
+	if (g_runningRemote || true)
+	{
+		DirHash hash(dirName.data, dirName.count);
+
+		SCOPED_WRITE_LOCK(g_directoryTable.m_lookupLock, lookLock);
+		auto insres = g_directoryTable.m_lookup.try_emplace(hash.key, &g_memoryBlock);
+		DirectoryTable::Directory& dir = insres.first->second;
+		if (insres.second)
+		{
+			if (g_directoryTable.EntryExistsNoLock(hash.key, dirName.data, dirName.count) != DirectoryTable::Exists_No)
+				Rpc_UpdateDirectory(hash.key, dirName.data, dirName.count, false);
+		}
+		bool exists = false;
+		if (dir.tableOffset != InvalidTableOffset)
+		{
+			u32 entryOffset = dir.tableOffset | 0x80000000;
+			DirectoryTable::EntryInformation entryInfo;
+			g_directoryTable.GetEntryInformation(entryInfo, entryOffset);
+			exists = entryInfo.attributes != 0;
+		}
+
+		if (!exists)
+		{
+			errno = ENOENT;
+			return nullptr;
+		}
+
+		g_directoryTable.PopulateDirectory(hash.open, dir);
+
+		auto dirInfo = new DirInfo();
+	
+		SCOPED_READ_LOCK(dir.lock, lock);
+		dirInfo->fileTableOffsets.resize(dir.files.size());
+		u32 it = 0;
+		for (auto& pair : dir.files)
+			dirInfo->fileTableOffsets[it++] = pair.second;
+		lock.Leave();
+	
+		DEBUG_LOG_DETOURED("opendir", "(%s) -> %p", dirName.data, dirInfo);
+		return (DIR*)dirInfo;
+	}
+
 	DIR* res = TRUE_WRAPPER(opendir)(dirName.data);
 	DEBUG_LOG_TRUE("opendir", "(%s) -> %p", dirName.data, res);
 	if (!res)
@@ -1009,7 +1061,13 @@ UBA_EXPORT DIR* UBA_WRAPPER(opendir)(const char* name)
 UBA_EXPORT int UBA_WRAPPER(dirfd)(DIR* dirp)
 {
 	UBA_INIT_DETOUR(dirfd, dirp);
-	dirp = (DIR*)(u64(dirp) & ~0x1000'0000'0000'0000);
+
+	if (IsDirInfo(dirp))
+	{
+		UBA_ASSERT(false);
+		return 1;
+	}
+
 	int res = TRUE_WRAPPER(dirfd)(dirp);
 	DEBUG_LOG_TRUE("dirfd", "(%p) -> %i", dirp, res);
 	return res;
@@ -1018,7 +1076,36 @@ UBA_EXPORT int UBA_WRAPPER(dirfd)(DIR* dirp)
 UBA_EXPORT dirent* UBA_WRAPPER(readdir)(DIR* dirp)
 {
 	UBA_INIT_DETOUR(readdir, dirp);
-	dirp = (DIR*)(u64(dirp) & ~0x1000'0000'0000'0000);
+
+	if (IsDirInfo(dirp))
+	{
+		auto& dirInfo = *AsDirInfo(dirp);
+		while (true)
+		{
+			++dirInfo.it;
+			if (dirInfo.it >= dirInfo.fileTableOffsets.size())
+			{
+				DEBUG_LOG_DETOURED("readdir", "(%p) -> nullptr", dirp);
+				return nullptr;
+			}
+			u32 fileTableOffset = dirInfo.fileTableOffsets[dirInfo.it];
+
+			DirectoryTable::EntryInformation info;
+			g_directoryTable.GetEntryInformation(info, fileTableOffset, dirInfo.ent.d_name, 256);
+			if (info.attributes == 0) // File was deleted
+				continue;
+
+			dirInfo.ent.d_ino = info.fileIndex;
+			dirInfo.ent.d_off = 0;
+			dirInfo.ent.d_reclen = sizeof(dirent);
+			dirInfo.ent.d_type = S_ISDIR(info.attributes) ? DT_DIR : DT_REG;
+			return &dirInfo.ent;
+		}
+
+		DEBUG_LOG_DETOURED("readdir", "(%p) -> %p", dirp, &dirInfo.ent);
+		return &dirInfo.ent;
+	}
+
 	auto res = TRUE_WRAPPER(readdir)(dirp);
 	DEBUG_LOG_TRUE("readdir", "(%p) -> %p", dirp, res);
 	return res;
@@ -1027,7 +1114,7 @@ UBA_EXPORT dirent* UBA_WRAPPER(readdir)(DIR* dirp)
 UBA_EXPORT void UBA_WRAPPER(rewinddir)(DIR* dirp)
 {
 	UBA_INIT_DETOUR(rewinddir, dirp);
-	UBA_ASSERTF(false, "rewinddir");
+	UBA_ASSERTF(!IsDirInfo(dirp), "rewinddir");
 	DEBUG_LOG_TRUE("rewinddir", "(%p)", dirp);
 	return TRUE_WRAPPER(rewinddir)(dirp);
 }
@@ -1035,14 +1122,14 @@ UBA_EXPORT void UBA_WRAPPER(rewinddir)(DIR* dirp)
 UBA_EXPORT int UBA_WRAPPER(scandir)(const char* dirp, dirent*** namelist, int (*filter)(const dirent*), int (*compar)(const dirent**, const dirent**))
 {
 	UBA_INIT_DETOUR(scandir, dirp, namelist, filter, compar);
-	UBA_ASSERTF(false, "scandir");
+	UBA_ASSERTF(!g_runningRemote, "scandir not implemented for remote");
 	DEBUG_LOG_TRUE("scandir", "(%p)", dirp);
 	return TRUE_WRAPPER(scandir)(dirp, namelist, filter, compar);
 }
 UBA_EXPORT void UBA_WRAPPER(seekdir)(DIR* dirp, long loc)
 {
 	UBA_INIT_DETOUR(seekdir, dirp, loc);
-	UBA_ASSERTF(false, "seekdir");
+	UBA_ASSERTF(!IsDirInfo(dirp), "seekdir");
 	DEBUG_LOG_TRUE("seekdir", "(%p)", dirp);
 	return TRUE_WRAPPER(seekdir)(dirp, loc);
 }
@@ -1050,7 +1137,7 @@ UBA_EXPORT void UBA_WRAPPER(seekdir)(DIR* dirp, long loc)
 UBA_EXPORT long UBA_WRAPPER(telldir)(DIR* dirp)
 {
 	UBA_INIT_DETOUR(telldir, dirp);
-	UBA_ASSERTF(false, "telldir");
+	UBA_ASSERTF(!IsDirInfo(dirp), "telldir");
 	DEBUG_LOG_TRUE("telldir", "(%p)", dirp);
 	return TRUE_WRAPPER(telldir)(dirp);
 }
@@ -1066,7 +1153,14 @@ UBA_EXPORT DIR* UBA_WRAPPER(fdopendir)(int fd)
 UBA_EXPORT int UBA_WRAPPER(closedir)(DIR* dirp)
 {
 	UBA_INIT_DETOUR(closedir, dirp);
-	dirp = (DIR*)(u64(dirp) & ~0x1000'0000'0000'0000);
+
+	if (IsDirInfo(dirp))
+	{
+		delete (DirInfo*)AsDirInfo(dirp);
+		DEBUG_LOG_DETOURED("closedir", "(%p)", dirp);
+		return 0;
+	}
+
 	DEBUG_LOG_TRUE("closedir", "(%p)", dirp);
 	return TRUE_WRAPPER(closedir)(dirp);
 }
