[UBA]
* Fixed crash in code handling corrupt cas files

diff --git a/Engine/Source/Programs/UnrealBuildAccelerator/Common/Private/UbaStorage.cpp b/Engine/Source/Programs/UnrealBuildAccelerator/Common/Private/UbaStorage.cpp
--- a/Engine/Source/Programs/UnrealBuildAccelerator/Common/Private/UbaStorage.cpp
+++ b/Engine/Source/Programs/UnrealBuildAccelerator/Common/Private/UbaStorage.cpp
@@ -990,7 +990,7 @@ namespace uba
 							BinaryReader reader(fa.GetData(), 0, fa.GetSize());
 							if (reader.GetLeft() < 12)
 							{
-								m_logger.Detail(TC("Corrupt cas. Is %llu, must be at least 12 bytes (%s)"), size, e.name);
+								m_logger.Detail(TC("Corrupt cas. Is %llu, must be at least 12 bytes (%s)"), e.name);
 								deleteFile = true;
 							}
 							else
@@ -1000,7 +1000,7 @@ namespace uba
 								{
 									if (reader.GetLeft() <= 8)
 									{
-										m_logger.Detail(TC("Corrupt cas. Missing beginning of block (%s)"), size, e.name);
+										m_logger.Detail(TC("Corrupt cas. Missing beginning of block (%s)"), e.name);
 										deleteFile = true;
 										break;
 									}
@@ -1010,7 +1010,7 @@ namespace uba
 
 									if (!compressedBlockSize || compressedBlockSize > reader.GetLeft())
 									{
-										m_logger.Detail(TC("Corrupt cas. Bad block (%s)"), size, e.name);
+										m_logger.Detail(TC("Corrupt cas. Bad block (%s)"), e.name);
 										deleteFile = true;
 										break;
 									}
