[UBA]
* Added so client session waits for processes to flush before exiting

diff --git a/Engine/Source/Programs/UnrealBuildAccelerator/Common/Private/UbaSessionClient.cpp b/Engine/Source/Programs/UnrealBuildAccelerator/Common/Private/UbaSessionClient.cpp
--- a/Engine/Source/Programs/UnrealBuildAccelerator/Common/Private/UbaSessionClient.cpp
+++ b/Engine/Source/Programs/UnrealBuildAccelerator/Common/Private/UbaSessionClient.cpp
@@ -1767,7 +1767,25 @@ namespace uba
 			RemoveInactiveProcesses();
 
 			if (activeProcesses.empty() && !m_remoteExecutionEnabled)
+			{
+				// There can be processes that are done (isDone is true) but are still in m_processes list (since they are removed from that after). give them some time
+				u64 counter = 300;
+				while (true)
+				{
+					if (!counter--)
+					{
+						m_logger.Warning(TC("Took a long time for processes to be removed after being finished"));
+						break;
+					}
+
+					SCOPED_READ_LOCK(m_processesLock, processesLock);
+					if (m_processes.empty())
+						break;
+					processesLock.Leave();
+					Sleep(10);
+				}
 				break;
+			}
 		}
 
 		CancelAllProcessesAndWait(); // If we got the exit from server there is no point sending anything more back.. cancel everything
