[UBA]
* Added debug logging stuff to session server FetchBegin to try to figure out why clients timeout now and then

diff --git a/Engine/Source/Programs/UnrealBuildAccelerator/Common/Private/UbaStorageServer.cpp b/Engine/Source/Programs/UnrealBuildAccelerator/Common/Private/UbaStorageServer.cpp
--- a/Engine/Source/Programs/UnrealBuildAccelerator/Common/Private/UbaStorageServer.cpp
+++ b/Engine/Source/Programs/UnrealBuildAccelerator/Common/Private/UbaStorageServer.cpp
@@ -333,8 +333,21 @@ namespace uba
 
 			case StorageMessageType_FetchBegin:
 			{
+				// TODO: Remove this when we have tracked down the issue where clients time out
+				u32 todoRemoveMe = 0;
+				auto timeoutGuard = MakeGuard([&, timeoutStartTime = GetTime()]()
+					{
+						u64 timeSpentMs = TimeToMs(GetTime() - timeoutStartTime);
+						if (timeSpentMs > 8 * 60 * 1000)
+						{
+							// Took more than 5 minutes to respond
+							m_logger.Warning(TC("Took more than 8 minutes to respond to FetchBegin (%u).. is this some sort of hang or just host being half dead?"), todoRemoveMe);
+						}
+					});
+
 				if (reader.ReadBool()) // Wants proxy
 				{
+					todoRemoveMe = 1;
 					SCOPED_READ_LOCK(m_connectionInfoLock, lock);
 					auto findIt = m_connectionInfo.find(connectionInfo.GetId());
 					UBA_ASSERT(findIt != m_connectionInfo.end());
@@ -401,6 +414,8 @@ namespace uba
 					}
 				}
 
+				todoRemoveMe = 2;
+
 				u64 start = GetTime();
 				CasKey casKey = reader.ReadCasKey();
 				StringBuffer<> hint;
@@ -412,6 +427,7 @@ namespace uba
 				bool has = HasCasFile(casKey, &casEntry); // HasCasFile also writes deferred cas entries if in queue
 				if (!has)
 				{
+					todoRemoveMe = 3;
 					if (!EnsureCasFile(casKey, nullptr) && m_allowFallback)
 					{
 						// Last resort.. use hint to load file into cas (hint should be renamed since it is now a critical parameter)
@@ -468,6 +484,8 @@ namespace uba
 					casEntry = &findIt->second;
 				}
 
+				todoRemoveMe = 4;
+
 				if (casEntry->disallowed)
 				{
 					writer.WriteU16(0);
@@ -486,6 +504,8 @@ namespace uba
 				MappedView mappedView;
 				auto mvg = MakeGuard([&](){ m_casDataBuffer.UnmapView(mappedView, TC("FetchBegin")); });
 
+				todoRemoveMe = 5;
+
 				bool useFileMapping = casEntry->mappingHandle.IsValid();
 				if (useFileMapping)
 				{
@@ -527,6 +547,8 @@ namespace uba
 					}
 				}
 
+				todoRemoveMe = 6;
+
 				if (m_trace)
 					m_trace->FileBeginFetch(connectionInfo.GetId(), casKey, fileSize, hint.data, m_traceFetch);
 
@@ -545,6 +567,9 @@ namespace uba
 				u64 capacityLeft = writer.GetCapacityLeft();
 				u32 toWrite = u32(Min(left, capacityLeft));
 				void* writeBuffer = writer.AllocWrite(toWrite);
+
+				todoRemoveMe = 7;
+
 				if (useFileMapping)
 				{
 					memcpy(writeBuffer, memoryPos, toWrite);
@@ -575,6 +600,8 @@ namespace uba
 					readFileHandle = InvalidFileHandle;
 				}
 
+				todoRemoveMe = 8;
+
 				u64 actualSize = fileSize;
 				if (m_storeCompressed)
 					actualSize = *(u64*)writeBuffer;
@@ -593,6 +620,8 @@ namespace uba
 					return true;
 				}
 
+				todoRemoveMe = 9;
+
 				mvg.Cancel();
 				cg.Cancel();
 				rfg.Cancel();
@@ -606,6 +635,8 @@ namespace uba
 				fetch.clientId = connectionInfo.GetId();
 				lock.Leave();
 
+				todoRemoveMe = 10;
+
 				mappedView.size = fileSize;
 
 				fetch.readFileHandle = readFileHandle;
