[UBA]
* Changed so actual wait time when timing out waiting for message is logged

diff --git a/Engine/Source/Programs/UnrealBuildAccelerator/Common/Private/UbaNetworkClient.cpp b/Engine/Source/Programs/UnrealBuildAccelerator/Common/Private/UbaNetworkClient.cpp
--- a/Engine/Source/Programs/UnrealBuildAccelerator/Common/Private/UbaNetworkClient.cpp
+++ b/Engine/Source/Programs/UnrealBuildAccelerator/Common/Private/UbaNetworkClient.cpp
@@ -552,10 +552,11 @@ namespace uba
 
 		if (response)
 		{
+			u64 waitStart = GetTime();
 			u32 timeoutMs = 10 * 60 * 1000;
 			if (!gotResponse.IsSet(timeoutMs))
 			{
-				m_logger.Error(TC("Timed out after 10 minutes waiting for message response from server."));
+				m_logger.Error(TC("Timed out after %s waiting for message response from server."), TimeToText(GetTime() - waitStart).str);
 				message.m_error = 4;
 			}
 			else if (m_cryptoKey && !message.m_error && message.m_responseSize)
