[UBA]
* Fixed so quotes are preserved for arguments when visualizer relaunching itself

diff --git a/Engine/Source/Programs/UnrealBuildAccelerator/Visualizer/Private/UbaVisualizerApp.cpp b/Engine/Source/Programs/UnrealBuildAccelerator/Visualizer/Private/UbaVisualizerApp.cpp
--- a/Engine/Source/Programs/UnrealBuildAccelerator/Visualizer/Private/UbaVisualizerApp.cpp
+++ b/Engine/Source/Programs/UnrealBuildAccelerator/Visualizer/Private/UbaVisualizerApp.cpp
@@ -51,7 +51,14 @@ struct MessageBoxLogWriter : public LogWriter
 	virtual void Log(LogEntryType type, const tchar* str, u32 strLen, const tchar* prefix = nullptr, u32 prefixLen = 0) override
 	{
 		if (type > LogEntryType_Warning)
+		{
+			#if UBA_DEBUG
+			StringBuffer<> buf;
+			buf.Append(str, strLen).Append(TC("\r\n"));
+			OutputDebugStringW(buf.data);
+			#endif
 			return;
+		}
 
 		HWND hwnd = NULL;
 		if (m_visualizer)
@@ -247,7 +254,13 @@ int WINAPI wWinMain(_In_ HINSTANCE hInstance, _In_opt_ HINSTANCE hPrevInstance,
 		StringBuffer<> args;
 		args.Append(ubaFileName);
 		for (int i = 1; i != argc; ++i)
-			args.Append(' ').Append(argv[i]);
+		{
+			args.Append(' ');
+			if (TStrchr(argv[i], ' '))
+				args.Append('\"').Append(argv[i]).Append('\"');
+			else
+				args.Append(argv[i]);
+		}
 		args.Append(" -nocopy");
 
 		OwnerInfo ownerInfo = GetOwnerInfo();
