[UBA]
* Fixed macos compile error

diff --git a/Engine/Source/Programs/UnrealBuildAccelerator/Detours/Private/PosixOS/UbaDetoursFunctionsPosix.cpp b/Engine/Source/Programs/UnrealBuildAccelerator/Detours/Private/PosixOS/UbaDetoursFunctionsPosix.cpp
--- a/Engine/Source/Programs/UnrealBuildAccelerator/Detours/Private/PosixOS/UbaDetoursFunctionsPosix.cpp
+++ b/Engine/Source/Programs/UnrealBuildAccelerator/Detours/Private/PosixOS/UbaDetoursFunctionsPosix.cpp
@@ -1096,7 +1096,9 @@ UBA_EXPORT dirent* UBA_WRAPPER(readdir)(DIR* dirp)
 				continue;
 
 			dirInfo.ent.d_ino = info.fileIndex;
+			#if PLATFORM_LINUX
 			dirInfo.ent.d_off = 0;
+			#endif
 			dirInfo.ent.d_reclen = sizeof(dirent);
 			dirInfo.ent.d_type = S_ISDIR(info.attributes) ? DT_DIR : DT_REG;
 			return &dirInfo.ent;
