[UBA]
* Added validation code for creating event used to wait on message response

diff --git a/Engine/Source/Programs/UnrealBuildAccelerator/Common/Private/UbaNetworkClient.cpp b/Engine/Source/Programs/UnrealBuildAccelerator/Common/Private/UbaNetworkClient.cpp
--- a/Engine/Source/Programs/UnrealBuildAccelerator/Common/Private/UbaNetworkClient.cpp
+++ b/Engine/Source/Programs/UnrealBuildAccelerator/Common/Private/UbaNetworkClient.cpp
@@ -463,7 +463,7 @@ namespace uba
 		BinaryWriter& writer = *message.m_sendWriter;
 
 		u16 messageId = 0;
-		Event gotResponse(true);
+		Event gotResponse;
 
 		if (response)
 		{
@@ -506,6 +506,14 @@ namespace uba
 					UBA_ASSERT(!message.m_doneFunc);
 					message.m_doneUserData = &gotResponse;
 					message.m_doneFunc = [](bool error, void* userData) { ((Event*)userData)->Set(); };
+
+					if (!gotResponse.Create(true))
+					{
+						m_logger.Error(TC("Failed to create event, this should not happen?!?"));
+						message.m_error = 13;
+						OnDisconnected(connection, 13);
+						return false;
+					}
 				}
 				break;
 			}
