[UBA]
* Moved event initialization code to be outside lock in order to reduce contention

diff --git a/Engine/Source/Programs/UnrealBuildAccelerator/Common/Private/UbaNetworkClient.cpp b/Engine/Source/Programs/UnrealBuildAccelerator/Common/Private/UbaNetworkClient.cpp
--- a/Engine/Source/Programs/UnrealBuildAccelerator/Common/Private/UbaNetworkClient.cpp
+++ b/Engine/Source/Programs/UnrealBuildAccelerator/Common/Private/UbaNetworkClient.cpp
@@ -465,8 +465,20 @@ namespace uba
 		u16 messageId = 0;
 		Event gotResponse;
 
+
 		if (response)
 		{
+			if (!async)
+			{
+				if (!gotResponse.Create(true))
+				{
+					m_logger.Error(TC("Failed to create event, this should not happen?!?"));
+					message.m_error = 13;
+					OnDisconnected(connection, 13);
+					return false;
+				}
+			}
+
 			while (true)
 			{
 				SCOPED_WRITE_LOCK(m_activeMessagesLock, lock);
@@ -506,14 +518,6 @@ namespace uba
 					UBA_ASSERT(!message.m_doneFunc);
 					message.m_doneUserData = &gotResponse;
 					message.m_doneFunc = [](bool error, void* userData) { ((Event*)userData)->Set(); };
-
-					if (!gotResponse.Create(true))
-					{
-						m_logger.Error(TC("Failed to create event, this should not happen?!?"));
-						message.m_error = 13;
-						OnDisconnected(connection, 13);
-						return false;
-					}
 				}
 				break;
 			}
