[UBA]
* Disabled opendir directory table usage on macos since it is failing for some reason I can't figure out why.

diff --git a/Engine/Source/Programs/UnrealBuildAccelerator/Detours/Private/PosixOS/UbaDetoursFunctionsPosix.cpp b/Engine/Source/Programs/UnrealBuildAccelerator/Detours/Private/PosixOS/UbaDetoursFunctionsPosix.cpp
--- a/Engine/Source/Programs/UnrealBuildAccelerator/Detours/Private/PosixOS/UbaDetoursFunctionsPosix.cpp
+++ b/Engine/Source/Programs/UnrealBuildAccelerator/Detours/Private/PosixOS/UbaDetoursFunctionsPosix.cpp
@@ -1007,55 +1007,59 @@ UBA_EXPORT DIR* UBA_WRAPPER(opendir)(const char* name)
 	StringBuffer<> dirName;
 	FixPath(dirName, name);
 
-	if (g_runningRemote || true)
+	bool useDetour = true;
+
+	// TODO: Someone needs to debug this.. it fails on the farm.
+	#if PLATFORM_MAC
+	useDetour = false;
+	#endif
+
+	if (!useDetour)
 	{
-		DirHash hash(dirName.data, dirName.count);
+		DIR* res = TRUE_WRAPPER(opendir)(dirName.data);
+		DEBUG_LOG_TRUE("opendir", "(%s) -> %p", dirName.data, res);
+		return res;
+	}
 
-		SCOPED_WRITE_LOCK(g_directoryTable.m_lookupLock, lookLock);
-		auto insres = g_directoryTable.m_lookup.try_emplace(hash.key, &g_memoryBlock);
-		DirectoryTable::Directory& dir = insres.first->second;
-		if (insres.second)
-		{
-			if (g_directoryTable.EntryExistsNoLock(hash.key, dirName.data, dirName.count) != DirectoryTable::Exists_No)
-				Rpc_UpdateDirectory(hash.key, dirName.data, dirName.count, false);
-		}
-		bool exists = false;
-		if (dir.tableOffset != InvalidTableOffset)
-		{
-			u32 entryOffset = dir.tableOffset | 0x80000000;
-			DirectoryTable::EntryInformation entryInfo;
-			g_directoryTable.GetEntryInformation(entryInfo, entryOffset);
-			exists = entryInfo.attributes != 0;
-		}
+	DirHash hash(dirName.data, dirName.count);
 
-		if (!exists)
-		{
-			errno = ENOENT;
-			return nullptr;
-		}
+	SCOPED_WRITE_LOCK(g_directoryTable.m_lookupLock, lookLock);
+	auto insres = g_directoryTable.m_lookup.try_emplace(hash.key, &g_memoryBlock);
+	DirectoryTable::Directory& dir = insres.first->second;
+	if (insres.second)
+	{
+		if (g_directoryTable.EntryExistsNoLock(hash.key, dirName.data, dirName.count) != DirectoryTable::Exists_No)
+			Rpc_UpdateDirectory(hash.key, dirName.data, dirName.count, false);
+	}
+	bool exists = false;
+	if (dir.tableOffset != InvalidTableOffset)
+	{
+		u32 entryOffset = dir.tableOffset | 0x80000000;
+		DirectoryTable::EntryInformation entryInfo;
+		g_directoryTable.GetEntryInformation(entryInfo, entryOffset);
+		exists = entryInfo.attributes != 0;
+	}
+
+	if (!exists)
+	{
+		errno = ENOENT;
+		return nullptr;
+	}
 
-		g_directoryTable.PopulateDirectory(hash.open, dir);
+	g_directoryTable.PopulateDirectory(hash.open, dir);
 
-		auto dirInfo = new DirInfo();
+	auto dirInfo = new DirInfo();
 	
-		SCOPED_READ_LOCK(dir.lock, lock);
-		dirInfo->fileTableOffsets.resize(dir.files.size());
-		u32 it = 0;
-		for (auto& pair : dir.files)
-			dirInfo->fileTableOffsets[it++] = pair.second;
-		lock.Leave();
+	SCOPED_READ_LOCK(dir.lock, lock);
+	dirInfo->fileTableOffsets.resize(dir.files.size());
+	u32 it = 0;
+	for (auto& pair : dir.files)
+		dirInfo->fileTableOffsets[it++] = pair.second;
+	lock.Leave();
 	
-		DEBUG_LOG_DETOURED("opendir", "(%s) -> %p", dirName.data, dirInfo);
-		return (DIR*)dirInfo;
-	}
+	DEBUG_LOG_DETOURED("opendir", "(%s) -> %p", dirName.data, dirInfo);
 
-	DIR* res = TRUE_WRAPPER(opendir)(dirName.data);
-	DEBUG_LOG_TRUE("opendir", "(%s) -> %p", dirName.data, res);
-	if (!res)
-		return res;
-	res = (DIR*)(u64(res) | 0x1000'0000'0000'0000);
-
-	return res;
+	return (DIR*)(u64(dirInfo) | 0x1000'0000'0000'0000);
 }
 
 UBA_EXPORT int UBA_WRAPPER(dirfd)(DIR* dirp)
