[Uba]
* Fixed variable name

diff --git a/Engine/Source/Programs/UnrealBuildAccelerator/Common/Private/UbaStorageServer.cpp b/Engine/Source/Programs/UnrealBuildAccelerator/Common/Private/UbaStorageServer.cpp
--- a/Engine/Source/Programs/UnrealBuildAccelerator/Common/Private/UbaStorageServer.cpp
+++ b/Engine/Source/Programs/UnrealBuildAccelerator/Common/Private/UbaStorageServer.cpp
@@ -104,7 +104,7 @@ namespace uba
 
 		u64 startTime = GetTime();
 		u32 timeout = 0;
-		while (!waitEntry.Done.IsSet(timeout)) // TODO WaitMultipleObjects (additional work and Done)
+		while (!waitEntry.done.IsSet(timeout)) // TODO WaitMultipleObjects (additional work and Done)
 		{
 			timeout = m_server.DoAdditionalWork() ? 0 : 50;
 
@@ -916,7 +916,7 @@ namespace uba
 					{
 						WaitEntry& waitEntry = waitFindIt->second;
 						waitEntry.Success = true;
-						waitEntry.Done.Set();
+						waitEntry.done.Set();
 					}
 					waitLock.Leave();
 
diff --git a/Engine/Source/Programs/UnrealBuildAccelerator/Common/Public/UbaStorageServer.h b/Engine/Source/Programs/UnrealBuildAccelerator/Common/Public/UbaStorageServer.h
--- a/Engine/Source/Programs/UnrealBuildAccelerator/Common/Public/UbaStorageServer.h
+++ b/Engine/Source/Programs/UnrealBuildAccelerator/Common/Public/UbaStorageServer.h
@@ -64,8 +64,8 @@ namespace uba
 
 		struct WaitEntry
 		{
-			WaitEntry() : Done(false) {}
-			Event Done;
+			WaitEntry() : done(false) {}
+			Event done;
 			bool Success = false;
 			u32 refCount = 0;
 		};
