[UBA]
* Disable mimalloc when tsan is enabled

diff --git a/Engine/Source/Programs/UnrealBuildAccelerator/UbaAgent.Target.cs b/Engine/Source/Programs/UnrealBuildAccelerator/UbaAgent.Target.cs
--- a/Engine/Source/Programs/UnrealBuildAccelerator/UbaAgent.Target.cs
+++ b/Engine/Source/Programs/UnrealBuildAccelerator/UbaAgent.Target.cs
@@ -107,7 +107,7 @@ public class UbaAgentTarget : TargetRules
 		int useMiMalloc = 0;
 		if (Target.Platform.IsInGroup(UnrealPlatformGroup.Windows) || (!Rules.bShouldCompileAsDLL && Target.Platform.IsInGroup(UnrealPlatformGroup.Linux)))
 		{
-			if (Rules.bUsePCHFiles)
+			if (Rules.bUsePCHFiles && !Rules.LinuxPlatform.bEnableThreadSanitizer)
 			{
 				useMiMalloc = 1;
 			}
