[UBA]
* Found another bug in cas file transfer code where if segments where aligned on a very specific size the receiving side would throw an error.

diff --git a/Engine/Source/Programs/UnrealBuildAccelerator/Common/Private/UbaStorageClient.cpp b/Engine/Source/Programs/UnrealBuildAccelerator/Common/Private/UbaStorageClient.cpp
--- a/Engine/Source/Programs/UnrealBuildAccelerator/Common/Private/UbaStorageClient.cpp
+++ b/Engine/Source/Programs/UnrealBuildAccelerator/Common/Private/UbaStorageClient.cpp
@@ -559,8 +559,10 @@ namespace uba
 
 							if (isFirstInBlock)
 							{
-								if (readPosition - readBuffer < sizeof(u32) * 2)
+								if ((readPosition - readBuffer) + responseSize < sizeof(u32) * 2)
+								{
 									return m_logger.Error(TC("Received less than minimum amount of data. Most likely corrupt cas file %s (Available: %u UncompressedSize: %llu LeftUncompressed: %llu)"), casFile.data, u32(readPosition - readBuffer), actualSize, leftUncompressed);
+								}
 								isFirstInBlock = false;
 								u32* blockSize = (u32*)readBuffer;
 								compressedSize = blockSize[0];
@@ -637,7 +639,12 @@ namespace uba
 						memmove(readBuffer, readPosition - overflow, overflow);
 						readPosition = readBuffer + overflow;
 						if (overflow)
-							responseSize = 0;
+						{
+							if (overflow < sizeof(u32) * 2) // Must always have the compressed and uncompressed size to be able to move on with logic above
+								sendSegmentMessage = true;
+							else
+								responseSize = 0;
+						}
 					} while (leftUncompressed);
 
 					if (tryAgain)
