[UBA]
* More info added to some error message related to decompressing memory to memory

diff --git a/Engine/Source/Programs/UnrealBuildAccelerator/Common/Private/UbaSession.cpp b/Engine/Source/Programs/UnrealBuildAccelerator/Common/Private/UbaSession.cpp
--- a/Engine/Source/Programs/UnrealBuildAccelerator/Common/Private/UbaSession.cpp
+++ b/Engine/Source/Programs/UnrealBuildAccelerator/Common/Private/UbaSession.cpp
@@ -553,7 +553,7 @@ namespace uba
 
 			if (isCompressed)
 			{
-				if (!m_storage.DecompressFileToMemory(fileName, fileHandle, mappedView.memory, size))
+				if (!m_storage.DecompressFileToMemory(fileName, fileHandle, mappedView.memory, size, TC("MappedMemory")))
 					return false;
 			}
 			else
@@ -640,7 +640,7 @@ namespace uba
 				auto mappedViewWrite = m_fileMappingBuffer.AllocAndMapView(MappedView_Transient, size, alignment, fileName);
 				auto unmapGuard = MakeGuard([&](){ m_fileMappingBuffer.UnmapView(mappedViewWrite, fileName); });
 
-				if (!m_storage.DecompressMemoryToMemory(readMemory, mappedViewWrite.memory, size, fileName))
+				if (!m_storage.DecompressMemoryToMemory(readMemory, mappedViewWrite.memory, size, fileName, TC("TransientMapping")))
 					return false;
 				unmapGuard.Execute();
 
diff --git a/Engine/Source/Programs/UnrealBuildAccelerator/Common/Private/UbaStorage.cpp b/Engine/Source/Programs/UnrealBuildAccelerator/Common/Private/UbaStorage.cpp
--- a/Engine/Source/Programs/UnrealBuildAccelerator/Common/Private/UbaStorage.cpp
+++ b/Engine/Source/Programs/UnrealBuildAccelerator/Common/Private/UbaStorage.cpp
@@ -602,7 +602,7 @@ namespace uba
 		return false;
 	}
 
-	bool StorageImpl::DecompressMemoryToMemory(u8* compressedData, u8* writeData, u64 decompressedSize, const tchar* readHint)
+	bool StorageImpl::DecompressMemoryToMemory(u8* compressedData, u8* writeData, u64 decompressedSize, const tchar* readHint, const tchar* writeHint)
 	{
 		UBA_ASSERT(compressedData);
 		UBA_ASSERT(writeData);
@@ -724,13 +724,13 @@ namespace uba
 					break;
 				u32 decompressedBlockSize = ((u32*)readPos)[1];
 				if (decompressedBlockSize == 0 || decompressedBlockSize > left)
-					return m_logger.Error(TC("Decompressed block size %u is invalid. Decompressed file is %u (%s)"), decompressedBlockSize, decompressedSize, readHint);
+					return m_logger.Error(TC("Decompressed block size %u is invalid. Decompressed file is %u (%s -> %s)"), decompressedBlockSize, decompressedSize, readHint, writeHint);
 				readPos += sizeof(u32) * 2;
 
 				TimerScope ts(stats.decompressToMem);
 				OO_SINTa decompLen = OodleLZ_Decompress(readPos, (OO_SINTa)compressedBlockSize, writePos, (OO_SINTa)decompressedBlockSize); (void)decompLen;
 				if (decompLen != decompressedBlockSize)
-					return m_logger.Error(TC("Expecting to be able to decompress to %u bytes but got %llu (%s)"), decompressedBlockSize, decompLen, readHint);
+					return m_logger.Error(TC("Expecting to be able to decompress to %u bytes but got %llu (%s -> %s)"), decompressedBlockSize, decompLen, readHint, writeHint);
 				writePos += decompressedBlockSize;
 				readPos += compressedBlockSize;
 				left -= decompressedBlockSize;
@@ -1682,7 +1682,7 @@ namespace uba
 								mem += sizeof(u64);
 								u8* dest = new u8[decompressedSize];
 								auto g = MakeGuard([dest]() { delete[] dest; });
-								if (!DecompressMemoryToMemory(mem, dest, decompressedSize, filePath.c_str()))
+								if (!DecompressMemoryToMemory(mem, dest, decompressedSize, filePath.c_str(), TC("Memory")))
 									return;
 								checkedKey = CalculateCasKey(dest, decompressedSize, true);
 							}
@@ -2488,12 +2488,12 @@ namespace uba
 
 					if (casEntry->mappingHandle.IsValid())
 					{
-						if (!DecompressMemoryToMemory(readData, destinationFile.GetData(), decompressedSize, casFile.data))
+						if (!DecompressMemoryToMemory(readData, destinationFile.GetData(), decompressedSize, casFile.data, destination))
 							return false;
 					}
 					else
 					{
-						if (!DecompressFileToMemory(CasKeyString(actualKey).str, readHandle, destinationFile.GetData(), decompressedSize))
+						if (!DecompressFileToMemory(CasKeyString(actualKey).str, readHandle, destinationFile.GetData(), decompressedSize, destination))
 							return false;
 					}
 				}
@@ -2711,7 +2711,7 @@ namespace uba
 		return ToCasKey(hasher, storeCompressed);
 	}
 
-	bool StorageImpl::DecompressFileToMemory(const tchar* fileName, FileHandle fileHandle, u8* dest, u64 decompressedSize)
+	bool StorageImpl::DecompressFileToMemory(const tchar* fileName, FileHandle fileHandle, u8* dest, u64 decompressedSize, const tchar* writeHint)
 	{
 		if (m_workManager && decompressedSize > BufferSlotSize*4) // Arbitrary size threshold. We want to at least catch the pch here
 		{
@@ -2727,7 +2727,7 @@ namespace uba
 				return m_logger.Error(TC("Failed to map view of file mapping for %s (%s)"), fileName, LastErrorToText().data);
 			auto udg = MakeGuard([&]() { UnmapViewOfFile(fileData, compressedSize, fileName); });
 			
-			if (!DecompressMemoryToMemory(fileData + 8, dest, decompressedSize, fileName))
+			if (!DecompressMemoryToMemory(fileData + 8, dest, decompressedSize, fileName, writeHint))
 				return false;
 		}
 		else
diff --git a/Engine/Source/Programs/UnrealBuildAccelerator/Common/Public/UbaStorage.h b/Engine/Source/Programs/UnrealBuildAccelerator/Common/Public/UbaStorage.h
--- a/Engine/Source/Programs/UnrealBuildAccelerator/Common/Public/UbaStorage.h
+++ b/Engine/Source/Programs/UnrealBuildAccelerator/Common/Public/UbaStorage.h
@@ -34,8 +34,8 @@ namespace uba
 		virtual bool GetZone(StringBufferBase& out) = 0;
 		virtual bool HasProxy(u32 clientId) { return false; }
 
-		virtual bool DecompressFileToMemory(const tchar* fileName, FileHandle fileHandle, u8* dest, u64 decompressedSize) = 0;
-		virtual bool DecompressMemoryToMemory(u8* compressedData, u8* writeData, u64 decompressedSize, const tchar* readHint) = 0;
+		virtual bool DecompressFileToMemory(const tchar* fileName, FileHandle fileHandle, u8* dest, u64 decompressedSize, const tchar* writeHint) = 0;
+		virtual bool DecompressMemoryToMemory(u8* compressedData, u8* writeData, u64 decompressedSize, const tchar* readHint, const tchar* writeHint) = 0;
 		virtual bool CreateDirectory(const tchar* dir) = 0;
 		virtual bool DeleteCasForFile(const tchar* file) = 0;
 
@@ -122,7 +122,7 @@ namespace uba
 		virtual bool StoreCompressed() const final { return m_storeCompressed; }
 		virtual void PrintSummary(Logger& logger) override;
 
-		virtual bool DecompressFileToMemory(const tchar* fileName, FileHandle fileHandle, u8* dest, u64 decompressedSize) override;
+		virtual bool DecompressFileToMemory(const tchar* fileName, FileHandle fileHandle, u8* dest, u64 decompressedSize, const tchar* writeHint) override;
 		virtual bool CreateDirectory(const tchar* dir) override;
 		virtual bool DeleteCasForFile(const tchar* file) override;
 		virtual bool RetrieveCasFile(RetrieveResult& out, const CasKey& casKey, const tchar* hint, FileMappingBuffer* mappingBuffer = nullptr, u64 memoryMapAlignment = 1, bool allowProxy = true) override;
@@ -166,7 +166,7 @@ namespace uba
 		bool AddCasFile(StringKey fileNameKey, const tchar* fileName, const CasKey& casKey, bool deferCreation, bool fileIsCompressed);
 		void CasEntryAccessed(const CasKey& casKey);
 		virtual bool IsDisallowedPath(const tchar* fileName);
-		virtual bool DecompressMemoryToMemory(u8* compressedData, u8* writeData, u64 decompressedSize, const tchar* readHint) override;
+		virtual bool DecompressMemoryToMemory(u8* compressedData, u8* writeData, u64 decompressedSize, const tchar* readHint, const tchar* writeHint) override;
 		bool DecompressMemoryToFile(u8* compressedData, FileAccessor& destination, u64 decompressedSize, bool useNoBuffering);
 
 		void CasEntryAccessed(CasEntry& entry);
