[UBA]
* Changed event to use clock CLOCK_MONOTONIC for linux and changed to pthread_cond_timedwait_relative_np for mac.
* Changed GetTime to use clock_gettime_nsec_np(CLOCK_UPTIME_RAW) for mac.

These changes are to attempt to fix the 10minute event timeouts that are not 10 minutes but rather a second.

diff --git a/Engine/Source/Programs/UnrealBuildAccelerator/Common/Private/UbaNetworkClient.cpp b/Engine/Source/Programs/UnrealBuildAccelerator/Common/Private/UbaNetworkClient.cpp
--- a/Engine/Source/Programs/UnrealBuildAccelerator/Common/Private/UbaNetworkClient.cpp
+++ b/Engine/Source/Programs/UnrealBuildAccelerator/Common/Private/UbaNetworkClient.cpp
@@ -568,7 +568,7 @@ namespace uba
 			u32 timeoutMs = 10 * 60 * 1000;
 			if (!gotResponse.IsSet(timeoutMs))
 			{
-				m_logger.Error(TC("Timed out after %s waiting for message response from server."), TimeToText(GetTime() - waitStart).str);
+				m_logger.Error(TC("Timed out after %s waiting for message response from server."), TimeToText(GetTime() - waitStart, true).str);
 				message.m_error = 4;
 			}
 			else if (m_cryptoKey && !message.m_error && message.m_responseSize)
diff --git a/Engine/Source/Programs/UnrealBuildAccelerator/Core/Private/UbaEvent.cpp b/Engine/Source/Programs/UnrealBuildAccelerator/Core/Private/UbaEvent.cpp
--- a/Engine/Source/Programs/UnrealBuildAccelerator/Core/Private/UbaEvent.cpp
+++ b/Engine/Source/Programs/UnrealBuildAccelerator/Core/Private/UbaEvent.cpp
@@ -11,7 +11,6 @@
 #endif
 
 #define UBA_TEST_WAIT_QUALITY 0
-#define UBA_USE_TIME_OF_DAY_CONDITION 1//PLATFORM_MAC
 
 namespace uba
 {
@@ -19,21 +18,6 @@ namespace uba
 
 	u64 GetMonoticTimeNs();
 
-	#if UBA_USE_TIME_OF_DAY_CONDITION
-	static u64 GetTimeOfDayOffset = []()
-		{
-			timeval tv;
-			gettimeofday(&tv, NULL);
-				
-			u64 tsNs = GetMonoticTimeNs();
-			u64 tvNs = u64(tv.tv_sec)*1'000'000'000 + tv.tv_usec * 1'000;
-
-			return tvNs - tsNs;
-		}();
-	#else
-	static constexpr u64 GetTimeOfDayOffset = 0;
-	#endif
-
 	struct EventImpl
 	{
 		EventImpl() = default;
@@ -85,7 +69,7 @@ namespace uba
 				return false;
 			}
 
-			#if !UBA_USE_TIME_OF_DAY_CONDITION
+			#if PLATFORM_LINUX
 			if (pthread_condattr_setclock(&attrcond, CLOCK_MONOTONIC) != 0)
 			{
 				UBA_ASSERTF(false, "pthread_condattr_setclock failed");
@@ -237,10 +221,15 @@ namespace uba
 					}
 					else  // timed wait.
 					{
-						struct timespec timeout = ToTimeSpec(startTimeNs + timeoutNs + GetTimeOfDayOffset);
-						int rc = pthread_cond_timedwait(&m_condition, &m_mutex, &timeout);    // unlocks Mutex while blocking...
+#if PLATFORM_MAC
+						struct timespec timeout = ToTimeSpec(timeoutNs);
+						int rc = pthread_cond_timedwait_relative_np(&m_condition, &m_mutex, &timeout); // unlocks Mutex while blocking...
+						UBA_ASSERTF((rc == 0) || (rc == ETIMEDOUT), "pthread_cond_timedwait_relative_np failed"); (void)rc;
+#else
+						struct timespec timeout = ToTimeSpec(startTimeNs + timeoutNs);
+						int rc = pthread_cond_timedwait(&m_condition, &m_mutex, &timeout); // unlocks Mutex while blocking...
 						UBA_ASSERTF((rc == 0) || (rc == ETIMEDOUT), "pthread_cond_timedwait failed"); (void)rc;
-
+#endif
 						#if UBA_TEST_WAIT_QUALITY
 						++loop;
 						#endif
diff --git a/Engine/Source/Programs/UnrealBuildAccelerator/Core/Private/UbaTimer.cpp b/Engine/Source/Programs/UnrealBuildAccelerator/Core/Private/UbaTimer.cpp
--- a/Engine/Source/Programs/UnrealBuildAccelerator/Core/Private/UbaTimer.cpp
+++ b/Engine/Source/Programs/UnrealBuildAccelerator/Core/Private/UbaTimer.cpp
@@ -15,10 +15,14 @@ namespace uba
 	#if !PLATFORM_WINDOWS
 	u64 GetMonoticTimeNs()
 	{
+		#if PLATFORM_MAC
+		return clock_gettime_nsec_np(CLOCK_UPTIME_RAW);
+		#else
 		struct timespec ts;
 		if (clock_gettime(CLOCK_MONOTONIC, &ts) == -1)
 			FatalError(1401, TC("clock_gettime(CLOCK_MONOTONIC) failed"));
 		return u64(ts.tv_sec * 1'000'000'000LL + ts.tv_nsec);
+		#endif
 	}
 	#endif
 
