[UBA]
* Bumped uba cache cas key entry table size from 32mb to 64mb

diff --git a/Engine/Source/Programs/UnrealBuildAccelerator/Common/Public/UbaNetwork.h b/Engine/Source/Programs/UnrealBuildAccelerator/Common/Public/UbaNetwork.h
--- a/Engine/Source/Programs/UnrealBuildAccelerator/Common/Public/UbaNetwork.h
+++ b/Engine/Source/Programs/UnrealBuildAccelerator/Common/Public/UbaNetwork.h
@@ -19,7 +19,7 @@ namespace uba
 	static constexpr u32 CacheNetworkVersion = 4;
 
 	static constexpr u32 CachePathTableMaxSize = 32*1024*1024;
-	static constexpr u32 CacheCasKeyTableMaxSize = 32*1024*1024;
+	static constexpr u32 CacheCasKeyTableMaxSize = 64*1024*1024;
 
 	// Messages used over network between client/server (system, storage and session)
 
