[UBA]
* Added more logging info around active fetches

diff --git a/Engine/Source/Programs/UnrealBuildAccelerator/Common/Private/UbaStorageServer.cpp b/Engine/Source/Programs/UnrealBuildAccelerator/Common/Private/UbaStorageServer.cpp
--- a/Engine/Source/Programs/UnrealBuildAccelerator/Common/Private/UbaStorageServer.cpp
+++ b/Engine/Source/Programs/UnrealBuildAccelerator/Common/Private/UbaStorageServer.cpp
@@ -195,6 +195,8 @@ namespace uba
 					continue;
 				}
 
+				m_logger.Detail(TC("Cancelled fetch id %u because of disconnect of client with id %u"), u32(it->first), clientId);
+
 				fetch.Release(*this, TC("OnDisconnected"));
 
 				if (m_traceFetch)
@@ -659,7 +661,7 @@ namespace uba
 				SCOPED_READ_LOCK(m_activeFetchesLock, lock);
 				auto findIt = m_activeFetches.find(fetchId);
 				if (findIt == m_activeFetches.end())
-					return m_logger.Error(TC("Can't find active fetch %u, disconnected client? (index %u)"), fetchId, fetchIndex);
+					return m_logger.Error(TC("Can't find active fetch %u, disconnected client? (fetch index %u, client id %u uid %s)"), fetchId, fetchIndex, connectionInfo.GetId(), GuidToString(connectionInfo.GetUid()).str);
 				ActiveFetch& fetch = findIt->second;
 				UBA_ASSERT(fetch.clientId == connectionInfo.GetId());
 				lock.Leave();
diff --git a/Engine/Source/Programs/UnrealBuildAccelerator/Common/Private/UbaStorageUtils.cpp b/Engine/Source/Programs/UnrealBuildAccelerator/Common/Private/UbaStorageUtils.cpp
--- a/Engine/Source/Programs/UnrealBuildAccelerator/Common/Private/UbaStorageUtils.cpp
+++ b/Engine/Source/Programs/UnrealBuildAccelerator/Common/Private/UbaStorageUtils.cpp
@@ -154,7 +154,7 @@ namespace uba
 		{
 			Entry& entry = entries[i];
 			if (!entry.done.IsSet(5*60*1000))
-				logger.Error(TC("SendBatchMessages timed out getting async message response"));
+				logger.Error(TC("SendBatchMessages timed out after 5 minutes getting async message response (%u/%u). Received %llu bytes so far. FetchId: %u"), i, inFlightCount, responseSize, fetchId);
 			if (!entry.message.ProcessAsyncResults(entry.reader))
 				success = false;
 			else
