[UBA]
* Added code to UbaCacheServer to automatically lower expiration time when cas key table ends up over 30mb large

diff --git a/Engine/Source/Programs/UnrealBuildAccelerator/Common/Private/UbaCacheServer.cpp b/Engine/Source/Programs/UnrealBuildAccelerator/Common/Private/UbaCacheServer.cpp
--- a/Engine/Source/Programs/UnrealBuildAccelerator/Common/Private/UbaCacheServer.cpp
+++ b/Engine/Source/Programs/UnrealBuildAccelerator/Common/Private/UbaCacheServer.cpp
@@ -51,6 +51,7 @@ namespace uba
 		Atomic<u64> lastSavedTime;
 		Atomic<u64> lastUsedTime;
 
+		u64 expirationTimeSeconds = 0;
 		u32 index = ~0u;
 	};
 
@@ -191,6 +192,8 @@ namespace uba
 		if (databaseVersion != CacheFileVersion)
 			bucket.needsSave = true;
 
+		bucket.expirationTimeSeconds = m_expirationTimeSeconds;
+
 		u32 pathTableSize = reader.ReadU32();
 		if (pathTableSize)
 		{
@@ -512,11 +515,6 @@ namespace uba
 		u64 oldest = 0;
 		u64 longestUnused = 0;
 
-		u64 lastUseTimeLimit = 0;
-		if (m_expirationTimeSeconds && GetFileTimeAsSeconds(now - m_creationTime) > m_expirationTimeSeconds)
-			lastUseTimeLimit = (now - m_creationTime) - GetSecondsAsFileTime(m_expirationTimeSeconds);
-
-
 		u32 workerCount = m_server.GetWorkerCount();
 		u32 workerCountToUse = workerCount > 0 ? workerCount - 1 : 0;
 		u32 workerCountToUseForBuckets = Min(workerCountToUse, u32(m_buckets.size()));
@@ -546,9 +544,24 @@ namespace uba
 				bucket.totalEntryCount = 0;
 				bucket.totalEntrySize = 0;
 
+
 				ReaderWriterLock keysToEraseLock;
 				Vector<CasKey> keysToErase;
 
+				u64 lastUseTimeLimit = ~0u;
+				if (bucket.expirationTimeSeconds)
+				{
+					// If cas key table size is over 30mb we need to shorten expiration time
+					if (bucket.m_casKeyTable.GetSize() > 30*1024*1024)
+					{
+						m_logger.Info(TC("Lowered expiration time for bucket %u to %s"), bucket.index, TimeToText(MsToTime(bucket.expirationTimeSeconds*1000), true).str);
+						bucket.expirationTimeSeconds -= 60*60; // Shorten by one hour.. this will be reset 
+					}
+					if (GetFileTimeAsSeconds(now - m_creationTime) > bucket.expirationTimeSeconds)
+						lastUseTimeLimit = (now - m_creationTime) - GetSecondsAsFileTime(bucket.expirationTimeSeconds);
+				}
+
+
 				m_server.ParallelFor(workerCountToUse, bucket.m_cacheEntryLookup, [&, touchedCas = Vector<u64*>()](auto& li) mutable
 				{
 					CacheEntries& entries = li->second;
@@ -945,9 +958,12 @@ namespace uba
 	{
 		SCOPED_WRITE_LOCK(m_bucketsLock, bucketsLock);
 		auto insres = m_buckets.try_emplace(id, id);
-		if (insres.second)
-			insres.first->second.index = u32(m_buckets.size() - 1);
-		return insres.first->second;
+		auto& bucket = insres.first->second;
+		if (!insres.second)
+			return bucket;
+		bucket.index = u32(m_buckets.size() - 1);
+		bucket.expirationTimeSeconds = m_expirationTimeSeconds;
+		return bucket;
 	}
 
 	u32 CacheServer::GetBucketWorkerCount()
