[UBA]
* Fixed so bucket expiration time is only lowered on first delete iteration when running maintenance

diff --git a/Engine/Source/Programs/UnrealBuildAccelerator/Common/Private/UbaCacheServer.cpp b/Engine/Source/Programs/UnrealBuildAccelerator/Common/Private/UbaCacheServer.cpp
--- a/Engine/Source/Programs/UnrealBuildAccelerator/Common/Private/UbaCacheServer.cpp
+++ b/Engine/Source/Programs/UnrealBuildAccelerator/Common/Private/UbaCacheServer.cpp
@@ -529,6 +529,7 @@ namespace uba
 		Atomic<u64> activeDropCount;
 		auto dropCasGuard = MakeGuard([&]() { while (activeDropCount != 0) Sleep(1); });
 
+		u32 deleteIteration = 0;
 		u64 deleteCacheEntriesStartTime = GetTime();
 		do
 		{
@@ -552,7 +553,7 @@ namespace uba
 				if (bucket.expirationTimeSeconds)
 				{
 					// If cas key table size is over 30mb we need to shorten expiration time
-					if (bucket.m_casKeyTable.GetSize() > 30*1024*1024)
+					if (deleteIteration == 0 && bucket.m_casKeyTable.GetSize() > 30*1024*1024)
 					{
 						m_logger.Info(TC("Lowered expiration time for bucket %u to %s"), bucket.index, TimeToText(MsToTime(bucket.expirationTimeSeconds*1000), true).str);
 						bucket.expirationTimeSeconds -= 60*60; // Shorten by one hour.. this will be reset 
@@ -739,6 +740,7 @@ namespace uba
 				++activeDropCount;
 				m_server.AddWork([&, key = casKey]() { m_storage.DropCasFile(key, true, TC("")); --activeDropCount; }, 1, TC(""));
 			}
+			++deleteIteration;
 		}
 		while (!deletedCasFiles.empty()); // if cas files are deleted we need to do another loop and check cache entry inputs to see if files were inputs
 
