[UBA]
* Added support for huge pages in MemoryBlock
* Changed so UbaCacheService tries to use huge pages for maintenance work.

diff --git a/Engine/Source/Programs/UnrealBuildAccelerator/CacheService/Private/UbaCacheService.cpp b/Engine/Source/Programs/UnrealBuildAccelerator/CacheService/Private/UbaCacheService.cpp
--- a/Engine/Source/Programs/UnrealBuildAccelerator/CacheService/Private/UbaCacheService.cpp
+++ b/Engine/Source/Programs/UnrealBuildAccelerator/CacheService/Private/UbaCacheService.cpp
@@ -198,7 +198,18 @@ namespace uba
 		#if UBA_DEBUG
 		dbgStr = TC(" (DEBUG)");
 		#endif
-		logger.Info(TC("UbaCacheService v%s(%u)%s (Workers: %u, Rootdir: \"%s\", StoreCapacity: %uGb, Expiration: %s)\n"), Version, CacheNetworkVersion, dbgStr, GetLogicalProcessorCount(), g_rootDir.data, storageCapacityGb, TimeToText(MsToTime(expirationTimeSeconds)*1000, true).str);
+		logger.Info(TC("UbaCacheService v%s(%u)%s (Workers: %u, Rootdir: \"%s\", StoreCapacity: %uGb, Expiration: %s)"), Version, CacheNetworkVersion, dbgStr, GetLogicalProcessorCount(), g_rootDir.data, storageCapacityGb, TimeToText(MsToTime(expirationTimeSeconds)*1000, true).str);
+
+		u64 maintenanceReserveSizeMb = 128;
+
+		#if PLATFORM_LINUX
+		u64 hugePageCount = GetHugePageCount();
+		u64 recommendedHugePageCount = (maintenanceReserveSizeMb*GetLogicalProcessorCount())/2;
+		if (hugePageCount < recommendedHugePageCount)
+			logger.Info(TC("  Improve maintenance performance by enabling %llu huge pages on system (%llu enabled)"), recommendedHugePageCount, hugePageCount);
+		#endif
+
+		logger.Info(TC(""));
 
 		u64 storageCapacity = u64(storageCapacityGb)*1000*1000*1000;
 
@@ -210,13 +221,6 @@ namespace uba
 
 		// TODO: Change workdir to make it full
 
-
-		StringBuffer<> logFile;
-		#if UBA_DEBUG
-		logFile.count = GetFullPathNameW(g_rootDir.data, logFile.capacity, logFile.data, nullptr);
-		logFile.EnsureEndsWithSlash().Append(TC("DebugLog.log"));
-		#endif
-
 		#if PLATFORM_WINDOWS
 		SetConsoleCtrlHandler(ConsoleHandler, TRUE);
 		#else
@@ -245,6 +249,7 @@ namespace uba
 
 		CacheServerCreateInfo cacheInfo(storageServer, g_rootDir.data, logWriter);
 		cacheInfo.expirationTimeSeconds = expirationTimeSeconds;
+		cacheInfo.maintenanceReserveSize = maintenanceReserveSizeMb * 1024 * 1024;
 		CacheServer cacheServer(cacheInfo);
 
 		if (!cacheServer.Load())
@@ -303,7 +308,8 @@ namespace uba
 #if PLATFORM_WINDOWS
 int wmain(int argc, wchar_t* argv[])
 {
-	return uba::WrappedMain(argc, argv);
+	int res = uba::WrappedMain(argc, argv);
+	return res;
 }
 #else
 int main(int argc, char* argv[])
diff --git a/Engine/Source/Programs/UnrealBuildAccelerator/Common/Private/UbaCacheServer.cpp b/Engine/Source/Programs/UnrealBuildAccelerator/Common/Private/UbaCacheServer.cpp
--- a/Engine/Source/Programs/UnrealBuildAccelerator/Common/Private/UbaCacheServer.cpp
+++ b/Engine/Source/Programs/UnrealBuildAccelerator/Common/Private/UbaCacheServer.cpp
@@ -73,8 +73,9 @@ namespace uba
 	,	m_storage(info.storage)
 	{
 		m_checkInputsForDeletedCas = info.checkInputsForDeletedCas;
-		m_startTime = GetTime();
+		m_bootTime = GetTime();
 
+		m_maintenanceReserveSize = info.maintenanceReserveSize;
 		m_expirationTimeSeconds = info.expirationTimeSeconds;
 
 		m_rootDir.count = GetFullPathNameW(info.rootDir, m_rootDir.capacity, m_rootDir.data, NULL);
@@ -447,11 +448,27 @@ namespace uba
 		bool forceAllSteps = m_forceAllSteps;
 		m_forceAllSteps = false;
 
+		bool entriesAdded = m_addsSinceMaintenance != 0;
+		m_addsSinceMaintenance = 0;
+
+		u64 startTime = GetTime();
+
+		if (entriesAdded)
+		{
+			auto& storageStats = m_storage.Stats();
+			u64 hits = m_cacheKeyHitCount;
+			u64 miss = m_cacheKeyFetchCount - hits;
+			m_logger.Info(TC("Stats since boot (%s ago)"), TimeToText(startTime - m_bootTime, true).str);
+			m_logger.Info(TC("  CacheServer %llu hits, %llu misses"), hits, miss);
+			u64 recvCount = storageStats.sendCas.count.load();
+			u64 sendCount = storageStats.recvCas.count.load();
+			m_logger.Info(TC("  StorageServer cas %llu (%s) sent, %llu (%s) received"), recvCount, BytesToText(storageStats.sendCasBytesComp).str, sendCount, BytesToText(storageStats.recvCasBytesComp).str);
+		}
+
 		if (m_shouldWipe)
 		{
 			m_shouldWipe = false;
 			m_logger.Info(TC("Obliterating database"));
-			m_addsSinceMaintenance = 0;
 			m_longestMaintenance = 0;
 			m_buckets.clear();
 			forceAllSteps = true;
@@ -459,14 +476,12 @@ namespace uba
 		}
 		else
 		{
-			m_logger.Info(TC("Maintenance started after %u added cache entries"), m_addsSinceMaintenance.load());
+			TimeToText lastTime(startTime - m_lastMaintenance, true);
+			m_logger.Info(TC("Maintenance started after %u added cache entries (Ran last time %s ago)"), entriesAdded, (m_lastMaintenance ? lastTime.str : TC("Never")));
 		}
-		
-		u64 startTime = GetTime();
-
-		bool entriesAdded = m_addsSinceMaintenance != 0;
-		m_addsSinceMaintenance = 0;
 
+		m_lastMaintenance = startTime;
+		
 		UnorderedSet<CasKey> deletedCasFiles;
 		m_storage.HandleOverflow(&deletedCasFiles);
 		u64 deletedCasCount = deletedCasFiles.size();
@@ -758,6 +773,8 @@ namespace uba
 		if (shouldExit())
 			return true;
 
+		u64 maxCommittedMemory = 0;
+
 		m_server.ParallelFor(workerCountToUseForBuckets, m_buckets, [&](auto& it)
 		{
 			u64 bucketStartTime = GetTime();
@@ -771,7 +788,10 @@ namespace uba
 			}
 			bucket.hasDeletedEntries = false;
 
-			MemoryBlock memoryBlock(128*1024*1024);
+			// Try to use large blocks
+			MemoryBlock memoryBlock;
+			if (!memoryBlock.Init(m_maintenanceReserveSize, nullptr, true))
+				memoryBlock.Init(m_maintenanceReserveSize, nullptr, false);
 
 			GrowingNoLockUnorderedSet<u32> usedCasKeyOffsets(&memoryBlock);
 			usedCasKeyOffsets.reserve(bucket.m_casKeyTable.GetKeyCount());
@@ -894,6 +914,10 @@ namespace uba
 			bucket.needsSave = true;
 
 			m_logger.Info(TC("    Bucket %u Done (%s). CacheEntries: %llu (%s) PathTable: %s CasTable: %s"), bucket.index, TimeToText(GetTime() - bucketStartTime).str, bucket.totalEntryCount.load(), BytesToText(bucket.totalEntrySize.load()).str, BytesToText(bucket.m_pathTable.GetSize()).str, BytesToText(bucket.m_casKeyTable.GetSize()).str);
+
+			SCOPED_WRITE_LOCK(existingCasLock, l); // Just reusing lock for other purpose
+			maxCommittedMemory = Max(maxCommittedMemory, memoryBlock.writtenSize);
+
 		}, TC(""), true);
 
 		// Need to make sure all cas entries are dropped before saving cas table
@@ -917,7 +941,7 @@ namespace uba
 		u64 oldestTime = oldest ? GetFileTimeAsTime(now - (m_creationTime + oldest)) : 0;
 		u64 longestUnusedTime = longestUnused ? GetFileTimeAsTime(now - (m_creationTime + longestUnused)) : 0;
 		u64 duration = GetTime() - startTime;
-		m_logger.Info(TC("Maintenance done! (%s) CasFiles: %llu (%s) Entries: %llu Oldest: %s LongestUnused: %s"), TimeToText(duration).str, totalCasCount - deletedCasCount, BytesToText(totalCasSize).str, totalEntryCount.load(), TimeToText(oldestTime, true).str, TimeToText(longestUnusedTime, true).str);
+		m_logger.Info(TC("Maintenance done! (%s) CasFiles: %llu (%s) Entries: %llu Oldest: %s LongestUnused: %s MaintenanceMem: %s"), TimeToText(duration).str, totalCasCount - deletedCasCount, BytesToText(totalCasSize).str, totalEntryCount.load(), TimeToText(oldestTime, true).str, TimeToText(longestUnusedTime, true).str, BytesToText(maxCommittedMemory).str);
 		
 		m_longestMaintenance = Max(m_longestMaintenance, duration);
 
@@ -1305,6 +1329,8 @@ namespace uba
 		Bucket& bucket = GetBucket(reader);
 		CasKey cmdKey = reader.ReadCasKey();
 
+		++m_cacheKeyFetchCount;
+
 		SCOPED_READ_LOCK(bucket.m_cacheEntryLookupLock, lock);
 		auto findIt = bucket.m_cacheEntryLookup.find(cmdKey);
 		if (findIt == bucket.m_cacheEntryLookup.end())
@@ -1325,6 +1351,8 @@ namespace uba
 		CasKey cmdKey = reader.ReadCasKey();
 		u64 entryId = reader.Read7BitEncoded();
 
+		++m_cacheKeyHitCount;
+
 		SCOPED_READ_LOCK(bucket.m_cacheEntryLookupLock, lock);
 		auto findIt = bucket.m_cacheEntryLookup.find(cmdKey);
 		if (findIt == bucket.m_cacheEntryLookup.end())
@@ -1473,7 +1501,7 @@ namespace uba
 		{
 			writeLine(TC("UbaCacheServer status"));
 			writeLine(line.Clear().Appendf(TC("  CreationTime: %s ago"), TimeToText(GetFileTimeAsTime(GetSystemTimeAsFileTime() - m_creationTime), true).str).data);
-			writeLine(line.Clear().Appendf(TC("  UpTime: %s"), TimeToText(GetTime() - m_startTime, true).str).data);
+			writeLine(line.Clear().Appendf(TC("  UpTime: %s"), TimeToText(GetTime() - m_bootTime, true).str).data);
 			writeLine(line.Clear().Appendf(TC("  Longest maintenance: %s"), TimeToText(m_longestMaintenance).str).data);
 			writeLine(line.Clear().Appendf(TC("  Buckets:")).data);
 			u32 index = 0;
diff --git a/Engine/Source/Programs/UnrealBuildAccelerator/Common/Private/UbaLogger.cpp b/Engine/Source/Programs/UnrealBuildAccelerator/Common/Private/UbaLogger.cpp
--- a/Engine/Source/Programs/UnrealBuildAccelerator/Common/Private/UbaLogger.cpp
+++ b/Engine/Source/Programs/UnrealBuildAccelerator/Common/Private/UbaLogger.cpp
@@ -70,13 +70,15 @@ namespace uba
 		va_list arg;
 		va_start(arg, format);
 		tchar buffer[1024];
-		int count = TSprintf_s(buffer, 1024, format, arg);
+		int count = Tvsprintf_s(buffer, 1024, format, arg);
 		if (count <= 0)
 			TStrcpy_s(buffer, 1024, format);
 		va_end(arg);
 #if PLATFORM_WINDOWS
 		wprintf(TC("FATAL ERROR %u: %s\n"), code, buffer);
 		fflush(stdout);
+		if (IsDebuggerPresent())
+			DebugBreak();
 		ExitProcess(code);
 #else
 		printf(TC("FATAL ERROR %u: %s\n"), code, buffer);
@@ -312,8 +314,10 @@ namespace uba
 			TSprintf_s(str, 32, TC("%.1fkb"), double(bytes) / 1000ull);
 		else if (bytes < 1000ull * 1000 * 1000)
 			TSprintf_s(str, 32, TC("%.1fmb"), double(bytes) / (1000ull * 1000));
-		else
+		else if (bytes < 1000ull * 1000 * 1000 * 1000)
 			TSprintf_s(str, 32, TC("%.1fgb"), double(bytes) / (1000ull * 1000 * 1000));
+		else
+			TSprintf_s(str, 32, TC("%.1ftb"), double(bytes) / (1000ull * 1000 * 1000 * 1000));
 	}
 
 #if UBA_DEBUG_LOGGER
diff --git a/Engine/Source/Programs/UnrealBuildAccelerator/Common/Public/UbaCacheServer.h b/Engine/Source/Programs/UnrealBuildAccelerator/Common/Public/UbaCacheServer.h
--- a/Engine/Source/Programs/UnrealBuildAccelerator/Common/Public/UbaCacheServer.h
+++ b/Engine/Source/Programs/UnrealBuildAccelerator/Common/Public/UbaCacheServer.h
@@ -35,6 +35,9 @@ namespace uba
 		// The time cache entries will stay around after they were last used in hours (defaults to two days)
 		// Set to zero to never expire
 		u64 expirationTimeSeconds = 2*24*60*60;
+
+		// The amount of reserved memory used per core when doing maintenance
+		u64 maintenanceReserveSize = 128ull * 1024 * 1024;
 	};
 
 	class CacheServer
@@ -82,6 +85,8 @@ namespace uba
 		StringBuffer<MaxPath> m_rootDir;
 
 		Atomic<u32> m_addsSinceMaintenance;
+		Atomic<u64> m_cacheKeyFetchCount;
+		Atomic<u64> m_cacheKeyHitCount;
 		Atomic<bool> m_isRunningMaintenance;
 
 		ReaderWriterLock m_bucketsLock;
@@ -92,8 +97,10 @@ namespace uba
 
 		Atomic<bool> m_shutdownRequested = false;
 
+		u64 m_maintenanceReserveSize = 0;
 		u64 m_creationTime = 0;
-		u64 m_startTime = 0;
+		u64 m_bootTime = 0;
+		u64 m_lastMaintenance = 0;
 		u64 m_longestMaintenance = 0;
 		u64 m_expirationTimeSeconds = 0;
 		bool m_dbfileDirty = false;
diff --git a/Engine/Source/Programs/UnrealBuildAccelerator/Core/Private/UbaMemory.cpp b/Engine/Source/Programs/UnrealBuildAccelerator/Core/Private/UbaMemory.cpp
--- a/Engine/Source/Programs/UnrealBuildAccelerator/Core/Private/UbaMemory.cpp
+++ b/Engine/Source/Programs/UnrealBuildAccelerator/Core/Private/UbaMemory.cpp
@@ -3,11 +3,18 @@
 #include "UbaMemory.h"
 #include "UbaPlatform.h"
 
+#if PLATFORM_LINUX
+#include <linux/mman.h>
+#endif
+
 namespace uba
 {
+	constexpr u64 MemoryBlock_ReserveAlign = 1 * 1024 * 1024;
+
+
 	MemoryBlock::MemoryBlock(u64 reserveSize_, void* baseAddress_)
 	{
-		Init(reserveSize_, baseAddress_);
+		Init(reserveSize_, baseAddress_, false);
 	}
 
 	MemoryBlock::MemoryBlock(u8* baseAddress_)
@@ -20,22 +27,39 @@ namespace uba
 		Deinit();
 	}
 
-	void MemoryBlock::Init(u64 reserveSize_, void* baseAddress_)
+	bool MemoryBlock::Init(u64 reserveSize_, void* baseAddress_, bool useHugePages)
 	{
-		reserveSize = AlignUp(reserveSize_, 1024 * 1024);
-
 		#if PLATFORM_WINDOWS
+		reserveSize = AlignUp(reserveSize_, MemoryBlock_ReserveAlign);
 		memory = (u8*)VirtualAlloc(baseAddress_, reserveSize, MEM_RESERVE, PAGE_READWRITE); // Max size of obj file?
 		if (!memory)
 			FatalError(1347, TC("Failed to reserve virtual memory (%u)"), GetLastError());
 		#else
-		memory = (u8*)mmap(baseAddress_, reserveSize, PROT_READ | PROT_WRITE, MAP_PRIVATE | MAP_ANONYMOUS, -1, 0);
+
+		int flags = MAP_PRIVATE | MAP_ANONYMOUS;
+		int reserveAlign = MemoryBlock_ReserveAlign;
+		#if PLATFORM_LINUX
+		if (useHugePages)
+		{
+			flags |= MAP_HUGETLB | MAP_HUGE_2MB;
+			reserveAlign = 2 * 1024 * 1024;
+		}
+		#endif
+
+		reserveSize = AlignUp(reserveSize_, reserveAlign);
+		memory = (u8*)mmap(baseAddress_, reserveSize, PROT_READ | PROT_WRITE, flags, -1, 0);
+
 		if (memory == MAP_FAILED)
-			FatalError(1347, "mmap failed to reserve %llu bytes: %s", reserveSize, strerror(errno));
+		{
+			if (useHugePages)
+				return false;
+			FatalError(1347, "mmap failed to reserve %llu bytes (asking for %llu): %s", reserveSize, reserveSize_, strerror(errno));
+		}
 		#endif
 
 		if (baseAddress_ && baseAddress_ != memory)
 			FatalError(9881, TC("Failed to reserve virtual memory at address (%u)"), GetLastError());
+		return true;
 	}
 
 	void MemoryBlock::Deinit()
@@ -72,7 +96,7 @@ namespace uba
 		#if PLATFORM_WINDOWS
 		if (newPos > mappedSize)
 		{
-			u64 toCommit = AlignUp(newPos - mappedSize, 1024 * 1024);
+			u64 toCommit = AlignUp(newPos - mappedSize, MemoryBlock_ReserveAlign);
 			if (mappedSize + toCommit > reserveSize)
 				toCommit = reserveSize - mappedSize;
 			if (!VirtualAlloc(memory + mappedSize, toCommit, MEM_COMMIT, PAGE_READWRITE))
@@ -92,7 +116,7 @@ namespace uba
 		u64 newPos = Min(writtenSize + bytes, reserveSize);
 		if (newPos <= mappedSize)
 			return;
-		u64 toCommit = AlignUp(newPos - mappedSize, 1024 * 1024);
+		u64 toCommit = AlignUp(newPos - mappedSize, MemoryBlock_ReserveAlign);
 		if (mappedSize + toCommit > reserveSize)
 			toCommit = reserveSize - mappedSize;
 		if (!VirtualAlloc(memory + mappedSize, toCommit, MEM_COMMIT, PAGE_READWRITE))
@@ -134,4 +158,19 @@ namespace uba
 		other.writtenSize = ws;
 		other.mappedSize = ms;
 	}
+
+	u64 GetHugePageCount()
+	{
+#if PLATFORM_LINUX
+		FILE* f = fopen("/proc/sys/vm/nr_hugepages", "r");
+		if (!f)
+			return 0;
+		u32 pageCount = 0;
+		fscanf(f, "%u", &pageCount);
+		fclose(f);
+		return pageCount;
+#else
+		return 0;
+#endif
+	}
 }
diff --git a/Engine/Source/Programs/UnrealBuildAccelerator/Core/Public/UbaMemory.h b/Engine/Source/Programs/UnrealBuildAccelerator/Core/Public/UbaMemory.h
--- a/Engine/Source/Programs/UnrealBuildAccelerator/Core/Public/UbaMemory.h
+++ b/Engine/Source/Programs/UnrealBuildAccelerator/Core/Public/UbaMemory.h
@@ -95,7 +95,7 @@ namespace uba
 		MemoryBlock(u64 reserveSize_, void* baseAddress_ = nullptr);
 		MemoryBlock(u8* baseAddress_ = nullptr);
 		~MemoryBlock();
-		void Init(u64 reserveSize_, void* baseAddress_ = nullptr);
+		bool Init(u64 reserveSize_, void* baseAddress_ = nullptr, bool useHugePages = false);
 		void Deinit();
 		void* Allocate(u64 bytes, u64 alignment, const tchar* hint);
 		void* AllocateNoLock(u64 bytes, u64 alignment, const tchar* hint);
@@ -273,4 +273,6 @@ namespace uba
 
 		return v;
 	}
+
+	u64 GetHugePageCount();
 }
diff --git a/Engine/Source/Programs/UnrealBuildAccelerator/Test/Private/UbaTestBasics.h b/Engine/Source/Programs/UnrealBuildAccelerator/Test/Private/UbaTestBasics.h
--- a/Engine/Source/Programs/UnrealBuildAccelerator/Test/Private/UbaTestBasics.h
+++ b/Engine/Source/Programs/UnrealBuildAccelerator/Test/Private/UbaTestBasics.h
@@ -260,10 +260,23 @@ namespace uba
 
 	bool TestMemoryBlock(Logger& logger, const StringBufferBase& rootDir)
 	{
-		MemoryBlock block(1024 * 1024);
-		u64* mem = (u64*)block.Allocate(8, 1, TC("Foo"));
-		*mem = 0x1234;
-		block.Free(mem);
+		{
+			MemoryBlock block(1024 * 1024);
+			u64* mem = (u64*)block.Allocate(8, 1, TC("Foo"));
+			*mem = 0x1234;
+			block.Free(mem);
+		}
+
+		if (GetHugePageCount())
+		{
+			MemoryBlock block;
+			if (!block.Init(1024 * 1024, nullptr, true))
+				return logger.Error(TC("Failed to allocate huge pages even though system says they exists"));
+			u64* mem = (u64*)block.Allocate(8, 1, TC("Foo"));
+			*mem = 0x1234;
+			block.Free(mem);
+		}
+
 		return true;
 	}
 
