UBA: .sarif is a possible output file for msvc

#rnx

diff --git a/Engine/Source/Programs/UnrealBuildAccelerator/Core/Private/UbaApplicationRules.cpp b/Engine/Source/Programs/UnrealBuildAccelerator/Core/Private/UbaApplicationRules.cpp
--- a/Engine/Source/Programs/UnrealBuildAccelerator/Core/Private/UbaApplicationRules.cpp
+++ b/Engine/Source/Programs/UnrealBuildAccelerator/Core/Private/UbaApplicationRules.cpp
@@ -56,6 +56,7 @@ namespace uba
 		{
 			return fileName.EndsWith(TC(".obj"))
 				|| fileName.EndsWith(TC(".dep.json"))
+				|| fileName.EndsWith(TC(".sarif"))
 				|| fileName.EndsWith(TC(".rc2.res")) // Not really an obj file.. 
 				;// || fileName.EndsWith(TC(".h.pch") // Not tested enough
 		}
@@ -69,6 +70,7 @@ namespace uba
 		virtual bool IsRarelyReadAfterWritten(const StringView& fileName) const override
 		{
 			return fileName.EndsWith(TC(".dep.json"))
+				|| fileName.EndsWith(TC(".sarif"))
 				|| fileName.EndsWith(TC(".exe"))
 				|| fileName.EndsWith(TC(".dll"));
 		}
