[UBA]
* Added support for including log lines in cache entries. Design is based on that the log lines are the same most of the time and we store a shared log line per cache key that is part of the response when requesting cache entries for a key. If entries have different log lines inside the same key we will receive those as part of the ReportUsedEntry message.
* Fixed linux arm compile error

diff --git a/Engine/Source/Programs/UnrealBuildAccelerator/CacheService/Private/UbaCacheService.cpp b/Engine/Source/Programs/UnrealBuildAccelerator/CacheService/Private/UbaCacheService.cpp
--- a/Engine/Source/Programs/UnrealBuildAccelerator/CacheService/Private/UbaCacheService.cpp
+++ b/Engine/Source/Programs/UnrealBuildAccelerator/CacheService/Private/UbaCacheService.cpp
@@ -202,12 +202,13 @@ namespace uba
 
 		u64 maintenanceReserveSizeMb = 128;
 
-		#if PLATFORM_LINUX
-		u64 hugePageCount = GetHugePageCount();
-		u64 recommendedHugePageCount = (maintenanceReserveSizeMb*GetLogicalProcessorCount())/2;
-		if (hugePageCount < recommendedHugePageCount)
-			logger.Info(TC("  Improve maintenance performance by enabling %llu huge pages on system (%llu enabled)"), recommendedHugePageCount, hugePageCount);
-		#endif
+		if (SupportsHugePages())
+		{
+			u64 hugePageCount = GetHugePageCount();
+			u64 recommendedHugePageCount = (maintenanceReserveSizeMb*GetLogicalProcessorCount())/2;
+			if (hugePageCount < recommendedHugePageCount)
+				logger.Info(TC("  Improve maintenance performance by enabling %llu huge pages on system (%llu enabled)"), recommendedHugePageCount, hugePageCount);
+		}
 
 		logger.Info(TC(""));
 
diff --git a/Engine/Source/Programs/UnrealBuildAccelerator/Cli/Private/UbaCli.cpp b/Engine/Source/Programs/UnrealBuildAccelerator/Cli/Private/UbaCli.cpp
--- a/Engine/Source/Programs/UnrealBuildAccelerator/Cli/Private/UbaCli.cpp
+++ b/Engine/Source/Programs/UnrealBuildAccelerator/Cli/Private/UbaCli.cpp
@@ -462,6 +462,7 @@ namespace uba
 
 			StringBuffer<> rootDir2(g_rootDir.data);
 			rootDir2.Append("_CHECKCAS2");
+			DeleteAllFiles(logger, rootDir2.data);
 			Client client;
 
 			auto g = MakeGuard([&]() { server.DisconnectClients(); });
@@ -480,13 +481,14 @@ namespace uba
 						{
 							Storage::RetrieveResult res;
 							storageServer.EnsureCasFile(casKey, TC("Dummy"));
-							CasKey casKey2 = AsCompressed(casKey, true);
-							if (!client.storageClient->RetrieveCasFile(res, casKey2, TC("")))
+							if (!client.storageClient->RetrieveCasFile(res, AsCompressed(casKey, false), TC("")))
+								success = false;
+							if (!client.storageClient->RetrieveCasFile(res, casKey, TC("")))
 								success = false;
 						}, 1, TC(""));
 				});
 			workManager.FlushWork();
-			return success;
+			return success ? 0 : -1;
 		}
 
 #if UBA_USE_AWS
@@ -712,9 +714,9 @@ namespace uba
 			u32 bucketId = 1337;
 			if (cacheClient)
 			{
-				bool hit = false;
-				cacheClient->FetchFromCache(hit, RootPaths(), bucketId, pinfo);
-				if (hit)
+				CacheResult cacheResult;
+				cacheClient->FetchFromCache(cacheResult, RootPaths(), bucketId, pinfo);
+				if (cacheResult.hit)
 				{
 					logger.Info(TC("%s run took %s [cached]"), (enableDetour ? TC("Boxed") : TC("Native")), TimeToText(GetTime() - start).str);
 					return true;
@@ -738,7 +740,7 @@ namespace uba
 			{
 				logger.Error(TC("Populating cache not implemented... todo"));
 				RootPaths rootPaths;
-				cacheClient->WriteToCache(rootPaths, 0, pinfo, nullptr, 1, nullptr, 0);
+				cacheClient->WriteToCache(rootPaths, 0, pinfo, nullptr, 1, nullptr, 0, nullptr, 0);
 			}
 			return true;
 		};
diff --git a/Engine/Source/Programs/UnrealBuildAccelerator/Common/Private/UbaCacheClient.cpp b/Engine/Source/Programs/UnrealBuildAccelerator/Common/Private/UbaCacheClient.cpp
--- a/Engine/Source/Programs/UnrealBuildAccelerator/Common/Private/UbaCacheClient.cpp
+++ b/Engine/Source/Programs/UnrealBuildAccelerator/Common/Private/UbaCacheClient.cpp
@@ -18,8 +18,6 @@
 #define UBA_LOG_WRITE_CACHE_INFO 0 // 0 = Disabled, 1 = Normal, 2 = Detailed
 #define UBA_LOG_FETCH_CACHE_INFO 0 // 0 = Disabled, 1 = Misses, 2 = Both misses and hits
 
-#define UBA_OLD_TEST 0
-
 namespace uba
 {
 	void CacheClientCreateInfo::Apply(Config& config)
@@ -128,7 +126,7 @@ namespace uba
 
 	CacheClient::~CacheClient() = default;
 
-	bool CacheClient::WriteToCache(const RootPaths& rootPaths, u32 bucketId, const ProcessStartInfo& info, const u8* inputs, u64 inputsSize, const u8* outputs, u64 outputsSize, u32 processId)
+	bool CacheClient::WriteToCache(const RootPaths& rootPaths, u32 bucketId, const ProcessStartInfo& info, const u8* inputs, u64 inputsSize, const u8* outputs, u64 outputsSize, const u8* logLines, u64 logLinesSize, u32 processId)
 	{
 		if (!m_connected)
 			return false;
@@ -350,7 +348,7 @@ namespace uba
 			return false;
 
 		// actual cache entry now when we know server has the needed tables
-		if (!SendCacheEntry(bucket, rootPaths, cmdKey, inputsStringToCasKey, outputsStringToCasKey, bytesSent))
+		if (!SendCacheEntry(bucket, rootPaths, cmdKey, inputsStringToCasKey, outputsStringToCasKey, logLines, logLinesSize, bytesSent))
 			return false;
 
 
@@ -386,9 +384,9 @@ namespace uba
 	}
 
 
-	bool CacheClient::FetchFromCache(bool& outCacheHit, const RootPaths& rootPaths, u32 bucketId, const ProcessStartInfo& info)
+	bool CacheClient::FetchFromCache(CacheResult& outResult, const RootPaths& rootPaths, u32 bucketId, const ProcessStartInfo& info)
 	{
-		outCacheHit = false;
+		outResult.hit = false;
 
 		if (!m_connected)
 			return false;
@@ -437,11 +435,6 @@ namespace uba
 				return false;
 		}
 
-		// Traverse entries and test inputs against local machine
-#if UBA_OLD_TEST
-		BinaryReader traverserReader(reader.GetPositionData(), 0, reader.GetLeft());
-#endif
-
 		u32 entryCount = reader.ReadU16();
 
 		#if UBA_LOG_FETCH_CACHE_INFO
@@ -530,45 +523,6 @@ namespace uba
 			};
 
 
-#if UBA_OLD_TEST
-		bool testIsMatch = false;
-		u32 matchingId = ~0u;
-		{
-			CacheEntriesTraverser traverser(traverserReader);
-			u32 entryIndex = 0;
-			for (; entryIndex!=entryCount; ++entryIndex)
-			{
-				{
-					bool isMatch = true;
-					bool result = traverser.TraverseEntryInputs([&](u32 casKeyOffset)
-						{
-							bool entryMatch;
-							if (!IsCasKeyMatch(entryMatch, casKeyOffset, entryIndex, true))
-								return false;
-							isMatch &= entryMatch;
-							return true;
-						});
-
-					if (isMatch && !result) // Returned false before setting isMatch, something went wrong
-						return false;
-
-					// No match, test next entry
-					if (!isMatch)
-					{
-						traverser.SkipEntryOutputs();
-						continue;
-					}
-				}
-
-				testIsMatch = true;
-				matchingId = traverser.lastId;
-			}
-		}
-#endif
-
-
-
-
 		struct Range
 		{
 			u32 begin;
@@ -576,6 +530,8 @@ namespace uba
 		};
 		Vector<Range> sharedMatchingRanges;
 
+		const u8* sharedLogLines;
+		u64 sharedLogLinesSize;
 
 		// Create ranges out of shared offsets that matches local state
 		{
@@ -585,6 +541,10 @@ namespace uba
 			BinaryReader sharedReader(reader.GetPositionData(), 0, sharedSize);
 			reader.Skip(sharedSize);
 
+			sharedLogLinesSize = reader.Read7BitEncoded();
+			sharedLogLines = reader.GetPositionData();
+			reader.Skip(sharedLogLinesSize);
+
 			u32 rangeBegin = 0;
 
 			auto addRange = [&](u32 rangeEnd)
@@ -645,6 +605,8 @@ namespace uba
 				u64 outSize = reader.Read7BitEncoded();
 				BinaryReader outputsReader(reader.GetPositionData(), 0, outSize);
 				reader.Skip(outSize);
+				
+				auto logLinesType = LogLinesType(reader.ReadByte());
 
 				{
 					TimerScope ts(cacheStats.testEntry);
@@ -680,22 +642,16 @@ namespace uba
 						continue;
 				}
 
-#if UBA_OLD_TEST
-				UBA_ASSERTF(testIsMatch, TC("%s"), info.GetDescription());
-				UBA_ASSERTF(matchingId == entryId, TC("%u vs %u"), matchingId, entryId);
-#endif
-
 				if (!m_useCacheHit)
 					return false;
 
-				{
-					StackBinaryWriter<128> writer;
-					NetworkMessage msg(m_client, CacheServiceId, CacheMessageType_ReportUsedEntry, writer);
-					writer.Write7BitEncoded(MakeId(bucket.id));
-					writer.WriteCasKey(cmdKey);
-					writer.Write7BitEncoded(entryId);
-					msg.Send();
-				}
+
+				if (logLinesType == LogLinesType_Shared)
+					if (!PopulateLogLines(outResult.logLines, sharedLogLines, sharedLogLinesSize))
+						return false;
+
+				if (!ReportUsedEntry(outResult.logLines, logLinesType == LogLinesType_Owned, bucket, cmdKey, entryId))
+					return false;
 
 				// Fetch output files from cache (and some files need to be "denormalized" before written to disk
 
@@ -807,16 +763,12 @@ namespace uba
 					if (!m_session.RegisterNewFile(path.data))
 						return false;
 				}
-				outCacheHit = true;
+				outResult.hit = true;
 				success = true;
 				return true;
 			}
 		}
 
-#if UBA_OLD_TEST
-		UBA_ASSERTF(!testIsMatch, TC("%s"), info.GetDescription());
-#endif
-
 		for (auto& miss : misses)
 			m_logger.Info(TC("Cache miss on %s because of mismatch of %s (entry: %u, local: %s cache: %s)"), info.GetDescription(), miss.path.data(), miss.entryIndex, CasKeyString(miss.local).str, CasKeyString(miss.cache).str);
 
@@ -937,7 +889,7 @@ namespace uba
 		return true;
 	}
 
-	bool CacheClient::SendCacheEntry(Bucket& bucket, const RootPaths& rootPaths, const CasKey& cmdKey, const Map<u32, u32>& inputsStringToCasKey, const Map<u32, u32>& outputsStringToCasKey, u64& outBytesSent)
+	bool CacheClient::SendCacheEntry(Bucket& bucket, const RootPaths& rootPaths, const CasKey& cmdKey, const Map<u32, u32>& inputsStringToCasKey, const Map<u32, u32>& outputsStringToCasKey, const u8* logLines, u64 logLinesSize, u64& outBytesSent)
 	{
 		StackBinaryReader<1024> reader;
 		{
@@ -947,6 +899,7 @@ namespace uba
 			writer.Write7BitEncoded(MakeId(bucket.id));
 			writer.WriteCasKey(cmdKey);
 
+			writer.Write7BitEncoded(inputsStringToCasKey.size());
 			writer.Write7BitEncoded(outputsStringToCasKey.size());
 			for (auto& kv : outputsStringToCasKey)
 				writer.Write7BitEncoded(kv.second);
@@ -954,6 +907,10 @@ namespace uba
 			for (auto& kv : inputsStringToCasKey)
 				writer.Write7BitEncoded(kv.second);
 
+			if (logLinesSize)
+				if (writer.GetCapacityLeft() > logLinesSize + Get7BitEncodedCount(logLinesSize))
+					writer.WriteBytes(logLines, logLinesSize);
+
 			if (!msg.Send(reader))
 				return false;
 		}
@@ -1149,6 +1106,36 @@ namespace uba
 		return true;
 	}
 
+	bool CacheClient::ReportUsedEntry(Vector<ProcessLogLine>& outLogLines, bool ownedLogLines, Bucket& bucket, const CasKey& cmdKey, u32 entryId)
+	{
+		StackBinaryWriter<128> writer;
+		NetworkMessage msg(m_client, CacheServiceId, CacheMessageType_ReportUsedEntry, writer);
+		writer.Write7BitEncoded(MakeId(bucket.id));
+		writer.WriteCasKey(cmdKey);
+		writer.Write7BitEncoded(entryId);
+
+		if (!ownedLogLines)
+			return msg.Send();
+
+		StackBinaryReader<SendMaxSize> reader;
+		if (!msg.Send(reader))
+			return false;
+
+		return PopulateLogLines(outLogLines, reader.GetPositionData(), reader.GetLeft());
+	}
+
+	bool CacheClient::PopulateLogLines(Vector<ProcessLogLine>& outLogLines, const u8* mem, u64 memLen)
+	{
+		BinaryReader reader(mem, 0, memLen);
+		while (reader.GetLeft())
+		{
+			auto& logLine = outLogLines.emplace_back();
+			logLine.text = reader.ReadString();
+			logLine.type = LogEntryType(reader.ReadByte());
+		}
+		return true;
+	}
+
 	CasKey CacheClient::GetCmdKey(const RootPaths& rootPaths, const ProcessStartInfo& info)
 	{
 		CasKeyHasher hasher;
diff --git a/Engine/Source/Programs/UnrealBuildAccelerator/Common/Private/UbaCacheEntry.cpp b/Engine/Source/Programs/UnrealBuildAccelerator/Common/Private/UbaCacheEntry.cpp
--- a/Engine/Source/Programs/UnrealBuildAccelerator/Common/Private/UbaCacheEntry.cpp
+++ b/Engine/Source/Programs/UnrealBuildAccelerator/Common/Private/UbaCacheEntry.cpp
@@ -8,17 +8,28 @@ namespace uba
 {
 	u64 CacheEntries::GetSharedSize()
 	{
-		return sizeof(u16) + Get7BitEncodedCount(sharedInputCasKeyOffsets.size()) + sharedInputCasKeyOffsets.size();
+		return sizeof(u16)
+			+ Get7BitEncodedCount(sharedInputCasKeyOffsets.size()) + sharedInputCasKeyOffsets.size()
+			+ Get7BitEncodedCount(sharedLogLines.size()) + sharedLogLines.size();
 	}
 
-	u64 CacheEntries::GetEntrySize(CacheEntry& entry, bool toDisk)
+	u64 CacheEntries::GetEntrySize(CacheEntry& entry, u32 clientVersion, bool toDisk)
 	{
 		u64 size = 0;
 
 		if (toDisk)
+		{
 			size += Get7BitEncodedCount(entry.creationTime) + Get7BitEncodedCount(entry.lastUsedTime);
+			if (clientVersion >= 5 && entry.logLinesType == LogLinesType_Owned)
+				size += Get7BitEncodedCount(entry.logLines.size()) + entry.logLines.size();
+		}
 		else
+		{
 			size += Get7BitEncodedCount(entry.id);
+		}
+
+		if (clientVersion >= 5)
+			++size; // logLinesType
 
 		auto& extra = entry.extraInputCasKeyOffsets;
 		size += Get7BitEncodedCount(extra.size()) + extra.size();
@@ -31,11 +42,11 @@ namespace uba
 		return size;
 	}
 
-	u64 CacheEntries::GetTotalSize(bool toDisk)
+	u64 CacheEntries::GetTotalSize(u32 clientVersion, bool toDisk)
 	{
 		u64 size = GetSharedSize();
 		for (auto& entry : entries)
-			size += GetEntrySize(entry, toDisk);
+			size += GetEntrySize(entry, clientVersion, toDisk);
 		return size;
 	}
 
@@ -69,16 +80,22 @@ namespace uba
 		}
 
 		{
-			auto& shared = sharedInputCasKeyOffsets;
+			auto& sharedOffsets = sharedInputCasKeyOffsets;
 
 			if (!toDisk)
 			{
-				u64 neededSize = Get7BitEncodedCount(shared.size()) + shared.size();
+				u64 neededSize = Get7BitEncodedCount(sharedOffsets.size()) + sharedOffsets.size();
+				neededSize += Get7BitEncodedCount(sharedLogLines.size()) + sharedLogLines.size();
 				if (neededSize > writer.GetCapacityLeft())
 					return true;
 			}
-			writer.Write7BitEncoded(shared.size());
-			writer.WriteBytes(shared.data(), shared.size());
+			writer.Write7BitEncoded(sharedOffsets.size());
+			writer.WriteBytes(sharedOffsets.data(), sharedOffsets.size());
+			if (clientVersion >= 5)
+			{
+				writer.Write7BitEncoded(sharedLogLines.size());
+				writer.WriteBytes(sharedLogLines.data(), sharedLogLines.size());
+			}
 		}
 
 		for (auto& entry : entries)
@@ -97,6 +114,8 @@ namespace uba
 				u64 neededSize = Get7BitEncodedCount(entry.id) + Get7BitEncodedCount(extra.size()) + extra.size();
 				neededSize += Get7BitEncodedCount(ranges.size()) + ranges.size();
 				neededSize += Get7BitEncodedCount(outputs.size()) + outputs.size();
+				if (clientVersion >= 5)
+					neededSize += 1; // hasLogLines
 
 				if (neededSize > writer.GetCapacityLeft())
 					return true;
@@ -113,6 +132,21 @@ namespace uba
 			writer.Write7BitEncoded(outputs.size());
 			writer.WriteBytes(outputs.data(), outputs.size());
 
+			// Log lines are not included in network data.
+			if (toDisk)
+			{
+				writer.WriteByte(entry.logLinesType);
+				if (entry.logLinesType == LogLinesType_Owned)
+				{
+					writer.Write7BitEncoded(entry.logLines.size());
+					writer.WriteBytes(entry.logLines.data(), entry.logLines.size());
+				}
+			}
+			else if (clientVersion >= 5)
+			{
+				writer.WriteByte(entry.logLinesType);
+			}
+
 			++entryCount;
 		}
 
@@ -134,11 +168,6 @@ namespace uba
 				cacheEntry.lastUsedTime = GetSystemTimeAsFileTime();
 
 				u32 inputSize = reader.ReadU32();
-				#if UBA_USE_OLD
-				cacheEntry.inputCasKeyOffsets.resize(inputSize);
-				reader.ReadBytes(cacheEntry.inputCasKeyOffsets.data(), inputSize);
-				reader.SetPosition(reader.GetPosition() - inputSize);
-				#endif
 
 				u64 inputEnd = reader.GetPosition() + inputSize;
 				temp.clear();
@@ -151,9 +180,6 @@ namespace uba
 				reader.ReadBytes(cacheEntry.outputCasKeyOffsets.data(), outputSize);
 			}
 
-			#if UBA_USE_OLD
-			ValidateEntries(logger);
-			#endif
 			return true;
 		}
 
@@ -162,6 +188,13 @@ namespace uba
 		sharedInputCasKeyOffsets.resize(sharedSize);
 		reader.ReadBytes(sharedInputCasKeyOffsets.data(), sharedSize);
 
+		if (databaseVersion >= 6)
+		{
+			u64 sharedLogLinesSize = reader.Read7BitEncoded();
+			sharedLogLines.resize(sharedLogLinesSize);
+			reader.ReadBytes(sharedLogLines.data(), sharedLogLinesSize);
+		}
+
 		while (entryCount--)
 		{
 			auto& entry = entries.emplace_back();
@@ -181,9 +214,17 @@ namespace uba
 			entry.outputCasKeyOffsets.resize(outputSize);
 			reader.ReadBytes(entry.outputCasKeyOffsets.data(), outputSize);
 
-			#if UBA_USE_OLD
-			Flatten(entry.inputCasKeyOffsets, entry);
-			#endif
+			if (databaseVersion >= 6)
+			{
+				entry.logLinesType = LogLinesType(reader.ReadByte());
+				if (entry.logLinesType == LogLinesType_Owned)
+				{
+					u64 logLinesSize = reader.Read7BitEncoded();
+					entry.logLines.resize(logLinesSize);
+					reader.ReadBytes(entry.logLines.data(), logLinesSize);
+				}
+			}
+
 		}
 
 		return true;
@@ -394,63 +435,11 @@ namespace uba
 		BuildInputsT(entry, inputs, entries.empty());
 	}
 
-	void CacheEntries::UpdateEntries()
-	{
-		if (entries.empty())
-			return;
-
-		// Checking if first entry is still the matching entry
-		auto& firstEntry = *entries.begin();
-		if (firstEntry.extraInputCasKeyOffsets.empty())
-		{
-			BinaryReader reader(firstEntry.sharedInputCasKeyOffsetRanges.data(), 0, firstEntry.sharedInputCasKeyOffsetRanges.size());
-			u64 begin = reader.Read7BitEncoded();
-			u64 end = reader.Read7BitEncoded();
-			if (begin == 0 && end == sharedInputCasKeyOffsets.size())
-				return;
-		}
-
-		// First entry is gone, need to refresh the shared table...
-
-
-		Vector<u8> oldShared;
-		bool isFirst = true;
-		Vector<u32> temp;
-		for (auto& entry : entries)
-		{
-			if (isFirst)
-			{
-				primaryId = entry.id;
-				// Flatten first entry into new shared
-				Flatten(oldShared, entry);
-				oldShared.swap(sharedInputCasKeyOffsets);
-				entry.extraInputCasKeyOffsets.clear();
-				entry.sharedInputCasKeyOffsetRanges.resize(1 + Get7BitEncodedCount(sharedInputCasKeyOffsets.size()));
-				BinaryWriter rangeWriter(entry.sharedInputCasKeyOffsetRanges.data(), 0, entry.sharedInputCasKeyOffsetRanges.size());
-				rangeWriter.Write7BitEncoded(0);
-				rangeWriter.Write7BitEncoded(sharedInputCasKeyOffsets.size());
-				isFirst = false;
-			}
-			else
-			{
-				// Flatten using old shared and rebuild it with new shared
-				Flatten(temp, entry, oldShared);
-				entry.extraInputCasKeyOffsets.clear();
-				entry.sharedInputCasKeyOffsetRanges.clear();
-				BuildInputsT(entry, temp, false);
-			}
-		}
-	}
-
 	void CacheEntries::UpdateEntries(Logger& logger, const GrowingNoLockUnorderedMap<u32, u32>& oldToNewCasKeyOffset, Vector<u32>& temp, Vector<u8>& temp2)
 	{
 		if (entries.empty())
 			return;
 
-		#if UBA_USE_OLD
-		ValidateEntries(logger);
-		#endif
-
 		auto convertOffsets = [&](Vector<u8>& offsets)
 			{
 				temp.clear();
@@ -475,12 +464,7 @@ namespace uba
 			};
 
 		for (auto& entry : entries)
-		{
-			#if UBA_USE_OLD
-			convertOffsets(entry.inputCasKeyOffsets);
-			#endif
 			convertOffsets(entry.outputCasKeyOffsets);
-		}
 
 		auto writePrimaryRange = [&](CacheEntry& entry, u64 newSize)
 			{
@@ -607,25 +591,9 @@ namespace uba
 					// Create new extras
 					convertOffsets(entry.extraInputCasKeyOffsets);
 				}
-
-				#if UBA_USE_OLD
-				ValidateEntry(logger, entry);
-				#endif
 			}
 		}
-
-		#if UBA_USE_OLD
-		ValidateEntries(logger);
-		#endif
-	}
-
-	#if UBA_USE_OLD
-	void CacheEntries::ValidateEntries(Logger& logger)
-	{
-		for (auto& entry : entries)
-			ValidateEntry(logger, entry, entry.inputCasKeyOffsets);
 	}
-	#endif
 
 	void CacheEntries::ValidateEntry(Logger& logger, CacheEntry& entry, Vector<u8>& inputCasKeyOffsets)
 	{
diff --git a/Engine/Source/Programs/UnrealBuildAccelerator/Common/Private/UbaCacheServer.cpp b/Engine/Source/Programs/UnrealBuildAccelerator/Common/Private/UbaCacheServer.cpp
--- a/Engine/Source/Programs/UnrealBuildAccelerator/Common/Private/UbaCacheServer.cpp
+++ b/Engine/Source/Programs/UnrealBuildAccelerator/Common/Private/UbaCacheServer.cpp
@@ -12,7 +12,7 @@
 
 namespace uba
 {
-	static constexpr u32 CacheFileVersion = 5;
+	static constexpr u32 CacheFileVersion = 6;
 	static constexpr u32 CacheFileCompatibilityVersion = 3;
 
 	bool IsCaseInsensitive(u64 id) { return (id & (1ull << 32)) == 0; }
@@ -184,7 +184,7 @@ namespace uba
 		}
 
 		u64 duration = GetTime() - startTime;
-		m_logger.Detail(TC("Database (v%u) loaded from %s in %s (%llu bucket(s) containing %s paths, %s keys, %llu cache entries)"), databaseVersion, fileName.data, TimeToText(duration).str, m_buckets.size(), BytesToText(stats.totalPathTableSize).str, BytesToText(stats.totalCasKeyTableSize).str, stats.totalCacheEntryCount.load());
+		m_logger.Detail(TC("Database loaded from %s (v%u)  in %s (%llu bucket(s) containing %s paths, %s keys, %llu cache entries)"), fileName.data, databaseVersion, TimeToText(duration).str, m_buckets.size(), BytesToText(stats.totalPathTableSize).str, BytesToText(stats.totalCasKeyTableSize).str, stats.totalCacheEntryCount.load());
 		return true;
 	}
 
@@ -351,7 +351,7 @@ namespace uba
 		{
 			file.Write(kv2.first);
 
-			temp.resize(kv2.second.GetTotalSize(true));
+			temp.resize(kv2.second.GetTotalSize(CacheNetworkVersion, true));
 			BinaryWriter writer(temp.data(), 0, temp.size());
 			kv2.second.Write(writer, CacheNetworkVersion, true);
 			UBA_ASSERT(writer.GetPosition() == temp.size());
@@ -477,7 +477,7 @@ namespace uba
 		else
 		{
 			TimeToText lastTime(startTime - m_lastMaintenance, true);
-			m_logger.Info(TC("Maintenance started after %u added cache entries (Ran last time %s ago)"), entriesAdded, (m_lastMaintenance ? lastTime.str : TC("Never")));
+			m_logger.Info(TC("Maintenance started after %u added cache entries (Ran last time %s ago)"), entriesAdded, (m_lastMaintenance ? lastTime.str : TC("<never>")));
 		}
 
 		m_lastMaintenance = startTime;
@@ -564,7 +564,7 @@ namespace uba
 				ReaderWriterLock keysToEraseLock;
 				Vector<CasKey> keysToErase;
 
-				u64 lastUseTimeLimit = ~0u;
+				u64 lastUseTimeLimit = 0; // This is the time relative to server startup time
 				if (bucket.expirationTimeSeconds)
 				{
 					// If cas key table size is over 30mb we need to shorten expiration time
@@ -608,7 +608,7 @@ namespace uba
 						auto& entry = *i;
 						bool deleteEntry = false;
 
-						u64 neededSize = entries.GetEntrySize(entry, false);
+						u64 neededSize = entries.GetEntrySize(entry, CacheNetworkVersion, false);
 						if (neededSize > capacityLeft)
 						{
 							deleteEntry = true;
@@ -722,7 +722,7 @@ namespace uba
 						keysToErase.push_back(li->first);
 					}
 					else
-						bucket.totalEntrySize += entries.GetTotalSize(false);
+						bucket.totalEntrySize += entries.GetTotalSize(CacheNetworkVersion, false);
 				});
 
 				for (auto& key : keysToErase)
@@ -965,11 +965,13 @@ namespace uba
 		lock.Leave();
 	}
 
-	CacheServer::ConnectionBucket& CacheServer::GetConnectionBucket(const ConnectionInfo& connectionInfo, BinaryReader& reader)
+	CacheServer::ConnectionBucket& CacheServer::GetConnectionBucket(const ConnectionInfo& connectionInfo, BinaryReader& reader, u32* outClientVersion)
 	{
 		u64 id = reader.Read7BitEncoded();
 		SCOPED_WRITE_LOCK(m_connectionsLock, lock);
 		auto& connection = m_connections[connectionInfo.GetId()];
+		if (outClientVersion)
+			*outClientVersion = connection.clientVersion;
 		return connection.buckets.try_emplace(id, id).first->second;
 	}
 
@@ -1037,8 +1039,9 @@ namespace uba
 		}
 		case CacheMessageType_StoreEntry:
 		{
-			auto& bucket = GetConnectionBucket(connectionInfo, reader);
-			return HandleStoreEntry(bucket, reader, writer);
+			u32 clientVersion;
+			auto& bucket = GetConnectionBucket(connectionInfo, reader, &clientVersion);
+			return HandleStoreEntry(bucket, reader, writer, clientVersion);
 		}
 		case CacheMessageType_StoreEntryDone:
 		{
@@ -1086,8 +1089,12 @@ namespace uba
 			return HandleExecuteCommand(reader, writer);
 
 		case CacheMessageType_ReportUsedEntry:
-			return HandleReportUsedEntry(reader, writer);
-
+		{
+			SCOPED_READ_LOCK(m_connectionsLock, lock);
+			u32 clientVersion = m_connections[connectionInfo.GetId()].clientVersion;
+			lock.Leave();
+			return HandleReportUsedEntry(reader, writer, clientVersion);
+		}
 		case CacheMessageType_RequestShutdown:
 		{
 			TString reason = reader.ReadString();
@@ -1102,10 +1109,15 @@ namespace uba
 		}
 	}
 
-	bool CacheServer::HandleStoreEntry(ConnectionBucket& connectionBucket, BinaryReader& reader, BinaryWriter& writer)
+	bool CacheServer::HandleStoreEntry(ConnectionBucket& connectionBucket, BinaryReader& reader, BinaryWriter& writer, u32 clientVersion)
 	{
 		CasKey cmdKey = reader.ReadCasKey();
 
+		u64 inputCount = ~0u;
+		if (clientVersion >= 5)
+			inputCount = reader.Read7BitEncoded();
+
+
 		u64 outputCount = reader.Read7BitEncoded();
 		u64 index = 0;
 
@@ -1118,8 +1130,11 @@ namespace uba
 
 		while (reader.GetLeft())
 		{
-			u32 offset = u32(reader.Read7BitEncoded());
 			bool isInput = index++ >= outputCount;
+			if (isInput && !inputCount--) // For client versions under 5 we will hit reader.GetLeft() == false first.
+				break;
+
+			u32 offset = u32(reader.Read7BitEncoded());
 			if (!isInput)
 				continue;
 
@@ -1148,6 +1163,14 @@ namespace uba
 			//m_logger.Info(TC("%s - %s"), path.data, CasKeyString(casKey).str);
 		}
 
+		// For client versions 5 and over we have log entries after the inputs
+		Vector<u8> logLines;
+		if (u64 logLinesSize = reader.GetLeft())
+		{
+			logLines.resize(logLinesSize);
+			reader.ReadBytes(logLines.data(), logLinesSize);
+		}
+
 		Vector<u8> inputCasKeyOffsets;
 		{
 			inputCasKeyOffsets.resize(bytesForInput);
@@ -1177,20 +1200,6 @@ namespace uba
 			break;
 		}
 
-		#if UBA_USE_OLD
-		List<CacheEntry>::iterator matchingEntry2 = cacheEntries.entries.end();
-		for (auto i=cacheEntries.entries.begin(), e=cacheEntries.entries.end(); i!=e; ++i)
-		{
-			if (i->inputCasKeyOffsets != inputCasKeyOffsets)
-				continue;
-			matchingEntry2 = i;
-			break;
-		}
-		if (matchingEntry2 != matchingEntry)
-			m_logger.Warning(L"CODE MISMATCH!!!");
-		#endif
-
-
 		// Already exists
 		if (matchingEntry != cacheEntries.entries.end())
 		{
@@ -1236,12 +1245,6 @@ namespace uba
 				return true;
 		}
 
-		#if UBA_USE_OLD
-		// Add new entry
-		newEntry.inputCasKeyOffsets.swap(inputCasKeyOffsets);
-		cacheEntries.ValidateEntry(m_logger, newEntry);
-		#endif
-
 		Set<u32> outputs;
 		u64 bytesForOutput = 0;
 
@@ -1282,6 +1285,28 @@ namespace uba
 		newEntry.creationTime = GetSystemTimeAsFileTime() - m_creationTime;
 		newEntry.id = cacheEntries.idCounter++;
 
+		if (logLines.empty())
+		{
+			newEntry.logLinesType = LogLinesType_Empty;
+		}
+		else if (cacheEntries.sharedLogLines.empty() && logLines.size() < 150) // If log line is very long it is most likely a warning that will be fixed
+		{
+			cacheEntries.sharedLogLines = std::move(logLines);
+			newEntry.logLinesType = LogLinesType_Shared;
+		}
+		else
+		{
+			if (cacheEntries.sharedLogLines == logLines)
+			{
+				newEntry.logLinesType = LogLinesType_Shared;
+			}
+			else
+			{
+				newEntry.logLinesType = LogLinesType_Owned;
+				newEntry.logLines = std::move(logLines);
+			}
+		}
+
 		// If cache server has all content we can put the new cache entry directly in the lookup.. otherwise we'll have to wait until client has uploaded content
 		if (hasAllContent)
 		{
@@ -1345,7 +1370,7 @@ namespace uba
 		return cacheEntries.Write(writer, clientVersion, false);
 	}
 
-	bool CacheServer::HandleReportUsedEntry(BinaryReader& reader, BinaryWriter& writer)
+	bool CacheServer::HandleReportUsedEntry(BinaryReader& reader, BinaryWriter& writer, u32 clientVersion)
 	{
 		Bucket& bucket = GetBucket(reader);
 		CasKey cmdKey = reader.ReadCasKey();
@@ -1368,6 +1393,10 @@ namespace uba
 			u64 fileTime = GetSystemTimeAsFileTime() - m_creationTime;
 			entry.lastUsedTime = fileTime;
 			bucket.lastUsedTime = fileTime;
+
+			if (clientVersion >= 5 && entry.logLinesType == LogLinesType_Owned)
+				if (entry.logLines.size() <= writer.GetCapacityLeft())
+					writer.WriteBytes(entry.logLines.data(), entry.logLines.size());
 			break;
 		}
 		return true;
diff --git a/Engine/Source/Programs/UnrealBuildAccelerator/Common/Private/UbaScheduler.cpp b/Engine/Source/Programs/UnrealBuildAccelerator/Common/Private/UbaScheduler.cpp
--- a/Engine/Source/Programs/UnrealBuildAccelerator/Common/Private/UbaScheduler.cpp
+++ b/Engine/Source/Programs/UnrealBuildAccelerator/Common/Private/UbaScheduler.cpp
@@ -454,8 +454,8 @@ namespace uba
 						ProcessStartInfo& si = *exitInfo->startInfo;
 						u64 startTime = GetTime();
 
-						bool cacheHit = false;
-						if (m_cacheClient->FetchFromCache(cacheHit, *m_rootPaths[0], 0, si) && cacheHit)
+						CacheResult cacheResult;
+						if (m_cacheClient->FetchFromCache(cacheResult, *m_rootPaths[0], 0, si) && cacheResult.hit)
 						{
 							auto process = new CachedProcess(si);
 							ProcessHandle ph(process);
diff --git a/Engine/Source/Programs/UnrealBuildAccelerator/Common/Private/UbaStorage.cpp b/Engine/Source/Programs/UnrealBuildAccelerator/Common/Private/UbaStorage.cpp
--- a/Engine/Source/Programs/UnrealBuildAccelerator/Common/Private/UbaStorage.cpp
+++ b/Engine/Source/Programs/UnrealBuildAccelerator/Common/Private/UbaStorage.cpp
@@ -1467,7 +1467,7 @@ namespace uba
 		if (logStats)
 		{
 			u64 duration = GetTime() - startTime;
-			m_logger.Detail(TC("Database loaded from %s in %s (contained %llu entries estimated to %s)"), fileName.data, TimeToText(duration).str, m_casLookup.size(), BytesToText(m_casTotalBytes).str);
+			m_logger.Detail(TC("Database loaded from %s (v%u) in %s (contained %llu entries estimated to %s)"), fileName.data, version, TimeToText(duration).str, m_casLookup.size(), BytesToText(m_casTotalBytes).str);
 		}
 
 		return true;
@@ -2121,8 +2121,9 @@ namespace uba
 
 		SCOPED_WRITE_LOCK(casEntry.lock, entryLock);
 
-		if (casEntry.verified)
-			return casEntry.exists;
+		if (casEntry.verified && casEntry.exists)
+			return true;
+
 		SCOPED_WRITE_LOCK(m_deferredCasCreationLookupLock, deferredLock);
 		auto findIt = m_deferredCasCreationLookup.find(casKey);
 		if (findIt == m_deferredCasCreationLookup.end())
diff --git a/Engine/Source/Programs/UnrealBuildAccelerator/Common/Private/UbaStorageUtils.cpp b/Engine/Source/Programs/UnrealBuildAccelerator/Common/Private/UbaStorageUtils.cpp
--- a/Engine/Source/Programs/UnrealBuildAccelerator/Common/Private/UbaStorageUtils.cpp
+++ b/Engine/Source/Programs/UnrealBuildAccelerator/Common/Private/UbaStorageUtils.cpp
@@ -150,11 +150,16 @@ namespace uba
 			if (!runInWaitFunc())
 				return false;
 
+		u32 timeOutTimeMs = 5*60*1000;
+
 		for (u32 i=0; i!=inFlightCount; ++i)
 		{
 			Entry& entry = entries[i];
-			if (!entry.done.IsSet(5*60*1000))
+			if (!entry.done.IsSet(timeOutTimeMs))
+			{
 				logger.Error(TC("SendBatchMessages timed out after 5 minutes getting async message response (%u/%u). Received %llu bytes so far. FetchId: %u"), i, inFlightCount, responseSize, fetchId);
+				timeOutTimeMs = 10;
+			}
 			if (!entry.message.ProcessAsyncResults(entry.reader))
 				success = false;
 			else
diff --git a/Engine/Source/Programs/UnrealBuildAccelerator/Common/Public/UbaCacheClient.h b/Engine/Source/Programs/UnrealBuildAccelerator/Common/Public/UbaCacheClient.h
--- a/Engine/Source/Programs/UnrealBuildAccelerator/Common/Public/UbaCacheClient.h
+++ b/Engine/Source/Programs/UnrealBuildAccelerator/Common/Public/UbaCacheClient.h
@@ -4,7 +4,7 @@
 
 #include "UbaHash.h"
 #include "UbaLogger.h"
-#include "UbaMemory.h"
+#include "UbaProcessHandle.h"
 
 namespace uba
 {
@@ -12,7 +12,6 @@ namespace uba
 	class CompactPathTable;
 	class Config;
 	class NetworkClient;
-	class ProcessHandle;
 	class RootPaths;
 	class Session;
 	class StorageImpl;
@@ -39,14 +38,20 @@ namespace uba
 		bool useCacheHit = true; // Set this to false to ignore found cache hits.. this is for debugging/testing only
 	};
 
+	struct CacheResult
+	{
+		bool hit = false;
+		Vector<ProcessLogLine> logLines;
+	};
+
 	class CacheClient
 	{
 	public:
 		CacheClient(const CacheClientCreateInfo& info);
 		~CacheClient();
 
-		bool WriteToCache(const RootPaths& rootPaths, u32 bucketId, const ProcessStartInfo& info, const u8* inputs, u64 inputsSize, const u8* outputs, u64 outputsSize, u32 processId = 0);
-		bool FetchFromCache(bool& outCacheHit, const RootPaths& rootPaths, u32 bucketId, const ProcessStartInfo& info);
+		bool WriteToCache(const RootPaths& rootPaths, u32 bucketId, const ProcessStartInfo& info, const u8* inputs, u64 inputsSize, const u8* outputs, u64 outputsSize, const u8* logLines, u64 logLinesSize, u32 processId = 0);
+		bool FetchFromCache(CacheResult& outResult, const RootPaths& rootPaths, u32 bucketId, const ProcessStartInfo& info);
 		bool RequestServerShutdown(const tchar* reason);
 
 		bool ExecuteCommand(Logger& logger, const tchar* command, const tchar* destinationFile = nullptr, const tchar* additionalInfo = nullptr);
@@ -61,8 +66,10 @@ namespace uba
 
 		bool SendPathTable(Bucket& bucket, u32 requiredPathTableSize);
 		bool SendCasTable(Bucket& bucket, u32 requiredCasTableSize);
-		bool SendCacheEntry(Bucket& bucket, const RootPaths& rootPaths, const CasKey& cmdKey, const Map<u32, u32>& inputsStringToCasKey, const Map<u32, u32>& outputsStringToCasKey, u64& outBytesSent);
+		bool SendCacheEntry(Bucket& bucket, const RootPaths& rootPaths, const CasKey& cmdKey, const Map<u32, u32>& inputsStringToCasKey, const Map<u32, u32>& outputsStringToCasKey, const u8* logLines, u64 logLinesSize, u64& outBytesSent);
 		bool FetchCasTable(Bucket& bucket, CacheStats& stats, u32 requiredCasTableOffset);
+		bool ReportUsedEntry(Vector<ProcessLogLine>& outLogLines, bool ownedLogLines, Bucket& bucket, const CasKey& cmdKey, u32 entryId);
+		bool PopulateLogLines(Vector<ProcessLogLine>& outLogLines, const u8* mem, u64 memLen);
 
 		CasKey GetCmdKey(const RootPaths& rootPaths, const ProcessStartInfo& info);
 		bool ShouldNormalize(const StringBufferBase& path);
diff --git a/Engine/Source/Programs/UnrealBuildAccelerator/Common/Public/UbaCacheEntry.h b/Engine/Source/Programs/UnrealBuildAccelerator/Common/Public/UbaCacheEntry.h
--- a/Engine/Source/Programs/UnrealBuildAccelerator/Common/Public/UbaCacheEntry.h
+++ b/Engine/Source/Programs/UnrealBuildAccelerator/Common/Public/UbaCacheEntry.h
@@ -4,23 +4,28 @@
 
 #include "UbaBinaryReaderWriter.h"
 #include "UbaLogger.h"
-
-#define UBA_USE_OLD 0
+#include "UbaProcessHandle.h"
 
 namespace uba
 {
+	enum LogLinesType : u8
+	{
+		LogLinesType_Empty,
+		LogLinesType_Shared,
+		LogLinesType_Owned,
+	};
+
 	struct CacheEntry
 	{
 		u64 creationTime = 0;
 		u64 lastUsedTime = 0;
 		u32 id = 0;
+		LogLinesType logLinesType = LogLinesType_Empty;
 
-		#if UBA_USE_OLD
-		Vector<u8> inputCasKeyOffsets;
-		#endif
 		Vector<u8> sharedInputCasKeyOffsetRanges;
 		Vector<u8> extraInputCasKeyOffsets;
 		Vector<u8> outputCasKeyOffsets;
+		Vector<u8> logLines;
 	};
 
 	struct CacheEntries
@@ -28,22 +33,18 @@ namespace uba
 		ReaderWriterLock lock;
 		List<CacheEntry> entries;
 		Vector<u8> sharedInputCasKeyOffsets;
+		Vector<u8> sharedLogLines;
 		u32 idCounter = 0;
 		u32 primaryId = ~0u; // Id of entry that shared offsets was made from
 
 		u64 GetSharedSize();
-		u64 GetEntrySize(CacheEntry& entry, bool toDisk);
-		u64 GetTotalSize(bool toDisk);
+		u64 GetEntrySize(CacheEntry& entry, u32 clientVersion, bool toDisk);
+		u64 GetTotalSize(u32 clientVersion, bool toDisk);
 		bool Write(BinaryWriter& writer, u32 clientVersion, bool toDisk);
 		bool Read(Logger& logger, BinaryReader& reader, u32 databaseVersion);
 		void BuildInputs(CacheEntry& entry, const Set<u32>& inputs);
-		void UpdateEntries();
 		void UpdateEntries(Logger& logger, const GrowingNoLockUnorderedMap<u32, u32>& oldToNewCasKeyOffset, Vector<u32>& temp, Vector<u8>& temp2);
 
-		#if UBA_USE_OLD
-		void ValidateEntries(Logger& logger);
-		#endif
-
 		void Flatten(Vector<u8>& out, const CacheEntry& entry);
 		void Flatten(Vector<u32>& out, const CacheEntry& entry, const Vector<u8>& sharedOffsets);
 
diff --git a/Engine/Source/Programs/UnrealBuildAccelerator/Common/Public/UbaCacheServer.h b/Engine/Source/Programs/UnrealBuildAccelerator/Common/Public/UbaCacheServer.h
--- a/Engine/Source/Programs/UnrealBuildAccelerator/Common/Public/UbaCacheServer.h
+++ b/Engine/Source/Programs/UnrealBuildAccelerator/Common/Public/UbaCacheServer.h
@@ -65,17 +65,17 @@ namespace uba
 		struct Connection;
 		struct ConnectionBucket;
 
-		ConnectionBucket& GetConnectionBucket(const ConnectionInfo& connectionInfo, BinaryReader& reader);
+		ConnectionBucket& GetConnectionBucket(const ConnectionInfo& connectionInfo, BinaryReader& reader, u32* outClientVersion = nullptr);
 		Bucket& GetBucket(BinaryReader& reader);
 		Bucket& GetBucket(u64 id);
 		u32 GetBucketWorkerCount();
 
 		bool HandleMessage(const ConnectionInfo& connectionInfo, u8 messageType, BinaryReader& reader, BinaryWriter& writer);
-		bool HandleStoreEntry(ConnectionBucket& bucket, BinaryReader& reader, BinaryWriter& writer);
+		bool HandleStoreEntry(ConnectionBucket& bucket, BinaryReader& reader, BinaryWriter& writer, u32 clientVersion);
 		bool HandleFetchPathTable(BinaryReader& reader, BinaryWriter& writer);
 		bool HandleFetchCasTable(BinaryReader& reader, BinaryWriter& writer);
 		bool HandleFetchEntries(BinaryReader& reader, BinaryWriter& writer, u32 clientVersion);
-		bool HandleReportUsedEntry(BinaryReader& reader, BinaryWriter& writer);
+		bool HandleReportUsedEntry(BinaryReader& reader, BinaryWriter& writer, u32 clientVersion);
 		bool HandleExecuteCommand(BinaryReader& reader, BinaryWriter& writer);
 
 		MutableLogger m_logger;
diff --git a/Engine/Source/Programs/UnrealBuildAccelerator/Common/Public/UbaNetwork.h b/Engine/Source/Programs/UnrealBuildAccelerator/Common/Public/UbaNetwork.h
--- a/Engine/Source/Programs/UnrealBuildAccelerator/Common/Public/UbaNetwork.h
+++ b/Engine/Source/Programs/UnrealBuildAccelerator/Common/Public/UbaNetwork.h
@@ -16,7 +16,7 @@ namespace uba
 	static constexpr u32 SystemNetworkVersion = 1339;
 	static constexpr u32 StorageNetworkVersion = 4;
 	static constexpr u32 SessionNetworkVersion = 34;
-	static constexpr u32 CacheNetworkVersion = 4;
+	static constexpr u32 CacheNetworkVersion = 5;
 
 	static constexpr u32 CachePathTableMaxSize = 32*1024*1024;
 	static constexpr u32 CacheCasKeyTableMaxSize = 64*1024*1024;
diff --git a/Engine/Source/Programs/UnrealBuildAccelerator/Core/Private/UbaMemory.cpp b/Engine/Source/Programs/UnrealBuildAccelerator/Core/Private/UbaMemory.cpp
--- a/Engine/Source/Programs/UnrealBuildAccelerator/Core/Private/UbaMemory.cpp
+++ b/Engine/Source/Programs/UnrealBuildAccelerator/Core/Private/UbaMemory.cpp
@@ -5,6 +5,13 @@
 
 #if PLATFORM_LINUX
 #include <linux/mman.h>
+#if defined(MAP_HUGE_2MB)
+#define UBA_SUPPORTS_HUGE_PAGES 1
+#endif
+#endif
+
+#if !defined(UBA_SUPPORTS_HUGE_PAGES)
+#define UBA_SUPPORTS_HUGE_PAGES 0
 #endif
 
 namespace uba
@@ -38,7 +45,7 @@ namespace uba
 
 		int flags = MAP_PRIVATE | MAP_ANONYMOUS;
 		int reserveAlign = MemoryBlock_ReserveAlign;
-		#if PLATFORM_LINUX
+		#if UBA_SUPPORTS_HUGE_PAGES
 		if (useHugePages)
 		{
 			flags |= MAP_HUGETLB | MAP_HUGE_2MB;
@@ -159,9 +166,13 @@ namespace uba
 		other.mappedSize = ms;
 	}
 
+	bool SupportsHugePages()
+	{
+		return UBA_SUPPORTS_HUGE_PAGES != 0;
+	}
 	u64 GetHugePageCount()
 	{
-#if PLATFORM_LINUX
+#if UBA_SUPPORTS_HUGE_PAGES
 		FILE* f = fopen("/proc/sys/vm/nr_hugepages", "r");
 		if (!f)
 			return 0;
diff --git a/Engine/Source/Programs/UnrealBuildAccelerator/Core/Public/UbaMemory.h b/Engine/Source/Programs/UnrealBuildAccelerator/Core/Public/UbaMemory.h
--- a/Engine/Source/Programs/UnrealBuildAccelerator/Core/Public/UbaMemory.h
+++ b/Engine/Source/Programs/UnrealBuildAccelerator/Core/Public/UbaMemory.h
@@ -274,5 +274,6 @@ namespace uba
 		return v;
 	}
 
+	bool SupportsHugePages();
 	u64 GetHugePageCount();
 }
diff --git a/Engine/Source/Programs/UnrealBuildAccelerator/Host/Private/UbaExports.cpp b/Engine/Source/Programs/UnrealBuildAccelerator/Host/Private/UbaExports.cpp
--- a/Engine/Source/Programs/UnrealBuildAccelerator/Host/Private/UbaExports.cpp
+++ b/Engine/Source/Programs/UnrealBuildAccelerator/Host/Private/UbaExports.cpp
@@ -603,15 +603,42 @@ extern "C"
 
 	bool CacheClient_WriteToCache(uba::CacheClient* cacheClient, uba::RootPaths* rootPaths, uba::u32 bucket, const uba::ProcessHandle* process, const uba::u8* inputs, uba::u32 inputsSize, const uba::u8* outputs, uba::u32 outputsSize)
 	{
-		return cacheClient->WriteToCache(*rootPaths, bucket, process->GetStartInfo(), inputs, inputsSize, outputs, outputsSize, process->GetId());
+		using namespace uba;
+		StackBinaryWriter<16*1024> logLinesWriter;
+
+		auto& logLines = process->GetLogLines();
+		for (auto& line : logLines)
+		{
+			if (logLinesWriter.GetCapacityLeft() < 1 + GetStringWriteSize(line.text.c_str(), line.text.size()))
+				break;
+			logLinesWriter.WriteString(line.text);
+			logLinesWriter.WriteByte(line.type);
+		}
+
+		return cacheClient->WriteToCache(*rootPaths, bucket, process->GetStartInfo(), inputs, inputsSize, outputs, outputsSize, logLinesWriter.GetData(), logLinesWriter.GetPosition(), process->GetId());
+	}
+
+	bool CacheClient_WriteToCache2(uba::CacheClient* cacheClient, uba::RootPaths* rootPaths, uba::u32 bucket, const uba::ProcessHandle* process, const uba::u8* inputs, uba::u32 inputsSize, const uba::u8* outputs, uba::u32 outputsSize, const uba::u8* logLines, uba::u32 logLinesSize)
+	{
+		return cacheClient->WriteToCache(*rootPaths, bucket, process->GetStartInfo(), inputs, inputsSize, outputs, outputsSize, logLines, logLinesSize, process->GetId());
+	}
+
+	uba::u32 CacheClient_FetchFromCache(uba::CacheClient* cacheClient, uba::RootPaths* rootPaths, uba::u32 bucket, const uba::ProcessStartInfo& info)
+	{
+		using namespace uba;
+		CacheResult cacheResult;
+		bool res = cacheClient->FetchFromCache(cacheResult, *rootPaths, bucket, info);
+		return (res && cacheResult.hit) ? 1 : 0;
 	}
 
-	bool CacheClient_FetchFromCache(uba::CacheClient* cacheClient, uba::RootPaths* rootPaths, uba::u32 bucket, const uba::ProcessStartInfo& info)
+	uba::CacheResult* CacheClient_FetchFromCache2(uba::CacheClient* cacheClient, uba::RootPaths* rootPaths, uba::u32 bucket, const uba::ProcessStartInfo& info)
 	{
 		using namespace uba;
-		bool cacheHit = false;
-		bool res = cacheClient->FetchFromCache(cacheHit, *rootPaths, bucket, info);
-		return res && cacheHit;
+		auto cacheResult = new CacheResult();
+		if (cacheClient->FetchFromCache(*cacheResult, *rootPaths, bucket, info))
+			return cacheResult;
+		delete cacheResult;
+		return nullptr;
 	}
 
 	void CacheClient_RequestServerShutdown(uba::CacheClient* cacheClient, const uba::tchar* reason)
@@ -628,6 +655,27 @@ extern "C"
 		delete &networkClient;
 	}
 
+	const uba::tchar* CacheResult_GetLogLine(uba::CacheResult* result, uba::u32 index)
+	{
+		auto& lines = result->logLines;
+		if (index >= lines.size())
+			return nullptr;
+		return lines[index].text.c_str();
+	}
+
+	uba::u32 CacheResult_GetLogLineType(uba::CacheResult* result, uba::u32 index)
+	{
+		auto& lines = result->logLines;
+		if (index >= lines.size())
+			return 0;
+		return uba::u32(lines[index].type);
+	}
+
+	void CacheResult_Delete(uba::CacheResult* result)
+	{
+		delete result;
+	}
+
 	void Uba_SetCustomAssertHandler(Uba_CustomAssertHandler* handler)
 	{
 		uba::SetCustomAssertHandler(handler);
diff --git a/Engine/Source/Programs/UnrealBuildAccelerator/Host/Public/UbaExports.h b/Engine/Source/Programs/UnrealBuildAccelerator/Host/Public/UbaExports.h
--- a/Engine/Source/Programs/UnrealBuildAccelerator/Host/Public/UbaExports.h
+++ b/Engine/Source/Programs/UnrealBuildAccelerator/Host/Public/UbaExports.h
@@ -16,6 +16,7 @@ namespace uba
 	class SessionServer;
 	class Scheduler;
 	class StorageServer;
+	struct CacheResult;
 	struct ProcessStartInfo;
 	struct SessionServerCreateInfo;
 
@@ -129,9 +130,14 @@ extern "C"
 	UBA_API uba::CacheClient* CacheClient_Create(uba::SessionServer* session, bool reportMissReason = false, const uba::tchar* crypto = nullptr);
 	UBA_API bool CacheClient_Connect(uba::CacheClient* cacheClient, const uba::tchar* host, int port);
 	UBA_API bool CacheClient_WriteToCache(uba::CacheClient* cacheClient, uba::RootPaths* rootPaths, uba::u32 bucket, const uba::ProcessHandle* process, const uba::u8* inputs, uba::u32 inputsSize, const uba::u8* outputs, uba::u32 outputsSize);
-	UBA_API bool CacheClient_FetchFromCache(uba::CacheClient* cacheClient, uba::RootPaths* rootPaths, uba::u32 bucket, const uba::ProcessStartInfo& info);
+	UBA_API bool CacheClient_WriteToCache2(uba::CacheClient* cacheClient, uba::RootPaths* rootPaths, uba::u32 bucket, const uba::ProcessHandle* process, const uba::u8* inputs, uba::u32 inputsSize, const uba::u8* outputs, uba::u32 outputsSize, const uba::u8* logLines, uba::u32 logLinesSize);
+	UBA_API uba::u32 CacheClient_FetchFromCache(uba::CacheClient* cacheClient, uba::RootPaths* rootPaths, uba::u32 bucket, const uba::ProcessStartInfo& info);
+	UBA_API uba::CacheResult* CacheClient_FetchFromCache2(uba::CacheClient* cacheClient, uba::RootPaths* rootPaths, uba::u32 bucket, const uba::ProcessStartInfo& info);
 	UBA_API void CacheClient_RequestServerShutdown(uba::CacheClient* cacheClient, const uba::tchar* reason);
 	UBA_API void CacheClient_Destroy(uba::CacheClient* cacheClient);
+	UBA_API const uba::tchar* CacheResult_GetLogLine(uba::CacheResult* result, uba::u32 index);
+	UBA_API uba::u32 CacheResult_GetLogLineType(uba::CacheResult* result, uba::u32 index);
+	UBA_API void CacheResult_Delete(uba::CacheResult* result);
 
 	// Misc
 	using Uba_CustomAssertHandler = void(const uba::tchar* text);
diff --git a/Engine/Source/Programs/UnrealBuildAccelerator/Test/Private/UbaTestAll.h b/Engine/Source/Programs/UnrealBuildAccelerator/Test/Private/UbaTestAll.h
--- a/Engine/Source/Programs/UnrealBuildAccelerator/Test/Private/UbaTestAll.h
+++ b/Engine/Source/Programs/UnrealBuildAccelerator/Test/Private/UbaTestAll.h
@@ -48,6 +48,7 @@ namespace uba
 		UBA_TEST(TestLocalSchedule) \
 		UBA_TEST(TestLocalScheduleReuse) \
 		UBA_TEST(TestRemoteScheduleReuse) \
+		UBA_TEST(TestCacheClientAndServer) \
 
 #define UBA_WINDOWS_TESTS \
 		UBA_NONMAC_TESTS \
diff --git a/Engine/Source/Programs/UnrealBuildAccelerator/Test/Private/UbaTestCache.h b/Engine/Source/Programs/UnrealBuildAccelerator/Test/Private/UbaTestCache.h
--- a/Engine/Source/Programs/UnrealBuildAccelerator/Test/Private/UbaTestCache.h
+++ b/Engine/Source/Programs/UnrealBuildAccelerator/Test/Private/UbaTestCache.h
@@ -2,7 +2,14 @@
 
 #pragma once
 
+#include "UbaCacheClient.h"
 #include "UbaCacheEntry.h"
+#include "UbaCacheServer.h"
+#include "UbaNetworkBackendTcp.h"
+#include "UbaNetworkClient.h"
+#include "UbaNetworkServer.h"
+#include "UbaSessionServer.h"
+#include "UbaStorageServer.h"
 
 namespace uba
 {
@@ -39,4 +46,164 @@ namespace uba
 
 		return true;
 	}
+
+	void GetTestAppPath(LoggerWithWriter& logger, StringBufferBase& out);
+	bool CreateTextFile(StringBufferBase& outPath, LoggerWithWriter& logger, const tchar* workingDir, const tchar* fileName, const char* text);
+	StringKey GetKeyAndFixedName(StringBuffer<>& fixedFilePath, const tchar* filePath);
+	void InvalidateCachedInfo(StorageImpl& storage, StringBufferBase& fileName)
+	{
+		StringBuffer<> fixedFilePath;
+		storage.InvalidateCachedFileInfo(GetKeyAndFixedName(fixedFilePath, fileName.data));
+	}
+
+	bool TestCacheClientAndServer(LoggerWithWriter& logger, const StringBufferBase& testRootDir)
+	{
+		LogWriter& logWriter = logger.m_writer;
+		NetworkBackendTcp tcpBackend(logWriter);
+
+		bool ctorSuccess = true;
+		NetworkServer server(ctorSuccess, { logWriter });
+
+		StringBuffer<MaxPath> rootDir;
+		rootDir.Append(testRootDir).Append(TC("Uba"));
+		if (!DeleteAllFiles(logger, rootDir.data))
+			return false;
+
+		StorageServerCreateInfo storageServerInfo(server, rootDir.data, logWriter);
+		storageServerInfo.casCapacityBytes = 1024ull * 1024 * 1024;
+		auto& storageServer = *new StorageServer(storageServerInfo);
+		auto ssg = MakeGuard([&]() { delete &storageServer; });
+
+		CacheServerCreateInfo csci(storageServer, rootDir.data, logWriter);
+		CacheServer cacheServer(csci);
+
+
+		SessionServerCreateInfo sessionInfo(storageServer, server, logWriter);
+		sessionInfo.rootDir = rootDir.data;
+		auto& session = *new SessionServer(sessionInfo);
+		auto sg = MakeGuard([&]() { delete &session; });
+
+		u16 port = 1356;
+		if (!server.StartListen(tcpBackend, port))
+			return logger.Error(TC("Failed to listen"));
+		auto disconnectServer = MakeGuard([&]() { server.DisconnectClients(); });
+
+		StringBuffer<MaxPath> workingDir;
+		workingDir.Append(testRootDir).Append(TC("WorkingDir"));
+		if (!DeleteAllFiles(logger, workingDir.data))
+			return false;
+		if (!storageServer.CreateDirectory(workingDir.data))
+			return false;
+		if (!DeleteAllFiles(logger, workingDir.data, false))
+			return false;
+
+		StringBuffer<> testApp;
+		GetTestAppPath(logger, testApp);
+
+		StringBuffer<MaxPath> inputFile;
+		if (!CreateTextFile(inputFile, logger, workingDir.data, TC("Input.txt"), "Foo"))
+			return false;
+		StringBuffer<MaxPath> outputFile;
+		if (!CreateTextFile(outputFile, logger, workingDir.data, TC("Output.txt"), "Foo"))
+			return false;
+
+		StackBinaryWriter<256> inputs;
+		inputs.WriteString(inputFile);
+
+		StackBinaryWriter<256> outputs;
+		outputs.WriteString(outputFile);
+
+		StackBinaryWriter<256> logLines;
+		logLines.WriteString(TC("Hello"));
+		logLines.WriteByte(1);
+
+		ProcessStartInfo psi;
+		psi.application = testApp.data;
+
+		{
+			NetworkClient client(ctorSuccess, { logWriter });
+			CacheClientCreateInfo ccci(logWriter, storageServer, client, session);
+			ccci.useRoots = false;
+			CacheClient cacheClient(ccci);
+
+			if (!client.Connect(tcpBackend, TC("127.0.0.1"), port))
+				return logger.Error(TC("Failed to connect"));
+			auto disconnectClient = MakeGuard([&]() { client.Disconnect(); });
+
+			{
+				CacheResult result;
+				if (cacheClient.FetchFromCache(result, RootPaths(), 0, psi) || result.hit)
+					return false;
+
+				if (!cacheClient.WriteToCache(RootPaths(), 0, psi, inputs.GetData(), inputs.GetPosition(), outputs.GetData(), outputs.GetPosition(), logLines.GetData(), logLines.GetPosition()))
+					return false;
+
+				if (!DeleteFileW(outputFile.data))
+					return false;
+				if (FileExists(logger, outputFile.data))
+					return false;
+				if (!cacheClient.FetchFromCache(result, RootPaths(), 0, psi))
+					return false;
+				if (!FileExists(logger, outputFile.data))
+					return false;
+				if (result.logLines.size() != 1)
+					return false;
+				if (result.logLines[0].text != TC("Hello"))
+					return false;
+			}
+
+			{
+				if (!DeleteFileW(inputFile.data))
+					return false;
+				if (!CreateTextFile(inputFile, logger, workingDir.data, TC("Input.txt"), "Bar"))
+					return false;
+				InvalidateCachedInfo(storageServer, inputFile);
+
+				CacheResult result;
+				if (cacheClient.FetchFromCache(result, RootPaths(), 0, psi) || result.hit)
+					return false;
+
+				if (!cacheClient.WriteToCache(RootPaths(), 0, psi, inputs.GetData(), inputs.GetPosition(), outputs.GetData(), outputs.GetPosition(), logLines.GetData(), logLines.GetPosition()))
+					return false;
+
+				if (!DeleteFileW(outputFile.data))
+					return false;
+				if (FileExists(logger, outputFile.data))
+					return false;
+				if (!cacheClient.FetchFromCache(result, RootPaths(), 0, psi))
+					return false;
+				if (!FileExists(logger, outputFile.data))
+					return false;
+				if (result.logLines.size() != 1)
+					return false;
+				if (result.logLines[0].text != TC("Hello"))
+					return false;
+			}
+		}
+
+		if (!cacheServer.RunMaintenance(true, []() { return false; }))
+			return false;
+
+		{
+			NetworkClient client(ctorSuccess, { logWriter });
+			CacheClientCreateInfo ccci(logWriter, storageServer, client, session);
+			ccci.useRoots = false;
+			CacheClient cacheClient(ccci);
+
+			if (!client.Connect(tcpBackend, TC("127.0.0.1"), port))
+				return logger.Error(TC("Failed to connect"));
+			auto disconnectClient = MakeGuard([&]() { client.Disconnect(); });
+
+			{
+				CacheResult result;
+				if (!cacheClient.FetchFromCache(result, RootPaths(), 0, psi))
+					return false;
+				if (result.logLines.size() != 1)
+					return false;
+				if (result.logLines[0].text != TC("Hello"))
+					return false;
+			}
+		}
+		return true;
+	}
 }
\ No newline at end of file
diff --git a/Engine/Source/Programs/UnrealBuildAccelerator/Test/Private/UbaTestSession.h b/Engine/Source/Programs/UnrealBuildAccelerator/Test/Private/UbaTestSession.h
--- a/Engine/Source/Programs/UnrealBuildAccelerator/Test/Private/UbaTestSession.h
+++ b/Engine/Source/Programs/UnrealBuildAccelerator/Test/Private/UbaTestSession.h
@@ -117,19 +117,25 @@ namespace uba
 		out.Append(IsWindows ? TC("UbaTestApp.exe") : TC("UbaTestApp"));
 	}
 
+	bool CreateTextFile(StringBufferBase& outPath, LoggerWithWriter& logger, const tchar* workingDir, const tchar* fileName, const char* text)
+	{
+		outPath.Clear().Append(workingDir).EnsureEndsWithSlash().Append(fileName);
+		FileAccessor fr(logger, outPath.data);
+		if (!fr.CreateWrite())
+			return false;
+		fr.Write(text, strlen(text) + 1);
+		return fr.Close();
+	}
+
 	bool RunTestApp(LoggerWithWriter& logger, SessionServer& session, const tchar* workingDir, const RunProcessFunction& runProcess)
 	{
 		StringBuffer<MaxPath> testApp;
 		GetTestAppPath(logger, testApp);
 
-		{
-			StringBuffer<MaxPath> fileR;
-			fileR.Append(workingDir).Append(TC("FileR.h"));
-			FileAccessor fr(logger, fileR.data);
-			fr.CreateWrite();
-			fr.Write("Foo", 4);
-			fr.Close();
-		}
+		StringBuffer<MaxPath> fileR;
+		if (!CreateTextFile(fileR, logger, workingDir, TC("FileR.h"), "Foo"))
+			return false;
+
 		{
 			StringBuffer<MaxPath> dir1;
 			dir1.Append(workingDir).Append(TC("Dir1"));
