[UBA]
* Improved ui for ubavisualizer for active processes
* Added windows message for when window is activated

diff --git a/Engine/Source/Programs/UnrealBuildAccelerator/Visualizer/Private/UbaTraceReader.cpp b/Engine/Source/Programs/UnrealBuildAccelerator/Visualizer/Private/UbaTraceReader.cpp
--- a/Engine/Source/Programs/UnrealBuildAccelerator/Visualizer/Private/UbaTraceReader.cpp
+++ b/Engine/Source/Programs/UnrealBuildAccelerator/Visualizer/Private/UbaTraceReader.cpp
@@ -348,21 +348,22 @@ namespace uba
 		{
 		case TraceType_SessionAdded:
 		{
-			StringBuffer<> sessionName;
+			StringBuffer<128> sessionName;
 			reader.ReadString(sessionName);
 			StringBuffer<> sessionInfo;
 			reader.ReadString(sessionInfo);
 			Guid clientUid = ReadClientId(out, reader);
 			u32 sessionIndex = reader.ReadU32();
 
-			sessionName.Append(L" (").Append(sessionInfo).Append(L")");
+			StringBuffer<> fullName;
+			fullName.Append(sessionName).Append(L" (").Append(sessionInfo).Append(L")");
 
 			// Check if we can re-use existing session (same machine was disconnected and then reconnected)
 			u32 virtualSessionIndex = sessionIndex;
 			for (u32 i = 0; i != out.sessions.size(); ++i)
 			{
 				auto& oldSession = out.sessions[i];
-				if (oldSession.name != sessionName.data)
+				if (oldSession.fullName != fullName.data)
 					continue;
 				if (oldSession.disconnectTime == ~u64(0))
 					break;
@@ -388,6 +389,7 @@ namespace uba
 
 			TraceView::Session& session = GetSession(out, sessionIndex);
 			session.name = sessionName.data;
+			session.fullName = fullName.data;
 			session.clientUid = clientUid;
 
 			++out.activeSessionCount;
diff --git a/Engine/Source/Programs/UnrealBuildAccelerator/Visualizer/Private/UbaTraceReader.h b/Engine/Source/Programs/UnrealBuildAccelerator/Visualizer/Private/UbaTraceReader.h
--- a/Engine/Source/Programs/UnrealBuildAccelerator/Visualizer/Private/UbaTraceReader.h
+++ b/Engine/Source/Programs/UnrealBuildAccelerator/Visualizer/Private/UbaTraceReader.h
@@ -82,6 +82,7 @@ namespace uba
 		struct Session
 		{
 			TString name;
+			TString fullName;
 			Guid clientUid;
 			Vector<Processor> processors;
 			Vector<SessionUpdate> updates;
diff --git a/Engine/Source/Programs/UnrealBuildAccelerator/Visualizer/Private/UbaVisualizer.cpp b/Engine/Source/Programs/UnrealBuildAccelerator/Visualizer/Private/UbaVisualizer.cpp
--- a/Engine/Source/Programs/UnrealBuildAccelerator/Visualizer/Private/UbaVisualizer.cpp
+++ b/Engine/Source/Programs/UnrealBuildAccelerator/Visualizer/Private/UbaVisualizer.cpp
@@ -211,6 +211,8 @@ namespace uba
 	,	m_config(config)
 	,	m_trace(logger)
 	{
+		memset(m_activeProcessFont, 0, sizeof(m_activeProcessFont));
+		memset(m_activeProcessCountHistory, 0, sizeof(m_activeProcessCountHistory));
 	}
 
 	Visualizer::~Visualizer()
@@ -537,15 +539,15 @@ namespace uba
 		wndClassEx.lpszClassName = TEXT("UbaVisualizer");
 		ATOM wndClassAtom = RegisterClassEx(&wndClassEx);
 
-		UpdateFont();
+		UpdateDefaultFont();
 
 		u32 defaultProcessFontHeight = 13;
 		m_zoomValue = defaultProcessFontHeight / 20.0f;
 		UpdateProcessFont();
 
 		const TCHAR* fontName = TEXT("Consolas");
-		m_popupFont = (HFONT)CreateFontW(-12, 0, 0, 0, FW_NORMAL, false, false, false, ANSI_CHARSET, OUT_DEFAULT_PRECIS, CLIP_DEFAULT_PRECIS, 0, FIXED_PITCH | FF_MODERN, fontName);
-		m_popupFontHeight = 14;
+		m_popupFont.handle = (HFONT)CreateFontW(-12, 0, 0, 0, FW_NORMAL, false, false, false, ANSI_CHARSET, OUT_DEFAULT_PRECIS, CLIP_DEFAULT_PRECIS, 0, FIXED_PITCH | FF_MODERN, fontName);
+		m_popupFont.height = 14;
 		//m_popupFont = (HFONT)GetStockObject(SYSTEM_FIXED_FONT);
 
 		DWORD windowStyle = WS_CAPTION | WS_SYSMENU | WS_THICKFRAME | WS_MAXIMIZEBOX | WS_MINIMIZEBOX | WS_CLIPCHILDREN | WS_VSCROLL | WS_HSCROLL;
@@ -619,7 +621,6 @@ namespace uba
 							SaveSettings();
 						DestroyWindow(m_hwnd);
 					}
-					DeleteObject(m_font);
 					UnregisterClass(windowClassName, hInstance);
 					m_hwnd = 0;
 					m_looping = false;
@@ -718,35 +719,37 @@ namespace uba
 		m_lastBitmap = 0;
 	}
 
-	void Visualizer::UpdateFont()
+	void Visualizer::UpdateFont(Font& font, int height, bool createUnderline)
 	{
 		//NONCLIENTMETRICS nonClientMetrics;
 		//nonClientMetrics.cbSize = sizeof(nonClientMetrics);
 		//SystemParametersInfo(SPI_GETNONCLIENTMETRICS, sizeof(nonClientMetrics), &nonClientMetrics, 0);
 		//m_font = (HFONT)CreateFontIndirect(&nonClientMetrics.lfMessageFont);
-		if (m_font)
-			DeleteObject(m_font);
-		if (m_fontUnderlined)
-			DeleteObject(m_fontUnderlined);
-		m_fontHeight = m_config.fontSize;
-		m_font = (HFONT)CreateFontW(4 - m_fontHeight, 0, 0, 0, FW_NORMAL, 0, 0, 0, ANSI_CHARSET, OUT_DEFAULT_PRECIS, CLIP_DEFAULT_PRECIS, CLEARTYPE_QUALITY, DEFAULT_PITCH, m_config.fontName.c_str());
-		m_fontUnderlined = (HFONT)CreateFontW(4 - m_fontHeight, 0, 0, 0, FW_NORMAL, 0, 1, 0, ANSI_CHARSET, OUT_DEFAULT_PRECIS, CLIP_DEFAULT_PRECIS, CLEARTYPE_QUALITY, DEFAULT_PITCH, m_config.fontName.c_str());
-		m_sessionStepY = m_fontHeight + 4;
+		font.height = height;
+
+		if (font.handle)
+			DeleteObject(font.handle);
+		if (font.handleUnderlined)
+			DeleteObject(font.handleUnderlined);
+		font.handle = (HFONT)CreateFontW(4 - height, 0, 0, 0, FW_NORMAL, 0, 0, 0, ANSI_CHARSET, OUT_DEFAULT_PRECIS, CLIP_DEFAULT_PRECIS, CLEARTYPE_QUALITY, DEFAULT_PITCH, m_config.fontName.c_str());
+		if (createUnderline)
+			font.handleUnderlined = (HFONT)CreateFontW(4 - height, 0, 0, 0, FW_NORMAL, 0, 1, 0, ANSI_CHARSET, OUT_DEFAULT_PRECIS, CLIP_DEFAULT_PRECIS, CLEARTYPE_QUALITY, DEFAULT_PITCH, m_config.fontName.c_str());
 	}
 
-	void Visualizer::UpdateProcessFont()
+	void Visualizer::UpdateDefaultFont()
 	{
-		int fontHeight = int(m_zoomValue*20.0f);
-
-		//m_processFont = m_font;
-		m_processFontHeight = fontHeight;
-		m_boxHeight = m_processFontHeight + 2;
-		m_progressRectLeft = int(5 + float(m_processFontHeight) * 1.8f);
+		UpdateFont(m_defaultFont, m_config.fontSize, true);
+		m_sessionStepY = m_defaultFont.height + 4;
 
-		if (m_processFont)
-			DeleteObject(m_processFont);
-		m_processFont = (HFONT)CreateFontW(1 - fontHeight, 0, 0, 0, FW_NORMAL, 0, 0, 0, ANSI_CHARSET, OUT_DEFAULT_PRECIS, CLIP_DEFAULT_PRECIS, CLEARTYPE_QUALITY, DEFAULT_PITCH, m_config.fontName.c_str());
+		m_timelineFont = m_defaultFont;
+	}
 
+	void Visualizer::UpdateProcessFont()
+	{
+		int fontHeight = int(m_zoomValue*20.0f);
+		UpdateFont(m_processFont, fontHeight, false);
+		m_boxHeight = m_processFont.height + 2;
+		m_progressRectLeft = int(5 + float(m_processFont.height) * 1.8f);
 		DirtyBitmaps(true);
 	}
 
@@ -754,8 +757,7 @@ namespace uba
 	{
 		m_config.fontSize += offset;
 		m_config.fontSize = Max(m_config.fontSize, 10u);
-		UpdateFont();
-		DirtyBitmaps(true);
+		UpdateDefaultFont();
 		Redraw();
 	}
 
@@ -821,9 +823,6 @@ namespace uba
 		u64 playTime = GetPlayTime();
 
 		int posY = int(m_scrollPosY);
-		float boxHeight = float(m_boxHeight);
-		int stepY = int(boxHeight) + 2;
-		int processStepY = int(boxHeight) + 2;
 		float scaleX = 50.0f*m_zoomValue*m_horizontalScaleValue;
 
 		RECT progressRect = clientRect;
@@ -831,17 +830,16 @@ namespace uba
 
 		if (m_config.showTimeline)
 		{
-			progressRect.bottom -= m_fontHeight + 10;
+			progressRect.bottom -= m_defaultFont.height + 10;
 		}
 		bool shouldDrawText = m_zoomValue > 0.25f;
 
 		SetBkMode(hdc, TRANSPARENT);
 		SetTextColor(hdc, m_textColor);
-		SelectObject(hdc, m_font);
 
 		HDC textDC = CreateCompatibleDC(hdc);
 		SetTextColor(textDC, m_textColor);
-		SelectObject(textDC, m_processFont);
+		SelectObject(textDC, m_processFont.handle);
 		SelectObject(textDC, GetStockObject(NULL_BRUSH));
 		SetBkMode(textDC, TRANSPARENT);
 		SetBkColor(hdc, m_config.DarkMode ? RGB(70, 70, 70) : RGB(180, 180, 180));
@@ -856,44 +854,52 @@ namespace uba
 
 		u64 lastStop = 0;
 
-		auto drawStatusText = [&](const StringView& text, LogEntryType type, int indent, bool moveY, bool underlined = false)
+		SetActiveFont(m_defaultFont);
+
+		auto drawStatusText = [&](const StringView& text, LogEntryType type, int posX, int endX, bool moveY, bool underlined = false)
 			{
 				RECT rect;
-				rect.left = 5 + indent*m_fontHeight;
-				rect.right = clientRect.right;
+				rect.left = posX;
+				rect.right = endX;
 				rect.top = posY;
-				rect.bottom = posY + m_fontHeight + 2;
+				rect.bottom = posY + m_activeFont.height + 2;
 				SetTextColor(hdc, type == LogEntryType_Info ? m_textColor : (type == LogEntryType_Error ? m_textErrorColor : m_textWarningColor));
 
 				if (underlined)
-					SelectObject(hdc, m_fontUnderlined);
+					SelectObject(m_activeHdc, m_activeFont.handleUnderlined);
 				ExtTextOutW(hdc, rect.left, posY, ETO_CLIPPED, &rect, text.data, text.count, NULL);
 				if (underlined)
-					SelectObject(hdc, m_font);
+					SelectObject(m_activeHdc, m_activeFont.handle);
 				if (moveY)
 					posY = rect.bottom;
 			};
 
+		auto drawIndentedText = [&](const StringView& text, LogEntryType type, int indent, bool moveY, bool underlined = false)
+			{
+				int posX = 5 + indent*m_defaultFont.height;
+				drawStatusText(text, type, posX, clientRect.right, moveY, underlined);
+			};
+
 		if (m_config.showProgress && m_traceView.progressProcessesTotal)
 		{
-			drawStatusText(ToView(L"Progress"), LogEntryType_Info, 1, false);
+			drawIndentedText(ToView(L"Progress"), LogEntryType_Info, 1, false);
 
 			float progress = float(m_traceView.progressProcessesDone) / m_traceView.progressProcessesTotal;
-			u32 width = m_fontHeight * 18;
+			u32 width = m_activeFont.height * 18;
 			RECT rect;
-			rect.left = 3 + 6*m_fontHeight;
+			rect.left = 3 + 6*m_activeFont.height;
 			rect.right = rect.left + width;
 			rect.top = posY;
-			rect.bottom = posY + m_fontHeight;
+			rect.bottom = posY + m_activeFont.height;
 			FillRect(hdc, &rect, m_processBrushes[0].inProgress);
 
 			rect.right = rect.left + int(progress*width);//durationMs/100;
 			FillRect(hdc, &rect, m_traceView.progressErrorCount ? m_processBrushes[0].error : m_processBrushes[0].success);
 
-			const tchar* remoteDisabled = m_traceView.remoteExecutionDisabled ? TC(" (Remote spawn disabled)") : TC("");
+			const tchar* remoteDisabled = m_traceView.remoteExecutionDisabled ? TC("   (Remote spawn disabled)") : TC("");
 			StringBuffer<> str;
-			str.Appendf(L"%u%% (%u/%u) %s", u32(progress*100.0f), m_traceView.progressProcessesDone, m_traceView.progressProcessesTotal, remoteDisabled);
-			drawStatusText(str, LogEntryType_Info, 6, true);
+			str.Appendf(L"%u%%    %u / %u %s", u32(progress*100.0f), m_traceView.progressProcessesDone, m_traceView.progressProcessesTotal, remoteDisabled);
+			drawIndentedText(str, LogEntryType_Info, 6, true);
 		}
 
 		if (m_config.showStatus && !m_traceView.statusMap.empty())
@@ -907,13 +913,13 @@ namespace uba
 					continue;
 				row = u32(kv.first >> 32);
 				if (lastRow != ~0u && lastRow != row)
-					posY += m_fontHeight + 2;
+					posY += m_activeFont.height + 2;
 				lastRow = row;
 				u32 column = u32(kv.first & ~0u);
-				drawStatusText(status.text, status.type, column, false, !status.link.empty());
+				drawIndentedText(status.text, status.type, column, false, !status.link.empty());
 			}
 			if (row != ~0u)
-				posY += m_fontHeight + 2;
+				posY += m_activeFont.height + 2;
 
 			SetTextColor(hdc, m_textColor);
 			posY += 3;
@@ -921,76 +927,73 @@ namespace uba
 		
 		if (m_config.showActiveProcesses && !m_trace.m_activeProcesses.empty())
 		{
-			Map<u64, TraceView::Process*> activeProcesses;
-			u32 remoteCount = 0;
-			for (auto& pair : m_trace.m_activeProcesses)
-			{
-				auto& active = pair.second;
-				TraceView::Process& process = m_trace.GetSession(m_traceView, active.sessionIndex).processors[active.processorIndex].processes[active.processIndex];
-				activeProcesses.try_emplace(~0llu - process.start, &process);
-				if (process.isRemote)
-					++remoteCount;
-			}
+			auto drawBox = [&](u64 start, u64 stop, int height, bool selected)
+				{
+					//int left = 3 + 6*m_activeFont.height;
+					//int right = rect.left + durationMs/100;
+					int posX = int(m_scrollPosX) + progressRect.left;
+					bool done = stop != ~u64(0);
+					if (!done)
+						stop = playTime;
+					int left = int(posX + TimeToS(start) * scaleX);
+					int right = int(posX + TimeToS(stop) * scaleX) - 1;
 
-			drawStatusText(ToView(L"Active"), LogEntryType_Info, 1, false);
-			StringBuffer<> str;
-			str.Appendf(L"%u", u32(activeProcesses.size()) - remoteCount);
-			if (remoteCount)
-				str.Appendf(L" local %u remote", remoteCount);
-			str.Append(L" process");
-			if (activeProcesses.size() > 1)
-				str.Append(L"es");
-			drawStatusText(str, LogEntryType_Info, 6, true);
-
-			u32 index = 0;
-			for (auto& kv : activeProcesses)
-			{
-				if (index++ == m_config.maxActiveVisible)
-					break;
-				
-				auto& process = *kv.second;
-				u32 durationMs = 0;
-				if (playTime > process.start)
-					durationMs = u32(TimeToMs(playTime - process.start));
-
-				//int left = 3 + 6*m_fontHeight;
-				//int right = rect.left + durationMs/100;
-				int posX = int(m_scrollPosX) + progressRect.left;
-				u64 stop = process.stop;
-				bool done = stop != ~u64(0);
-				if (!done)
-					stop = playTime;
-				int left = int(posX + TimeToS(process.start) * scaleX);
-				int right = int(posX + TimeToS(stop) * scaleX) - 1;
+					RECT rect;
+					rect.left = left;
+					rect.right = right;
+					rect.top = posY;
+					rect.bottom = posY + height;
+					FillRect(hdc, &rect, m_processBrushes[selected].inProgress);
+					return rect;
+				};
 
-				RECT rect;
-				rect.left = left;
-				rect.right = right;
-				rect.top = posY;
-				rect.bottom = posY + m_fontHeight;
-				FillRect(hdc, &rect, m_processBrushes[0].inProgress);
+			m_activeProcessCountHistory[m_activeProcessCountHistoryIterator++ % sizeof_array(m_activeProcessCountHistory)] = u32(m_trace.m_activeProcesses.size());
 
-				str.Clear();
-				auto firstPar = process.description.find_first_of('(');
-				if (firstPar != -1 && *process.description.rbegin() == ')')
+			PaintActiveProcesses(posY, clientRect, [&](TraceView::ProcessLocation& processLocation, u32 boxHeight, bool firstWithHeight)
 				{
-					str.Append(process.description.c_str() + firstPar + 1).Resize(str.count-1);
-					str.Append(L"  ").Append(process.description.c_str(), firstPar);
-				}
-				else
-					str.Append(process.description);
+					auto& session = m_trace.GetSession(m_traceView, processLocation.sessionIndex);
+					TraceView::Process& process = session.processors[processLocation.processorIndex].processes[processLocation.processIndex];
 
-				if (process.isRemote)
-					str.Append(L" [remote]");
-				else if (process.cacheFetch)
-					str.Append(L" [cache]");
-				if (durationMs > 1000)
-					str.Appendf(L" - %us", durationMs/1000);
-				drawStatusText(str, LogEntryType_Info, 6, true);
-			}
-			posY += 3;
+					bool selected = m_processSelected && m_processSelectedLocation == processLocation;
+
+					RECT boxRect = drawBox(process.start, process.stop, boxHeight, selected);
+
+					u32 v = boxHeight - 1;
+
+					if (v > 5)
+					{
+						u32 fontIndex = Min(v, u32(sizeof_array(m_activeProcessFont) - 1));
+						if (!m_activeProcessFont[fontIndex].handle)
+							UpdateFont(m_activeProcessFont[fontIndex], fontIndex+2, false);
+						if (firstWithHeight)
+							SetActiveFont(m_activeProcessFont[fontIndex]);
+
+						StringBuffer<> str;
+						auto firstPar = -1;//process.description.find_first_of('(');
+						if (firstPar != -1 && *process.description.rbegin() == ')')
+						{
+							str.Append(process.description.c_str() + firstPar + 1).Resize(str.count-1);
+							str.Append(L"  ").Append(process.description.c_str(), firstPar);
+						}
+						else
+							str.Append(process.description);
+
+						if (process.isRemote)
+							str.Append(L" [").Append(session.name).Append(']');
+						else if (process.cacheFetch)
+							str.Append(L" [cache]");
+						if (boxRect.left < 0)
+							str.Appendf(L"   %s", TimeToText(playTime - process.start, true).str);
+						//drawIndentedText(str, LogEntryType_Info, 1, true);
+						drawStatusText(str, LogEntryType_Info, Max(int(boxRect.left+1), 1), boxRect.right, false);
+					}
+				});
 		}
 
+		float boxHeight = float(m_boxHeight);
+		int stepY = int(boxHeight) + 2;
+		int processStepY = int(boxHeight) + 2;
+
 		TraceView::WorkRecord selectedWork;
 
 		SessionRec sortedSessions[1024];
@@ -1024,7 +1027,7 @@ namespace uba
 					LineTo(hdc, clientRect.right, posY);
 
 					StringBuffer<> text;
-					text.Append(session.name);
+					text.Append(session.fullName);
 
 					if (hasUpdates && session.disconnectTime == ~u64(0))
 					{
@@ -1164,7 +1167,7 @@ namespace uba
 				PaintDetailedStats(posY, progressRect, session, isRemote, playTime, drawText);
 			}
 
-			SelectObject(hdc, m_processFont);
+			SetActiveFont(m_processFont);
 
 			if (m_config.showProcessBars)
 			{
@@ -1186,7 +1189,7 @@ namespace uba
 
 						const int textHeight = int(barHeight);
 						const int rectBottom = posY + textHeight;
-						const int offsetY = (textHeight - m_processFontHeight + textOffsetY) / 2;
+						const int offsetY = (textHeight - m_processFont.height + textOffsetY) / 2;
 
 						processLocation.processIndex = 0;
 						int posX = int(m_scrollPosX) + progressRect.left;
@@ -1244,7 +1247,7 @@ namespace uba
 										}
 										process.bitmap = m_lastBitmap;
 										process.bitmapOffset = m_lastBitmapOffset;
-										m_lastBitmapOffset += m_processFontHeight;
+										m_lastBitmapOffset += m_processFont.height;
 									}
 									if (lastSelectedBitmap != process.bitmap)
 									{
@@ -1252,8 +1255,8 @@ namespace uba
 										lastSelectedBitmap = process.bitmap;
 									}
 
-									RECT rect2{ 0, int(process.bitmapOffset), 256, int(process.bitmapOffset) + m_processFontHeight };
-									RECT rect3{ 0, int(process.bitmapOffset), processWidth, int(process.bitmapOffset) + m_processFontHeight };
+									RECT rect2{ 0, int(process.bitmapOffset), 256, int(process.bitmapOffset) + m_processFont.height };
+									RECT rect3{ 0, int(process.bitmapOffset), processWidth, int(process.bitmapOffset) + m_processFont.height };
 									if (!done)
 										rect3.right = 256;
 
@@ -1262,7 +1265,9 @@ namespace uba
 									rect2.left += 3; // Move in text a bit
 									
 									//SetTextAlign(textDC, TA_LEFT|TA_TOP|TA_NOUPDATECP|VTA_CENTER);
-									int textY = rect2.top - 1;// - int(m_processFontHeight/9);
+									int textY = rect2.top;// - int(m_processFont.height/9);
+									if (m_processFont.height < 12)
+										textY += 1;
 
 									bool dropShadow = m_config.DarkMode;
 									if (dropShadow)
@@ -1295,7 +1300,7 @@ namespace uba
 									bitmapOffsetY -= bltOffsetY;
 									bltOffsetY = 0;
 								}
-								int height = Min(textHeight, m_processFontHeight);
+								int height = Min(textHeight, m_processFont.height);
 								if (bltOffsetY + height > textHeight)
 									height = textHeight - bltOffsetY;
 
@@ -1362,7 +1367,7 @@ namespace uba
 
 						const int textHeight = int(barHeight);
 						const int rectBottom = posY + textHeight;
-						const int offsetY = (textHeight - m_processFontHeight + textOffsetY) / 2;
+						const int offsetY = (textHeight - m_processFont.height + textOffsetY) / 2;
 
 						if (shouldDrawText)
 						{
@@ -1467,13 +1472,13 @@ namespace uba
 									}
 									SelectObject(textDC, m_lastBitmap);
 
-									RECT rect2{ 0,m_lastBitmapOffset,256, m_lastBitmapOffset + m_processFontHeight };
+									RECT rect2{ 0,m_lastBitmapOffset,256, m_lastBitmapOffset + m_processFont.height };
 
 									FillRect(textDC, &rect2, m_workBrush);
 									ExtTextOutW(textDC, rect2.left, rect2.top, ETO_CLIPPED, &rect2, work.description, int(wcslen(work.description)), NULL);
 									work.bitmap = m_lastBitmap;
 									work.bitmapOffset = m_lastBitmapOffset;
-									m_lastBitmapOffset += m_processFontHeight;
+									m_lastBitmapOffset += m_processFont.height;
 								}
 
 								if (lastSelectedBitmap != work.bitmap)
@@ -1490,7 +1495,7 @@ namespace uba
 									bitmapOffsetY -= bltOffsetY;
 									bltOffsetY = 0;
 								}
-								int height = Min(textHeight, m_processFontHeight);
+								int height = Min(textHeight, m_processFont.height);
 								if (bltOffsetY + height > textHeight)
 									height = textHeight - bltOffsetY;
 
@@ -1516,7 +1521,7 @@ namespace uba
 				}
 			}
 
-			SelectObject(hdc, m_font);
+			SetActiveFont(m_defaultFont);
 		}
 
 		SelectObject(textDC, oldBmp);
@@ -1524,6 +1529,8 @@ namespace uba
 		DeleteDC(textDC);
 
 		m_contentWidth = m_progressRectLeft + Max(0, int(TimeToS((lastStop != 0 && lastStop != ~u64(0)) ? lastStop : playTime) * scaleX));
+		
+
 		m_contentHeight = posY - int(m_scrollPosY) + stepY + 14;
 
 		float timelineSelected = m_timelineSelected;
@@ -1580,8 +1587,8 @@ namespace uba
 					b.Append('0');
 				b.AppendValue(milliseconds);
 
-				SelectObject(hdc, m_popupFont);
-				DrawTextLogger logger(m_hwnd, hdc, m_popupFontHeight, m_tooltipBackgroundBrush);
+				SetActiveFont(m_popupFont);
+				DrawTextLogger logger(m_hwnd, hdc, m_popupFont.height, m_tooltipBackgroundBrush);
 
 				logger.Info(L"%s", b.data);
 				logger.DrawAtPos(left + 4, timelineTop - 20);
@@ -1618,7 +1625,7 @@ namespace uba
 			left -= 20;
 			if (m_config.showCpuMemStats)
 			{
-				SelectObject(hdc, m_font);
+				SetActiveFont(m_defaultFont);
 
 				SetTextColor(hdc, m_cpuColor);
 				ExtTextOutW(hdc, left, top, 0, NULL, L"CPU", 3, NULL);
@@ -1630,7 +1637,7 @@ namespace uba
 			}
 			if (m_config.showNetworkStats)
 			{
-				SelectObject(hdc, m_font);
+				SetActiveFont(m_defaultFont);
 
 				SetTextColor(hdc, m_sendColor);
 				ExtTextOutW(hdc, left, top, 0, NULL, L"SND", 3, NULL);
@@ -1678,8 +1685,8 @@ namespace uba
 				duration = playTime - process.start;
 			}
 
-			SelectObject(hdc, m_popupFont);
-			DrawTextLogger logger(m_hwnd, hdc, m_popupFontHeight, m_tooltipBackgroundBrush);
+			SetActiveFont(m_popupFont);
+			DrawTextLogger logger(m_hwnd, hdc, m_popupFont.height, m_tooltipBackgroundBrush);
 
 			logger.AddTextOffset(-10); // Remove spaces in the front
 			logger.AddWidth(3);
@@ -1794,8 +1801,8 @@ namespace uba
 			else
 				duration = playTime - selectedWork.start;
 
-			SelectObject(hdc, m_popupFont);
-			DrawTextLogger logger(m_hwnd, hdc, m_popupFontHeight, m_tooltipBackgroundBrush);
+			SetActiveFont(m_popupFont);
+			DrawTextLogger logger(m_hwnd, hdc, m_popupFont.height, m_tooltipBackgroundBrush);
 
 			logger.AddSpace();
 			logger.Info(L"  %ls", selectedWork.description);
@@ -1819,7 +1826,7 @@ namespace uba
 				width = 380;
 			}
 
-			int height = int(summary.size())*m_popupFontHeight;
+			int height = int(summary.size())*m_popupFont.height;
 
 			POINT p;
 			GetCursorPos(&p);
@@ -1841,18 +1848,18 @@ namespace uba
 			FillRect(hdc, &r, m_tooltipBackgroundBrush);
 
 			r.top += 5;
-			SelectObject(hdc, m_popupFont);
+			SetActiveFont(m_popupFont);
 
 			for (auto& line : summary)
 			{
 				DrawTextW(hdc, line.c_str(), int(line.size()), &r, DT_SINGLELINE);
-				r.top += m_popupFontHeight;
+				r.top += m_popupFont.height;
 			}
 		}
 		else if (m_statsSelected)
 		{
-			SelectObject(hdc, m_popupFont);
-			DrawTextLogger logger(m_hwnd, hdc, m_popupFontHeight, m_tooltipBackgroundBrush);
+			SetActiveFont(m_popupFont);
+			DrawTextLogger logger(m_hwnd, hdc, m_popupFont.height, m_tooltipBackgroundBrush);
 			logger.AddSpace(3);
 			logger.SetColor(m_cpuColor).Info(L"  Cpu: %.1f%%", m_stats.cpuLoad * 100.0f);
 			logger.SetColor(m_memColor).Info(L"  Mem: %ls/%ls", BytesToText(m_stats.memTotal - m_stats.memAvail).str, BytesToText(m_stats.memTotal).str);
@@ -1872,8 +1879,8 @@ namespace uba
 				#undef UBA_VISUALIZER_FLAG
 			};
 
-			SelectObject(hdc, m_popupFont);
-			DrawTextLogger logger(m_hwnd, hdc, m_popupFontHeight, m_tooltipBackgroundBrush);
+			SetActiveFont(m_popupFont);
+			DrawTextLogger logger(m_hwnd, hdc, m_popupFont.height, m_tooltipBackgroundBrush);
 			logger.Info(L"%ls %ls", L"Show", tooltip[m_buttonSelected]);
 			logger.DrawAtCursor();
 		}
@@ -1886,15 +1893,15 @@ namespace uba
 				/*
 				int colWidth = 500;
 				int width = colWidth * 2;
-				int height = Min(int(clientRect.bottom), int(fetchedFiles.size() * m_popupFontHeight));
+				int height = Min(int(clientRect.bottom), int(fetchedFiles.size() * m_popupFont.height));
 
-				SelectObject(hdc, m_font);
-				DrawTextLogger logger(m_hwnd, hdc, r, m_fontHeight, m_tooltipBackgroundBrush);
+				SetActiveFont(m_defaultFont);
+				DrawTextLogger logger(m_hwnd, hdc, r, m_font.height, m_tooltipBackgroundBrush);
 				for (auto& f : fetchedFiles)
 				{
 					if (f.hint == TC("KnownInput"))
 						continue;
-					if (logger.rect.top >= r.bottom - m_fontHeight)
+					if (logger.rect.top >= r.bottom - m_font.height)
 					{
 						if (logger.rect.left + colWidth >= r.right)
 						{
@@ -1914,6 +1921,106 @@ namespace uba
 
 	u64 ConvertTime(const TraceView& view, u64 time);
 
+	void Visualizer::PaintActiveProcesses(int& posY, const RECT& clientRect, const Function<void(TraceView::ProcessLocation&, u32, bool)>& drawProcess)
+	{
+		SetActiveFont(m_processFont);
+		int startPosY = posY;
+
+		Map<u64, TraceView::ProcessLocation*> activeProcesses;
+		u32 remoteCount = 0;
+		for (auto& kv : m_trace.m_activeProcesses)
+		{
+			auto& active = kv.second;
+			TraceView::Process& process = m_trace.GetSession(m_traceView, active.sessionIndex).processors[active.processorIndex].processes[active.processIndex];
+			u64 start = process.start; //~0llu - process.start;
+			activeProcesses.try_emplace(start, &active);
+			if (process.isRemote)
+				++remoteCount;
+		}
+
+		u32 maxHeight = clientRect.bottom;
+
+		bool fillHeight = !m_config.showDetailedData && !m_config.showTitleBars && !m_config.showCpuMemStats && !m_config.showNetworkStats && !m_config.showProcessBars;
+		if (fillHeight)
+		{
+			maxHeight = clientRect.bottom - posY;
+			if (m_config.showTimeline)
+				maxHeight -= GetTimelineHeight();
+		}
+		else
+		{
+			u32 maxHeight2 = m_config.maxActiveVisible * (m_activeFont.height + 2);
+			maxHeight = Min(maxHeight, maxHeight2);
+		}
+
+		u32 maxSize = m_processFont.height;
+		u32 maxSizeMinusOne = maxSize - 1;
+
+		u32 counts[128] = { 0 };
+
+		u32 highestHistoryCount = 0;
+		for (u32 i=0; i!= sizeof_array(m_activeProcessCountHistory) - 1; ++i)
+			highestHistoryCount = Max(highestHistoryCount, m_activeProcessCountHistory[i]);
+
+		u32 activeProcessCount = highestHistoryCount;
+		counts[0] = activeProcessCount;
+		u32 totalHeight = counts[0]*2;
+		while (totalHeight < maxHeight && counts[maxSizeMinusOne] != activeProcessCount)
+		{
+			bool changed = false;
+			for (u32 i=0; i!=maxSizeMinusOne; ++i)
+			{
+				if (counts[i] && counts[i] > counts[i+1]*2 + 1)
+				{
+					--counts[i];
+					++counts[i+1];
+					++totalHeight;
+					changed = true;
+				}
+			}
+			if (!changed)
+			{
+				for (u32 j=0; j!=maxSizeMinusOne; ++j)
+				{
+					if (!counts[j])
+						continue;
+					++counts[j + 1];
+					--counts[j];
+					++totalHeight;
+					break;
+				}
+			}
+		}
+
+		auto it = activeProcesses.begin();
+		auto itEnd = activeProcesses.end();
+		int endY = int(startPosY + maxHeight);
+		for (u32 i=0;i!=maxSize;++i)
+		{
+			u32 v = maxSizeMinusOne - i;
+			u32 boxHeight = v + 1;
+
+			for (u32 j=0;j!=counts[v];++j)
+			{
+				if (it == itEnd || posY >= endY)
+					break;
+				auto& active = *it->second;
+				++it;
+				drawProcess(active, boxHeight, j == 0);
+				posY += boxHeight + 1;
+			}
+		}
+
+		if (fillHeight)
+			posY = startPosY; // To prevent vertical scrollbar
+		else if (counts[maxSizeMinusOne] != activeProcessCount)
+			posY = startPosY + maxHeight;
+		else
+			posY += 3;
+
+		SetActiveFont(m_defaultFont);
+	}
+
 	void Visualizer::PaintProcessRect(TraceView::Process& process, HDC hdc, RECT rect, const RECT& progressRect, bool selected, bool writingBitmap)
 	{
 		auto clampRect = [&](RECT& r) { r.left = Min(Max(r.left, progressRect.left), progressRect.right); r.right = Max(Min(r.right, progressRect.right), progressRect.left); };
@@ -1997,6 +2104,7 @@ namespace uba
 
 	void Visualizer::PaintTimeline(HDC hdc, const RECT& clientRect)
 	{
+		SetActiveFont(m_timelineFont);
 		int top = GetTimelineTop(clientRect);
 		float timeScale = (m_horizontalScaleValue*m_zoomValue)*50.0f;
 		float startOffset = ((m_scrollPosX/timeScale) - int(m_scrollPosX/timeScale)) * timeScale;
@@ -2054,7 +2162,7 @@ namespace uba
 				{
 					RECT textRect;
 					textRect.top = top + 8;
-					textRect.bottom = textRect.top + m_fontHeight;
+					textRect.bottom = textRect.top + m_activeFont.height;
 					textRect.right = pos + 20;
 					textRect.left = pos - 20;
 					DrawTextW(hdc, buffer.data, buffer.count, &textRect, DT_SINGLELINE | DT_CENTER);
@@ -2088,7 +2196,7 @@ namespace uba
 
 	void Visualizer::PaintDetailedStats(int& posY, const RECT& progressRect, TraceView::Session& session, bool isRemote, u64 playTime, const DrawTextFunc& drawTextFunc)
 	{
-		int stepY = m_fontHeight;
+		int stepY = m_activeFont.height;
 		int startPosY = posY;
 		int posX = progressRect.left + 5;
 		RECT textRect;
@@ -2210,15 +2318,23 @@ namespace uba
 		return playTime;
 	}
 
+	int Visualizer::GetTimelineHeight()
+	{
+		return m_timelineFont.height + 8;
+	}
+
 	int Visualizer::GetTimelineTop(const RECT& clientRect)
 	{
-		int posY = m_contentHeight - m_fontHeight - 8;
-		int maxY = int(clientRect.bottom - m_fontHeight - 8);
+		int timelineHeight = GetTimelineHeight();
+		int posY = m_contentHeight - timelineHeight;
+		int maxY = int(clientRect.bottom - timelineHeight);
 		return m_config.LockTimelineToBottom ? maxY : Min(posY, maxY);
 	}
 
 	void Visualizer::HitTest(HitTestResult& outResult, const POINT& pos)
 	{
+		SetActiveFont(m_defaultFont);
+
 		u64 playTime = GetPlayTime();
 
 		RECT clientRect;
@@ -2252,10 +2368,12 @@ namespace uba
 			}
 		}
 
+		outResult.section = 0;
+
 		u64 lastStop = 0;
 
 		if (m_config.showProgress && m_traceView.progressProcessesTotal)
-			posY += m_fontHeight + 2;
+			posY += m_activeFont.height + 2;
 
 		if (m_config.showStatus && !m_traceView.statusMap.empty())
 		{
@@ -2268,12 +2386,12 @@ namespace uba
 
 				row = u32(kv.first >> 32);
 				if (lastRow != ~0u && lastRow != row)
-					posY += m_fontHeight + 2;
+					posY += m_activeFont.height + 2;
 				lastRow = row;
 
 				if (!kv.second.link.empty())
 				{
-					if (pos.y >= posY && pos.y < posY + m_fontHeight && pos.x > 20 && pos.x < 80) // TODO: x is just hard coded to fit horde for now
+					if (pos.y >= posY && pos.y < posY + m_activeFont.height && pos.x > 20 && pos.x < 80) // TODO: x is just hard coded to fit horde for now
 					{
 						outResult.hyperLink = kv.second.link;
 						return;
@@ -2281,17 +2399,46 @@ namespace uba
 				}
 			}
 			if (row != ~0u)
-				posY += m_fontHeight + 2;
+				posY += m_activeFont.height + 2;
 			posY += 3;
 		}
 
+		if (pos.y <= posY)
+			return;
+
+		outResult.section = 1;
+
+		TraceView::ProcessLocation& outLocation = outResult.processLocation;
+
 		if (m_config.showActiveProcesses && !m_trace.m_activeProcesses.empty())
 		{
-			u32 lines = 1 + Min(u32(m_trace.m_activeProcesses.size()), m_config.maxActiveVisible);
-			posY += (m_fontHeight + 2)*lines;
-		}
+			PaintActiveProcesses(posY, clientRect, [&](TraceView::ProcessLocation& processLocation, u32 boxHeight, bool firstWithHeight)
+				{
+					if (pos.y < posY || pos.y > posY + int(boxHeight))
+						return;
+					auto& session = m_trace.GetSession(m_traceView, processLocation.sessionIndex);
+					TraceView::Process& process = session.processors[processLocation.processorIndex].processes[processLocation.processIndex];
 
-		TraceView::ProcessLocation& outLocation = outResult.processLocation;
+					int posX = int(m_scrollPosX) + progressRect.left;
+					u64 stop = process.stop;
+					bool done = stop != ~u64(0);
+					if (!done)
+						stop = playTime;
+					int left = int(posX + TimeToS(process.start) * scaleX);
+					int right = int(posX + TimeToS(stop) * scaleX) - 1;
+
+					if (pos.x >= left && pos.x <= right)
+					{
+						outLocation.sessionIndex = processLocation.sessionIndex;
+						outLocation.processorIndex = processLocation.processorIndex;
+						outLocation.processIndex = processLocation.processIndex;
+						outResult.processSelected = true;
+						return;
+					}
+				});
+			if (outResult.processSelected)
+				return;
+		}
 
 		SessionRec sortedSessions[1024];
 		Populate(sortedSessions, m_traceView, m_config.SortActiveRemoteSessions);
@@ -2516,8 +2663,9 @@ namespace uba
 		if (m_config.showTimeline && !m_traceView.sessions.empty())
 		{
 			int timelineTop = GetTimelineTop(clientRect);
-			if (pos.y >= timelineTop && pos.y < timelineTop + 40)
+			if (pos.y >= timelineTop)// && pos.y < timelineTop + 40)
 			{
+				outResult.section = 2;
 				float timeScale = (m_horizontalScaleValue * m_zoomValue)*50.0f;
 				float startOffset = -(m_scrollPosX / timeScale);
 				outResult.timelineSelected = startOffset + (pos.x - m_progressRectLeft) / timeScale;
@@ -2641,6 +2789,9 @@ namespace uba
 
 		HitTestResult res;
 		HitTest(res, pos);
+
+		m_activeSection = res.section;
+
 		if (res.processSelected == m_processSelected && res.processLocation == m_processSelectedLocation &&
 			res.sessionSelectedIndex == m_sessionSelectedIndex &&
 			res.statsSelected == m_statsSelected && memcmp(&res.stats, &m_stats, sizeof(Stats)) == 0 &&
@@ -2674,7 +2825,7 @@ namespace uba
 		si.fMask = SIF_ALL | SIF_DISABLENOSCROLL;
 		si.nMin = 0;
 		si.nMax = m_contentHeight;
-		si.nPage = rect.bottom;
+		si.nPage = int(rect.bottom);
 		si.nPos = -int(m_scrollPosY);
 		si.nTrackPos = 0;
 		SetScrollInfo(m_hwnd, SB_VERT, &si, redraw);
@@ -2685,6 +2836,13 @@ namespace uba
 		SetScrollInfo(m_hwnd, SB_HORZ, &si, redraw);
 	}
 
+	void Visualizer::SetActiveFont(const Font& font)
+	{
+		m_activeFont = font;
+		if (m_activeHdc)
+			SelectObject(m_activeHdc, font.handle);
+	}
+
 	void Visualizer::PostNewTrace(u32 replay, bool paused)
 	{
 		KillTimer(m_hwnd, 0);
@@ -2769,7 +2927,9 @@ namespace uba
 			PaintClient([&](HDC hdc, HDC memDC, RECT& rect)
 				{
 					FillRect(memDC, &rect, m_backgroundBrush);
+					m_activeHdc = memDC;
 					PaintAll(memDC, rect);
+					m_activeHdc = 0;
 					BitBlt(hdc, 0, 0, rect.right - rect.left, rect.bottom - rect.top, memDC, 0, 0, SRCCOPY);
 				});
 			break;
@@ -2782,7 +2942,7 @@ namespace uba
 			int width = LOWORD(lParam);
 			if (m_contentWidth && m_contentWidth + m_scrollPosX < width)
 				m_scrollPosX = float(Min(0, width - m_contentWidth));
-			UpdateScrollbars(false);
+			UpdateScrollbars(true);
 			break;
 		}
 		case WM_TIMER:
@@ -2838,66 +2998,79 @@ namespace uba
 
 			if (m_config.ScaleHorizontalWithScrollWheel || controlDown || shiftDown)
 			{
-				RECT r;
-				GetClientRect(hWnd, &r);
+				if (m_activeSection >= 1)
+				{
+					RECT r;
+					GetClientRect(hWnd, &r);
 
-				// Use mouse cursor as scroll anchor point
-				POINT cursorPos = {};
-				GetCursorPos(&cursorPos);
-				ScreenToClient(m_hwnd, &cursorPos);
+					// Use mouse cursor as scroll anchor point
+					POINT cursorPos = {};
+					GetCursorPos(&cursorPos);
+					ScreenToClient(m_hwnd, &cursorPos);
 
-				float newZoomValue = m_zoomValue;
-				float newScaleValue = m_horizontalScaleValue;
+					float newZoomValue = m_zoomValue;
+					float newScaleValue = m_horizontalScaleValue;
 
-				if (controlDown)
-				{
-					if (delta < 0)
+					if (controlDown)
 					{
-						newZoomValue = (m_processFontHeight - 1) / 20.0f;
-						if (newZoomValue == 0)
-							newZoomValue = m_zoomValue;
+						if (delta < 0)
+						{
+							newZoomValue = (m_processFont.height - 1) / 20.0f;
+							if (newZoomValue == 0)
+								newZoomValue = m_zoomValue;
+						}
+						else if (delta > 0)
+							newZoomValue = (m_processFont.height + 1) / 20.0f;
 					}
-					else if (delta > 0)
-						newZoomValue = (m_processFontHeight + 1) / 20.0f;
-				}
-				else
-					newScaleValue = m_horizontalScaleValue + m_horizontalScaleValue*float(delta)*0.0006f;
+					else
+						newScaleValue = m_horizontalScaleValue + m_horizontalScaleValue*float(delta)*0.0006f;
 
-				// TODO: m_progressRectLeft changes with zoom so anchor logic is wrong
-				const float scrollAnchorOffsetX = float(cursorPos.x) - m_progressRectLeft;
-				const float scrollAnchorOffsetY = 0;//float(cursorPos.y)*m_zoomValue;// - m_progressRectLeft;
+					// TODO: m_progressRectLeft changes with zoom so anchor logic is wrong
+					const float scrollAnchorOffsetX = float(cursorPos.x) - m_progressRectLeft;
+					const float scrollAnchorOffsetY = 0;//float(cursorPos.y)*m_zoomValue;// - m_progressRectLeft;
 
-				m_scrollPosY = Min(0.0f, float(m_scrollPosY - scrollAnchorOffsetY)*(newZoomValue/m_zoomValue) + scrollAnchorOffsetY);
-				m_scrollPosX = Min(0.0f, float(m_scrollPosX - scrollAnchorOffsetX)*(newZoomValue/m_zoomValue)*(newScaleValue/m_horizontalScaleValue) + scrollAnchorOffsetX);//LOWORD(lParam);
+					m_scrollPosY = Min(0.0f, float(m_scrollPosY - scrollAnchorOffsetY)*(newZoomValue/m_zoomValue) + scrollAnchorOffsetY);
+					m_scrollPosX = Min(0.0f, float(m_scrollPosX - scrollAnchorOffsetX)*(newZoomValue/m_zoomValue)*(newScaleValue/m_horizontalScaleValue) + scrollAnchorOffsetX);//LOWORD(lParam);
 
-				if (m_horizontalScaleValue != newScaleValue)
-					m_horizontalScaleValue = newScaleValue;
+					if (m_horizontalScaleValue != newScaleValue)
+						m_horizontalScaleValue = newScaleValue;
 
-				if (m_zoomValue != newZoomValue)
-				{
-					m_zoomValue = newZoomValue;
-					UpdateProcessFont();
-				}
+					if (m_zoomValue != newZoomValue)
+					{
+						m_zoomValue = newZoomValue;
+						UpdateProcessFont();
+					}
 
 
-				UpdateAutoscroll();
-				UpdateSelection();
+					UpdateAutoscroll();
+					UpdateSelection();
 
-				int minScroll = r.right - m_contentWidth;
-				m_scrollPosX = Min(0.0f, Max(m_scrollPosX, float(minScroll)));
-				m_scrollPosY = Min(0.0f, Max(m_scrollPosY, float(r.bottom - m_contentHeight)));
+					int minScroll = r.right - m_contentWidth;
+					m_scrollPosX = Min(0.0f, Max(m_scrollPosX, float(minScroll)));
+					m_scrollPosY = Min(0.0f, Max(m_scrollPosY, float(r.bottom - m_contentHeight)));
 
-				//if (!m_traceView.finished && m_scrollPosX <= minScroll)
-				//	m_autoScroll = true;
+					//if (!m_traceView.finished && m_scrollPosX <= minScroll)
+					//	m_autoScroll = true;
 
 		
-				if (m_config.ShowReadWriteColors)
-					for (auto& session : m_traceView.sessions)
-						for (auto& processor : session.processors)
-							for (auto& process : processor.processes)
-								//if (TimeToMs(process.writeFilesTime, m_traceView.frequency) >= 300 || TimeToMs(process.createFilesTime, m_traceView.frequency) >= 300)
-									process.bitmapDirty = true;
-
+					if (m_config.ShowReadWriteColors)
+						for (auto& session : m_traceView.sessions)
+							for (auto& processor : session.processors)
+								for (auto& process : processor.processes)
+									//if (TimeToMs(process.writeFilesTime, m_traceView.frequency) >= 300 || TimeToMs(process.createFilesTime, m_traceView.frequency) >= 300)
+										process.bitmapDirty = true;
+				}
+				else if (m_activeSection == 0 || m_activeSection == 2)
+				{
+					if (controlDown)
+					{
+						if (delta < 0)
+							m_config.fontSize -= 1;
+						else if (delta > 0)
+							m_config.fontSize += 1;
+						UpdateDefaultFont();
+					}
+				}
 				UpdateScrollbars(true);
 				RedrawWindow(hWnd, NULL, NULL, RDW_INVALIDATE);
 			}
@@ -2918,6 +3091,8 @@ namespace uba
 		}
 		case WM_MOUSEMOVE:
 		{
+			//m_logger.Info(TC("Section: %u"), m_activeSection);
+
 			POINTS p = MAKEPOINTS(lParam);
 			POINT pos{ p.x, p.y };
 			if (m_dragToScrollCounter > 0)
@@ -2994,9 +3169,19 @@ namespace uba
 
 		case WM_MOUSEACTIVATE:
 		{
-			if (LOWORD(lParam) == HTCLIENT)
-				return MA_ACTIVATEANDEAT;
-			break;
+			if (LOWORD(lParam) != HTCLIENT)
+				break;
+
+			if (m_parentHwnd)
+			{
+				//PostMessage(wParam, WM_MOUSEACTIVATE, wParam, lParam);
+				//SetForegroundWindow(m_parentHwnd);
+				//SetFocus(m_parentHwnd);
+				//SetWindowPos((HWND)wParam, HWND_TOP, 0, 0, 0, 0, SWP_ASYNCWINDOWPOS | SWP_NOMOVE | SWP_NOSIZE);
+				PostMessage(m_parentHwnd, 0x0445, 0, 0);
+				//return MA_NOACTIVATEANDEAT;
+			}
+			return MA_ACTIVATEANDEAT;
 		}
 
 		case WM_LBUTTONDOWN:
@@ -3016,6 +3201,10 @@ namespace uba
 				{
 					float timelineSelected = Max(m_timelineSelected, 0.0f);
 					Reset();
+
+					bool changed;
+					u64 time = MsToTime(u64(timelineSelected * 1000.0));
+
 					if (!m_fileName.IsEmpty())
 					{
 						if (!m_trace.ReadFile(m_traceView, m_fileName.data, true))
@@ -3025,14 +3214,10 @@ namespace uba
 					{
 						if (!m_trace.StartReadNamed(m_traceView, nullptr, true, true))
 							return false;
+						if (m_traceView.realStartTime + time > m_startTime)
+							time = m_startTime - m_traceView.realStartTime;
 					}
 
-					bool changed;
-					u64 time = MsToTime(u64(timelineSelected * 1000.0));
-
-					if (m_traceView.realStartTime + time > m_startTime)
-						time = m_startTime - m_traceView.realStartTime;
-
 					m_traceView.finished = false;
 					if (!m_fileName.IsEmpty())
 						m_trace.UpdateReadFile(m_traceView, time, changed);
@@ -3264,7 +3449,7 @@ namespace uba
 			{
 				TString str;
 				auto& session = m_traceView.sessions[m_sessionSelectedIndex];
-				str.append(session.name).append(TC("\n"));
+				str.append(session.fullName).append(TC("\n"));
 				for (auto& line : session.summary)
 					str.append(line).append(TC("\n"));
 				CopyTextToClipboard(str);
@@ -3314,9 +3499,9 @@ namespace uba
 			if (wParam == VK_SPACE)
 				Pause(!m_paused);
 			if (wParam == VK_ADD)
-				ChangeFontSize(1);
+				++m_replay;
 			if (wParam == VK_SUBTRACT)
-				ChangeFontSize(-1);
+				--m_replay;
 			break;
 		}
 		case WM_VSCROLL:
diff --git a/Engine/Source/Programs/UnrealBuildAccelerator/Visualizer/Private/UbaVisualizer.h b/Engine/Source/Programs/UnrealBuildAccelerator/Visualizer/Private/UbaVisualizer.h
--- a/Engine/Source/Programs/UnrealBuildAccelerator/Visualizer/Private/UbaVisualizer.h
+++ b/Engine/Source/Programs/UnrealBuildAccelerator/Visualizer/Private/UbaVisualizer.h
@@ -32,7 +32,7 @@ namespace uba
 		UBA_VISUALIZER_FLAG(AutoSaveSettings, true, L"Auto save Position/Settings on close") \
 		UBA_VISUALIZER_FLAG(ShowAllTraces, true, L"Show all traces started on channel") \
 		UBA_VISUALIZER_FLAG(SortActiveRemoteSessions, true, L"Sort active sessions on top") \
-		UBA_VISUALIZER_FLAG(AutoScaleHorizontal, true, L"Automatically scale horizontal to fit processes") \
+		UBA_VISUALIZER_FLAG(AutoScaleHorizontal, true, L"Automatically scale horizontally to fit processes") \
 		UBA_VISUALIZER_FLAG(LockTimelineToBottom, true, L"Lock timeline to always paint at bottom") \
 
 	struct VisualizerConfig
@@ -84,6 +84,7 @@ namespace uba
 		void Reset();
 		void PaintClient(const Function<void(HDC hdc, HDC memDC, RECT& clientRect)>& paintFunc);
 		void PaintAll(HDC hdc, const RECT& clientRect);
+		void PaintActiveProcesses(int& posY, const RECT& clientRect, const Function<void(TraceView::ProcessLocation&, u32, bool)>& drawProcess);
 		void PaintProcessRect(TraceView::Process& process, HDC hdc, RECT rect, const RECT& progressRect, bool selected, bool writingBitmap);
 		void PaintTimeline(HDC hdc, const RECT& clientRect);
 		using DrawTextFunc = Function<void(const StringBufferBase& text, RECT& rect)>;
@@ -100,6 +101,7 @@ namespace uba
 		};
 		struct HitTestResult
 		{
+			u32 section = ~0u;
 			TraceView::ProcessLocation processLocation;
 			bool processSelected = false;
 			u32 sessionSelectedIndex = ~0u;
@@ -114,6 +116,7 @@ namespace uba
 			TString hyperLink;
 		};
 		u64 GetPlayTime();
+		int GetTimelineHeight();
 		int GetTimelineTop(const RECT& clientRect);
 		void HitTest(HitTestResult& outResult, const POINT& pos);
 
@@ -131,10 +134,20 @@ namespace uba
 		void StopDragToScroll();
 		void SaveSettings();
 		void DirtyBitmaps(bool full);
-		void UpdateFont();
+
+		struct Font
+		{
+			HFONT handle = 0;
+			HFONT handleUnderlined = 0;
+			int height = 0;
+		};
+
+		void UpdateFont(Font& font, int height, bool createUnderline);
+		void UpdateDefaultFont();
 		void UpdateProcessFont();
 		void ChangeFontSize(int offset);
 		void Redraw();
+		void SetActiveFont(const Font& font);
 
 		StringBuffer<256> m_namedTrace;
 		StringBuffer<256> m_fileName;
@@ -176,16 +189,22 @@ namespace uba
 		HPEN m_memPen = 0;
 		HPEN m_processUpdatePen = 0;
 		HPEN m_checkboxPen = 0;
-		HFONT m_font = 0;
-		HFONT m_fontUnderlined = 0;
-		HFONT m_processFont = 0;
-		int m_progressRectLeft = 30;
-		int m_processFontHeight = 0;
-		int m_fontHeight = 0;
 		int m_boxHeight = 0;
 		int m_sessionStepY = 0;
-		HFONT m_popupFont = 0;
-		int m_popupFontHeight = 0;
+
+		Font m_defaultFont;
+		Font m_processFont;
+		Font m_timelineFont;
+		Font m_popupFont;
+
+		Font m_activeProcessFont[32];
+		u32 m_activeProcessCountHistory[5];
+		u32 m_activeProcessCountHistoryIterator = 0;
+
+		HDC m_activeHdc = 0;
+		Font m_activeFont;
+
+		int m_progressRectLeft = 30;
 
 		Logger& m_logger;
 		VisualizerConfig m_config;
@@ -214,6 +233,7 @@ namespace uba
 		HBITMAP m_lastBitmap = 0;
 		int m_lastBitmapOffset = BitmapCacheHeight;
 
+		u32 m_activeSection = ~0u;
 		TraceView::ProcessLocation m_processSelectedLocation;
 		bool m_processSelected = false;
 		u32 m_sessionSelectedIndex = ~0u;
