[UBA]
* Improved some logging to help find reason for error "... that is already being written"

diff --git a/Engine/Source/Programs/UnrealBuildAccelerator/Common/Private/UbaNetworkServer.cpp b/Engine/Source/Programs/UnrealBuildAccelerator/Common/Private/UbaNetworkServer.cpp
--- a/Engine/Source/Programs/UnrealBuildAccelerator/Common/Private/UbaNetworkServer.cpp
+++ b/Engine/Source/Programs/UnrealBuildAccelerator/Common/Private/UbaNetworkServer.cpp
@@ -395,7 +395,7 @@ namespace uba
 				SCOPED_READ_LOCK(m_server.m_onDisconnectFunctionsLock, l);
 				for (auto& entry : m_server.m_onDisconnectFunctions)
 					entry.function(m_client->uid, m_client->id);
-				m_server.m_logger.Detail(TC("Client %s disconnected"), GuidToString(m_client->uid).str);
+				m_server.m_logger.Detail(TC("Client %s (Id: %u) disconnected"), GuidToString(m_client->uid).str, m_client->id);
 			}
 			m_disconnected = true;
 		}
diff --git a/Engine/Source/Programs/UnrealBuildAccelerator/Common/Private/UbaStorageServer.cpp b/Engine/Source/Programs/UnrealBuildAccelerator/Common/Private/UbaStorageServer.cpp
--- a/Engine/Source/Programs/UnrealBuildAccelerator/Common/Private/UbaStorageServer.cpp
+++ b/Engine/Source/Programs/UnrealBuildAccelerator/Common/Private/UbaStorageServer.cpp
@@ -283,22 +283,32 @@ namespace uba
 		return false;
 	}
 
-	bool StorageServer::WaitForWritten(CasEntry& casEntry, ScopedWriteLock& entryLock, const tchar* hint)
+	bool StorageServer::WaitForWritten(CasEntry& casEntry, ScopedWriteLock& entryLock, const ConnectionInfo& connectionInfo, const tchar* hint)
 	{
 		int waitCount = 0;
 		while (true)
 		{
 			if (!casEntry.beingWritten)
 				return true;
+			CasKey key = casEntry.key;
 			entryLock.Leave();
 			Sleep(100);
 			entryLock.Enter();
 
-			if (++waitCount == 10000)
+			if (++waitCount < 12*60*10)
+				continue;
+
+			// Something went wrong.. should not take 12 minutes to write a file
+
+			SCOPED_READ_LOCK(m_activeStoresLock, activeLock);
+			for (auto& kv : m_activeStores)
 			{
-				m_logger.Error(TC("Got store for file %s that is already being written. Waited 1000 seconds for it to finish without success."), hint);
-				return false;
+				if (kv.second.casEntry != &casEntry)
+					continue;
+				ActiveStore& as = kv.second;
+				return m_logger.Error(TC("Client %u waited more than 12 minutes for file %s (%s) to be written by client %u (Written %llu/%llu)"), connectionInfo.GetId(), CasKeyString(key).str, hint, as.clientId, as.totalWritten.load(), as.fileSize);
 			}
+			return m_logger.Error(TC("Client %u waited more than 12 minutes for file %s (%s) to be written but there are no active writes. This should not be possible!"), connectionInfo.GetId(), CasKeyString(key).str, hint);
 		}
 	}
 
@@ -713,7 +723,7 @@ namespace uba
 
 				SCOPED_WRITE_LOCK(casEntry.lock, entryLock);
 				
-				if (!WaitForWritten(casEntry, entryLock, TC("UNKNOWN")))
+				if (!WaitForWritten(casEntry, entryLock, connectionInfo, TC("UNKNOWN")))
 					return false;
 
 				bool exists = casEntry.verified && casEntry.exists;
@@ -733,6 +743,7 @@ namespace uba
 							if (!uba::DeleteFileW(casFile.data))
 								return m_logger.Error(TC("Failed to delete %s. Clean cas folder and restart"), casFile.data);
 							casEntry.exists = false;
+							casEntry.verified = true;
 						}
 						else
 						{
@@ -745,8 +756,8 @@ namespace uba
 					else
 					{
 						casEntry.exists = false;
+						casEntry.verified = true;
 					}
-					casEntry.verified = true;
 #endif
 				}
 				writer.WriteBool(exists);
@@ -776,7 +787,7 @@ namespace uba
 				}
 				else
 				{
-					if (!WaitForWritten(casEntry, entryLock, hint.data))
+					if (!WaitForWritten(casEntry, entryLock, connectionInfo, hint.data))
 						return false;
 
 					if (casEntry.exists)
@@ -958,7 +969,6 @@ namespace uba
 					s.casEntry = firstStore->casEntry;
 					s.totalWritten = firstStore->totalWritten.load();
 					s.recvCasTime = firstStore->recvCasTime.load();
-					s.error = firstStore->error;
 				}
 				return true;
 			}
diff --git a/Engine/Source/Programs/UnrealBuildAccelerator/Common/Public/UbaStorageServer.h b/Engine/Source/Programs/UnrealBuildAccelerator/Common/Public/UbaStorageServer.h
--- a/Engine/Source/Programs/UnrealBuildAccelerator/Common/Public/UbaStorageServer.h
+++ b/Engine/Source/Programs/UnrealBuildAccelerator/Common/Public/UbaStorageServer.h
@@ -46,7 +46,7 @@ namespace uba
 
 		void OnDisconnected(u32 clientId);
 		bool HandleMessage(const ConnectionInfo& connectionInfo, u8 messageType, BinaryReader& reader, BinaryWriter& writer);
-		bool WaitForWritten(CasEntry& casEntry, ScopedWriteLock& entryLock, const tchar* hint);
+		bool WaitForWritten(CasEntry& casEntry, ScopedWriteLock& entryLock, const ConnectionInfo& connectionInfo, const tchar* hint);
 
 		u16 PopId();
 		void PushId(u16 id);
@@ -82,7 +82,6 @@ namespace uba
 			Atomic<u64> recvCasTime;
 			u64 fileSize = 0;
 			u64 actualSize = 0;
-			bool error = false;
 		};
 		ReaderWriterLock m_activeStoresLock;
 		UnorderedMap<u16, ActiveStore> m_activeStores;
