[UBA]
* Fixed visualizer hang when horizontal autoscaling was enabled and window was minimized

diff --git a/Engine/Source/Programs/UnrealBuildAccelerator/Visualizer/Private/UbaVisualizer.cpp b/Engine/Source/Programs/UnrealBuildAccelerator/Visualizer/Private/UbaVisualizer.cpp
--- a/Engine/Source/Programs/UnrealBuildAccelerator/Visualizer/Private/UbaVisualizer.cpp
+++ b/Engine/Source/Programs/UnrealBuildAccelerator/Visualizer/Private/UbaVisualizer.cpp
@@ -2762,13 +2762,16 @@ namespace uba
 
 		RECT rect;
 		GetClientRect(m_hwnd, &rect);
+		if (rect.right == 0)
+			return false;
+
 		float timeS = TimeToS(playTime);
 
 		if (m_config.AutoScaleHorizontal)
 		{
 			m_scrollPosX = 0;
 			timeS = Max(timeS, 20.0f/m_zoomValue);
-			m_horizontalScaleValue = float(rect.right - m_progressRectLeft - 2)/(m_zoomValue*timeS*50.0f);
+			m_horizontalScaleValue = Max(float(rect.right - m_progressRectLeft - 2)/(m_zoomValue*timeS*50.0f), 0.001f);
 			return true;
 		}
 		else
@@ -3023,7 +3026,7 @@ namespace uba
 							newZoomValue = (m_processFont.height + 1) / 20.0f;
 					}
 					else
-						newScaleValue = m_horizontalScaleValue + m_horizontalScaleValue*float(delta)*0.0006f;
+						newScaleValue = Max(m_horizontalScaleValue + m_horizontalScaleValue*float(delta)*0.0006f, 0.001f);
 
 					// TODO: m_progressRectLeft changes with zoom so anchor logic is wrong
 					const float scrollAnchorOffsetX = float(cursorPos.x) - m_progressRectLeft;
