[UBA]
* Added ld-link.exe to rules list
* Improved some logging text in network server

diff --git a/Engine/Source/Programs/UnrealBuildAccelerator/Common/Private/UbaNetworkServer.cpp b/Engine/Source/Programs/UnrealBuildAccelerator/Common/Private/UbaNetworkServer.cpp
--- a/Engine/Source/Programs/UnrealBuildAccelerator/Common/Private/UbaNetworkServer.cpp
+++ b/Engine/Source/Programs/UnrealBuildAccelerator/Common/Private/UbaNetworkServer.cpp
@@ -299,10 +299,10 @@ namespace uba
 			{
 				if (server.m_onConnectionFunction)
 					server.m_onConnectionFunction(clientUid, clientId);
-				server.m_logger.Detail(TC("Client %s connected on connection %s"), GuidToString(clientUid).str, GuidToString(connectionUid).str);
+				server.m_logger.Detail(TC("Client %u (%s) connected on connection %s"), clientId, GuidToString(clientUid).str, GuidToString(connectionUid).str);
 			}
 			else
-				server.m_logger.Detail(TC("Client %s additional connection %s connected"), GuidToString(clientUid).str, GuidToString(connectionUid).str);
+				server.m_logger.Detail(TC("Client %u (%s) additional connection %s connected"), clientId, GuidToString(clientUid).str, GuidToString(connectionUid).str);
 
 
 			return true;
@@ -395,7 +395,7 @@ namespace uba
 				SCOPED_READ_LOCK(m_server.m_onDisconnectFunctionsLock, l);
 				for (auto& entry : m_server.m_onDisconnectFunctions)
 					entry.function(m_client->uid, m_client->id);
-				m_server.m_logger.Detail(TC("Client %s (Id: %u) disconnected"), GuidToString(m_client->uid).str, m_client->id);
+				m_server.m_logger.Detail(TC("Client %u (%s) disconnected"), m_client->id, GuidToString(m_client->uid).str);
 			}
 			m_disconnected = true;
 		}
diff --git a/Engine/Source/Programs/UnrealBuildAccelerator/Common/Private/UbaProcess.cpp b/Engine/Source/Programs/UnrealBuildAccelerator/Common/Private/UbaProcess.cpp
--- a/Engine/Source/Programs/UnrealBuildAccelerator/Common/Private/UbaProcess.cpp
+++ b/Engine/Source/Programs/UnrealBuildAccelerator/Common/Private/UbaProcess.cpp
@@ -1224,6 +1224,7 @@ namespace uba
 		WrittenFile& tempFile = insres.first->second;
 		FileMappingHandle oldMapping = tempFile.mappingHandle;
 		tempFile.mappingHandle = newHandle;
+		tempFile.mappingWritten = mappingHandleSize;
 		tempLock.Leave();
 		CloseFileMapping(oldMapping);
 		return true;
@@ -1784,6 +1785,8 @@ namespace uba
 				}
 			}
 
+			m_processStats.startupTime = GetTime() - m_startTime;
+
 			m_nativeProcessHandle = (ProcHandle)1;
 			m_nativeProcessId = u32(processID);
 
diff --git a/Engine/Source/Programs/UnrealBuildAccelerator/Core/Private/UbaApplicationRules.cpp b/Engine/Source/Programs/UnrealBuildAccelerator/Core/Private/UbaApplicationRules.cpp
--- a/Engine/Source/Programs/UnrealBuildAccelerator/Core/Private/UbaApplicationRules.cpp
+++ b/Engine/Source/Programs/UnrealBuildAccelerator/Core/Private/UbaApplicationRules.cpp
@@ -122,7 +122,8 @@ namespace uba
 		virtual bool NeedsSharedMemory(const tchar* file) const override
 		{
 			return Contains(file, TC("lnk{")) // This file is shared from link.exe to mt.exe and rc.exe so we need to put it shared memory
-			 	|| Contains(file, TC("\\_cl_")); // When link.exe is spawned by cl.exe we might use this which is in shared memory
+			 	|| Contains(file, TC("\\_cl_")) // When link.exe is spawned by cl.exe we might use this which is in shared memory
+			 	|| EndsWith(file, TStrlen(file), TC(".manifest")); // lld-link.exe is using a different name for files shared with child processes
 		}
 
 		virtual bool IsRarelyRead(const StringBufferBase& file) const override
@@ -471,6 +472,7 @@ namespace uba
 			{ TC("cvtres.exe"),					new ApplicationRulesLinkExe() },
 			{ TC("mt.exe"),						new ApplicationRulesLinkExe() },
 			{ TC("rc.exe"),						new ApplicationRulesLinkExe() },
+			{ TC("lld-link.exe"),				new ApplicationRulesLinkExe() },
 			{ TC("clang++.exe"),				new ApplicationRulesClangPlusPlusExe() },
 			{ TC("clang-cl.exe"),				new ApplicationRulesClangPlusPlusExe() },
 			{ TC("verse-clang-cl.exe"),			new ApplicationRulesClangPlusPlusExe() },
diff --git a/Engine/Source/Programs/UnrealBuildAccelerator/Detours/Private/Windows/UbaDetoursFunctionsNtDll.inl b/Engine/Source/Programs/UnrealBuildAccelerator/Detours/Private/Windows/UbaDetoursFunctionsNtDll.inl
--- a/Engine/Source/Programs/UnrealBuildAccelerator/Detours/Private/Windows/UbaDetoursFunctionsNtDll.inl
+++ b/Engine/Source/Programs/UnrealBuildAccelerator/Detours/Private/Windows/UbaDetoursFunctionsNtDll.inl
@@ -1050,10 +1050,17 @@ NTSTATUS NTAPI Shared_NtCreateFile(bool IsCreateFunc, PHANDLE hFileHandle, ACCES
 
 	if (keepInMemory || info.memoryFile)
 	{
+		#if UBA_DEBUG_LOG_ENABLED
+		const wchar_t* memoryType = L"MEMORY";
+		#endif
+
 		if (!info.memoryFile)
 		{
 			if (NeedsSharedMemory(fileName.data))
 			{
+				#if UBA_DEBUG_LOG_ENABLED
+				memoryType = L"SHAREDMEMORY";
+				#endif
 				if (isWrite)
 				{
 					info.memoryFile = new MemoryFile(false, FileTypeMaxSize(fileName, isSystemOrTempFile));
@@ -1118,6 +1125,14 @@ NTSTATUS NTAPI Shared_NtCreateFile(bool IsCreateFunc, PHANDLE hFileHandle, ACCES
 			info.memoryFile->volumeSerial = 1;
 			info.memoryFile->fileIndex = InterlockedDecrement(&g_memoryFileIndexCounter);
 		}
+		else
+		{
+			#if UBA_DEBUG_LOG_ENABLED
+			if (!info.memoryFile->isLocalOnly)
+				memoryType = L"SHAREDMEMORY";
+			#endif
+		}
+
 		_.Leave();
 
 		auto fileObject = new FileObject();
@@ -1132,7 +1147,7 @@ NTSTATUS NTAPI Shared_NtCreateFile(bool IsCreateFunc, PHANDLE hFileHandle, ACCES
 
 		TrackFileInput();
 
-		DEBUG_LOG_DETOURED(funcName, L"(MEMORY)%ls %llu (%ls) (%ls) -> Success", isWriteStr, uintptr_t(*hFileHandle), lpFileName, (fileName.data != lpFileName ? fileName.data : L""));
+		DEBUG_LOG_DETOURED(funcName, L"(%s)%ls %llu (%ls) (%ls) -> Success", memoryType, isWriteStr, uintptr_t(*hFileHandle), lpFileName, (fileName.data != lpFileName ? fileName.data : L""));
 		return STATUS_SUCCESS;
 	}
 
