[UBA]
* Fixed so temporary files stored in memory is properly updated when created in process, written to in child process and read in parent process.
* Changed so stdin/out/err are different handle types and fixed up code around this
* Fixed so ispc.exe handles extract import/export when building for non-windows

diff --git a/Engine/Source/Programs/UnrealBuildAccelerator/Cli/Private/UbaCli.cpp b/Engine/Source/Programs/UnrealBuildAccelerator/Cli/Private/UbaCli.cpp
--- a/Engine/Source/Programs/UnrealBuildAccelerator/Cli/Private/UbaCli.cpp
+++ b/Engine/Source/Programs/UnrealBuildAccelerator/Cli/Private/UbaCli.cpp
@@ -905,6 +905,16 @@ namespace uba
 		}
 		auto cg = MakeGuard([&]() { coordinator.Destroy(); });
 
+#if 0	// Annoying that link.exe/lld-link.exe needs path to windows folder.. 
+		if (!useScheduler)
+		{
+			StringBuffer<2048> temp;
+			temp.count = GetEnvironmentVariableW(TC("PATH"), temp.data, temp.capacity);
+			temp.Append(TC("c:\\sdk\\AutoSDK\\HostWin64\\Win64\\Windows Kits\\10\\bin\\10.0.19041.0\\x64;"));
+			SetEnvironmentVariableW(TC("PATH"), temp.data);
+
+		}
+#endif
 
 		for (u32 i=0; i!=loopCount; ++i)
 		{
diff --git a/Engine/Source/Programs/UnrealBuildAccelerator/Common/Private/UbaFileMapping.cpp b/Engine/Source/Programs/UnrealBuildAccelerator/Common/Private/UbaFileMapping.cpp
--- a/Engine/Source/Programs/UnrealBuildAccelerator/Common/Private/UbaFileMapping.cpp
+++ b/Engine/Source/Programs/UnrealBuildAccelerator/Common/Private/UbaFileMapping.cpp
@@ -395,12 +395,12 @@ namespace uba
 		}
 
 		u64 committedBefore = AlignUp(offset, m_pageSize);
-		u64 commitedAfter = AlignUp(newOffset, m_pageSize);
+		u64 committedAfter = AlignUp(newOffset, m_pageSize);
 
-		if (f.commitOnAlloc && committedBefore != commitedAfter)
+		if (f.commitOnAlloc && committedBefore != committedAfter)
 		{
 			u64 commitStart = committedBefore - alignedOffsetStart;
-			u64 commitSize = commitedAfter - committedBefore;
+			u64 commitSize = committedAfter - committedBefore;
 			if (!MapViewCommit(data + commitStart, commitSize))
 			{
 				UnmapViewOfFile(data, mapSize, hint);
diff --git a/Engine/Source/Programs/UnrealBuildAccelerator/Common/Private/UbaProcess.cpp b/Engine/Source/Programs/UnrealBuildAccelerator/Common/Private/UbaProcess.cpp
--- a/Engine/Source/Programs/UnrealBuildAccelerator/Common/Private/UbaProcess.cpp
+++ b/Engine/Source/Programs/UnrealBuildAccelerator/Common/Private/UbaProcess.cpp
@@ -835,6 +835,23 @@ namespace uba
 	{
 		writer.WriteU32(m_session.GetDirectoryTableSize());
 		writer.WriteU32(m_session.GetFileMappingSize());
+
+#if PLATFORM_WINDOWS
+		if (!m_tempFilesModified)
+		{
+			writer.WriteU32(0);
+			return true;
+		}
+
+		// This should be very rare (lld-link.exe uses it for mt.exe).. if we get too many temp files we'll need to rethink this design
+		SCOPED_READ_LOCK(m_tempFilesLock, l);
+		writer.WriteU32(u32(m_tempFiles.size()));
+		for (auto& kv : m_tempFiles)
+		{
+			writer.WriteStringKey(kv.first);
+			writer.WriteU64(kv.second.mappingWritten);
+		}
+#endif
 		return true;
 	}
 
@@ -1210,7 +1227,7 @@ namespace uba
 		FileMappingHandle source;
 		source.FromU64(mappingHandle);
 		FileMappingHandle newHandle;
-		if (!DuplicateFileMapping(nativeProcessHandle, source, GetCurrentProcessHandle(), &newHandle, FILE_MAP_READ, false, 0))
+		if (!DuplicateFileMapping(nativeProcessHandle, source, GetCurrentProcessHandle(), &newHandle, 0, false, DUPLICATE_SAME_ACCESS))
 		{
 			m_session.m_logger.Error(TC("Failed to duplicate handle for temp file (%s)"), fileName.data);
 			return true;
@@ -1221,6 +1238,7 @@ namespace uba
 		if (insres.second)
 			return true;
 
+		++m_tempFilesModified;
 		WrittenFile& tempFile = insres.first->second;
 		FileMappingHandle oldMapping = tempFile.mappingHandle;
 		tempFile.mappingHandle = newHandle;
diff --git a/Engine/Source/Programs/UnrealBuildAccelerator/Common/Private/UbaSession.cpp b/Engine/Source/Programs/UnrealBuildAccelerator/Common/Private/UbaSession.cpp
--- a/Engine/Source/Programs/UnrealBuildAccelerator/Common/Private/UbaSession.cpp
+++ b/Engine/Source/Programs/UnrealBuildAccelerator/Common/Private/UbaSession.cpp
@@ -2758,6 +2758,9 @@ namespace uba
 
 	bool Session::ExtractSymbolsFromObjectFile(const CloseFileMessage& msg, const tchar* fileName, u64 fileSize)
 	{
+		if (!msg.mappingHandle)
+			return m_logger.Error(TC("Can't extract symbols from obj file that is written directly to disk (%s writing %s)"), msg.process.m_startInfo.application, fileName);
+
 		FileMappingHandle source;
 		source.FromU64(msg.mappingHandle);
 		FileMappingHandle objectFileMappingHandle;
diff --git a/Engine/Source/Programs/UnrealBuildAccelerator/Common/Public/UbaProcess.h b/Engine/Source/Programs/UnrealBuildAccelerator/Common/Public/UbaProcess.h
--- a/Engine/Source/Programs/UnrealBuildAccelerator/Common/Public/UbaProcess.h
+++ b/Engine/Source/Programs/UnrealBuildAccelerator/Common/Public/UbaProcess.h
@@ -171,6 +171,7 @@ namespace uba
 		UnorderedMap<TString, WrittenFile>& m_writtenFiles;
 		ReaderWriterLock& m_tempFilesLock;
 		UnorderedMap<StringKey, WrittenFile>& m_tempFiles;
+		Atomic<u32> m_tempFilesModified;
 		SessionStats m_sessionStats;
 		StorageStats m_storageStats;
 		ProcessStats m_processStats;
diff --git a/Engine/Source/Programs/UnrealBuildAccelerator/Core/Private/UbaApplicationRules.cpp b/Engine/Source/Programs/UnrealBuildAccelerator/Core/Private/UbaApplicationRules.cpp
--- a/Engine/Source/Programs/UnrealBuildAccelerator/Core/Private/UbaApplicationRules.cpp
+++ b/Engine/Source/Programs/UnrealBuildAccelerator/Core/Private/UbaApplicationRules.cpp
@@ -160,6 +160,23 @@ namespace uba
 		}
 	};
 
+	class ApplicationRulesLldLinkExe : public ApplicationRulesLinkExe
+	{
+		using Super = ApplicationRulesLinkExe;
+
+		virtual bool KeepInMemory(const StringView& fileName, const tchar* systemTemp) const override
+		{
+			return fileName.EndsWith(TC(".manifest"));
+		}
+
+		virtual bool IsOutputFile(const StringView& fileName) const override
+		{
+			return false//Super::IsOutputFile(fileName)
+				|| fileName.Contains(TC(".exe.tmp"))
+				|| fileName.Contains(TC(".pdb.tmp"));
+		}
+	};
+
 	// ==== Clang tool chain ====
 
 	class ApplicationRulesClang : public ApplicationRules
@@ -375,7 +392,8 @@ namespace uba
 			return fileName.Contains(TC(".generated.dummy"))
 				|| fileName.EndsWith(TC(".ispc.bc"))
 				|| fileName.EndsWith(TC(".ispc.txt"))
-				|| fileName.EndsWith(TC(".obj"));
+				|| fileName.EndsWith(TC(".obj"))
+				|| fileName.EndsWith(TC(".o")); // Used when compiling for linux
 		}
 
 		virtual bool IsCacheable() const override
@@ -472,7 +490,7 @@ namespace uba
 			{ TC("cvtres.exe"),					new ApplicationRulesLinkExe() },
 			{ TC("mt.exe"),						new ApplicationRulesLinkExe() },
 			{ TC("rc.exe"),						new ApplicationRulesLinkExe() },
-			{ TC("lld-link.exe"),				new ApplicationRulesLinkExe() },
+			{ TC("lld-link.exe"),				new ApplicationRulesLldLinkExe() },
 			{ TC("clang++.exe"),				new ApplicationRulesClangPlusPlusExe() },
 			{ TC("clang-cl.exe"),				new ApplicationRulesClangPlusPlusExe() },
 			{ TC("verse-clang-cl.exe"),			new ApplicationRulesClangPlusPlusExe() },
diff --git a/Engine/Source/Programs/UnrealBuildAccelerator/Core/Private/UbaPlatform.cpp b/Engine/Source/Programs/UnrealBuildAccelerator/Core/Private/UbaPlatform.cpp
--- a/Engine/Source/Programs/UnrealBuildAccelerator/Core/Private/UbaPlatform.cpp
+++ b/Engine/Source/Programs/UnrealBuildAccelerator/Core/Private/UbaPlatform.cpp
@@ -128,7 +128,7 @@ namespace uba
 		#endif
 	}
 
-	void WriteAssertInfo(StringBufferBase& out, const tchar* text, const char* file, u32 line, const char* expr, u32 skipCallstack = 0)
+	UBA_NOINLINE void WriteAssertInfo(StringBufferBase& out, const tchar* text, const char* file, u32 line, const char* expr, u32 skipCallstack = 0)
 	{
 #if PLATFORM_WINDOWS
 		if (text)
@@ -164,7 +164,7 @@ namespace uba
 				if (file)
 					out.Appendf(L"\n\n");
 				out.Appendf(L"  Callstack:");
-				for (u32 i = 0; i < count; i++)
+				for (u32 i = skipCallstack; i < count; i++)
 				{
 					tchar str[1024];
 					auto addr = u64(callers[i]);
diff --git a/Engine/Source/Programs/UnrealBuildAccelerator/Core/Public/UbaPlatform.h b/Engine/Source/Programs/UnrealBuildAccelerator/Core/Public/UbaPlatform.h
--- a/Engine/Source/Programs/UnrealBuildAccelerator/Core/Public/UbaPlatform.h
+++ b/Engine/Source/Programs/UnrealBuildAccelerator/Core/Public/UbaPlatform.h
@@ -47,7 +47,6 @@
 	#endif
 
 	#define ANALYSIS_NORETURN // __attribute__((analyzer_noreturn))
-    #define ERROR_SUCCESS 0
 #endif
 
 #define UBA_EXPERIMENTAL 0
@@ -124,6 +123,7 @@ namespace uba
 	#define TStrcpy_s wcscpy_s
 	#define TStrcat_s wcscat_s
 	#define TStrdup _wcsdup
+	#define UBA_NOINLINE __declspec(noinline)
 #else
 	inline constexpr tchar PathSeparator = '/';
 	inline constexpr tchar NonPathSeparator = '\\';
@@ -147,6 +147,9 @@ namespace uba
 	#define st_mtimespec st_mtim
 	#else
 	inline constexpr bool CaseInsensitiveFs = true;
-#endif
+	#endif
+	#define UBA_NOINLINE
+    #define ERROR_SUCCESS 0
+	#define DUPLICATE_SAME_ACCESS 0
 #endif // PLATFORM_WINDOWS
 }
diff --git a/Engine/Source/Programs/UnrealBuildAccelerator/Detours/Private/UbaDetoursFileMappingTable.cpp b/Engine/Source/Programs/UnrealBuildAccelerator/Detours/Private/UbaDetoursFileMappingTable.cpp
--- a/Engine/Source/Programs/UnrealBuildAccelerator/Detours/Private/UbaDetoursFileMappingTable.cpp
+++ b/Engine/Source/Programs/UnrealBuildAccelerator/Detours/Private/UbaDetoursFileMappingTable.cpp
@@ -4,6 +4,10 @@
 #include "UbaDirectoryTable.h"
 #include "UbaPathUtils.h"
 
+#if PLATFORM_WINDOWS
+#include "Windows/UbaDetoursUtilsWin.h"
+#endif
+
 namespace uba
 {
 	MappedFileTable::MappedFileTable(MemoryBlock& memoryBlock) : m_memoryBlock(memoryBlock), m_lookup(&memoryBlock)
@@ -153,18 +157,45 @@ namespace uba
 
 	void Rpc_UpdateTables()
 	{
-		TimerScope ts(g_stats.updateTables);
-		SCOPED_WRITE_LOCK(g_communicationLock, pcs);
-		BinaryWriter writer;
-		writer.WriteByte(MessageType_UpdateTables);
-		writer.Flush();
-		BinaryReader reader;
-		u32 directoryTableSize = reader.ReadU32();
-		u32 fileMappingTableSize = reader.ReadU32();
-		pcs.Leave();
+		u32 directoryTableSize;
+		u32 fileMappingTableSize;
+		{
+			TimerScope ts(g_stats.updateTables);
+			SCOPED_WRITE_LOCK(g_communicationLock, pcs);
+			BinaryWriter writer;
+			writer.WriteByte(MessageType_UpdateTables);
+			writer.Flush();
+			BinaryReader reader;
+			directoryTableSize = reader.ReadU32();
+			fileMappingTableSize = reader.ReadU32();
+
+#if PLATFORM_WINDOWS
+			if (u32 tempFileCount = reader.ReadU32())
+			{
+				SCOPED_WRITE_LOCK(g_mappedFileTable.m_lookupLock, _);
+				while (tempFileCount--)
+				{
+					StringKey fileNameKey = reader.ReadStringKey();
+					u64 fileSize = reader.ReadU64();
+					auto findIt = g_mappedFileTable.m_lookup.find(fileNameKey);
+					if (findIt == g_mappedFileTable.m_lookup.end())
+						continue;
+					FileInfo& info = findIt->second;
+					UBA_ASSERT(info.memoryFile);
+					if (!info.memoryFile)
+						continue;
+					info.memoryFile->writtenSize = fileSize;
+					if (fileSize <= info.memoryFile->committedSize)
+						continue;
+					UBA_ASSERT(info.memoryFile->committedSize == 0);
+					info.memoryFile->EnsureCommitted(DetouredHandle(HandleType_File), fileSize);
+				}
+			}
+#endif
+			DEBUG_LOG_PIPE(L"UpdateTables", L"");
+		}
 		g_directoryTable.ParseDirectoryTable(directoryTableSize);
 		g_mappedFileTable.Parse(fileMappingTableSize);
-		DEBUG_LOG_PIPE(L"UpdateTables", L"");
 	}
 
 	u32 Rpc_GetEntryOffset(const StringKey& entryNameKey, const tchar* entryName, u64 entryNameLen, bool checkIfDir)
diff --git a/Engine/Source/Programs/UnrealBuildAccelerator/Detours/Private/Windows/UbaDetoursFunctionsKernelBase.inl b/Engine/Source/Programs/UnrealBuildAccelerator/Detours/Private/Windows/UbaDetoursFunctionsKernelBase.inl
--- a/Engine/Source/Programs/UnrealBuildAccelerator/Detours/Private/Windows/UbaDetoursFunctionsKernelBase.inl
+++ b/Engine/Source/Programs/UnrealBuildAccelerator/Detours/Private/Windows/UbaDetoursFunctionsKernelBase.inl
@@ -364,7 +364,7 @@ BOOL Detoured_ReadFile(HANDLE hFile, LPVOID lpBuffer, DWORD nNumberOfBytesToRead
 	{
 		auto& dh = asDetouredHandle(hFile);
 
-		if (dh.type == HandleType_Std) // HACK HACK
+		if (dh.type == HandleType_StdIn) // HACK HACK
 		{
 			UBA_ASSERTF(false, L"Trying to read input from stdin while application is running in a way console can not be accessed");
 			memcpy(lpBuffer, "Y\r\n", 3);
@@ -635,14 +635,18 @@ BOOL Detoured_WriteFile(HANDLE hFile, LPCVOID lpBuffer, DWORD nNumberOfBytesToWr
 		auto& dh = asDetouredHandle(hFile);
 		auto& fo = *dh.fileObject;
 
-		if (dh.type == HandleType_Std)
+		if (dh.type >= HandleType_StdErr)
 		{
-			//DEBUG_LOG_DETOURED(L"WriteStdFile1", L"%llu", uintptr_t(hFile));
-			WriteStdFile(lpBuffer, nNumberOfBytesToWrite, hFile == g_stdHandle[0]);
+			if (dh.type != HandleType_StdIn)
+			{
+				//DEBUG_LOG_DETOURED(L"WriteStdFile1", L"%llu", uintptr_t(hFile));
+				WriteStdFile(lpBuffer, nNumberOfBytesToWrite, dh.type == HandleType_StdErr);
+			}
 			*lpNumberOfBytesWritten = nNumberOfBytesToWrite;
 			SetLastError(ERROR_SUCCESS);
 			return true;
 		}
+
 		auto& fi = *fo.fileInfo;
 		if (MemoryFile* mf = fi.memoryFile)
 		{
@@ -833,7 +837,7 @@ DWORD Detoured_SetFilePointerEx(HANDLE hFile, LARGE_INTEGER liDistanceToMove, PL
 	if (isDetouredHandle(hFile))
 	{
 		DetouredHandle& dh = asDetouredHandle(hFile);
-		if (dh.type == HandleType_Std)
+		if (dh.type >= HandleType_StdErr)
 		{
 			if (lpNewFilePointer)
 				*lpNewFilePointer = ToLargeInteger(0);
@@ -887,7 +891,7 @@ BOOL Detoured_SetEndOfFile(HANDLE hFile)
 			DEBUG_LOG_DETOURED(L"SetEndOfFile (MEMORY)", L"%llu (%ls) -> Success", uintptr_t(hFile), HandleToName(hFile));
 			mf->writtenSize = dh.pos;
 			mf->isReported = false;
-			mf->EnsureCommited(dh, mf->writtenSize);
+			mf->EnsureCommitted(dh, mf->writtenSize);
 			SetLastError(ERROR_SUCCESS);
 			return true;
 		}
@@ -955,12 +959,12 @@ DWORD Detoured_GetFileType(HANDLE hFile)
 	{
 		DetouredHandle& dh = asDetouredHandle(hFile);
 		SetLastError(ERROR_SUCCESS);
-		if (dh.type == HandleType_Std)
+		if (dh.type >= HandleType_StdErr)
 		{
 			DEBUG_LOG_DETOURED(L"GetFileType", L"%llu (%ls) -> FILE_TYPE_CHAR", uintptr_t(hFile), HandleToName(hFile));
 			return FILE_TYPE_CHAR;
 		}
-		UBA_ASSERT(dh.type == HandleType_File);
+		UBA_ASSERTF(dh.type == HandleType_File, L"HandleType: %u", dh.type);
 		DEBUG_LOG_DETOURED(L"GetFileType", L"%llu (%ls) -> FILE_TYPE_DISK", uintptr_t(hFile), HandleToName(hFile));
 		return FILE_TYPE_DISK;
 	}
@@ -1989,7 +1993,7 @@ BOOL Detoured_GetFileInformationByHandle(HANDLE hFile, LPBY_HANDLE_FILE_INFORMAT
 
 		if (MemoryFile* mf = fi.memoryFile)
 		{
-			DEBUG_LOG_DETOURED(L"GetFileInformationByHandle", L"(memoryfile) %llu (%ls) -> Success", uintptr_t(hFile), HandleToName(hFile));
+			DEBUG_LOG_DETOURED(L"GetFileInformationByHandle", L"(memoryfile) %llu (%ls) -> Success (Size: %llu)", uintptr_t(hFile), HandleToName(hFile), mf->writtenSize);
 			lpFileInformation->dwFileAttributes = FILE_ATTRIBUTE_NORMAL;
 			(u64&)lpFileInformation->ftCreationTime = mf->fileTime;
 			(u64&)lpFileInformation->ftLastAccessTime = mf->fileTime;
@@ -1999,7 +2003,7 @@ BOOL Detoured_GetFileInformationByHandle(HANDLE hFile, LPBY_HANDLE_FILE_INFORMAT
 			lpFileInformation->nFileIndexHigh = li.HighPart;
 			lpFileInformation->nFileIndexLow = li.LowPart;
 			lpFileInformation->nNumberOfLinks = 1;//~u32(0); // TODO
-			li = ToLargeInteger(fi.memoryFile->writtenSize);
+			li = ToLargeInteger(mf->writtenSize);
 			lpFileInformation->nFileSizeHigh = li.HighPart;
 			lpFileInformation->nFileSizeLow = li.LowPart;
 			return TRUE;
@@ -2151,7 +2155,7 @@ HANDLE Detoured_CreateFileMappingW(HANDLE hFile, LPSECURITY_ATTRIBUTES lpFileMap
 				li.LowPart = dwMaximumSizeLow;
 				if (!(flProtect & MEM_RESERVE) && li.QuadPart)
 				{
-					mf->EnsureCommited(*mdh, li.QuadPart);
+					mf->EnsureCommitted(*mdh, li.QuadPart);
 					if (!mf->writtenSize && (flProtect & PAGE_READWRITE)) // TODO: Maybe we should always set writtenSize?
 						mf->writtenSize = li.QuadPart;
 				}
@@ -2723,7 +2727,9 @@ BOOL Detoured_CreateProcessW(LPCWSTR lpApplicationName, LPWSTR lpCommandLine, LP
 		writer.Flush();
 		BinaryReader reader;
 		processId = reader.ReadU32();
-		UBA_ASSERT(processId > 0);
+		UBA_ASSERTF(processId > 0, L"Failed to create process %s", originalCmd);
+		if (!processId)
+			return FALSE;
 
 		reader.Skip(sizeof(u32)); // Rules index
 
@@ -2738,18 +2744,30 @@ BOOL Detoured_CreateProcessW(LPCWSTR lpApplicationName, LPWSTR lpCommandLine, LP
 
 	LPCSTR dlls[] = { dll };
 
-	UBA_ASSERT(!isDetouredHandle(lpStartupInfo->hStdInput));
-
-	if (g_isDetachedProcess)
+	if (isDetouredHandle(lpStartupInfo->hStdError))
 	{
-		if (lpStartupInfo->hStdError == g_stdHandle[0])
-			lpStartupInfo->hStdError = 0;
-		if (lpStartupInfo->hStdOutput == g_stdHandle[1])
-			lpStartupInfo->hStdOutput = 0;
+		DetouredHandle& dh = asDetouredHandle(lpStartupInfo->hStdError);
+		if (dh.type == HandleType_StdErr)
+			lpStartupInfo->hStdError = g_isDetachedProcess ? 0 : True_GetStdHandle(STD_ERROR_HANDLE);
+		else
+			UBA_ASSERTF(false, L"hStdError is detoured (%s)", lpApplicationName);
+	}
+	if (isDetouredHandle(lpStartupInfo->hStdOutput))
+	{
+		DetouredHandle& dh = asDetouredHandle(lpStartupInfo->hStdOutput);
+		if (dh.type == HandleType_StdOut)
+			lpStartupInfo->hStdOutput = g_isDetachedProcess ? 0 : True_GetStdHandle(STD_OUTPUT_HANDLE);
+		else
+			UBA_ASSERTF(false, L"hStdOutput is detoured (%s)", lpApplicationName);
+	}
+	if (isDetouredHandle(lpStartupInfo->hStdInput))
+	{
+		DetouredHandle& dh = asDetouredHandle(lpStartupInfo->hStdInput);
+		if (dh.type == HandleType_StdIn)
+			lpStartupInfo->hStdInput = g_isDetachedProcess ? 0 : True_GetStdHandle(STD_INPUT_HANDLE);
+		else
+			UBA_ASSERTF(false, L"hStdInput is detoured (%s)", lpApplicationName);
 	}
-
-	UBA_ASSERT(!isDetouredHandle(lpStartupInfo->hStdOutput));
-	UBA_ASSERTF(!isDetouredHandle(lpStartupInfo->hStdError), L"Got detoured handle for stderror: %llu", lpStartupInfo->hStdError);
 
 	lpStartupInfo->dwFlags |= STARTF_USESHOWWINDOW;
 	lpStartupInfo->wShowWindow = SW_HIDE;
@@ -2785,7 +2803,6 @@ BOOL Detoured_CreateProcessW(LPCWSTR lpApplicationName, LPWSTR lpCommandLine, LP
 		continue;
 	}
 	--t_disallowDetour;
-	UBA_ASSERTF(res, L"Failed to spawn process %ls (Error code: %u)", commandLine.c_str(), lastError);
 
 	if (isChild)
 	{
@@ -2803,6 +2820,8 @@ BOOL Detoured_CreateProcessW(LPCWSTR lpApplicationName, LPWSTR lpCommandLine, LP
 		DEBUG_LOG_PIPE(L"StartProcess", L"%ls %ls", lpApplicationName, originalCmd);
 	}
 
+	UBA_ASSERTF(res, L"Failed to spawn process %ls (Error code: %u)", commandLine.c_str(), lastError);
+
 	HANDLE trueHandle = lpProcessInformation->hProcess;
 
 	if (!res || trueHandle == INVALID_HANDLE_VALUE)
diff --git a/Engine/Source/Programs/UnrealBuildAccelerator/Detours/Private/Windows/UbaDetoursFunctionsNtDll.inl b/Engine/Source/Programs/UnrealBuildAccelerator/Detours/Private/Windows/UbaDetoursFunctionsNtDll.inl
--- a/Engine/Source/Programs/UnrealBuildAccelerator/Detours/Private/Windows/UbaDetoursFunctionsNtDll.inl
+++ b/Engine/Source/Programs/UnrealBuildAccelerator/Detours/Private/Windows/UbaDetoursFunctionsNtDll.inl
@@ -688,7 +688,8 @@ NTSTATUS NTAPI Shared_NtCreateFile(bool IsCreateFunc, PHANDLE hFileHandle, ACCES
 
 	bool useContent = IsContentUse(DesiredAccess, CreateDisposition);
 	bool isWrite = IsWrite(DesiredAccess, CreateDisposition);
-	bool keepInMemory = (KeepInMemory(fileName) && useContent) || IsOutputFile(fileName, isWrite, isDeleteOnClose);
+	//bool keepInMemory = (KeepInMemory(fileName) && (useContent || NeedsSharedMemory(fileName.data))) || IsOutputFile(fileName, isWrite, isDeleteOnClose);
+	bool keepInMemory = KeepInMemory(fileName) || IsOutputFile(fileName, isWrite, isDeleteOnClose);
 
 #if UBA_DEBUG_LOG_ENABLED
 	const wchar_t* isWriteStr = isWrite ? L" WRITE" : L""; (void)isWriteStr;
@@ -1061,11 +1062,6 @@ NTSTATUS NTAPI Shared_NtCreateFile(bool IsCreateFunc, PHANDLE hFileHandle, ACCES
 				#if UBA_DEBUG_LOG_ENABLED
 				memoryType = L"SHAREDMEMORY";
 				#endif
-				if (isWrite)
-				{
-					info.memoryFile = new MemoryFile(false, FileTypeMaxSize(fileName, isSystemOrTempFile));
-				}
-				else
 				{
 					TimerScope ts(g_stats.openTempFile);
 					SCOPED_WRITE_LOCK(g_communicationLock, pcs);
@@ -1080,16 +1076,20 @@ NTSTATUS NTAPI Shared_NtCreateFile(bool IsCreateFunc, PHANDLE hFileHandle, ACCES
 					pcs.Leave();
 					if (mappingHandle)
 					{
-						info.memoryFile = new MemoryFile();
+						info.memoryFile = new MemoryFile(nullptr, false);
 						MemoryFile& mf = *info.memoryFile;
 						True_DuplicateHandle(g_hostProcess, (HANDLE)mappingHandle, GetCurrentProcess(), &mf.mappingHandle, 0, FALSE, DUPLICATE_SAME_ACCESS);
 						UBA_ASSERTF(mf.mappingHandle, L"DuplicateHandle failed when opening temp file %ls (%u)", fileName.data, GetLastError());
 						mf.writtenSize = mappingHandleSize;
-
 						TimerScope ts2(g_kernelStats.mapViewOfFile);
-						mf.baseAddress = (u8*)True_MapViewOfFile(mf.mappingHandle, FILE_MAP_READ, 0, 0, mappingHandleSize);
+						mf.baseAddress = (u8*)True_MapViewOfFile(mf.mappingHandle, FILE_MAP_READ | FILE_MAP_ALL_ACCESS, 0, 0, mappingHandleSize);
 						UBA_ASSERTF(mf.baseAddress, L"MapViewOfFile failed when opening temp file %ls (%u)", fileName.data, GetLastError());
-						mf.committedSize = mappingHandleSize;
+						//mf.committedSize = mappingHandleSize;
+						mf.reserveSize = FileTypeMaxSize(fileName, isSystemOrTempFile);
+					}
+					else if (isWrite)
+					{
+						info.memoryFile = new MemoryFile(false, FileTypeMaxSize(fileName, isSystemOrTempFile));
 					}
 					else
 					{
@@ -1311,9 +1311,6 @@ NTSTATUS NTAPI Detoured_NtClose(HANDLE handle)
 
 	DetouredHandle& dh = asDetouredHandle(handle);
 
-	if (dh.type == HandleType_Std)
-		return STATUS_SUCCESS;
-
 	NTSTATUS res = STATUS_SUCCESS;
 
 	if (dh.trueHandle != INVALID_HANDLE_VALUE)
@@ -1325,6 +1322,8 @@ NTSTATUS NTAPI Detoured_NtClose(HANDLE handle)
 	FileObject* fo = dh.fileObject;
 	if (!fo)
 	{
+		if (dh.type >= HandleType_StdErr) // TODO: This might leak if handle is duplicated.. but ignore for now
+			return res;
 		DEBUG_LOG_TRUE(L"NtClose", L"%llu (%ls) -> %ls", uintptr_t(handle), HandleToName(handle), ToString(res));
 		delete& dh;
 		return res;
@@ -1421,6 +1420,7 @@ NTSTATUS NTAPI Detoured_NtClose(HANDLE handle)
 		}
 		else if (NeedsSharedMemory(fi.originalName) && IsWrite(fo->desiredAccess, 0))
 		{
+			UBA_ASSERT(!fi.memoryFile->isLocalOnly);
 			StringBuffer<> fixedName;
 			FixPath(fixedName, path);
 
diff --git a/Engine/Source/Programs/UnrealBuildAccelerator/Detours/Private/Windows/UbaDetoursFunctionsUcrtBase.inl b/Engine/Source/Programs/UnrealBuildAccelerator/Detours/Private/Windows/UbaDetoursFunctionsUcrtBase.inl
--- a/Engine/Source/Programs/UnrealBuildAccelerator/Detours/Private/Windows/UbaDetoursFunctionsUcrtBase.inl
+++ b/Engine/Source/Programs/UnrealBuildAccelerator/Detours/Private/Windows/UbaDetoursFunctionsUcrtBase.inl
@@ -407,6 +407,7 @@ errno_t Detoured_getenv_s(size_t* pReturnValue, char* buffer, size_t numberOfEle
 
 errno_t Detoured__wmakepath_s(wchar_t* path, size_t sizeInWords, const wchar_t* drive, const wchar_t* dir, const wchar_t* fname, const wchar_t* ext)
 {
+	DETOURED_CALL(_wmakepath_s);
 	auto res = True__wmakepath_s(path, sizeInWords, drive, dir, fname, ext);
 	DEBUG_LOG_TRUE(L"_wmakepath_s", L"%ls %ls %ls %ls %ls", path, drive, dir, fname, ext);
 	return res;
@@ -414,8 +415,9 @@ errno_t Detoured__wmakepath_s(wchar_t* path, size_t sizeInWords, const wchar_t*
 char* Detoured__getcwd(char* buffer, int maxlen)
 {
 	DETOURED_CALL(_getcwd);
-	UBA_ASSERT(!g_virtualCommandLineA);
-	return True__getcwd(buffer, maxlen);
+	auto res = True__getcwd(buffer, maxlen); // We know this calls GetCommandLine for both wine and windows
+	DEBUG_LOG_TRUE(L"_getcwd", L"%hs", res);
+	return res;
 }
 
 #endif // defined(DETOURED_INCLUDE_DEBUG)
diff --git a/Engine/Source/Programs/UnrealBuildAccelerator/Detours/Private/Windows/UbaDetoursFunctionsWin.cpp b/Engine/Source/Programs/UnrealBuildAccelerator/Detours/Private/Windows/UbaDetoursFunctionsWin.cpp
--- a/Engine/Source/Programs/UnrealBuildAccelerator/Detours/Private/Windows/UbaDetoursFunctionsWin.cpp
+++ b/Engine/Source/Programs/UnrealBuildAccelerator/Detours/Private/Windows/UbaDetoursFunctionsWin.cpp
@@ -234,7 +234,7 @@ struct SuppressCreateFileDetourScope
 	~SuppressCreateFileDetourScope() { --t_disallowCreateFileDetour; }
 };
 
-const wchar_t* HandleToName(DetouredHandle& dh)
+const wchar_t* HandleToName(const DetouredHandle& dh)
 {
 	if (dh.fileObject)
 		if (const wchar_t* name = dh.fileObject->fileInfo->name)
@@ -242,11 +242,12 @@ const wchar_t* HandleToName(DetouredHandle& dh)
 	return L"Unknown";
 }
 
+BlockAllocator<FileObject> g_fileObjectAllocator(g_memoryBlock);
 
 void MemoryFile::Write(DetouredHandle& handle, LPCVOID lpBuffer, u64 nNumberOfBytesToWrite)
 {
 	u64 newPos = handle.pos + nNumberOfBytesToWrite;
-	EnsureCommited(handle, newPos);
+	EnsureCommitted(handle, newPos);
 	memcpy(baseAddress + handle.pos, lpBuffer, nNumberOfBytesToWrite);
 	handle.pos += nNumberOfBytesToWrite;
 	if (writtenSize < newPos)
@@ -256,7 +257,7 @@ void MemoryFile::Write(DetouredHandle& handle, LPCVOID lpBuffer, u64 nNumberOfBy
 	}
 }
 
-void MemoryFile::EnsureCommited(DetouredHandle& handle, u64 size)
+void MemoryFile::EnsureCommitted(const DetouredHandle& handle, u64 size)
 {
 	if (committedSize >= size)
 		return;
@@ -279,14 +280,7 @@ void MemoryFile::EnsureCommited(DetouredHandle& handle, u64 size)
 		}
 
 		if (shouldRemap)
-		{
-			True_UnmapViewOfFile(baseAddress);
-			mappedSize = Min(reserveSize, AlignUp(Max(size, mappedSize * 4), g_pageSize));
-			TimerScope ts(g_kernelStats.mapViewOfFile);
-			baseAddress = (u8*)True_MapViewOfFile(mappingHandle, FILE_MAP_READ | FILE_MAP_WRITE, 0, 0, mappedSize);
-			if (!baseAddress)
-				FatalError(1347, L"MapViewOfFile failed trying to map %llu for %ls. ReservedSize: %llu (Error code: %u)", mappedSize, HandleToName(handle), reserveSize, GetLastError());
-		}
+			Remap(handle, size);
 	}
 
 	u64 toCommit = Min(reserveSize, AlignUp(size - committedSize, g_pageSize));
@@ -295,6 +289,16 @@ void MemoryFile::EnsureCommited(DetouredHandle& handle, u64 size)
 	committedSize += toCommit;
 }
 
+void MemoryFile::Remap(const DetouredHandle& handle, u64 size)
+{
+	True_UnmapViewOfFile(baseAddress);
+	mappedSize = Min(reserveSize, AlignUp(Max(size, mappedSize * 4), g_pageSize));
+	TimerScope ts(g_kernelStats.mapViewOfFile);
+	baseAddress = (u8*)True_MapViewOfFile(mappingHandle, FILE_MAP_READ | FILE_MAP_WRITE, 0, 0, mappedSize);
+	if (!baseAddress)
+		FatalError(1347, L"MapViewOfFile failed trying to map %llu for %ls. ReservedSize: %llu (Error code: %u)", mappedSize, HandleToName(handle), reserveSize, GetLastError());
+}
+
 void ToInvestigate(const wchar_t* format, ...)
 {
 #if UBA_DEBUG_LOG_ENABLED
@@ -309,12 +313,12 @@ void ToInvestigate(const wchar_t* format, ...)
 #endif
 }
 
-void UbaAssert(const wchar_t* text, const char* file, u32 line, const char* expr, u32 terminateCode, bool allowTerminate)
+UBA_NOINLINE void UbaAssert(const wchar_t* text, const char* file, u32 line, const char* expr, u32 terminateCode, bool allowTerminate)
 {
 	SuppressDetourScope _;
 	
 	StringBuffer<32*1024> b;
-	WriteAssertInfo(b, text, file, line, expr, 1);
+	WriteAssertInfo(b, text, file, line, expr, 3);
 	Rpc_WriteLog(b.data, b.count, true, true);
 	#if UBA_DEBUG_LOG_ENABLED
 	FlushDebugLog();
@@ -803,9 +807,9 @@ void Init(const DetoursPayload& payload, u64 startTime)
 
 	if (g_isDetachedProcess)
 	{
-		g_stdHandle[0] = makeDetouredHandle(new DetouredHandle(HandleType_Std)); // STD_ERR
-		g_stdHandle[1] = makeDetouredHandle(new DetouredHandle(HandleType_Std)); // STD_OUT
-		g_stdHandle[2] = makeDetouredHandle(new DetouredHandle(HandleType_Std)); // STD_IN
+		g_stdHandle[0] = makeDetouredHandle(new DetouredHandle(HandleType_StdErr)); // STD_ERR
+		g_stdHandle[1] = makeDetouredHandle(new DetouredHandle(HandleType_StdOut)); // STD_OUT
+		g_stdHandle[2] = makeDetouredHandle(new DetouredHandle(HandleType_StdIn)); // STD_IN
 	}
 
 	if (payload.trackInputs)
diff --git a/Engine/Source/Programs/UnrealBuildAccelerator/Detours/Private/Windows/UbaDetoursUtilsWin.h b/Engine/Source/Programs/UnrealBuildAccelerator/Detours/Private/Windows/UbaDetoursUtilsWin.h
--- a/Engine/Source/Programs/UnrealBuildAccelerator/Detours/Private/Windows/UbaDetoursUtilsWin.h
+++ b/Engine/Source/Programs/UnrealBuildAccelerator/Detours/Private/Windows/UbaDetoursUtilsWin.h
@@ -19,9 +19,9 @@ namespace uba
 		bool ownsFileInfo = false;
 		TString newName;
 	};
-	BlockAllocator<FileObject> g_fileObjectAllocator(g_memoryBlock);
-	void* FileObject::operator new(size_t size) { return g_fileObjectAllocator.Allocate(); }
-	void FileObject::operator delete(void* p) { g_fileObjectAllocator.Free(p); }
+	extern BlockAllocator<FileObject> g_fileObjectAllocator;
+	inline void* FileObject::operator new(size_t size) { return g_fileObjectAllocator.Allocate(); }
+	inline void FileObject::operator delete(void* p) { g_fileObjectAllocator.Free(p); }
 
 
 	enum HandleType
@@ -29,7 +29,10 @@ namespace uba
 		HandleType_File,
 		HandleType_FileMapping,
 		HandleType_Process,
-		HandleType_Std,
+		HandleType_StdErr,
+		HandleType_StdOut,
+		HandleType_StdIn,
+		// Std handle types must be last
 	};
 
 
@@ -101,7 +104,8 @@ namespace uba
 		}
 
 		void Write(struct DetouredHandle& handle, LPCVOID lpBuffer, u64 nNumberOfBytesToWrite);
-		void EnsureCommited(struct DetouredHandle& handle, u64 size);
+		void EnsureCommitted(const struct DetouredHandle& handle, u64 size);
+		void Remap(const struct DetouredHandle& handle, u64 size);
 
 		u64 fileIndex = ~u64(0);
 		u64 fileTime = ~u64(0);
