[UBA]
* Fixed some logging
* Added check for "isdone" inside lock in cancel call in case done message has already been sent

diff --git a/Engine/Source/Programs/UnrealBuildAccelerator/Common/Private/UbaSessionClient.cpp b/Engine/Source/Programs/UnrealBuildAccelerator/Common/Private/UbaSessionClient.cpp
--- a/Engine/Source/Programs/UnrealBuildAccelerator/Common/Private/UbaSessionClient.cpp
+++ b/Engine/Source/Programs/UnrealBuildAccelerator/Common/Private/UbaSessionClient.cpp
@@ -1556,6 +1556,8 @@ namespace uba
 					if (rec.isKilled || rec.isDone)
 						continue;
 					SCOPED_WRITE_LOCK(rec.lock, lock);
+					if (rec.isDone)
+						continue;
 					rec.handle.Cancel(true);
 					rec.isKilled = true;
 					SendReturnProcess(rec.handle.GetId(), TC("Running out of memory"));
diff --git a/Engine/Source/Programs/UnrealBuildAccelerator/Common/Private/UbaSessionServer.cpp b/Engine/Source/Programs/UnrealBuildAccelerator/Common/Private/UbaSessionServer.cpp
--- a/Engine/Source/Programs/UnrealBuildAccelerator/Common/Private/UbaSessionServer.cpp
+++ b/Engine/Source/Programs/UnrealBuildAccelerator/Common/Private/UbaSessionServer.cpp
@@ -1257,7 +1257,7 @@ namespace uba
 		if (session.enabled)
 			++m_availableRemoteSlotCount;
 
-		m_logger.Detail(TC("Client %s returned process to queue (%s)"), session.name.c_str(), reason.data);
+		m_logger.Detail(TC("Client %s returned process %u to queue (%s)"), session.name.c_str(), processId, reason.data);
 		++m_returnedRemoteProcessCount;
 
 		process->m_executingHost.clear();
