[UBA]
* Fixed compile warning on linux

diff --git a/Engine/Source/Programs/UnrealBuildAccelerator/Common/Private/UbaFileMapping.cpp b/Engine/Source/Programs/UnrealBuildAccelerator/Common/Private/UbaFileMapping.cpp
--- a/Engine/Source/Programs/UnrealBuildAccelerator/Common/Private/UbaFileMapping.cpp
+++ b/Engine/Source/Programs/UnrealBuildAccelerator/Common/Private/UbaFileMapping.cpp
@@ -543,7 +543,7 @@ namespace uba
 		if (!view.handle.IsValid())
 			return;
 
-		auto unmap = [=](const tchar* hint)
+		auto unmap = [=, this](const tchar* hint)
 			{
 				u8 storageIndex = 255;
 				File& file = GetFile(view.handle, storageIndex);
