[UBA]
* Added -testcompress and -testdecompress to UbaCli. This can be used to test obj file compression
* Added 4k alignment to decompressed memory for .h.obj files .. don't remember if/why this was needed but is done in UbaSession as well.
* Changed WrappedMain to return a bool instead to make it easier to exit on error

diff --git a/Engine/Source/Programs/UnrealBuildAccelerator/Cli/Private/UbaCli.cpp b/Engine/Source/Programs/UnrealBuildAccelerator/Cli/Private/UbaCli.cpp
--- a/Engine/Source/Programs/UnrealBuildAccelerator/Cli/Private/UbaCli.cpp
+++ b/Engine/Source/Programs/UnrealBuildAccelerator/Cli/Private/UbaCli.cpp
@@ -3,6 +3,7 @@
 #include "UbaCacheClient.h"
 #include "UbaCacheServer.h"
 #include "UbaClient.h"
+#include "UbaCompressedObjFileHeader.h"
 #include "UbaCoordinatorWrapper.h"
 #include "UbaFileAccessor.h"
 #include "UbaNetworkBackendTcp.h"
@@ -42,7 +43,7 @@ namespace uba
 		}();
 	u32				DefaultProcessorCount = []() { return GetLogicalProcessorCount(); }();
 
-	int PrintHelp(const tchar* message)
+	bool PrintHelp(const tchar* message)
 	{
 		LoggerWithWriter logger(g_consoleLogWriter, TC(""));
 		if (*message)
@@ -94,7 +95,7 @@ namespace uba
 		logger.Info(TC("   -oidc=<name>            Name of oidc"));
 		logger.Info(TC("   -maxcores=<number>      Max number of cores that will be asked for from coordinator"));
 		logger.Info(TC(""));
-		return -1;
+		return false;
 	}
 
 	StorageServer* g_storageServer;
@@ -146,7 +147,7 @@ namespace uba
 	//}
 
 
-	int WrappedMain(int argc, tchar* argv[])
+	bool WrappedMain(int argc, tchar* argv[])
 	{
 		using namespace uba;
 		//SetUnhandledExceptionFilter(UbaUnhandledExceptionFilter);
@@ -177,6 +178,8 @@ namespace uba
 		TString checkFileTable;
 		TString cacheFilterString;
 		TString cacheCommand;
+		TString testCompress;
+		TString testDecompress;
 
 		u32 loopCount = 1;
 
@@ -354,6 +357,22 @@ namespace uba
 			{
 				checkAws = true;
 			}
+			else if (name.Equals(TC("-testcompress")))
+			{
+				if (value.IsEmpty())
+					return PrintHelp(TC("-testCompress needs a value"));
+				testCompress = value.data;
+			}
+			else if (name.Equals(TC("-testdecompress")))
+			{
+				if (value.IsEmpty())
+				{
+					if (testCompress.empty())
+						return PrintHelp(TC("-testDecompress needs a value"));
+					value.Clear().Append(g_rootDir).EnsureEndsWithSlash().Append(TC("castemp")).EnsureEndsWithSlash().Append(TC("TestCompress.tmp"));
+				}
+				testDecompress = value.data;
+			}
 			else if (name.Equals(TC("-deletecas")))
 			{
 				deleteCas = true;
@@ -433,8 +452,7 @@ namespace uba
 			storageInfo.casCapacityBytes = 0;
 			storageInfo.storeCompressed = storeCompressed;
 			StorageImpl storage(storageInfo);
-			bool success = storage.CheckCasContent(DefaultProcessorCount);
-			return success ? 0 : -1;
+			return storage.CheckCasContent(DefaultProcessorCount);
 		}
 
 		if (!checkFileTable.empty())
@@ -444,9 +462,8 @@ namespace uba
 			storageInfo.storeCompressed = storeCompressed;
 			StorageImpl storage(storageInfo);
 			if (!storage.LoadCasTable())
-				return -1;
-			bool success = storage.CheckFileTable(checkFileTable.data(), DefaultProcessorCount);
-			return success ? 0 : -1;
+				return false;
+			return storage.CheckFileTable(checkFileTable.data(), DefaultProcessorCount);
 		}
 
 		if (checkCas2) // Creates a storage server and storage client and transfer _all_ cas files over network
@@ -488,7 +505,7 @@ namespace uba
 						}, 1, TC(""));
 				});
 			workManager.FlushWork();
-			return success ? 0 : -1;
+			return success;
 		}
 
 #if UBA_USE_AWS
@@ -507,10 +524,93 @@ namespace uba
 			}
 			else
 				logger.Info(TC("Seems like we are not running inside aws."));
-			return 0;
+			return true;
 		}
 #endif
 		
+		u64 testCompressOriginalSize = 0;
+		if (!testCompress.empty())
+		{
+			WorkManagerImpl workManager(DefaultProcessorCount);
+
+			FileAccessor fa(logger, testCompress.c_str());
+			if (!fa.OpenMemoryRead())
+				return logger.Error(TC("Failed to open file %s"), testCompress.c_str());
+			u64 fileSize = fa.GetSize();
+			u8* mem = fa.GetData();
+
+			testCompressOriginalSize = fileSize;
+
+			StorageCreateInfo storageInfo(g_rootDir.data, logWriter);
+			storageInfo.casCapacityBytes = 0;
+			storageInfo.storeCompressed = storeCompressed;
+			storageInfo.workManager = &workManager;
+			StorageImpl storage(storageInfo);
+
+			Storage::WriteResult res;
+			CompressedObjFileHeader header { CalculateCasKey(mem, fileSize, true, &workManager, testCompress.c_str()) };
+
+			StringBuffer<> dest;
+			dest.Append(storage.GetTempPath()).Append(TC("TestCompress.tmp"));
+			if (!storage.WriteCompressed(res, TC("MemoryMap"), InvalidFileHandle, mem, fileSize, dest.data, &header, sizeof(header), 0))
+				return false;
+			if (testDecompress.empty())
+				return true;
+		}
+
+		if (!testDecompress.empty())
+		{
+			WorkManagerImpl workManager(DefaultProcessorCount);
+
+			FileAccessor fa(logger, testDecompress.c_str());
+			if (!fa.OpenMemoryRead())
+				return logger.Error(TC("Failed to open file %s"), testDecompress.c_str());
+			u64 fileSize = fa.GetSize();
+			u8* mem = fa.GetData();
+
+			StorageCreateInfo storageInfo(g_rootDir.data, logWriter);
+			storageInfo.casCapacityBytes = 0;
+			storageInfo.storeCompressed = storeCompressed;
+			storageInfo.workManager = &workManager;
+			StorageImpl storage(storageInfo);
+
+			auto& h = *(CompressedObjFileHeader*)mem;
+			if (!h.IsValid())
+				return logger.Error(TC("File %s is not a compressed file"), testDecompress.c_str());
+
+			BinaryReader reader(mem, 0, fileSize);
+			reader.Skip(sizeof(CompressedObjFileHeader));
+			u64 decompressedSize = reader.ReadU64();
+
+			if (testCompressOriginalSize && decompressedSize != testCompressOriginalSize)
+				return logger.Error(TC("Compressed file %s has wrong decompressed size."), testDecompress.c_str());
+
+			StringBuffer<> dest;
+			dest.Append(storage.GetTempPath()).Append(TC("TestDecompress.tmp"));
+			FileAccessor faDest(logger, dest.data);
+			if (!faDest.CreateMemoryWrite(false, DefaultAttributes(), decompressedSize))
+				return false;
+			u8* destMem = faDest.GetData();
+
+			OO_SINTa decoredMemSize = OodleLZDecoder_MemorySizeNeeded(OodleLZ_Compressor_Kraken);
+			void* decoderMem = malloc(decoredMemSize);
+			auto mg = MakeGuard([decoderMem]() { free(decoderMem); });
+
+			while (reader.GetLeft())
+			{
+				u32 compressedBlockSize = reader.ReadU32();
+				u32 decompressedBlockSize = reader.ReadU32();
+
+				OO_SINTa decompLen = OodleLZ_Decompress(reader.GetPositionData(), (OO_SINTa)compressedBlockSize, destMem, (OO_SINTa)decompressedBlockSize, OodleLZ_FuzzSafe_Yes, OodleLZ_CheckCRC_No, OodleLZ_Verbosity_None, NULL, 0, NULL, NULL, decoderMem, decoredMemSize);
+				if (decompLen != decompressedBlockSize)
+					return logger.Error(TC("Failed to decompress %s"), testDecompress.c_str());
+				destMem += decompressedBlockSize;
+				reader.Skip(compressedBlockSize);
+			}
+
+			return faDest.Close();
+		}
+
 		if (commandType == CommandType_NotSet)
 		{
 			const tchar* errorMsg = argc == 1 ? TC("") : TC("\nERROR: First argument must be command type. Options are 'local,remote or native'");
@@ -556,7 +656,7 @@ namespace uba
 				logger.Info(TC("  Is64Bit: %s"), (is64Bit ? TC("true") : TC("false")));
 				logger.Info(TC("  Size: %llu"), fileSize);
 				logger.Info(TC("  CasKey: %s"), CasKeyString(key).str);
-				return 0;
+				return true;
 			}
 		}
 
@@ -594,7 +694,7 @@ namespace uba
 		NetworkServer& networkServer = *new NetworkServer(ctorSuccess, nsci);
 		auto destroyServer = MakeGuard([&]() { delete &networkServer; });
 		if (!ctorSuccess)
-			return -1;
+			return false;
 
 		if (!crypto.empty())
 		{
@@ -647,20 +747,15 @@ namespace uba
 		{
 			CreateCacheClient();
 			if (!cacheClient->GetClient().Connect(networkBackend, cacheHost.data, DefaultCachePort))
-			{
-				logger.Error(TC("Failed to connect to cache server"));
-				return -1;
-			}
+				return logger.Error(TC("Failed to connect to cache server"));
 
 			if (!storageServer.LoadCasTable(true))
-				return -1;
+				return false;
 
 			if (!cacheCommand.empty())
 			{
 				LoggerWithWriter consoleLogger(g_consoleLogWriter);
-				if (!cacheClient->ExecuteCommand(consoleLogger, cacheCommand.data()))
-					return -1;
-				return 0;
+				return cacheClient->ExecuteCommand(consoleLogger, cacheCommand.data());
 			}
 
 			if (writeCacheSummary)
@@ -670,13 +765,13 @@ namespace uba
 				CreateGuid(guid);
 				tempFile.Append(GuidToString(guid).str).Append(TC(".txt"));
 				if (!cacheClient->ExecuteCommand(logger, TC("content"), tempFile.data, cacheFilterString.data()))
-					return -1;
+					return false;
 				logger.Info(TC("Cache status summary written to %s"), tempFile.data);
 
 				#if PLATFORM_WINDOWS
 				ShellExecuteW(NULL, L"open", tempFile.data, NULL, NULL, SW_SHOW);
 				#endif
-				return 0;
+				return true;
 			}
 		}
 
@@ -694,9 +789,9 @@ namespace uba
 		{
 			if (!storageServer.m_casTableLoaded)
 				if (!storageServer.LoadCasTable(true))
-					return -1;
+					return false;
 			if (!networkServer.StartListen(networkBackend, port, listenIp.data))
-				return -1;
+				return false;
 		}
 		auto stopServer = MakeGuard([&]() { networkServer.DisconnectClients(); });
 
@@ -942,7 +1037,7 @@ namespace uba
 				}
 			}
 			if (!success)
-				return -1;
+				return false;
 
 			if (false)
 				networkServer.DisconnectClients();
@@ -959,20 +1054,20 @@ namespace uba
 		}
 		logger.EndScope();
 
-		return 0;
+		return true;
 	}
 }
 
 #if PLATFORM_WINDOWS
 int wmain(int argc, wchar_t* argv[])
 {
-	int res = uba::WrappedMain(argc, argv);
+	int res = uba::WrappedMain(argc, argv) ? 0 : -1;
 	Sleep(1); // Here to be able to put a breakpoint just before exit :-)
 	return res;
 }
 #else
 int main(int argc, char* argv[])
 {
-	return uba::WrappedMain(argc, argv);
+	return uba::WrappedMain(argc, argv) ? 0 : -1;
 }
 #endif
diff --git a/Engine/Source/Programs/UnrealBuildAccelerator/Detours/Private/Windows/UbaDetoursObjFilesPreloader.h b/Engine/Source/Programs/UnrealBuildAccelerator/Detours/Private/Windows/UbaDetoursObjFilesPreloader.h
--- a/Engine/Source/Programs/UnrealBuildAccelerator/Detours/Private/Windows/UbaDetoursObjFilesPreloader.h
+++ b/Engine/Source/Programs/UnrealBuildAccelerator/Detours/Private/Windows/UbaDetoursObjFilesPreloader.h
@@ -191,6 +191,10 @@ namespace uba
 				preload.objWriteOffset = 0;
 				preload.objLeft = preload.objDecompressedSize;
 				CloseHandle(objFileMappingHandle);
+
+				if (ToView(preload.fileInfo->originalName).EndsWith(TC(".h.obj"))) // TODO: Copied this from UbaSession align code
+					totalMemSize = AlignUp(totalMemSize, 4 * 1024);
+
 				totalMemSize += preload.objDecompressedSize;
 			}
 
@@ -235,6 +239,9 @@ namespace uba
 
 				if (preload.objWriteOffset == 0)
 				{
+					if (ToView(info.originalName).EndsWith(TC(".h.obj"))) // TODO: Copied this from UbaSession align code
+						m_totalMemOffset = AlignUp(m_totalMemOffset, 4 * 1024);
+
 					info.fileMapMem = m_totalMem + m_totalMemOffset;
 					m_totalMemOffset += preload.objDecompressedSize;
 				}
@@ -253,7 +260,10 @@ namespace uba
 
 				OO_SINTa decompLen = OodleLZ_Decompress(reader.GetPositionData(), (OO_SINTa)compressedBlockSize, destMem, (OO_SINTa)decompressedBlockSize,
 					OodleLZ_FuzzSafe_Yes, OodleLZ_CheckCRC_No, OodleLZ_Verbosity_None, NULL, 0, NULL, NULL, decoderMem, decoredMemSize);
-				UBA_ASSERTF(decompLen == decompressedBlockSize, TC("Failed to decompress .obj file %s (%s)"), info.name, info.originalName);(void)decompLen;
+				if (decompLen != decompressedBlockSize)
+				{
+					FatalError(1356, TC("Failed to decompress .obj file %s (%s)"), info.name, info.originalName);
+				}
 
 				bool isDone = preload.objLeft.fetch_sub(decompressedBlockSize) == decompressedBlockSize;
 				if (!isDone)
