[UBA]
* Changed so all remote processes that fails that have a failed message will send back for retry

diff --git a/Engine/Source/Programs/UnrealBuildAccelerator/Common/Private/UbaProcess.cpp b/Engine/Source/Programs/UnrealBuildAccelerator/Common/Private/UbaProcess.cpp
--- a/Engine/Source/Programs/UnrealBuildAccelerator/Common/Private/UbaProcess.cpp
+++ b/Engine/Source/Programs/UnrealBuildAccelerator/Common/Private/UbaProcess.cpp
@@ -368,6 +368,11 @@ namespace uba
 		#endif
 	}
 
+	bool ProcessImpl::HasFailedMessage()
+	{
+		return !m_messageSuccess;
+	}
+
 	bool ProcessImpl::WaitForExit(u32 millisecondsTimeout)
 	{
 		return m_messageThread.Wait(millisecondsTimeout);
diff --git a/Engine/Source/Programs/UnrealBuildAccelerator/Common/Private/UbaSessionClient.cpp b/Engine/Source/Programs/UnrealBuildAccelerator/Common/Private/UbaSessionClient.cpp
--- a/Engine/Source/Programs/UnrealBuildAccelerator/Common/Private/UbaSessionClient.cpp
+++ b/Engine/Source/Programs/UnrealBuildAccelerator/Common/Private/UbaSessionClient.cpp
@@ -1713,6 +1713,13 @@ namespace uba
 									session.SendReturnProcess(rec->handle.GetId(), session.m_terminationReason);
 								return;
 							}
+
+							if (process.HasFailedMessage()) // If there are failure caused by failed messages we send back for retry
+							{
+								if (session.m_loop)
+									session.SendReturnProcess(rec->handle.GetId(), TC("Failed message"));
+								return;
+							}
 						}
 
 						if (exitCode == 0 || startInfo.writeOutputFilesOnFail)
@@ -2006,8 +2013,10 @@ namespace uba
 
 	bool SessionClient::LogLine(ProcessImpl& process, const tchar* line, LogEntryType logType)
 	{
-		// TODO: Remove this once we have figured out a bug that seems to exist for remote execution
-#if PLATFORM_WINDOWS
+		// Remove this once we have figured out a bug that seems to exist for remote execution
+		// Update: Bug has been found for macos... for windows we believe the bug is related to uninformed shutdown and having multiple tcp connections..
+		// ... one tcp connection is disconnected, causing file not found while another connection manages to send "process finished"
+#if 0 // PLATFORM_WINDOWS
 
 		auto rules = process.m_startInfo.rules;
 		if (!rules)
diff --git a/Engine/Source/Programs/UnrealBuildAccelerator/Common/Private/UbaStorageClient.cpp b/Engine/Source/Programs/UnrealBuildAccelerator/Common/Private/UbaStorageClient.cpp
--- a/Engine/Source/Programs/UnrealBuildAccelerator/Common/Private/UbaStorageClient.cpp
+++ b/Engine/Source/Programs/UnrealBuildAccelerator/Common/Private/UbaStorageClient.cpp
@@ -314,7 +314,7 @@ namespace uba
 				{
 					if (proxy)
 						continue;
-					return m_logger.Error(TC("Failed to send fetch begin message for cas %s (%s)"), casFile.data, hint);
+					return m_logger.Error(TC("Failed to send fetch begin message for cas %s (%s). Error: %u"), casFile.data, hint, msg.GetError());
 				}
 				sizeOfFirstMessage = u32(reader.GetLeft());
 				fetchId = reader.ReadU16();
@@ -347,7 +347,6 @@ namespace uba
 					{
 						reader.ReadString(proxyHost);
 						proxyPort = reader.ReadU16();
-						// TODO: Check if self, then use loopback.. this could be received before local proxy is started due to network ordering
 					}
 
 					SCOPED_WRITE_LOCK(m_proxyClientLock, proxyLock2);
diff --git a/Engine/Source/Programs/UnrealBuildAccelerator/Common/Public/UbaProcess.h b/Engine/Source/Programs/UnrealBuildAccelerator/Common/Public/UbaProcess.h
--- a/Engine/Source/Programs/UnrealBuildAccelerator/Common/Public/UbaProcess.h
+++ b/Engine/Source/Programs/UnrealBuildAccelerator/Common/Public/UbaProcess.h
@@ -64,6 +64,7 @@ namespace uba
 
 		bool IsActive();
 		bool IsCancelled();
+		bool HasFailedMessage();
 
 		bool WaitForRead(PipeReader& outReader, PipeReader& errReader);
 		void SetWritten();
