[UBA]
* Improved logging on UbaCacheServer

diff --git a/Engine/Source/Programs/UnrealBuildAccelerator/Common/Private/UbaCacheServer.cpp b/Engine/Source/Programs/UnrealBuildAccelerator/Common/Private/UbaCacheServer.cpp
--- a/Engine/Source/Programs/UnrealBuildAccelerator/Common/Private/UbaCacheServer.cpp
+++ b/Engine/Source/Programs/UnrealBuildAccelerator/Common/Private/UbaCacheServer.cpp
@@ -184,7 +184,7 @@ namespace uba
 		}
 
 		u64 duration = GetTime() - startTime;
-		m_logger.Detail(TC("Database loaded from %s (v%u)  in %s (%llu bucket(s) containing %s paths, %s keys, %llu cache entries)"), fileName.data, databaseVersion, TimeToText(duration).str, m_buckets.size(), BytesToText(stats.totalPathTableSize).str, BytesToText(stats.totalCasKeyTableSize).str, stats.totalCacheEntryCount.load());
+		m_logger.Detail(TC("Database loaded from %s (v%u)  in %s (%llu bucket(s) containing %s paths, %s keys, %s cache entries)"), fileName.data, databaseVersion, TimeToText(duration).str, m_buckets.size(), BytesToText(stats.totalPathTableSize).str, BytesToText(stats.totalCasKeyTableSize).str, CountToText(stats.totalCacheEntryCount.load()).str);
 		return true;
 	}
 
@@ -459,10 +459,10 @@ namespace uba
 			u64 hits = m_cacheKeyHitCount;
 			u64 miss = m_cacheKeyFetchCount - hits;
 			m_logger.Info(TC("Stats since boot (%s ago)"), TimeToText(startTime - m_bootTime, true).str);
-			m_logger.Info(TC("  CacheServer %llu hits, %llu misses"), hits, miss);
+			m_logger.Info(TC("  CacheServer %s hits, %s misses"), CountToText(hits).str, CountToText(miss).str);
 			u64 recvCount = storageStats.sendCas.count.load();
 			u64 sendCount = storageStats.recvCas.count.load();
-			m_logger.Info(TC("  StorageServer cas %llu (%s) sent, %llu (%s) received"), recvCount, BytesToText(storageStats.sendCasBytesComp).str, sendCount, BytesToText(storageStats.recvCasBytesComp).str);
+			m_logger.Info(TC("  StorageServer cas %s (%s) sent, %s (%s) received"), CountToText(recvCount).str, BytesToText(storageStats.sendCasBytesComp).str, CountToText(sendCount).str, BytesToText(storageStats.recvCasBytesComp).str);
 		}
 
 		if (m_shouldWipe)
@@ -518,9 +518,9 @@ namespace uba
 		}
 
 		if (removedNonExisting)
-			m_logger.Detail(TC("  Removed %u cas entries (marked as not existing)"), removedNonExisting);
+			m_logger.Detail(TC("  Removed %s cas entries (marked as not existing)"), CountToText(removedNonExisting).str);
 
-		m_logger.Detail(TC("  Found %llu cas files and %llu deleted by overflow (%s)"), existingCas.size(), deletedCasFiles.size(), BytesToText(totalCasSize).str);
+		m_logger.Detail(TC("  Found %s cas files and %s deleted by overflow (%s)"), CountToText(existingCas.size()).str, CountToText(deletedCasFiles.size()).str, BytesToText(totalCasSize).str);
 		u64 totalCasCount = existingCas.size() + deletedCasCount;
 
 		if (shouldExit())
@@ -762,7 +762,7 @@ namespace uba
 		if (overflowedEntryCount)
 			m_logger.Detail(TC("  Found %llu overflowed cache entries"), overflowedEntryCount.load());
 		if (expiredEntryCount)
-			m_logger.Detail(TC("  Found %llu expired cache entries (older than %s)"), expiredEntryCount.load(), TimeToText(MsToTime(m_expirationTimeSeconds*1000), true).str);
+			m_logger.Detail(TC("  Found %llu expired cache entries"), expiredEntryCount.load());
 		if (missingOutputEntryCount)
 			m_logger.Detail(TC("  Found %llu cache entries with missing output cas"), missingOutputEntryCount.load());
 		if (missingInputEntryCount)
@@ -783,7 +783,7 @@ namespace uba
 
 			if (!bucket.hasDeletedEntries && !forceAllSteps)
 			{
-				m_logger.Detail(TC("    Bucket %u skipped updating. (%llu entries)"), bucket.index, bucket.totalEntryCount.load());
+				m_logger.Detail(TC("    Bucket %u skipped updating. (%s entries)"), bucket.index, CountToText(bucket.totalEntryCount.load()).str);
 				return;
 			}
 			bucket.hasDeletedEntries = false;
@@ -815,7 +815,7 @@ namespace uba
 					collectUsedCasKeyOffsets(entry.outputCasKeyOffsets);
 				}
 			}
-			m_logger.Detail(TC("    Bucket %u Collected %llu used caskeys. (%s)"), bucket.index, usedCasKeyOffsets.size(), TimeToText(GetTime() - collectUsedCasKeysStart).str);
+			m_logger.Detail(TC("    Bucket %u Collected %s used caskeys. (%s)"), bucket.index, CountToText(usedCasKeyOffsets.size()).str, TimeToText(GetTime() - collectUsedCasKeysStart).str);
 
 			u64 recreatePathTableStart = GetTime();
 
@@ -913,7 +913,7 @@ namespace uba
 
 			bucket.needsSave = true;
 
-			m_logger.Info(TC("    Bucket %u Done (%s). CacheEntries: %llu (%s) PathTable: %s CasTable: %s"), bucket.index, TimeToText(GetTime() - bucketStartTime).str, bucket.totalEntryCount.load(), BytesToText(bucket.totalEntrySize.load()).str, BytesToText(bucket.m_pathTable.GetSize()).str, BytesToText(bucket.m_casKeyTable.GetSize()).str);
+			m_logger.Info(TC("    Bucket %u Done (%s). CacheEntries: %s (%s) PathTable: %s CasTable: %s Expiration: %s"), bucket.index, TimeToText(GetTime() - bucketStartTime).str, CountToText(bucket.totalEntryCount.load()).str, BytesToText(bucket.totalEntrySize.load()).str, BytesToText(bucket.m_pathTable.GetSize()).str, BytesToText(bucket.m_casKeyTable.GetSize()).str, TimeToText(MsToTime(bucket.expirationTimeSeconds*1000), true).str);
 
 			SCOPED_WRITE_LOCK(existingCasLock, l); // Just reusing lock for other purpose
 			maxCommittedMemory = Max(maxCommittedMemory, memoryBlock.writtenSize);
@@ -941,7 +941,7 @@ namespace uba
 		u64 oldestTime = oldest ? GetFileTimeAsTime(now - (m_creationTime + oldest)) : 0;
 		u64 longestUnusedTime = longestUnused ? GetFileTimeAsTime(now - (m_creationTime + longestUnused)) : 0;
 		u64 duration = GetTime() - startTime;
-		m_logger.Info(TC("Maintenance done! (%s) CasFiles: %llu (%s) Entries: %llu Oldest: %s LongestUnused: %s MaintenanceMem: %s"), TimeToText(duration).str, totalCasCount - deletedCasCount, BytesToText(totalCasSize).str, totalEntryCount.load(), TimeToText(oldestTime, true).str, TimeToText(longestUnusedTime, true).str, BytesToText(maxCommittedMemory).str);
+		m_logger.Info(TC("Maintenance done! (%s) CasFiles: %s (%s) Entries: %s Oldest: %s LongestUnused: %s MaintenanceMem: %s"), TimeToText(duration).str, CountToText(totalCasCount - deletedCasCount).str, BytesToText(totalCasSize).str, CountToText(totalEntryCount.load()).str, TimeToText(oldestTime, true).str, TimeToText(longestUnusedTime, true).str, BytesToText(maxCommittedMemory).str);
 		
 		m_longestMaintenance = Max(m_longestMaintenance, duration);
 
diff --git a/Engine/Source/Programs/UnrealBuildAccelerator/Common/Private/UbaLogger.cpp b/Engine/Source/Programs/UnrealBuildAccelerator/Common/Private/UbaLogger.cpp
--- a/Engine/Source/Programs/UnrealBuildAccelerator/Common/Private/UbaLogger.cpp
+++ b/Engine/Source/Programs/UnrealBuildAccelerator/Common/Private/UbaLogger.cpp
@@ -320,6 +320,20 @@ namespace uba
 			TSprintf_s(str, 32, TC("%.1ftb"), double(bytes) / (1000ull * 1000 * 1000 * 1000));
 	}
 
+	CountToText::CountToText(u64 count)
+	{
+		if (count < 1000)
+			TSprintf_s(str, 32, TC("%u"), u32(count));
+		else if (count < 1000 * 1000)
+			TSprintf_s(str, 32, TC("%.1fk"), double(count) / 1000ull);
+		else if (count < 1000ull * 1000 * 1000)
+			TSprintf_s(str, 32, TC("%.1fm"), double(count) / (1000ull * 1000));
+		else if (count < 1000ull * 1000 * 1000 * 1000)
+			TSprintf_s(str, 32, TC("%.1fg"), double(count) / (1000ull * 1000 * 1000));
+		else
+			TSprintf_s(str, 32, TC("%.1ft"), double(count) / (1000ull * 1000 * 1000 * 1000));
+	}
+
 #if UBA_DEBUG_LOGGER
 	thread_local u32 t_debugLogScopeCount = 0;
 
diff --git a/Engine/Source/Programs/UnrealBuildAccelerator/Common/Public/UbaLogger.h b/Engine/Source/Programs/UnrealBuildAccelerator/Common/Public/UbaLogger.h
--- a/Engine/Source/Programs/UnrealBuildAccelerator/Common/Public/UbaLogger.h
+++ b/Engine/Source/Programs/UnrealBuildAccelerator/Common/Public/UbaLogger.h
@@ -73,6 +73,13 @@ namespace uba
 		tchar str[32];
 	};
 
+	struct CountToText
+	{
+		CountToText(u64 count);
+		operator const tchar* () const { return str; };
+		tchar str[32];
+	};
+
 
 	#if UBA_DEBUG_LOGGER
 	bool StartDebugLogger(Logger& outerLogger, const tchar* fileName);
