[UBA]
* Improved ubavisualizer output
* Fixed visualizer crash
* Fixed support for dpi scaling through remote desktop
* Fixed so visualizer never ends up outside screen(s) boundaries
*

diff --git a/Engine/Source/Programs/UnrealBuildAccelerator/Visualizer/Private/UbaTraceReader.cpp b/Engine/Source/Programs/UnrealBuildAccelerator/Visualizer/Private/UbaTraceReader.cpp
--- a/Engine/Source/Programs/UnrealBuildAccelerator/Visualizer/Private/UbaTraceReader.cpp
+++ b/Engine/Source/Programs/UnrealBuildAccelerator/Visualizer/Private/UbaTraceReader.cpp
@@ -8,9 +8,18 @@
 
 namespace uba
 {
-	TraceView::Process* TraceView::GetProcess(const ProcessLocation& loc)
+	const TraceView::Process& TraceView::GetProcess(const ProcessLocation& loc)
 	{
-		return &(sessions[loc.sessionIndex].processors[loc.processorIndex].processes[loc.processIndex]);
+		static TraceView::Process emptyProcess;
+		if (loc.sessionIndex >= sessions.size())
+			return emptyProcess;
+		auto& session = sessions[loc.sessionIndex];
+		if (loc.processorIndex >= session.processors.size())
+			return emptyProcess;
+		auto& processor = session.processors[loc.processorIndex];
+		if (loc.processIndex >= processor.processes.size())
+			return emptyProcess;
+		return processor.processes[loc.processIndex];
 	}
 
 	void TraceView::Clear()
diff --git a/Engine/Source/Programs/UnrealBuildAccelerator/Visualizer/Private/UbaTraceReader.h b/Engine/Source/Programs/UnrealBuildAccelerator/Visualizer/Private/UbaTraceReader.h
--- a/Engine/Source/Programs/UnrealBuildAccelerator/Visualizer/Private/UbaTraceReader.h
+++ b/Engine/Source/Programs/UnrealBuildAccelerator/Visualizer/Private/UbaTraceReader.h
@@ -126,7 +126,7 @@ namespace uba
 			bool success = false;
 		};
 
-		Process* GetProcess(const ProcessLocation& loc);
+		const Process& GetProcess(const ProcessLocation& loc);
 		void Clear();
 
 		Vector<Session> sessions;
diff --git a/Engine/Source/Programs/UnrealBuildAccelerator/Visualizer/Private/UbaVisualizer.cpp b/Engine/Source/Programs/UnrealBuildAccelerator/Visualizer/Private/UbaVisualizer.cpp
--- a/Engine/Source/Programs/UnrealBuildAccelerator/Visualizer/Private/UbaVisualizer.cpp
+++ b/Engine/Source/Programs/UnrealBuildAccelerator/Visualizer/Private/UbaVisualizer.cpp
@@ -6,7 +6,9 @@
 
 #include <uxtheme.h>
 #include <dwmapi.h>
+#include <shellscalingapi.h>
 
+#pragma comment (lib, "Shcore.lib")
 #pragma comment (lib, "UxTheme.lib")
 #pragma comment (lib, "Dwmapi.lib")
 #define WM_NEWTRACE WM_USER+1
@@ -60,6 +62,7 @@ namespace uba
 		config.GetValueAsU32(fontSize, L"FontSize");
 		config.GetValueAsString(fontName, L"FontName");
 		config.GetValueAsU32(maxActiveVisible, L"MaxActiveVisible");
+		config.GetValueAsU32(maxActiveProcessHeight, L"MaxActiveProcessHeight");
 		#define UBA_VISUALIZER_FLAG(name, defaultValue, desc) config.GetValueAsBool(show##name, L"Show" #name);
 		UBA_VISUALIZER_FLAGS1
 		#undef UBA_VISUALIZER_FLAG
@@ -79,6 +82,7 @@ namespace uba
 		config.AddValue(L"FontSize", fontSize);
 		config.AddValue(L"FontName", fontName.c_str());
 		config.AddValue(L"MaxActiveVisible", maxActiveVisible);
+		config.AddValue(L"MaxActiveProcessHeight", maxActiveProcessHeight);
 		#define UBA_VISUALIZER_FLAG(name, defaultValue, desc) config.AddValue(L"Show" #name, show##name);
 		UBA_VISUALIZER_FLAGS1
 		#undef UBA_VISUALIZER_FLAG
@@ -515,6 +519,11 @@ namespace uba
 
 	void Visualizer::ThreadLoop()
 	{
+		if (m_config.parent)
+		{
+			SetProcessDpiAwareness(PROCESS_SYSTEM_DPI_AWARE);
+		}
+
 		InitBrushes();
 
 		LOGBRUSH br = { 0 };
@@ -522,10 +531,24 @@ namespace uba
 		m_processUpdatePen = CreatePen(PS_SOLID, 2, RGB(GetRValue(br.lbColor), GetGValue(br.lbColor), GetBValue(br.lbColor)));
 
 		HINSTANCE hInstance = GetModuleHandle(NULL);
-		u32 winPosX = m_config.x;
-		u32 winPosY = m_config.y;
-		u32 winWidth = m_config.width;
-		u32 winHeight = m_config.height;
+		int winPosX = m_config.x;
+		int winPosY = m_config.y;
+		int winWidth = m_config.width;
+		int winHeight = m_config.height;
+
+		RECT rectCombined;
+		SetRectEmpty(&rectCombined);
+		EnumDisplayMonitors(0, 0, [](HMONITOR hMon,HDC hdc,LPRECT lprcMonitor,LPARAM pData)
+		{
+			RECT* rectCombined = reinterpret_cast<RECT*>(pData);
+			UnionRect(rectCombined, rectCombined, lprcMonitor);
+			return TRUE;
+		}, (LPARAM)&rectCombined);
+
+		winPosX = Max(int(rectCombined.left), winPosX);
+		winPosY = Max(int(rectCombined.top), winPosY);
+		winPosX = Min(int(rectCombined.right) - winWidth, winPosX);
+		winPosY = Min(int(rectCombined.bottom) - winHeight, winPosY);
 
 		WNDCLASSEX wndClassEx;
 		ZeroMemory(&wndClassEx, sizeof(wndClassEx));
@@ -546,7 +569,7 @@ namespace uba
 		UpdateProcessFont();
 
 		const TCHAR* fontName = TEXT("Consolas");
-		m_popupFont.handle = (HFONT)CreateFontW(-12, 0, 0, 0, FW_NORMAL, false, false, false, ANSI_CHARSET, OUT_DEFAULT_PRECIS, CLIP_DEFAULT_PRECIS, 0, FIXED_PITCH | FF_MODERN, fontName);
+		m_popupFont.handle = (HFONT)CreateFontW(-12, 0, 0, 0, FW_NORMAL, false, false, false, ANSI_CHARSET, OUT_DEFAULT_PRECIS, CLIP_DEFAULT_PRECIS, CLEARTYPE_QUALITY, FIXED_PITCH | FF_MODERN, fontName);
 		m_popupFont.height = 14;
 		//m_popupFont = (HFONT)GetStockObject(SYSTEM_FIXED_FONT);
 
@@ -668,7 +691,7 @@ namespace uba
 			m_scrollAtAnchorX = m_scrollPosX;
 			m_scrollAtAnchorY = m_scrollPosY;
 			SetCapture(m_hwnd);
-			RedrawWindow(m_hwnd, NULL, NULL, RDW_INVALIDATE);
+			Redraw(false);
 		}
 		++m_dragToScrollCounter;
 	}
@@ -682,7 +705,7 @@ namespace uba
 
 		ReleaseCapture();
 		if (UpdateSelection())
-			RedrawWindow(m_hwnd, NULL, NULL, RDW_INVALIDATE);
+			Redraw(false);
 	}
 
 	void Visualizer::SaveSettings()
@@ -721,19 +744,39 @@ namespace uba
 
 	void Visualizer::UpdateFont(Font& font, int height, bool createUnderline)
 	{
-		//NONCLIENTMETRICS nonClientMetrics;
-		//nonClientMetrics.cbSize = sizeof(nonClientMetrics);
-		//SystemParametersInfo(SPI_GETNONCLIENTMETRICS, sizeof(nonClientMetrics), &nonClientMetrics, 0);
-		//m_font = (HFONT)CreateFontIndirect(&nonClientMetrics.lfMessageFont);
 		font.height = height;
 
+		int fh = height;
+		font.offset = 0;
+		if (height <= 13)
+		{
+			++fh;
+			--font.offset;
+		}
+		if (height <= 11)
+			++fh;
+		if (height <= 9)
+			++fh;
+		if (height <= 8)
+			++fh;
+		if (height <= 6)
+			++fh;
+		if (height <= 4)
+			--font.offset;
+
 		if (font.handle)
 			DeleteObject(font.handle);
 		if (font.handleUnderlined)
 			DeleteObject(font.handleUnderlined);
-		font.handle = (HFONT)CreateFontW(4 - height, 0, 0, 0, FW_NORMAL, 0, 0, 0, ANSI_CHARSET, OUT_DEFAULT_PRECIS, CLIP_DEFAULT_PRECIS, CLEARTYPE_QUALITY, DEFAULT_PITCH, m_config.fontName.c_str());
+
+		//NONCLIENTMETRICS nonClientMetrics;
+		//nonClientMetrics.cbSize = sizeof(nonClientMetrics);
+		//SystemParametersInfo(SPI_GETNONCLIENTMETRICS, sizeof(nonClientMetrics), &nonClientMetrics, 0);
+		//m_font = (HFONT)CreateFontIndirect(&nonClientMetrics.lfMessageFont);
+
+		font.handle = (HFONT)CreateFontW(4 - fh, 0, 0, 0, FW_NORMAL, 0, 0, 0, ANSI_CHARSET, OUT_DEFAULT_PRECIS, CLIP_DEFAULT_PRECIS, CLEARTYPE_QUALITY, DEFAULT_PITCH, m_config.fontName.c_str());
 		if (createUnderline)
-			font.handleUnderlined = (HFONT)CreateFontW(4 - height, 0, 0, 0, FW_NORMAL, 0, 1, 0, ANSI_CHARSET, OUT_DEFAULT_PRECIS, CLIP_DEFAULT_PRECIS, CLEARTYPE_QUALITY, DEFAULT_PITCH, m_config.fontName.c_str());
+			font.handleUnderlined = (HFONT)CreateFontW(4 - fh, 0, 0, 0, FW_NORMAL, 0, 1, 0, ANSI_CHARSET, OUT_DEFAULT_PRECIS, CLIP_DEFAULT_PRECIS, CLEARTYPE_QUALITY, DEFAULT_PITCH, m_config.fontName.c_str());
 	}
 
 	void Visualizer::UpdateDefaultFont()
@@ -748,6 +791,7 @@ namespace uba
 	{
 		int fontHeight = int(m_zoomValue*20.0f);
 		UpdateFont(m_processFont, fontHeight, false);
+		m_processFont.height = fontHeight;
 		m_boxHeight = m_processFont.height + 2;
 		m_progressRectLeft = int(5 + float(m_processFont.height) * 1.8f);
 		DirtyBitmaps(true);
@@ -758,12 +802,19 @@ namespace uba
 		m_config.fontSize += offset;
 		m_config.fontSize = Max(m_config.fontSize, 10u);
 		UpdateDefaultFont();
-		Redraw();
+		Redraw(true);
 	}
 
-	void Visualizer::Redraw()
+	void Visualizer::Redraw(bool now)
 	{
-		RedrawWindow(m_hwnd, NULL, NULL, RDW_INVALIDATE | RDW_UPDATENOW);
+		u32 flags = RDW_INVALIDATE;
+		if (now)
+			flags |= RDW_UPDATENOW;
+		RedrawWindow(m_hwnd, NULL, NULL, flags);
+
+		u32 activeProcessCount = u32(m_trace.m_activeProcesses.size());
+		for (u32 i=0; i!=sizeof_array(m_activeProcessCountHistory); ++i)
+			m_activeProcessCountHistory[i] = activeProcessCount;
 	}
 
 	void Visualizer::PaintClient(const Function<void(HDC hdc, HDC memDC, RECT& clientRect)>& paintFunc)
@@ -832,7 +883,6 @@ namespace uba
 		{
 			progressRect.bottom -= m_defaultFont.height + 10;
 		}
-		bool shouldDrawText = m_zoomValue > 0.25f;
 
 		SetBkMode(hdc, TRANSPARENT);
 		SetTextColor(hdc, m_textColor);
@@ -861,7 +911,7 @@ namespace uba
 				RECT rect;
 				rect.left = posX;
 				rect.right = endX;
-				rect.top = posY;
+				rect.top = posY + m_activeFont.offset;
 				rect.bottom = posY + m_activeFont.height + 2;
 				SetTextColor(hdc, type == LogEntryType_Info ? m_textColor : (type == LogEntryType_Error ? m_textErrorColor : m_textWarningColor));
 
@@ -927,7 +977,7 @@ namespace uba
 		
 		if (m_config.showActiveProcesses && !m_trace.m_activeProcesses.empty())
 		{
-			auto drawBox = [&](u64 start, u64 stop, int height, bool selected)
+			auto drawBox = [&](u64 start, u64 stop, int height, bool selected, bool inProgress)
 				{
 					//int left = 3 + 6*m_activeFont.height;
 					//int right = rect.left + durationMs/100;
@@ -943,7 +993,7 @@ namespace uba
 					rect.right = right;
 					rect.top = posY;
 					rect.bottom = posY + height;
-					FillRect(hdc, &rect, m_processBrushes[selected].inProgress);
+					FillRect(hdc, &rect, inProgress ? m_processBrushes[selected].inProgress : m_processBrushes[selected].success);
 					return rect;
 				};
 
@@ -956,15 +1006,29 @@ namespace uba
 
 					bool selected = m_processSelected && m_processSelectedLocation == processLocation;
 
-					RECT boxRect = drawBox(process.start, process.stop, boxHeight, selected);
+					if (m_config.showFinishedProcesses)
+					{
+						u32 index = processLocation.processIndex;
+						while (index > 0)
+						{
+							--index;
+							TraceView::Process& process2 = session.processors[processLocation.processorIndex].processes[index];
+							drawBox(process2.start, process2.stop, boxHeight, false, false);
+						}
+					}
+
+					RECT boxRect = drawBox(process.start, process.stop, boxHeight, selected, true);
 
 					u32 v = boxHeight - 1;
 
-					if (v > 5)
+					if (v > 4)
 					{
 						u32 fontIndex = Min(v, u32(sizeof_array(m_activeProcessFont) - 1));
 						if (!m_activeProcessFont[fontIndex].handle)
-							UpdateFont(m_activeProcessFont[fontIndex], fontIndex+2, false);
+						{
+							UpdateFont(m_activeProcessFont[fontIndex], fontIndex - 1, false);
+							m_activeProcessFont[fontIndex].offset += 1;
+						}
 						if (firstWithHeight)
 							SetActiveFont(m_activeProcessFont[fontIndex]);
 
@@ -1169,6 +1233,8 @@ namespace uba
 
 			SetActiveFont(m_processFont);
 
+			bool shouldDrawText = m_zoomValue > 0.1f;
+
 			if (m_config.showProcessBars)
 			{
 				processLocation.processorIndex = 0;
@@ -1266,8 +1332,7 @@ namespace uba
 									
 									//SetTextAlign(textDC, TA_LEFT|TA_TOP|TA_NOUPDATECP|VTA_CENTER);
 									int textY = rect2.top;// - int(m_processFont.height/9);
-									if (m_processFont.height < 12)
-										textY += 1;
+									textY += m_processFont.offset;
 
 									bool dropShadow = m_config.DarkMode;
 									if (dropShadow)
@@ -1622,36 +1687,36 @@ namespace uba
 			}
 
 			top -= 2;
-			left -= 20;
+
+			auto drawText = [&](const tchar* text, COLORREF color)
+				{
+					SetTextColor(hdc, color);
+					RECT r{left, top, left+200, top+200};
+					auto strLen = TStrlen(text);
+					DrawTextW(hdc, text, strLen, &r, DT_SINGLELINE|DT_NOCLIP|DT_CALCRECT);
+					left -= (r.right - r.left) + 5;
+					r.left = left;
+					DrawTextW(hdc, text, strLen, &r, DT_SINGLELINE|DT_NOCLIP);
+				};
+
 			if (m_config.showCpuMemStats)
 			{
 				SetActiveFont(m_defaultFont);
-
-				SetTextColor(hdc, m_cpuColor);
-				ExtTextOutW(hdc, left, top, 0, NULL, L"CPU", 3, NULL);
-				left -= 25;
-
-				SetTextColor(hdc, m_memColor);
-				ExtTextOutW(hdc, left, top, 0, NULL, L"MEM", 3, NULL);
-				left -= 25;
+				drawText(L"CPU", m_cpuColor);
+				drawText(L"MEM", m_memColor);
 			}
 			if (m_config.showNetworkStats)
 			{
 				SetActiveFont(m_defaultFont);
-
-				SetTextColor(hdc, m_sendColor);
-				ExtTextOutW(hdc, left, top, 0, NULL, L"SND", 3, NULL);
-				left -= 25;
-
-				SetTextColor(hdc, m_recvColor);
-				ExtTextOutW(hdc, left, top, 0, NULL, L"RCV", 3, NULL);
+				drawText(L"SND", m_sendColor);
+				drawText(L"RCV", m_recvColor);
 			}
 			SetTextColor(hdc, m_textColor);
 		}
 
 		if (m_processSelected)
 		{
-			TraceView::Process& process = *m_traceView.GetProcess(m_processSelectedLocation);
+			const TraceView::Process& process = m_traceView.GetProcess(m_processSelectedLocation);
 			u64 duration = 0;
 			
 			Vector<TString> logLines;
@@ -1953,7 +2018,7 @@ namespace uba
 			maxHeight = Min(maxHeight, maxHeight2);
 		}
 
-		u32 maxSize = m_processFont.height;
+		u32 maxSize = Min(m_config.maxActiveProcessHeight, 32u);
 		u32 maxSizeMinusOne = maxSize - 1;
 
 		u32 counts[128] = { 0 };
@@ -2011,9 +2076,9 @@ namespace uba
 			}
 		}
 
-		if (fillHeight)
-			posY = startPosY; // To prevent vertical scrollbar
-		else if (counts[maxSizeMinusOne] != activeProcessCount)
+		//if (fillHeight)
+		//	posY = clientRect.bottom-1;//startPosY; // To prevent vertical scrollbar
+		if (fillHeight || counts[maxSizeMinusOne] != activeProcessCount)
 			posY = startPosY + maxHeight;
 		else
 			posY += 3;
@@ -2440,6 +2505,11 @@ namespace uba
 				return;
 		}
 
+		if (pos.y <= posY)
+			return;
+
+		outResult.section = 2;
+
 		SessionRec sortedSessions[1024];
 		Populate(sortedSessions, m_traceView, m_config.SortActiveRemoteSessions);
 
@@ -2665,7 +2735,7 @@ namespace uba
 			int timelineTop = GetTimelineTop(clientRect);
 			if (pos.y >= timelineTop)// && pos.y < timelineTop + 40)
 			{
-				outResult.section = 2;
+				outResult.section = 3;
 				float timeScale = (m_horizontalScaleValue * m_zoomValue)*50.0f;
 				float startOffset = -(m_scrollPosX / timeScale);
 				outResult.timelineSelected = startOffset + (pos.x - m_progressRectLeft) / timeScale;
@@ -2673,7 +2743,7 @@ namespace uba
 		}
 	}
 
-	void Visualizer::WriteProcessStats(Logger& out, TraceView::Process& process)
+	void Visualizer::WriteProcessStats(Logger& out, const TraceView::Process& process)
 	{
 		bool hasExited = process.stop != ~u64(0);
 		out.Info(L"  %ls", process.description.c_str());
@@ -2749,7 +2819,7 @@ namespace uba
 	void Visualizer::UnselectAndRedraw()
 	{
 		if (Unselect() || m_config.showCursorLine)
-			RedrawWindow(m_hwnd, NULL, NULL, RDW_INVALIDATE);
+			Redraw(false);
 	}
 
 	bool Visualizer::UpdateAutoscroll()
@@ -2875,7 +2945,7 @@ namespace uba
 
 			auto g = MakeGuard([&]()
 				{
-					RedrawWindow(hWnd, NULL, NULL, RDW_INVALIDATE|RDW_UPDATENOW);
+					Redraw(true);
 					UpdateScrollbars(true);
 				});
 
@@ -2890,7 +2960,6 @@ namespace uba
 			else if (!m_fileName.IsEmpty())
 			{
 				m_trace.ReadFile(m_traceView, m_fileName.data, m_replay != 0);
-				Redraw();
 				title.Append(m_fileName);
 				m_traceView.finished = m_replay == 0;
 			}
@@ -2972,7 +3041,7 @@ namespace uba
 			{
 				m_autoScroll = false;
 				KillTimer(m_hwnd, 0);
-				UpdateScrollbars(true);
+				changed = true;
 			}
 
 			changed = UpdateAutoscroll() || changed;
@@ -2982,7 +3051,7 @@ namespace uba
 				UpdateScrollbars(true);
 
 				u64 startTime = GetTime();
-				RedrawWindow(hWnd, NULL, NULL, RDW_INVALIDATE|RDW_UPDATENOW);
+				RedrawWindow(m_hwnd, NULL, NULL, RDW_INVALIDATE|RDW_UPDATENOW);
 				u64 paintTimeMs = TimeToMs(GetTime() - startTime);
 				u32 waitTime = u32(Min(paintTimeMs * 5, 200ull));
 				if (!m_traceView.finished)
@@ -3001,7 +3070,7 @@ namespace uba
 
 			if (m_config.ScaleHorizontalWithScrollWheel || controlDown || shiftDown)
 			{
-				if (m_activeSection >= 1)
+				if (m_activeSection == 2 || !controlDown) // process bars
 				{
 					RECT r;
 					GetClientRect(hWnd, &r);
@@ -3063,19 +3132,23 @@ namespace uba
 									//if (TimeToMs(process.writeFilesTime, m_traceView.frequency) >= 300 || TimeToMs(process.createFilesTime, m_traceView.frequency) >= 300)
 										process.bitmapDirty = true;
 				}
-				else if (m_activeSection == 0 || m_activeSection == 2)
+				else if (m_activeSection == 1) // active processes
 				{
-					if (controlDown)
-					{
-						if (delta < 0)
-							m_config.fontSize -= 1;
-						else if (delta > 0)
-							m_config.fontSize += 1;
-						UpdateDefaultFont();
-					}
+					if (delta < 0)
+						m_config.maxActiveProcessHeight = Max(m_config.maxActiveProcessHeight - 1u, 5u);
+					else if (delta > 0)
+						m_config.maxActiveProcessHeight = Min(m_config.maxActiveProcessHeight + 1u, 32u);
+				}
+				else if (m_activeSection == 0 || m_activeSection == 3) // status/timeline
+				{
+					if (delta < 0)
+						m_config.fontSize -= 1;
+					else if (delta > 0)
+						m_config.fontSize += 1;
+					UpdateDefaultFont();
 				}
 				UpdateScrollbars(true);
-				RedrawWindow(hWnd, NULL, NULL, RDW_INVALIDATE);
+				Redraw(false);
 			}
 			else
 			{
@@ -3087,7 +3160,7 @@ namespace uba
 				if (oldScrollY != m_scrollPosY)
 				{
 					UpdateScrollbars(true);
-					RedrawWindow(hWnd, NULL, NULL, RDW_INVALIDATE);
+					Redraw(false);
 				}
 			}
 			break;
@@ -3118,12 +3191,12 @@ namespace uba
 						m_autoScroll = true;
 				}
 				UpdateScrollbars(true);
-				RedrawWindow(hWnd, NULL, NULL, RDW_INVALIDATE);
+				Redraw(false);
 			}
 			else
 			{
 				if (UpdateSelection() || m_config.showCursorLine)
-					RedrawWindow(hWnd, NULL, NULL, RDW_INVALIDATE);
+					Redraw(false);
 				/*
 				else
 				{
@@ -3196,7 +3269,7 @@ namespace uba
 				HitTestResult res;
 				HitTest(res, { -1, -1 });
 				UpdateScrollbars(true);
-				RedrawWindow(hWnd, NULL, NULL, RDW_INVALIDATE);
+				Redraw(false);
 			}
 			else if (m_timelineSelected)
 			{
@@ -3250,7 +3323,7 @@ namespace uba
 					m_scrollPosY = Min(0.0f, Max(m_scrollPosY, float(r.bottom - m_contentHeight)));
 
 					UpdateScrollbars(true);
-					Redraw();
+					Redraw(true);
 				}
 			}
 			else if (!m_hyperLinkSelected.empty())
@@ -3310,7 +3383,7 @@ namespace uba
 			}
 			else if (m_processSelected)
 			{
-				TraceView::Process& process = *m_traceView.GetProcess(m_processSelectedLocation);
+				const TraceView::Process& process = m_traceView.GetProcess(m_processSelectedLocation);
 
 				AppendMenuW(hMenu, MF_STRING, Popup_CopyProcessInfo, L"&Copy Process Info");
 				if (!process.logLines.empty())
@@ -3370,12 +3443,12 @@ namespace uba
 			}
 			case Popup_ShowProcessText:
 				m_config.ShowProcessText = !m_config.ShowProcessText;
-				Redraw();
+				Redraw(true);
 				break;
 			case Popup_ShowReadWriteColors:
 				m_config.ShowReadWriteColors = !m_config.ShowReadWriteColors;
 				DirtyBitmaps(false);
-				Redraw();
+				Redraw(true);
 				break;
 			case Popup_ScaleHorizontalWithScrollWheel:
 				m_config.ScaleHorizontalWithScrollWheel = !m_config.ScaleHorizontalWithScrollWheel;
@@ -3385,15 +3458,15 @@ namespace uba
 				break;
 			case Popup_SortActiveRemoteSessions:
 				m_config.SortActiveRemoteSessions = !m_config.SortActiveRemoteSessions;
-				Redraw();
+				Redraw(true);
 				break;
 			case Popup_AutoScaleHorizontal:
 				m_config.AutoScaleHorizontal = !m_config.AutoScaleHorizontal;
-				Redraw();
+				Redraw(true);
 				break;
 			case Popup_LockTimelineToBottom:
 				m_config.LockTimelineToBottom = !m_config.LockTimelineToBottom;
-				Redraw();
+				Redraw(true);
 				break;
 			case Popup_DarkMode:
 			{
@@ -3405,7 +3478,7 @@ namespace uba
 				BOOL useDarkMode = m_config.DarkMode;
 				u32 attribute = 20; // DWMWA_USE_IMMERSIVE_DARK_MODE
 				DwmSetWindowAttribute(m_hwnd, attribute, &useDarkMode, sizeof(useDarkMode));
-				Redraw();
+				Redraw(true);
 				break;
 			}
 			case Popup_AutoSaveSettings:
@@ -3463,7 +3536,7 @@ namespace uba
 			{
 				TString str;
 				WriteTextLogger logger(str);
-				TraceView::Process& process = *m_traceView.GetProcess(m_processSelectedLocation);
+				const TraceView::Process& process = m_traceView.GetProcess(m_processSelectedLocation);
 				WriteProcessStats(logger, process);
 				CopyTextToClipboard(str);
 				break;
@@ -3471,7 +3544,7 @@ namespace uba
 			case Popup_CopyProcessLog:
 			{
 				TString str;
-				TraceView::Process& process = *m_traceView.GetProcess(m_processSelectedLocation);
+				const TraceView::Process& process = m_traceView.GetProcess(m_processSelectedLocation);
 				bool isFirst = true;
 				for (auto line : process.logLines)
 				{
@@ -3543,7 +3616,7 @@ namespace uba
 				if (oldScrollY != m_scrollPosY)
 				{
 					UpdateScrollbars(true);
-					RedrawWindow(hWnd, NULL, NULL, RDW_INVALIDATE);
+					Redraw(false);
 				}
 				return 0;
 			}
@@ -3595,7 +3668,7 @@ namespace uba
 				if (oldScrollX != m_scrollPosX)
 				{
 					UpdateScrollbars(true);
-					RedrawWindow(hWnd, NULL, NULL, RDW_INVALIDATE);
+					Redraw(false);
 				}
 				return 0;
 			}
@@ -3610,6 +3683,7 @@ namespace uba
 		{
 			thisPtr = (Visualizer*)lParam;
 			SetWindowLongPtr(hWnd, GWLP_USERDATA, lParam);
+			//EnableNonClientDpiScaling(hWnd);
 		}
 		if (thisPtr && hWnd == thisPtr->m_hwnd)
 			return thisPtr->WinProc(hWnd, Msg, wParam, lParam);
diff --git a/Engine/Source/Programs/UnrealBuildAccelerator/Visualizer/Private/UbaVisualizer.h b/Engine/Source/Programs/UnrealBuildAccelerator/Visualizer/Private/UbaVisualizer.h
--- a/Engine/Source/Programs/UnrealBuildAccelerator/Visualizer/Private/UbaVisualizer.h
+++ b/Engine/Source/Programs/UnrealBuildAccelerator/Visualizer/Private/UbaVisualizer.h
@@ -51,6 +51,7 @@ namespace uba
 		u32 fontSize = 13;
 		TString fontName;
 		u32 maxActiveVisible = 5;
+		u32 maxActiveProcessHeight = 16;
 
 		#define UBA_VISUALIZER_FLAG(name, defaultValue, desc) bool show##name = defaultValue;
 		UBA_VISUALIZER_FLAGS1
@@ -120,7 +121,7 @@ namespace uba
 		int GetTimelineTop(const RECT& clientRect);
 		void HitTest(HitTestResult& outResult, const POINT& pos);
 
-		void WriteProcessStats(Logger& out, TraceView::Process& process);
+		void WriteProcessStats(Logger& out, const TraceView::Process& process);
 		void CopyTextToClipboard(const TString& str);
 		void UnselectAndRedraw();
 		bool UpdateAutoscroll();
@@ -140,13 +141,14 @@ namespace uba
 			HFONT handle = 0;
 			HFONT handleUnderlined = 0;
 			int height = 0;
+			int offset = 0;
 		};
 
 		void UpdateFont(Font& font, int height, bool createUnderline);
 		void UpdateDefaultFont();
 		void UpdateProcessFont();
 		void ChangeFontSize(int offset);
-		void Redraw();
+		void Redraw(bool now);
 		void SetActiveFont(const Font& font);
 
 		StringBuffer<256> m_namedTrace;
@@ -197,6 +199,8 @@ namespace uba
 		Font m_timelineFont;
 		Font m_popupFont;
 
+		int m_processFontOffsetY = 0;
+
 		Font m_activeProcessFont[32];
 		u32 m_activeProcessCountHistory[5];
 		u32 m_activeProcessCountHistoryIterator = 0;
