[UBA]
* Added code to recognize LLVM IR obj files. Format is not handled but at least it now throws an error instead of crashing

diff --git a/Engine/Source/Programs/UnrealBuildAccelerator/Common/Private/UbaObjectFile.cpp b/Engine/Source/Programs/UnrealBuildAccelerator/Common/Private/UbaObjectFile.cpp
--- a/Engine/Source/Programs/UnrealBuildAccelerator/Common/Private/UbaObjectFile.cpp
+++ b/Engine/Source/Programs/UnrealBuildAccelerator/Common/Private/UbaObjectFile.cpp
@@ -5,6 +5,7 @@
 #include "UbaFileAccessor.h"
 #include "UbaObjectFileCoff.h"
 #include "UbaObjectFileElf.h"
+#include "UbaObjectFileLLVMIR.h"
 
 namespace uba
 {
@@ -29,16 +30,19 @@ namespace uba
 
 	ObjectFile* ObjectFile::Parse(Logger& logger, u8* data, u64 dataSize, const tchar* hint)
 	{
-		if (dataSize < 4)
-			return nullptr;
-
 		ObjectFile* objectFile = nullptr;
 
-		constexpr u8 elfMagic[] = { 0x7f, 'E', 'L', 'F' };
-		if (memcmp(data, elfMagic, sizeof(elfMagic)) == 0)
+		if (IsElfFile(data, dataSize))
 			objectFile = new ObjectFileElf();
-		else
+		else if (IsLLVMIRFile(data, dataSize))
+			objectFile = new ObjectFileLLVMIR();
+		else if (IsCoffFile(data, dataSize))
 			objectFile = new ObjectFileCoff();
+		else
+		{
+			logger.Error(TC("Unknown object file format. Maybe msvc FE IL? (%s)"), hint);
+			return nullptr;
+		}
 
 		objectFile->m_data = data;
 		objectFile->m_dataSize = dataSize;
diff --git a/Engine/Source/Programs/UnrealBuildAccelerator/Common/Private/UbaObjectFileCoff.cpp b/Engine/Source/Programs/UnrealBuildAccelerator/Common/Private/UbaObjectFileCoff.cpp
--- a/Engine/Source/Programs/UnrealBuildAccelerator/Common/Private/UbaObjectFileCoff.cpp
+++ b/Engine/Source/Programs/UnrealBuildAccelerator/Common/Private/UbaObjectFileCoff.cpp
@@ -174,19 +174,7 @@ namespace uba
 #endif
 	};
 
-	template<typename SymbolType>
-	ObjectFile::StringView GetSymbolName(SymbolType& symbol, const u8* data, u32 stringTableMemPos)
-	{
-		if (symbol.N.Name.Short == 0)
-		{
-			auto name = (const char*)(data + stringTableMemPos + symbol.N.Name.Long);
-			return { name, name + strlen(name) };
-		}
-		auto shortName = (char*)symbol.N.ShortName;
-		return { shortName, shortName + strnlen(shortName, ImageSizeofShortName) };
-	}
-
-	bool IsBigObj(u8* data, u64 size)
+	bool IsBigObj(const u8* data, u64 size)
 	{
 		if (size < sizeof(AnonObjectHeaderBigobj))
 			return false;
@@ -203,6 +191,38 @@ namespace uba
 		return true;
 	}
 
+	bool IsCoffFile(const u8* data, u64 dataSize)
+	{
+		if (IsBigObj(data, dataSize))
+			return true;
+		if (dataSize < sizeof(ImageFileHeader) + 8)
+			return false;
+
+		// TODO: This is not a solid way to identify a coff file.. don't know how this should be done
+		auto header = *(ImageFileHeader*)data;
+		if (header.Machine != 0x8664) // TODO: Add whatever other machines supported (ARM)
+			return false;
+		if (header.SizeOfOptionalHeader != 0) // Should always be 0
+			return false;
+		if (header.Characteristics != 0) // Should always be 0
+			return false;
+		// We expect .text or .drectve sections first.. this is also hacky.. but don't know how to verify that a file is a coff file
+		const u8* firstSection = data + sizeof(ImageFileHeader);
+		return memcmp(firstSection, ".text", 5) == 0 || memcmp(firstSection, ".drectve", 8) == 0;
+	}
+
+	template<typename SymbolType>
+	ObjectFile::StringView GetSymbolName(SymbolType& symbol, const u8* data, u32 stringTableMemPos)
+	{
+		if (symbol.N.Name.Short == 0)
+		{
+			auto name = (const char*)(data + stringTableMemPos + symbol.N.Name.Long);
+			return { name, name + strlen(name) };
+		}
+		auto shortName = (char*)symbol.N.ShortName;
+		return { shortName, shortName + strnlen(shortName, ImageSizeofShortName) };
+	}
+
 	ObjectFileCoff::ObjectFileCoff()
 	{
 		m_type = ObjectFileType_Coff;
diff --git a/Engine/Source/Programs/UnrealBuildAccelerator/Common/Private/UbaObjectFileElf.cpp b/Engine/Source/Programs/UnrealBuildAccelerator/Common/Private/UbaObjectFileElf.cpp
--- a/Engine/Source/Programs/UnrealBuildAccelerator/Common/Private/UbaObjectFileElf.cpp
+++ b/Engine/Source/Programs/UnrealBuildAccelerator/Common/Private/UbaObjectFileElf.cpp
@@ -112,6 +112,13 @@ namespace uba
 #endif
 	};
 
+	bool IsElfFile(const u8* data, u64 dataSize)
+	{
+		constexpr u8 magic[] = { 0x7f, 'E', 'L', 'F' };
+		return dataSize > 4 && memcmp(data, magic, sizeof(magic)) == 0;
+	}
+
+
 	ObjectFileElf::ObjectFileElf()
 	{
 		m_type = ObjectFileType_Elf;
diff --git a/Engine/Source/Programs/UnrealBuildAccelerator/Common/Public/UbaObjectFileCoff.h b/Engine/Source/Programs/UnrealBuildAccelerator/Common/Public/UbaObjectFileCoff.h
--- a/Engine/Source/Programs/UnrealBuildAccelerator/Common/Public/UbaObjectFileCoff.h
+++ b/Engine/Source/Programs/UnrealBuildAccelerator/Common/Public/UbaObjectFileCoff.h
@@ -45,4 +45,6 @@ namespace uba
 		static UnorderedSymbols PotentiallyDuplicatedSymbols;
 		static UnorderedSymbols ExportsToKeep;
 	};
+
+	bool IsCoffFile(const u8* data, u64 dataSize);
 }
diff --git a/Engine/Source/Programs/UnrealBuildAccelerator/Common/Public/UbaObjectFileElf.h b/Engine/Source/Programs/UnrealBuildAccelerator/Common/Public/UbaObjectFileElf.h
--- a/Engine/Source/Programs/UnrealBuildAccelerator/Common/Public/UbaObjectFileElf.h
+++ b/Engine/Source/Programs/UnrealBuildAccelerator/Common/Public/UbaObjectFileElf.h
@@ -20,4 +20,6 @@ namespace uba
 
 		static UnorderedSymbols PotentiallyDuplicatedSymbols;
 	};
+
+	bool IsElfFile(const u8* data, u64 dataSize);
 }
diff --git a/Engine/Source/Programs/UnrealBuildAccelerator/Common/Public/UbaObjectFileLLVMIR.h b/Engine/Source/Programs/UnrealBuildAccelerator/Common/Public/UbaObjectFileLLVMIR.h
new file mode 100644
--- /dev/null
+++ b/Engine/Source/Programs/UnrealBuildAccelerator/Common/Public/UbaObjectFileLLVMIR.h
@@ -0,0 +1,31 @@
+// Copyright Epic Games, Inc. All Rights Reserved.
+
+#pragma once
+
+#include "UbaObjectFile.h"
+
+namespace uba
+{
+	// TODO: https://llvm.org/docs/BitCodeFormat.html
+
+	class ObjectFileLLVMIR : public ObjectFile
+	{
+	public:
+		virtual bool Parse(Logger& logger, const tchar* filename) override
+		{
+			return logger.Error(TC("LLVM IR obj file format not supported (yet)"));
+		}
+
+	private:
+		virtual bool StripExports(Logger& logger, u8* newData, const UnorderedSymbols& allNeededImports, u32& outKeptExportCount) override { return false; }
+		virtual bool CreateExtraFile(Logger& logger, MemoryBlock& memoryBlock, const UnorderedSymbols& allNeededImports, const UnorderedSymbols& allSharedImports, const UnorderedExports& allSharedExports, bool includeExportsInFile) override { return false; }
+	};
+
+
+	bool IsLLVMIRFile(const u8* data, u64 dataSize)
+	{
+		constexpr u8 wrapperMagic[] = { 'B', 'C', 0xc0, 0xde };
+		constexpr u8 magic[] = { 'B', 'C', 0x04, 0xc4, 0xe4, 0xd4 };
+		return dataSize >= 6 && (memcmp(data, wrapperMagic, sizeof(wrapperMagic)) == 0 || memcmp(data, magic, sizeof(magic)) == 0);
+	}
+}
