[UBA]
* Added more logging around shutdown/close socket errors

diff --git a/Engine/Source/Programs/UnrealBuildAccelerator/Common/Private/UbaNetworkBackendTcp.cpp b/Engine/Source/Programs/UnrealBuildAccelerator/Common/Private/UbaNetworkBackendTcp.cpp
--- a/Engine/Source/Programs/UnrealBuildAccelerator/Common/Private/UbaNetworkBackendTcp.cpp
+++ b/Engine/Source/Programs/UnrealBuildAccelerator/Common/Private/UbaNetworkBackendTcp.cpp
@@ -117,12 +117,24 @@ namespace uba
 		return true;
 	}
 
-	bool CloseSocket(Logger& logger, SOCKET s)
+	bool ShutdownSocket(Logger& logger, SOCKET s, const tchar* hint)
 	{
-		if (s != INVALID_SOCKET)
-			if (closesocket(s) == SOCKET_ERROR)
-				return logger.Error(TC("failed to close socket (%s)"), LastErrorToText(WSAGetLastError()).data);
-		return true;
+		if (s == INVALID_SOCKET)
+			return true;
+		if (shutdown(s, SD_BOTH) != SOCKET_ERROR)
+			return true;
+		logger.Info(TC("failed to shutdown socket %llu for %s (%s)"), u64(s), hint, LastErrorToText(WSAGetLastError()).data);
+		return false;
+	}
+
+	bool CloseSocket(Logger& logger, SOCKET s, const tchar* hint)
+	{
+		if (s == INVALID_SOCKET)
+			return true;
+		if (closesocket(s) != SOCKET_ERROR)
+			return true;
+		logger.Info(TC("failed to close socket %llu for %s (%s)"), u64(s), hint, LastErrorToText(WSAGetLastError()).data);
+		return false;
 	}
 
 
@@ -144,10 +156,10 @@ namespace uba
 				continue;
 			SOCKET s = conn.socket;
 			conn.socket = INVALID_SOCKET;
-			shutdown(s, SD_BOTH);
+			ShutdownSocket(conn.logger, s, TC("tcp dtor"));
 			lock2.Leave();
 			conn.recvThread.Wait();
-			CloseSocket(conn.logger, s);
+			CloseSocket(conn.logger, s, TC("tcp dtor"));
 		}
 		m_connections.clear();
 
@@ -163,7 +175,7 @@ namespace uba
 		ScopedCriticalSection lock(conn.shutdownLock);
 		if (conn.socket == INVALID_SOCKET)
 			return;
-		shutdown(conn.socket, SD_BOTH);
+		ShutdownSocket(conn.logger, conn.socket, TC("shutdown"));
 	}
 
 	bool NetworkBackendTcp::Send(Logger& logger, void* connection, const void* data, u32 dataSize, SendContext& sendContext)
@@ -341,7 +353,7 @@ namespace uba
 		if (listenSocket == INVALID_SOCKET)
 			return logger.Error(TC("socket failed (%s)"), LastErrorToText(WSAGetLastError()).data);
 
-		auto listenSocketCleanup = MakeGuard([&]() { CloseSocket(logger, listenSocket); });
+		auto listenSocketCleanup = MakeGuard([&]() { CloseSocket(logger, listenSocket, TC("listen cleanup")); });
 
 		u32 reuseAddr = 1;
 		if (::setsockopt(listenSocket, SOL_SOCKET, SO_REUSEADDR, (const char*)&reuseAddr, sizeof reuseAddr) == SOCKET_ERROR)
@@ -411,7 +423,7 @@ namespace uba
 
 			if (!DisableNagle(logger, clientSocket) || !SetKeepAlive(logger, clientSocket))
 			{
-				CloseSocket(logger, clientSocket);
+				CloseSocket(logger, clientSocket, TC("disable nagle"));
 				continue;
 			}
 
@@ -425,7 +437,7 @@ namespace uba
 
 			if (!entry.connectedFunc(&conn, remoteSockAddr))
 			{
-				shutdown(clientSocket, SD_BOTH);
+				ShutdownSocket(logger, clientSocket, TC("thread listen"));
 				conn.ready.Set();
 				conn.recvThread.Wait();
 				SCOPED_WRITE_LOCK(m_connectionsLock, lock2);
@@ -510,8 +522,8 @@ namespace uba
 
 		if (s == INVALID_SOCKET)
 			return;
-		shutdown(s, SD_BOTH);
-		CloseSocket(logger, s);
+		ShutdownSocket(logger, s, TC("threadrecv"));
+		CloseSocket(logger, s, TC("threadrecv"));
 	}
 
 	bool NetworkBackendTcp::Connect(Logger& logger, const tchar* ip, const ConnectedFunc& connectedFunc, u16 port, bool* timedOut)
@@ -569,7 +581,7 @@ namespace uba
 			return logger.Error(TC("socket failed (%s)"), LastErrorToText(WSAGetLastError()).data);
 
 		// Create guard in case we fail to connect (will be cancelled further down if we succeed)
-		auto socketClose = MakeGuard([&]() { CloseSocket(logger, socketFd); });
+		auto socketClose = MakeGuard([&]() { CloseSocket(logger, socketFd, TC("connect")); });
 
 		// Set to non-blocking just for the connect call (we want to control the connect timeout after connect using select instead)
 		if (!SetBlocking(logger, socketFd, false))
@@ -675,7 +687,7 @@ namespace uba
 
 		if (!connectedFunc(&conn, remoteSocketAddr, timedOut))
 		{
-			shutdown(socketFd, SD_BOTH);
+			ShutdownSocket(logger, socketFd, TC("connect"));
 			conn.ready.Set();
 			conn.recvThread.Wait();
 			SCOPED_WRITE_LOCK(m_connectionsLock, lock2);
@@ -945,7 +957,7 @@ namespace uba
 		if (m_socket != INVALID_SOCKET)
 		{
 			LoggerWithWriter logger(g_nullLogWriter);
-			CloseSocket(logger, m_socket);
+			CloseSocket(logger, m_socket, TC("http dtor"));
 		}
 
 		#if PLATFORM_WINDOWS
@@ -1001,7 +1013,7 @@ namespace uba
 		// TODO: Fix so we reuse socket connection for multiple queries
 		if (*m_host)// && _stricmp(m_host, host) != 0)
 		{
-			CloseSocket(logger, m_socket);
+			CloseSocket(logger, m_socket, TC("http query"));
 			m_socket = INVALID_SOCKET;
 			*m_host = 0;
 		}
