[UBA]
* Moved format specific object file headers to private
* Added object file source file for llvm ir

diff --git a/Engine/Source/Programs/UnrealBuildAccelerator/Common/Private/UbaObjectFileCoff.h b/Engine/Source/Programs/UnrealBuildAccelerator/Common/Private/UbaObjectFileCoff.h
new file mode 100644
--- /dev/null
+++ b/Engine/Source/Programs/UnrealBuildAccelerator/Common/Private/UbaObjectFileCoff.h
@@ -0,0 +1,50 @@
+// Copyright Epic Games, Inc. All Rights Reserved.
+
+#pragma once
+
+#include "UbaObjectFile.h"
+
+namespace uba
+{
+
+	class ObjectFileCoff : public ObjectFile
+	{
+	public:
+		ObjectFileCoff();
+		virtual bool Parse(Logger& logger, const tchar* filename) override;
+
+	private:
+		struct Info;
+
+		bool ParseExports();
+		template<typename SymbolType> void ParseImports();
+
+		virtual bool StripExports(Logger& logger, u8* newData, const UnorderedSymbols& allNeededImports, u32& outKeptExportCount) override;
+		virtual bool CreateExtraFile(Logger& logger, MemoryBlock& memoryBlock, const UnorderedSymbols& allNeededImports, const UnorderedSymbols& allSharedImports, const UnorderedExports& allSharedExports, bool includeExportsInFile) override;
+
+		template<typename SymbolType> void CalculateImports(Logger& logger, Vector<u32>& outImports);
+		template<typename SymbolType> void WriteImports(Logger& logger, u8* newData, Info& newInfo, const Vector<u32>& symbolsToAdd);
+		template<typename SymbolType> void RemoveSymbols(Logger& logger, u8* newData, Info& newInfo);
+
+		struct Info
+		{
+			u32 sectionsMemOffset = 0;
+			u32 sectionCount = 0;
+			u64 directiveSectionMemOffset = 0;
+			u32 stringTableMemPos = 0;
+			u32 symbolsMemPos = 0;
+			u32 symbolCount = 0;
+		};
+
+		bool m_isBigObj = false;
+		Info m_info;
+
+		UnorderedSymbols m_loopbacksToAdd;
+		UnorderedSymbols m_toRemove;
+
+		static UnorderedSymbols PotentiallyDuplicatedSymbols;
+		static UnorderedSymbols ExportsToKeep;
+	};
+
+	bool IsCoffFile(const u8* data, u64 dataSize);
+}
diff --git a/Engine/Source/Programs/UnrealBuildAccelerator/Common/Private/UbaObjectFileElf.h b/Engine/Source/Programs/UnrealBuildAccelerator/Common/Private/UbaObjectFileElf.h
new file mode 100644
--- /dev/null
+++ b/Engine/Source/Programs/UnrealBuildAccelerator/Common/Private/UbaObjectFileElf.h
@@ -0,0 +1,25 @@
+// Copyright Epic Games, Inc. All Rights Reserved.
+
+#pragma once
+
+#include "UbaObjectFile.h"
+
+namespace uba
+{
+	class ObjectFileElf : public ObjectFile
+	{
+	public:
+		ObjectFileElf();
+		virtual bool Parse(Logger& logger, const tchar* filename) override;
+
+	private:
+		virtual bool StripExports(Logger& logger, u8* newData, const UnorderedSymbols& allNeededImports, u32& outKeptExportCount) override;
+		virtual bool CreateExtraFile(Logger& logger, MemoryBlock& memoryBlock, const UnorderedSymbols& allNeededImports, const UnorderedSymbols& allSharedImports, const UnorderedExports& allSharedExports, bool includeExportsInFile) override;
+
+		UnorderedSymbols m_toRemove;
+
+		static UnorderedSymbols PotentiallyDuplicatedSymbols;
+	};
+
+	bool IsElfFile(const u8* data, u64 dataSize);
+}
diff --git a/Engine/Source/Programs/UnrealBuildAccelerator/Common/Private/UbaObjectFileLLVMIR.cpp b/Engine/Source/Programs/UnrealBuildAccelerator/Common/Private/UbaObjectFileLLVMIR.cpp
new file mode 100644
--- /dev/null
+++ b/Engine/Source/Programs/UnrealBuildAccelerator/Common/Private/UbaObjectFileLLVMIR.cpp
@@ -0,0 +1,23 @@
+// Copyright Epic Games, Inc. All Rights Reserved.
+
+#pragma once
+
+#include "UbaObjectFileLLVMIR.h"
+
+namespace uba
+{
+	// TODO: https://llvm.org/docs/BitCodeFormat.html
+
+
+	bool ObjectFileLLVMIR::Parse(Logger& logger, const tchar* filename)
+	{
+		return logger.Error(TC("LLVM IR obj file format not supported (yet)"));
+	}
+
+	bool IsLLVMIRFile(const u8* data, u64 dataSize)
+	{
+		constexpr u8 wrapperMagic[] = { 'B', 'C', 0xc0, 0xde };
+		constexpr u8 magic[] = { 'B', 'C', 0x04, 0xc4, 0xe4, 0xd4 };
+		return dataSize >= 6 && (memcmp(data, wrapperMagic, sizeof(wrapperMagic)) == 0 || memcmp(data, magic, sizeof(magic)) == 0);
+	}
+}
diff --git a/Engine/Source/Programs/UnrealBuildAccelerator/Common/Private/UbaObjectFileLLVMIR.h b/Engine/Source/Programs/UnrealBuildAccelerator/Common/Private/UbaObjectFileLLVMIR.h
new file mode 100644
--- /dev/null
+++ b/Engine/Source/Programs/UnrealBuildAccelerator/Common/Private/UbaObjectFileLLVMIR.h
@@ -0,0 +1,21 @@
+// Copyright Epic Games, Inc. All Rights Reserved.
+
+#pragma once
+
+#include "UbaObjectFile.h"
+
+namespace uba
+{
+	class ObjectFileLLVMIR : public ObjectFile
+	{
+	public:
+		virtual bool Parse(Logger& logger, const tchar* filename) override;
+
+	private:
+		virtual bool StripExports(Logger& logger, u8* newData, const UnorderedSymbols& allNeededImports, u32& outKeptExportCount) override { return false; }
+		virtual bool CreateExtraFile(Logger& logger, MemoryBlock& memoryBlock, const UnorderedSymbols& allNeededImports, const UnorderedSymbols& allSharedImports, const UnorderedExports& allSharedExports, bool includeExportsInFile) override { return false; }
+	};
+
+
+	bool IsLLVMIRFile(const u8* data, u64 dataSize);
+}
diff --git a/Engine/Source/Programs/UnrealBuildAccelerator/Common/Public/UbaObjectFileCoff.h b/Engine/Source/Programs/UnrealBuildAccelerator/Common/Public/UbaObjectFileCoff.h
deleted file mode 100644
diff --git a/Engine/Source/Programs/UnrealBuildAccelerator/Common/Public/UbaObjectFileElf.h b/Engine/Source/Programs/UnrealBuildAccelerator/Common/Public/UbaObjectFileElf.h
deleted file mode 100644
diff --git a/Engine/Source/Programs/UnrealBuildAccelerator/Common/Public/UbaObjectFileLLVMIR.h b/Engine/Source/Programs/UnrealBuildAccelerator/Common/Public/UbaObjectFileLLVMIR.h
deleted file mode 100644
