[UBA]
* Fixed bug in application rules where lld-link.exe expected all .manifest files to be kept in memory but it is only the ones in temp dir that we want in memory

diff --git a/Engine/Source/Programs/UnrealBuildAccelerator/Core/Private/UbaApplicationRules.cpp b/Engine/Source/Programs/UnrealBuildAccelerator/Core/Private/UbaApplicationRules.cpp
--- a/Engine/Source/Programs/UnrealBuildAccelerator/Core/Private/UbaApplicationRules.cpp
+++ b/Engine/Source/Programs/UnrealBuildAccelerator/Core/Private/UbaApplicationRules.cpp
@@ -166,7 +166,7 @@ namespace uba
 
 		virtual bool KeepInMemory(const StringView& fileName, const tchar* systemTemp) const override
 		{
-			return fileName.EndsWith(TC(".manifest"));
+			return fileName.EndsWith(TC(".manifest")) && fileName.Contains(systemTemp);
 		}
 
 		virtual bool IsOutputFile(const StringView& fileName) const override
