[UBA]
* Attempt at fixing what seems like deadlock in the directory table update in session clients. Theory being that if response from server is out of order and message capacity is exactly the same as the remainding size of directory table we end up with a thread timing out.. and the code did not clean up the linked list properly on timeout so anything could go wrong after that
* Changed some asserts to errors to make sure they happen in shipping builds too

diff --git a/Engine/Source/Programs/UnrealBuildAccelerator/Common/Private/UbaSessionClient.cpp b/Engine/Source/Programs/UnrealBuildAccelerator/Common/Private/UbaSessionClient.cpp
--- a/Engine/Source/Programs/UnrealBuildAccelerator/Common/Private/UbaSessionClient.cpp
+++ b/Engine/Source/Programs/UnrealBuildAccelerator/Common/Private/UbaSessionClient.cpp
@@ -916,8 +916,9 @@ namespace uba
 		u32 readPos = 0;
 		ActiveUpdateDirectoryEntry* prev = nullptr;
 		ActiveUpdateDirectoryEntry* next = nullptr;
+		bool success = true;
 
-		static bool Wait(SessionClient& client, ActiveUpdateDirectoryEntry*& first, ScopedWriteLock& lock, u32 readPos)
+		static bool Wait(SessionClient& client, ActiveUpdateDirectoryEntry*& first, ScopedWriteLock& lock, u32 readPos, const tchar* hint)
 		{
 			ActiveUpdateDirectoryEntry item;
 			item.next = first;
@@ -926,20 +927,25 @@ namespace uba
 			item.readPos = readPos;
 			first = &item;
 			item.done.Create(true);
+
 			lock.Leave();
-			while (!item.done.IsSet(10000))
-			{
-				client.m_logger.Error(TC("Timed out waiting for update directory message"));
-				return false;
-			}
+			bool res = item.done.IsSet(5*60*60);
 			lock.Enter();
+
 			if (item.prev)
 				item.prev->next = item.next;
 			else
 				first = item.next;
 			if (item.next)
 				item.next->prev = item.prev;
-			return true;
+
+			if (res)
+				return item.success;
+
+			u32 activeCount = 0;
+			for (auto i = first; i; i = i->next)
+				++activeCount;
+			return client.m_logger.Error(TC("Timed out after 5 minutes waiting for update directory message to reach read position %u  (%u active in %s wait)"), readPos, activeCount, hint);
 		}
 
 		static void UpdateReadPosMatching(ActiveUpdateDirectoryEntry*& first, u32 readPos)
@@ -953,32 +959,40 @@ namespace uba
 			}
 		}
 
-		static void UpdateReadPosLess(ActiveUpdateDirectoryEntry*& first, u32 readPos)
+		static void UpdateReadPosLessOrEqual(ActiveUpdateDirectoryEntry*& first, u32 readPos)
 		{
 			for (auto i = first; i; i = i->next)
 				if (i->readPos <= readPos)
 					i->done.Set();
 		}
+
+		static void UpdateError(ActiveUpdateDirectoryEntry*& first)
+		{
+			for (auto i = first; i; i = i->next)
+			{
+				i->success = false;
+				i->done.Set();
+			}
+		}
 	};
 
 	bool SessionClient::UpdateDirectoryTableFromServer(StackBinaryReader<SendMaxSize>& reader)
 	{
 		auto& dirTable = m_directoryTable;
 
-		bool isFirst = true;
-		while (true)
-		{
-			if (!isFirst)
+		auto updateMemorySizeAndSignal = [&]
 			{
-				reader.Reset();
+				SCOPED_WRITE_LOCK(dirTable.m_memoryLock, lock);
+				dirTable.m_memorySize = m_directoryTableMemPos;
+				lock.Leave();
+				ActiveUpdateDirectoryEntry::UpdateReadPosLessOrEqual(m_firstEmptyWait, m_directoryTableMemPos);
+				return true;
+			};
 
-				StackBinaryWriter<1024> writer;
-				NetworkMessage msg(m_client, ServiceId, SessionMessageType_GetDirectoriesFromServer, writer);
-				writer.WriteU32(m_sessionId);
-				if (!msg.Send(reader, Stats().getDirsMsg))
-					return false;
-			}
+		u32 lastWriteEnd = ~0u;
 
+		while (true)
+		{
 			u32 readPos = reader.ReadU32();
 
 			u8* pos = dirTable.m_memory + readPos;
@@ -986,39 +1000,59 @@ namespace uba
 
 			SCOPED_WRITE_LOCK(m_directoryTableLock, lock);
 
+			if (m_directoryTableError)
+				return false;
+
 			if (toRead == 0)
 			{
+				// We wrote to lastWriteEnd and now we got an empty message where readPos is the same..
+				// This means that it was a good cut-off and we can increase m_memorySize
+				// If m_directoryTableMemPos is different it means that we have another thread going on that will update things a little bit later.
+				if (lastWriteEnd == readPos && lastWriteEnd == m_directoryTableMemPos)
+					return updateMemorySizeAndSignal();
+
 				// We might share this position with others
-				if (readPos > dirTable.m_memorySize)
-					if (!ActiveUpdateDirectoryEntry::Wait(*this, m_firstEmptyWait, lock, readPos))
+				if (dirTable.m_memorySize < readPos)
+					if (!ActiveUpdateDirectoryEntry::Wait(*this, m_firstEmptyWait, lock, readPos, TC("empty")))
 						return false;
 				return true;
 			}
 
 			reader.ReadBytes(pos, toRead);
 
+			// Wait until all data before readPos has been read
 			if (readPos != m_directoryTableMemPos)
-				if (!ActiveUpdateDirectoryEntry::Wait(*this, m_firstReadWait, lock, readPos))
+				if (!ActiveUpdateDirectoryEntry::Wait(*this, m_firstReadWait, lock, readPos, TC("read")))
 					return false;
 
 			m_directoryTableMemPos += toRead;
 			
+			// Find potential waiter waiting for this exact size and wake it up
 			ActiveUpdateDirectoryEntry::UpdateReadPosMatching(m_firstReadWait, m_directoryTableMemPos);
 
+
+			// If there is space left in the message it means that we caught up with the directory table server side..
+			// And we will stop asking for more data.
+			// Note, we can only set m_memorySize when getting messages that reads less than capacity since we don't know if we reached a good position in the directory table
 			if (reader.GetPosition() < m_client.GetMessageMaxSize() - m_client.GetMessageReceiveHeaderSize())
-			{
-				//dirTable.ParseDirectoryTable(m_directoryTableMemPos); // This is not needed.. we never read from the directory table in the client session
-				{
-					SCOPED_WRITE_LOCK(dirTable.m_memoryLock, lock2);
-					dirTable.m_memorySize = m_directoryTableMemPos;
-				}
-				ActiveUpdateDirectoryEntry::UpdateReadPosLess(m_firstEmptyWait, m_directoryTableMemPos);
-				break;
-			}
-			isFirst = false;
-		}
+				return updateMemorySizeAndSignal();
 
-		return true;
+			lastWriteEnd = m_directoryTableMemPos;
+
+			StackBinaryWriter<1024> writer;
+			NetworkMessage msg(m_client, ServiceId, SessionMessageType_GetDirectoriesFromServer, writer);
+			writer.WriteU32(m_sessionId);
+
+			reader.Reset();
+			if (msg.Send(reader, Stats().getDirsMsg))
+				continue;
+
+			// Let's signal waiters to exit faster since we will not get out of this situation (most likely a disconnect)
+			m_directoryTableError = true;
+			ActiveUpdateDirectoryEntry::UpdateError(m_firstReadWait);
+			ActiveUpdateDirectoryEntry::UpdateError(m_firstEmptyWait);
+			return false;
+		}
 	}
 
 	bool SessionClient::UpdateNameToHashTableFromServer(StackBinaryReader<SendMaxSize>& reader)
@@ -1362,7 +1396,6 @@ namespace uba
 		StackBinaryWriter<32> writer;
 		NetworkMessage msg(m_client, ServiceId, SessionMessageType_GetDirectoriesFromServer, writer);
 		writer.WriteU32(m_sessionId);
-		writer.WriteU32(~u32(0));
 		if (!msg.Send(reader, Stats().getDirsMsg))
 			return false;
 		return UpdateDirectoryTableFromServer(reader);
diff --git a/Engine/Source/Programs/UnrealBuildAccelerator/Common/Private/UbaSessionServer.cpp b/Engine/Source/Programs/UnrealBuildAccelerator/Common/Private/UbaSessionServer.cpp
--- a/Engine/Source/Programs/UnrealBuildAccelerator/Common/Private/UbaSessionServer.cpp
+++ b/Engine/Source/Programs/UnrealBuildAccelerator/Common/Private/UbaSessionServer.cpp
@@ -1001,7 +1001,8 @@ namespace uba
 		u32 sessionId = reader.ReadU32();
 		u32 sessionIndex = sessionId - 1;
 		ScopedCriticalSection lock(m_remoteProcessAndSessionLock);
-		UBA_ASSERT(sessionIndex < m_clientSessions.size());
+		if (sessionIndex >= m_clientSessions.size())
+			return m_logger.Error(TC("Got ListDirectory message from connection using bad sessionid (%u/%llu)"), sessionIndex, m_clientSessions.size());
 		ClientSession& session = *m_clientSessions[sessionIndex];
 		lock.Leave();
 
@@ -1020,7 +1021,8 @@ namespace uba
 		u32 sessionId = reader.ReadU32();
 		u32 sessionIndex = sessionId - 1;
 		ScopedCriticalSection lock(m_remoteProcessAndSessionLock);
-		UBA_ASSERT(sessionIndex < m_clientSessions.size());
+		if (sessionIndex >= m_clientSessions.size())
+			return m_logger.Error(TC("Got GetDirectories message from connection using bad sessionid (%u/%llu)"), sessionIndex, m_clientSessions.size());
 		ClientSession& session = *m_clientSessions[sessionIndex];
 		lock.Leave();
 		WriteDirectoryTable(session, reader, writer);
@@ -1050,7 +1052,8 @@ namespace uba
 		u32 sessionIndex = sessionId - 1;
 
 		ScopedCriticalSection sessionsLock(m_remoteProcessAndSessionLock);
-		UBA_ASSERT(sessionIndex < m_clientSessions.size());
+		if (sessionIndex >= m_clientSessions.size())
+			return m_logger.Error(TC("Got ProcessAvailable message from connection using bad sessionid (%u/%llu)"), sessionIndex, m_clientSessions.size());
 		ClientSession& session = *m_clientSessions[sessionIndex];
 		sessionsLock.Leave();
 
@@ -1183,7 +1186,8 @@ namespace uba
 			return true;
 		}
 		u32 sessionIndex = process.m_sessionId - 1;
-		UBA_ASSERT(sessionIndex < m_clientSessions.size());
+		if (sessionIndex >= m_clientSessions.size())
+			return m_logger.Error(TC("Got ProcessFinished message from connection using bad sessionid (%u/%llu)"), sessionIndex, m_clientSessions.size());
 		auto& session = *m_clientSessions[sessionIndex];
 		++m_finishedRemoteProcessCount;
 		--session.usedSlotCount;
@@ -1251,7 +1255,8 @@ namespace uba
 			return true;
 		}
 		u32 sessionIndex = process->m_sessionId - 1;
-		UBA_ASSERT(sessionIndex < m_clientSessions.size());
+		if (sessionIndex >= m_clientSessions.size())
+			return m_logger.Error(TC("Got ProcessReturned message from connection using bad sessionid (%u/%llu)"), sessionIndex, m_clientSessions.size());
 		auto& session = *m_clientSessions[sessionIndex];
 		--session.usedSlotCount;
 		if (session.enabled)
@@ -1305,7 +1310,8 @@ namespace uba
 		u64 pingTime = GetTime();
 		u32 sessionIndex = sessionId - 1;
 		ScopedCriticalSection lock(m_remoteProcessAndSessionLock);
-		UBA_ASSERT(sessionIndex < m_clientSessions.size());
+		if (sessionIndex >= m_clientSessions.size())
+			return m_logger.Error(TC("Got Pingmessage from connection using bad sessionid (%u/%llu)"), sessionIndex, m_clientSessions.size());
 		auto& session = *m_clientSessions[sessionIndex];
 		session.pingTime = pingTime;
 		session.lastPing = lastPing;
diff --git a/Engine/Source/Programs/UnrealBuildAccelerator/Common/Public/UbaSessionClient.h b/Engine/Source/Programs/UnrealBuildAccelerator/Common/Public/UbaSessionClient.h
--- a/Engine/Source/Programs/UnrealBuildAccelerator/Common/Public/UbaSessionClient.h
+++ b/Engine/Source/Programs/UnrealBuildAccelerator/Common/Public/UbaSessionClient.h
@@ -142,6 +142,7 @@ namespace uba
 
 		ReaderWriterLock m_directoryTableLock;
 		u32 m_directoryTableMemPos = 0;
+		bool m_directoryTableError = false;
 		struct ActiveUpdateDirectoryEntry;
 		ActiveUpdateDirectoryEntry* m_firstEmptyWait = nullptr;
 		ActiveUpdateDirectoryEntry* m_firstReadWait = nullptr;
