[UBA]
* Fixed so zoom is based on box height which is integer values. Now it is possible to zoom down to only 1 pixel high process bars

diff --git a/Engine/Source/Programs/UnrealBuildAccelerator/Visualizer/Private/UbaVisualizer.cpp b/Engine/Source/Programs/UnrealBuildAccelerator/Visualizer/Private/UbaVisualizer.cpp
--- a/Engine/Source/Programs/UnrealBuildAccelerator/Visualizer/Private/UbaVisualizer.cpp
+++ b/Engine/Source/Programs/UnrealBuildAccelerator/Visualizer/Private/UbaVisualizer.cpp
@@ -564,8 +564,6 @@ namespace uba
 
 		UpdateDefaultFont();
 
-		u32 defaultProcessFontHeight = 13;
-		m_zoomValue = defaultProcessFontHeight / 20.0f;
 		UpdateProcessFont();
 
 		const TCHAR* fontName = TEXT("Consolas");
@@ -789,10 +787,9 @@ namespace uba
 
 	void Visualizer::UpdateProcessFont()
 	{
-		int fontHeight = int(m_zoomValue*20.0f);
+		m_zoomValue = 1.0f + m_boxHeight / 30.0f;
+		int fontHeight = Max(m_boxHeight - 2, 1);
 		UpdateFont(m_processFont, fontHeight, false);
-		m_processFont.height = fontHeight;
-		m_boxHeight = m_processFont.height + 2;
 		m_progressRectLeft = int(5 + float(m_processFont.height) * 1.8f);
 		DirtyBitmaps(true);
 	}
@@ -946,7 +943,7 @@ namespace uba
 			rect.right = rect.left + int(progress*width);//durationMs/100;
 			FillRect(hdc, &rect, m_traceView.progressErrorCount ? m_processBrushes[0].error : m_processBrushes[0].success);
 
-			const tchar* remoteDisabled = m_traceView.remoteExecutionDisabled ? TC("   (Remote spawn disabled)") : TC("");
+			const tchar* remoteDisabled = TC(""); // TODO: Only show this in "advanced mode" m_traceView.remoteExecutionDisabled ? TC("   (Remote spawn disabled)") : TC("");
 			StringBuffer<> str;
 			str.Appendf(L"%u%%    %u / %u %s", u32(progress*100.0f), m_traceView.progressProcessesDone, m_traceView.progressProcessesTotal, remoteDisabled);
 			drawIndentedText(str, LogEntryType_Info, 6, true);
@@ -1054,9 +1051,9 @@ namespace uba
 				});
 		}
 
-		float boxHeight = float(m_boxHeight);
+		int boxHeight = m_boxHeight;
 		int stepY = int(boxHeight) + 2;
-		int processStepY = int(boxHeight) + 2;
+		int processStepY = boxHeight + 1;
 
 		TraceView::WorkRecord selectedWork;
 
@@ -1233,7 +1230,7 @@ namespace uba
 
 			SetActiveFont(m_processFont);
 
-			bool shouldDrawText = m_zoomValue > 0.1f;
+			bool shouldDrawText = m_processFont.height > 4;
 
 			if (m_config.showProcessBars)
 			{
@@ -1244,16 +1241,16 @@ namespace uba
 
 					if (posY + m_sessionStepY >= progressRect.top && posY < progressRect.bottom)
 					{
-						float barHeight = boxHeight;
+						int barHeight = boxHeight;
 						int textOffsetY = 0;
-						if (posY + int(boxHeight) > progressRect.bottom)
+						if (posY + boxHeight > progressRect.bottom)
 						{
-							float newBarHeight = Min(barHeight, float(progressRect.bottom - posY));
+							int newBarHeight = Min(barHeight, int(progressRect.bottom - posY));
 							textOffsetY = int(barHeight - newBarHeight);
 							barHeight = newBarHeight;
 						}
 
-						const int textHeight = int(barHeight);
+						const int textHeight = barHeight;
 						const int rectBottom = posY + textHeight;
 						const int offsetY = (textHeight - m_processFont.height + textOffsetY) / 2;
 
@@ -1281,7 +1278,7 @@ namespace uba
 							rect.left = left;
 							rect.right = int(posX + TimeToS(stop) * scaleX) - 1;
 							rect.top = posY;
-							rect.bottom = rectBottom;
+							rect.bottom = rectBottom - 1;
 
 							if (rect.right <= progressRect.left)
 								continue;
@@ -1330,9 +1327,7 @@ namespace uba
 
 									rect2.left += 3; // Move in text a bit
 									
-									//SetTextAlign(textDC, TA_LEFT|TA_TOP|TA_NOUPDATECP|VTA_CENTER);
-									int textY = rect2.top;// - int(m_processFont.height/9);
-									textY += m_processFont.offset;
+									int textY = rect2.top + m_processFont.offset;
 
 									bool dropShadow = m_config.DarkMode;
 									if (dropShadow)
@@ -1422,15 +1417,15 @@ namespace uba
 					if (posY + m_sessionStepY >= progressRect.top && posY <= progressRect.bottom)
 					{
 						int textOffsetY = 0;
-						float barHeight = boxHeight;
+						int barHeight = boxHeight;
 						if (posY + int(boxHeight) > progressRect.bottom)
 						{
-							float newBarHeight = Min(barHeight, float(progressRect.bottom - posY));
-							textOffsetY = int(barHeight - newBarHeight);
+							int newBarHeight = Min(barHeight, int(progressRect.bottom - posY));
+							textOffsetY = barHeight - newBarHeight;
 							barHeight = newBarHeight;
 						}
 
-						const int textHeight = int(barHeight);
+						const int textHeight = barHeight;
 						const int rectBottom = posY + textHeight;
 						const int offsetY = (textHeight - m_processFont.height + textOffsetY) / 2;
 
@@ -1494,7 +1489,6 @@ namespace uba
 							rect.right = int(posX + stopTime * scaleX) - 1;
 							rect.top = posY;
 							rect.bottom = rectBottom;
-							//rect.bottom = posY + int(float(18) * m_zoomValue);
 
 							if (rect.right <= progressRect.left)
 							{
@@ -2406,8 +2400,8 @@ namespace uba
 		GetClientRect(m_hwnd, &clientRect);
 
 		int posY = int(m_scrollPosY);
-		float boxHeight = float(m_boxHeight);
-		int processStepY = int(boxHeight) + 2;
+		int boxHeight = m_boxHeight;
+		int processStepY = boxHeight + 1;
 		float scaleX = 50.0f*m_zoomValue*m_horizontalScaleValue;
 
 		RECT progressRect = clientRect;
@@ -2468,7 +2462,7 @@ namespace uba
 			posY += 3;
 		}
 
-		if (pos.y <= posY)
+		if (pos.y < posY)
 			return;
 
 		outResult.section = 1;
@@ -2505,7 +2499,7 @@ namespace uba
 				return;
 		}
 
-		if (pos.y <= posY)
+		if (pos.y < posY)
 			return;
 
 		outResult.section = 2;
@@ -3080,20 +3074,19 @@ namespace uba
 					GetCursorPos(&cursorPos);
 					ScreenToClient(m_hwnd, &cursorPos);
 
-					float newZoomValue = m_zoomValue;
 					float newScaleValue = m_horizontalScaleValue;
+					int newBoxHeight = m_boxHeight;
 
 					if (controlDown)
 					{
 						if (delta < 0)
 						{
-							newZoomValue = (m_processFont.height - 1) / 20.0f;
-							if (newZoomValue == 0)
-								newZoomValue = m_zoomValue;
+							if (newBoxHeight > 1)
+								--newBoxHeight;
 						}
 						else if (delta > 0)
-							newZoomValue = (m_processFont.height + 1) / 20.0f;
-					}
+							++newBoxHeight;
+												}
 					else
 						newScaleValue = Max(m_horizontalScaleValue + m_horizontalScaleValue*float(delta)*0.0006f, 0.001f);
 
@@ -3101,18 +3094,19 @@ namespace uba
 					const float scrollAnchorOffsetX = float(cursorPos.x) - m_progressRectLeft;
 					const float scrollAnchorOffsetY = 0;//float(cursorPos.y)*m_zoomValue;// - m_progressRectLeft;
 
-					m_scrollPosY = Min(0.0f, float(m_scrollPosY - scrollAnchorOffsetY)*(newZoomValue/m_zoomValue) + scrollAnchorOffsetY);
-					m_scrollPosX = Min(0.0f, float(m_scrollPosX - scrollAnchorOffsetX)*(newZoomValue/m_zoomValue)*(newScaleValue/m_horizontalScaleValue) + scrollAnchorOffsetX);//LOWORD(lParam);
-
-					if (m_horizontalScaleValue != newScaleValue)
-						m_horizontalScaleValue = newScaleValue;
-
-					if (m_zoomValue != newZoomValue)
+					float oldZoomValue = m_zoomValue;
+					if (newBoxHeight != m_boxHeight)
 					{
-						m_zoomValue = newZoomValue;
+						m_boxHeight = newBoxHeight;
 						UpdateProcessFont();
 					}
 
+					m_scrollPosY = Min(0.0f, float(m_scrollPosY - scrollAnchorOffsetY)*(m_zoomValue/oldZoomValue) + scrollAnchorOffsetY);
+					m_scrollPosX = Min(0.0f, float(m_scrollPosX - scrollAnchorOffsetX)*(m_zoomValue/oldZoomValue)*(newScaleValue/m_horizontalScaleValue) + scrollAnchorOffsetX);//LOWORD(lParam);
+
+					if (m_horizontalScaleValue != newScaleValue)
+						m_horizontalScaleValue = newScaleValue;
+
 
 					UpdateAutoscroll();
 					UpdateSelection();
diff --git a/Engine/Source/Programs/UnrealBuildAccelerator/Visualizer/Private/UbaVisualizer.h b/Engine/Source/Programs/UnrealBuildAccelerator/Visualizer/Private/UbaVisualizer.h
--- a/Engine/Source/Programs/UnrealBuildAccelerator/Visualizer/Private/UbaVisualizer.h
+++ b/Engine/Source/Programs/UnrealBuildAccelerator/Visualizer/Private/UbaVisualizer.h
@@ -191,7 +191,7 @@ namespace uba
 		HPEN m_memPen = 0;
 		HPEN m_processUpdatePen = 0;
 		HPEN m_checkboxPen = 0;
-		int m_boxHeight = 0;
+		int m_boxHeight = 12;
 		int m_sessionStepY = 0;
 
 		Font m_defaultFont;
