[UBA]
* Fixed obj preload code so it handles rsp paths that contains spaces

diff --git a/Engine/Source/Programs/UnrealBuildAccelerator/Detours/Private/Windows/UbaDetoursObjFilesPreloader.h b/Engine/Source/Programs/UnrealBuildAccelerator/Detours/Private/Windows/UbaDetoursObjFilesPreloader.h
--- a/Engine/Source/Programs/UnrealBuildAccelerator/Detours/Private/Windows/UbaDetoursObjFilesPreloader.h
+++ b/Engine/Source/Programs/UnrealBuildAccelerator/Detours/Private/Windows/UbaDetoursObjFilesPreloader.h
@@ -92,21 +92,15 @@ namespace uba
 			if (line.data[0] == '@')
 			{
 				const wchar_t* begin = line.data + 1;
-				const wchar_t* end = nullptr;
+				const wchar_t* end = line.data + line.count;
 				StringBuffer<> rspFile;
 				if (*begin == '\"')
 				{
 					++begin;
-					end = wcschr(begin, '\"');
+					--end;
 					if (!end)
 						return;
 				}
-				else
-				{
-					end = wcschr(begin, ' ');
-					if (!end)
-						end = begin + wcslen(begin);
-				}
 				rspFile.Append(begin, end - begin);
 				ParseRsp(rspFile);
 				return;
