[UBA]
* Added reason for process being returned by client to trace file and visualizer

diff --git a/Engine/Source/Programs/UnrealBuildAccelerator/Common/Private/UbaSessionServer.cpp b/Engine/Source/Programs/UnrealBuildAccelerator/Common/Private/UbaSessionServer.cpp
--- a/Engine/Source/Programs/UnrealBuildAccelerator/Common/Private/UbaSessionServer.cpp
+++ b/Engine/Source/Programs/UnrealBuildAccelerator/Common/Private/UbaSessionServer.cpp
@@ -510,7 +510,7 @@ namespace uba
 			it = m_activeRemoteProcesses.erase(it);
 			remoteProcess->m_executingHost.clear();
 
-			m_trace.ProcessReturned(remoteProcess->m_processId);
+			m_trace.ProcessReturned(remoteProcess->m_processId, ToView(TC("Disconnected")));
 
 			ProcessHandle h = ProcessRemoved(remoteProcess->m_processId);
 			if (!h.m_process)
@@ -1269,7 +1269,7 @@ namespace uba
 		process->m_clientId = ~0u;
 		process->m_sessionId = 0;
 
-		m_trace.ProcessReturned(process->m_processId);
+		m_trace.ProcessReturned(process->m_processId, reason);
 		m_queuedRemoteProcesses.push_front(h);
 
 		if (m_remoteProcessReturnedEvent)
diff --git a/Engine/Source/Programs/UnrealBuildAccelerator/Common/Private/UbaTrace.cpp b/Engine/Source/Programs/UnrealBuildAccelerator/Common/Private/UbaTrace.cpp
--- a/Engine/Source/Programs/UnrealBuildAccelerator/Common/Private/UbaTrace.cpp
+++ b/Engine/Source/Programs/UnrealBuildAccelerator/Common/Private/UbaTrace.cpp
@@ -281,10 +281,11 @@ namespace uba
 		writer.WriteByte(255);
 	}
 
-	void Trace::ProcessReturned(u32 processId)
+	void Trace::ProcessReturned(u32 processId, const StringView& reason)
 	{
 		BEGIN_TRACE_ENTRY(ProcessReturned);
 		writer.WriteU32(processId);
+		writer.WriteString(reason);
 	}
 
 	void Trace::ProxyCreated(u32 clientId, const tchar* proxyName)
diff --git a/Engine/Source/Programs/UnrealBuildAccelerator/Common/Public/UbaTrace.h b/Engine/Source/Programs/UnrealBuildAccelerator/Common/Public/UbaTrace.h
--- a/Engine/Source/Programs/UnrealBuildAccelerator/Common/Public/UbaTrace.h
+++ b/Engine/Source/Programs/UnrealBuildAccelerator/Common/Public/UbaTrace.h
@@ -68,7 +68,7 @@ namespace uba
 	using Color = u32;
 	inline Color toColor(u8 r, u8 g, u8 b) { return (r << 16) + (g << 8) + b; }
 
-	static constexpr u32 TraceVersion = 32;
+	static constexpr u32 TraceVersion = 33;
 	static constexpr u32 TraceReadCompatibilityVersion = 6;
 
 	class Trace : public WorkTracker
@@ -88,7 +88,7 @@ namespace uba
 		void ProcessAdded(u32 sessionId, u32 processId, const tchar* description);
 		void ProcessEnvironmentUpdated(u32 processId, const tchar* reason, const u8* data, u64 dataSize);
 		void ProcessExited(u32 processId, u32 exitCode, const u8* data, u64 dataSize, const Vector<ProcessLogLine>& logLines);
-		void ProcessReturned(u32 processId);
+		void ProcessReturned(u32 processId, const StringView& reason);
 		void ProxyCreated(u32 clientId, const tchar* proxyName);
 		void ProxyUsed(u32 clientId, const tchar* proxyName);
 		void FileBeginFetch(u32 clientId, const CasKey& key, u64 size, const tchar* hint, bool detailed);
diff --git a/Engine/Source/Programs/UnrealBuildAccelerator/Core/Public/UbaBinaryReaderWriter.h b/Engine/Source/Programs/UnrealBuildAccelerator/Core/Public/UbaBinaryReaderWriter.h
--- a/Engine/Source/Programs/UnrealBuildAccelerator/Core/Public/UbaBinaryReaderWriter.h
+++ b/Engine/Source/Programs/UnrealBuildAccelerator/Core/Public/UbaBinaryReaderWriter.h
@@ -17,7 +17,7 @@ namespace uba
 		inline void WriteU64(u64 value);
 		inline void WriteString(const tchar* str);
 		inline void WriteString(const tchar* str, u64 len);
-		inline void WriteString(const StringBufferBase& str);
+		inline void WriteString(const StringView& str);
 		inline void WriteStringf(const tchar* format, ...);
 		inline void WriteString(const TString& str);
 		inline void WriteStringKey(const StringKey& g);
@@ -199,7 +199,7 @@ namespace uba
 		#endif
 	}
 
-	void BinaryWriter::WriteString(const StringBufferBase& str)
+	void BinaryWriter::WriteString(const StringView& str)
 	{
 		WriteString(str.data, str.count);
 	}
diff --git a/Engine/Source/Programs/UnrealBuildAccelerator/Visualizer/Private/UbaTraceReader.cpp b/Engine/Source/Programs/UnrealBuildAccelerator/Visualizer/Private/UbaTraceReader.cpp
--- a/Engine/Source/Programs/UnrealBuildAccelerator/Visualizer/Private/UbaTraceReader.cpp
+++ b/Engine/Source/Programs/UnrealBuildAccelerator/Visualizer/Private/UbaTraceReader.cpp
@@ -683,6 +683,12 @@ namespace uba
 					return true;
 				}
 			}
+			TString reason;
+			if (out.version >= 33)
+				reason = reader.ReadString();
+			if (reason.empty())
+				reason = TC("Unknown");
+
 			auto findIt = m_activeProcesses.find(id);
 			if (findIt == m_activeProcesses.end())
 				return false;
@@ -696,7 +702,7 @@ namespace uba
 			TraceView::Process& process = session.processors[active.processorIndex].processes[active.processIndex];
 			process.exitCode = 0;
 			process.stop = time;
-			process.returned = true;
+			process.returnedReason = reason;
 			process.bitmapDirty = true;
 			break;
 		}
@@ -1011,7 +1017,8 @@ namespace uba
 			u32 sessionIndex;
 			TraceView::Process& process = *ProcessEnd(out, sessionIndex, id, time);
 			process.exitCode = 0;//success ? 0 : -1;
-			process.returned = !success;
+			if (!success)
+				process.returnedReason = TC("M");
 
 			CacheStats cacheStats;
 			KernelStats kernelStats;
diff --git a/Engine/Source/Programs/UnrealBuildAccelerator/Visualizer/Private/UbaTraceReader.h b/Engine/Source/Programs/UnrealBuildAccelerator/Visualizer/Private/UbaTraceReader.h
--- a/Engine/Source/Programs/UnrealBuildAccelerator/Visualizer/Private/UbaTraceReader.h
+++ b/Engine/Source/Programs/UnrealBuildAccelerator/Visualizer/Private/UbaTraceReader.h
@@ -20,10 +20,10 @@ namespace uba
 			u64 start = 0;
 			u64 stop = 0;
 			TString description;
+			TString returnedReason;
 			HBITMAP bitmap = 0;
 			u32 bitmapOffset = 0;
 			bool bitmapDirty = true;
-			bool returned = false;
 			bool cacheFetch = false;
 			bool isRemote = false;
 			u64 createFilesTime = 0;
diff --git a/Engine/Source/Programs/UnrealBuildAccelerator/Visualizer/Private/UbaVisualizer.cpp b/Engine/Source/Programs/UnrealBuildAccelerator/Visualizer/Private/UbaVisualizer.cpp
--- a/Engine/Source/Programs/UnrealBuildAccelerator/Visualizer/Private/UbaVisualizer.cpp
+++ b/Engine/Source/Programs/UnrealBuildAccelerator/Visualizer/Private/UbaVisualizer.cpp
@@ -1755,6 +1755,8 @@ namespace uba
 			logger.Info(L"  ProcessId: %6u", process.id);
 			logger.Info(L"  Start:    %7ls", TimeToText(process.start, true).str);
 			logger.Info(L"  Duration: %7ls", TimeToText(duration, true).str);
+			if (!process.returnedReason.empty())
+				logger.Info(L"  Returned: %7s", process.returnedReason.data());
 			if (hasExited && process.exitCode != 0)
 				logger.Info(L"  ExitCode: %7u", process.exitCode);
 
@@ -1769,7 +1771,7 @@ namespace uba
 
 				if (process.cacheFetch)
 				{
-					if (process.returned)
+					if (!process.returnedReason.empty())
 						logger.Info(L"  Cache:       Miss");
 					else
 						logger.Info(L"  Cache:        Hit");
@@ -2086,7 +2088,7 @@ namespace uba
 		bool done = process.stop != ~u64(0);
 
 		HBRUSH brush = m_processBrushes[selected].success;
-		if (process.returned)
+		if (!process.returnedReason.empty())
 			brush = m_processBrushes[selected].returned;
 		else if (!done)
 			brush = m_processBrushes[selected].inProgress;
