[UBA]
* Fixed so files created with deleteonclose always are tracked by session. Fixes builds using -sarif

diff --git a/Engine/Source/Programs/UnrealBuildAccelerator/Detours/Private/Windows/UbaDetoursFunctionsKernelBase.inl b/Engine/Source/Programs/UnrealBuildAccelerator/Detours/Private/Windows/UbaDetoursFunctionsKernelBase.inl
--- a/Engine/Source/Programs/UnrealBuildAccelerator/Detours/Private/Windows/UbaDetoursFunctionsKernelBase.inl
+++ b/Engine/Source/Programs/UnrealBuildAccelerator/Detours/Private/Windows/UbaDetoursFunctionsKernelBase.inl
@@ -1290,7 +1290,7 @@ BOOL Detoured_CopyFileExW(LPCWSTR lpExistingFileName, LPCWSTR lpNewFileName, LPP
 		SuppressCreateFileDetourScope cfs;
 		res = True_CopyFileExW(newFromName.data, newToName.data, lpProgressRoutine, lpData, pbCancel, dwCopyFlags);
 	}
-	DEBUG_LOG_TRUE(L"CopyFileExW", L"%ls to %ls  (%ls to %ls) -> %ls", lpExistingFileName, lpNewFileName, newFromName.data, newToName.data, ToString(res));
+	DEBUG_LOG_TRUE(L"CopyFileExW", L"%ls to %ls flags: %u (%ls to %ls) -> %ls", lpExistingFileName, lpNewFileName, dwCopyFlags, newFromName.data, newToName.data, ToString(res));
 
 	// We need to report the new file that has been added (and we must do it _after_ it has been copied
 	if (!closeId)
diff --git a/Engine/Source/Programs/UnrealBuildAccelerator/Detours/Private/Windows/UbaDetoursFunctionsNtDll.inl b/Engine/Source/Programs/UnrealBuildAccelerator/Detours/Private/Windows/UbaDetoursFunctionsNtDll.inl
--- a/Engine/Source/Programs/UnrealBuildAccelerator/Detours/Private/Windows/UbaDetoursFunctionsNtDll.inl
+++ b/Engine/Source/Programs/UnrealBuildAccelerator/Detours/Private/Windows/UbaDetoursFunctionsNtDll.inl
@@ -931,7 +931,7 @@ NTSTATUS NTAPI Shared_NtCreateFile(bool IsCreateFunc, PHANDLE hFileHandle, ACCES
 		{
 			bool lastWasWrite = IsContentWrite(info.lastDesiredAccess, 0);
 			UBA_ASSERT(!info.isFileMap);
-			bool shouldReport = !lastWasWrite || info.deleted;
+			bool shouldReport = !lastWasWrite || info.deleted || isDeleteOnClose;
 			shouldReport = shouldReport && !keepInMemory;
 			if (shouldReport)
 			{
@@ -1219,7 +1219,7 @@ NTSTATUS NTAPI Shared_NtCreateFile(bool IsCreateFunc, PHANDLE hFileHandle, ACCES
 	dh->fileObject->fileInfo = &info;
 	dh->fileObject->deleteOnClose = isDeleteOnClose;
 	*hFileHandle = makeDetouredHandle(dh);
-	DEBUG_LOG_TRUE(funcName, L"%ls %llu (%ls) -> %ls", isWriteStr, uintptr_t(*hFileHandle), tempFileName, ToString(res));
+	DEBUG_LOG_TRUE(funcName, L"%ls %llu (%ls)%s -> %ls", isWriteStr, uintptr_t(*hFileHandle), tempFileName, isDeleteOnClose ? TC(" DeleteOnClose") : TC(""), ToString(res));
 	return res;
 }
 
