[UBA]
* Fixed printf bug in code path handling corrupt cas files...

diff --git a/Engine/Source/Programs/UnrealBuildAccelerator/Cli/Private/UbaCli.cpp b/Engine/Source/Programs/UnrealBuildAccelerator/Cli/Private/UbaCli.cpp
--- a/Engine/Source/Programs/UnrealBuildAccelerator/Cli/Private/UbaCli.cpp
+++ b/Engine/Source/Programs/UnrealBuildAccelerator/Cli/Private/UbaCli.cpp
@@ -51,9 +51,15 @@ namespace uba
 			logger.Info(TC(""));
 			logger.Error(TC("%s"), message);
 		}
+
+		const tchar* dbgStr = TC("");
+		#if UBA_DEBUG
+		dbgStr = TC(" (DEBUG)");
+		#endif
+
 		logger.Info(TC(""));
 		logger.Info(TC("-------------------------------------------"));
-		logger.Info(TC("   UbaCli v%s"), Version);
+		logger.Info(TC("   UbaCli v%s%s"), Version, dbgStr);
 		logger.Info(TC("-------------------------------------------"));
 		logger.Info(TC(""));
 		logger.Info(TC("  UbaCli.exe [options...] <commandtype> <executable> [arguments...]"));
diff --git a/Engine/Source/Programs/UnrealBuildAccelerator/Common/Private/UbaStorage.cpp b/Engine/Source/Programs/UnrealBuildAccelerator/Common/Private/UbaStorage.cpp
--- a/Engine/Source/Programs/UnrealBuildAccelerator/Common/Private/UbaStorage.cpp
+++ b/Engine/Source/Programs/UnrealBuildAccelerator/Common/Private/UbaStorage.cpp
@@ -990,7 +990,7 @@ namespace uba
 							BinaryReader reader(fa.GetData(), 0, fa.GetSize());
 							if (reader.GetLeft() < 12)
 							{
-								m_logger.Detail(TC("Corrupt cas. Is %llu, must be at least 12 bytes (%s)"), e.name);
+								m_logger.Detail(TC("Corrupt cas. Is %llu, must be at least 12 bytes (%s)"), reader.GetLeft(), e.name);
 								deleteFile = true;
 							}
 							else
@@ -1073,7 +1073,7 @@ namespace uba
 				if (DeleteFileW(fullPath.data))
 					return;
 
-				m_logger.Error(TC("Failed to delete file (%s)"), fullPath.data);
+				m_logger.Error(TC("Failed to delete file %s (%s)"), fullPath.data, LastErrorToText().data);
 				success = false;
 
 			}, true);
@@ -1151,19 +1151,19 @@ namespace uba
 		FileHandle sparseFile = uba::CreateFileW(sparseFileName.data, GENERIC_READ | GENERIC_WRITE, FILE_SHARE_READ | FILE_SHARE_WRITE, createDisposition, DefaultAttributes());
 		auto sparseFileGuard = MakeGuard([&](){ CloseFile(sparseFile); });
 		if (sparseFile == InvalidFileHandle)
-			return m_logger.Error(TC("Failed to create database file %s (%s)"), sparseFileName.data, LastErrorToText());
+			return m_logger.Error(TC("Failed to create database file %s (%s)"), sparseFileName.data, LastErrorToText().data);
 		
 		if (createFile)
 		{
 			DWORD dwTemp;
 			if (!::DeviceIoControl(sparseFile, FSCTL_SET_SPARSE, NULL, 0, NULL, 0, &dwTemp, NULL))
-				return m_logger.Error(TC("Failed to make file %s sparse (%s)"), sparseFileName.data, LastErrorToText());
+				return m_logger.Error(TC("Failed to make file %s sparse (%s)"), sparseFileName.data, LastErrorToText().data);
 		}
 		else
 		{
 			u64 fileSize;
 			if (!uba::GetFileSizeEx(fileSize, sparseFile))
-				return m_logger.Error(TC("GetFileSize failed for %s (%s)"), fileName, LastErrorToText());
+				return m_logger.Error(TC("GetFileSize failed for %s (%s)"), fileName, LastErrorToText().data);
 
 			if (fileSize < size)
 				return m_logger.Error(TC("Sparse file size is smaller than what cas db think it is. Expected %llu, was %llu (%s)"), size, fileSize, sparseFileName.data);
@@ -1493,7 +1493,7 @@ namespace uba
 			FILE_ALLOCATED_RANGE_BUFFER allocRanges[1024];
 			DWORD nbytes;
 			if (!DeviceIoControl(m_casDataBuffer.GetPersistentFile(0), FSCTL_QUERY_ALLOCATED_RANGES, &queryRange, sizeof(queryRange), allocRanges, sizeof(allocRanges), &nbytes, NULL))
-				return m_logger.Error(TC("Failed to make file %s sparse (%s)"), TC("FOO"), LastErrorToText());
+				return m_logger.Error(TC("Failed to make file %s sparse (%s)"), TC("FOO"), LastErrorToText().data);
 			DWORD dwAllocRangeCount = nbytes / sizeof(FILE_ALLOCATED_RANGE_BUFFER);
 			printf("Range count: %u"), dwAllocRangeCount);
 		}
