[UBA]
* Changed code in cache server maintenance logic to use bit array instead of unordered set with an attempt to improve performance (some steps takes 1.8 seconds on the deployed server right now)

diff --git a/Engine/Source/Programs/UnrealBuildAccelerator/Common/Private/UbaCacheServer.cpp b/Engine/Source/Programs/UnrealBuildAccelerator/Common/Private/UbaCacheServer.cpp
--- a/Engine/Source/Programs/UnrealBuildAccelerator/Common/Private/UbaCacheServer.cpp
+++ b/Engine/Source/Programs/UnrealBuildAccelerator/Common/Private/UbaCacheServer.cpp
@@ -10,6 +10,12 @@
 #include "UbaStorageServer.h"
 //#include <oodle2.h>
 
+#if PLATFORM_WINDOWS
+#define UBA_FORCEINLINE __forceinline
+#else
+#define UBA_FORCEINLINE inline __attribute__ ((always_inline))
+#endif
+
 namespace uba
 {
 	static constexpr u32 CacheFileVersion = 6;
@@ -55,6 +61,82 @@ namespace uba
 		u32 index = ~0u;
 	};
 
+	static constexpr UBA_FORCEINLINE u64 CountBits(u64 Bits)
+	{
+		// https://en.wikipedia.org/wiki/Hamming_weight
+		Bits -= (Bits >> 1) & 0x5555555555555555ull;
+		Bits = (Bits & 0x3333333333333333ull) + ((Bits >> 2) & 0x3333333333333333ull);
+		Bits = (Bits + (Bits >> 4)) & 0x0f0f0f0f0f0f0f0full;
+		return (Bits * 0x0101010101010101) >> 56;
+	}
+
+	static constexpr UBA_FORCEINLINE u64 FindFirstBit(u64 Value)
+	{
+		u64 pos = 0;
+		if (Value >= 1ull<<32) { Value >>= 32; pos += 32; }
+		if (Value >= 1ull<<16) { Value >>= 16; pos += 16; }
+		if (Value >= 1ull<< 8) { Value >>=  8; pos +=  8; }
+		if (Value >= 1ull<< 4) { Value >>=  4; pos +=  4; }
+		if (Value >= 1ull<< 2) { Value >>=  2; pos +=  2; }
+		if (Value >= 1ull<< 1) {               pos +=  1; }
+		return pos;
+	}
+
+	struct BitArray
+	{
+		BitArray(MemoryBlock& memoryBlock, u32 bitCount)
+		{
+			size = AlignUp(bitCount / 8, 8); // Align up to 64 bits
+			data = (u64*)memoryBlock.Allocate(size, 8, TC(""));
+			memset(data, 0, size);
+		}
+
+		UBA_FORCEINLINE void Set(u32 index)
+		{
+			u32 byteIndex = index / 64;
+			u32 bitIndex = index - byteIndex * 64;
+			data[byteIndex] |= 1ull << bitIndex;
+		}
+
+		UBA_FORCEINLINE u32 CountSetBits()
+		{
+			u64 count = 0;
+			for (u64 i=0,e=size/8; i!=e; ++i)
+				count += CountBits(data[i]);
+			return u32(count);
+		}
+
+		template<typename Func>
+		UBA_FORCEINLINE void Traverse(const Func& func)
+		{
+			u32 index = 0;
+			for (u64 i=0,e=size/8; i!=e; ++i)
+			{
+				u64 v = data[i];
+				/*
+				if (!v)
+				{
+					u32 index2 = index + u32(FindFirstBit(v));
+					func(index2);
+					index += 128;
+				}
+				*/
+				u32 index2 = index;
+				while (v)
+				{
+					if (v & 1)
+						func(index2);
+					v >>= 1;
+					++index2;
+				}
+				index += 64;
+			}
+		}
+
+		u64* data;
+		u32 size;
+	};
+
 	const tchar* ToString(CacheMessageType type)
 	{
 		switch (type)
@@ -477,7 +559,7 @@ namespace uba
 		else
 		{
 			TimeToText lastTime(startTime - m_lastMaintenance, true);
-			m_logger.Info(TC("Maintenance started after %u added cache entries (Ran last time %s ago)"), entriesAdded, (m_lastMaintenance ? lastTime.str : TC("<never>")));
+			m_logger.Info(TC("Maintenance started after %u added cache entries (Ran last time %s ago)"), m_addsSinceMaintenance.load(), (m_lastMaintenance ? lastTime.str : TC("<never>")));
 		}
 
 		m_lastMaintenance = startTime;
@@ -793,8 +875,7 @@ namespace uba
 			if (!memoryBlock.Init(m_maintenanceReserveSize, nullptr, true))
 				memoryBlock.Init(m_maintenanceReserveSize, nullptr, false);
 
-			GrowingNoLockUnorderedSet<u32> usedCasKeyOffsets(&memoryBlock);
-			usedCasKeyOffsets.reserve(bucket.m_casKeyTable.GetKeyCount());
+			BitArray usedCasKeyOffsets(memoryBlock, u32(bucket.m_casKeyTable.GetSize()));
 
 			u64 collectUsedCasKeysStart = GetTime();
 
@@ -805,7 +886,10 @@ namespace uba
 					{
 						BinaryReader reader2(offsets.data(), 0, offsets.size());
 						while (reader2.GetLeft())
-							usedCasKeyOffsets.insert(u32(reader2.Read7BitEncoded()));
+						{
+							u32 offset = u32(reader2.Read7BitEncoded());
+							usedCasKeyOffsets.Set(offset);
+						}
 					};
 
 				collectUsedCasKeyOffsets(kv2.second.sharedInputCasKeyOffsets);
@@ -815,32 +899,34 @@ namespace uba
 					collectUsedCasKeyOffsets(entry.outputCasKeyOffsets);
 				}
 			}
-			m_logger.Detail(TC("    Bucket %u Collected %s used caskeys. (%s)"), bucket.index, CountToText(usedCasKeyOffsets.size()).str, TimeToText(GetTime() - collectUsedCasKeysStart).str);
+			u64 usedCasKeyOffsetsCount = usedCasKeyOffsets.CountSetBits();
+
+			m_logger.Detail(TC("    Bucket %u Collected %s used caskeys. (%s)"), bucket.index, CountToText(usedCasKeyOffsetsCount).str, TimeToText(GetTime() - collectUsedCasKeysStart).str);
 
 			u64 recreatePathTableStart = GetTime();
 
 			// Traverse all caskeys in caskey table and figure out which ones we can delete
-			GrowingNoLockUnorderedSet<u32> usedPathOffsets(&memoryBlock);
-			usedPathOffsets.reserve(usedCasKeyOffsets.size());
+			BitArray usedPathOffsets(memoryBlock, u32(bucket.m_pathTable.GetSize()));
 
-			for (u32 casKeyOffset : usedCasKeyOffsets)
-			{
-				BinaryReader reader2(bucket.m_casKeyTable.GetMemory(), casKeyOffset, bucket.m_casKeyTable.GetSize());
-				u32 pathOffset = u32(reader2.Read7BitEncoded());
-				usedPathOffsets.insert(pathOffset);
-			}
+			BinaryReader casKeyTableReader(bucket.m_casKeyTable.GetMemory(), 0, bucket.m_casKeyTable.GetSize());
+			usedCasKeyOffsets.Traverse([&](u32 casKeyOffset)
+				{
+					casKeyTableReader.SetPosition(casKeyOffset);
+					u32 pathOffset = u32(casKeyTableReader.Read7BitEncoded());
+					usedPathOffsets.Set(pathOffset);
+				});
 
 			// Build new path table based on used offsets
 			GrowingNoLockUnorderedMap<u32, u32> oldToNewPathOffset(&memoryBlock);
 			u32 oldSize = bucket.m_pathTable.GetSize();
 			{
 				CompactPathTable newPathTable(CachePathTableMaxSize, CompactPathTable::V1, bucket.m_pathTable.GetPathCount(), bucket.m_pathTable.GetSegmentCount());
-				oldToNewPathOffset.reserve(usedPathOffsets.size());
+				oldToNewPathOffset.reserve(usedPathOffsets.CountSetBits());
 
-				for (u32 pathOffset : usedPathOffsets)
+				StringBuffer<> temp;
+				usedPathOffsets.Traverse([&](u32 pathOffset)
 				{
-					StringBuffer<> temp;
-					bucket.m_pathTable.GetString(temp, pathOffset);
+					bucket.m_pathTable.GetString(temp.Clear(), pathOffset);
 					u32 newOffset = newPathTable.AddNoLock(temp.data, temp.count);
 
 					#if 0
@@ -851,7 +937,7 @@ namespace uba
 
 					auto res = oldToNewPathOffset.try_emplace(pathOffset, newOffset);
 					UBA_ASSERT(res.second);(void)res;
-				}
+				});
 				bucket.m_pathTable.Swap(newPathTable);
 			}
 			m_logger.Detail(TC("    Bucket %u Recreated path table. %s -> %s (%s)"), bucket.index, BytesToText(oldSize).str, BytesToText(bucket.m_pathTable.GetSize()).str, TimeToText(GetTime() - recreatePathTableStart).str);
@@ -862,21 +948,22 @@ namespace uba
 			GrowingNoLockUnorderedMap<u32, u32> oldToNewCasKeyOffset(&memoryBlock);
 			oldSize = bucket.m_casKeyTable.GetSize();
 			{
-				oldToNewCasKeyOffset.reserve(usedCasKeyOffsets.size());
-				CompactCasKeyTable newCasKeyTable(CacheCasKeyTableMaxSize, usedCasKeyOffsets.size());
-				for (u32 casKeyOffset : usedCasKeyOffsets)
+				oldToNewCasKeyOffset.reserve(usedCasKeyOffsetsCount);
+				CompactCasKeyTable newCasKeyTable(CacheCasKeyTableMaxSize, usedCasKeyOffsetsCount);
+				BinaryReader reader2(bucket.m_casKeyTable.GetMemory(), 0, bucket.m_casKeyTable.GetSize());
+				usedCasKeyOffsets.Traverse([&](u32 casKeyOffset)
 				{
-					BinaryReader reader2(bucket.m_casKeyTable.GetMemory(), casKeyOffset, bucket.m_casKeyTable.GetSize());
+					reader2.SetPosition(casKeyOffset);
 					u32 oldPathOffset = u32(reader2.Read7BitEncoded());
 					CasKey casKey = reader2.ReadCasKey();
 					auto findIt = oldToNewPathOffset.find(oldPathOffset);
 					UBA_ASSERT(findIt != oldToNewPathOffset.end());
 					u32 newCasKeyOffset = newCasKeyTable.Add(casKey, findIt->second);
 					if (casKeyOffset == newCasKeyOffset)
-						continue;
+						return;
 					auto res = oldToNewCasKeyOffset.try_emplace(casKeyOffset, newCasKeyOffset);
 					UBA_ASSERT(res.second);(void)res;
-				}
+				});
 				bucket.m_casKeyTable.Swap(newCasKeyTable);
 			}
 			m_logger.Detail(TC("    Bucket %u Recreated caskey table. %s -> %s (%s)"), bucket.index, BytesToText(oldSize).str, BytesToText(bucket.m_casKeyTable.GetSize()).str, TimeToText(GetTime() - recreateCasKeyTableStart).str);
