[UBA]
* Fixed logging bug in cache server
* Did some minor optimizations to cache server maintenance code

diff --git a/Engine/Source/Programs/UnrealBuildAccelerator/Common/Private/UbaCacheServer.cpp b/Engine/Source/Programs/UnrealBuildAccelerator/Common/Private/UbaCacheServer.cpp
--- a/Engine/Source/Programs/UnrealBuildAccelerator/Common/Private/UbaCacheServer.cpp
+++ b/Engine/Source/Programs/UnrealBuildAccelerator/Common/Private/UbaCacheServer.cpp
@@ -61,34 +61,14 @@ namespace uba
 		u32 index = ~0u;
 	};
 
-	static constexpr UBA_FORCEINLINE u64 CountBits(u64 Bits)
-	{
-		// https://en.wikipedia.org/wiki/Hamming_weight
-		Bits -= (Bits >> 1) & 0x5555555555555555ull;
-		Bits = (Bits & 0x3333333333333333ull) + ((Bits >> 2) & 0x3333333333333333ull);
-		Bits = (Bits + (Bits >> 4)) & 0x0f0f0f0f0f0f0f0full;
-		return (Bits * 0x0101010101010101) >> 56;
-	}
-
-	static constexpr UBA_FORCEINLINE u64 FindFirstBit(u64 Value)
-	{
-		u64 pos = 0;
-		if (Value >= 1ull<<32) { Value >>= 32; pos += 32; }
-		if (Value >= 1ull<<16) { Value >>= 16; pos += 16; }
-		if (Value >= 1ull<< 8) { Value >>=  8; pos +=  8; }
-		if (Value >= 1ull<< 4) { Value >>=  4; pos +=  4; }
-		if (Value >= 1ull<< 2) { Value >>=  2; pos +=  2; }
-		if (Value >= 1ull<< 1) {               pos +=  1; }
-		return pos;
-	}
-
 	struct BitArray
 	{
 		BitArray(MemoryBlock& memoryBlock, u32 bitCount)
 		{
-			size = AlignUp(bitCount / 8, 8); // Align up to 64 bits
-			data = (u64*)memoryBlock.Allocate(size, 8, TC(""));
-			memset(data, 0, size);
+			u32 bytes = AlignUp(bitCount / 8, 8u); // Align up to 64 bits
+			data = (u64*)memoryBlock.Allocate(bytes, 8, TC(""));
+			memset(data, 0, bytes);
+			count = bytes / 8;
 		}
 
 		UBA_FORCEINLINE void Set(u32 index)
@@ -100,41 +80,58 @@ namespace uba
 
 		UBA_FORCEINLINE u32 CountSetBits()
 		{
-			u64 count = 0;
-			for (u64 i=0,e=size/8; i!=e; ++i)
-				count += CountBits(data[i]);
-			return u32(count);
+			u64 bits = 0;
+			for (u64 i=0,e=count; i!=e; ++i)
+				bits += CountBits(data[i]);
+			return u32(bits);
 		}
 
 		template<typename Func>
-		UBA_FORCEINLINE void Traverse(const Func& func)
+		void Traverse(const Func& func)
 		{
 			u32 index = 0;
-			for (u64 i=0,e=size/8; i!=e; ++i)
+			for (u64 i=0,e=count; i!=e; ++i)
 			{
 				u64 v = data[i];
-				/*
-				if (!v)
-				{
-					u32 index2 = index + u32(FindFirstBit(v));
-					func(index2);
-					index += 128;
-				}
-				*/
-				u32 index2 = index;
 				while (v)
 				{
-					if (v & 1)
-						func(index2);
-					v >>= 1;
-					++index2;
+					u64 bitIndex = FindFirstBit(v);
+					func(index + u32(bitIndex));
+					v &= ~(1ull << bitIndex);
 				}
 				index += 64;
 			}
 		}
 
+		static constexpr UBA_FORCEINLINE u64 CountBits(u64 bits)
+		{
+			// https://en.wikipedia.org/wiki/Hamming_weight
+			bits -= (bits >> 1) & 0x5555555555555555ull;
+			bits = (bits & 0x3333333333333333ull) + ((bits >> 2) & 0x3333333333333333ull);
+			bits = (bits + (bits >> 4)) & 0x0f0f0f0f0f0f0f0full;
+			return (bits * 0x0101010101010101) >> 56;
+		}
+
+		static UBA_FORCEINLINE u64 FindFirstBit(u64 v)
+		{
+			#if PLATFORM_WINDOWS
+			return _tzcnt_u64(v);
+			#elif PLATFORM_LINUX
+			return __builtin_ia32_tzcnt_u64(v);
+			#else
+			u64 pos = 0;
+			if (v >= 1ull<<32) { v >>= 32; pos += 32; }
+			if (v >= 1ull<<16) { v >>= 16; pos += 16; }
+			if (v >= 1ull<< 8) { v >>=  8; pos +=  8; }
+			if (v >= 1ull<< 4) { v >>=  4; pos +=  4; }
+			if (v >= 1ull<< 2) { v >>=  2; pos +=  2; }
+			if (v >= 1ull<< 1) {           pos +=  1; }
+			return pos;
+			#endif
+		}
+
 		u64* data;
-		u32 size;
+		u32 count;
 	};
 
 	const tchar* ToString(CacheMessageType type)
@@ -530,7 +527,8 @@ namespace uba
 		bool forceAllSteps = m_forceAllSteps;
 		m_forceAllSteps = false;
 
-		bool entriesAdded = m_addsSinceMaintenance != 0;
+		u32 addsSinceMaintenance = m_addsSinceMaintenance;
+		bool entriesAdded = addsSinceMaintenance != 0;
 		m_addsSinceMaintenance = 0;
 
 		u64 startTime = GetTime();
@@ -559,7 +557,7 @@ namespace uba
 		else
 		{
 			TimeToText lastTime(startTime - m_lastMaintenance, true);
-			m_logger.Info(TC("Maintenance started after %u added cache entries (Ran last time %s ago)"), m_addsSinceMaintenance.load(), (m_lastMaintenance ? lastTime.str : TC("<never>")));
+			m_logger.Info(TC("Maintenance started after %u added cache entries (Ran last time %s ago)"), addsSinceMaintenance, (m_lastMaintenance ? lastTime.str : TC("<never>")));
 		}
 
 		m_lastMaintenance = startTime;
