[UBA]
* Fixed bug in UbaCacheService where bit array was 1 byte too small

diff --git a/Engine/Source/Programs/UnrealBuildAccelerator/Common/Private/UbaCacheServer.cpp b/Engine/Source/Programs/UnrealBuildAccelerator/Common/Private/UbaCacheServer.cpp
--- a/Engine/Source/Programs/UnrealBuildAccelerator/Common/Private/UbaCacheServer.cpp
+++ b/Engine/Source/Programs/UnrealBuildAccelerator/Common/Private/UbaCacheServer.cpp
@@ -65,17 +65,18 @@ namespace uba
 	{
 		BitArray(MemoryBlock& memoryBlock, u32 bitCount)
 		{
-			u32 bytes = AlignUp(bitCount / 8, 8u); // Align up to 64 bits
+			u32 bytes = AlignUp((bitCount+7) / 8, 8u); // Align up to 64 bits
 			data = (u64*)memoryBlock.Allocate(bytes, 8, TC(""));
 			memset(data, 0, bytes);
 			count = bytes / 8;
 		}
 
-		UBA_FORCEINLINE void Set(u32 index)
+		UBA_FORCEINLINE void Set(u32 bitIndex)
 		{
-			u32 byteIndex = index / 64;
-			u32 bitIndex = index - byteIndex * 64;
-			data[byteIndex] |= 1ull << bitIndex;
+			u32 index = bitIndex / 64;
+			UBA_ASSERTF(index < count, TC("Out of bounds (%u/%u). Bit index : %u"), index, count, bitIndex);
+			u32 bitOffset = bitIndex - index * 64;
+			data[index] |= 1ull << bitOffset;
 		}
 
 		UBA_FORCEINLINE u32 CountSetBits()
@@ -955,7 +956,7 @@ namespace uba
 					u32 oldPathOffset = u32(reader2.Read7BitEncoded());
 					CasKey casKey = reader2.ReadCasKey();
 					auto findIt = oldToNewPathOffset.find(oldPathOffset);
-					UBA_ASSERT(findIt != oldToNewPathOffset.end());
+					UBA_ASSERTF(findIt != oldToNewPathOffset.end(), TC("Can't find entry with offset %u"), oldPathOffset);
 					u32 newCasKeyOffset = newCasKeyTable.Add(casKey, findIt->second);
 					if (casKeyOffset == newCasKeyOffset)
 						return;
