[UBA]
* Fixed arm compile errors in cache server

Thank you ZeroErrors[REDACTED_HANDLE] for the fix

diff --git a/Engine/Source/Programs/UnrealBuildAccelerator/Common/Private/UbaCacheServer.cpp b/Engine/Source/Programs/UnrealBuildAccelerator/Common/Private/UbaCacheServer.cpp
--- a/Engine/Source/Programs/UnrealBuildAccelerator/Common/Private/UbaCacheServer.cpp
+++ b/Engine/Source/Programs/UnrealBuildAccelerator/Common/Private/UbaCacheServer.cpp
@@ -115,10 +115,18 @@ namespace uba
 
 		static UBA_FORCEINLINE u64 FindFirstBit(u64 v)
 		{
-			#if PLATFORM_WINDOWS
+			#if PLATFORM_WINDOWS && (defined(_M_X64) || defined(_M_IX86))
+			// Use TZCNT intrinsic on Windows x86/x64
 			return _tzcnt_u64(v);
-			#elif PLATFORM_LINUX
+			#elif PLATFORM_WINDOWS && defined(_M_ARM64)
+			// Use the ARM64 equivalent
+			return _CountTrailingZeros64(v);
+			#elif PLATFORM_LINUX && (defined(__x86_64__) || defined(__i386__))
+			// Use GCC's built-in TZCNT equivalent for x86/x64
 			return __builtin_ia32_tzcnt_u64(v);
+			#elif PLATFORM_LINUX && defined(__aarch64__)
+			// Use the ARM64 equivalent
+			return __builtin_ctzll(v);
 			#else
 			u64 pos = 0;
 			if (v >= 1ull<<32) { v >>= 32; pos += 32; }
