[UBA]
* Changed so network client/server listen to 0.0.0.0 by default instead of traversing network adapters.

Thanks ZeroErrors[REDACTED_HANDLE] for reporting the issue with the network adapter traversal.

diff --git a/Engine/Source/Programs/UnrealBuildAccelerator/Common/Private/UbaNetworkClient.cpp b/Engine/Source/Programs/UnrealBuildAccelerator/Common/Private/UbaNetworkClient.cpp
--- a/Engine/Source/Programs/UnrealBuildAccelerator/Common/Private/UbaNetworkClient.cpp
+++ b/Engine/Source/Programs/UnrealBuildAccelerator/Common/Private/UbaNetworkClient.cpp
@@ -298,9 +298,9 @@ namespace uba
 		FlushWork();
 	}
 
-	bool NetworkClient::StartListen(NetworkBackend& backend, u16 port)
+	bool NetworkClient::StartListen(NetworkBackend& backend, u16 port, const tchar* ip)
 	{
-		backend.StartListen(m_logger, port, nullptr, [&](void* connection, const sockaddr& remoteSockAddr)
+		backend.StartListen(m_logger, port, ip, [&](void* connection, const sockaddr& remoteSockAddr)
 			{
 				return AddConnection(backend, connection, nullptr);
 			});
diff --git a/Engine/Source/Programs/UnrealBuildAccelerator/Common/Public/UbaNetworkClient.h b/Engine/Source/Programs/UnrealBuildAccelerator/Common/Public/UbaNetworkClient.h
--- a/Engine/Source/Programs/UnrealBuildAccelerator/Common/Public/UbaNetworkClient.h
+++ b/Engine/Source/Programs/UnrealBuildAccelerator/Common/Public/UbaNetworkClient.h
@@ -36,7 +36,7 @@ namespace uba
 		bool Connect(NetworkBackend& backend, const tchar* ip, u16 port = DefaultPort, bool* timedOut = nullptr);
 		void Disconnect();
 
-		bool StartListen(NetworkBackend& backend, u16 port = DefaultPort);
+		bool StartListen(NetworkBackend& backend, u16 port = DefaultPort, const tchar* ip = TC("0.0.0.0"));
 		bool SetConnectionCount(u32 count);
 		bool SendKeepAlive();
 
diff --git a/Engine/Source/Programs/UnrealBuildAccelerator/Common/Public/UbaNetworkServer.h b/Engine/Source/Programs/UnrealBuildAccelerator/Common/Public/UbaNetworkServer.h
--- a/Engine/Source/Programs/UnrealBuildAccelerator/Common/Public/UbaNetworkServer.h
+++ b/Engine/Source/Programs/UnrealBuildAccelerator/Common/Public/UbaNetworkServer.h
@@ -51,7 +51,7 @@ namespace uba
 		NetworkServer(bool& outCtorSuccess, const NetworkServerCreateInfo& info = {}, const tchar* name = TC("UbaServer"));
 		virtual ~NetworkServer();
 
-		bool StartListen(NetworkBackend& backend, u16 port = DefaultPort, const tchar* ip = nullptr, bool requiresCrypto = false); // Start listen for new connections/clients
+		bool StartListen(NetworkBackend& backend, u16 port = DefaultPort, const tchar* ip = TC("0.0.0.0"), bool requiresCrypto = false); // Start listen for new connections/clients
 		void DisallowNewClients();	// Disallow new clients to connect but old clients can still create more connections
 		void DisconnectClients();	// Disconnect all active connections
 
