[UBA]
* Changed so *-telemetry.json are only thrown away when process is executing remotely

diff --git a/Engine/Source/Programs/UnrealBuildAccelerator/Core/Private/UbaApplicationRules.cpp b/Engine/Source/Programs/UnrealBuildAccelerator/Core/Private/UbaApplicationRules.cpp
--- a/Engine/Source/Programs/UnrealBuildAccelerator/Core/Private/UbaApplicationRules.cpp
+++ b/Engine/Source/Programs/UnrealBuildAccelerator/Core/Private/UbaApplicationRules.cpp
@@ -30,17 +30,17 @@ namespace uba
 				return false;
 			return true;
 		}
-		virtual bool IsThrowAway(const StringView& fileName) const override
+		virtual bool IsThrowAway(const StringView& fileName, bool isRunningRemote) const override
 		{
-			return fileName.Contains(TC("vctip_")) || Super::IsThrowAway(fileName);
+			return fileName.Contains(TC("vctip_")) || Super::IsThrowAway(fileName, isRunningRemote);
 		}
-		virtual bool KeepInMemory(const StringView& fileName, const tchar* systemTemp) const override
+		virtual bool KeepInMemory(const StringView& fileName, const tchar* systemTemp, bool isRunningRemote) const override
 		{
 			if (fileName.Contains(TC("\\vctip_")))
 				return true;
 			if (fileName.Contains(systemTemp))
 				return true;
-			return Super::KeepInMemory(fileName, systemTemp);
+			return Super::KeepInMemory(fileName, systemTemp, isRunningRemote);
 		}
 
 		virtual bool IsExitCodeSuccess(u32 exitCode) const override
@@ -110,7 +110,7 @@ namespace uba
 				|| fileName.EndsWith(TC(".rc2.res")); // Not really an obj file.. 
 		}
 
-		virtual bool IsThrowAway(const StringView& fileName) const override
+		virtual bool IsThrowAway(const StringView& fileName, bool isRunningRemote) const override
 		{
 			return fileName.Contains(TC(".sup.")); // .sup.lib/exp are throw-away files that we don't want created
 		}
@@ -164,7 +164,7 @@ namespace uba
 	{
 		using Super = ApplicationRulesLinkExe;
 
-		virtual bool KeepInMemory(const StringView& fileName, const tchar* systemTemp) const override
+		virtual bool KeepInMemory(const StringView& fileName, const tchar* systemTemp, bool isRunningRemote) const override
 		{
 			return fileName.EndsWith(TC(".manifest")) && fileName.Contains(systemTemp);
 		}
@@ -298,9 +298,9 @@ namespace uba
 	{
 		using Super = ApplicationRulesClangPlusPlusExe;
 
-		virtual bool IsThrowAway(const StringView& fileName) const override
+		virtual bool IsThrowAway(const StringView& fileName, bool isRunningRemote) const override
 		{
-			return fileName.EndsWith(TC("-telemetry.json")) || Super::IsThrowAway(fileName);
+			return fileName.EndsWith(TC("-telemetry.json")) || Super::IsThrowAway(fileName, isRunningRemote);
 		}
 
 		//virtual bool NeedsSharedMemory(const tchar* file)
@@ -317,9 +317,9 @@ namespace uba
 		//{
 		//	return fileName.EndsWith(TC(".self")) || Equals(file, TC("Symbols.map"));
 		//}
-		virtual bool KeepInMemory(const StringView& fileName, const tchar* systemTemp) const override
+		virtual bool KeepInMemory(const StringView& fileName, const tchar* systemTemp, bool isRunningRemote) const override
 		{
-			return Super::KeepInMemory(fileName, systemTemp)
+			return Super::KeepInMemory(fileName, systemTemp, isRunningRemote)
 				|| fileName.Contains(TC("thinlto-"));// Used by a clang based platform's link time optimization pass. Shared from lto process back to linker process
 		}
 		virtual bool NeedsSharedMemory(const tchar* file) const override
@@ -342,9 +342,10 @@ namespace uba
 			return fileName.Contains(TC(".self")) || Super::IsOutputFile(fileName);
 		}
 
-		virtual bool IsThrowAway(const StringView& fileName) const override
+		virtual bool IsThrowAway(const StringView& fileName, bool isRunningRemote) const override
 		{
-			return fileName.EndsWith(TC("-telemetry.json")) || Super::IsThrowAway(fileName);
+			return Super::IsThrowAway(fileName, isRunningRemote)
+				|| isRunningRemote && fileName.EndsWith(TC("-telemetry.json"));
 		}
 
 		//virtual bool NeedsSharedMemory(const tchar* file)
@@ -362,9 +363,10 @@ namespace uba
 			return fileName.Contains(TC(".self"));
 		}
 
-		virtual bool IsThrowAway(const StringView& fileName) const override
+		virtual bool IsThrowAway(const StringView& fileName, bool isRunningRemote) const override
 		{
-			return fileName.EndsWith(TC("-telemetry.json")) || Super::IsThrowAway(fileName);
+			return Super::IsThrowAway(fileName, isRunningRemote)
+				|| isRunningRemote && fileName.EndsWith(TC("-telemetry.json"));
 		}
 
 		//virtual bool NeedsSharedMemory(const tchar* file)
diff --git a/Engine/Source/Programs/UnrealBuildAccelerator/Core/Public/UbaApplicationRules.h b/Engine/Source/Programs/UnrealBuildAccelerator/Core/Public/UbaApplicationRules.h
--- a/Engine/Source/Programs/UnrealBuildAccelerator/Core/Public/UbaApplicationRules.h
+++ b/Engine/Source/Programs/UnrealBuildAccelerator/Core/Public/UbaApplicationRules.h
@@ -34,16 +34,16 @@ namespace uba
 		}
 
 		// Throw-away means that the file is temporary and will not be used after process exists. (By default these are kept in memory and never touch disk)
-		virtual bool IsThrowAway(const StringView& fileName) const
+		virtual bool IsThrowAway(const StringView& fileName, bool isRunningRemote) const
 		{
 			return false;
 		}
 
 		// Keep file in memory
 		// If this returns true it means that file will be kept in memory and never touch disk.
-		virtual bool KeepInMemory(const StringView& fileName, const tchar* systemTemp) const
+		virtual bool KeepInMemory(const StringView& fileName, const tchar* systemTemp, bool isRunningRemote) const
 		{
-			return IsThrowAway(fileName);
+			return IsThrowAway(fileName, isRunningRemote);
 		}
 
 		// For files that are kept in memory but shared between process (temporary files where one process write and another read)
diff --git a/Engine/Source/Programs/UnrealBuildAccelerator/Detours/Private/UbaDetoursShared.h b/Engine/Source/Programs/UnrealBuildAccelerator/Detours/Private/UbaDetoursShared.h
--- a/Engine/Source/Programs/UnrealBuildAccelerator/Detours/Private/UbaDetoursShared.h
+++ b/Engine/Source/Programs/UnrealBuildAccelerator/Detours/Private/UbaDetoursShared.h
@@ -119,7 +119,7 @@ namespace uba
 	};
 
 	inline bool CanDetour(const tchar* file) { return !t_disallowDetour && g_rules->CanDetour(file); }
-	inline bool KeepInMemory(const StringView& fileName) { return g_allowKeepFilesInMemory && g_rules->KeepInMemory(fileName, g_systemTemp.data); }
+	inline bool KeepInMemory(const StringView& fileName) { return g_allowKeepFilesInMemory && g_rules->KeepInMemory(fileName, g_systemTemp.data, g_runningRemote); }
 
 	void Shared_WriteConsole(const char* chars, u32 charCount, bool isError);
 	void Shared_WriteConsole(const wchar_t* chars, u32 charCount, bool isError);
diff --git a/Engine/Source/Programs/UnrealBuildAccelerator/Detours/Private/Windows/UbaDetoursFunctionsKernelBase.inl b/Engine/Source/Programs/UnrealBuildAccelerator/Detours/Private/Windows/UbaDetoursFunctionsKernelBase.inl
--- a/Engine/Source/Programs/UnrealBuildAccelerator/Detours/Private/Windows/UbaDetoursFunctionsKernelBase.inl
+++ b/Engine/Source/Programs/UnrealBuildAccelerator/Detours/Private/Windows/UbaDetoursFunctionsKernelBase.inl
@@ -1448,7 +1448,7 @@ bool Shared_MoveFile(LPCWSTR lpExistingFileName, LPCWSTR lpNewFileName, DWORD dw
 			lock2.Leave();
 			FileInfo& destInfo = insres.first->second;
 			UBA_ASSERT(!insres.second); // This is here just to get a chance to investigate this scenario.. might work
-			UBA_ASSERTF(!destInfo.trueFileMapHandle && (!destInfo.memoryFile || g_rules->IsThrowAway(dest)), TC("Moving file %s to %s that is an output file that is not a memory file is not supported"), source.data, lpNewFileName);
+			UBA_ASSERTF(!destInfo.trueFileMapHandle && (!destInfo.memoryFile || g_rules->IsThrowAway(dest, g_runningRemote)), TC("Moving file %s to %s that is an output file that is not a memory file is not supported"), source.data, lpNewFileName);
 			destInfo.memoryFile = sourceInfo.memoryFile;
 			sourceInfo.memoryFile = nullptr;
 			DEBUG_LOG_DETOURED(L"MoveFileExW", L"(memfile->memfile) %ls to %ls -> Success", lpExistingFileName, lpNewFileName);
diff --git a/Engine/Source/Programs/UnrealBuildAccelerator/Detours/Private/Windows/UbaDetoursFunctionsNtDll.inl b/Engine/Source/Programs/UnrealBuildAccelerator/Detours/Private/Windows/UbaDetoursFunctionsNtDll.inl
--- a/Engine/Source/Programs/UnrealBuildAccelerator/Detours/Private/Windows/UbaDetoursFunctionsNtDll.inl
+++ b/Engine/Source/Programs/UnrealBuildAccelerator/Detours/Private/Windows/UbaDetoursFunctionsNtDll.inl
@@ -1100,7 +1100,7 @@ NTSTATUS NTAPI Shared_NtCreateFile(bool IsCreateFunc, PHANDLE hFileHandle, ACCES
 			}
 			else
 			{
-				if (g_rules->IsThrowAway(fileName))
+				if (g_rules->IsThrowAway(fileName, g_runningRemote))
 				{
 					//isDeleteOnClose = true;
 				}
@@ -1376,7 +1376,7 @@ NTSTATUS NTAPI Detoured_NtClose(HANDLE handle)
 		mappingWritten = fi.memoryFile->writtenSize;
 
 		u32 orginalNameLen = TStrlen(fi.originalName);
-		if (IsOutputFile(StringView(fi.originalName, orginalNameLen), IsWrite(fo->desiredAccess, 0), fo->deleteOnClose) && !g_rules->IsThrowAway(StringView(fi.originalName, orginalNameLen)))
+		if (IsOutputFile(StringView(fi.originalName, orginalNameLen), IsWrite(fo->desiredAccess, 0), fo->deleteOnClose) && !g_rules->IsThrowAway(StringView(fi.originalName, orginalNameLen), g_runningRemote))
 		{
 			// Need to report this file to host so it can be tracked in directory table
 			if (!fi.memoryFile->isReported)
