[UBA]
* Optimization when path table offsets from old and new path table matches (which I expect can happen now when they are sorted)

diff --git a/Engine/Source/Programs/UnrealBuildAccelerator/Common/Private/UbaCacheServer.cpp b/Engine/Source/Programs/UnrealBuildAccelerator/Common/Private/UbaCacheServer.cpp
--- a/Engine/Source/Programs/UnrealBuildAccelerator/Common/Private/UbaCacheServer.cpp
+++ b/Engine/Source/Programs/UnrealBuildAccelerator/Common/Private/UbaCacheServer.cpp
@@ -532,7 +532,7 @@ namespace uba
 			});
 
 
-		//m_forceAllSteps = true;
+		m_forceAllSteps = true;
 		bool forceAllSteps = m_forceAllSteps;
 		m_forceAllSteps = false;
 
@@ -934,16 +934,19 @@ namespace uba
 				usedPathOffsets.Traverse([&](u32 pathOffset)
 				{
 					bucket.m_pathTable.GetString(temp.Clear(), pathOffset);
-					u32 newOffset = newPathTable.AddNoLock(temp.data, temp.count);
+					u32 newPathOffset = newPathTable.AddNoLock(temp.data, temp.count);
 
 					#if 0
 					StringBuffer<> test;
-					newPathTable.GetString(test, newOffset);
+					newPathTable.GetString(test, newPathOffset);
 					UBA_ASSERT(test.Equals(temp.data));
 					#endif
 
-					auto res = oldToNewPathOffset.try_emplace(pathOffset, newOffset);
-					UBA_ASSERT(res.second);(void)res;
+					if (pathOffset != newPathOffset)
+					{
+						auto res = oldToNewPathOffset.try_emplace(pathOffset, newPathOffset);
+						UBA_ASSERT(res.second);(void)res;
+					}
 				});
 				bucket.m_pathTable.Swap(newPathTable);
 			}
@@ -963,13 +966,16 @@ namespace uba
 					reader2.SetPosition(casKeyOffset);
 					u32 oldPathOffset = u32(reader2.Read7BitEncoded());
 					CasKey casKey = reader2.ReadCasKey();
+					u32 newPathOffset = oldPathOffset;
 					auto findIt = oldToNewPathOffset.find(oldPathOffset);
-					UBA_ASSERTF(findIt != oldToNewPathOffset.end(), TC("Can't find entry with offset %u"), oldPathOffset);
-					u32 newCasKeyOffset = newCasKeyTable.Add(casKey, findIt->second);
-					if (casKeyOffset == newCasKeyOffset)
-						return;
-					auto res = oldToNewCasKeyOffset.try_emplace(casKeyOffset, newCasKeyOffset);
-					UBA_ASSERT(res.second);(void)res;
+					if (findIt != oldToNewPathOffset.end())
+						newPathOffset = findIt->second;
+					u32 newCasKeyOffset = newCasKeyTable.Add(casKey, newPathOffset);
+					if (casKeyOffset != newCasKeyOffset)
+					{
+						auto res = oldToNewCasKeyOffset.try_emplace(casKeyOffset, newCasKeyOffset);
+						UBA_ASSERT(res.second);(void)res;
+					}
 				});
 				bucket.m_casKeyTable.Swap(newCasKeyTable);
 			}
