[UBA]
* Added executing host in information popup when hovering over process

diff --git a/Engine/Source/Programs/UnrealBuildAccelerator/Visualizer/Private/UbaTraceReader.cpp b/Engine/Source/Programs/UnrealBuildAccelerator/Visualizer/Private/UbaTraceReader.cpp
--- a/Engine/Source/Programs/UnrealBuildAccelerator/Visualizer/Private/UbaTraceReader.cpp
+++ b/Engine/Source/Programs/UnrealBuildAccelerator/Visualizer/Private/UbaTraceReader.cpp
@@ -22,6 +22,14 @@ namespace uba
 		return processor.processes[loc.processIndex];
 	}
 
+	const TraceView::Session& TraceView::GetSession(const ProcessLocation& loc)
+	{
+		if (loc.sessionIndex < sessions.size())
+			return sessions[loc.sessionIndex];
+		static TraceView::Session emptySession;
+		return emptySession;
+	}
+
 	void TraceView::Clear()
 	{
 		sessions.clear();
diff --git a/Engine/Source/Programs/UnrealBuildAccelerator/Visualizer/Private/UbaTraceReader.h b/Engine/Source/Programs/UnrealBuildAccelerator/Visualizer/Private/UbaTraceReader.h
--- a/Engine/Source/Programs/UnrealBuildAccelerator/Visualizer/Private/UbaTraceReader.h
+++ b/Engine/Source/Programs/UnrealBuildAccelerator/Visualizer/Private/UbaTraceReader.h
@@ -127,6 +127,7 @@ namespace uba
 		};
 
 		const Process& GetProcess(const ProcessLocation& loc);
+		const Session& GetSession(const ProcessLocation& loc);
 		void Clear();
 
 		Vector<Session> sessions;
diff --git a/Engine/Source/Programs/UnrealBuildAccelerator/Visualizer/Private/UbaVisualizer.cpp b/Engine/Source/Programs/UnrealBuildAccelerator/Visualizer/Private/UbaVisualizer.cpp
--- a/Engine/Source/Programs/UnrealBuildAccelerator/Visualizer/Private/UbaVisualizer.cpp
+++ b/Engine/Source/Programs/UnrealBuildAccelerator/Visualizer/Private/UbaVisualizer.cpp
@@ -1752,9 +1752,10 @@ namespace uba
 
 			logger.AddSpace(2);
 			logger.Info(L"  %ls", process.description.c_str());
-			logger.Info(L"  ProcessId: %6u", process.id);
-			logger.Info(L"  Start:    %7ls", TimeToText(process.start, true).str);
-			logger.Info(L"  Duration: %7ls", TimeToText(duration, true).str);
+			logger.Info(L"  Host:        %ls", m_processSelectedLocation.sessionIndex == 0 ? TC("local") : m_traceView.GetSession(m_processSelectedLocation).name.c_str());
+			logger.Info(L"  ProcessId:  %6u", process.id);
+			logger.Info(L"  Start:     %7ls", TimeToText(process.start, true).str);
+			logger.Info(L"  Duration:  %7ls", TimeToText(duration, true).str);
 			if (!process.returnedReason.empty())
 				logger.Info(L"  Returned: %7s", process.returnedReason.data());
 			if (hasExited && process.exitCode != 0)
