[UBA]
* Fixed linux compile error

diff --git a/Engine/Source/Programs/UnrealBuildAccelerator/Core/Private/UbaApplicationRules.cpp b/Engine/Source/Programs/UnrealBuildAccelerator/Core/Private/UbaApplicationRules.cpp
--- a/Engine/Source/Programs/UnrealBuildAccelerator/Core/Private/UbaApplicationRules.cpp
+++ b/Engine/Source/Programs/UnrealBuildAccelerator/Core/Private/UbaApplicationRules.cpp
@@ -345,7 +345,7 @@ namespace uba
 		virtual bool IsThrowAway(const StringView& fileName, bool isRunningRemote) const override
 		{
 			return Super::IsThrowAway(fileName, isRunningRemote)
-				|| isRunningRemote && fileName.EndsWith(TC("-telemetry.json"));
+				|| (isRunningRemote && fileName.EndsWith(TC("-telemetry.json")));
 		}
 
 		//virtual bool NeedsSharedMemory(const tchar* file)
@@ -366,7 +366,7 @@ namespace uba
 		virtual bool IsThrowAway(const StringView& fileName, bool isRunningRemote) const override
 		{
 			return Super::IsThrowAway(fileName, isRunningRemote)
-				|| isRunningRemote && fileName.EndsWith(TC("-telemetry.json"));
+				|| (isRunningRemote && fileName.EndsWith(TC("-telemetry.json")));
 		}
 
 		//virtual bool NeedsSharedMemory(const tchar* file)
