[UBA]
* Added some more information to maintenance log in cache server

diff --git a/Engine/Source/Programs/UnrealBuildAccelerator/Common/Private/UbaCacheServer.cpp b/Engine/Source/Programs/UnrealBuildAccelerator/Common/Private/UbaCacheServer.cpp
--- a/Engine/Source/Programs/UnrealBuildAccelerator/Common/Private/UbaCacheServer.cpp
+++ b/Engine/Source/Programs/UnrealBuildAccelerator/Common/Private/UbaCacheServer.cpp
@@ -583,9 +583,12 @@ namespace uba
 
 		m_storage.WaitForActiveWork();
 
-		u32 removedNonExisting = 0;
 		{
-			// TODO: Make this cleaner... 
+			u64 collectCasStartTime = GetTime();
+
+			u32 removedNonExisting = 0;
+
+			// TODO: Make this cleaner... (inside UbaStorage instead)
 			SCOPED_WRITE_LOCK(m_storage.m_casLookupLock, lookupLock);
 			
 			existingCas.reserve(m_storage.m_casLookup.size());
@@ -604,12 +607,14 @@ namespace uba
 				existingCas.try_emplace(i->first, CasFileInfo{i->second.size, 0ull});
 				++i;
 			}
-		}
+			lookupLock.Leave();
 
-		if (removedNonExisting)
-			m_logger.Detail(TC("  Removed %s cas entries (marked as not existing)"), CountToText(removedNonExisting).str);
+			if (removedNonExisting)
+				m_logger.Detail(TC("  Removed %s cas entries (marked as not existing)"), CountToText(removedNonExisting).str);
+
+			m_logger.Detail(TC("  Found %s (%s) cas files and %s deleted by overflow (%s)"), CountToText(existingCas.size()).str, BytesToText(totalCasSize).str, CountToText(deletedCasFiles.size()).str, TimeToText(GetTime() - collectCasStartTime).str);
+		}
 
-		m_logger.Detail(TC("  Found %s cas files and %s deleted by overflow (%s)"), CountToText(existingCas.size()).str, CountToText(deletedCasFiles.size()).str, BytesToText(totalCasSize).str);
 		u64 totalCasCount = existingCas.size() + deletedCasCount;
 
 		if (shouldExit())
@@ -857,7 +862,7 @@ namespace uba
 		if (missingInputEntryCount)
 			m_logger.Detail(TC("  Found %llu cache entries with missing input cas"), missingInputEntryCount.load());
 
-		m_logger.Detail(TC("  Deleted %llu cas files and %llu cache entries (%s)"), deletedCasCount, deleteEntryCount.load(), TimeToText(GetTime() - deleteCacheEntriesStartTime).str);
+		m_logger.Detail(TC("  Deleted %llu cas files and %llu cache entries over %u buckets (%s)"), deletedCasCount, deleteEntryCount.load(), u32(m_buckets.size()), TimeToText(GetTime() - deleteCacheEntriesStartTime).str);
 
 		if (shouldExit())
 			return true;
