[UBA]
* Made StringBuffer::Parse(char* out, u64 capacity exist on non-windows just to simplify code (it is just a memcpy in those cases)

diff --git a/Engine/Source/Programs/UnrealBuildAccelerator/Core/Private/UbaStringBuffer.cpp b/Engine/Source/Programs/UnrealBuildAccelerator/Core/Private/UbaStringBuffer.cpp
--- a/Engine/Source/Programs/UnrealBuildAccelerator/Core/Private/UbaStringBuffer.cpp
+++ b/Engine/Source/Programs/UnrealBuildAccelerator/Core/Private/UbaStringBuffer.cpp
@@ -376,15 +376,22 @@ namespace uba
 		return out != 0 || Equals(TC("0"));
 	}
 
-	#if PLATFORM_WINDOWS
 	u32 StringBufferBase::Parse(char* out, u64 outCapacity)
 	{
+		#if PLATFORM_WINDOWS
 		size_t destLen;
 		if (wcstombs_s(&destLen, out, outCapacity, data, outCapacity-1) != 0)
 			return 0;
 		return (u32)destLen;
+		#else
+		if (outCapacity == 0)
+			return 0;
+		u32 toCopy = Min(u32(outCapacity - 1), count);
+		memcpy(out, data, toCopy);
+		out[toCopy] = 0;
+		return toCopy;
+		#endif
 	}
-	#endif
 
 	StringView ToView(const tchar* s)
 	{
diff --git a/Engine/Source/Programs/UnrealBuildAccelerator/Core/Public/UbaStringBuffer.h b/Engine/Source/Programs/UnrealBuildAccelerator/Core/Public/UbaStringBuffer.h
--- a/Engine/Source/Programs/UnrealBuildAccelerator/Core/Public/UbaStringBuffer.h
+++ b/Engine/Source/Programs/UnrealBuildAccelerator/Core/Public/UbaStringBuffer.h
@@ -66,10 +66,7 @@ namespace uba
 		bool Parse(u32& out);
 		bool Parse(u16& out);
 		bool Parse(float& out);
-
-		#if PLATFORM_WINDOWS
 		u32 Parse(char* out, u64 outCapacity);
-		#endif
 
 		u32 count;
 		u32 capacity;
