[UBA]
* Added support for parsing LLVM IR stream obj files and extracing import/export symbols plus clearing out export symbols
* Added support for writing import libs from list of exported symbols
* Some fixes around extract symbol feature when writing obj files

diff --git a/Engine/Source/Programs/UnrealBuildAccelerator/Cli/Private/UbaCli.cpp b/Engine/Source/Programs/UnrealBuildAccelerator/Cli/Private/UbaCli.cpp
--- a/Engine/Source/Programs/UnrealBuildAccelerator/Cli/Private/UbaCli.cpp
+++ b/Engine/Source/Programs/UnrealBuildAccelerator/Cli/Private/UbaCli.cpp
@@ -730,7 +730,8 @@ namespace uba
 		//info.shouldWriteToDisk = shouldWriteToDisk;
 		info.rootDir = g_rootDir.data;
 		//info.traceName.Append(TC("TESTTRACE"));
-		info.storeObjFilesCompressed = true;
+		//info.storeObjFilesCompressed = true;
+		//info.extractObjFilesSymbols = true;
 		#if UBA_DEBUG_LOG_ENABLED
 		info.remoteLogEnabled = true;
 		#endif
diff --git a/Engine/Source/Programs/UnrealBuildAccelerator/Common/Private/UbaObjectFile.cpp b/Engine/Source/Programs/UnrealBuildAccelerator/Common/Private/UbaObjectFile.cpp
--- a/Engine/Source/Programs/UnrealBuildAccelerator/Common/Private/UbaObjectFile.cpp
+++ b/Engine/Source/Programs/UnrealBuildAccelerator/Common/Private/UbaObjectFile.cpp
@@ -67,14 +67,12 @@ namespace uba
 
 	bool ObjectFile::StripExports(Logger& logger)
 	{
-		u32 keptExportCount = 0;
-		return StripExports(logger, m_data, {}, keptExportCount);
-		return true;
+		return StripExports(logger, m_data, {});
 	}
 
 	bool ObjectFile::WriteImportsAndExports(Logger& logger, MemoryBlock& memoryBlock)
 	{
-		auto write = [&](const void* data, u64 dataSize) { memcpy(memoryBlock.Allocate(dataSize, 1, TC("")), data, dataSize); };
+		auto write = [&](const void* data, u64 dataSize) { memcpy(memoryBlock.Allocate(dataSize, 1, TC("ObjectFile::WriteImportsAndExports")), data, dataSize); };
 
 		write(&SymbolFileVersion, 1);
 		write(&m_type, 1);
@@ -91,7 +89,7 @@ namespace uba
 		for (auto& kv : m_exports)
 		{
 			write(kv.first.c_str(), kv.first.size());
-			write(kv.second.c_str(), kv.second.size());
+			write(kv.second.extra.c_str(), kv.second.extra.size());
 			write("", 1);
 		}
 		write("", 1);
@@ -121,7 +119,7 @@ namespace uba
 		for (auto& kv : m_exports)
 		{
 			write(kv.first.c_str(), kv.first.size());
-			write(kv.second.c_str(), kv.second.size());
+			write(kv.second.extra.c_str(), kv.second.extra.size());
 			write("", 1);
 		}
 		write("", 1);
@@ -162,8 +160,23 @@ namespace uba
 	{
 		ObjectFileCoff objectFileCoff;
 		ObjectFileElf objectFileElf;
+		ObjectFileLLVMIR objectFileLLVMIR;
 		
-		ObjectFile& objectFile = type == ObjectFileType_Coff ? (ObjectFile&)objectFileCoff : (ObjectFile&)objectFileElf;
+		ObjectFile& objectFile = [&]() -> ObjectFile&
+			{
+				switch (type)
+				{
+				case ObjectFileType_Coff:
+					return objectFileCoff;
+				case ObjectFileType_Elf:
+					return objectFileCoff;
+				case ObjectFileType_LLVMIR:
+					return objectFileLLVMIR;
+				default:
+					UBA_ASSERT(false);
+					return objectFileCoff;
+				}
+			}();
 
 
 		MemoryBlock memoryBlock(16*1024*1024);
diff --git a/Engine/Source/Programs/UnrealBuildAccelerator/Common/Private/UbaObjectFileCoff.cpp b/Engine/Source/Programs/UnrealBuildAccelerator/Common/Private/UbaObjectFileCoff.cpp
--- a/Engine/Source/Programs/UnrealBuildAccelerator/Common/Private/UbaObjectFileCoff.cpp
+++ b/Engine/Source/Programs/UnrealBuildAccelerator/Common/Private/UbaObjectFileCoff.cpp
@@ -284,6 +284,8 @@ namespace uba
 		std::string tmp;
 		std::string extra;
 
+		u32 index = 0;
+
 		auto str = (char*)directiveData;
 		while (str)
 		{
@@ -291,12 +293,22 @@ namespace uba
 			if (!exportStr)
 				break;
 			exportStr += 8;
-			char* exportEnd = strchr(exportStr, ' ');
-			str = exportEnd;
-			if (!exportEnd)
-				exportEnd = exportStr + strlen(exportStr);
+			char* exportEnd;
+			if (*exportStr == '\"')
+			{
+				++exportStr;
+				exportEnd = strchr(exportStr, '\"');
+				str = exportEnd + 1;
+			}
 			else
-				++str;
+			{
+				exportEnd = strchr(exportStr, ' ');
+				str = exportEnd;
+				if (!exportEnd)
+					exportEnd = exportStr + strlen(exportStr);
+				else
+					++str;
+			}
 
 			tmp.assign(exportStr, exportEnd);
 			extra.clear();
@@ -306,7 +318,7 @@ namespace uba
 				tmp = tmp.substr(0, comma - tmp.data());
 			}
 					
-			m_exports.emplace(tmp, extra);
+			m_exports.emplace(tmp, ExportInfo{extra, index++});
 		}
 		return true;
 	}
@@ -340,7 +352,7 @@ namespace uba
 		}
 	}
 
-	bool ObjectFileCoff::StripExports(Logger& logger, u8* newData, const UnorderedSymbols& allNeededImports, u32& outKeptExportCount)
+	bool ObjectFileCoff::StripExports(Logger& logger, u8* newData, const UnorderedSymbols& allNeededImports)
 	{
 		if (!m_info.directiveSectionMemOffset)
 			return true;
@@ -363,7 +375,6 @@ namespace uba
 		auto readEnd = readPos + directiveSection->SizeOfRawData;
 		auto readLastPossiblePos = readEnd - 9;
 
-		u32 exportCount = 0;
 		while (true)
 		{
 			const char* exportStr = nullptr;
@@ -383,22 +394,32 @@ namespace uba
 
 			if (!exportStr)
 			{
-				readPos = readEnd - 1;
+				readPos = readEnd;
 				break;
 			}
 
-			++exportCount;
-
 			const char* startPos = exportStr;
 			exportStr += 8;
-			const char* exportEnd = strchr(exportStr, ' ');
-			if (!exportEnd)
-				readPos = exportEnd = exportStr + strlen(exportStr);
+			const char* exportEnd;
+			if (*exportStr == '\"')
+			{
+				++exportStr;
+				exportEnd = strchr(exportStr, '\"');
+				readPos = exportEnd + 1;
+				if (strncmp(readPos, ",DATA", 5) == 0)
+					readPos += 5;
+				readPos = Min(readPos, readEnd);
+			}
 			else
-				readPos = exportEnd;
-
-			if (strncmp(exportEnd-5, ",DATA", 5) == 0)
-				exportEnd -= 5;
+			{
+				exportEnd = strchr(exportStr, ' ');
+				if (!exportEnd)
+					readPos = exportEnd = exportStr + strlen(exportStr);
+				else
+					readPos = exportEnd;
+				if (strncmp(exportEnd-5, ",DATA", 5) == 0)
+					exportEnd -= 5;
+			}
 
 			tmp.assign(exportStr, exportEnd - exportStr);
 			if (allNeededImports.find(tmp) != allNeededImports.end())
@@ -407,8 +428,6 @@ namespace uba
 			if (allNeededImports.find(tmp) != allNeededImports.end())
 				continue;
 			
-			--exportCount;
-
 			u64 toCopy = startPos - lastCopyPos - 1;
 			memcpy(writePos, lastCopyPos, toCopy);
 			writePos += toCopy;
@@ -420,7 +439,6 @@ namespace uba
 		u64 toCopy = readPos - lastCopyPos;
 		memcpy(writePos, lastCopyPos, toCopy);
 		writePos += toCopy;
-		*writePos = 0;
 
 		u32 sizeOfRawData = directiveSection->SizeOfRawData;
 		newDirectiveSection->SizeOfRawData = u32((u8*)writePos - newDirectiveData);
@@ -428,12 +446,15 @@ namespace uba
 
 		memset(writePos, 0, sizeOfRawData - newDirectiveSection->SizeOfRawData);
 
-		outKeptExportCount = exportCount;
-
 		return true;
 	}
 
 	bool ObjectFileCoff::CreateExtraFile(Logger& logger, MemoryBlock& memoryBlock, const UnorderedSymbols& allNeededImports, const UnorderedSymbols& allSharedImports, const UnorderedExports& allSharedExports, bool includeExportsInFile)
+	{
+		return CreateExtraFile2(logger, memoryBlock, allNeededImports, allSharedImports, allSharedExports, includeExportsInFile);
+	}
+
+	bool ObjectFileCoff::CreateExtraFile2(Logger& logger, MemoryBlock& memoryBlock, const UnorderedSymbols& allNeededImports, const UnorderedSymbols& allSharedImports, const UnorderedExports& allSharedExports, bool includeExportsInFile)
 	{
 		std::string tmp;
 
@@ -459,6 +480,7 @@ namespace uba
 
 		// Header
 		auto& header = *(ImageFileHeader*)allocate(sizeof(ImageFileHeader));
+		header.Machine = 0x8664;
 
 		// Session for loopbacks
 		auto& textSection = *(ImageSectionHeader*)allocate(sizeof(ImageSectionHeader));
@@ -492,7 +514,7 @@ namespace uba
 
 				write(slashExport, sizeof(slashExport) - 1);
 				write(symbol.data(), symbol.size());
-				write(kv.second.data(), kv.second.size());
+				write(kv.second.extra.data(), kv.second.extra.size());
 				write(" ", 1);
 			}
 			write("", 1);
diff --git a/Engine/Source/Programs/UnrealBuildAccelerator/Common/Private/UbaObjectFileCoff.h b/Engine/Source/Programs/UnrealBuildAccelerator/Common/Private/UbaObjectFileCoff.h
--- a/Engine/Source/Programs/UnrealBuildAccelerator/Common/Private/UbaObjectFileCoff.h
+++ b/Engine/Source/Programs/UnrealBuildAccelerator/Common/Private/UbaObjectFileCoff.h
@@ -13,13 +13,15 @@ namespace uba
 		ObjectFileCoff();
 		virtual bool Parse(Logger& logger, const tchar* filename) override;
 
+		static bool CreateExtraFile2(Logger& logger, MemoryBlock& memoryBlock, const UnorderedSymbols& allNeededImports, const UnorderedSymbols& allSharedImports, const UnorderedExports& allSharedExports, bool includeExportsInFile);
+
 	private:
 		struct Info;
 
 		bool ParseExports();
 		template<typename SymbolType> void ParseImports();
 
-		virtual bool StripExports(Logger& logger, u8* newData, const UnorderedSymbols& allNeededImports, u32& outKeptExportCount) override;
+		virtual bool StripExports(Logger& logger, u8* newData, const UnorderedSymbols& allNeededImports) override;
 		virtual bool CreateExtraFile(Logger& logger, MemoryBlock& memoryBlock, const UnorderedSymbols& allNeededImports, const UnorderedSymbols& allSharedImports, const UnorderedExports& allSharedExports, bool includeExportsInFile) override;
 
 		template<typename SymbolType> void CalculateImports(Logger& logger, Vector<u32>& outImports);
diff --git a/Engine/Source/Programs/UnrealBuildAccelerator/Common/Private/UbaObjectFileElf.cpp b/Engine/Source/Programs/UnrealBuildAccelerator/Common/Private/UbaObjectFileElf.cpp
--- a/Engine/Source/Programs/UnrealBuildAccelerator/Common/Private/UbaObjectFileElf.cpp
+++ b/Engine/Source/Programs/UnrealBuildAccelerator/Common/Private/UbaObjectFileElf.cpp
@@ -230,14 +230,15 @@ namespace uba
 		return newFile.Close();
 	}
 	*/
-	bool ObjectFileElf::StripExports(Logger& logger, u8* newData, const UnorderedSymbols& allNeededImports, u32& outKeptExportCount)
+	bool ObjectFileElf::StripExports(Logger& logger, u8* newData, const UnorderedSymbols& allNeededImports)
 	{
 		return true;
 	}
 
 	bool ObjectFileElf::CreateExtraFile(Logger& logger, MemoryBlock& memoryBlock, const UnorderedSymbols& allNeededImports, const UnorderedSymbols& allSharedImports, const UnorderedExports& allSharedExports, bool includeExportsInFile)
 	{
-		auto& header = *(Elf64Header*)memoryBlock.Allocate(sizeof(Elf64Header), 1, TC(""));
+		u8* data = (u8*)memoryBlock.Allocate(sizeof(Elf64Header) + sizeof(Elf64SectionHeader) + 1, 1, TC(""));
+		auto& header = *(Elf64Header*)data;
 
 		header.e_ident[0] = 0x7f;
 		header.e_ident[1] = 'E';
@@ -248,7 +249,21 @@ namespace uba
 		header.e_ident[6] = 1;
 		header.e_type = 1;
 		header.e_machine = EM_X86_64;
+		header.e_version = 1;
 		header.e_ehsize = sizeof(Elf64Header);
+		header.e_flags = 0;//0x04000000;
+		header.e_shoff = sizeof(Elf64Header);
+		header.e_shentsize = sizeof(Elf64SectionHeader);
+		header.e_shnum = 1;
+		header.e_shstrndx = 0;
+
+		auto& section = *(Elf64SectionHeader*)(data + sizeof(Elf64Header));
+		section.sh_name = 0;
+		section.sh_type = 3;
+		section.sh_flags = 0;//0xC35d00;
+		section.sh_addralign = 1;
+		section.sh_offset = sizeof(Elf64Header) + sizeof(Elf64SectionHeader);
+		section.sh_size = 1;
 		return true;
 	}
 }
diff --git a/Engine/Source/Programs/UnrealBuildAccelerator/Common/Private/UbaObjectFileElf.h b/Engine/Source/Programs/UnrealBuildAccelerator/Common/Private/UbaObjectFileElf.h
--- a/Engine/Source/Programs/UnrealBuildAccelerator/Common/Private/UbaObjectFileElf.h
+++ b/Engine/Source/Programs/UnrealBuildAccelerator/Common/Private/UbaObjectFileElf.h
@@ -13,7 +13,7 @@ namespace uba
 		virtual bool Parse(Logger& logger, const tchar* filename) override;
 
 	private:
-		virtual bool StripExports(Logger& logger, u8* newData, const UnorderedSymbols& allNeededImports, u32& outKeptExportCount) override;
+		virtual bool StripExports(Logger& logger, u8* newData, const UnorderedSymbols& allNeededImports) override;
 		virtual bool CreateExtraFile(Logger& logger, MemoryBlock& memoryBlock, const UnorderedSymbols& allNeededImports, const UnorderedSymbols& allSharedImports, const UnorderedExports& allSharedExports, bool includeExportsInFile) override;
 
 		UnorderedSymbols m_toRemove;
diff --git a/Engine/Source/Programs/UnrealBuildAccelerator/Common/Private/UbaObjectFileLLVMIR.cpp b/Engine/Source/Programs/UnrealBuildAccelerator/Common/Private/UbaObjectFileLLVMIR.cpp
--- a/Engine/Source/Programs/UnrealBuildAccelerator/Common/Private/UbaObjectFileLLVMIR.cpp
+++ b/Engine/Source/Programs/UnrealBuildAccelerator/Common/Private/UbaObjectFileLLVMIR.cpp
@@ -1,23 +1,1064 @@
 // Copyright Epic Games, Inc. All Rights Reserved.
 
-#pragma once
-
 #include "UbaObjectFileLLVMIR.h"
+#include "UbaBinaryReaderWriter.h"
+#include "UbaObjectFileCoff.h"
+
+#define UBA_LOG_STREAM 0
 
 namespace uba
 {
-	// TODO: https://llvm.org/docs/BitCodeFormat.html
+	// https://llvm.org/docs/BitCodeFormat.html
 
+	bool IsRawBitcode(const u8* data, u64 dataSize)
+	{
+		constexpr u8 magic[] = { 'B', 'C', 0xc0, 0xde };
+		return dataSize >= 4 && memcmp(data, magic, sizeof(magic)) == 0;
+	}
 
-	bool ObjectFileLLVMIR::Parse(Logger& logger, const tchar* filename)
+	bool IsWrappedBitcode(const u8* data, u64 dataSize)
 	{
-		return logger.Error(TC("LLVM IR obj file format not supported (yet)"));
+		constexpr u8 wrapperMagic[] = { 0xDE, 0xC0, 0x17, 0x0B };
+		return dataSize >= 4 && memcmp(data, wrapperMagic, sizeof(wrapperMagic)) == 0;
 	}
 
 	bool IsLLVMIRFile(const u8* data, u64 dataSize)
 	{
-		constexpr u8 wrapperMagic[] = { 'B', 'C', 0xc0, 0xde };
-		constexpr u8 magic[] = { 'B', 'C', 0x04, 0xc4, 0xe4, 0xd4 };
-		return dataSize >= 6 && (memcmp(data, wrapperMagic, sizeof(wrapperMagic)) == 0 || memcmp(data, magic, sizeof(magic)) == 0);
+		return IsRawBitcode(data, dataSize) || IsWrappedBitcode(data, dataSize);
+	}
+
+
+	#define BLOCK_IDS \
+		BLOCK_ID(MODULE_BLOCK_ID) \
+		BLOCK_ID(PARAMATTR_BLOCK_ID) \
+		BLOCK_ID(PARAMATTR_GROUP_BLOCK_ID) \
+		BLOCK_ID(CONSTANTS_BLOCK_ID) \
+		BLOCK_ID(FUNCTION_BLOCK_ID) \
+		BLOCK_ID(IDENTIFICATION_BLOCK_ID) \
+		BLOCK_ID(VALUE_SYMTAB_BLOCK_ID) \
+		BLOCK_ID(METADATA_BLOCK_ID) \
+		BLOCK_ID(METADATA_ATTACHMENT_ID) \
+		BLOCK_ID(TYPE_BLOCK_ID_NEW) \
+		BLOCK_ID(USELIST_BLOCK_ID) \
+		BLOCK_ID(MODULE_STRTAB_BLOCK_ID) \
+		BLOCK_ID(GLOBALVAL_SUMMARY_BLOCK_ID) \
+		BLOCK_ID(OPERAND_BUNDLE_TAGS_BLOCK_ID) \
+		BLOCK_ID(METADATA_KIND_BLOCK_ID) \
+		BLOCK_ID(STRTAB_BLOCK_ID) \
+		BLOCK_ID(FULL_LTO_GLOBALVAL_SUMMARY_BLOCK_ID) \
+		BLOCK_ID(SYMTAB_BLOCK_ID) \
+		BLOCK_ID(SYNC_SCOPE_NAMES_BLOCK_ID) \
+
+	enum BlockIDs : u8
+	{
+		BLOCKINFO_BLOCK_ID = 0,
+		BEFORE_FIRST_APPLICATION_BLOCKID = 7,
+		#define BLOCK_ID(x) x,
+		BLOCK_IDS
+		#undef BLOCK_ID
+	};
+
+	const tchar* BlockIdToString(u32 id)
+	{
+		switch(id)
+		{
+		case 0: return TC("BLOCKINFO_BLOCK_ID");
+		#define BLOCK_ID(x) case x: return TC(#x);
+		BLOCK_IDS
+		#undef BLOCK_ID
+		}
+		return TC("UNKNOWN_BLOCK_ID");
+	}
+
+	#define MODULE_CODES \
+		MODULE_CODE(VERSION, 1) \
+		MODULE_CODE(TRIPLE, 2) \
+		MODULE_CODE(DATALAYOUT, 3) \
+		MODULE_CODE(ASM, 4) \
+		MODULE_CODE(SECTIONNAME, 5) \
+		MODULE_CODE(DEPLIB, 6) \
+		MODULE_CODE(GLOBALVAR, 7) \
+		MODULE_CODE(FUNCTION, 8) \
+		MODULE_CODE(ALIAS_OLD, 9) \
+		MODULE_CODE(GCNAME, 11) \
+		MODULE_CODE(COMDAT, 12) \
+		MODULE_CODE(VSTOFFSET, 13) \
+		MODULE_CODE(ALIAS, 14) \
+		MODULE_CODE(METADATA_VALUES_UNUSED, 15) \
+		MODULE_CODE(SOURCE_FILENAME, 16) \
+		MODULE_CODE(HASH, 17) \
+		MODULE_CODE(IFUNC, 18)
+
+	enum ModuleCodes {
+		#define MODULE_CODE(x, v) MODULE_CODE_##x = v,
+		MODULE_CODES
+		#undef MODULE_CODE
+	};
+
+	const tchar* ModuleCodeToString(u32 code)
+	{
+		switch(code)
+		{
+		#define MODULE_CODE(x, v) case MODULE_CODE_##x: return TC("MODULE_CODE_" #x);
+		MODULE_CODES
+		#undef MODULE_CODE
+		}
+		return TC("MODULE_CODE_UNKNOWN");
+	}
+
+	enum ObjectFileLLVMIR::BlockInfoCodes : u8
+	{
+	  BLOCKINFO_CODE_SETBID = 1,
+	  BLOCKINFO_CODE_BLOCKNAME = 2,
+	  BLOCKINFO_CODE_SETRECORDNAME = 3
+	};
+
+	enum ObjectFileLLVMIR::FixedAbbrevIDs : u8
+	{
+		END_BLOCK = 0,
+		ENTER_SUBBLOCK = 1,
+		DEFINE_ABBREV = 2,
+		UNABBREV_RECORD = 3,
+		FIRST_APPLICATION_ABBREV = 4
+	};
+
+	enum ObjectFileLLVMIR::EntryKind : u8
+	{
+		EntryKind_Error,
+		EntryKind_EndBlock,
+		EntryKind_SubBlock,
+		EntryKind_Record
+	};
+
+	enum ObjectFileLLVMIR::Encoding : u8
+	{
+		Encoding_Fixed = 1,
+		Encoding_VBR   = 2,
+		Encoding_Array = 3,
+		Encoding_Char6 = 4,
+		Encoding_Blob  = 5 
+	};
+
+	enum
+	{
+		AF_DontPopBlockAtEnd = 1,
+		AF_DontAutoprocessAbbrevs = 2
+	};
+
+	struct ObjectFileLLVMIR::AbbrevOp
+	{
+		explicit AbbrevOp(u64 v) : val(v), isLiteral(true) {}
+		explicit AbbrevOp(Encoding e, u64 data = 0) : val(data), isLiteral(false), encoding(e) {}
+		u64 val;
+		bool isLiteral : 1;
+		u32 encoding   : 3;
+	};
+
+	struct ObjectFileLLVMIR::Abbrev
+	{
+		Vector<AbbrevOp> operands;
+	};
+
+	struct ObjectFileLLVMIR::Entry
+	{
+		EntryKind kind;
+		u32 id;
+	};
+
+	struct ObjectFileLLVMIR::BlockInfo
+	{
+		struct Record
+		{
+			u32 blockId = 0;
+			Vector<AbbrevPtr> abbrevs;
+			std::string name;
+			Vector<std::pair<u32, std::string>> recordNames;
+		};
+
+		Vector<Record> records;
+
+		const Record* GetBlockInfo(u32 blockId) const
+		{
+			if (!records.empty() && records.back().blockId == blockId)
+				return &records.back();
+			for (const Record& bi : records)
+				if (bi.blockId == blockId)
+					return &bi;
+			return nullptr;
+		}
+
+		Record& GetOrCreateBlockInfo(u32 blockId)
+		{
+			if (const Record* bi = GetBlockInfo(blockId))
+				return *const_cast<Record*>(bi);
+			records.emplace_back();
+			records.back().blockId = blockId;
+			return records.back();
+		}
+	};
+
+	class ObjectFileLLVMIR::BitStreamReader
+	{
+	public:
+		BitStreamReader(ObjectFileLLVMIR& owner, Logger& logger, u8* data, u64 dataSize)
+		:	m_owner(owner)
+		,	m_logger(logger)
+		,	m_begin(data)
+		,	m_pos(data)
+		,	m_end(data + dataSize)
+		{
+		}
+
+		u64 GetCurrentBitNo() const
+		{
+			return u64(m_pos - m_begin)*8 - m_wordBits;
+		}
+
+		bool CanSkipToPos(u64 pos) const
+		{
+			return pos <= u64(m_end - m_begin);
+		}
+
+		void SkipToEnd()
+		{
+			m_pos = m_end;
+		}
+
+		void JumpToBit(u64 bitNo)
+		{
+			u64 byteNo = u64(bitNo/8) & ~(sizeof(u32)-1);
+			u32 wordBitNo = u32(bitNo & (sizeof(u32)*8-1));
+			m_pos = m_begin + byteNo;
+			m_wordBits = 0;
+			if (wordBitNo)
+				Read(wordBitNo);
+		}
+
+		u32 Read(u32 bits)
+		{
+			if (m_wordBits >= bits)
+			{
+				u32 res = m_word & (~0u >> (32 - bits));
+				m_word >>= (bits & 0x1f);
+				m_wordBits -= bits;
+				return res;
+			}
+			u32 res = m_wordBits ? m_word : 0;
+			u32 bitsLeft = bits - m_wordBits;
+
+			UBA_ASSERT((m_end - m_pos) >= 4);
+			m_word = *(u32*)m_pos;
+			m_wordBits = 32;
+			m_pos += 4;
+
+			u32 res2 = m_word & (~0u >> (32 - bitsLeft));
+			m_word >>= (bitsLeft & 0x1f);
+			m_wordBits -= bitsLeft;
+
+			res |= res2 << (bits - bitsLeft);
+
+			return res;
+		}
+
+		u32 ReadVBR(u32 bits)
+		{
+			u32 piece = Read(bits);
+			u32 maskBitOrder = (bits - 1);
+			u32 mask = 1u << maskBitOrder;
+
+			if ((piece & mask) == 0)
+			  return piece;
+
+			u32 result = 0;
+			u32 nextBit = 0;
+
+			while (true)
+			{
+				result |= (piece & (mask - 1)) << nextBit;
+
+				if ((piece & mask) == 0)
+					return result;
+
+				nextBit += bits-1;
+				UBA_ASSERT(nextBit < 32);
+				piece = Read(bits);
+			}
+		}
+
+		u64 ReadVBR64(u32 bits)
+		{
+			u32 piece = Read(bits);
+			u32 maskBitOrder = (bits - 1);
+			u32 mask = 1u << maskBitOrder;
+
+			if ((piece & mask) == 0)
+			  return piece;
+
+			u64 result = 0;
+			u32 nextBit = 0;
+
+			while (true)
+			{
+				result |= u64(piece & (mask - 1)) << nextBit;
+
+				if ((piece & mask) == 0)
+					return result;
+
+				nextBit += bits-1;
+				UBA_ASSERT(nextBit < 64);
+				piece = Read(bits);
+			}
+		}
+
+		void ClearVBR64(u32 bits)
+		{
+			if (m_wordBits)
+			{
+				u8* oldPos = m_pos - 4;
+				u32 mask = ((1u << (32u - m_wordBits)) - 1u);
+				if (bits < m_wordBits)
+				{
+					mask |= ~((1u << (32u - (m_wordBits - bits))) - 1u);
+					bits = 0;
+				}
+				else
+					bits -= m_wordBits;
+				*(u32*)oldPos &= mask;
+			}
+			if (bits)
+			{
+				u32 mask = ~((1u << (bits)) - 1u);
+				*(u32*)m_pos &= mask;
+			}
+		}
+
+		u32 ReadCode()
+		{
+			return Read(m_currentCodeSize);
+		}
+
+		void SkipToFourByteBoundary()
+		{
+			m_wordBits = 0;
+		}
+
+		bool IsDone()
+		{
+			return m_pos == m_end;
+		}
+
+		void SkipBlock()
+		{
+			ReadVBR(4); // CodeLenWidth
+			SkipToFourByteBoundary();
+			size_t numFourBytes = Read(32); // BlockSizeWidth
+			size_t skipTo = GetCurrentBitNo() + numFourBytes * 4 * 8;
+			JumpToBit(skipTo);
+		}
+
+		void EnterSubBlock(u32& outNumWords, u32 blockId)
+		{
+			m_blockScope.push_back(Block(m_currentCodeSize));
+			m_blockScope.back().prevAbbrevs.swap(m_curAbbrevs);
+
+			if (m_blockInfo)
+				if (const BlockInfo::Record* record = m_blockInfo->GetBlockInfo(blockId))
+					m_curAbbrevs.insert(m_curAbbrevs.end(), record->abbrevs.begin(), record->abbrevs.end());
+
+			m_currentCodeSize = ReadVBR(4);
+			UBA_ASSERT(m_currentCodeSize);
+
+			SkipToFourByteBoundary();
+			u32 numWords = Read(32);
+			UBA_ASSERT(numWords);
+			outNumWords = numWords;
+		}
+
+		BlockInfo ReadBlockInfoBlock(bool readBlockInfoNames)
+		{
+			u32 numWords;
+			EnterSubBlock(numWords, BLOCKINFO_BLOCK_ID);
+
+			BlockInfo newBlockInfo;
+
+			Vector<u64> record;
+			BlockInfo::Record* curBlockInfo = nullptr;
+
+			while (true)
+			{
+				Entry entry = AdvanceSkippingSubblocks(AF_DontAutoprocessAbbrevs);
+
+				switch (entry.kind)
+				{
+				case EntryKind_SubBlock:
+				case EntryKind_Error:
+					UBA_ASSERT(false);
+					return {};
+				case EntryKind_EndBlock:
+					return newBlockInfo;
+				case EntryKind_Record:
+					break;
+				}
+
+				if (entry.id == DEFINE_ABBREV)
+				{
+					if (!curBlockInfo)
+					{
+						UBA_ASSERT(false);
+						return {};
+					}
+					ReadAbbrevRecord();
+					curBlockInfo->abbrevs.push_back(std::move(m_curAbbrevs.back()));
+					m_curAbbrevs.pop_back();
+					continue;
+				}
+
+				record.clear();
+				u32 code = ReadRecord(record, entry.id);
+				switch (code)
+				{
+				default:
+					break;
+				case BLOCKINFO_CODE_SETBID:
+					if (record.size() < 1)
+					{
+						UBA_ASSERT(false);
+						return {};
+					}
+					curBlockInfo = &newBlockInfo.GetOrCreateBlockInfo(u32(record[0]));
+					break;
+				case BLOCKINFO_CODE_BLOCKNAME: {
+					if (!curBlockInfo)
+					{
+						UBA_ASSERT(false);
+						return {};
+					}
+					UBA_ASSERT(!readBlockInfoNames);
+					break;
+				}
+				case BLOCKINFO_CODE_SETRECORDNAME: {
+					if (!curBlockInfo)
+					{
+						UBA_ASSERT(false);
+						return {};
+					}
+					UBA_ASSERT(!readBlockInfoNames);
+					break;
+				}
+				}
+			}
+			return newBlockInfo;
+		}
+
+		bool ReadAbbrevRecord()
+		{
+			auto abbv = std::make_shared<Abbrev>();
+
+			u32 numOpInfo = ReadVBR(5);
+			for (u32 i = 0; i != numOpInfo; ++i)
+			{
+				bool isLiteral = Read(1);
+				if (isLiteral)
+				{
+					u64 op = ReadVBR64(8);
+					abbv->operands.emplace_back(op);
+					continue;
+				}
+
+				Encoding encoding = Encoding(Read(3));
+				UBA_ASSERT(encoding >= 1 && encoding <= 5);
+				if (encoding == Encoding_Fixed || encoding == Encoding_VBR)
+				{
+					u64 data = ReadVBR64(5);
+					if (data == 0)
+						abbv->operands.emplace_back(0);
+					else
+						abbv->operands.emplace_back(encoding, data);
+				}
+				else
+					abbv->operands.emplace_back(encoding);
+			}
+
+			m_curAbbrevs.push_back(abbv);
+			return true;
+		}
+
+		u64 ReadAbbreviatedField(const AbbrevOp& op)
+		{
+			switch (op.encoding)
+			{
+			case Encoding_Fixed:
+				return Read(u32(op.val));
+			case Encoding_VBR:
+				return ReadVBR64(u32(op.val));
+			case Encoding_Char6:
+				return "abcdefghijklmnopqrstuvwxyzABCDEFGHIJKLMNOPQRSTUVWXYZ0123456789._"[Read(6)];
+			default:
+				UBA_ASSERT(false);
+				return ~0ull;
+			}
+		}
+
+		void SkipRecord(u32 abbrevId)
+		{
+			if (abbrevId == UNABBREV_RECORD)
+			{
+				ReadVBR(6);
+				u32 numElts = ReadVBR(6);
+				for (u32 i = 0; i != numElts; ++i)
+					ReadVBR64(6);
+				return;
+			}
+
+			u32 abbrevIndex = abbrevId - FIRST_APPLICATION_ABBREV;
+			Abbrev& abbv = *m_curAbbrevs[abbrevIndex];
+			UBA_ASSERT(!abbv.operands.empty());
+			AbbrevOp& codeOp = abbv.operands[0];
+			if (!codeOp.isLiteral)
+				ReadAbbreviatedField(codeOp);
+
+			for (u64 i = 1, e = abbv.operands.size(); i != e; ++i) {
+				const AbbrevOp& op = abbv.operands[i];
+				if (op.isLiteral)
+					continue;
+
+				if (op.encoding != Encoding_Array && op.encoding != Encoding_Blob)
+				{
+					ReadAbbreviatedField(op);
+					continue;
+				}
+
+				if (op.encoding == Encoding_Array)
+				{
+					u32 numElts = ReadVBR(6);
+
+					const AbbrevOp& eltEnc = abbv.operands[++i];
+
+					switch (eltEnc.encoding)
+					{
+					default:
+						UBA_ASSERT(false);
+						return;
+					case Encoding_Fixed:
+						JumpToBit(GetCurrentBitNo() + u64(numElts) * eltEnc.val);
+						break;
+					case Encoding_VBR:
+						for (; numElts; --numElts)
+							ReadVBR64(u32(eltEnc.val));
+						break;
+					case Encoding_Char6:
+						JumpToBit(GetCurrentBitNo() + numElts * 6);
+					}
+					continue;
+				}
+
+				UBA_ASSERT(op.encoding == Encoding_Blob);
+				u32 numElts = ReadVBR(6);
+				SkipToFourByteBoundary();
+
+				u64 newEnd = GetCurrentBitNo() + AlignUp(numElts, 4) * 8;
+				if (!CanSkipToPos(newEnd/8))
+				{
+					SkipToEnd();
+					break;
+				}
+				JumpToBit(newEnd);
+			}
+			return;
+		}
+
+		void ReadRecordOperands(Vector<u64>& outVals, const Vector<AbbrevOp>& operands, std::string* blob = nullptr, u32 blockId = 0)
+		{
+			for (u64 i = 1, e = operands.size(); i != e; ++i) {
+				const AbbrevOp& op = operands[i];
+				if (op.isLiteral)
+				{
+					outVals.push_back(op.val);
+					continue;
+				}
+
+				if (op.encoding != Encoding_Array && op.encoding != Encoding_Blob)
+				{
+					u64 val = ReadAbbreviatedField(op);(void)val;
+					outVals.push_back(val);
+					continue;
+				}
+
+				if (op.encoding == Encoding_Array)
+				{
+					u32 numElts = ReadVBR(6);
+					outVals.reserve(outVals.size() + numElts);
+
+					const AbbrevOp& eltEnc = operands[++i];
+					UBA_ASSERT(!eltEnc.isLiteral);
+
+					switch (eltEnc.encoding)
+					{
+					default:
+						UBA_ASSERT(false);
+						return;
+					case Encoding_Fixed:
+						for (; numElts; --numElts)
+							outVals.push_back(Read(u32(eltEnc.val)));
+						break;
+					case Encoding_VBR:
+						for (; numElts; --numElts)
+							outVals.push_back(ReadVBR64(u32(eltEnc.val)));
+						break;
+					case Encoding_Char6:
+						for (; numElts; --numElts)
+							outVals.push_back("abcdefghijklmnopqrstuvwxyzABCDEFGHIJKLMNOPQRSTUVWXYZ0123456789._"[Read(6)]);
+					}
+					continue;
+				}
+
+				UBA_ASSERT(op.encoding == Encoding_Blob);
+				u32 numElts = ReadVBR(6);
+				SkipToFourByteBoundary();  // 32-bit alignment
+
+				size_t curBitPos = GetCurrentBitNo();
+				size_t newEnd = curBitPos + AlignUp(numElts, 4) * 8;
+
+				UBA_ASSERT(CanSkipToPos(newEnd/8));
+				JumpToBit(newEnd);
+
+				if (blockId == STRTAB_BLOCK_ID)
+				{
+					m_owner.m_strTabPos = u64(curBitPos/8);
+					m_owner.m_strTabSize = numElts;
+				}
+
+				char* ptr = (char*)(m_begin + (curBitPos/8));
+				if (blob)
+					*blob = std::string(ptr, numElts);
+				else
+				{
+					auto* uptr = reinterpret_cast<const unsigned char*>(ptr);
+					outVals.insert(outVals.end(), uptr, uptr + numElts);
+				}
+			}
+		}
+
+		bool CanBeExported(u32 blockId, u32 code)
+		{
+			if (blockId != MODULE_BLOCK_ID)
+				return false;
+			switch (code)
+			{
+			case MODULE_CODE_FUNCTION:
+			case MODULE_CODE_GLOBALVAR:
+			case MODULE_CODE_IFUNC:
+			case MODULE_CODE_ALIAS:
+			case MODULE_CODE_ALIAS_OLD:
+				return true;
+			default:
+				return false;
+			}
+		}
+
+		u32 ReadRecord(Vector<u64>& outVals, u32 abbrevId, std::string* blob = nullptr, u32 blockId = 0)
+		{
+			if (abbrevId == UNABBREV_RECORD)
+			{
+				u32 code = ReadVBR(6);
+
+				if (CanBeExported(blockId, code))
+					m_owner.m_globalVarOrFunctionRecords.push_back(BitStreamEntry{u64(m_pos - m_begin), m_word, m_wordBits, code, false});
+
+				u32 numElts = ReadVBR(6);
+
+				outVals.reserve(outVals.size() + numElts);
+				for (u32 i = 0; i != numElts; ++i)
+					outVals.push_back(ReadVBR64(6));
+
+#if UBA_LOG_STREAM
+				if (blockId == MODULE_BLOCK_ID)
+				{
+					auto str = ModuleCodeToString(code);
+					u32 indent = 2;
+					m_logger.Info(TC("%s%s (%u)"), TC("        ") + 8 - indent, str, code);
+				}
+#endif
+				return code;
+			}
+
+			u32 abbrevIndex = abbrevId - FIRST_APPLICATION_ABBREV;
+			UBA_ASSERT(abbrevIndex < m_curAbbrevs.size());
+			Abbrev& abbv = *m_curAbbrevs[abbrevIndex];
+			UBA_ASSERT(!abbv.operands.empty());
+			AbbrevOp& codeOp = abbv.operands[0];
+			u32 code;
+			if (codeOp.isLiteral)
+				code = u32(codeOp.val);
+			else
+			{
+				UBA_ASSERT(codeOp.encoding != Encoding_Array && codeOp.encoding != Encoding_Blob);
+				code = u32(ReadAbbreviatedField(codeOp));
+			}
+
+			if (CanBeExported(blockId, code))
+			{
+				UBA_ASSERT(!abbv.operands.empty());
+				m_owner.m_globalVarOrFunctionRecords.push_back(BitStreamEntry{u64(m_pos - m_begin), m_word, m_wordBits, code, false, abbv.operands });
+			}
+
+#if UBA_LOG_STREAM
+			if (blockId == MODULE_BLOCK_ID)
+			{
+				auto str = ModuleCodeToString(code);
+				u32 indent = 2;
+				m_logger.Info(TC("%s%s (%u)"), TC("        ") + 8 - indent, str, code);
+			}
+#endif
+
+			ReadRecordOperands(outVals, abbv.operands, blob, blockId);
+
+			return code;
+		}
+
+		Entry Advance(u32 flags)
+		{
+			while (true)
+			{
+				UBA_ASSERT(!IsDone());
+
+				u32 code = ReadCode();
+				if (code == END_BLOCK)
+				{
+					if (flags & AF_DontPopBlockAtEnd)
+						return { EntryKind_EndBlock };
+					UBA_ASSERT(!m_blockScope.empty());
+					SkipToFourByteBoundary();
+					m_currentCodeSize = m_blockScope.back().prevCodeSize;
+					m_curAbbrevs = std::move(m_blockScope.back().prevAbbrevs);
+					m_blockScope.pop_back();
+					return { EntryKind_EndBlock };
+				}
+
+				if (code == ENTER_SUBBLOCK)
+				{
+					u32 subBlockId =  ReadVBR(8);
+					return { EntryKind_SubBlock, subBlockId };
+				}
+
+				if (code == DEFINE_ABBREV && !(flags & AF_DontAutoprocessAbbrevs))
+				{
+					ReadAbbrevRecord();
+					continue;
+				}
+
+				return { EntryKind_Record, code };
+			}
+		}
+
+		Entry AdvanceSkippingSubblocks(u32 flags)
+		{
+			while (true)
+			{
+				Entry entry = Advance(flags);
+				if (entry.kind != EntryKind_SubBlock)
+					return entry;
+				SkipBlock();
+			}
+		}
+
+		ObjectFileLLVMIR& m_owner;
+		Logger& m_logger;
+		u32 m_word = 0;
+		u32 m_wordBits = 0;
+		u8* m_begin;
+		u8* m_pos;
+		u8* m_end;
+
+
+		u32 m_currentCodeSize = 2;
+
+		struct Block
+		{
+			u32 prevCodeSize;
+			Vector<AbbrevPtr> prevAbbrevs;
+		};
+		Vector<Block> m_blockScope;
+		Vector<AbbrevPtr> m_curAbbrevs;
+		BlockInfo* m_blockInfo = nullptr;
+	};
+
+	ObjectFileLLVMIR::ObjectFileLLVMIR()
+	{
+		m_type = ObjectFileType_LLVMIR;
+	}
+
+	ObjectFileLLVMIR::~ObjectFileLLVMIR() = default;
+
+	u32 GetDllStorageIndex(u32 code)
+	{
+		if (code == MODULE_CODE_FUNCTION)
+			return 13;
+		if (code == MODULE_CODE_GLOBALVAR)
+			return 12;
+		if (code == MODULE_CODE_ALIAS)
+			return 7;
+		if (code == MODULE_CODE_ALIAS_OLD)
+			return 6;
+		UBA_ASSERTF(false, TC("module code %s not supported"), ModuleCodeToString(code));
+		return ~0u;
+	}
+
+	bool ObjectFileLLVMIR::Parse(Logger& logger, const tchar* filename)
+	{
+		//BinaryReader reader(m_data, 0, m_dataSize);
+		BitStreamReader reader(*this, logger, m_data, m_dataSize);
+
+		if (IsWrappedBitcode(m_data, m_dataSize))
+		{
+			UBA_ASSERT(false);
+			//reader.Skip(4);
+			//u32 version = reader.ReadU32();(void)version;
+			//u32 offset = reader.ReadU32();(void)offset;
+			//u32 size = reader.ReadU32();(void)size;
+			//u32 cpuType = reader.ReadU32();(void)cpuType;
+		}
+
+		reader.Read(32); // Skip magic
+
+		BlockInfo blockInfo;
+		reader.m_blockInfo = &blockInfo;
+
+		while (!reader.IsDone())
+		{
+			u32 code = reader.ReadCode();(void)code;
+			UBA_ASSERT(code == ENTER_SUBBLOCK);
+
+			u32 blockId = reader.ReadVBR(8);
+
+			u64 subBlockBitStart = reader.GetCurrentBitNo();(void)subBlockBitStart;
+			ParseBlock(logger, reader, blockInfo, blockId, 0);
+		}
+
+		u32 index = 0;
+		for (auto& record : m_globalVarOrFunctionRecords)
+		{
+			reader.m_pos = m_data + record.pos;
+			reader.m_word = record.word;
+			reader.m_wordBits = record.wordBits;
+
+			Vector<u64> recordData;
+			if (record.operands.empty())
+			{
+				u32 numElts = reader.ReadVBR(6);
+				recordData.reserve(recordData.size() + numElts);
+				for (u32 i = 0; i != numElts; ++i)
+					recordData.push_back(reader.ReadVBR64(6));
+			}
+			else
+			{
+				reader.ReadRecordOperands(recordData, record.operands);
+			}
+
+			std::string name;
+			u64* recIt = recordData.data();
+			u64 recSize = recordData.size();
+
+			if (recIt[0] + recIt[1] >= m_strTabSize)
+				continue;
+			name = std::string((char*)(m_data + m_strTabPos) + recIt[0], recIt[1]);
+
+			u32 dllStorageIndex = GetDllStorageIndex(record.code);
+
+			u64 dllStorage = ~0ull;
+			if (recSize > dllStorageIndex)
+				dllStorage = recIt[dllStorageIndex];
+
+			UBA_ASSERT(!record.isExport);
+
+			if (dllStorage == 1)
+				m_imports.emplace(name);
+			else if (dllStorage == 2)
+			{
+				UBA_ASSERT(!name.empty());
+				record.isExport = true;
+				m_exports.emplace(name, ExportInfo{"", index++});
+			}
+		}
+		return true;
+	}
+
+	bool ObjectFileLLVMIR::ParseBlock(Logger& logger, BitStreamReader& reader, BlockInfo& blockInfo, u32 blockId, u32 indent)
+	{
+		if (blockId == BLOCKINFO_BLOCK_ID)
+		{
+			uint64_t blockBitStart = reader.GetCurrentBitNo();
+			blockInfo = reader.ReadBlockInfoBlock(true);
+			reader.JumpToBit(blockBitStart);
+		}
+
+		u32 numWords = 0;
+		reader.EnterSubBlock(numWords, blockId);
+
+		#if UBA_LOG_STREAM
+		const tchar* str = BlockIdToString(blockId);
+		logger.Info(TC("%s%s (%u)"), TC("        ") + 8 - indent, str, blockId);
+		#endif
+
+		Vector<u64> record;
+		std::string blob;
+
+		if (blockId == GLOBALVAL_SUMMARY_BLOCK_ID)
+		{
+			// TODO: /EXPORT:xyz is probably in here somehow.. this format is sooo hard to parse... for now we just write directly into the strtab blob
+		}
+
+		while (true)
+		{
+			UBA_ASSERT(!reader.IsDone());
+		
+			Entry entry = reader.Advance(AF_DontAutoprocessAbbrevs);
+			switch (entry.kind)
+			{
+			case EntryKind_Record:
+				break;
+			case EntryKind_EndBlock:
+				return true;
+			case EntryKind_SubBlock:
+			{
+				u64 subBlockBitStart = reader.GetCurrentBitNo();(void)subBlockBitStart;
+				ParseBlock(logger, reader, blockInfo, entry.id, indent + 2);
+				continue;
+			}
+			default:
+				UBA_ASSERT(false);
+			}
+
+			if (entry.id == DEFINE_ABBREV)
+			{
+				reader.ReadAbbrevRecord();
+				continue;
+			}
+
+			if (blockId == MODULE_BLOCK_ID || blockId == STRTAB_BLOCK_ID)
+			{
+				record.clear();
+				reader.ReadRecord(record, entry.id, nullptr, blockId);
+			}
+			else
+				reader.SkipRecord(entry.id);
+
+#if UBA_LOG_STREAM
+			//logger.Log(LogEntryType_Info, blob.c_str(), blob.size());
+			#endif
+		}
+
+		return true;
+	}
+
+	bool ObjectFileLLVMIR::StripExports(Logger& logger, u8* newData, const UnorderedSymbols& allNeededImports)
+	{ 
+		if (m_exports.empty())
+			return true;
+
+		for (auto& record : m_globalVarOrFunctionRecords)
+		{
+			if (!record.isExport)
+				continue;
+	
+			BitStreamReader reader(*this, logger, newData, m_dataSize);
+			reader.m_pos = newData + record.pos;
+			reader.m_word = record.word;
+			reader.m_wordBits = record.wordBits;
+
+			u32 dllStorageIndex = GetDllStorageIndex(record.code);
+
+			UBA_ASSERT(record.operands.empty());
+
+			reader.ReadVBR(6);
+			for (u32 i=0;i!=dllStorageIndex; ++i)
+				reader.ReadVBR64(6);
+
+			// Write 0 to dll entry (no import/export)
+			//u64 res = reader.ReadVBR64(6);
+			//UBA_ASSERT(res == 2);
+			reader.ClearVBR64(6);
+
+#if 0
+			BitStreamReader reader2(*this, logger, newData, m_dataSize);
+			reader2.m_pos = newData + record.pos;
+			reader2.m_word = record.word;
+			reader2.m_wordBits = record.wordBits;
+			reader2.ReadVBR(6);
+			for (u32 i=0;i!=dllStorageIndex; ++i)
+				reader2.ReadVBR64(6);
+			u64 res2 = reader2.ReadVBR64(6);
+			UBA_ASSERT(res2  == 0);
+#endif
+		}
+
+		std::string tmp;
+
+		char* readPos = (char*)newData + m_strTabPos;
+		auto readEnd = readPos + m_strTabSize;
+		auto readLastPossiblePos = readEnd - 9;
+
+		while (true)
+		{
+			char* exportStr = nullptr;
+			char* it = readPos;
+
+			while (it < readLastPossiblePos)
+			{
+				if (memcmp(it, "/EXPORT:", 8) != 0)
+				{
+					++it;
+					continue;
+				}
+
+				exportStr = it;
+				break;
+			}
+
+			if (!exportStr)
+			{
+				readPos = readEnd;
+				break;
+			}
+
+			char* startPos = exportStr;
+			exportStr += 8;
+			char* exportEnd;
+			if (*exportStr == '\"')
+			{
+				++exportStr;
+				exportEnd = strchr(exportStr, '\"');
+				readPos = exportEnd + 1;
+				if (strncmp(readPos, ",DATA", 5) == 0)
+					readPos += 5;
+				readPos = Min(readPos, readEnd);
+			}
+			else
+			{
+				exportEnd = strchr(exportStr, ' ');
+				if (!exportEnd)
+					readPos = exportEnd = exportStr + strlen(exportStr);
+				else
+					readPos = exportEnd;
+				if (strncmp(exportEnd-5, ",DATA", 5) == 0)
+					exportEnd -= 5;
+			}
+
+			tmp.assign(exportStr, exportEnd - exportStr);
+			if (allNeededImports.find(tmp) != allNeededImports.end())
+				continue;
+			tmp.assign("__imp_").append(exportStr, exportEnd - exportStr);
+			if (allNeededImports.find(tmp) != allNeededImports.end())
+				continue;
+			
+			memset(startPos, ' ', readPos - startPos);
+		}
+
+		return true;
+	}
+	
+	bool ObjectFileLLVMIR::CreateExtraFile(Logger& logger, MemoryBlock& memoryBlock, const UnorderedSymbols& allNeededImports, const UnorderedSymbols& allSharedImports, const UnorderedExports& allSharedExports, bool includeExportsInFile)
+	{
+		return ObjectFileCoff::CreateExtraFile2(logger, memoryBlock, allNeededImports, allSharedImports, allSharedExports, includeExportsInFile);
 	}
 }
diff --git a/Engine/Source/Programs/UnrealBuildAccelerator/Common/Private/UbaObjectFileLLVMIR.h b/Engine/Source/Programs/UnrealBuildAccelerator/Common/Private/UbaObjectFileLLVMIR.h
--- a/Engine/Source/Programs/UnrealBuildAccelerator/Common/Private/UbaObjectFileLLVMIR.h
+++ b/Engine/Source/Programs/UnrealBuildAccelerator/Common/Private/UbaObjectFileLLVMIR.h
@@ -9,11 +9,43 @@ namespace uba
 	class ObjectFileLLVMIR : public ObjectFile
 	{
 	public:
+		ObjectFileLLVMIR();
+		~ObjectFileLLVMIR();
+
 		virtual bool Parse(Logger& logger, const tchar* filename) override;
 
 	private:
-		virtual bool StripExports(Logger& logger, u8* newData, const UnorderedSymbols& allNeededImports, u32& outKeptExportCount) override { return false; }
-		virtual bool CreateExtraFile(Logger& logger, MemoryBlock& memoryBlock, const UnorderedSymbols& allNeededImports, const UnorderedSymbols& allSharedImports, const UnorderedExports& allSharedExports, bool includeExportsInFile) override { return false; }
+		virtual bool StripExports(Logger& logger, u8* newData, const UnorderedSymbols& allNeededImports) override;
+		virtual bool CreateExtraFile(Logger& logger, MemoryBlock& memoryBlock, const UnorderedSymbols& allNeededImports, const UnorderedSymbols& allSharedImports, const UnorderedExports& allSharedExports, bool includeExportsInFile) override;
+
+		enum BlockInfoCodes : u8;
+		enum Encoding : u8;
+		enum EntryKind : u8;
+		enum FixedAbbrevIDs : u8;
+		struct Abbrev;
+		struct AbbrevOp;
+		struct BlockInfo;
+		struct Entry;
+		using AbbrevPtr = std::shared_ptr<Abbrev>;
+
+		class BitStreamReader;
+		bool ParseBlock(Logger& logger, BitStreamReader& reader, BlockInfo& blockInfo, u32 blockId, u32 indent);
+
+		struct BitStreamEntry
+		{
+			u64 pos;
+			u32 word;
+			u32 wordBits;
+			u32 code;
+			bool isExport;
+			Vector<AbbrevOp> operands;
+		};
+
+		Vector<BitStreamEntry> m_globalVarOrFunctionRecords;
+
+		u64 m_strTabPos;
+		u64 m_strTabSize;
+		//std::string m_strTab;
 	};
 
 
diff --git a/Engine/Source/Programs/UnrealBuildAccelerator/Common/Private/UbaProcess.cpp b/Engine/Source/Programs/UnrealBuildAccelerator/Common/Private/UbaProcess.cpp
--- a/Engine/Source/Programs/UnrealBuildAccelerator/Common/Private/UbaProcess.cpp
+++ b/Engine/Source/Programs/UnrealBuildAccelerator/Common/Private/UbaProcess.cpp
@@ -1525,6 +1525,7 @@ namespace uba
 			WaitForParent();
 			if (m_nativeProcessHandle == InvalidProcHandle) // Failed to create the child process
 				return UBA_EXIT_CODE(7);
+			m_extractExports = m_parentProcess->m_extractExports;
 		}
 
 		auto closeThreadHandle = MakeGuard([&]() { CloseHandle(m_nativeThreadHandle); m_nativeThreadHandle = 0; });
diff --git a/Engine/Source/Programs/UnrealBuildAccelerator/Common/Private/UbaSession.cpp b/Engine/Source/Programs/UnrealBuildAccelerator/Common/Private/UbaSession.cpp
--- a/Engine/Source/Programs/UnrealBuildAccelerator/Common/Private/UbaSession.cpp
+++ b/Engine/Source/Programs/UnrealBuildAccelerator/Common/Private/UbaSession.cpp
@@ -2787,7 +2787,7 @@ namespace uba
 		StringBuffer<> exportsFile;
 		exportsFile.Append(fileName, lastDot - fileName).Append(TC(".exi"));
 
-		MemoryBlock memoryBlock(1*1024*1024);
+		MemoryBlock memoryBlock(8*1024*1024);
 		if (!objectFile->WriteImportsAndExports(m_logger, memoryBlock))
 			return false;
 
diff --git a/Engine/Source/Programs/UnrealBuildAccelerator/Common/Public/UbaObjectFile.h b/Engine/Source/Programs/UnrealBuildAccelerator/Common/Public/UbaObjectFile.h
--- a/Engine/Source/Programs/UnrealBuildAccelerator/Common/Public/UbaObjectFile.h
+++ b/Engine/Source/Programs/UnrealBuildAccelerator/Common/Public/UbaObjectFile.h
@@ -10,13 +10,16 @@ namespace uba
 	class FileAccessor;
 	
 	using UnorderedSymbols = UnorderedSet<std::string>;
-	using UnorderedExports = UnorderedMap<std::string, std::string>;
+
+	struct ExportInfo { std::string extra; u32 index = 0; };
+	using UnorderedExports = UnorderedMap<std::string, ExportInfo>;
 
 	enum ObjectFileType : u8
 	{
 		ObjectFileType_Unknown,
 		ObjectFileType_Coff,
 		ObjectFileType_Elf,
+		ObjectFileType_LLVMIR,
 	};
 
 	class ObjectFile
@@ -39,9 +42,12 @@ namespace uba
 
 		virtual ~ObjectFile();
 
+		u8* GetData() { return m_data; }
+		u64 GetDataSize() { return m_dataSize; }
+
 	protected:
 		virtual bool Parse(Logger& logger, const tchar* filename) = 0;
-		virtual bool StripExports(Logger& logger, u8* newData, const UnorderedSymbols& allNeededImports, u32& outKeptExportCount) = 0;
+		virtual bool StripExports(Logger& logger, u8* newData, const UnorderedSymbols& allNeededImports) = 0;
 		virtual bool CreateExtraFile(Logger& logger, MemoryBlock& memoryBlock, const UnorderedSymbols& allNeededImports, const UnorderedSymbols& allSharedImports, const UnorderedExports& allSharedExports, bool includeExportsInFile) = 0;
 
 		FileAccessor* m_file = nullptr;
diff --git a/Engine/Source/Programs/UnrealBuildAccelerator/ObjTool/Private/UbaImportLibWriter.cpp b/Engine/Source/Programs/UnrealBuildAccelerator/ObjTool/Private/UbaImportLibWriter.cpp
new file mode 100644
--- /dev/null
+++ b/Engine/Source/Programs/UnrealBuildAccelerator/ObjTool/Private/UbaImportLibWriter.cpp
@@ -0,0 +1,678 @@
+// Copyright Epic Games, Inc. All Rights Reserved.
+
+#include "UbaImportLibWriter.h"
+#include "UbaFileAccessor.h"
+#include "UbaObjectFile.h"
+
+#if PLATFORM_WINDOWS
+#include <windows.h>
+#include <algorithm>
+#endif
+
+namespace uba
+{
+#if PLATFORM_WINDOWS
+
+	class ImportLib
+	{
+		#pragma pack(push, 1)
+		struct DATA_DESCRIPTOR
+		{
+			u16 machine;
+			u16 a;
+			u32 date;
+			u32 sizeOfData;
+			u32 b;
+			u32 flags;
+		};
+		struct SECTION_DESCRIPTOR_LONG
+		{
+			char name[16];
+			u32 sizeOfData;
+			u32 offset;
+			u32 aOffset;
+			u32 a;
+			u32 aCount;
+			u32 b;
+		};
+		struct SECTION_DESCRIPTOR_SHORT
+		{
+			union { char name[8]; struct { u32 unknown; u32 offset; } a; };
+			u32 b;
+			u32 c;
+			u16 type;
+		};
+		#pragma pack(pop)
+
+		u8 m_unknownData0[12] = { 0x02, 0x00, 0x00, 0x00, 0x20, 0x00, 0x09, 0x00, 0x00, 0x00, 0x00, 0x00 };
+		u8 m_unknownData1[41] = { 0x27, 0x00, 0x13, 0x10, 0x07, 0x00, 0x00, 0x00, 0xD0, 0x00, 0x00, 0x00, 0x00, 0x00, 0x00, 0x00, 0x0E, 0x00, 0x26, 0x00, 0x6A, 0x81, 0x12, 0x4D, 0x69, 0x63, 0x72, 0x6F, 0x73, 0x6F, 0x66, 0x74, 0x20, 0x28, 0x52, 0x29, 0x20, 0x4C, 0x49, 0x4E, 0x4B };
+		u8 m_unknownData2[20] = { 0x00, 0x00, 0x00, 0x00, 0x00, 0x00, 0x00, 0x00, 0x00, 0x00, 0x00, 0x00, 0x00, 0x00, 0x00, 0x00, 0x00, 0x00, 0x00, 0x00 };
+		u8 m_unknownData3[30] = { 0xC, 0x00, 0x00, 0x00, 0x03, 0x00, 0x00, 0x00, 0x03, 0x00, 0x00, 0x00, 0x00, 0x00, 0x04, 0x00, 0x00, 0x00, 0x03, 0x00, 0x10, 0x00, 0x00, 0x00, 0x05, 0x00, 0x00, 0x00, 0x03, 0x00 };
+
+		u32 AfterCompId = 0x0101816A;
+		u32 PredefinedSymbolsCount = 3;
+
+		struct SymbolInfo
+		{
+			std::string name;
+			u32 offsetIndex;
+			union { u32 hint; u32 ordinal; u32 value; };
+			u32 type;
+			u32 tempOrder;
+			u8 extra;
+			u8 isData;
+		};
+
+		enum CallingConvention
+		{
+			CallingConvention_Cdecl,
+			CallingConvention_StdCall,
+			CallingConvention_FastCall,
+			CallingConvention_VectorCall
+		};
+
+		enum ImportType
+		{
+			ImportType_Ordinal,
+			ImportType_SpecifiedName,
+			ImportType_DecoratedName 
+		};
+
+		MemoryBlock m_memory;
+		u32 m_memoryOffset;
+
+		Vector<SymbolInfo*> m_symbols;
+		u32 m_extraSymbolCount = 0;
+		u32 m_dataCount;
+		u16 m_machine = IMAGE_FILE_MACHINE_AMD64;
+		u64 m_date;
+		std::string m_moduleName;
+		std::string m_headerName;
+		bool m_noHeaderName;
+
+		template<typename Func>
+		void TraverseSortedSymbols(const Func& func)
+		{
+			auto sortedVec(m_symbols);
+			std::sort(sortedVec.begin(), sortedVec.end(), [](SymbolInfo* a, SymbolInfo* b) { return a->name < b->name; });
+			for (auto& s : sortedVec)
+				func(*s);
+		}
+
+		void AddSymbol(const char* name, const char* extra, u32 value, u32 ArgListSizeInBytes, CallingConvention CallingConvention, ImportType importType)
+		{
+			u32 offsetIndex = u32(m_symbols.size()) - m_extraSymbolCount;
+			SymbolInfo& s0 = *m_symbols.emplace_back(new SymbolInfo());
+
+			s0.isData = strstr(extra, "DATA") != 0;
+			s0.tempOrder = u32(m_symbols.size()-1) - m_extraSymbolCount - 3;
+			s0.offsetIndex = offsetIndex;
+
+			if (s0.isData)
+			{
+				s0.name = std::string("__imp_") + name;
+				s0.value = 0;
+				s0.type = 5;
+				return;
+			}
+
+
+			char buf[256];
+			if (CallingConvention == CallingConvention_VectorCall)
+			{
+				sprintf_s(buf, sizeof(buf), "%s@@%d", name, ArgListSizeInBytes);
+				name = buf;
+			}
+			s0.name = name;
+
+			switch (importType)
+			{
+			case ImportType_Ordinal:
+				s0.ordinal = value;
+				s0.type = 0;
+				break;
+			case ImportType_SpecifiedName:
+				s0.hint = value;
+				if (CallingConvention == CallingConvention_VectorCall)
+					s0.type = 0xC;
+				else
+					s0.type = 4;
+				break;
+			case ImportType_DecoratedName:
+				s0.hint = value;
+				s0.type = 4;
+				break;
+			}
+
+			auto& s1 = *m_symbols.emplace_back(new SymbolInfo());
+			s1.name = "__imp_" + s0.name;
+			s1.offsetIndex = offsetIndex;
+			s1.value = 0;
+			s1.type = 0;
+			s1.extra = 1;
+
+			++m_extraSymbolCount;
+		}
+
+		void InitSymbolList(const char* name)
+		{
+			char buf[200];
+
+			sprintf_s(buf, sizeof(buf), "__IMPORT_DESCRIPTOR_%s", name);
+			auto& s0 = *m_symbols.emplace_back(new SymbolInfo());
+			s0.name = buf;
+			s0.offsetIndex = 0;
+			s0.value = 0;
+			s0.type = 0;
+			auto& s1 = *m_symbols.emplace_back(new SymbolInfo());
+			s1.name = "__NULL_IMPORT_DESCRIPTOR";
+			s1.offsetIndex = 1;
+			s1.value = 0;
+			s1.type = 0;
+			sprintf_s(buf, sizeof(buf), "\x7F%s_NULL_THUNK_DATA", name);
+			auto& s2 = *m_symbols.emplace_back(new SymbolInfo());
+			s2.name = buf;
+			s2.offsetIndex = 2;
+			s2.value = 0;
+			s2.type = 0;
+		}
+
+		void Write(const void* data, u64 bytes)
+		{
+			u32 newOffset = m_memoryOffset + u32(bytes);
+			if (m_memory.writtenSize < newOffset)
+				m_memory.AllocateNoLock(newOffset - m_memory.writtenSize, 1, TC(""));
+			memcpy(m_memory.memory + m_memoryOffset, data, bytes);
+			m_memoryOffset = newOffset;
+		}
+
+		u32 ByteSwap(u32 Value)
+		{
+			return _byteswap_ulong(Value);
+		}
+		void WriteU8(u8 value) { Write(&value, sizeof(value)); }
+		void WriteU16(u16 value) { Write(&value, sizeof(value)); }
+		void WriteU32(u32 value) { Write(&value, sizeof(value)); }
+		void SkipWrite(u32 distance) { m_memoryOffset += distance; }
+		u32 GetWritePosition() { return m_memoryOffset; }
+		void SetWritePosition(u32 offset) { m_memoryOffset = offset; }
+
+		void WriteFileHeader(const char* str, bool prefixWithSlash, u32 a, u32 size)
+		{
+			IMAGE_ARCHIVE_MEMBER_HEADER FileHeader;
+			memset(&FileHeader, ' ', IMAGE_SIZEOF_ARCHIVE_MEMBER_HDR);
+		
+			u32 stringLength = u32(strlen(str));
+			if (prefixWithSlash)
+			{
+				FileHeader.Name[0] = '/';
+				memcpy(FileHeader.Name + 1, str, stringLength);
+			}
+			else
+			{
+				memcpy(FileHeader.Name, str, stringLength);
+				FileHeader.Name[stringLength] = '/';
+			}
+
+			sprintf_s((char*)FileHeader.Date, 13, "%-12I64d", m_date);
+			FileHeader.UserID[0] = ' '; // Remove 0 termination from Date
+			sprintf_s((char*)FileHeader.Mode, 9, "%-8ho", a);
+			sprintf_s((char*)FileHeader.Size, 11, "%-10d", size);
+			FileHeader.EndHeader[0] = 0x60;
+			FileHeader.EndHeader[1] = 0xA;
+			Write(&FileHeader, IMAGE_SIZEOF_ARCHIVE_MEMBER_HDR);
+		}
+
+		void WriteSectionDescriptionLong(const char* sectionName, u32 sizeOfData, u32 offset, u32 aOffset, u32 a, u32 aCount, u32 b)
+		{
+			SECTION_DESCRIPTOR_LONG desc;
+			memset(&desc, 0, sizeof(desc));
+			strcpy_s(desc.name, sizeof(desc.name), sectionName);
+			desc.sizeOfData = sizeOfData;
+			desc.offset = offset;
+			desc.aOffset = aOffset;
+			desc.a = a;
+			desc.aCount = aCount;
+			desc.b = b;
+			Write(&desc, sizeof(desc));
+		}
+
+		void WriteSectionDescriptionShort(const char* sectionName, u32 b, u32 c, u16 type)
+		{
+			SECTION_DESCRIPTOR_SHORT desc;
+			memset(&desc, 0, sizeof(desc));
+			memcpy(desc.name, sectionName, strlen(sectionName));
+			desc.b = b;
+			desc.c = c;
+			desc.type = type;
+			Write(&desc, sizeof(desc));
+		}
+
+		void WriteSectionDescriptionShort(u32 offset, u32 b, u32 c, u16 type)
+		{
+			SECTION_DESCRIPTOR_SHORT desc;
+			memset(&desc, 0, sizeof(desc));
+			desc.a.offset = offset;
+			desc.b = b;
+			desc.c = c;
+			desc.type = type;
+			Write(&desc, sizeof(desc));
+		}
+
+		void WriteDataDescriptor(u16 a, u32 sizeOfData, u32 b)
+		{
+			DATA_DESCRIPTOR desc;
+			desc.machine = m_machine;
+			desc.a = a;
+			desc.date = u32(m_date);
+			desc.sizeOfData = sizeOfData;
+			desc.b = b;
+			desc.flags = 0;
+			Write(&desc, sizeof(desc));
+		}
+
+		void WriteImportLibrary()
+		{
+			Write(IMAGE_ARCHIVE_START, IMAGE_ARCHIVE_START_SIZE);
+
+			Vector<u32> symbolOffsets;
+			symbolOffsets.resize(2 + m_symbols.size() - m_extraSymbolCount);
+
+			{
+				// File 1
+				u32 offset = GetWritePosition();
+				symbolOffsets[0] = offset;
+
+				SkipWrite(IMAGE_SIZEOF_ARCHIVE_MEMBER_HDR);
+				WriteU32(ByteSwap(u32(m_symbols.size())));
+				SkipWrite(u32(m_symbols.size()) * sizeof(u32));
+				for (auto symbol : m_symbols)
+					Write(symbol->name.data(), symbol->name.size() + 1);
+
+				u32 offset2 = GetWritePosition();
+
+				SetWritePosition(offset);
+				WriteFileHeader("", false, 0, offset2 - (offset + IMAGE_SIZEOF_ARCHIVE_MEMBER_HDR));
+				SetWritePosition(offset2);
+				if (offset2 & 1)
+					Write("\n", 1);
+
+				// File 2
+				offset = GetWritePosition();
+				symbolOffsets[1] = offset;
+
+				SkipWrite(IMAGE_SIZEOF_ARCHIVE_MEMBER_HDR);
+				u32 symbolCount = u32(m_symbols.size()) - m_extraSymbolCount;
+				WriteU32(symbolCount);
+				SkipWrite(symbolCount * sizeof(u32));
+				symbolCount = u32(m_symbols.size());
+				WriteU32(symbolCount);
+				SkipWrite(symbolCount * sizeof(u16));
+				TraverseSortedSymbols([&](SymbolInfo& symbol) { Write(symbol.name.data(), symbol.name.size() + 1); });
+
+				offset2 = GetWritePosition();
+
+				SetWritePosition(offset);
+				WriteFileHeader("", false, 0, offset2 - (offset + IMAGE_SIZEOF_ARCHIVE_MEMBER_HDR));
+				SetWritePosition(offset2);
+				if (offset2 & 1)
+					Write("\n", 1);
+
+				u32 len = u32(m_moduleName.size()) + 1;
+				if (len > 0x10)
+				{
+					WriteFileHeader("/", false, 0, len);
+					Write(m_moduleName.data(), len);
+					if (len & 1)
+						Write("\n", 1);
+				}
+			}
+
+			{
+				u32 offset = GetWritePosition();
+				symbolOffsets[2] = offset;
+
+				SkipWrite(IMAGE_SIZEOF_ARCHIVE_MEMBER_HDR);
+
+				u8 moduleNameLen = u8(m_moduleName.size());
+				u32 stringPad = (moduleNameLen + 1) & 1;
+
+				u32 sizeOfData = (sizeof(DATA_DESCRIPTOR) + (sizeof(SECTION_DESCRIPTOR_LONG) * 3)) +
+					(sizeof(m_unknownData0) + sizeof(moduleNameLen) + moduleNameLen +
+					sizeof(m_unknownData1)) + (sizeof(m_unknownData2) + sizeof(m_unknownData3)) +
+					(moduleNameLen + 1) + stringPad;
+
+				WriteDataDescriptor(3, sizeOfData, 8);
+
+				// header0
+				sizeOfData = sizeof(m_unknownData0) + sizeof(moduleNameLen) + moduleNameLen + sizeof(m_unknownData1);
+				u32 offset3 = sizeof(DATA_DESCRIPTOR) + (sizeof(SECTION_DESCRIPTOR_LONG) * 3);
+				WriteSectionDescriptionLong(".debug$S", sizeOfData, offset3, 0, 0, 0, 0x42100040);
+
+				// header1
+				offset3 += sizeOfData;
+				sizeOfData = sizeof(m_unknownData2);
+				u32 AOffset = offset3 + sizeOfData;
+				WriteSectionDescriptionLong(".idata$2", sizeOfData, offset3, AOffset, 0, 3, 0xC0300040);
+
+				// header2
+				offset3 += sizeOfData + sizeof(m_unknownData3);
+				sizeOfData = moduleNameLen + 1 + stringPad;
+				WriteSectionDescriptionLong(".idata$6", sizeOfData, offset3, AOffset, 0, 0, 0xC0200040);
+
+				// data0
+				Write(m_unknownData0, sizeof(m_unknownData0));
+				Write(&moduleNameLen, sizeof(moduleNameLen));
+				Write(m_moduleName.data(), moduleNameLen);
+				Write(m_unknownData1, sizeof(m_unknownData1));
+
+				// data1
+				Write(m_unknownData2, sizeof(m_unknownData2));
+
+				// data2
+				Write(m_unknownData3, sizeof(m_unknownData3));
+
+				// 
+				Write(m_moduleName.data(), moduleNameLen + 1);
+				if (stringPad)
+					WriteU8(0);
+
+				WriteSectionDescriptionShort("@comp.id", AfterCompId, 0xFFFF, 3);
+
+				offset3 = sizeof(sizeOfData);
+				WriteSectionDescriptionShort(offset3, 0, 2, 2);
+
+				WriteSectionDescriptionShort(".idata$2", 0xC0000040, 2, 0x68);
+				WriteSectionDescriptionShort(".idata$6", 0, 3, 3);
+				WriteSectionDescriptionShort(".idata$4", 0xC0000040, 0, 0x68);
+				WriteSectionDescriptionShort(".idata$5", 0xC0000040, 0, 0x68);
+
+				offset3 += u32(m_symbols[0]->name.size()) + 1;
+				WriteSectionDescriptionShort(offset3, 0, 0, 2);
+
+				offset3 += u32(m_symbols[1]->name.size()) + 1;
+				WriteSectionDescriptionShort(offset3, 0, 0, 2);
+
+				offset3 += u32(m_symbols[2]->name.size()) + 1;
+				WriteU32(offset3);
+
+				for (u32 i = 0; i < PredefinedSymbolsCount; ++i)
+					Write(m_symbols[i]->name.data(), u32(m_symbols[i]->name.size()) + 1);
+
+				u32 offset2 = GetWritePosition();
+				SetWritePosition(offset);
+				WriteFileHeader(m_headerName.data(), m_noHeaderName, 0, offset2 - (offset + IMAGE_SIZEOF_ARCHIVE_MEMBER_HDR));
+				SetWritePosition(offset2);
+				if (offset2 & 1)
+					Write("\n", 1);
+			}
+
+			{
+				u32 offset = GetWritePosition();
+				symbolOffsets[3] = offset;
+
+				SkipWrite(IMAGE_SIZEOF_ARCHIVE_MEMBER_HDR);
+
+				// DATA DESCRIPTOR
+				u8 moduleNameLen = u8(m_moduleName.size());
+				u32 sizeOfData = (sizeof(DATA_DESCRIPTOR) + (sizeof(SECTION_DESCRIPTOR_LONG) * 2)) +
+					(sizeof(m_unknownData0) + sizeof(moduleNameLen) + moduleNameLen + sizeof(m_unknownData1)) +
+					(sizeof(m_unknownData2));
+
+				WriteDataDescriptor(2, sizeOfData, 2);
+
+				// DATA HEADER 1
+				sizeOfData = sizeof(m_unknownData0) + sizeof(moduleNameLen) + moduleNameLen + sizeof(m_unknownData1);
+				u32 offset3 = sizeof(DATA_DESCRIPTOR) + (sizeof(SECTION_DESCRIPTOR_LONG) * 2);
+				WriteSectionDescriptionLong(".debug$S", sizeOfData, offset3, 0, 0, 0, 0x42100040);
+
+				// DATA HEADER 2
+				offset3 += sizeOfData;
+				sizeOfData = sizeof(m_unknownData2);
+				WriteSectionDescriptionLong(".idata$3", sizeOfData, offset3, 0, 0, 0, 0xC0300040);
+
+				// DATA 1
+				Write(m_unknownData0, sizeof(m_unknownData0));
+				Write(&moduleNameLen, sizeof(moduleNameLen));
+				Write(m_moduleName.data(), moduleNameLen);
+				Write(m_unknownData1, sizeof(m_unknownData1));
+
+				// DATA 2
+				Write(m_unknownData2, sizeof(m_unknownData2));
+
+				WriteSectionDescriptionShort("@comp.id", AfterCompId, 0xFFFF, 3);
+
+				offset3 = sizeof(sizeOfData);
+				WriteSectionDescriptionShort(offset3, 0, 2, 2);
+
+				SymbolInfo& nullImportSymbol = *m_symbols[1];
+				WriteU32(offset3 + u32(nullImportSymbol.name.size() + 1));
+				Write(nullImportSymbol.name.data(), nullImportSymbol.name.size() + 1);
+
+				u32 offset2 = GetWritePosition();
+				SetWritePosition(offset);
+				WriteFileHeader(m_headerName.data(), m_noHeaderName, 0, offset2 - (offset + IMAGE_SIZEOF_ARCHIVE_MEMBER_HDR));
+				SetWritePosition(offset2);
+				if (offset2 & 1)
+					Write("\n", 1);
+			}
+
+			{
+				u32 offset = GetWritePosition();
+				symbolOffsets[4] = offset;
+
+				SkipWrite(IMAGE_SIZEOF_ARCHIVE_MEMBER_HDR);
+
+				u8 moduleNameLen = u8(m_moduleName.size());
+
+				u32 pad[2] = { 0 };
+
+				u32 sizeOfData = (sizeof(DATA_DESCRIPTOR) + (sizeof(SECTION_DESCRIPTOR_LONG) * 3)) +
+					(sizeof(m_unknownData0) + sizeof(moduleNameLen) + moduleNameLen + sizeof(m_unknownData1)) +
+					sizeof(pad) +
+					sizeof(pad);
+
+				WriteDataDescriptor(3, sizeOfData, 2);
+
+				// header 0
+				sizeOfData = sizeof(m_unknownData0) + sizeof(moduleNameLen) + moduleNameLen + sizeof(m_unknownData1);
+				u32 offset3 = sizeof(DATA_DESCRIPTOR) + (sizeof(SECTION_DESCRIPTOR_LONG) * 3);
+				WriteSectionDescriptionLong(".debug$S", sizeOfData, offset3, 0, 0, 0, 0x42100040);
+
+				// header 1
+				offset3 += sizeOfData;
+				sizeOfData = sizeof(pad);
+				WriteSectionDescriptionLong(".idata$5", sizeOfData, offset3, 0, 0, 0, 0xC0400040);
+
+				// header 2
+				offset3 += sizeOfData;
+				sizeOfData = sizeof(pad);
+				WriteSectionDescriptionLong(".idata$4", sizeOfData, offset3, 0, 0, 0, 0xC0400040);
+
+				// data 0
+				Write(m_unknownData0, sizeof(m_unknownData0));
+				Write(&moduleNameLen, sizeof(moduleNameLen));
+				Write(m_moduleName.data(), moduleNameLen);
+				Write(m_unknownData1, sizeof(m_unknownData1));
+
+				// data 1 & 2
+				Write(pad, sizeof(pad));
+				Write(pad, sizeof(pad));
+
+				WriteSectionDescriptionShort("@comp.id", AfterCompId, 0xFFFF, 3);
+
+				offset3 = sizeof(sizeOfData);
+				WriteSectionDescriptionShort(offset3, 0, 2, 2);
+
+				SymbolInfo& thunkDataSymbol = *m_symbols[2];
+				WriteU32(offset3 + u32(thunkDataSymbol.name.size() + 1));
+				Write(thunkDataSymbol.name.data(), thunkDataSymbol.name.size() + 1);
+
+				u32 offset2 = GetWritePosition();
+				SetWritePosition(offset);
+				WriteFileHeader(m_headerName.data(), m_noHeaderName, 0, offset2 - (offset + IMAGE_SIZEOF_ARCHIVE_MEMBER_HDR));
+				SetWritePosition(offset2);
+				if (offset2 & 1)
+					Write("\n", 1);
+			}
+
+			{
+				Vector<SymbolInfo*> sortedFunctions;
+				for (u32 i = PredefinedSymbolsCount; i < m_symbols.size(); ++i)
+					if (!m_symbols[i]->extra)
+						sortedFunctions.push_back(m_symbols[i]);
+
+				std::sort(sortedFunctions.begin(), sortedFunctions.end(), [](SymbolInfo* a, SymbolInfo* b)
+					{
+						const char* aName = a->name.data();
+						if (a->isData)
+							aName += 6;
+						const char* bName = b->name.data();
+						if (b->isData)
+							bName += 6;
+						return strcmp(aName, bName) < 0;
+					});
+
+				for (u32 i = 0, e = u32(sortedFunctions.size()); i != e; ++i)
+				{
+					auto& symbol = *sortedFunctions[i];
+					symbolOffsets[2 + PredefinedSymbolsCount + symbol.tempOrder] = GetWritePosition();
+
+					const char* symbolName = symbol.name.data();
+					u32 symbolNameLen = u32(symbol.name.size());
+					if (symbol.isData)
+					{
+						symbolName += 6;
+						symbolNameLen -= 6;
+					}
+
+					u32 sizeOfData = u32(m_moduleName.size()) + symbolNameLen + 2;
+					WriteFileHeader(m_headerName.data(), m_noHeaderName, 0, sizeOfData + sizeof(IMPORT_OBJECT_HEADER));
+
+					IMPORT_OBJECT_HEADER header;
+					memset(&header, 0, sizeof(IMPORT_OBJECT_HEADER));
+					header.Sig1 = 0;
+					header.Sig2 = 0xFFFF;
+					header.Version = 0;
+					header.Machine = m_machine;
+					header.TimeDateStamp = u32(m_date);
+					header.SizeOfData = sizeOfData;
+					header.Ordinal = u16(i);
+					header.Type = symbol.type & 3;
+					header.NameType = symbol.type >> 2;
+					header.Reserved = 0;
+
+					Write(&header, sizeof(header));
+					Write(symbolName, symbolNameLen + 1);
+					Write(m_moduleName.data(), m_moduleName.size() + 1);
+					if (GetWritePosition() & 1)
+						Write("\n", 1);
+				}
+
+				u32 Offset = GetWritePosition();
+
+				// Offsets
+				SetWritePosition(symbolOffsets[0] + IMAGE_SIZEOF_ARCHIVE_MEMBER_HDR + sizeof(u32));
+				for (auto& symbol : m_symbols)
+					WriteU32(ByteSwap(symbolOffsets[2 + symbol->offsetIndex]));
+
+				// Offset table
+				SetWritePosition(symbolOffsets[1] + IMAGE_SIZEOF_ARCHIVE_MEMBER_HDR + sizeof(u32));
+				for (u32 i = 0; i < (m_symbols.size() - m_extraSymbolCount); ++i)
+					Write((symbolOffsets.data() + (2 + i)), sizeof(u32));
+
+				// Offset indices
+				SkipWrite(sizeof(u32));
+				TraverseSortedSymbols([&](SymbolInfo& symbol) { WriteU16(u16(symbol.offsetIndex + 1)); });
+				SetWritePosition(Offset);
+			}
+		}
+
+	public:
+
+		ImportLib() : m_memory(128*1024*1024) // 128mb should be more than enough reserved memory for the biggest import libs
+		{
+		}
+
+		~ImportLib()
+		{
+			for (auto ptr : m_symbols)
+				delete ptr;
+		}
+
+		bool WriteFile(Logger& logger, const Vector<ObjectFile*>& objFiles, const char* libName, const tchar* libFile)
+		{
+			char temp[256];
+			const char* libExt = strrchr(libName, '.');
+			if (libExt)
+			{
+				u32 len = u32(libExt - libName);
+				memcpy(temp, libName, len);
+				temp[len] = 0;
+				libName = temp;
+			}
+			else
+				libExt = ".dll";
+
+			InitSymbolList(libName);
+
+			UnorderedSymbols handled;
+			for (auto objFile : objFiles)
+			{
+				if (!objFile)
+					continue;
+
+				Map<u32, const UnorderedExports::value_type*> sortedSymbols;
+				for (auto& kv : objFile->GetExports())
+					if (handled.insert(kv.first).second)
+						sortedSymbols.try_emplace(kv.second.index, &kv);
+				for (auto& outerKv : sortedSymbols)
+					AddSymbol(outerKv.second->first.c_str(), outerKv.second->second.extra.c_str(), 0, 4, CallingConvention_StdCall, ImportType_SpecifiedName);
+			}
+
+			m_moduleName = std::string(libName) + libExt;
+
+			if ((m_moduleName.size() + 1) > 0x10)
+			{
+				m_headerName = "0";
+				m_noHeaderName = TRUE;
+			}
+			else
+			{
+				m_headerName = m_moduleName;
+				m_noHeaderName = FALSE;
+			}
+
+			m_date = ~0ull;//_time64(NULL);
+
+			m_unknownData0[4] = u8(m_moduleName.size()) + 7; // Have no idea why this matches and what it does.
+			// BlankProgram-Projects.lib - 0x20
+			// BlankProgram-Core.lib - 0x1C
+			// BlankProgram-Json.lib - 0x1C
+			// BlankProgram-BuildSettings.lib - 0x25
+
+			m_memoryOffset = 0;
+
+			WriteImportLibrary();
+
+			FileAccessor fa(logger, libFile);
+			if (!fa.CreateMemoryWrite(false, DefaultAttributes(), m_memory.writtenSize))
+				return false;
+
+			memcpy(fa.GetData(), m_memory.memory, m_memory.writtenSize);
+
+			return fa.Close();
+		}
+	};
+#endif // PLATFORM_WINDOWS
+
+	bool ImportLibWriter::Write(Logger& logger, const Vector<ObjectFile*>& objFiles, const char* libName, const tchar* libFile)
+	{
+#if PLATFORM_WINDOWS
+		ImportLib lib;
+		return lib.WriteFile(logger, objFiles, libName, libFile);
+#else
+		return false;
+#endif
+	}
+}
diff --git a/Engine/Source/Programs/UnrealBuildAccelerator/ObjTool/Private/UbaImportLibWriter.h b/Engine/Source/Programs/UnrealBuildAccelerator/ObjTool/Private/UbaImportLibWriter.h
new file mode 100644
--- /dev/null
+++ b/Engine/Source/Programs/UnrealBuildAccelerator/ObjTool/Private/UbaImportLibWriter.h
@@ -0,0 +1,17 @@
+// Copyright Epic Games, Inc. All Rights Reserved.
+
+#pragma once
+
+#include "UbaLogger.h"
+#include "UbaStringBuffer.h"
+
+namespace uba
+{
+	class ObjectFile;
+	
+	class ImportLibWriter
+	{
+	public:
+		bool Write(Logger& logger, const Vector<ObjectFile*>& objFiles, const char* libName, const tchar* libFile);
+	};
+}
diff --git a/Engine/Source/Programs/UnrealBuildAccelerator/ObjTool/Private/UbaObjTool.cpp b/Engine/Source/Programs/UnrealBuildAccelerator/ObjTool/Private/UbaObjTool.cpp
--- a/Engine/Source/Programs/UnrealBuildAccelerator/ObjTool/Private/UbaObjTool.cpp
+++ b/Engine/Source/Programs/UnrealBuildAccelerator/ObjTool/Private/UbaObjTool.cpp
@@ -1,8 +1,12 @@
 // Copyright Epic Games, Inc. All Rights Reserved.
 
+#define Local_GetLongPathNameW uba::GetLongPathNameW
+
 #include "UbaObjectFile.h"
 #include "UbaDirectoryIterator.h"
 #include "UbaFileAccessor.h"
+#include "UbaImportLibWriter.h"
+#include "UbaPathUtils.h"
 #include "UbaVersion.h"
 #include "UbaWorkManager.h"
 
@@ -64,12 +68,17 @@ namespace uba
 		TString objFile;
 		bool printSymbols = false;
 		bool stripExports = false;
+		bool writeImpLib = false;
+		bool isImpLibRsp = false;
 
 		Vector<TString> objFilesToStrip;
 		Vector<TString> objFilesDependencies;
 		TString extraObjFile;
+		Vector<TString> objFilesForImpLib;
+		std::string impLibName;
+		TString impLibFile;
 
-		auto parseArg = [&](const tchar* arg)
+		auto parseArg = [&](const tchar* arg, bool isRsp)
 			{
 				StringBuffer<> name;
 				StringBuffer<> value;
@@ -81,29 +90,99 @@ namespace uba
 				}
 				else
 				{
-					name.Append(arg);
+					const tchar* colon = TStrchr(arg,':');
+					if (colon && colon[1] != '\\')
+					{
+						name.Append(arg, colon - arg);
+						const tchar* valueStart = colon+1;
+						if (*valueStart == '\"')
+							++valueStart;
+						value.Append(valueStart);
+						if (value.data[value.count-1] == '\"')
+							value.Resize(value.count-1);
+
+					}
+					else
+					{
+						name.Append(arg);
+					}
 				}
 
-				if (name.StartsWith(TC("/D:")))
+				if (isImpLibRsp)
 				{
-					objFilesDependencies.push_back(name.data + 3);
+					if (name.Equals(TC("/NOLOGO")))
+					{
+					}
+					else if (name.Equals(TC("/errorReport")))
+					{
+					}
+					else if (name.Equals(TC("/MACHINE")))
+					{
+						//if (!value.Equals(TC("x64")))
+						//{
+						//	logger.Error(TC("only x64 supported"));
+						//	return 0;
+						//}
+					}
+					else if (name.Equals(TC("/SUBSYSTEM")))
+					{
+					}
+					else if (name.Equals(TC("/DEF")))
+					{
+						writeImpLib = true;
+					}
+					else if (name.Equals(TC("/NAME")))
+					{
+						char buffer[256];
+						value.Parse(buffer, sizeof(buffer));
+						impLibName = buffer;
+					}
+					else if (name.Equals(TC("/OUT")))
+					{
+						impLibFile = value.data;
+					}
+					else if (name.Equals(TC("/IGNORE")))
+					{
+					}
+					else if (name.Equals(TC("/NODEFAULTLIB")))
+					{
+					}
+					else if (name.Equals(TC("/LTCG")))
+					{
+					}
+					else
+					{
+						objFilesForImpLib.push_back(name.data);
+					}
+				}
+				else if (name.StartsWith(TC("/D")))
+				{
+					objFilesDependencies.push_back(value.data);
 				}
-				else if (name.StartsWith(TC("/S:")))
+				else if (name.StartsWith(TC("/S")))
 				{
-					objFilesToStrip.push_back(name.data + 3);
+					objFilesToStrip.push_back(value.data);
 				}
-				else if (name.StartsWith(TC("/O:")))
+				else if (name.StartsWith(TC("/O")))
 				{
-					extraObjFile = name.data + 3;
+					extraObjFile = value.data;
 				}
 				else if (name.Equals(TC("-printsymbols")))
 				{
 					printSymbols = true;
 				}
+				else if (name.Equals(TC("-writeimplib")))
+				{
+				}
 				else if (name.Equals(TC("-stripexports")))
 				{
 					stripExports = true;
 				}
+				else if (name.Equals(TC("/LIB")))
+				{
+					isImpLibRsp = true;
+					writeImpLib = true;
+				}
 				else if (name.Equals(TC("-?")))
 				{
 					return PrintHelp(TC(""));
@@ -139,7 +218,7 @@ namespace uba
 				LoggerWithWriter logger(g_consoleLogWriter, TC(""));
 				if (!ReadLines(logger, arg, [&](const TString& line)
 					{
-						res = parseArg(line.c_str());
+						res = parseArg(line.c_str(), true);
 						return res == 0;
 					}))
 					return -1;
@@ -147,7 +226,7 @@ namespace uba
 					return res;
 				continue;
 			}
-			int res = parseArg(arg);
+			int res = parseArg(arg, false);
 			if (res != 0)
 				return res;
 		}
@@ -179,7 +258,6 @@ namespace uba
 						return;
 					}
 					ScopedCriticalSection _(cs);
-					UBA_ASSERT(type == ObjectFileType_Unknown || type == symbolFile.type);
 					if (type == ObjectFileType_Unknown)
 						type = symbolFile.type;
 					allNeededImports.insert(symbolFile.imports.begin(), symbolFile.imports.end());
@@ -203,7 +281,6 @@ namespace uba
 						return;
 					}
 					ScopedCriticalSection _(cs);
-					UBA_ASSERT(type == ObjectFileType_Unknown || type == symbolFile.type);
 					if (type == ObjectFileType_Unknown)
 						type = symbolFile.type;
 					allSharedImports.insert(symbolFile.imports.begin(), symbolFile.imports.end());
@@ -218,6 +295,37 @@ namespace uba
 
 			//logger.Info(TC("Reduced export count from %llu to %llu"), totalExportCount.load(), totalKeptExportCount.size());
 		}
+		else if (writeImpLib)
+		{
+			StringBuffer<> currentDir;
+			GetCurrentDirectoryW(currentDir);
+			currentDir.EnsureEndsWithSlash();
+			ImportLibWriter writer;
+			if (objFilesForImpLib.empty() && !objFile.empty())
+				objFilesForImpLib.push_back(objFile);
+
+			Atomic<bool> success = true;
+			Vector<ObjectFile*> objFiles;
+			objFiles.resize(objFilesForImpLib.size());
+			u32 workerCount = DefaultProcessorCount;
+			WorkManagerImpl workManager(workerCount);
+			workManager.ParallelFor(workerCount, objFilesForImpLib, [&](auto& it)
+			{
+				auto& o = *it;
+				StringBuffer<> fixedPath;
+				FixPath(o.c_str(), currentDir.data, currentDir.count, fixedPath);
+				if (fixedPath.EndsWith(TC(".res")) || fixedPath.EndsWith(TC(".lib")))
+					return;
+				if (ObjectFile* objectFile = ObjectFile::OpenAndParse(logger, fixedPath.data))
+					objFiles[it - objFilesForImpLib.begin()] = objectFile;
+				else
+					success = false;
+			});
+			if (!success)
+				return -1;
+
+			writer.Write(logger, objFiles, impLibName.c_str(), impLibFile.c_str());
+		}
 		else
 		{
 			if (objFile.empty())
@@ -234,7 +342,7 @@ namespace uba
 					logger.Info(TC("I %S"), symbol.c_str());
 
 				for (auto& kv : objectFile->GetExports())
-					logger.Info(TC("E %S%S"), kv.first.c_str(), kv.second.c_str());
+					logger.Info(TC("E %S%S"), kv.first.c_str(), kv.second.extra.c_str());
 			}
 
 			if (stripExports)
@@ -251,11 +359,16 @@ namespace uba
 				if (!objectFile->WriteImportsAndExports(logger, exportsFile.data))
 					return false;
 
-				//u32 keptExportCount = 0;
-				//StringBuffer<> newFilename;
-				//newFilename.Append(fileName, lastDot - fileName).Append(TC(".TEST")).Append(lastDot);
-				//if (!objectFile->CreateStripped(logger, newFilename.data, {}, keptExportCount))
-				//	return false;
+#if 0
+				objectFile->StripExports(logger);
+				StringBuffer<> newFilename;
+				newFilename.Append(fileName, lastDot - fileName).Append(TC(".TEST")).Append(lastDot);
+				FileAccessor fa(logger, newFilename.data);
+				fa.CreateWrite();
+				fa.Write(objectFile->GetData(), objectFile->GetDataSize());
+				fa.Close();
+
+#endif
 			}
 		}
 		return 0;
