[UBA]
* Added support for reading symbols out of import libraries
* Fixed so objfiletool supports compressed obj files
* Minor cleanup in FileAccessor

diff --git a/Engine/Source/Programs/UnrealBuildAccelerator/Cli/Private/UbaCli.cpp b/Engine/Source/Programs/UnrealBuildAccelerator/Cli/Private/UbaCli.cpp
--- a/Engine/Source/Programs/UnrealBuildAccelerator/Cli/Private/UbaCli.cpp
+++ b/Engine/Source/Programs/UnrealBuildAccelerator/Cli/Private/UbaCli.cpp
@@ -100,6 +100,9 @@ namespace uba
 		logger.Info(TC("   -pool=<name>            Name of helper pool inside coordinator"));
 		logger.Info(TC("   -oidc=<name>            Name of oidc"));
 		logger.Info(TC("   -maxcores=<number>      Max number of cores that will be asked for from coordinator"));
+		logger.Info(TC(""));
+		logger.Info(TC("  If <executable> is a .yaml-file UbaCli creates a scheduler to execute commands from the yaml file instead"));
+		
 		logger.Info(TC(""));
 		return false;
 	}
diff --git a/Engine/Source/Programs/UnrealBuildAccelerator/Common/Private/UbaFileAccessor.cpp b/Engine/Source/Programs/UnrealBuildAccelerator/Common/Private/UbaFileAccessor.cpp
--- a/Engine/Source/Programs/UnrealBuildAccelerator/Common/Private/UbaFileAccessor.cpp
+++ b/Engine/Source/Programs/UnrealBuildAccelerator/Common/Private/UbaFileAccessor.cpp
@@ -346,7 +346,7 @@ namespace uba
 			m_mappingHandle = uba::CreateFileMappingW(InvalidFileHandle, PAGE_READONLY, 1, m_fileName);
 
 		if (!m_mappingHandle.IsValid())
-			return m_logger.Error(TC("Failed to create mapping handle for %s (%s)"), m_fileName, LastErrorToText().data);
+			return m_logger.Error(TC("Failed to create mapping handle for %s with size %llu (%s)"), m_fileName, m_size, LastErrorToText().data);
 #else
 		m_mappingHandle = { asFileDescriptor(m_fileHandle) };
 		if (offset == m_size)
@@ -380,69 +380,74 @@ namespace uba
 			m_mappingHandle = {};
 		}
 
-		if (m_fileHandle != InvalidFileHandle)
-		{
-			const tchar* realFileName = m_fileName;
-			StringBuffer<> tempFile;
+		if (m_fileHandle == InvalidFileHandle)
+			return true;
+		const tchar* realFileName = m_fileName;
+		StringBuffer<> tempFile;
 
-			if (m_isWrite)
+		auto closeFile = MakeGuard([&]()
 			{
-				#if !PLATFORM_WINDOWS
-				if (m_tempPath)
-					realFileName = tempFile.Append(m_tempPath).Append("Temp_").AppendValue(m_tempFileIndex).data;
-				#endif
+				if (!CloseFile(realFileName, m_fileHandle))
+					return m_logger.Error(TC("Failed to close file %s (%s)"), realFileName, LastErrorToText().data);
+				m_fileHandle = InvalidFileHandle;
+				return true;
+			});
 
-				if (success)
-				{
-					#if PLATFORM_WINDOWS
-					if (!SetDeleteOnClose(m_logger, realFileName, m_fileHandle, false))
-						return m_logger.Error(TC("Failed to remove delete on close for file %s (%s)"), realFileName, LastErrorToText().data);
-					#else
-					if (m_tempPath && rename(realFileName, m_fileName) == -1)
-					{
-						if (errno != EXDEV)
-							return m_logger.Error(TC("Failed to rename temporary file %s to %s (%s)"), realFileName, m_fileName, strerror(errno));
-
-						// Need to copy, can't rename over devices
-						int targetFd = open(m_fileName, O_CREAT | O_TRUNC | O_WRONLY | O_CLOEXEC, m_flagsAndAttributes);
-						auto g = MakeGuard([targetFd]() { close(targetFd); });
-						if (targetFd == -1)
-							return m_logger.Error(TC("Failed to create file %s for move from temporary file %s (%s)"), m_fileName, realFileName, strerror(errno));
+
+		if (!m_isWrite)
+			return closeFile.Execute();
+
+		#if !PLATFORM_WINDOWS
+		if (m_tempPath)
+			realFileName = tempFile.Append(m_tempPath).Append("Temp_").AppendValue(m_tempFileIndex).data;
+		#endif
+
+		if (success)
+		{
+			#if PLATFORM_WINDOWS
+			if (!SetDeleteOnClose(m_logger, realFileName, m_fileHandle, false))
+				return m_logger.Error(TC("Failed to remove delete on close for file %s (%s)"), realFileName, LastErrorToText().data);
+			#else
+			if (m_tempPath && rename(realFileName, m_fileName) == -1)
+			{
+				if (errno != EXDEV)
+					return m_logger.Error(TC("Failed to rename temporary file %s to %s (%s)"), realFileName, m_fileName, strerror(errno));
+
+				// Need to copy, can't rename over devices
+				int targetFd = open(m_fileName, O_CREAT | O_TRUNC | O_WRONLY | O_CLOEXEC, m_flagsAndAttributes);
+				auto g = MakeGuard([targetFd]() { close(targetFd); });
+				if (targetFd == -1)
+					return m_logger.Error(TC("Failed to create file %s for move from temporary file %s (%s)"), m_fileName, realFileName, strerror(errno));
 						
-						#if PLATFORM_MAC
-						if (fcopyfile(asFileDescriptor(m_fileHandle), targetFd, 0, COPYFILE_ALL) == -1)
-							return m_logger.Error(TC("Failed to do fcopyfile from temporary %s to file %s (%s)"), realFileName, m_fileName, strerror(errno));
-						#else
-						int sourceFd = asFileDescriptor(m_fileHandle);
-						if (lseek(sourceFd, 0, SEEK_SET) == -1)
-							return m_logger.Error(TC("Failed to do lseek to beginning for sendfile (%s)"), strerror(errno));
-						if (sendfile(targetFd, sourceFd, NULL, m_size) != m_size)
-							return m_logger.Error(TC("Failed to do sendfile from temporary %s to file %s (%s)"), realFileName, m_fileName, strerror(errno));
-						#endif
-
-						remove(realFileName); // Remove real file now when we have copied it over
-					}
-					#endif
+				#if PLATFORM_MAC
+				if (fcopyfile(asFileDescriptor(m_fileHandle), targetFd, 0, COPYFILE_ALL) == -1)
+					return m_logger.Error(TC("Failed to do fcopyfile from temporary %s to file %s (%s)"), realFileName, m_fileName, strerror(errno));
+				#else
+				int sourceFd = asFileDescriptor(m_fileHandle);
+				if (lseek(sourceFd, 0, SEEK_SET) == -1)
+					return m_logger.Error(TC("Failed to do lseek to beginning for sendfile (%s)"), strerror(errno));
+				if (sendfile(targetFd, sourceFd, NULL, m_size) != m_size)
+					return m_logger.Error(TC("Failed to do sendfile from temporary %s to file %s (%s)"), realFileName, m_fileName, strerror(errno));
+				#endif
 
-					if (lastWriteTime)
-					{
-						*lastWriteTime = 0;
-						if (!GetFileLastWriteTime(*lastWriteTime, m_fileHandle))
-							m_logger.Warning(TC("Failed to get file time for %s (%s)"), realFileName, LastErrorToText().data);
-					}
-				}
-				else
-				{
-					#if !PLATFORM_WINDOWS
-					if (m_tempPath && remove(realFileName) == -1)
-						return m_logger.Error(TC("Failed to remove temporary file %s (%s)"), realFileName, strerror(errno));
-					#endif
-				}
+				remove(realFileName); // Remove real file now when we have copied it over
+			}
+			#endif
+
+			if (lastWriteTime)
+			{
+				*lastWriteTime = 0;
+				if (!GetFileLastWriteTime(*lastWriteTime, m_fileHandle))
+					m_logger.Warning(TC("Failed to get file time for %s (%s)"), realFileName, LastErrorToText().data);
 			}
-			if (!CloseFile(realFileName, m_fileHandle))
-				return m_logger.Error(TC("Failed to close file %s (%s)"), realFileName, LastErrorToText().data);
-			m_fileHandle = InvalidFileHandle;
 		}
-		return true;
+		else
+		{
+			#if !PLATFORM_WINDOWS
+			if (m_tempPath && remove(realFileName) == -1)
+				return m_logger.Error(TC("Failed to remove temporary file %s (%s)"), realFileName, strerror(errno));
+			#endif
+		}
+		return closeFile.Execute();
 	}
 }
diff --git a/Engine/Source/Programs/UnrealBuildAccelerator/Common/Private/UbaObjectFile.cpp b/Engine/Source/Programs/UnrealBuildAccelerator/Common/Private/UbaObjectFile.cpp
--- a/Engine/Source/Programs/UnrealBuildAccelerator/Common/Private/UbaObjectFile.cpp
+++ b/Engine/Source/Programs/UnrealBuildAccelerator/Common/Private/UbaObjectFile.cpp
@@ -2,10 +2,13 @@
 
 #include "UbaObjectFile.h"
 #include "UbaBinaryReaderWriter.h"
+#include "UbaCompressedObjFileHeader.h"
 #include "UbaFileAccessor.h"
 #include "UbaObjectFileCoff.h"
 #include "UbaObjectFileElf.h"
+#include "UbaObjectFileImportLib.h"
 #include "UbaObjectFileLLVMIR.h"
+#include <oodle2.h>
 
 namespace uba
 {
@@ -32,24 +35,70 @@ namespace uba
 	{
 		ObjectFile* objectFile = nullptr;
 
+		bool ownsData = false;
+		if (dataSize >= sizeof(CompressedObjFileHeader) && ((CompressedObjFileHeader*)data)->IsValid())
+		{
+			u64 decompressedSize = *(u64*)(data + sizeof(CompressedObjFileHeader));
+			u8* readPos = data + sizeof(CompressedObjFileHeader) + 8;
+
+			u8* decompressedData = (u8*)malloc(decompressedSize);
+			u8* writePos = decompressedData;
+
+			OO_SINTa decoredMemSize = OodleLZDecoder_MemorySizeNeeded(OodleLZ_Compressor_Kraken);
+			void* decoderMem = malloc(decoredMemSize);
+			auto mg = MakeGuard([decoderMem]() { free(decoderMem); });
+
+			u64 left = decompressedSize;
+			while (left)
+			{
+				u32 compressedBlockSize = *(u32*)readPos;
+				readPos += 4;
+				u32 decompressedBlockSize = *(u32*)readPos;
+				readPos += 4;
+
+				OO_SINTa decompLen = OodleLZ_Decompress(readPos, (OO_SINTa)compressedBlockSize, writePos, (OO_SINTa)decompressedBlockSize,
+					OodleLZ_FuzzSafe_Yes, OodleLZ_CheckCRC_No, OodleLZ_Verbosity_None, NULL, 0, NULL, NULL, decoderMem, decoredMemSize);
+				if (decompLen != decompressedBlockSize)
+				{
+					logger.Error(TC("Failed to decompress file %s"), hint);
+					return nullptr;
+				}
+
+				readPos += compressedBlockSize;
+				writePos += decompressedBlockSize;
+				left -= decompressedBlockSize;
+			}
+
+			data = decompressedData;
+			dataSize = decompressedSize;
+			ownsData = true;
+		}
+
 		if (IsElfFile(data, dataSize))
 			objectFile = new ObjectFileElf();
 		else if (IsLLVMIRFile(data, dataSize))
 			objectFile = new ObjectFileLLVMIR();
 		else if (IsCoffFile(data, dataSize))
 			objectFile = new ObjectFileCoff();
+		else if (IsImportLib(data, dataSize))
+			objectFile = new ObjectFileImportLib();
 		else
 		{
+			if (ownsData)
+				free(data);
 			logger.Error(TC("Unknown object file format. Maybe msvc FE IL? (%s)"), hint);
 			return nullptr;
 		}
 
 		objectFile->m_data = data;
 		objectFile->m_dataSize = dataSize;
+		objectFile->m_ownsData = ownsData;
 
 		if (objectFile->Parse(logger, hint))
 			return objectFile;
 
+		if (ownsData)
+			free(data);
 		delete objectFile;
 		return nullptr;
 	}
@@ -58,6 +107,8 @@ namespace uba
 	{
 		u8* data = (u8*)malloc(m_dataSize);
 		memcpy(data, m_data, m_dataSize);
+		if (m_ownsData)
+			free(m_data);
 		m_data = data;
 		m_ownsData = true;
 		delete m_file;
@@ -129,6 +180,12 @@ namespace uba
 		return exportsFile.Close();
 	}
 
+	const char* ObjectFile::GetLibName()
+	{
+		UBA_ASSERT(false);
+		return "";
+	}
+
 	ObjectFile::~ObjectFile()
 	{
 		if (m_ownsData)
@@ -136,6 +193,11 @@ namespace uba
 		delete m_file;
 	}
 
+	void ObjectFile::RemoveExportedSymbol(const char* symbol)
+	{
+		m_exports.erase(symbol);
+	}
+
 	const tchar* ObjectFile::GetFileName() const
 	{
 		return m_file->GetFileName();
diff --git a/Engine/Source/Programs/UnrealBuildAccelerator/Common/Private/UbaObjectFileCoff.cpp b/Engine/Source/Programs/UnrealBuildAccelerator/Common/Private/UbaObjectFileCoff.cpp
--- a/Engine/Source/Programs/UnrealBuildAccelerator/Common/Private/UbaObjectFileCoff.cpp
+++ b/Engine/Source/Programs/UnrealBuildAccelerator/Common/Private/UbaObjectFileCoff.cpp
@@ -228,7 +228,7 @@ namespace uba
 		m_type = ObjectFileType_Coff;
 	}
 
-	bool ObjectFileCoff::Parse(Logger& logger, const tchar* filename)
+	bool ObjectFileCoff::Parse(Logger& logger, const tchar* hint)
 	{
 		m_isBigObj = IsBigObj(m_data, m_dataSize);
 		
diff --git a/Engine/Source/Programs/UnrealBuildAccelerator/Common/Private/UbaObjectFileCoff.h b/Engine/Source/Programs/UnrealBuildAccelerator/Common/Private/UbaObjectFileCoff.h
--- a/Engine/Source/Programs/UnrealBuildAccelerator/Common/Private/UbaObjectFileCoff.h
+++ b/Engine/Source/Programs/UnrealBuildAccelerator/Common/Private/UbaObjectFileCoff.h
@@ -11,7 +11,7 @@ namespace uba
 	{
 	public:
 		ObjectFileCoff();
-		virtual bool Parse(Logger& logger, const tchar* filename) override;
+		virtual bool Parse(Logger& logger, const tchar* hint) override;
 
 		static bool CreateExtraFile2(Logger& logger, MemoryBlock& memoryBlock, const UnorderedSymbols& allNeededImports, const UnorderedSymbols& allSharedImports, const UnorderedExports& allSharedExports, bool includeExportsInFile);
 
diff --git a/Engine/Source/Programs/UnrealBuildAccelerator/Common/Private/UbaObjectFileElf.cpp b/Engine/Source/Programs/UnrealBuildAccelerator/Common/Private/UbaObjectFileElf.cpp
--- a/Engine/Source/Programs/UnrealBuildAccelerator/Common/Private/UbaObjectFileElf.cpp
+++ b/Engine/Source/Programs/UnrealBuildAccelerator/Common/Private/UbaObjectFileElf.cpp
@@ -125,7 +125,7 @@ namespace uba
 	}
 
 
-	bool ObjectFileElf::Parse(Logger& logger, const tchar* filename)
+	bool ObjectFileElf::Parse(Logger& logger, const tchar* hint)
 	{
 		auto& header = *(Elf64Header*)m_data;
 		if (header.e_ident[4] != 2) // 64-bit
diff --git a/Engine/Source/Programs/UnrealBuildAccelerator/Common/Private/UbaObjectFileElf.h b/Engine/Source/Programs/UnrealBuildAccelerator/Common/Private/UbaObjectFileElf.h
--- a/Engine/Source/Programs/UnrealBuildAccelerator/Common/Private/UbaObjectFileElf.h
+++ b/Engine/Source/Programs/UnrealBuildAccelerator/Common/Private/UbaObjectFileElf.h
@@ -10,7 +10,7 @@ namespace uba
 	{
 	public:
 		ObjectFileElf();
-		virtual bool Parse(Logger& logger, const tchar* filename) override;
+		virtual bool Parse(Logger& logger, const tchar* hint) override;
 
 	private:
 		virtual bool StripExports(Logger& logger, u8* newData, const UnorderedSymbols& allNeededImports) override;
diff --git a/Engine/Source/Programs/UnrealBuildAccelerator/Common/Private/UbaObjectFileImportLib.cpp b/Engine/Source/Programs/UnrealBuildAccelerator/Common/Private/UbaObjectFileImportLib.cpp
new file mode 100644
--- /dev/null
+++ b/Engine/Source/Programs/UnrealBuildAccelerator/Common/Private/UbaObjectFileImportLib.cpp
@@ -0,0 +1,75 @@
+// Copyright Epic Games, Inc. All Rights Reserved.
+
+#include "UbaObjectFileImportLib.h"
+
+namespace uba
+{
+	bool IsImportLib(const u8* data, u64 dataSize)
+	{
+		return dataSize > 5 && memcmp(data, "!<arch>", 5) == 0;
+	}
+
+	bool ObjectFileImportLib::Parse(Logger& logger, const tchar* hint)
+	{
+#if PLATFORM_WINDOWS
+
+		u8* pos = m_data;
+		pos += IMAGE_ARCHIVE_START_SIZE;
+		//auto& header = *(IMAGE_ARCHIVE_MEMBER_HEADER*)pos;
+		pos += sizeof(IMAGE_ARCHIVE_MEMBER_HEADER);
+		
+		u32 symbolCount = _byteswap_ulong(*(u32*)pos);
+		pos += sizeof(u32);
+
+		u32* symbolOffsets = (u32*)pos;
+		pos += sizeof(u32)*symbolCount;
+
+		Vector<char*> impSymbols;
+		for (u32 i=0; i!=symbolCount; ++i)
+		{
+			u32 symbolOffset = _byteswap_ulong(symbolOffsets[i]);(void)symbolOffset;
+			auto symbolStr = (char*)pos;
+			std::string symbol(symbolStr);
+			pos += symbol.size() + 1;
+			if (i == 0)
+			{
+				m_libName = symbol.data() + strlen("__IMPORT_DESCRIPTOR_");
+			}
+			if (i < 3) // Skip the predefined symbols
+				continue;
+			if (strncmp(symbol.data(), "__imp_", 6) == 0)
+			{
+				impSymbols.push_back(symbolStr + 6);
+				continue;
+			}
+			m_exports.try_emplace(std::move(symbol), ExportInfo{ ",DATA", i });
+		}
+
+		for (auto symbol : impSymbols)
+		{
+			auto findIt = m_exports.find(symbol);
+			if (findIt != m_exports.end())
+				findIt->second.extra = std::string();
+		}
+
+		return true;
+#else
+		return false;
+#endif
+	}
+
+	const char* ObjectFileImportLib::GetLibName()
+	{
+		return m_libName.c_str();
+	}
+
+	bool ObjectFileImportLib::StripExports(Logger& logger, u8* newData, const UnorderedSymbols& allNeededImports)
+	{
+		return logger.Error(TC("Stripping exports from import lib file not supported"));
+	}
+
+	bool ObjectFileImportLib::CreateExtraFile(Logger& logger, MemoryBlock& memoryBlock, const UnorderedSymbols& allNeededImports, const UnorderedSymbols& allSharedImports, const UnorderedExports& allSharedExports, bool includeExportsInFile)
+	{
+		return logger.Error(TC("Creating extra file from import lib file not supported"));
+	}
+}
diff --git a/Engine/Source/Programs/UnrealBuildAccelerator/Common/Private/UbaObjectFileImportLib.h b/Engine/Source/Programs/UnrealBuildAccelerator/Common/Private/UbaObjectFileImportLib.h
new file mode 100644
--- /dev/null
+++ b/Engine/Source/Programs/UnrealBuildAccelerator/Common/Private/UbaObjectFileImportLib.h
@@ -0,0 +1,24 @@
+// Copyright Epic Games, Inc. All Rights Reserved.
+
+#pragma once
+
+#include "UbaObjectFile.h"
+
+namespace uba
+{
+	class ObjectFileImportLib : public ObjectFile
+	{
+	public:
+		virtual bool Parse(Logger& logger, const tchar* hint) override;
+		virtual const char* GetLibName() override;
+
+	private:
+		virtual bool StripExports(Logger& logger, u8* newData, const UnorderedSymbols& allNeededImports) override;
+		virtual bool CreateExtraFile(Logger& logger, MemoryBlock& memoryBlock, const UnorderedSymbols& allNeededImports, const UnorderedSymbols& allSharedImports, const UnorderedExports& allSharedExports, bool includeExportsInFile) override;
+
+		std::string m_libName;
+	};
+
+
+	bool IsImportLib(const u8* data, u64 dataSize);
+}
diff --git a/Engine/Source/Programs/UnrealBuildAccelerator/Common/Private/UbaObjectFileLLVMIR.cpp b/Engine/Source/Programs/UnrealBuildAccelerator/Common/Private/UbaObjectFileLLVMIR.cpp
--- a/Engine/Source/Programs/UnrealBuildAccelerator/Common/Private/UbaObjectFileLLVMIR.cpp
+++ b/Engine/Source/Programs/UnrealBuildAccelerator/Common/Private/UbaObjectFileLLVMIR.cpp
@@ -804,7 +804,7 @@ namespace uba
 		return ~0u;
 	}
 
-	bool ObjectFileLLVMIR::Parse(Logger& logger, const tchar* filename)
+	bool ObjectFileLLVMIR::Parse(Logger& logger, const tchar* hint)
 	{
 		//BinaryReader reader(m_data, 0, m_dataSize);
 		BitStreamReader reader(*this, logger, m_data, m_dataSize);
diff --git a/Engine/Source/Programs/UnrealBuildAccelerator/Common/Private/UbaObjectFileLLVMIR.h b/Engine/Source/Programs/UnrealBuildAccelerator/Common/Private/UbaObjectFileLLVMIR.h
--- a/Engine/Source/Programs/UnrealBuildAccelerator/Common/Private/UbaObjectFileLLVMIR.h
+++ b/Engine/Source/Programs/UnrealBuildAccelerator/Common/Private/UbaObjectFileLLVMIR.h
@@ -12,7 +12,7 @@ namespace uba
 		ObjectFileLLVMIR();
 		~ObjectFileLLVMIR();
 
-		virtual bool Parse(Logger& logger, const tchar* filename) override;
+		virtual bool Parse(Logger& logger, const tchar* hint) override;
 
 	private:
 		virtual bool StripExports(Logger& logger, u8* newData, const UnorderedSymbols& allNeededImports) override;
diff --git a/Engine/Source/Programs/UnrealBuildAccelerator/Common/Private/UbaSession.cpp b/Engine/Source/Programs/UnrealBuildAccelerator/Common/Private/UbaSession.cpp
--- a/Engine/Source/Programs/UnrealBuildAccelerator/Common/Private/UbaSession.cpp
+++ b/Engine/Source/Programs/UnrealBuildAccelerator/Common/Private/UbaSession.cpp
@@ -2204,7 +2204,8 @@ namespace uba
 
 					if (!m_storage.WriteCompressed(res, TC("MemoryMap"), InvalidFileHandle, mem, fileSize, file.name.c_str(), &header, sizeof(header), file.lastWriteTime))
 						return false;
-					shouldEvictFromMemory = true;
+
+					shouldEvictFromMemory = false; // Can't evict without properly update filemappingtable.. the file on disk does now not match what was registered for write
 				}
 				else
 				{
diff --git a/Engine/Source/Programs/UnrealBuildAccelerator/Common/Public/UbaObjectFile.h b/Engine/Source/Programs/UnrealBuildAccelerator/Common/Public/UbaObjectFile.h
--- a/Engine/Source/Programs/UnrealBuildAccelerator/Common/Public/UbaObjectFile.h
+++ b/Engine/Source/Programs/UnrealBuildAccelerator/Common/Public/UbaObjectFile.h
@@ -25,13 +25,14 @@ namespace uba
 	class ObjectFile
 	{
 	public:
-		static ObjectFile* OpenAndParse(Logger& logger, const tchar* filename);
+		static ObjectFile* OpenAndParse(Logger& logger, const tchar* hint);
 		static ObjectFile* Parse(Logger& logger, u8* data, u64 dataSize, const tchar* hint);
 
 		virtual bool CopyMemoryAndClose();
 		virtual bool StripExports(Logger& logger);
 		virtual bool WriteImportsAndExports(Logger& logger, MemoryBlock& memoryBlock);
 		virtual bool WriteImportsAndExports(Logger& logger, const tchar* exportsFilename);
+		virtual const char* GetLibName();
 
 		const tchar* GetFileName() const;
 		const UnorderedSymbols& GetImports() const;
@@ -42,11 +43,12 @@ namespace uba
 
 		virtual ~ObjectFile();
 
+		void RemoveExportedSymbol(const char* symbol);
 		u8* GetData() { return m_data; }
 		u64 GetDataSize() { return m_dataSize; }
 
 	protected:
-		virtual bool Parse(Logger& logger, const tchar* filename) = 0;
+		virtual bool Parse(Logger& logger, const tchar* hint) = 0;
 		virtual bool StripExports(Logger& logger, u8* newData, const UnorderedSymbols& allNeededImports) = 0;
 		virtual bool CreateExtraFile(Logger& logger, MemoryBlock& memoryBlock, const UnorderedSymbols& allNeededImports, const UnorderedSymbols& allSharedImports, const UnorderedExports& allSharedExports, bool includeExportsInFile) = 0;
 
diff --git a/Engine/Source/Programs/UnrealBuildAccelerator/ObjTool/Private/UbaObjTool.cpp b/Engine/Source/Programs/UnrealBuildAccelerator/ObjTool/Private/UbaObjTool.cpp
--- a/Engine/Source/Programs/UnrealBuildAccelerator/ObjTool/Private/UbaObjTool.cpp
+++ b/Engine/Source/Programs/UnrealBuildAccelerator/ObjTool/Private/UbaObjTool.cpp
@@ -33,11 +33,12 @@ namespace uba
 		logger.Info(TC("   UbaObjTool v%s%s"), Version, dbgStr);
 		logger.Info(TC("-------------------------------------------"));
 		logger.Info(TC(""));
-		logger.Info(TC("  UbaObjTool.exe [options...] <objfile>"));
+		logger.Info(TC("  UbaObjTool.exe [options...] <objfile/libfile>"));
 		logger.Info(TC(""));
 		logger.Info(TC("   Options:"));
 		logger.Info(TC("    -printsymbols            Print the symbols found in obj file"));
 		logger.Info(TC("    -stripexports            Will strip exports and write them out in a .exp file"));
+		logger.Info(TC("    -writeimplib=<file>      Will create a import library from symbols collected from obj/lib files"));
 		logger.Info(TC(""));
 		logger.Info(TC("  --- OR ---"));
 		logger.Info(TC(""));
@@ -69,6 +70,7 @@ namespace uba
 		bool printSymbols = false;
 		bool stripExports = false;
 		bool writeImpLib = false;
+		bool allowLibInputs = false;
 		bool isImpLibRsp = false;
 
 		Vector<TString> objFilesToStrip;
@@ -173,6 +175,9 @@ namespace uba
 				}
 				else if (name.Equals(TC("-writeimplib")))
 				{
+					impLibFile = value.data;
+					writeImpLib = true;
+					allowLibInputs = true;
 				}
 				else if (name.Equals(TC("-stripexports")))
 				{
@@ -314,22 +319,28 @@ namespace uba
 				auto& o = *it;
 				StringBuffer<> fixedPath;
 				FixPath(o.c_str(), currentDir.data, currentDir.count, fixedPath);
-				if (fixedPath.EndsWith(TC(".res")) || fixedPath.EndsWith(TC(".lib")))
+				if (fixedPath.EndsWith(TC(".res")) || (fixedPath.EndsWith(TC(".lib")) && !allowLibInputs))
 					return;
 				if (ObjectFile* objectFile = ObjectFile::OpenAndParse(logger, fixedPath.data))
+				{
+					objectFile->RemoveExportedSymbol("DllMain");
 					objFiles[it - objFilesForImpLib.begin()] = objectFile;
+				}
 				else
 					success = false;
 			});
 			if (!success)
 				return -1;
 
+			if (impLibName.empty() && objFiles.size() == 1)
+				impLibName = objFiles[0]->GetLibName();
+
 			writer.Write(logger, objFiles, impLibName.c_str(), impLibFile.c_str());
 		}
 		else
 		{
 			if (objFile.empty())
-				return PrintHelp(TC("No obj or rsp file provided"));
+				return PrintHelp(TC("No obj, lib or rsp file provided"));
 
 			ObjectFile* objectFile = ObjectFile::OpenAndParse(logger, objFile.c_str());
 			if (!objectFile)
