[UBA]
* Attempt at fixing macos compile error

diff --git a/Engine/Source/Programs/UnrealBuildAccelerator/Common/Private/UbaObjectFileLLVMIR.cpp b/Engine/Source/Programs/UnrealBuildAccelerator/Common/Private/UbaObjectFileLLVMIR.cpp
--- a/Engine/Source/Programs/UnrealBuildAccelerator/Common/Private/UbaObjectFileLLVMIR.cpp
+++ b/Engine/Source/Programs/UnrealBuildAccelerator/Common/Private/UbaObjectFileLLVMIR.cpp
@@ -358,8 +358,9 @@ namespace uba
 
 		void EnterSubBlock(u32& outNumWords, u32 blockId)
 		{
-			m_blockScope.push_back(Block(m_currentCodeSize));
-			m_blockScope.back().prevAbbrevs.swap(m_curAbbrevs);
+			auto& block = m_blockScope.emplace_back();
+			block.prevCodeSize = m_currentCodeSize;
+			block.prevAbbrevs.swap(m_curAbbrevs);
 
 			if (m_blockInfo)
 				if (const BlockInfo::Record* record = m_blockInfo->GetBlockInfo(blockId))
@@ -770,7 +771,6 @@ namespace uba
 		u8* m_pos;
 		u8* m_end;
 
-
 		u32 m_currentCodeSize = 2;
 
 		struct Block
