[UBA]
* Fixed another macos compile error

diff --git a/Engine/Source/Programs/UnrealBuildAccelerator/Common/Private/UbaObjectFile.cpp b/Engine/Source/Programs/UnrealBuildAccelerator/Common/Private/UbaObjectFile.cpp
--- a/Engine/Source/Programs/UnrealBuildAccelerator/Common/Private/UbaObjectFile.cpp
+++ b/Engine/Source/Programs/UnrealBuildAccelerator/Common/Private/UbaObjectFile.cpp
@@ -280,13 +280,13 @@ namespace uba
 		while (*readPos)
 		{
 			auto strEnd = strlen(readPos);
-			std::string extra;
+			ExportInfo info;
 			if (const char* comma = strchr(readPos, ','))
 			{
 				strEnd = comma - readPos;
-				extra = comma;
+				info.extra = comma;
 			}
-			exports.emplace(std::string(readPos, readPos + strEnd), extra);
+			exports.emplace(std::string(readPos, readPos + strEnd), info);
 			readPos = readPos + strEnd + 1;
 		}
 		return true;
