[UBA]
* Attempt at optimizing cache server maintenance logic by changing how it figures out which cas key offsets that are deleted

diff --git a/Engine/Source/Programs/UnrealBuildAccelerator/Common/Private/UbaCacheServer.cpp b/Engine/Source/Programs/UnrealBuildAccelerator/Common/Private/UbaCacheServer.cpp
--- a/Engine/Source/Programs/UnrealBuildAccelerator/Common/Private/UbaCacheServer.cpp
+++ b/Engine/Source/Programs/UnrealBuildAccelerator/Common/Private/UbaCacheServer.cpp
@@ -23,47 +23,9 @@
 
 	bool IsCaseInsensitive(u64 id) { return (id & (1ull << 32)) == 0; }
 
-	struct CacheServer::ConnectionBucket
-	{
-		ConnectionBucket(u64 i) : pathTable(CachePathTableMaxSize, CompactPathTable::V1, IsCaseInsensitive(i)), casKeyTable(CacheCasKeyTableMaxSize), id(i) {}
-		CompactPathTable pathTable;
-		CompactCasKeyTable casKeyTable;
-
-		ReaderWriterLock cacheEntryLookupLock;
-		UnorderedMap<CasKey, CacheEntry> cacheEntryLookup;
-		u64 id;
-	};
-
-	struct CacheServer::Connection
-	{
-		u32 clientVersion;
-		UnorderedMap<u64, ConnectionBucket> buckets;
-	};
-
-	struct CacheServer::Bucket
-	{
-		Bucket(u64 id) : m_pathTable(CachePathTableMaxSize, CompactPathTable::V1, IsCaseInsensitive(id)), m_casKeyTable(CacheCasKeyTableMaxSize) {}
-		ReaderWriterLock m_cacheEntryLookupLock;
-		UnorderedMap<CasKey, CacheEntries> m_cacheEntryLookup;
-
-		CompactPathTable m_pathTable;
-		CompactCasKeyTable m_casKeyTable;
-
-		Atomic<u64> totalEntryCount;
-		Atomic<u64> totalEntrySize;
-		Atomic<bool> hasDeletedEntries;
-		Atomic<bool> needsSave;
-
-		Atomic<u64> lastSavedTime;
-		Atomic<u64> lastUsedTime;
-
-		u64 expirationTimeSeconds = 0;
-		u32 index = ~0u;
-	};
-
 	struct BitArray
 	{
-		BitArray(MemoryBlock& memoryBlock, u32 bitCount)
+		void Init(MemoryBlock& memoryBlock, u32 bitCount)
 		{
 			u32 bytes = AlignUp((bitCount+7) / 8, 8u); // Align up to 64 bits
 			data = (u64*)memoryBlock.Allocate(bytes, 8, TC(""));
@@ -79,6 +41,14 @@
 			data[index] |= 1ull << bitOffset;
 		}
 
+		UBA_FORCEINLINE bool IsSet(u32 bitIndex)
+		{
+			u32 index = bitIndex / 64;
+			UBA_ASSERTF(index < count, TC("Out of bounds (%u/%u). Bit index : %u"), index, count, bitIndex);
+			u32 bitOffset = bitIndex - index * 64;
+			return (data[index] & (1ull << bitOffset)) != 0;
+		}
+
 		UBA_FORCEINLINE u32 CountSetBits()
 		{
 			u64 bits = 0;
@@ -139,8 +109,54 @@
 			#endif
 		}
 
-		u64* data;
-		u32 count;
+		u64* data = nullptr;
+		u32 count = 0;
 	};
 
+	struct CacheServer::ConnectionBucket
+	{
+		ConnectionBucket(u64 i) : pathTable(CachePathTableMaxSize, CompactPathTable::V1, IsCaseInsensitive(i)), casKeyTable(CacheCasKeyTableMaxSize), id(i) {}
+		CompactPathTable pathTable;
+		CompactCasKeyTable casKeyTable;
+
+		ReaderWriterLock cacheEntryLookupLock;
+		UnorderedMap<CasKey, CacheEntry> cacheEntryLookup;
+		u64 id;
+	};
+
+	struct CacheServer::Connection
+	{
+		u32 clientVersion;
+		UnorderedMap<u64, ConnectionBucket> buckets;
+	};
+
+	struct CacheServer::Bucket
+	{
+		Bucket(u64 id) : m_pathTable(CachePathTableMaxSize, CompactPathTable::V1, IsCaseInsensitive(id)), m_casKeyTable(CacheCasKeyTableMaxSize) {}
+		ReaderWriterLock m_cacheEntryLookupLock;
+		UnorderedMap<CasKey, CacheEntries> m_cacheEntryLookup;
+
+		CompactPathTable m_pathTable;
+		CompactCasKeyTable m_casKeyTable;
+
+		Atomic<u64> totalEntryCount;
+		Atomic<u64> totalEntrySize;
+		Atomic<bool> hasDeletedEntries;
+		Atomic<bool> needsSave;
+
+		Atomic<u64> lastSavedTime;
+		Atomic<u64> lastUsedTime;
+
+		u64 expirationTimeSeconds = 0;
+		u32 index = ~0u;
+
+		struct MaintenanceContext
+		{
+			MemoryBlock memoryBlock;
+			BitArray deletedOffsets;
+			bool isInitialized = false;
+			bool shouldTest = false;
+		}* m_maintenanceContext;
+	};
+
 	const tchar* ToString(CacheMessageType type)
@@ -532,7 +548,7 @@
 			});
 
 
-		m_forceAllSteps = true;
+		//m_forceAllSteps = true;
 		bool forceAllSteps = m_forceAllSteps;
 		m_forceAllSteps = false;
 
@@ -577,9 +593,8 @@
 
 		u64 totalCasSize = 0;
 
-		struct CasFileInfo { u64 size; u64 refCount; };
+		struct CasFileInfo { CasFileInfo(u64 s = 0) : size(s) {} u64 size; Atomic<bool> isUsed; };
 		UnorderedMap<CasKey, CasFileInfo> existingCas;
-		ReaderWriterLock existingCasLock;
 
 		m_storage.WaitForActiveWork();
 
@@ -604,7 +619,7 @@
 					continue;
 				}
 				totalCasSize += i->second.size;
-				existingCas.try_emplace(i->first, CasFileInfo{i->second.size, 0ull});
+				existingCas.try_emplace(i->first, i->second.size);
 				++i;
 			}
 			lookupLock.Leave();
@@ -620,6 +635,7 @@
 		if (shouldExit())
 			return true;
 
+		ReaderWriterLock globalStatsLock;
 		u64 now = GetSystemTimeAsFileTime();
 		u64 oldest = 0;
 		u64 longestUnused = 0;
@@ -638,6 +654,18 @@
 		Atomic<u64> activeDropCount;
 		auto dropCasGuard = MakeGuard([&]() { while (activeDropCount != 0) Sleep(1); });
 
+		auto EnsureBucketContextInitialized = [&](Bucket& bucket)
+			{
+				auto& context = *bucket.m_maintenanceContext;
+				if (!context.isInitialized)
+				{
+					if (!context.memoryBlock.Init(m_maintenanceReserveSize, nullptr, true)) // Try to use large blocks
+						context.memoryBlock.Init(m_maintenanceReserveSize, nullptr, false);
+					context.deletedOffsets.Init(context.memoryBlock, bucket.m_casKeyTable.GetSize());
+					context.isInitialized = true;
+				}
+			};
+
 		u32 deleteIteration = 0;
 		u64 deleteCacheEntriesStartTime = GetTime();
 		do
@@ -651,6 +679,24 @@
 			m_server.ParallelFor(workerCountToUseForBuckets, m_buckets, [&](auto& it)
 			{
 				Bucket& bucket = it->second;
+				auto context = bucket.m_maintenanceContext;
+				if (!context)
+					context = bucket.m_maintenanceContext = new Bucket::MaintenanceContext;
+
+				bool foundDeletedCasKey = false;
+				for (auto& cas : deletedCasFiles)
+					bucket.m_casKeyTable.TraverseOffsets(cas, [&](u32 casKeyOffset)
+						{
+							EnsureBucketContextInitialized(bucket);
+							foundDeletedCasKey = true;
+							context->deletedOffsets.Set(casKeyOffset);
+						});
+
+				if (!foundDeletedCasKey)
+					checkInputsForDeletes = false;
+
+				auto& deletedOffsets = context->deletedOffsets;
+
 				bucket.totalEntryCount = 0;
 				bucket.totalEntrySize = 0;
 
@@ -671,8 +717,7 @@
 						lastUseTimeLimit = (now - m_creationTime) - GetSecondsAsFileTime(bucket.expirationTimeSeconds);
 				}
 
-
-				m_server.ParallelFor(workerCountToUse, bucket.m_cacheEntryLookup, [&, touchedCas = Vector<u64*>()](auto& li) mutable
+				m_server.ParallelFor(workerCountToUse, bucket.m_cacheEntryLookup, [&, touchedCas = Vector<Atomic<bool>*>()](auto& li) mutable
 				{
 					CacheEntries& entries = li->second;
 
@@ -680,8 +725,6 @@
 					// all the entries that overflow that number
 					u64 capacityLeft = SendMaxSize - 32 - entries.GetSharedSize();
 
-					auto IsOffsetDeleted = [&](u64 offset) { CasKey casKey; bucket.m_casKeyTable.GetKey(casKey, offset); return deletedCasFiles.find(casKey) != deletedCasFiles.end(); };
-
 					// Check if any offset has been deleted in shared offsets..
 					bool offsetDeletedInShared = false;
 					auto& sharedOffsets = entries.sharedInputCasKeyOffsets;
@@ -690,7 +733,7 @@
 						BinaryReader reader2(sharedOffsets.data(), 0, sharedOffsets.size());
 						while (reader2.GetLeft())
 						{
-							if (!IsOffsetDeleted(reader2.Read7BitEncoded()))
+							if (!deletedOffsets.IsSet(u32(reader2.Read7BitEncoded())))
 								continue;
 							offsetDeletedInShared = true;
 							break;
@@ -731,7 +774,7 @@
 									BinaryReader inputReader(sharedOffsets.data() + begin, 0, end - begin);
 									while (inputReader.GetLeft())
 									{
-										if (!IsOffsetDeleted(inputReader.Read7BitEncoded()))
+										if (!deletedOffsets.IsSet(u32(inputReader.Read7BitEncoded())))
 											continue;
 										deleteEntry = true;
 										++missingInputEntryCount;
@@ -746,7 +789,7 @@
 								BinaryReader extraReader(extraInputs.data(), 0, extraInputs.size());
 								while (extraReader.GetLeft())
 								{
-									if (!IsOffsetDeleted(extraReader.Read7BitEncoded()))
+									if (!deletedOffsets.IsSet(u32(extraReader.Read7BitEncoded())))
 										continue;
 									deleteEntry = true;
 									++missingInputEntryCount;
@@ -771,7 +814,7 @@
 								auto findIt = existingCas.find(casKey);
 								if (findIt != existingCas.end())
 								{
-									touchedCas.push_back(&findIt->second.refCount);
+									touchedCas.push_back(&findIt->second.isUsed);
 									continue;
 								}
 								deleteEntry = true;
@@ -796,15 +839,16 @@
 
 						capacityLeft -= neededSize;
 
-						SCOPED_WRITE_LOCK(existingCasLock, l);
-						if (!oldest || entry.creationTime < oldest)
-							oldest = entry.creationTime;
-						if (!longestUnused || entry.lastUsedTime < longestUnused)
-							longestUnused = entry.lastUsedTime;
+						{
+							SCOPED_WRITE_LOCK(globalStatsLock, l);
+							if (!oldest || entry.creationTime < oldest)
+								oldest = entry.creationTime;
+							if (!longestUnused || entry.lastUsedTime < longestUnused)
+								longestUnused = entry.lastUsedTime;
+						}
 
-						for (u64* v : touchedCas)
-							++(*v);
-						l.Leave();
+						for (auto v : touchedCas)
+							*v = true;
 
 						++i;
 					}
@@ -830,9 +874,9 @@
 
 			for (auto i=existingCas.begin(), e=existingCas.end(); i!=e;)
 			{
-				if (i->second.refCount != 0)
+				if (i->second.isUsed)
 				{
-					i->second.refCount = 0;
+					i->second.isUsed = false;
 					++i;
 					continue;
 				}
@@ -874,6 +918,7 @@
 			u64 bucketStartTime = GetTime();
 
 			Bucket& bucket = it->second;
+			auto deleteContext = MakeGuard([&]() { delete bucket.m_maintenanceContext; bucket.m_maintenanceContext = nullptr; });
 
 			if (!bucket.hasDeletedEntries && !forceAllSteps)
 			{
@@ -882,12 +927,11 @@
 			}
 			bucket.hasDeletedEntries = false;
 
-			// Try to use large blocks
-			MemoryBlock memoryBlock;
-			if (!memoryBlock.Init(m_maintenanceReserveSize, nullptr, true))
-				memoryBlock.Init(m_maintenanceReserveSize, nullptr, false);
+			EnsureBucketContextInitialized(bucket);
+			MemoryBlock& memoryBlock = bucket.m_maintenanceContext->memoryBlock;
 
-			BitArray usedCasKeyOffsets(memoryBlock, u32(bucket.m_casKeyTable.GetSize()));
+			BitArray usedCasKeyOffsets;
+			usedCasKeyOffsets.Init(memoryBlock, bucket.m_casKeyTable.GetSize());
 
 			u64 collectUsedCasKeysStart = GetTime();
 
@@ -918,7 +962,8 @@
 			u64 recreatePathTableStart = GetTime();
 
 			// Traverse all caskeys in caskey table and figure out which ones we can delete
-			BitArray usedPathOffsets(memoryBlock, u32(bucket.m_pathTable.GetSize()));
+			BitArray usedPathOffsets;
+			usedPathOffsets.Init(memoryBlock, bucket.m_pathTable.GetSize());
 
 			BinaryReader casKeyTableReader(bucket.m_casKeyTable.GetMemory(), 0, bucket.m_casKeyTable.GetSize());
 			usedCasKeyOffsets.Traverse([&](u32 casKeyOffset)
@@ -1020,7 +1065,7 @@
 
 			m_logger.Info(TC("    Bucket %u Done (%s). CacheEntries: %s (%s) PathTable: %s CasTable: %s Expiration: %s"), bucket.index, TimeToText(GetTime() - bucketStartTime).str, CountToText(bucket.totalEntryCount.load()).str, BytesToText(bucket.totalEntrySize.load()).str, BytesToText(bucket.m_pathTable.GetSize()).str, BytesToText(bucket.m_casKeyTable.GetSize()).str, TimeToText(MsToTime(bucket.expirationTimeSeconds*1000), true).str);
 
-			SCOPED_WRITE_LOCK(existingCasLock, l); // Just reusing lock for other purpose
+			SCOPED_WRITE_LOCK(globalStatsLock, l);
 			maxCommittedMemory = Max(maxCommittedMemory, memoryBlock.writtenSize);
 
 		}, TC(""), true);
diff --git a/Engine/Source/Programs/UnrealBuildAccelerator/Common/Private/UbaCompactTables.cpp b/Engine/Source/Programs/UnrealBuildAccelerator/Common/Private/UbaCompactTables.cpp
--- a/Engine/Source/Programs/UnrealBuildAccelerator/Common/Private/UbaCompactTables.cpp
+++ b/Engine/Source/Programs/UnrealBuildAccelerator/Common/Private/UbaCompactTables.cpp
@@ -255,30 +255,82 @@ namespace uba
 			m_offsets.reserve(reserveOffsetsCount);
 	}
 
+	CompactCasKeyTable::~CompactCasKeyTable()
+	{
+		for (auto& kv : m_offsets)
+			if (kv.second.count > 1)
+				delete[] kv.second.stringAndCasKeyOffsets;
+
+	}
+
 	u32 CompactCasKeyTable::Add(const CasKey& casKey, u64 stringOffset, u32* outRequiredCasTableSize)
 	{
 		SCOPED_WRITE_LOCK(m_lock, lock2)
 		if (!m_mem.memory)
 			m_mem.Init(m_reserveSize);
-		auto insres = m_offsets.try_emplace({ casKey, u32(stringOffset)});
-		if (insres.second)
+
+		bool added = false;
+		u32* casKeyOffset = InternalAdd(casKey, stringOffset, added);
+
+		if (added)
 		{
 			u8 bytesForStringOffset = Get7BitEncodedCount(stringOffset);
 			u8* mem = (u8*)m_mem.AllocateNoLock(bytesForStringOffset + sizeof(CasKey), 1, TC(""));
 			BinaryWriter writer(mem, 0, 1000);
 			writer.Write7BitEncoded(stringOffset);
 			writer.WriteCasKey(casKey);
-			insres.first->second = u32(mem - m_mem.memory);
+			*casKeyOffset = u32(mem - m_mem.memory);
 			if (outRequiredCasTableSize)
 				*outRequiredCasTableSize = (u32)m_mem.writtenSize;
 		}
 		else if (outRequiredCasTableSize)
 		{
-			BinaryReader reader(m_mem.memory, insres.first->second, ~0u);
+			BinaryReader reader(m_mem.memory, *casKeyOffset, ~0u);
 			reader.Read7BitEncoded();
 			*outRequiredCasTableSize = Max(*outRequiredCasTableSize, u32(reader.GetPosition() + sizeof(CasKey)));
 		}
-		return insres.first->second;
+		return *casKeyOffset;
+	}
+
+	u32* CompactCasKeyTable::InternalAdd(const CasKey& casKey, u64 stringOffset, bool& outAdded)
+	{
+		auto insres = m_offsets.try_emplace(casKey);
+		Value& value = insres.first->second;
+		if (insres.second)
+		{
+			value.count = 1;
+			value.single.stringOffset = u32(stringOffset);
+			outAdded = true;
+			return &value.single.casKeyOffset;
+		}
+		
+		if (value.count == 1)
+		{
+			if (value.single.stringOffset == stringOffset)
+				return &value.single.casKeyOffset;
+
+			u32* newOffsets = new u32[4];
+			newOffsets[0] = value.single.stringOffset;
+			newOffsets[1] = value.single.casKeyOffset;
+			newOffsets[2] = u32(stringOffset);
+			value.stringAndCasKeyOffsets = newOffsets;
+			value.count = 2;
+			outAdded = true;
+			return &newOffsets[3];
+		}
+
+		for (u32 i=0, e=value.count*2; i!=e; i+=2)
+			if (value.stringAndCasKeyOffsets[i] == stringOffset)
+				return &value.stringAndCasKeyOffsets[i + 1];
+
+		u32* newOffsets = new u32[(value.count+1)*2];
+		memcpy(newOffsets, value.stringAndCasKeyOffsets, value.count*2*sizeof(u32));
+		newOffsets[value.count*2] = u32(stringOffset);
+		delete[] value.stringAndCasKeyOffsets;
+		value.stringAndCasKeyOffsets = newOffsets;
+		++value.count;
+		outAdded = true;
+		return &newOffsets[value.count*2 - 1];
 	}
 
 	void CompactCasKeyTable::GetKey(CasKey& outKey, u64 offset) const
@@ -334,7 +386,10 @@ namespace uba
 			u32 offset = u32(reader2.GetPosition());
 			u64 stringOffset = reader2.Read7BitEncoded();
 			CasKey casKey = reader2.ReadCasKey();
-			m_offsets.try_emplace({casKey, u32(stringOffset)}, offset);
+			bool added = false;
+			u32* casKeyOffset = InternalAdd(casKey, stringOffset, added);
+			UBA_ASSERT(added);
+			*casKeyOffset = offset;
 		}
 	}
 
diff --git a/Engine/Source/Programs/UnrealBuildAccelerator/Common/Public/UbaCompactTables.h b/Engine/Source/Programs/UnrealBuildAccelerator/Common/Public/UbaCompactTables.h
--- a/Engine/Source/Programs/UnrealBuildAccelerator/Common/Public/UbaCompactTables.h
+++ b/Engine/Source/Programs/UnrealBuildAccelerator/Common/Public/UbaCompactTables.h
@@ -47,9 +47,13 @@ namespace uba
 	{
 	public:
 		CompactCasKeyTable(u64 reserveSize, u64 reserveOffsetsCount = 0);
+		~CompactCasKeyTable();
 
 		u32 Add(const CasKey& casKey, u64 stringOffset, u32* outRequiredCasTableSize = nullptr);
 
+		template<typename Func>
+		void TraverseOffsets(const CasKey& casKey, const Func& func) const;
+
 		void GetKey(CasKey& outKey, u64 offset) const;
 		void GetPathAndKey(StringBufferBase& outPath, CasKey& outKey, const CompactPathTable& pathTable, u64 offset) const;
 
@@ -63,15 +67,40 @@ namespace uba
 		u64 GetKeyCount() { return m_offsets.size(); }
 
 	private:
+		u32* InternalAdd(const CasKey& casKey, u64 stringOffset, bool& outAdded);
+
 		ReaderWriterLock m_lock;
 		MemoryBlock m_mem;
-		struct Key
+		struct Value
 		{
-			bool operator==(const Key& o) const { return ck == o.ck && offset == o.offset; }
-			CasKey ck; u32 offset;
+			union
+			{
+				u32* stringAndCasKeyOffsets;
+				struct
+				{
+					u32 stringOffset;
+					u32 casKeyOffset;
+				} single;
+			};
+			u32 count; // If count is 1, single is set, otherwise stringAndCasKeyOffsets is allocated and contains two offsets per entry
 		};
-		struct KeyHash { size_t operator()(const Key& k) const { return k.ck.a ^ k.ck.b ^ k.ck.c ^ k.offset; } };
-		UnorderedMap<Key, u32, KeyHash> m_offsets;
+
+		UnorderedMap<CasKey, Value> m_offsets;
 		u64 m_reserveSize;
 	};
+
+
+	template<typename Func>
+	void CompactCasKeyTable::TraverseOffsets(const CasKey& casKey, const Func& func) const
+	{
+		auto findIt = m_offsets.find(casKey);
+		if (findIt == m_offsets.end())
+			return;
+		const Value& value = findIt->second;
+		if (value.count == 1)
+			func(value.single.casKeyOffset);
+		else
+			for (u32 i=1, e=value.count*2+1; i!=e; i+=2)
+				func(value.stringAndCasKeyOffsets[i]);
+	}
 }
