[UBA]
* Moved around variables in UbaStorage
* Added log warning to show if something happens that should never happen

diff --git a/Engine/Source/Programs/UnrealBuildAccelerator/Common/Private/UbaStorageServer.cpp b/Engine/Source/Programs/UnrealBuildAccelerator/Common/Private/UbaStorageServer.cpp
--- a/Engine/Source/Programs/UnrealBuildAccelerator/Common/Private/UbaStorageServer.cpp
+++ b/Engine/Source/Programs/UnrealBuildAccelerator/Common/Private/UbaStorageServer.cpp
@@ -137,7 +137,12 @@ namespace uba
 				server.m_casDataBuffer.UnmapView(mappedView, TC("OnDisconnected"));
 		}
 		else
+		{
+			if (!memoryBegin)
+				server.m_logger.Warning(TC("This should not happen. It means there is a race between a fetch and a disconnect. Report to honk (%s)"), reason);
 			server.m_bufferSlots.Push(memoryBegin);
+			memoryBegin = nullptr;
+		}
 	}
 
 	void StorageServer::OnDisconnected(u32 clientId)
@@ -592,7 +597,7 @@ namespace uba
 					if (!ReadFile(m_logger, casFile.data, readFileHandle, writeBuffer, toWrite))
 					{
 						UBA_ASSERT(false); // Implement
-						return false;
+						return m_logger.Error(TC("Failed to read file %s (%s) (1)"), casFile.data, LastErrorToText().data);;
 					}
 				}
 				else
@@ -603,7 +608,7 @@ namespace uba
  					if (!ReadFile(m_logger, casFile.data, readFileHandle, memoryBegin, toRead))
 					{
 						UBA_ASSERT(false); // Implement
-						return false;
+						return m_logger.Error(TC("Failed to read file %s (%s) (2)"), casFile.data, LastErrorToText().data);;
 					}
 					memcpy(writeBuffer, memoryPos, toWrite);
 					memoryPos += toWrite;
