[UBA]
* Fixed so agent connects and reports if it doesn't have the right xcode sdk version

diff --git a/Engine/Source/Programs/UnrealBuildAccelerator/Agent/Private/UbaAgent.cpp b/Engine/Source/Programs/UnrealBuildAccelerator/Agent/Private/UbaAgent.cpp
--- a/Engine/Source/Programs/UnrealBuildAccelerator/Agent/Private/UbaAgent.cpp
+++ b/Engine/Source/Programs/UnrealBuildAccelerator/Agent/Private/UbaAgent.cpp
@@ -112,7 +112,7 @@ namespace uba
 		logger.Info(TC("  -populateCasFromAllXcodes   Prepopulate cas database with files from local xcode installation that matches the version."));
 		#endif
 		logger.Info(TC(""));
-		return -1;
+		return false;
 	}
 
 	ReaderWriterLock* g_exitLock = new ReaderWriterLock();
@@ -248,11 +248,11 @@ namespace uba
 		return success;
 	}
 
-	int LaunchReal(Logger& logger, StringBufferBase& relaunchPath, int argc, tchar* argv[])
+	bool LaunchReal(Logger& logger, StringBufferBase& relaunchPath, int argc, tchar* argv[])
 	{
 		StringBuffer<256> currentDir;
 		if (!GetDirectoryOfCurrentModule(logger, currentDir))
-			return -1;
+			return false;
 		logger.Info(TC("Copying new binaries..."));
 		for (auto file : g_ubaAgentBinaries)
 		{
@@ -263,10 +263,7 @@ namespace uba
 			to.Append('\\').Append(file);
 
 			if (!uba::CopyFileW(from.data, to.data, false))
-			{
-				logger.Error(TC("Failed to copy file for relaunch"));
-				return -1;
-			}
+				return logger.Error(TC("Failed to copy file for relaunch"));
 		}
 
 		StringBuffer<> args;
@@ -278,13 +275,11 @@ namespace uba
 				args.Append(' ').Append(argv[i]);
 		logger.Info(TC("Relaunching new %s..."), UBA_AGENT_EXECUTABLE);
 		logger.Info(TC(""));
-		if (!LaunchProcess(args.data))
-			return -1;
-		return 0;
+		return LaunchProcess(args.data);
 	}
 #endif // UBA_AUTO_UPDATE
 
-	int ExpandEnvironmentVariables(StringBufferBase& str)
+	bool ExpandEnvironmentVariables(StringBufferBase& str)
 	{
 		StringBuffer<> expandedDir;
 		u64 offset = 0;
@@ -316,7 +311,7 @@ namespace uba
 			expandedDir.Append(str.data + offset, beginOffset - offset).Append(value);
 			offset = endOffset + 1;
 		}
-		return 0;
+		return true;
 	}
 
 	bool IsTerminating(Logger& logger, const tchar* eventFile, StringBufferBase& outReason, u64& outTerminationTimeMs)
@@ -394,7 +389,7 @@ namespace uba
 		return true;
 	}
 
-	int WrappedMain(int argc, tchar*argv[])
+	bool WrappedMain(int argc, tchar*argv[])
 	{
 		#if UBA_USE_EXCEPTION_HANDLER
 		SetUnhandledExceptionFilter(UbaUnhandledExceptionFilter);
@@ -496,8 +491,8 @@ namespace uba
 			{
 				if (value.IsEmpty())
 					return PrintHelp(TC("-dir needs a value"));
-				if (int res = ExpandEnvironmentVariables(value))
-					return res;
+				if (!ExpandEnvironmentVariables(value))
+					return false;
 				configFile.Append(value);
 			}
 			else if (name.Equals(TC("-stats")))
@@ -569,8 +564,8 @@ namespace uba
 			{
 				if (value.IsEmpty())
 					return PrintHelp(TC("-dir needs a value"));
-				if (int res = ExpandEnvironmentVariables(value))
-					return res;
+				if (!ExpandEnvironmentVariables(value))
+					return false;
 				if ((g_rootDir.count = GetFullPathNameW(value.Replace('\\', PathSeparator).data, g_rootDir.capacity, g_rootDir.data, nullptr)) == 0)
 					return PrintHelp(StringBuffer<>().Appendf(TC("-dir has invalid path %s"), value.data).data);
 			}
@@ -621,8 +616,8 @@ namespace uba
 			{
 				if (value.IsEmpty())
 					return PrintHelp(TC("-eventfile needs a value"));
-				if (int res = ExpandEnvironmentVariables(value))
-					return res;
+				if (!ExpandEnvironmentVariables(value))
+					return false;
 				if ((eventFile.count = GetFullPathNameW(value.Replace('\\', PathSeparator).data, eventFile.capacity, eventFile.data, nullptr)) == 0)
 					return PrintHelp(StringBuffer<>().Appendf(TC("-eventfile has invalid path %s"), value.data).data);
 			}
@@ -719,7 +714,7 @@ namespace uba
 #if UBA_AUTO_UPDATE
 		if (waitProcessId != ~0u)
 			if (!WaitForProcess(waitProcessId))
-				return -1;
+				return false;
 		if (relaunchPath.count)
 			return LaunchReal(logger, relaunchPath, argc, argv);
 #endif // UBA_AUTO_UPDATE
@@ -814,12 +809,14 @@ namespace uba
 			if (resetStore)
 			{
 				if (!storage.Reset())
-					return -1;
+					return false;
 			}
 			else if (!storage.LoadCasTable(false))
-				return -1;
+				return false;
 		}
 
+		StringBuffer<512> terminationReason;
+
 #if PLATFORM_MAC
 		
 		Vector<TString> xcodeDirectories;
@@ -885,26 +882,18 @@ namespace uba
 			StringBuffer<512> xcodeSelectOutput;
 			FILE* xcodeSelect = popen("/usr/bin/xcode-select -p", "r");
 			if (xcodeSelect == nullptr || fgets(xcodeSelectOutput.data, xcodeSelectOutput.capacity, xcodeSelect) == nullptr || pclose(xcodeSelect) != 0)
+				terminationReason.Append("Failed to get an Xcode from xcode-select");
+			else
 			{
-				logger.Error("Failed to get an Xcode from xcode-select");
-				return -1;
-			}
-
-			xcodeSelectOutput.count = strlen(xcodeSelectOutput.data);
-			while (isspace(xcodeSelectOutput.data[xcodeSelectOutput.count-1]))
-			{
-				xcodeSelectOutput.data[xcodeSelectOutput.count-1] = 0;
-				xcodeSelectOutput.count--;
+				xcodeSelectOutput.count = strlen(xcodeSelectOutput.data);
+				while (isspace(xcodeSelectOutput.data[xcodeSelectOutput.count-1]))
+					xcodeSelectOutput.data[--xcodeSelectOutput.count] = 0;
+				xcodeDirectories.push_back(xcodeSelectOutput.data);
 			}
-			
-			xcodeDirectories.push_back(xcodeSelectOutput.data);
 		}
 
 		if (xcodeDirectories.size() == 0)
-		{
-			logger.Error("Unable to populate from any Xcodes. Agent is unusable.");
-			return -1;
-		}
+			terminationReason.Append("Unable to populate from any Xcodes. Agent is unusable.");
 
 		for (TString& xcodeDir : xcodeDirectories)
 		{
@@ -974,14 +963,13 @@ namespace uba
 		#endif
 
 		bool relaunch = false;
-		StringBuffer<512> terminationReason;
 		u64 terminationTimeMs = 0;
 
 		#if UBA_USE_AWS
 		if (aws.IsTerminating(logger, terminationReason, terminationTimeMs))
 		{
 			LoggerWithWriter(g_consoleLogWriter, TC("")).Info(TC("%s. Exiting UbaAgent before starting session"), terminationReason.data);
-			return 0;
+			return true;
 		}
 		#endif
 
@@ -1012,7 +1000,7 @@ namespace uba
 			NetworkClient* client = new NetworkClient(ctorSuccess, ncci);
 			auto csg = MakeGuard([&]() { client->Disconnect(); delete client; });
 			if (!ctorSuccess)
-				return -1;
+				return false;
 
 			if (useListen)
 			{
@@ -1021,14 +1009,11 @@ namespace uba
 				while (!client->IsOrWasConnected(200))
 				{
 					if (ShouldExit())
-						return 0;
+						return true;
 
 					u64 waitTime = GetTime() - startTime;
 					if (!poll && TimeToMs(waitTime) > listenTimeoutSec*1000)
-					{
-						logger.Error(TC("Failed to get connection while listening for %s"), TimeToText(waitTime).str);
-						return -1;
-					}
+						return logger.Error(TC("Failed to get connection while listening for %s"), TimeToText(waitTime).str);
 				}
 			}
 			else
@@ -1040,16 +1025,13 @@ namespace uba
 				while (!client->Connect(*networkBackend, host.data, port, &timedOut))
 				{
 					if (ShouldExit())
-						return 0;
+						return true;
 
 					if (!timedOut)
-						return -1;
+						return false;
 
 					if (!poll && !--retryCount)
-					{
-						logger.Error(TC("Failed to connect to %s:%u (after %s)"), host.data, port, TimeToText(GetTime() - startTime).str);
-						return -1;
-					}
+						return logger.Error(TC("Failed to connect to %s:%u (after %s)"), host.data, port, TimeToText(GetTime() - startTime).str);
 				}
 			}
 
@@ -1060,10 +1042,7 @@ namespace uba
 				writer.WriteString(command);
 				StackBinaryReader<8*1024> reader;
 				if (!msg.Send(reader))
-				{
-					logger.Error(TC("Failed to send command to host"));
-					return -1;
-				}
+					return logger.Error(TC("Failed to send command to host"));
 				LoggerWithWriter commandLogger(g_consoleLogWriter, TC(""));
 				commandLogger.Info(TC("----------------------------------"));
 				while (true)
@@ -1075,7 +1054,7 @@ namespace uba
 					commandLogger.Log(logType, result.c_str(), u32(result.size()));
 				}
 				commandLogger.Info(TC("----------------------------------"));
-				return 0;
+				return true;
 			}
 
 
@@ -1163,10 +1142,10 @@ namespace uba
 			auto bscsg = MakeGuard([&]() { delete storageClient; });
 
 			if (!storageClient->LoadCasTable(true))
-				return -1;
+				return false;
 
 			if (!storageClient->PopulateCasFromDirs(populateCasDirs, maxProcessCount))
-				return -1;
+				return false;
 
 			proxy.storageClient = storageClient;
 
@@ -1254,11 +1233,11 @@ namespace uba
 #if UBA_AUTO_UPDATE
 				logger.Info(TC("Downloading new binaries..."));
 				if (!DownloadBinaries(*storageClient, keys))
-					return -1;
+					return false;
 				relaunch = true;
 				break;
 #else
-				return -1;
+				return false;
 #endif
 			}
 
@@ -1280,6 +1259,13 @@ namespace uba
 			storageClient->Start();
 			sessionClient->Start();
 
+
+			if (terminationReason.count)
+			{
+				isTerminating = true;
+				sessionClient->SetIsTerminating(terminationReason.data, 0);
+			}
+
 			while (!ShouldExit())
 			{
 				if (useListen)
@@ -1382,10 +1368,10 @@ namespace uba
 #if UBA_AUTO_UPDATE
 		if (relaunch)
 			if (!LaunchTemp(logger, argc, argv))
-				return -1;
+				return false;
 #endif
 
-		return 0;
+		return true;
 	}
 }
 
@@ -1395,7 +1381,7 @@ int wmain(int argc, wchar_t* argv[])
 	using namespace uba;
 	__try
 	{
-		return WrappedMain(argc, argv);
+		return WrappedMain(argc, argv) ? 0 : -1;
 	}
 	__except(ReportSEH(GetExceptionInformation()))
 	{
@@ -1405,6 +1391,6 @@ int wmain(int argc, wchar_t* argv[])
 #else
 int main(int argc, char* argv[])
 {
-	return uba::WrappedMain(argc, argv);
+	return uba::WrappedMain(argc, argv) ? 0 : -1;
 }
 #endif
