[UBA]
* Cosmetic fixes for log entries

diff --git a/Engine/Source/Programs/UnrealBuildAccelerator/Common/Private/UbaNetworkBackendTcp.cpp b/Engine/Source/Programs/UnrealBuildAccelerator/Common/Private/UbaNetworkBackendTcp.cpp
--- a/Engine/Source/Programs/UnrealBuildAccelerator/Common/Private/UbaNetworkBackendTcp.cpp
+++ b/Engine/Source/Programs/UnrealBuildAccelerator/Common/Private/UbaNetworkBackendTcp.cpp
@@ -123,7 +123,7 @@ namespace uba
 			return true;
 		if (shutdown(s, SD_BOTH) != SOCKET_ERROR)
 			return true;
-		logger.Info(TC("failed to shutdown socket %llu for %s (%s)"), u64(s), hint, LastErrorToText(WSAGetLastError()).data);
+		logger.Info(TC("Failed to shutdown socket %llu in %s (%s)"), u64(s), hint, LastErrorToText(WSAGetLastError()).data);
 		return false;
 	}
 
@@ -133,7 +133,7 @@ namespace uba
 			return true;
 		if (closesocket(s) != SOCKET_ERROR)
 			return true;
-		logger.Info(TC("failed to close socket %llu for %s (%s)"), u64(s), hint, LastErrorToText(WSAGetLastError()).data);
+		logger.Info(TC("Failed to close socket %llu in %s (%s)"), u64(s), hint, LastErrorToText(WSAGetLastError()).data);
 		return false;
 	}
 
@@ -156,10 +156,10 @@ namespace uba
 				continue;
 			SOCKET s = conn.socket;
 			conn.socket = INVALID_SOCKET;
-			ShutdownSocket(conn.logger, s, TC("tcp dtor"));
+			ShutdownSocket(conn.logger, s, TC("Dtor"));
 			lock2.Leave();
 			conn.recvThread.Wait();
-			CloseSocket(conn.logger, s, TC("tcp dtor"));
+			CloseSocket(conn.logger, s, TC("Dtor"));
 		}
 		m_connections.clear();
 
@@ -175,7 +175,7 @@ namespace uba
 		ScopedCriticalSection lock(conn.shutdownLock);
 		if (conn.socket == INVALID_SOCKET)
 			return;
-		ShutdownSocket(conn.logger, conn.socket, TC("shutdown"));
+		ShutdownSocket(conn.logger, conn.socket, TC("Shutdown"));
 	}
 
 	bool NetworkBackendTcp::Send(Logger& logger, void* connection, const void* data, u32 dataSize, SendContext& sendContext)
@@ -437,7 +437,7 @@ namespace uba
 
 			if (!entry.connectedFunc(&conn, remoteSockAddr))
 			{
-				ShutdownSocket(logger, clientSocket, TC("thread listen"));
+				ShutdownSocket(logger, clientSocket, TC("ThreadListen"));
 				conn.ready.Set();
 				conn.recvThread.Wait();
 				SCOPED_WRITE_LOCK(m_connectionsLock, lock2);
@@ -522,8 +522,8 @@ namespace uba
 
 		if (s == INVALID_SOCKET)
 			return;
-		ShutdownSocket(logger, s, TC("threadrecv"));
-		CloseSocket(logger, s, TC("threadrecv"));
+		ShutdownSocket(logger, s, TC("ThreadRecv"));
+		CloseSocket(logger, s, TC("ThreadRecv"));
 	}
 
 	bool NetworkBackendTcp::Connect(Logger& logger, const tchar* ip, const ConnectedFunc& connectedFunc, u16 port, bool* timedOut)
@@ -581,7 +581,7 @@ namespace uba
 			return logger.Error(TC("socket failed (%s)"), LastErrorToText(WSAGetLastError()).data);
 
 		// Create guard in case we fail to connect (will be cancelled further down if we succeed)
-		auto socketClose = MakeGuard([&]() { CloseSocket(logger, socketFd, TC("connect")); });
+		auto socketClose = MakeGuard([&]() { CloseSocket(logger, socketFd, TC("Connect")); });
 
 		// Set to non-blocking just for the connect call (we want to control the connect timeout after connect using select instead)
 		if (!SetBlocking(logger, socketFd, false))
@@ -687,7 +687,7 @@ namespace uba
 
 		if (!connectedFunc(&conn, remoteSocketAddr, timedOut))
 		{
-			ShutdownSocket(logger, socketFd, TC("connect"));
+			ShutdownSocket(logger, socketFd, TC("Connect"));
 			conn.ready.Set();
 			conn.recvThread.Wait();
 			SCOPED_WRITE_LOCK(m_connectionsLock, lock2);
@@ -957,7 +957,7 @@ namespace uba
 		if (m_socket != INVALID_SOCKET)
 		{
 			LoggerWithWriter logger(g_nullLogWriter);
-			CloseSocket(logger, m_socket, TC("http dtor"));
+			CloseSocket(logger, m_socket, TC("HttpDtor"));
 		}
 
 		#if PLATFORM_WINDOWS
@@ -1013,7 +1013,7 @@ namespace uba
 		// TODO: Fix so we reuse socket connection for multiple queries
 		if (*m_host)// && _stricmp(m_host, host) != 0)
 		{
-			CloseSocket(logger, m_socket, TC("http query"));
+			CloseSocket(logger, m_socket, TC("HttpQuery"));
 			m_socket = INVALID_SOCKET;
 			*m_host = 0;
 		}
