[UBA]
* Fixed so agent populates cas after it has started the session.. this in order to prevent timeouts after connecting if populating cas from local path takes a long time

diff --git a/Engine/Source/Programs/UnrealBuildAccelerator/Agent/Private/UbaAgent.cpp b/Engine/Source/Programs/UnrealBuildAccelerator/Agent/Private/UbaAgent.cpp
--- a/Engine/Source/Programs/UnrealBuildAccelerator/Agent/Private/UbaAgent.cpp
+++ b/Engine/Source/Programs/UnrealBuildAccelerator/Agent/Private/UbaAgent.cpp
@@ -1110,7 +1110,8 @@ namespace uba
 					return true;
 				};
 
-			client->RegisterOnDisconnected([&]() { networkBackend->StopListen(); if (auto proxyServer = proxy.server.load()) proxyServer->DisconnectClients(); });
+			Atomic<bool> isDisconnected;
+			client->RegisterOnDisconnected([&]() { isDisconnected = true; networkBackend->StopListen(); if (auto proxyServer = proxy.server.load()) proxyServer->DisconnectClients(); });
 
 			struct NetworkBackends
 			{
@@ -1144,9 +1145,6 @@ namespace uba
 			if (!storageClient->LoadCasTable(true))
 				return false;
 
-			if (!storageClient->PopulateCasFromDirs(populateCasDirs, maxProcessCount))
-				return false;
-
 			proxy.storageClient = storageClient;
 
 			SessionClient* sessionClient = nullptr;
@@ -1256,9 +1254,21 @@ namespace uba
 			//SetPriorityClass(GetCurrentProcess(), ABOVE_NORMAL_PRIORITY_CLASS);
 			//#endif
 
+			bool needPrepopulate = !populateCasDirs.empty();
+			if (needPrepopulate)
+				sessionClient->SetAllowSpawn(false);
+
 			storageClient->Start();
 			sessionClient->Start();
 
+			// We do population here to make sure session thread is running which will send pings to host (to prevent timeouts
+			if (needPrepopulate)
+			{
+				if (storageClient->PopulateCasFromDirs(populateCasDirs, maxProcessCount, [&]() { return isDisconnected.load(); }))
+					sessionClient->SetAllowSpawn(true);
+				else
+					terminationReason.Append(TC("Failed to prepopulate cas from local directory"));
+			}
 
 			if (terminationReason.count)
 			{
diff --git a/Engine/Source/Programs/UnrealBuildAccelerator/Common/Private/UbaSessionClient.cpp b/Engine/Source/Programs/UnrealBuildAccelerator/Common/Private/UbaSessionClient.cpp
--- a/Engine/Source/Programs/UnrealBuildAccelerator/Common/Private/UbaSessionClient.cpp
+++ b/Engine/Source/Programs/UnrealBuildAccelerator/Common/Private/UbaSessionClient.cpp
@@ -31,6 +31,7 @@ namespace uba
 	,	m_terminationTime(~0ull)
 	,	m_waitToSendEvent(false)
 	,	m_loop(true)
+	,	m_allowSpawn(true)
 	{
 		m_maxProcessCount = info.maxProcessCount;
 		m_dedicated = info.dedicated;
@@ -118,6 +119,11 @@ namespace uba
 		m_maxProcessCount = count;
 	}
 
+	void SessionClient::SetAllowSpawn(bool allow)
+	{
+		m_allowSpawn = allow;
+	}
+
 	u64 SessionClient::GetBestPing()
 	{
 		return m_bestPing;
@@ -1555,7 +1561,7 @@ namespace uba
 				m_logger.Info(TC("%s. Will stop scheduling processes and send failing processes back for retry"), m_terminationReason.load());
 			}
 
-			if (!activeProcesses.empty())
+			if (!activeProcesses.empty() || !m_allowSpawn)
 			{
 				idleStartTime = GetTime();
 				processRequestCount = 0;
@@ -1601,7 +1607,7 @@ namespace uba
 				lastWaitTime = GetTime();
 			}
 
-			bool canSpawn = TimeToMs(GetTime() - lastWaitTime) > waitTimeToSpawnAfterKillMs;
+			bool canSpawn = TimeToMs(GetTime() - lastWaitTime) > waitTimeToSpawnAfterKillMs && m_allowSpawn;
 			if (!canSpawn)
 				waitTimeoutMs = 500;
 
diff --git a/Engine/Source/Programs/UnrealBuildAccelerator/Common/Private/UbaStorageClient.cpp b/Engine/Source/Programs/UnrealBuildAccelerator/Common/Private/UbaStorageClient.cpp
--- a/Engine/Source/Programs/UnrealBuildAccelerator/Common/Private/UbaStorageClient.cpp
+++ b/Engine/Source/Programs/UnrealBuildAccelerator/Common/Private/UbaStorageClient.cpp
@@ -100,7 +100,7 @@ namespace uba
 			m_proxyClient->client.Disconnect();
 	}
 
-	bool StorageClient::PopulateCasFromDirs(const DirVector& directories, u32 workerCount)
+	bool StorageClient::PopulateCasFromDirs(const DirVector& directories, u32 workerCount, const Function<bool()>& shouldExit)
 	{
 		if (directories.empty())
 			return true;
@@ -114,7 +114,7 @@ namespace uba
 		ReaderWriterLock seenIdsLock;
 
 		for (auto& dir : directories)
-			success = PopulateCasFromDirsRecursive(dir.c_str(), workManager, seenIds, seenIdsLock) && success;
+			success = PopulateCasFromDirsRecursive(dir.c_str(), workManager, seenIds, seenIdsLock, shouldExit) && success;
 		workManager.FlushWork();
 
 		if (u32 fileCount = u32(m_localStorageFiles.size()))
@@ -833,7 +833,7 @@ namespace uba
 		return sender.SendFileCompressed(casKey, fileName, sourceMem, sourceSize, hint);
 	}
 
-	bool StorageClient::PopulateCasFromDirsRecursive(const tchar* dir, WorkManager& workManager, UnorderedSet<u64>& seenIds, ReaderWriterLock& seenIdsLock)
+	bool StorageClient::PopulateCasFromDirsRecursive(const tchar* dir, WorkManager& workManager, UnorderedSet<u64>& seenIds, ReaderWriterLock& seenIdsLock, const Function<bool()>& shouldExit)
 	{
 		StringBuffer<> fullPath;
 		fullPath.Append(dir).EnsureEndsWithSlash();
@@ -847,12 +847,15 @@ namespace uba
 					if (!seenIds.insert(e.id).second)
 						return;
 					lock.Leave();
-					PopulateCasFromDirsRecursive(fullPath.data, workManager, seenIds, seenIdsLock);
+					PopulateCasFromDirsRecursive(fullPath.data, workManager, seenIds, seenIdsLock, shouldExit);
 					return;
 				}
 
 				workManager.AddWork([&, filePath = TString(fullPath.data), name = TString(e.name)]()
 					{
+						if (shouldExit && shouldExit())
+							return;
+
 						FileInformation info;
 						if (!GetFileInformation(info, m_logger, filePath.c_str()))
 						{
diff --git a/Engine/Source/Programs/UnrealBuildAccelerator/Common/Public/UbaSessionClient.h b/Engine/Source/Programs/UnrealBuildAccelerator/Common/Public/UbaSessionClient.h
--- a/Engine/Source/Programs/UnrealBuildAccelerator/Common/Public/UbaSessionClient.h
+++ b/Engine/Source/Programs/UnrealBuildAccelerator/Common/Public/UbaSessionClient.h
@@ -41,6 +41,7 @@ namespace uba
 		void SendSummary(const Function<void(Logger&)>& extraInfo);
 		void SetIsTerminating(const tchar* reason = TC("Terminating"), u64 delayMs = 0); // Session stores pointer directly. Can't be temporary
 		void SetMaxProcessCount(u32 count);
+		void SetAllowSpawn(bool allow);
 
 		u64 GetBestPing();
 
@@ -150,6 +151,7 @@ namespace uba
 		Event m_waitToSendEvent;
 		Thread m_loopThread;
 		Atomic<bool> m_loop;
+		Atomic<bool> m_allowSpawn;
 
 		Function<void(const ProcessHandle&)> m_processFinished;
 
diff --git a/Engine/Source/Programs/UnrealBuildAccelerator/Common/Public/UbaStorageClient.h b/Engine/Source/Programs/UnrealBuildAccelerator/Common/Public/UbaStorageClient.h
--- a/Engine/Source/Programs/UnrealBuildAccelerator/Common/Public/UbaStorageClient.h
+++ b/Engine/Source/Programs/UnrealBuildAccelerator/Common/Public/UbaStorageClient.h
@@ -43,7 +43,7 @@ namespace uba
 		void StopProxy();
 
 		using DirVector = Vector<TString>;
-		bool PopulateCasFromDirs(const DirVector& directories, u32 workerCount);
+		bool PopulateCasFromDirs(const DirVector& directories, u32 workerCount, const Function<bool()>& shouldExit = {});
 
 		#if !UBA_USE_SPARSEFILE
 		virtual bool GetCasFileName(StringBufferBase& out, const CasKey& casKey) override;
@@ -61,7 +61,7 @@ namespace uba
 
 	private:
 		bool SendFile(const CasKey& casKey, const tchar* fileName, u8* sourceMem, u64 sourceSize, const tchar* hint);
-		bool PopulateCasFromDirsRecursive(const tchar* dir, WorkManager& workManager, UnorderedSet<u64>& seenIds, ReaderWriterLock& seenIdsLock);
+		bool PopulateCasFromDirsRecursive(const tchar* dir, WorkManager& workManager, UnorderedSet<u64>& seenIds, ReaderWriterLock& seenIdsLock, const Function<bool()>& shouldExit);
 
 		NetworkClient& m_client;
 		bool m_sendCompressed;
