[UBA]
* Expanded lock a little bit when populating local cas files because it is now populated at the same time as being used
* Fixed uninitialized variable that caused cas population from all xcode installations

diff --git a/Engine/Source/Programs/UnrealBuildAccelerator/Agent/Private/UbaAgent.cpp b/Engine/Source/Programs/UnrealBuildAccelerator/Agent/Private/UbaAgent.cpp
--- a/Engine/Source/Programs/UnrealBuildAccelerator/Agent/Private/UbaAgent.cpp
+++ b/Engine/Source/Programs/UnrealBuildAccelerator/Agent/Private/UbaAgent.cpp
@@ -437,7 +437,7 @@ namespace uba
 
 		#if PLATFORM_MAC
 		StringBuffer<32> populateCasFromXcodeVersion;
-		bool populateCasFromAllXcodes;
+		bool populateCasFromAllXcodes = false;
 		#endif
 
 		for (int i=1; i!=argc; ++i)
diff --git a/Engine/Source/Programs/UnrealBuildAccelerator/Common/Private/UbaStorageClient.cpp b/Engine/Source/Programs/UnrealBuildAccelerator/Common/Private/UbaStorageClient.cpp
--- a/Engine/Source/Programs/UnrealBuildAccelerator/Common/Private/UbaStorageClient.cpp
+++ b/Engine/Source/Programs/UnrealBuildAccelerator/Common/Private/UbaStorageClient.cpp
@@ -875,8 +875,6 @@ namespace uba
 						if (!insres.second)
 							return;
 						LocalFile& localFile = insres.first->second;
-						lookupLock.Leave();
-
 						localFile.casEntry.size = info.size;
 						localFile.fileName = filePath;
 
