[UBA]
* Added so UbaObjTool writing import lib also writes .exp file to make UBT happy

diff --git a/Engine/Source/Programs/UnrealBuildAccelerator/Core/Private/UbaApplicationRules.cpp b/Engine/Source/Programs/UnrealBuildAccelerator/Core/Private/UbaApplicationRules.cpp
--- a/Engine/Source/Programs/UnrealBuildAccelerator/Core/Private/UbaApplicationRules.cpp
+++ b/Engine/Source/Programs/UnrealBuildAccelerator/Core/Private/UbaApplicationRules.cpp
@@ -469,7 +469,8 @@ namespace uba
 	{
 		virtual bool IsOutputFile(const StringView& fileName) const override
 		{
-			return fileName.EndsWith(TC(".obj"));
+			return fileName.EndsWith(TC(".obj"))
+				|| fileName.EndsWith(TC(".exp"));
 		}
 
 		virtual bool StoreFileCompressed(const StringView& fileName) const
diff --git a/Engine/Source/Programs/UnrealBuildAccelerator/ObjTool/Private/UbaImportLibWriter.cpp b/Engine/Source/Programs/UnrealBuildAccelerator/ObjTool/Private/UbaImportLibWriter.cpp
--- a/Engine/Source/Programs/UnrealBuildAccelerator/ObjTool/Private/UbaImportLibWriter.cpp
+++ b/Engine/Source/Programs/UnrealBuildAccelerator/ObjTool/Private/UbaImportLibWriter.cpp
@@ -661,7 +661,18 @@ namespace uba
 
 			memcpy(fa.GetData(), m_memory.memory, m_memory.writtenSize);
 
-			return fa.Close();
+			if (!fa.Close())
+				return false;
+
+			// Create exp file
+			StringBuffer<> expFile;
+			expFile.Append(libFile);
+			if (const tchar* dot = expFile.Last('.'))
+				expFile.Resize(dot - expFile.data).Append(TC(".exp"));
+			FileAccessor faExp(logger, expFile.data);
+			if (!faExp.CreateWrite())
+				return false;
+			return faExp.Close();
 		}
 	};
 #endif // PLATFORM_WINDOWS
