[UBA]
* Fixed so PopulateCasFromDirs now use file lookup table to prevent cas calculation if file has been calculated before

diff --git a/Engine/Source/Programs/UnrealBuildAccelerator/Common/Private/UbaStorageClient.cpp b/Engine/Source/Programs/UnrealBuildAccelerator/Common/Private/UbaStorageClient.cpp
--- a/Engine/Source/Programs/UnrealBuildAccelerator/Common/Private/UbaStorageClient.cpp
+++ b/Engine/Source/Programs/UnrealBuildAccelerator/Common/Private/UbaStorageClient.cpp
@@ -863,15 +863,33 @@ namespace uba
 							return;
 						}
 
-						CasKey casKey;
-						if (!CalculateCasKey(casKey, filePath.c_str()))
+						StringBuffer<> forKey;
+						FixPath(filePath.c_str(), nullptr, 0, forKey);
+						if (CaseInsensitiveFs)
+							forKey.MakeLower();
+						StringKey fileNameKey = ToStringKey(forKey);
+
+
+						FileEntry& fileEntry = GetOrCreateFileEntry(fileNameKey);
+						SCOPED_WRITE_LOCK(fileEntry.lock, fileEntryLock);
+
+						if (info.size != fileEntry.size || info.lastWriteTime != fileEntry.lastWritten)
 						{
-							m_logger.Error(TC("Failed to calculate cas key for %s"), filePath.c_str());
-							return;
+							CasKey casKey;
+							if (!CalculateCasKey(casKey, filePath.c_str()))
+							{
+								m_logger.Error(TC("Failed to calculate cas key for %s"), filePath.c_str());
+								return;
+							}
+							fileEntry.size = info.size;
+							fileEntry.lastWritten = info.lastWriteTime;
+							fileEntry.casKey = AsCompressed(casKey, false);
 						}
+						fileEntry.verified = true;
+						fileEntryLock.Leave();
 
 						SCOPED_WRITE_LOCK(m_localStorageFilesLock, lookupLock);
-						auto insres = m_localStorageFiles.try_emplace(AsCompressed(casKey, false));
+						auto insres = m_localStorageFiles.try_emplace(fileEntry.casKey);
 						if (!insres.second)
 							return;
 						LocalFile& localFile = insres.first->second;
