[UBA]
* More optimizations around PopulateCasFromDirs

diff --git a/Engine/Source/Programs/UnrealBuildAccelerator/Common/Private/UbaStorageClient.cpp b/Engine/Source/Programs/UnrealBuildAccelerator/Common/Private/UbaStorageClient.cpp
--- a/Engine/Source/Programs/UnrealBuildAccelerator/Common/Private/UbaStorageClient.cpp
+++ b/Engine/Source/Programs/UnrealBuildAccelerator/Common/Private/UbaStorageClient.cpp
@@ -847,11 +847,16 @@ namespace uba
 					if (!seenIds.insert(e.id).second)
 						return;
 					lock.Leave();
-					PopulateCasFromDirsRecursive(fullPath.data, workManager, seenIds, seenIdsLock, shouldExit);
+					workManager.AddWork([&, filePath = TString(fullPath.data)]()
+						{
+							if (shouldExit && shouldExit())
+								return;
+							PopulateCasFromDirsRecursive(filePath.c_str(), workManager, seenIds, seenIdsLock, shouldExit);
+						}, 1, TC(""));
 					return;
 				}
 
-				workManager.AddWork([&, filePath = TString(fullPath.data), name = TString(e.name)]()
+				workManager.AddWork([&, filePath = TString(fullPath.data)]()
 					{
 						if (shouldExit && shouldExit())
 							return;
