[UBA]
* Gracefully handle socket shutdown when socket is not connected

diff --git a/Engine/Source/Programs/UnrealBuildAccelerator/Common/Private/UbaNetworkBackendTcp.cpp b/Engine/Source/Programs/UnrealBuildAccelerator/Common/Private/UbaNetworkBackendTcp.cpp
--- a/Engine/Source/Programs/UnrealBuildAccelerator/Common/Private/UbaNetworkBackendTcp.cpp
+++ b/Engine/Source/Programs/UnrealBuildAccelerator/Common/Private/UbaNetworkBackendTcp.cpp
@@ -29,6 +29,7 @@
 #define INVALID_SOCKET -1
 #define SD_BOTH SHUT_RDWR
 #define WSAHOST_NOT_FOUND 0
+#define WSAENOTCONN ENOTCONN
 #define WSAEADDRINUSE EADDRINUSE
 #define closesocket(a) close(a)
 #define addrinfoW addrinfo
@@ -123,6 +124,8 @@ namespace uba
 			return true;
 		if (shutdown(s, SD_BOTH) != SOCKET_ERROR)
 			return true;
+		if (WSAGetLastError() == WSAENOTCONN)
+			return true;
 		logger.Info(TC("Failed to shutdown socket %llu in %s (%s)"), u64(s), hint, LastErrorToText(WSAGetLastError()).data);
 		return false;
 	}
