[UBA]
* Fixed so wrapped llvm ir files can be parsed
* Fixed so correct extra symbol file is created for various platforms
* Added more handling of .elf files. It should find exported and imported symbols better now

diff --git a/Engine/Source/Programs/UnrealBuildAccelerator/Common/Private/UbaObjectFile.cpp b/Engine/Source/Programs/UnrealBuildAccelerator/Common/Private/UbaObjectFile.cpp
--- a/Engine/Source/Programs/UnrealBuildAccelerator/Common/Private/UbaObjectFile.cpp
+++ b/Engine/Source/Programs/UnrealBuildAccelerator/Common/Private/UbaObjectFile.cpp
@@ -218,35 +218,23 @@ namespace uba
 		return m_potentialDuplicates;
 	}
 
-	bool ObjectFile::CreateExtraFile(Logger& logger, const tchar* extraObjFilename, ObjectFileType type, const UnorderedSymbols& allNeededImports, const UnorderedSymbols& allSharedImports, const UnorderedExports& allSharedExports, bool includeExportsInFile)
+	bool ObjectFile::CreateExtraFile(Logger& logger, const StringView& extraObjFilename, const StringView& platform, const UnorderedSymbols& allNeededImports, const UnorderedSymbols& allSharedImports, const UnorderedExports& allSharedExports, bool includeExportsInFile)
 	{
 		ObjectFileCoff objectFileCoff;
 		ObjectFileElf objectFileElf;
-		ObjectFileLLVMIR objectFileLLVMIR;
 		
-		ObjectFile& objectFile = [&]() -> ObjectFile&
-			{
-				switch (type)
-				{
-				case ObjectFileType_Coff:
-					return objectFileCoff;
-				case ObjectFileType_Elf:
-					return objectFileCoff;
-				case ObjectFileType_LLVMIR:
-					return objectFileLLVMIR;
-				default:
-					UBA_ASSERT(false);
-					return objectFileCoff;
-				}
-			}();
-
-
 		MemoryBlock memoryBlock(16*1024*1024);
 
-		if (!objectFile.CreateExtraFile(logger, memoryBlock, allNeededImports, allSharedImports, allSharedExports, includeExportsInFile))
+		bool res;
+		if (platform.Equals(TC("win64")) || platform.Equals(TC("wingdk")) || platform.Equals(TC("xb1")) || platform.Equals(TC("xsx"))) 
+			res = ObjectFileCoff::CreateExtraFile(logger, platform, memoryBlock, allNeededImports, allSharedImports, allSharedExports, includeExportsInFile);
+		else
+			res = ObjectFileElf::CreateExtraFile(logger, platform, memoryBlock, allNeededImports, allSharedImports, allSharedExports, includeExportsInFile);
+
+		if (!res)
 			return false;
 
-		FileAccessor extraFile(logger, extraObjFilename);
+		FileAccessor extraFile(logger, extraObjFilename.data);
 		if (!extraFile.CreateWrite())
 			return false;
 
diff --git a/Engine/Source/Programs/UnrealBuildAccelerator/Common/Private/UbaObjectFileCoff.cpp b/Engine/Source/Programs/UnrealBuildAccelerator/Common/Private/UbaObjectFileCoff.cpp
--- a/Engine/Source/Programs/UnrealBuildAccelerator/Common/Private/UbaObjectFileCoff.cpp
+++ b/Engine/Source/Programs/UnrealBuildAccelerator/Common/Private/UbaObjectFileCoff.cpp
@@ -212,7 +212,7 @@ namespace uba
 	}
 
 	template<typename SymbolType>
-	ObjectFile::StringView GetSymbolName(SymbolType& symbol, const u8* data, u32 stringTableMemPos)
+	ObjectFile::AnsiStringView GetSymbolName(SymbolType& symbol, const u8* data, u32 stringTableMemPos)
 	{
 		if (symbol.N.Name.Short == 0)
 		{
@@ -335,7 +335,7 @@ namespace uba
 			if (symbol.StorageClass != ImageSymClassExternal)
 				continue;
 
-			StringView symbolName = GetSymbolName(symbol, m_data, m_info.stringTableMemPos);
+			AnsiStringView symbolName = GetSymbolName(symbol, m_data, m_info.stringTableMemPos);
 			symbolName.ToString(symbolString);
 
 			if (symbol.SectionNumber != ImageSymUndefined)
@@ -449,12 +449,7 @@ namespace uba
 		return true;
 	}
 
-	bool ObjectFileCoff::CreateExtraFile(Logger& logger, MemoryBlock& memoryBlock, const UnorderedSymbols& allNeededImports, const UnorderedSymbols& allSharedImports, const UnorderedExports& allSharedExports, bool includeExportsInFile)
-	{
-		return CreateExtraFile2(logger, memoryBlock, allNeededImports, allSharedImports, allSharedExports, includeExportsInFile);
-	}
-
-	bool ObjectFileCoff::CreateExtraFile2(Logger& logger, MemoryBlock& memoryBlock, const UnorderedSymbols& allNeededImports, const UnorderedSymbols& allSharedImports, const UnorderedExports& allSharedExports, bool includeExportsInFile)
+	bool ObjectFileCoff::CreateExtraFile(Logger& logger, const StringView& platform, MemoryBlock& memoryBlock, const UnorderedSymbols& allNeededImports, const UnorderedSymbols& allSharedImports, const UnorderedExports& allSharedExports, bool includeExportsInFile)
 	{
 		std::string tmp;
 
@@ -588,7 +583,7 @@ namespace uba
 				continue;
 			if (symbol.SectionNumber != ImageSymUndefined)
 				continue;
-			StringView symbolName = GetSymbolName(symbol, m_data, m_info.stringTableMemPos);
+			AnsiStringView symbolName = GetSymbolName(symbol, m_data, m_info.stringTableMemPos);
 			if (!symbolName.StartsWith("__imp_", 6))
 				continue;
 			symbolName.strBegin += 6;
@@ -736,7 +731,7 @@ namespace uba
 			if (!symbol.SectionNumber)
 				continue;
 
-			StringView symbolName = GetSymbolName(symbol, m_data, m_info.stringTableMemPos);
+			AnsiStringView symbolName = GetSymbolName(symbol, m_data, m_info.stringTableMemPos);
 			auto findIt = m_toRemove.find(symbolName.ToString(tmp));
 			if (findIt == m_toRemove.end())
 				continue;
diff --git a/Engine/Source/Programs/UnrealBuildAccelerator/Common/Private/UbaObjectFileCoff.h b/Engine/Source/Programs/UnrealBuildAccelerator/Common/Private/UbaObjectFileCoff.h
--- a/Engine/Source/Programs/UnrealBuildAccelerator/Common/Private/UbaObjectFileCoff.h
+++ b/Engine/Source/Programs/UnrealBuildAccelerator/Common/Private/UbaObjectFileCoff.h
@@ -13,7 +13,7 @@ namespace uba
 		ObjectFileCoff();
 		virtual bool Parse(Logger& logger, const tchar* hint) override;
 
-		static bool CreateExtraFile2(Logger& logger, MemoryBlock& memoryBlock, const UnorderedSymbols& allNeededImports, const UnorderedSymbols& allSharedImports, const UnorderedExports& allSharedExports, bool includeExportsInFile);
+		static bool CreateExtraFile(Logger& logger, const StringView& platform, MemoryBlock& memoryBlock, const UnorderedSymbols& allNeededImports, const UnorderedSymbols& allSharedImports, const UnorderedExports& allSharedExports, bool includeExportsInFile);
 
 	private:
 		struct Info;
@@ -22,7 +22,6 @@ namespace uba
 		template<typename SymbolType> void ParseImports();
 
 		virtual bool StripExports(Logger& logger, u8* newData, const UnorderedSymbols& allNeededImports) override;
-		virtual bool CreateExtraFile(Logger& logger, MemoryBlock& memoryBlock, const UnorderedSymbols& allNeededImports, const UnorderedSymbols& allSharedImports, const UnorderedExports& allSharedExports, bool includeExportsInFile) override;
 
 		template<typename SymbolType> void CalculateImports(Logger& logger, Vector<u32>& outImports);
 		template<typename SymbolType> void WriteImports(Logger& logger, u8* newData, Info& newInfo, const Vector<u32>& symbolsToAdd);
diff --git a/Engine/Source/Programs/UnrealBuildAccelerator/Common/Private/UbaObjectFileElf.cpp b/Engine/Source/Programs/UnrealBuildAccelerator/Common/Private/UbaObjectFileElf.cpp
--- a/Engine/Source/Programs/UnrealBuildAccelerator/Common/Private/UbaObjectFileElf.cpp
+++ b/Engine/Source/Programs/UnrealBuildAccelerator/Common/Private/UbaObjectFileElf.cpp
@@ -7,63 +7,96 @@ namespace uba
 {
 	struct Elf64Header
 	{
-	  unsigned char	e_ident[16];	/* Magic number and other info */
-	  u16	e_type;			/* Object file type */
-	  u16	e_machine;		/* Architecture */
-	  u32	e_version;		/* Object file version */
-	  u64	e_entry;		/* Entry point virtual address */
-	  u64	e_phoff;		/* Program header table file offset */
-	  u64	e_shoff;		/* Section header table file offset */
-	  u32	e_flags;		/* Processor-specific flags */
-	  u16	e_ehsize;		/* ELF header size in bytes */
-	  u16	e_phentsize;	/* Program header table entry size */
-	  u16	e_phnum;		/* Program header table entry count */
-	  u16	e_shentsize;	/* Section header table entry size */
-	  u16	e_shnum;		/* Section header table entry count */
-	  u16	e_shstrndx;		/* Section header string table index */
+		unsigned char	e_ident[16];	/* Magic number and other info */
+		u16	e_type;			/* Object file type */
+		u16	e_machine;		/* Architecture */
+		u32	e_version;		/* Object file version */
+		u64	e_entry;		/* Entry point virtual address */
+		u64	e_phoff;		/* Program header table file offset */
+		u64	e_shoff;		/* Section header table file offset */
+		u32	e_flags;		/* Processor-specific flags */
+		u16	e_ehsize;		/* ELF header size in bytes */
+		u16	e_phentsize;	/* Program header table entry size */
+		u16	e_phnum;		/* Program header table entry count */
+		u16	e_shentsize;	/* Section header table entry size */
+		u16	e_shnum;		/* Section header table entry count */
+		u16	e_shstrndx;		/* Section header string table index */
 	};
 
 	struct Elf64SectionHeader
 	{
-	  u32	sh_name;		/* Section name (string tbl index) */
-	  u32	sh_type;		/* Section type */
-	  u64	sh_flags;		/* Section flags */
-	  u64	sh_addr;		/* Section virtual addr at execution */
-	  u64	sh_offset;		/* Section file offset */
-	  u64	sh_size;		/* Section size in bytes */
-	  u32	sh_link;		/* Link to another section */
-	  u32	sh_info;		/* Additional section information */
-	  u64	sh_addralign;	/* Section alignment */
-	  u64	sh_entsize;		/* Entry size if section holds table */
+		u32	sh_name;		/* Section name (string tbl index) */
+		u32	sh_type;		/* Section type */
+		u64	sh_flags;		/* Section flags */
+		u64	sh_addr;		/* Section virtual addr at execution */
+		u64	sh_offset;		/* Section file offset */
+		u64	sh_size;		/* Section size in bytes */
+		u32	sh_link;		/* Link to another section */
+		u32	sh_info;		/* Additional section information */
+		u64	sh_addralign;	/* Section alignment */
+		u64	sh_entsize;		/* Entry size if section holds table */
 	};
 
 	struct Elf64Sym
 	{
-	  u32	st_name;		/* Symbol name (string tbl index) */
-	  u8	st_info;		/* Symbol type and binding */
-	  u8	st_other;		/* Symbol visibility */
-	  u16	st_shndx;		/* Section index */
-	  u64	st_value;		/* Symbol value */
-	  u64	st_size;		/* Symbol size */
+		u32	st_name;		/* Symbol name (string tbl index) */
+		u8	st_info;		/* Symbol type and binding */
+		u8	st_other;		/* Symbol visibility */
+		u16	st_shndx;		/* Section index */
+		u64	st_value;		/* Symbol value */
+		u64	st_size;		/* Symbol size */
 	};
 
+	struct Elf64Rel
+	{
+		u64	r_offset;
+		u64	r_info;
+	};
+ 
+	struct Elf64Rela
+	{
+		u64	r_offset;
+		u64	r_info;
+		u64	r_addend;
+	};
+
+	#define ELF64_R_SYM(i)			((i) >> 32)
+	#define ELF64_R_TYPE(i)			((i) & 0xffffffff)
+	#define ELF64_R_INFO(sym,type)		((((u64) (sym)) << 32) + (type))
+
+	#define ELF32_ST_VISIBILITY(o)	((o) & 0x03)
+	#define ELF64_ST_VISIBILITY(o)	ELF32_ST_VISIBILITY (o)
+
 	#define EM_X86_64	62	/* AMD x86-64 architecture */
 
 	#define SHT_SYMTAB	  2		/* Symbol table */
+	#define SHT_STRTAB	  3
+	#define SHT_RELA	  4		/* Relocation entries with addends */
+	#define SHT_DYNAMIC	  6		/* Dynamic linking information */
+	#define SHT_REL		  9		/* Relocation entries, no addends */
 	#define SHT_DYNSYM	  11		/* Dynamic linker symbol table */
+	#define SHT_SYMTAB_SHNDX  18		/* Extended section indeces */
 
+	#define STT_NOTYPE	0		/* Symbol type is unspecified */
 	#define STT_OBJECT	1		/* Symbol is a data object */
 	#define STT_FUNC	2		/* Symbol is a code object */
-	
+	#define STT_SECTION	3		/* Symbol associated with a section */
+
+	#define STB_LOCAL	0		/* Local symbol */
+	#define STB_GLOBAL	1		/* Global symbol */
 	#define STB_WEAK	2		/* Weak symbol */
 
+	#define STV_DEFAULT	0		// Default symbol visibility rules
+	#define STV_INTERNAL	1	// Processor specific hidden class
+	#define STV_HIDDEN	2		// Sym unavailable in other modules
+	#define STV_PROTECTED	3	// Not preemptible, not exported
 
 	#define ELF32_ST_BIND(val)		(((unsigned char) (val)) >> 4)
 	#define ELF32_ST_TYPE(val)		((val) & 0xf)
 	#define ELF32_ST_INFO(bind, type)	(((bind) << 4) + ((type) & 0xf))
 
 	/* Both Elf32_Sym and Elf64_Sym use the same one-byte st_info field.  */
-	#define ELF64_ST_BIND(val)		ELF32_ST_BIND (val)
+	#define ELF64_ST_BIND(val)		ELF32_ST_BIND(val)
 	#define ELF64_ST_TYPE(val)		ELF32_ST_TYPE (val)
 	#define ELF64_ST_INFO(bind, type)	ELF32_ST_INFO ((bind), (type))
 
@@ -128,13 +161,13 @@ namespace uba
 	bool ObjectFileElf::Parse(Logger& logger, const tchar* hint)
 	{
 		auto& header = *(Elf64Header*)m_data;
-		if (header.e_ident[4] != 2) // 64-bit
+		if (header.e_ident[4] != 2) // 64-bit  (1 is 32-bit)
 			return false;
 		if (header.e_ident[5] != 1) // Little endian
 			return false;
 		if (header.e_ident[6] != 1) // Version
 			return false;
-		if (header.e_type != 1) // Relocatable file
+		if (header.e_type != 1 && header.e_type != 0xfe0c) // Relocatable file
 			return false;
 
 		UBA_ASSERT(sizeof(Elf64SectionHeader) == header.e_shentsize);
@@ -142,7 +175,10 @@ namespace uba
 		auto sections = (Elf64SectionHeader*)(m_data + header.e_shoff);
 
 		auto& namesSection = sections[header.e_shstrndx];
-		u8* stringTableData = m_data + namesSection.sh_offset;
+		UBA_ASSERT(namesSection.sh_type == SHT_STRTAB);
+		char* sectionNamesTable = (char*)m_data + namesSection.sh_offset;
+		char* symTableNames = nullptr;
+		char* dynTableNames = nullptr;
 
 		#if 0
 		u64 stringTableSize = namesSection.sh_size;
@@ -161,11 +197,24 @@ namespace uba
 		for (u16 i=0; i!=header.e_shnum; ++i)
 		{
 			auto& section = sections[i];
-			//auto sectionName = (char*)stringTableData + section.sh_name;
+			if (section.sh_type != SHT_STRTAB)
+				continue;
+			char* sectionName = sectionNamesTable + section.sh_name;
+			if (strcmp(sectionName, ".shstrtab") == 0)
+				printf("");
+			else if (strcmp(sectionName, ".strtab") == 0)
+				symTableNames = (char*)m_data + section.sh_offset;
+			else if (strcmp(sectionName, ".dynstr") == 0)
+				dynTableNames = (char*)m_data + section.sh_offset;
+		}
+
+		for (u16 i=0; i!=header.e_shnum; ++i)
+		{
+			auto& section = sections[i];
 
 			if (section.sh_type == SHT_SYMTAB || section.sh_type == SHT_DYNSYM)
 			{
-				UBA_ASSERT(section.sh_link == header.e_shstrndx);
+				//UBA_ASSERT(section.sh_link == header.e_shstrndx);
 				UBA_ASSERT(section.sh_entsize == sizeof(Elf64Sym));
 				auto symbols = (Elf64Sym*)(m_data + section.sh_offset);
 				u64 symbolCount = section.sh_size / sizeof(Elf64Sym);
@@ -173,16 +222,47 @@ namespace uba
 				{
 					auto& symbol = symbols[j];
 					u8 symbolType = ELF64_ST_TYPE(symbol.st_info);
-					if (symbolType != STT_FUNC && symbolType != STT_OBJECT)
+					const char* symbolName;
+
+					if (symbolType == STT_FUNC || symbolType == STT_OBJECT)
+						symbolName = (section.sh_type == SHT_DYNSYM ? dynTableNames : symTableNames) + symbol.st_name;
+					else if (symbolType == STT_SECTION)
+						symbolName = sectionNamesTable + sections[symbol.st_shndx].sh_name;
+					else if (symbolType == STT_NOTYPE)
+						symbolName = symTableNames + symbol.st_name;
+					else
 						continue;
 
-					auto name = (char*)stringTableData + symbol.st_name;
-					if (PotentiallyDuplicatedSymbols.find(name) == PotentiallyDuplicatedSymbols.end())
-						continue;
-
-					m_potentialDuplicates.emplace(name);
+					u32 symbolBinding = ELF64_ST_BIND(symbol.st_info);
+
+					if (symbolBinding == STB_GLOBAL)
+					{
+						u32 vis = ELF64_ST_VISIBILITY(symbol.st_other);
+						(void)vis;
+						if (symbolType != STT_NOTYPE)
+						{
+							// On ELF, protected visibility indicates that the symbol will be placed in the dynamic symbol table, but that references within the defining module will bind to the local symbol. That is, the symbol cannot be overridden by another module.
+							if (vis == STV_PROTECTED) //|| vis == STV_DEFAULT
+								m_exports.emplace(symbolName, ExportInfo{std::string(), 0});//index++});
+						}
+						else
+							m_imports.emplace(symbolName);//index++});
+					}
 				}
 			}
+			else if (section.sh_type == SHT_RELA)
+			{
+				// ".rela.*"
+				char* sectionName = sectionNamesTable + section.sh_name;
+				(void)sectionName;
+				auto& data = *(Elf64Rela*)(m_data + section.sh_offset);
+				ELF64_R_SYM(data.r_info);
+				ELF64_R_TYPE(data.r_info);
+					(void)data;
+				//printf("");
+			}
+			//else
+			//	printf("TYPE: %u\n", section.sh_type);
 		}
 		return true;
 	}
@@ -235,7 +315,7 @@ namespace uba
 		return true;
 	}
 
-	bool ObjectFileElf::CreateExtraFile(Logger& logger, MemoryBlock& memoryBlock, const UnorderedSymbols& allNeededImports, const UnorderedSymbols& allSharedImports, const UnorderedExports& allSharedExports, bool includeExportsInFile)
+	bool ObjectFileElf::CreateExtraFile(Logger& logger, const StringView& platform, MemoryBlock& memoryBlock, const UnorderedSymbols& allNeededImports, const UnorderedSymbols& allSharedImports, const UnorderedExports& allSharedExports, bool includeExportsInFile)
 	{
 		u8* data = (u8*)memoryBlock.Allocate(sizeof(Elf64Header) + sizeof(Elf64SectionHeader) + 1, 1, TC(""));
 		auto& header = *(Elf64Header*)data;
@@ -247,6 +327,7 @@ namespace uba
 		header.e_ident[4] = 2;
 		header.e_ident[5] = 1;
 		header.e_ident[6] = 1;
+		header.e_ident[7] = platform.Equals(TC("PS4")) ? 9 : 0; // OS ABI identification
 		header.e_type = 1;
 		header.e_machine = EM_X86_64;
 		header.e_version = 1;
diff --git a/Engine/Source/Programs/UnrealBuildAccelerator/Common/Private/UbaObjectFileElf.h b/Engine/Source/Programs/UnrealBuildAccelerator/Common/Private/UbaObjectFileElf.h
--- a/Engine/Source/Programs/UnrealBuildAccelerator/Common/Private/UbaObjectFileElf.h
+++ b/Engine/Source/Programs/UnrealBuildAccelerator/Common/Private/UbaObjectFileElf.h
@@ -12,9 +12,10 @@ namespace uba
 		ObjectFileElf();
 		virtual bool Parse(Logger& logger, const tchar* hint) override;
 
+		static bool CreateExtraFile(Logger& logger, const StringView& platform, MemoryBlock& memoryBlock, const UnorderedSymbols& allNeededImports, const UnorderedSymbols& allSharedImports, const UnorderedExports& allSharedExports, bool includeExportsInFile);
+
 	private:
 		virtual bool StripExports(Logger& logger, u8* newData, const UnorderedSymbols& allNeededImports) override;
-		virtual bool CreateExtraFile(Logger& logger, MemoryBlock& memoryBlock, const UnorderedSymbols& allNeededImports, const UnorderedSymbols& allSharedImports, const UnorderedExports& allSharedExports, bool includeExportsInFile) override;
 
 		UnorderedSymbols m_toRemove;
 
diff --git a/Engine/Source/Programs/UnrealBuildAccelerator/Common/Private/UbaObjectFileImportLib.cpp b/Engine/Source/Programs/UnrealBuildAccelerator/Common/Private/UbaObjectFileImportLib.cpp
--- a/Engine/Source/Programs/UnrealBuildAccelerator/Common/Private/UbaObjectFileImportLib.cpp
+++ b/Engine/Source/Programs/UnrealBuildAccelerator/Common/Private/UbaObjectFileImportLib.cpp
@@ -67,9 +67,4 @@ namespace uba
 	{
 		return logger.Error(TC("Stripping exports from import lib file not supported"));
 	}
-
-	bool ObjectFileImportLib::CreateExtraFile(Logger& logger, MemoryBlock& memoryBlock, const UnorderedSymbols& allNeededImports, const UnorderedSymbols& allSharedImports, const UnorderedExports& allSharedExports, bool includeExportsInFile)
-	{
-		return logger.Error(TC("Creating extra file from import lib file not supported"));
-	}
 }
diff --git a/Engine/Source/Programs/UnrealBuildAccelerator/Common/Private/UbaObjectFileImportLib.h b/Engine/Source/Programs/UnrealBuildAccelerator/Common/Private/UbaObjectFileImportLib.h
--- a/Engine/Source/Programs/UnrealBuildAccelerator/Common/Private/UbaObjectFileImportLib.h
+++ b/Engine/Source/Programs/UnrealBuildAccelerator/Common/Private/UbaObjectFileImportLib.h
@@ -14,7 +14,6 @@ namespace uba
 
 	private:
 		virtual bool StripExports(Logger& logger, u8* newData, const UnorderedSymbols& allNeededImports) override;
-		virtual bool CreateExtraFile(Logger& logger, MemoryBlock& memoryBlock, const UnorderedSymbols& allNeededImports, const UnorderedSymbols& allSharedImports, const UnorderedExports& allSharedExports, bool includeExportsInFile) override;
 
 		std::string m_libName;
 	};
diff --git a/Engine/Source/Programs/UnrealBuildAccelerator/Common/Private/UbaObjectFileLLVMIR.cpp b/Engine/Source/Programs/UnrealBuildAccelerator/Common/Private/UbaObjectFileLLVMIR.cpp
--- a/Engine/Source/Programs/UnrealBuildAccelerator/Common/Private/UbaObjectFileLLVMIR.cpp
+++ b/Engine/Source/Programs/UnrealBuildAccelerator/Common/Private/UbaObjectFileLLVMIR.cpp
@@ -804,19 +804,47 @@ namespace uba
 		return ~0u;
 	}
 
+	u32 GetLinkageIndex(u32 code)
+	{
+		if (code == MODULE_CODE_FUNCTION)
+			return 5;
+		if (code == MODULE_CODE_GLOBALVAR)
+			return 5;
+		if (code == MODULE_CODE_ALIAS)
+			return 3;
+		if (code == MODULE_CODE_ALIAS_OLD)
+			return 2;
+		UBA_ASSERTF(false, TC("module code %s not supported"), ModuleCodeToString(code));
+		return ~0u;
+	}
+
+	u32 GetVisibilityIndex(u32 code)
+	{
+		if (code == MODULE_CODE_FUNCTION)
+			return 9;
+		if (code == MODULE_CODE_GLOBALVAR)
+			return 10;
+		if (code == MODULE_CODE_ALIAS)
+			return 8;
+		if (code == MODULE_CODE_ALIAS_OLD)
+			return 7;
+		UBA_ASSERTF(false, TC("module code %s not supported"), ModuleCodeToString(code));
+		return ~0u;
+	}
+
 	bool ObjectFileLLVMIR::Parse(Logger& logger, const tchar* hint)
 	{
-		//BinaryReader reader(m_data, 0, m_dataSize);
 		BitStreamReader reader(*this, logger, m_data, m_dataSize);
 
 		if (IsWrappedBitcode(m_data, m_dataSize))
 		{
-			UBA_ASSERT(false);
-			//reader.Skip(4);
-			//u32 version = reader.ReadU32();(void)version;
-			//u32 offset = reader.ReadU32();(void)offset;
-			//u32 size = reader.ReadU32();(void)size;
-			//u32 cpuType = reader.ReadU32();(void)cpuType;
+			reader.Read(32);
+			u32 version = reader.Read(32);(void)version;
+			u32 bitcodeOffset = reader.Read(32);(void)bitcodeOffset;
+			u32 bitcodeSize = reader.Read(32);(void)bitcodeSize;
+			u32 cpuType = reader.Read(32);(void)cpuType;
+			reader.JumpToBit(bitcodeOffset*8);
+			reader.m_end = reader.m_pos + bitcodeSize;
 		}
 
 		reader.Read(32); // Skip magic
@@ -879,6 +907,38 @@ namespace uba
 				record.isExport = true;
 				m_exports.emplace(name, ExportInfo{"", index++});
 			}
+			else // Check visibility... might be a non-windows based llvm ir stream file
+			{
+				/*
+				u64 visibility = ~0ull;
+				u32 visibilityIndex = GetVisibilityIndex(record.code);
+				bool hidden = false;
+				if (recSize > visibilityIndex)
+				{
+					visibility = recIt[visibilityIndex];
+					if (visibility != 0)
+						hidden = true;
+				}
+				*/
+
+				u64 linkage = ~0ull;
+				u32 linkageIndex = GetLinkageIndex(record.code);
+				if (recSize > linkageIndex)
+				{
+					linkage = recIt[linkageIndex];
+
+					//if (strncmp(name.data(), "llvm.umax", 9) == 0)
+					//	printf("");
+
+					if (linkage == 0 || linkage == 5 || linkage == 6 || linkage == 15)
+					{
+						record.isExport = true;
+						m_exports.emplace(name, ExportInfo{"", index++});
+					}
+					else if (linkage == 1 || linkage == 16)
+						m_imports.emplace(name);
+				}
+			}
 		}
 		return true;
 	}
@@ -1056,9 +1116,4 @@ namespace uba
 
 		return true;
 	}
-	
-	bool ObjectFileLLVMIR::CreateExtraFile(Logger& logger, MemoryBlock& memoryBlock, const UnorderedSymbols& allNeededImports, const UnorderedSymbols& allSharedImports, const UnorderedExports& allSharedExports, bool includeExportsInFile)
-	{
-		return ObjectFileCoff::CreateExtraFile2(logger, memoryBlock, allNeededImports, allSharedImports, allSharedExports, includeExportsInFile);
-	}
 }
diff --git a/Engine/Source/Programs/UnrealBuildAccelerator/Common/Private/UbaObjectFileLLVMIR.h b/Engine/Source/Programs/UnrealBuildAccelerator/Common/Private/UbaObjectFileLLVMIR.h
--- a/Engine/Source/Programs/UnrealBuildAccelerator/Common/Private/UbaObjectFileLLVMIR.h
+++ b/Engine/Source/Programs/UnrealBuildAccelerator/Common/Private/UbaObjectFileLLVMIR.h
@@ -16,7 +16,6 @@ namespace uba
 
 	private:
 		virtual bool StripExports(Logger& logger, u8* newData, const UnorderedSymbols& allNeededImports) override;
-		virtual bool CreateExtraFile(Logger& logger, MemoryBlock& memoryBlock, const UnorderedSymbols& allNeededImports, const UnorderedSymbols& allSharedImports, const UnorderedExports& allSharedExports, bool includeExportsInFile) override;
 
 		enum BlockInfoCodes : u8;
 		enum Encoding : u8;
diff --git a/Engine/Source/Programs/UnrealBuildAccelerator/Common/Public/UbaObjectFile.h b/Engine/Source/Programs/UnrealBuildAccelerator/Common/Public/UbaObjectFile.h
--- a/Engine/Source/Programs/UnrealBuildAccelerator/Common/Public/UbaObjectFile.h
+++ b/Engine/Source/Programs/UnrealBuildAccelerator/Common/Public/UbaObjectFile.h
@@ -39,7 +39,7 @@ namespace uba
 		const UnorderedExports& GetExports() const;
 		const UnorderedSymbols& GetPotentialDuplicates() const;
 
-		static bool CreateExtraFile(Logger& logger, const tchar* extraObjFilename, ObjectFileType type, const UnorderedSymbols& allNeededImports, const UnorderedSymbols& allSharedImports, const UnorderedExports& allSharedExports, bool includeExportsInFile);
+		static bool CreateExtraFile(Logger& logger, const StringView& extraObjFilename, const StringView& platform, const UnorderedSymbols& allNeededImports, const UnorderedSymbols& allSharedImports, const UnorderedExports& allSharedExports, bool includeExportsInFile);
 
 		virtual ~ObjectFile();
 
@@ -50,7 +50,6 @@ namespace uba
 	protected:
 		virtual bool Parse(Logger& logger, const tchar* hint) = 0;
 		virtual bool StripExports(Logger& logger, u8* newData, const UnorderedSymbols& allNeededImports) = 0;
-		virtual bool CreateExtraFile(Logger& logger, MemoryBlock& memoryBlock, const UnorderedSymbols& allNeededImports, const UnorderedSymbols& allSharedImports, const UnorderedExports& allSharedExports, bool includeExportsInFile) = 0;
 
 		FileAccessor* m_file = nullptr;
 		u8* m_data = nullptr;
@@ -63,7 +62,7 @@ namespace uba
 		UnorderedSymbols m_potentialDuplicates;
 
 	public:
-		struct StringView
+		struct AnsiStringView
 		{
 			const char* strBegin;
 			const char* strEnd;
diff --git a/Engine/Source/Programs/UnrealBuildAccelerator/Core/Public/UbaStringBuffer.h b/Engine/Source/Programs/UnrealBuildAccelerator/Core/Public/UbaStringBuffer.h
--- a/Engine/Source/Programs/UnrealBuildAccelerator/Core/Public/UbaStringBuffer.h
+++ b/Engine/Source/Programs/UnrealBuildAccelerator/Core/Public/UbaStringBuffer.h
@@ -98,6 +98,7 @@ namespace uba
 		bool IsEmpty() const { return count == 0; }
 		bool EndsWith(const tchar* value, bool ignoreCase = true) const { return uba::EndsWith(data, count, value, ignoreCase); }
 		bool Contains(const tchar* str, bool ignoreCase = true) const { return uba::Contains(data, str, ignoreCase); }
+		bool Equals(const tchar* str, bool ignoreCase = true) const { return uba::Equals(data, str, ignoreCase); }
 
 		const tchar* data;
 		u32 count;
diff --git a/Engine/Source/Programs/UnrealBuildAccelerator/ObjTool/Private/UbaObjTool.cpp b/Engine/Source/Programs/UnrealBuildAccelerator/ObjTool/Private/UbaObjTool.cpp
--- a/Engine/Source/Programs/UnrealBuildAccelerator/ObjTool/Private/UbaObjTool.cpp
+++ b/Engine/Source/Programs/UnrealBuildAccelerator/ObjTool/Private/UbaObjTool.cpp
@@ -79,6 +79,7 @@ namespace uba
 		Vector<TString> objFilesForImpLib;
 		std::string impLibName;
 		TString impLibFile;
+		TString platform;
 
 		auto parseArg = [&](const tchar* arg, bool isRsp)
 			{
@@ -169,6 +170,10 @@ namespace uba
 				{
 					extraObjFile = value.data;
 				}
+				else if (name.StartsWith(TC("/T")))
+				{
+					platform = value.data;
+				}
 				else if (name.Equals(TC("-printsymbols")))
 				{
 					printSymbols = true;
@@ -241,8 +246,6 @@ namespace uba
 
 		if (!objFilesToStrip.empty())
 		{
-			ObjectFileType type = ObjectFileType_Unknown;
-
 			CriticalSection cs;
 			Atomic<bool> success = true;
 			UnorderedSymbols allNeededImports; // Imports needed from the outside of the stripped obj files
@@ -263,8 +266,6 @@ namespace uba
 						return;
 					}
 					ScopedCriticalSection _(cs);
-					if (type == ObjectFileType_Unknown)
-						type = symbolFile.type;
 					allNeededImports.insert(symbolFile.imports.begin(), symbolFile.imports.end());
 				});
 			if (!success)
@@ -286,8 +287,6 @@ namespace uba
 						return;
 					}
 					ScopedCriticalSection _(cs);
-					if (type == ObjectFileType_Unknown)
-						type = symbolFile.type;
 					allSharedImports.insert(symbolFile.imports.begin(), symbolFile.imports.end());
 					allSharedExports.insert(symbolFile.exports.begin(), symbolFile.exports.end());
 				});
@@ -295,7 +294,7 @@ namespace uba
 				return -1;
 
 			if (!extraObjFile.empty())
-				if (!ObjectFile::CreateExtraFile(logger, extraObjFile.c_str(), type, allNeededImports, allSharedImports, allSharedExports, true))
+				if (!ObjectFile::CreateExtraFile(logger, extraObjFile, platform, allNeededImports, allSharedImports, allSharedExports, true))
 					return -1;
 
 			//logger.Info(TC("Reduced export count from %llu to %llu"), totalExportCount.load(), totalKeptExportCount.size());
@@ -323,7 +322,7 @@ namespace uba
 					return;
 				if (ObjectFile* objectFile = ObjectFile::OpenAndParse(logger, fixedPath.data))
 				{
-					objectFile->RemoveExportedSymbol("DllMain");
+					//objectFile->RemoveExportedSymbol("DllMain"); // This is used to patch xinput.lib
 					objFiles[it - objFilesForImpLib.begin()] = objectFile;
 				}
 				else
