[UBA]
* More optimizations in PopulateCasFromDirs

diff --git a/Engine/Source/Programs/UnrealBuildAccelerator/Common/Private/UbaStorageClient.cpp b/Engine/Source/Programs/UnrealBuildAccelerator/Common/Private/UbaStorageClient.cpp
--- a/Engine/Source/Programs/UnrealBuildAccelerator/Common/Private/UbaStorageClient.cpp
+++ b/Engine/Source/Programs/UnrealBuildAccelerator/Common/Private/UbaStorageClient.cpp
@@ -856,50 +856,56 @@ namespace uba
 					return;
 				}
 
-				workManager.AddWork([&, filePath = TString(fullPath.data)]()
+				StringBuffer<> forKey;
+				FixPath(fullPath.data, nullptr, 0, forKey);
+				if (CaseInsensitiveFs)
+					forKey.MakeLower();
+				StringKey fileNameKey = ToStringKey(forKey);
+				FileEntry& fileEntry = GetOrCreateFileEntry(fileNameKey);
+				fileEntry.lock.EnterWrite();
+				if (e.size == fileEntry.size && e.lastWritten == fileEntry.lastWritten)
+				{
+					fileEntry.verified = true;
+					fileEntry.lock.LeaveWrite();
+
+					SCOPED_WRITE_LOCK(m_localStorageFilesLock, lookupLock);
+					auto insres = m_localStorageFiles.try_emplace(fileEntry.casKey);
+					LocalFile& localFile = insres.first->second;
+					if (insres.second)
+					{
+						localFile.casEntry.size = e.size;
+						localFile.fileName = fullPath.data;
+					}
+					return;
+				}
+
+				workManager.AddWork([&, fe = &fileEntry, lw = e.lastWritten, s = e.size, filePath = TString(fullPath.data)]()
 					{
+						auto feLockLeave = MakeGuard([fe]() { fe->lock.LeaveWrite(); });
+
 						if (shouldExit && shouldExit())
 							return;
 
-						FileInformation info;
-						if (!GetFileInformation(info, m_logger, filePath.c_str()))
+						CasKey casKey;
+						if (!CalculateCasKey(casKey, filePath.c_str()))
 						{
-							m_logger.Error(TC("Failed to get information for file %s"), filePath.c_str());
+							m_logger.Error(TC("Failed to calculate cas key for %s"), filePath.c_str());
 							return;
 						}
-
-						StringBuffer<> forKey;
-						FixPath(filePath.c_str(), nullptr, 0, forKey);
-						if (CaseInsensitiveFs)
-							forKey.MakeLower();
-						StringKey fileNameKey = ToStringKey(forKey);
-
-
-						FileEntry& fileEntry = GetOrCreateFileEntry(fileNameKey);
-						SCOPED_WRITE_LOCK(fileEntry.lock, fileEntryLock);
-
-						if (info.size != fileEntry.size || info.lastWriteTime != fileEntry.lastWritten)
-						{
-							CasKey casKey;
-							if (!CalculateCasKey(casKey, filePath.c_str()))
-							{
-								m_logger.Error(TC("Failed to calculate cas key for %s"), filePath.c_str());
-								return;
-							}
-							fileEntry.size = info.size;
-							fileEntry.lastWritten = info.lastWriteTime;
-							fileEntry.casKey = AsCompressed(casKey, false);
-						}
-						fileEntry.verified = true;
-						fileEntryLock.Leave();
+						fe->size = s;
+						fe->lastWritten = lw;
+						fe->casKey = AsCompressed(casKey, false);
+						fe->verified = true;
+						feLockLeave.Execute();
 
 						SCOPED_WRITE_LOCK(m_localStorageFilesLock, lookupLock);
-						auto insres = m_localStorageFiles.try_emplace(fileEntry.casKey);
-						if (!insres.second)
-							return;
+						auto insres = m_localStorageFiles.try_emplace(casKey);
 						LocalFile& localFile = insres.first->second;
-						localFile.casEntry.size = info.size;
-						localFile.fileName = filePath;
+						if (insres.second)
+						{
+							localFile.casEntry.size = s;
+							localFile.fileName = filePath;
+						}
 
 					}, 1, TC(""));
 			});
