[UBA]
* Fixed so caskey is stored as not compressed when populating cas from dir

diff --git a/Engine/Source/Programs/UnrealBuildAccelerator/Common/Private/UbaStorageClient.cpp b/Engine/Source/Programs/UnrealBuildAccelerator/Common/Private/UbaStorageClient.cpp
--- a/Engine/Source/Programs/UnrealBuildAccelerator/Common/Private/UbaStorageClient.cpp
+++ b/Engine/Source/Programs/UnrealBuildAccelerator/Common/Private/UbaStorageClient.cpp
@@ -899,7 +899,7 @@ namespace uba
 						feLockLeave.Execute();
 
 						SCOPED_WRITE_LOCK(m_localStorageFilesLock, lookupLock);
-						auto insres = m_localStorageFiles.try_emplace(casKey);
+						auto insres = m_localStorageFiles.try_emplace(fe->casKey);
 						LocalFile& localFile = insres.first->second;
 						if (insres.second)
 						{
