[UBA]
* Fix for mistake where file lookup db was populated with compressed caskeys

diff --git a/Engine/Source/Programs/UnrealBuildAccelerator/Common/Private/UbaStorageClient.cpp b/Engine/Source/Programs/UnrealBuildAccelerator/Common/Private/UbaStorageClient.cpp
--- a/Engine/Source/Programs/UnrealBuildAccelerator/Common/Private/UbaStorageClient.cpp
+++ b/Engine/Source/Programs/UnrealBuildAccelerator/Common/Private/UbaStorageClient.cpp
@@ -835,6 +835,9 @@ namespace uba
 
 	bool StorageClient::PopulateCasFromDirsRecursive(const tchar* dir, WorkManager& workManager, UnorderedSet<u64>& seenIds, ReaderWriterLock& seenIdsLock, const Function<bool()>& shouldExit)
 	{
+		if (shouldExit && shouldExit())
+			return true;
+
 		StringBuffer<> fullPath;
 		fullPath.Append(dir).EnsureEndsWithSlash();
 		u32 dirLen = fullPath.count;
@@ -849,8 +852,6 @@ namespace uba
 					lock.Leave();
 					workManager.AddWork([&, filePath = TString(fullPath.data)]()
 						{
-							if (shouldExit && shouldExit())
-								return;
 							PopulateCasFromDirsRecursive(filePath.c_str(), workManager, seenIds, seenIdsLock, shouldExit);
 						}, 1, TC(""));
 					return;
@@ -866,6 +867,7 @@ namespace uba
 				if (e.size == fileEntry.size && e.lastWritten == fileEntry.lastWritten)
 				{
 					fileEntry.verified = true;
+					fileEntry.casKey = AsCompressed(fileEntry.casKey, false); // TODO: Remove this when machines have flushed their db
 					fileEntry.lock.LeaveWrite();
 
 					SCOPED_WRITE_LOCK(m_localStorageFilesLock, lookupLock);
