[UBA]
* Fixed exit time stats for remote processes

diff --git a/Engine/Source/Programs/UnrealBuildAccelerator/Common/Private/UbaSessionClient.cpp b/Engine/Source/Programs/UnrealBuildAccelerator/Common/Private/UbaSessionClient.cpp
--- a/Engine/Source/Programs/UnrealBuildAccelerator/Common/Private/UbaSessionClient.cpp
+++ b/Engine/Source/Programs/UnrealBuildAccelerator/Common/Private/UbaSessionClient.cpp
@@ -1384,12 +1384,20 @@ namespace uba
 		writer.WriteU32(CountLogLines(process));
 		WriteLogLines(writer, process);
 
+		// This is normally set after callback so we need to calculate it here
+		auto& exitTime = process.m_processStats.exitTime;
+		auto oldExitTime = exitTime.load();
+		if (exitTime)
+			exitTime = GetTime() - exitTime;
+
 		// Must be written last
 		process.m_processStats.Write(writer);
 		process.m_sessionStats.Write(writer);
 		process.m_storageStats.Write(writer);
 		process.m_kernelStats.Write(writer);
 
+		exitTime = oldExitTime;
+
 		StackBinaryReader<16> reader;
 		if (!msg.Send(reader, m_stats.procFinishedMsg) && m_loop)
 			return m_logger.Error(TC("Failed to send ProcessFinished message!"));
