[UBA]
* Fixed so cvtres.exe also preloads bcrypt.dll in order to prevent deadlock that can happen because of mimalloc

diff --git a/Engine/Source/Programs/UnrealBuildAccelerator/Core/Private/UbaApplicationRules.cpp b/Engine/Source/Programs/UnrealBuildAccelerator/Core/Private/UbaApplicationRules.cpp
--- a/Engine/Source/Programs/UnrealBuildAccelerator/Core/Private/UbaApplicationRules.cpp
+++ b/Engine/Source/Programs/UnrealBuildAccelerator/Core/Private/UbaApplicationRules.cpp
@@ -96,7 +96,7 @@ namespace uba
 		}
 	};
 
-	class ApplicationRulesLinkExe : public ApplicationRulesVC
+	class ApplicationRulesVcLink : public ApplicationRulesVC
 	{
 		using Super = ApplicationRulesVC;
 	public:
@@ -160,9 +160,27 @@ namespace uba
 		}
 	};
 
-	class ApplicationRulesLldLinkExe : public ApplicationRulesLinkExe
+	class ApplicationRulesLinkExe : public ApplicationRulesVcLink
 	{
-		using Super = ApplicationRulesLinkExe;
+		using Super = ApplicationRulesVcLink;
+
+		virtual const tchar* const* LibrariesToPreload() const override
+		{
+			// Special handling.. it seems loading bcrypt.dll can deadlock when using mimalloc so we make sure to load it here directly instead
+			// There is a setting to disable bcrypt dll loading inside mimalloc but with that change mimalloc does not work with older versions of windows
+			static constexpr const tchar* preloads[] = 
+			{
+				TC("bcrypt.dll"),
+				TC("bcryptprimitives.dll"),
+				nullptr,
+			};
+			return preloads;
+		}
+	};
+
+	class ApplicationRulesLldLinkExe : public ApplicationRulesVcLink
+	{
+		using Super = ApplicationRulesVcLink;
 
 		virtual bool KeepInMemory(const StringView& fileName, const tchar* systemTemp, bool isRunningRemote) const override
 		{
@@ -489,10 +507,10 @@ namespace uba
 			{ TC(""),							new ApplicationRules() },		// Must be index 0
 			{ TC("cl.exe"),						new ApplicationRulesClExe() },	// Must be index 1
 			{ TC("link.exe"),					new ApplicationRulesLinkExe() }, // Must be index 2
-			{ TC("lib.exe"),					new ApplicationRulesLinkExe() },
+			{ TC("lib.exe"),					new ApplicationRulesVcLink() },
 			{ TC("cvtres.exe"),					new ApplicationRulesLinkExe() },
-			{ TC("mt.exe"),						new ApplicationRulesLinkExe() },
-			{ TC("rc.exe"),						new ApplicationRulesLinkExe() },
+			{ TC("mt.exe"),						new ApplicationRulesVcLink() },
+			{ TC("rc.exe"),						new ApplicationRulesVcLink() },
 			{ TC("lld-link.exe"),				new ApplicationRulesLldLinkExe() },
 			{ TC("clang++.exe"),				new ApplicationRulesClangPlusPlusExe() },
 			{ TC("clang-cl.exe"),				new ApplicationRulesClangPlusPlusExe() },
diff --git a/Engine/Source/Programs/UnrealBuildAccelerator/Core/Public/UbaApplicationRules.h b/Engine/Source/Programs/UnrealBuildAccelerator/Core/Public/UbaApplicationRules.h
--- a/Engine/Source/Programs/UnrealBuildAccelerator/Core/Public/UbaApplicationRules.h
+++ b/Engine/Source/Programs/UnrealBuildAccelerator/Core/Public/UbaApplicationRules.h
@@ -144,6 +144,11 @@ namespace uba
 			return false;
 		}
 
+		virtual const tchar* const* LibrariesToPreload() const // Array should be null terminated
+		{
+			return nullptr;
+		}
+
 		u32 index = ~0u;
 	};
 
diff --git a/Engine/Source/Programs/UnrealBuildAccelerator/Detours/Private/Windows/UbaDetoursFunctionsWin.cpp b/Engine/Source/Programs/UnrealBuildAccelerator/Detours/Private/Windows/UbaDetoursFunctionsWin.cpp
--- a/Engine/Source/Programs/UnrealBuildAccelerator/Detours/Private/Windows/UbaDetoursFunctionsWin.cpp
+++ b/Engine/Source/Programs/UnrealBuildAccelerator/Detours/Private/Windows/UbaDetoursFunctionsWin.cpp
@@ -759,15 +759,10 @@ void PreInit(const DetoursPayload& payload)
 			FatalError(1349, L"Failed to reserve memory for cl.exe (%u)", GetLastError());
 	}
 
-	// Special link.exe handling.. it seems loading bcrypt.dll can deadlock when using mimalloc so we make sure to load it here directly instead
-	// There is a setting to disable bcrypt dll loading inside mimalloc but with that change mimalloc does not work with older versions of windows
-	if (payload.rulesIndex == 2)
-	{
-		if (!LoadLibraryExW(L"bcrypt.dll", NULL, LOAD_LIBRARY_SEARCH_SYSTEM32))
-			FatalError(1351, L"Failed to load bcrypt.dll (%u)", GetLastError());
-		if (!LoadLibraryExW(L"bcryptprimitives.dll", NULL, LOAD_LIBRARY_SEARCH_SYSTEM32))
-			FatalError(1352, L"Failed to load bcryptprimitives.dll (%u)", GetLastError());
-	}
+	if (const tchar* const* preloads = g_rules->LibrariesToPreload())
+		for (auto it = preloads; *it; ++it)
+			if (!LoadLibraryExW(*it, NULL, LOAD_LIBRARY_SEARCH_SYSTEM32))
+				FatalError(1351, L"Failed to preload %s (%u)", *it, GetLastError());
 }
 
 void Init(const DetoursPayload& payload, u64 startTime)
