[UBA]
* Added UBA_DETOURED=1 environment variable to environment for detoured processes so applications can query if they are detoured or not

diff --git a/Engine/Source/Programs/UnrealBuildAccelerator/Common/Private/UbaSession.cpp b/Engine/Source/Programs/UnrealBuildAccelerator/Common/Private/UbaSession.cpp
--- a/Engine/Source/Programs/UnrealBuildAccelerator/Common/Private/UbaSession.cpp
+++ b/Engine/Source/Programs/UnrealBuildAccelerator/Common/Private/UbaSession.cpp
@@ -1814,6 +1814,9 @@ namespace uba
 		}
 		AddEnvironmentVariableNoLock("TMPDIR", m_tempPath.data);
 #endif
+
+		AddEnvironmentVariableNoLock(TC("UBA_DETOURED"), TC("1"));
+
 		m_environmentVariables.push_back(0);
 		return m_environmentVariables.data();
 	}
