[UBA]
* Fixed bug with file mapping for when ld-lld.exe generated really big pdb files

diff --git a/Engine/Source/Programs/UnrealBuildAccelerator/Core/Private/UbaApplicationRules.cpp b/Engine/Source/Programs/UnrealBuildAccelerator/Core/Private/UbaApplicationRules.cpp
--- a/Engine/Source/Programs/UnrealBuildAccelerator/Core/Private/UbaApplicationRules.cpp
+++ b/Engine/Source/Programs/UnrealBuildAccelerator/Core/Private/UbaApplicationRules.cpp
@@ -9,6 +9,7 @@ namespace uba
 	{
 		using Super = ApplicationRules;
 
+	public:
 		virtual bool AllowDetach() const override
 		{
 			return true;
@@ -187,6 +188,13 @@ namespace uba
 			return fileName.EndsWith(TC(".manifest")) && fileName.Contains(systemTemp);
 		}
 
+		virtual u64 FileTypeMaxSize(const StringBufferBase& file, bool isSystemOrTempFile) const override
+		{
+			if (file.Contains(TC(".pdb.tmp")))
+				return 14ull * 1024 * 1024 * 1024; // This is ridiculous
+			return ApplicationRulesVcLink::FileTypeMaxSize(file, isSystemOrTempFile);
+		}
+
 		virtual bool IsOutputFile(const StringView& fileName) const override
 		{
 			return false//Super::IsOutputFile(fileName)
