[UBA]
* Fixed bug where too many symbols were exported in importlib created from llvm ir obj files

diff --git a/Engine/Source/Programs/UnrealBuildAccelerator/Common/Private/UbaObjectFileLLVMIR.cpp b/Engine/Source/Programs/UnrealBuildAccelerator/Common/Private/UbaObjectFileLLVMIR.cpp
--- a/Engine/Source/Programs/UnrealBuildAccelerator/Common/Private/UbaObjectFileLLVMIR.cpp
+++ b/Engine/Source/Programs/UnrealBuildAccelerator/Common/Private/UbaObjectFileLLVMIR.cpp
@@ -927,10 +927,10 @@ namespace uba
 				{
 					linkage = recIt[linkageIndex];
 
-					//if (strncmp(name.data(), "llvm.umax", 9) == 0)
+					//if (strncmp(name.data(), "_Init_thread_footer", 18) == 0)
 					//	printf("");
 
-					if (linkage == 0 || linkage == 5 || linkage == 6 || linkage == 15)
+					if (linkage == 5 || linkage == 6 || linkage == 15)
 					{
 						record.isExport = true;
 						m_exports.emplace(name, ExportInfo{"", index++});
