[UbaObjTool] Terminate export symbol names at the end of a section

Some object files do not end their directive sections with nul terminators or whitespace, so without this extra check UbaObjTool can produce corrupt import libs.

Thank you Russel for finding and fixing!

diff --git a/Engine/Source/Programs/UnrealBuildAccelerator/Common/Private/UbaObjectFileCoff.cpp b/Engine/Source/Programs/UnrealBuildAccelerator/Common/Private/UbaObjectFileCoff.cpp
--- a/Engine/Source/Programs/UnrealBuildAccelerator/Common/Private/UbaObjectFileCoff.cpp
+++ b/Engine/Source/Programs/UnrealBuildAccelerator/Common/Private/UbaObjectFileCoff.cpp
@@ -277,6 +277,7 @@ namespace uba
 
 		auto directiveSection = (ImageSectionHeader*)(m_data + m_info.directiveSectionMemOffset);
 		u8* directiveData = m_data + directiveSection->PointerToRawData;
+		u8* directiveEnd = directiveData + directiveSection->SizeOfRawData;
 
 		static constexpr u8 utf8Bom[3] = { 0xef, 0xbb, 0xbf };
 		UBA_ASSERT(memcmp(directiveData, utf8Bom, 3) != 0);
@@ -287,6 +288,7 @@ namespace uba
 		u32 index = 0;
 
 		auto str = (char*)directiveData;
+		auto strEnd = (char*)directiveEnd;
 		while (str)
 		{
 			char* exportStr = strstr(str, "/EXPORT:");
@@ -302,10 +304,10 @@ namespace uba
 			}
 			else
 			{
-				exportEnd = strchr(exportStr, ' ');
+				exportEnd = (char*)memchr(exportStr, ' ', strEnd - exportStr);
 				str = exportEnd;
 				if (!exportEnd)
-					exportEnd = exportStr + strlen(exportStr);
+					exportEnd = strEnd;
 				else
 					++str;
 			}
