[UBA]
* Fixed so ubacacheserver increase expiration time for bucket if cas bucket size gets more than 1mb under max size

diff --git a/Engine/Source/Programs/UnrealBuildAccelerator/Common/Private/UbaCacheServer.cpp b/Engine/Source/Programs/UnrealBuildAccelerator/Common/Private/UbaCacheServer.cpp
--- a/Engine/Source/Programs/UnrealBuildAccelerator/Common/Private/UbaCacheServer.cpp
+++ b/Engine/Source/Programs/UnrealBuildAccelerator/Common/Private/UbaCacheServer.cpp
@@ -10,6 +10,8 @@
 #include "UbaStorageServer.h"
 //#include <oodle2.h>
 
+// MAKE IT RECOVER TIME WHEN GETTING UNDER CERTAIN SIZE!
+
 #if PLATFORM_WINDOWS
 #define UBA_FORCEINLINE __forceinline
 #else
@@ -181,6 +183,7 @@ namespace uba
 
 		m_maintenanceReserveSize = info.maintenanceReserveSize;
 		m_expirationTimeSeconds = info.expirationTimeSeconds;
+		m_bucketCasTableMaxSize = info.bucketCasTableMaxSize;
 
 		m_rootDir.count = GetFullPathNameW(info.rootDir, m_rootDir.capacity, m_rootDir.data, NULL);
 		m_rootDir.Replace('/', PathSeparator).EnsureEndsWithSlash();
@@ -705,15 +708,22 @@ namespace uba
 				Vector<CasKey> keysToErase;
 
 				u64 lastUseTimeLimit = 0; // This is the time relative to server startup time
-				if (bucket.expirationTimeSeconds)
+				if (bucket.expirationTimeSeconds && m_bucketCasTableMaxSize)
 				{
-					// If cas key table size is over 30mb we need to shorten expiration time
-					if (deleteIteration == 0 && bucket.m_casKeyTable.GetSize() > 30*1024*1024)
+					if (deleteIteration == 0)
 					{
-						m_logger.Info(TC("Lowered expiration time for bucket %u to %s"), bucket.index, TimeToText(MsToTime(bucket.expirationTimeSeconds*1000), true).str);
-						bucket.expirationTimeSeconds -= 60*60; // Shorten by one hour.. this will be reset 
+						u64 bucketCasTableSize = bucket.m_casKeyTable.GetSize();
+						u64 oldExpirationTime = bucket.expirationTimeSeconds;
+						if (bucketCasTableSize >= m_bucketCasTableMaxSize)
+							bucket.expirationTimeSeconds -= 60*60; // Decreased by one hour
+						else if (bucket.expirationTimeSeconds < m_expirationTimeSeconds && (bucketCasTableSize + 1ull*1024*1024) < m_bucketCasTableMaxSize)
+							bucket.expirationTimeSeconds += 60*60; // Increased by one hour
+						if (oldExpirationTime != bucket.expirationTimeSeconds)
+							m_logger.Info(TC("%s expiration time for bucket %u to %s"), oldExpirationTime < bucket.expirationTimeSeconds ? TC("Increased") : TC("Decreased"), bucket.index, TimeToText(MsToTime(bucket.expirationTimeSeconds*1000), true).str);
 					}
-					if (GetFileTimeAsSeconds(now - m_creationTime) > bucket.expirationTimeSeconds)
+
+					u64 secondsRunning = GetFileTimeAsSeconds(now - m_creationTime);
+					if (secondsRunning > bucket.expirationTimeSeconds)
 						lastUseTimeLimit = (now - m_creationTime) - GetSecondsAsFileTime(bucket.expirationTimeSeconds);
 				}
 
