[UBA]
* missed file in previous submit

diff --git a/Engine/Source/Programs/UnrealBuildAccelerator/Common/Public/UbaCacheServer.h b/Engine/Source/Programs/UnrealBuildAccelerator/Common/Public/UbaCacheServer.h
--- a/Engine/Source/Programs/UnrealBuildAccelerator/Common/Public/UbaCacheServer.h
+++ b/Engine/Source/Programs/UnrealBuildAccelerator/Common/Public/UbaCacheServer.h
@@ -38,6 +38,9 @@ namespace uba
 
 		// The amount of reserved memory used per core when doing maintenance
 		u64 maintenanceReserveSize = 128ull * 1024 * 1024;
+
+		// Max size of cas bucket. When within 2mb it will start decreasing expiry time by one hour
+		u64 bucketCasTableMaxSize = 32ull * 1024 * 1024;
 	};
 
 	class CacheServer
@@ -98,6 +101,7 @@ namespace uba
 		Atomic<bool> m_shutdownRequested = false;
 
 		u64 m_maintenanceReserveSize = 0;
+		u64 m_bucketCasTableMaxSize = 0;
 		u64 m_creationTime = 0;
 		u64 m_bootTime = 0;
 		u64 m_lastMaintenance = 0;
