[UBA]
* Fixed crash that could happen when very large error messages were added to the trace stream

diff --git a/Engine/Source/Programs/UnrealBuildAccelerator/Common/Private/UbaTrace.cpp b/Engine/Source/Programs/UnrealBuildAccelerator/Common/Private/UbaTrace.cpp
--- a/Engine/Source/Programs/UnrealBuildAccelerator/Common/Private/UbaTrace.cpp
+++ b/Engine/Source/Programs/UnrealBuildAccelerator/Common/Private/UbaTrace.cpp
@@ -46,8 +46,10 @@ namespace uba
 
 		bool EnsureMemory(u64 size)
 		{
-			if (m_isValid)
-				m_isValid = m_trace.EnsureMemory(size);
+			if (!m_isValid)
+				return false;
+			m_trace.m_memoryPos = GetPosition();
+			m_isValid = m_trace.EnsureMemory(size);
 			return m_isValid;
 		}
 
@@ -273,7 +275,7 @@ namespace uba
 		{
 			if (lineCounter++ == 100) // We don't want to write the entire error in the trace stream to blow the entire buffer
 				break;
-			if (!writer.EnsureMemory(1 + (line.text.size()+1)*sizeof(tchar)))
+			if (!writer.EnsureMemory(2 + (line.text.size()+2)*sizeof(tchar)))
 				return;
 			writer.WriteByte(line.type);
 			writer.WriteString(line.text);
