[UBA]
* Added zones to some special computer names. Will be removed once environment variables are set

diff --git a/Engine/Source/Programs/UnrealBuildAccelerator/Agent/Private/UbaAgent.cpp b/Engine/Source/Programs/UnrealBuildAccelerator/Agent/Private/UbaAgent.cpp
--- a/Engine/Source/Programs/UnrealBuildAccelerator/Agent/Private/UbaAgent.cpp
+++ b/Engine/Source/Programs/UnrealBuildAccelerator/Agent/Private/UbaAgent.cpp
@@ -758,7 +758,7 @@ namespace uba
 		#endif
 
 		if (!zone.count)
-			zone.count = GetEnvironmentVariableW(TC("UBA_ZONE"), zone.data, zone.capacity);
+			GetZone(zone);
 
 		if (zone.count)
 			extraInfo.Append(TC(", ")).Append(zone);
diff --git a/Engine/Source/Programs/UnrealBuildAccelerator/Common/Private/UbaSession.cpp b/Engine/Source/Programs/UnrealBuildAccelerator/Common/Private/UbaSession.cpp
--- a/Engine/Source/Programs/UnrealBuildAccelerator/Common/Private/UbaSession.cpp
+++ b/Engine/Source/Programs/UnrealBuildAccelerator/Common/Private/UbaSession.cpp
@@ -2875,4 +2875,27 @@ namespace uba
 		if (counterSuffix)
 			out.Appendf(TC("_%03u"), counterSuffix);
 	}
+
+	bool GetZone(StringBufferBase& outZone)
+	{
+		outZone.count = GetEnvironmentVariableW(TC("UBA_ZONE"), outZone.data, outZone.capacity);
+		if (outZone.count)
+			return true;
+
+		// TODO: Remove.
+		#if PLATFORM_MAC
+		if (!GetComputerName(outZone.data, outZone.capacity))
+			return false;
+
+		outZone.count = TStrlen(outZone.data);
+		if (outZone.StartsWith(TC("dc4-mac")) || outZone.StartsWith(TC("rdu-mac")))
+		{
+			outZone.Resize(7);
+			return true;
+		}
+		outZone.count = 0;
+		#endif
+
+		return false;
+	}
 }
diff --git a/Engine/Source/Programs/UnrealBuildAccelerator/Common/Public/UbaSession.h b/Engine/Source/Programs/UnrealBuildAccelerator/Common/Public/UbaSession.h
--- a/Engine/Source/Programs/UnrealBuildAccelerator/Common/Public/UbaSession.h
+++ b/Engine/Source/Programs/UnrealBuildAccelerator/Common/Public/UbaSession.h
@@ -265,6 +265,8 @@ namespace uba
 	};
 
 	void GenerateNameForProcess(StringBufferBase& out, const tchar* arguments, u32 counterSuffix);
+	bool GetZone(StringBufferBase& outZone);
+
 
 	using FileAccess = u8;
 
@@ -474,7 +476,6 @@ namespace uba
 		TString description;
 		TString logFile;
 	};
-
 }
 
 template<> struct std::hash<uba::ProcessHandle> { size_t operator()(const uba::ProcessHandle& g) const { return g.GetHash(); } };
diff --git a/Engine/Source/Programs/UnrealBuildAccelerator/Host/Private/UbaExports.cpp b/Engine/Source/Programs/UnrealBuildAccelerator/Host/Private/UbaExports.cpp
--- a/Engine/Source/Programs/UnrealBuildAccelerator/Host/Private/UbaExports.cpp
+++ b/Engine/Source/Programs/UnrealBuildAccelerator/Host/Private/UbaExports.cpp
@@ -239,10 +239,8 @@ extern "C"
 
 		StringBuffer<256> zoneTemp;
 		if (!zone || !*zone)
-		{
-			zoneTemp.count = GetEnvironmentVariableW(TC("UBA_ZONE"), zoneTemp.data, zoneTemp.capacity);
-			zone = zoneTemp.data;
-		}
+			if (GetZone(zoneTemp))
+				zone = zoneTemp.data;
 
 		info.casCapacityBytes = casCapacityBytes;
 		info.storeCompressed = storeCompressed;
