[UBA]
* Fixed so GetPlatformName returns false if failing on non-windows platform

diff --git a/Engine/Source/Programs/UnrealBuildAccelerator/Core/Private/UbaPlatform.cpp b/Engine/Source/Programs/UnrealBuildAccelerator/Core/Private/UbaPlatform.cpp
--- a/Engine/Source/Programs/UnrealBuildAccelerator/Core/Private/UbaPlatform.cpp
+++ b/Engine/Source/Programs/UnrealBuildAccelerator/Core/Private/UbaPlatform.cpp
@@ -123,8 +123,7 @@ namespace uba
 		DWORD nSize = bufferLen;
 		return ::GetComputerNameW(buffer, &nSize);
 		#else
-		gethostname(buffer, bufferLen);
-		return true;
+		return gethostname(buffer, bufferLen) == 0;
 		#endif
 	}
 
