[UBA]
* Changed GetComputerName to GetComputerNameW

diff --git a/Engine/Source/Programs/UnrealBuildAccelerator/Common/Private/UbaSession.cpp b/Engine/Source/Programs/UnrealBuildAccelerator/Common/Private/UbaSession.cpp
--- a/Engine/Source/Programs/UnrealBuildAccelerator/Common/Private/UbaSession.cpp
+++ b/Engine/Source/Programs/UnrealBuildAccelerator/Common/Private/UbaSession.cpp
@@ -2884,7 +2884,7 @@ namespace uba
 
 		// TODO: Remove.
 		#if PLATFORM_MAC
-		if (!GetComputerName(outZone.data, outZone.capacity))
+		if (!GetComputerNameW(outZone.data, outZone.capacity))
 			return false;
 
 		outZone.count = TStrlen(outZone.data);
