[UBA]
* Changed so SetWindowName is called on message thread

diff --git a/Engine/Source/Programs/UnrealBuildAccelerator/Visualizer/Private/UbaVisualizer.cpp b/Engine/Source/Programs/UnrealBuildAccelerator/Visualizer/Private/UbaVisualizer.cpp
--- a/Engine/Source/Programs/UnrealBuildAccelerator/Visualizer/Private/UbaVisualizer.cpp
+++ b/Engine/Source/Programs/UnrealBuildAccelerator/Visualizer/Private/UbaVisualizer.cpp
@@ -12,6 +12,7 @@
 #pragma comment (lib, "UxTheme.lib")
 #pragma comment (lib, "Dwmapi.lib")
 #define WM_NEWTRACE WM_USER+1
+#define WM_SETTITLE WM_USER+2
 
 namespace uba
 {
@@ -251,7 +252,7 @@ namespace uba
 			StringBuffer<> title;
 			GetTitlePrefix(title);
 			title.Appendf(L"Listening for new sessions on channel '%s'", m_listenChannel.data);
-			SetWindowTextW(m_hwnd, title.data);
+			PostNewTitle(title);
 		}
 
 		StringBuffer<256> traceName;
@@ -333,7 +334,7 @@ namespace uba
 			StringBuffer<> title;
 			GetTitlePrefix(title);
 			title.Appendf(L"Trying to connect to %s:%u%s", host, port, dots + ((dotsCounter--) % 4));
-			SetWindowTextW(m_hwnd, title.data);
+			PostNewTitle(title);
 
 			if (!m_client->Connect(backend, host, port))
 				continue;
@@ -2919,6 +2920,11 @@ namespace uba
 		PostMessage(m_hwnd, WM_NEWTRACE, replay, paused);
 	}
 
+	void Visualizer::PostNewTitle(const StringView& title)
+	{
+		PostMessage(m_hwnd, WM_SETTITLE, 0, (LPARAM)_wcsdup(title.data));
+	}
+
 	void Visualizer::PostQuit()
 	{
 		m_looping = false;
@@ -2929,6 +2935,13 @@ namespace uba
 	{
 		switch (Msg)
 		{
+		case WM_SETTITLE:
+		{
+			auto title = (wchar_t*)lParam;
+			SetWindowTextW(hWnd, title);
+			free(title);
+			break;
+		}
 		case WM_NEWTRACE:
 		{
 			m_replay = u32(wParam);
@@ -2969,7 +2982,7 @@ namespace uba
 				title.Appendf(L"%s (Listening for new sessions on channel '%s')", m_namedTrace.data, m_listenChannel.data);
 			}
 
-			SetWindowTextW(m_hwnd, title.data);
+			PostNewTitle(title);
 			SetTimer(m_hwnd, 0, 200, NULL);
 			return 0;
 		}
diff --git a/Engine/Source/Programs/UnrealBuildAccelerator/Visualizer/Private/UbaVisualizer.h b/Engine/Source/Programs/UnrealBuildAccelerator/Visualizer/Private/UbaVisualizer.h
--- a/Engine/Source/Programs/UnrealBuildAccelerator/Visualizer/Private/UbaVisualizer.h
+++ b/Engine/Source/Programs/UnrealBuildAccelerator/Visualizer/Private/UbaVisualizer.h
@@ -269,6 +269,7 @@ namespace uba
 		Thread m_thread;
 
 		void PostNewTrace(u32 replay, bool paused);
+		void PostNewTitle(const StringView& title);
 		void PostQuit();
 		LRESULT WinProc(HWND hWnd, UINT Msg, WPARAM wParam, LPARAM lParam);
 		static LRESULT CALLBACK StaticWinProc(HWND hWnd, UINT Msg, WPARAM wParam, LPARAM lParam);
