[UBA]
* Added detouring of LoadLibrary when running debug. It seems like maybe windows server 2025 has changed how things are called and we might need to do a proper detour of this function

diff --git a/Engine/Source/Programs/UnrealBuildAccelerator/Detours/Private/Windows/UbaDetoursFunctionsKernelBase.inl b/Engine/Source/Programs/UnrealBuildAccelerator/Detours/Private/Windows/UbaDetoursFunctionsKernelBase.inl
--- a/Engine/Source/Programs/UnrealBuildAccelerator/Detours/Private/Windows/UbaDetoursFunctionsKernelBase.inl
+++ b/Engine/Source/Programs/UnrealBuildAccelerator/Detours/Private/Windows/UbaDetoursFunctionsKernelBase.inl
@@ -3129,6 +3129,14 @@ BOOL Detoured_GetFileAttributesExA(LPCSTR lpFileName, GET_FILEEX_INFO_LEVELS fIn
 	return True_GetFileAttributesExA(lpFileName, fInfoLevelId, lpFileInformation);
 }
 
+HMODULE Detoured_LoadLibraryW(LPCWSTR lpLibFileName)
+{
+	DETOURED_CALL(LoadLibraryW);
+	DEBUG_LOG_TRUE(L"LoadLibraryW", L"(%ls)", lpLibFileName);
+	return True_LoadLibraryW(lpLibFileName);
+}
+
+
 DWORD Shared_GetModuleFileNameA(HMODULE hModule, const wchar_t* moduleName, u32 moduleNameLen, LPSTR lpFilename, DWORD nSize)
 {
 	if (nSize <= moduleNameLen)
diff --git a/Engine/Source/Programs/UnrealBuildAccelerator/Detours/Private/Windows/UbaDetoursFunctionsWin.h b/Engine/Source/Programs/UnrealBuildAccelerator/Detours/Private/Windows/UbaDetoursFunctionsWin.h
--- a/Engine/Source/Programs/UnrealBuildAccelerator/Detours/Private/Windows/UbaDetoursFunctionsWin.h
+++ b/Engine/Source/Programs/UnrealBuildAccelerator/Detours/Private/Windows/UbaDetoursFunctionsWin.h
@@ -237,6 +237,7 @@
 	DETOURED_FUNCTION(GetFullPathNameA) \
 	DETOURED_FUNCTION(GetFileAttributesA) \
 	DETOURED_FUNCTION(GetFileAttributesExA) \
+	DETOURED_FUNCTION(LoadLibraryW) \
 	DETOURED_FUNCTION(GetModuleFileNameA) \
 	DETOURED_FUNCTION(GetModuleFileNameExA) \
 	DETOURED_FUNCTION(GetModuleBaseNameA) \
