[UBA]
* Added support for stripping out exported symbols from elf obj files
* Added support for writing out --dynamic-list files

diff --git a/Engine/Source/Programs/UnrealBuildAccelerator/Common/Private/UbaObjectFile.cpp b/Engine/Source/Programs/UnrealBuildAccelerator/Common/Private/UbaObjectFile.cpp
--- a/Engine/Source/Programs/UnrealBuildAccelerator/Common/Private/UbaObjectFile.cpp
+++ b/Engine/Source/Programs/UnrealBuildAccelerator/Common/Private/UbaObjectFile.cpp
@@ -218,7 +218,7 @@ namespace uba
 		return m_potentialDuplicates;
 	}
 
-	bool ObjectFile::CreateExtraFile(Logger& logger, const StringView& extraObjFilename, const StringView& platform, const UnorderedSymbols& allNeededImports, const UnorderedSymbols& allSharedImports, const UnorderedExports& allSharedExports, bool includeExportsInFile)
+	bool ObjectFile::CreateExtraFile(Logger& logger, const StringView& extraObjFilename, const StringView& platform, const UnorderedSymbols& allExternalImports, const UnorderedSymbols& allInternalImports, const UnorderedExports& allExports, bool includeExportsInFile)
 	{
 		ObjectFileCoff objectFileCoff;
 		ObjectFileElf objectFileElf;
@@ -227,9 +227,11 @@ namespace uba
 
 		bool res;
 		if (platform.Equals(TC("win64")) || platform.Equals(TC("wingdk")) || platform.Equals(TC("xb1")) || platform.Equals(TC("xsx"))) 
-			res = ObjectFileCoff::CreateExtraFile(logger, platform, memoryBlock, allNeededImports, allSharedImports, allSharedExports, includeExportsInFile);
+			res = ObjectFileCoff::CreateExtraFile(logger, platform, memoryBlock, allExternalImports, allInternalImports, allExports, includeExportsInFile);
+		else if (extraObjFilename.EndsWith(TC("ldscript")))
+			res = CreateDynamicListFile(logger, memoryBlock, allExternalImports, allInternalImports, allExports, includeExportsInFile);
 		else
-			res = ObjectFileElf::CreateExtraFile(logger, platform, memoryBlock, allNeededImports, allSharedImports, allSharedExports, includeExportsInFile);
+			res = ObjectFileElf::CreateExtraFile(logger, platform, memoryBlock, allExternalImports, allInternalImports, allExports, includeExportsInFile);
 
 		if (!res)
 			return false;
@@ -279,4 +281,33 @@ namespace uba
 		}
 		return true;
 	}
+
+	bool ObjectFile::CreateDynamicListFile(Logger& logger, MemoryBlock& memoryBlock, const UnorderedSymbols& allExternalImports, const UnorderedSymbols& allInternalImports, const UnorderedExports& allExports, bool includeExportsInFile)
+	{
+		auto WriteString = [&](const char* str, u64 strLen)
+			{
+				memcpy(memoryBlock.Allocate(strLen, 1, TC("")), str, strLen);
+			};
+
+		//WriteString("VERSION ", 8);
+		WriteString("{", 1);
+
+		bool isFirst = true;
+		for (auto& symbol : allExports)
+		{
+			//if (strncmp(symbol.first.c_str(), "_ZTV", 4) != 0)
+			//	continue;
+			if (allExternalImports.find(symbol.first) == allExternalImports.end())
+				continue;
+			if (isFirst)
+				WriteString("global: ", 8);
+			WriteString(symbol.first.c_str(), symbol.first.size());
+			WriteString(";", 1);
+			isFirst = false;
+		}
+		//WriteString("local: *;", 9);
+		WriteString("};", 2);
+
+		return true;
+	}
 }
diff --git a/Engine/Source/Programs/UnrealBuildAccelerator/Common/Private/UbaObjectFileCoff.cpp b/Engine/Source/Programs/UnrealBuildAccelerator/Common/Private/UbaObjectFileCoff.cpp
--- a/Engine/Source/Programs/UnrealBuildAccelerator/Common/Private/UbaObjectFileCoff.cpp
+++ b/Engine/Source/Programs/UnrealBuildAccelerator/Common/Private/UbaObjectFileCoff.cpp
@@ -354,7 +354,7 @@ namespace uba
 		}
 	}
 
-	bool ObjectFileCoff::StripExports(Logger& logger, u8* newData, const UnorderedSymbols& allNeededImports)
+	bool ObjectFileCoff::StripExports(Logger& logger, u8* newData, const UnorderedSymbols& allExternalImports)
 	{
 		if (!m_info.directiveSectionMemOffset)
 			return true;
@@ -424,10 +424,10 @@ namespace uba
 			}
 
 			tmp.assign(exportStr, exportEnd - exportStr);
-			if (allNeededImports.find(tmp) != allNeededImports.end())
+			if (allExternalImports.find(tmp) != allExternalImports.end())
 				continue;
 			tmp.assign("__imp_").append(exportStr, exportEnd - exportStr);
-			if (allNeededImports.find(tmp) != allNeededImports.end())
+			if (allExternalImports.find(tmp) != allExternalImports.end())
 				continue;
 			
 			u64 toCopy = startPos - lastCopyPos - 1;
@@ -451,18 +451,18 @@ namespace uba
 		return true;
 	}
 
-	bool ObjectFileCoff::CreateExtraFile(Logger& logger, const StringView& platform, MemoryBlock& memoryBlock, const UnorderedSymbols& allNeededImports, const UnorderedSymbols& allSharedImports, const UnorderedExports& allSharedExports, bool includeExportsInFile)
+	bool ObjectFileCoff::CreateExtraFile(Logger& logger, const StringView& platform, MemoryBlock& memoryBlock, const UnorderedSymbols& allExternalImports, const UnorderedSymbols& allInternalImports, const UnorderedExports& allExports, bool includeExportsInFile)
 	{
 		std::string tmp;
 
 		u32 totalStringSize = 0;
 		UnorderedSymbols neededLoopbacks;
-		for (auto& symbol : allSharedImports)
+		for (auto& symbol : allInternalImports)
 		{
 			if (strncmp(symbol.data(), "__imp_", 6) != 0)
 				continue;
 			tmp = symbol.substr(6);
-			if (allSharedExports.find(tmp) == allSharedExports.end())
+			if (allExports.find(tmp) == allExports.end())
 				continue;
 			if (neededLoopbacks.insert(symbol).second)
 				totalStringSize += u32(symbol.size()) + 1;
@@ -498,14 +498,14 @@ namespace uba
 			u32 directiveRawDataStart = u32(memoryBlock.writtenSize);
 			directiveSection.PointerToRawData = directiveRawDataStart;
 			char slashExport[] = "/EXPORT:";
-			for (auto& kv : allSharedExports)
+			for (auto& kv : allExports)
 			{
 				auto& symbol = kv.first;
 
-				if (allNeededImports.find(symbol) == allNeededImports.end())
+				if (allExternalImports.find(symbol) == allExternalImports.end())
 				{
 					tmp.assign("__imp_").append(symbol);
-					if (allNeededImports.find(tmp) == allNeededImports.end())
+					if (allExternalImports.find(tmp) == allExternalImports.end())
 						continue;
 				}
 
diff --git a/Engine/Source/Programs/UnrealBuildAccelerator/Common/Private/UbaObjectFileCoff.h b/Engine/Source/Programs/UnrealBuildAccelerator/Common/Private/UbaObjectFileCoff.h
--- a/Engine/Source/Programs/UnrealBuildAccelerator/Common/Private/UbaObjectFileCoff.h
+++ b/Engine/Source/Programs/UnrealBuildAccelerator/Common/Private/UbaObjectFileCoff.h
@@ -13,7 +13,7 @@ namespace uba
 		ObjectFileCoff();
 		virtual bool Parse(Logger& logger, const tchar* hint) override;
 
-		static bool CreateExtraFile(Logger& logger, const StringView& platform, MemoryBlock& memoryBlock, const UnorderedSymbols& allNeededImports, const UnorderedSymbols& allSharedImports, const UnorderedExports& allSharedExports, bool includeExportsInFile);
+		static bool CreateExtraFile(Logger& logger, const StringView& platform, MemoryBlock& memoryBlock, const UnorderedSymbols& allExternalImports, const UnorderedSymbols& allInternalImports, const UnorderedExports& allSharedExports, bool includeExportsInFile);
 
 	private:
 		struct Info;
@@ -21,7 +21,7 @@ namespace uba
 		bool ParseExports();
 		template<typename SymbolType> void ParseImports();
 
-		virtual bool StripExports(Logger& logger, u8* newData, const UnorderedSymbols& allNeededImports) override;
+		virtual bool StripExports(Logger& logger, u8* newData, const UnorderedSymbols& allExternalImports) override;
 
 		template<typename SymbolType> void CalculateImports(Logger& logger, Vector<u32>& outImports);
 		template<typename SymbolType> void WriteImports(Logger& logger, u8* newData, Info& newInfo, const Vector<u32>& symbolsToAdd);
diff --git a/Engine/Source/Programs/UnrealBuildAccelerator/Common/Private/UbaObjectFileElf.cpp b/Engine/Source/Programs/UnrealBuildAccelerator/Common/Private/UbaObjectFileElf.cpp
--- a/Engine/Source/Programs/UnrealBuildAccelerator/Common/Private/UbaObjectFileElf.cpp
+++ b/Engine/Source/Programs/UnrealBuildAccelerator/Common/Private/UbaObjectFileElf.cpp
@@ -7,44 +7,44 @@ namespace uba
 {
 	struct Elf64Header
 	{
-		unsigned char	e_ident[16];	/* Magic number and other info */
-		u16	e_type;			/* Object file type */
-		u16	e_machine;		/* Architecture */
-		u32	e_version;		/* Object file version */
-		u64	e_entry;		/* Entry point virtual address */
-		u64	e_phoff;		/* Program header table file offset */
-		u64	e_shoff;		/* Section header table file offset */
-		u32	e_flags;		/* Processor-specific flags */
-		u16	e_ehsize;		/* ELF header size in bytes */
-		u16	e_phentsize;	/* Program header table entry size */
-		u16	e_phnum;		/* Program header table entry count */
-		u16	e_shentsize;	/* Section header table entry size */
-		u16	e_shnum;		/* Section header table entry count */
-		u16	e_shstrndx;		/* Section header string table index */
+		unsigned char	e_ident[16];	// Magic number and other info
+		u16	e_type;			// Object file type
+		u16	e_machine;		// Architecture
+		u32	e_version;		// Object file version
+		u64	e_entry;		// Entry point virtual address
+		u64	e_phoff;		// Program header table file offset
+		u64	e_shoff;		// Section header table file offset
+		u32	e_flags;		// Processor-specific flags
+		u16	e_ehsize;		// ELF header size in bytes
+		u16	e_phentsize;	// Program header table entry size
+		u16	e_phnum;		// Program header table entry count
+		u16	e_shentsize;	// Section header table entry size
+		u16	e_shnum;		// Section header table entry count
+		u16	e_shstrndx;		// Section header string table index
 	};
 
 	struct Elf64SectionHeader
 	{
-		u32	sh_name;		/* Section name (string tbl index) */
-		u32	sh_type;		/* Section type */
-		u64	sh_flags;		/* Section flags */
-		u64	sh_addr;		/* Section virtual addr at execution */
-		u64	sh_offset;		/* Section file offset */
-		u64	sh_size;		/* Section size in bytes */
-		u32	sh_link;		/* Link to another section */
-		u32	sh_info;		/* Additional section information */
-		u64	sh_addralign;	/* Section alignment */
-		u64	sh_entsize;		/* Entry size if section holds table */
+		u32	sh_name;		// Section name (string tbl index)
+		u32	sh_type;		// Section type
+		u64	sh_flags;		// Section flags
+		u64	sh_addr;		// Section virtual addr at execution
+		u64	sh_offset;		// Section file offset
+		u64	sh_size;		// Section size in bytes
+		u32	sh_link;		// Link to another section
+		u32	sh_info;		// Additional section information
+		u64	sh_addralign;	// Section alignment
+		u64	sh_entsize;		// Entry size if section holds table
 	};
 
 	struct Elf64Sym
 	{
-		u32	st_name;		/* Symbol name (string tbl index) */
-		u8	st_info;		/* Symbol type and binding */
-		u8	st_other;		/* Symbol visibility */
-		u16	st_shndx;		/* Section index */
-		u64	st_value;		/* Symbol value */
-		u64	st_size;		/* Symbol size */
+		u32	st_name;		// Symbol name (string tbl index)
+		u8	st_info;		// Symbol type and binding
+		u8	st_other;		// Symbol visibility
+		u16	st_shndx;		// Section index
+		u64	st_value;		// Symbol value
+		u64	st_size;		// Symbol size
 	};
 
 	struct Elf64Rel
@@ -67,24 +67,25 @@ namespace uba
 	#define ELF32_ST_VISIBILITY(o)	((o) & 0x03)
 	#define ELF64_ST_VISIBILITY(o)	ELF32_ST_VISIBILITY (o)
 
-	#define EM_X86_64	62	/* AMD x86-64 architecture */
+	#define EM_X86_64	62		// AMD x86-64 architecture
 
-	#define SHT_SYMTAB	  2		/* Symbol table */
+	#define SHT_PROGBITS  1		// Program data
+	#define SHT_SYMTAB	  2		// Symbol table
 	#define SHT_STRTAB	  3
-	#define SHT_RELA	  4		/* Relocation entries with addends */
-	#define SHT_DYNAMIC	  6		/* Dynamic linking information */
-	#define SHT_REL		  9		/* Relocation entries, no addends */
-	#define SHT_DYNSYM	  11		/* Dynamic linker symbol table */
-	#define SHT_SYMTAB_SHNDX  18		/* Extended section indeces */
+	#define SHT_RELA	  4		// Relocation entries with addends
+	#define SHT_DYNAMIC	  6		// Dynamic linking information
+	#define SHT_REL		  9		// Relocation entries, no addends
+	#define SHT_DYNSYM	  11	// Dynamic linker symbol table
+	#define SHT_SYMTAB_SHNDX  18 // Extended section indeces
 
-	#define STT_NOTYPE	0		/* Symbol type is unspecified */
-	#define STT_OBJECT	1		/* Symbol is a data object */
-	#define STT_FUNC	2		/* Symbol is a code object */
-	#define STT_SECTION	3		/* Symbol associated with a section */
+	#define STT_NOTYPE	0		// Symbol type is unspecified
+	#define STT_OBJECT	1		// Symbol is a data object
+	#define STT_FUNC	2		// Symbol is a code object
+	#define STT_SECTION	3		// Symbol associated with a section
 
-	#define STB_LOCAL	0		/* Local symbol */
-	#define STB_GLOBAL	1		/* Global symbol */
-	#define STB_WEAK	2		/* Weak symbol */
+	#define STB_LOCAL	0		// Local symbol
+	#define STB_GLOBAL	1		// Global symbol
+	#define STB_WEAK	2		// Weak symbol
 
 	#define STV_DEFAULT	0		// Default symbol visibility rules
 	#define STV_INTERNAL	1	// Processor specific hidden class
@@ -95,56 +96,11 @@ namespace uba
 	#define ELF32_ST_TYPE(val)		((val) & 0xf)
 	#define ELF32_ST_INFO(bind, type)	(((bind) << 4) + ((type) & 0xf))
 
-	/* Both Elf32_Sym and Elf64_Sym use the same one-byte st_info field.  */
+	// Both Elf32_Sym and Elf64_Sym use the same one-byte st_info field.
 	#define ELF64_ST_BIND(val)		ELF32_ST_BIND(val)
 	#define ELF64_ST_TYPE(val)		ELF32_ST_TYPE (val)
 	#define ELF64_ST_INFO(bind, type)	ELF32_ST_INFO ((bind), (type))
 
-	// These are symbols that are added to all dlls through some macros.
-	// When merging dlls we need to remove duplicates of these
-	UnorderedSymbols ObjectFileElf::PotentiallyDuplicatedSymbols = 
-	{
-#if 0 // This is not needed anymore
-		// REPLACEMENT_OPERATOR_NEW_AND_DELETE
-		"_Znwm",
-		"_Znam",
-		"_ZnwmRKSt9nothrow_t",
-		"_ZnamRKSt9nothrow_t",
-		"_ZnwmSt11align_val_t",
-		"_ZnamSt11align_val_t",
-		"_ZnwmSt11align_val_tRKSt9nothrow_t",
-		"_ZnamSt11align_val_tRKSt9nothrow_t",
-		"_ZdlPv",
-		"_ZdaPv",
-		"_ZdlPvRKSt9nothrow_t",
-		"_ZdaPvRKSt9nothrow_t",
-		"_ZdlPvm",
-		"_ZdaPvm",
-		"_ZdlPvmRKSt9nothrow_t",
-		"_ZdaPvmRKSt9nothrow_t",
-		"_ZdlPvSt11align_val_t",
-		"_ZdaPvSt11align_val_t",
-		"_ZdlPvSt11align_val_tRKSt9nothrow_t",
-		"_ZdaPvSt11align_val_tRKSt9nothrow_t",
-		"_ZdlPvmSt11align_val_t",
-		"_ZdaPvmSt11align_val_t",
-		"_ZdlPvmSt11align_val_tRKSt9nothrow_t",
-		"_ZdaPvmSt11align_val_tRKSt9nothrow_t",
-
-		// UE_DEFINE_FMEMORY_WRAPPERS
-		"_Z14FMemory_Mallocmm",
-		"_Z15FMemory_ReallocPvmm",
-		"_Z12FMemory_FreePv",
-
-		// UE4_VISUALIZERS_HELPERS
-		"GNameBlocksDebug",
-		"GObjectIndexToPackedObjectRefDebug",
-		"GObjectArrayForDebugVisualizers",
-		"GComplexObjectPathDebug",
-		"GObjectHandlePackageDebug",
-#endif
-	};
-
 	bool IsElfFile(const u8* data, u64 dataSize)
 	{
 		constexpr u8 magic[] = { 0x7f, 'E', 'L', 'F' };
@@ -167,8 +123,8 @@ namespace uba
 			return false;
 		if (header.e_ident[6] != 1) // Version
 			return false;
-		if (header.e_type != 1 && header.e_type != 0xfe0c) // Relocatable file
-			return false;
+		//if (header.e_type != 1 && header.e_type != 0xfe0c) // Relocatable file
+		//	return false;
 
 		UBA_ASSERT(sizeof(Elf64SectionHeader) == header.e_shentsize);
 
@@ -180,41 +136,33 @@ namespace uba
 		char* symTableNames = nullptr;
 		char* dynTableNames = nullptr;
 
-		#if 0
-		u64 stringTableSize = namesSection.sh_size;
-		u32 counter = 0;
-		char* it = (char*)stringTableData;
-		char* end = it + stringTableSize;
-		while (it != end)
-		{
-			printf("%s\n", it);
-			u64 len = strlen(it);
-			it += len + 1;
-			counter++;
-		}
-		#endif
+		u64 sectionCount = header.e_shnum;
+		if (sectionCount == 0)
+			sectionCount = sections[0].sh_size;
 
-		for (u16 i=0; i!=header.e_shnum; ++i)
+		for (u64 i=0; i!=sectionCount; ++i)
 		{
 			auto& section = sections[i];
 			if (section.sh_type != SHT_STRTAB)
 				continue;
 			char* sectionName = sectionNamesTable + section.sh_name;
-			if (strcmp(sectionName, ".shstrtab") == 0)
-				printf("");
-			else if (strcmp(sectionName, ".strtab") == 0)
+			if (strcmp(sectionName, ".strtab") == 0)
 				symTableNames = (char*)m_data + section.sh_offset;
 			else if (strcmp(sectionName, ".dynstr") == 0)
 				dynTableNames = (char*)m_data + section.sh_offset;
+			//else if (strcmp(sectionName, ".shstrtab") == 0)
+				;//printf("");
 		}
 
-		for (u16 i=0; i!=header.e_shnum; ++i)
+		for (u64 i=0; i!=sectionCount; ++i)
 		{
-			auto& section = sections[i];
+			auto& section = sections[i]; 
+
+			char* sectionName = sectionNamesTable + section.sh_name;
+			//printf("TYPE: %3i - %14u %-30s  Offset: 0x%llx Size: %llu\n", i, section.sh_type, sectionName, section.sh_offset, section.sh_size);
 
 			if (section.sh_type == SHT_SYMTAB || section.sh_type == SHT_DYNSYM)
 			{
-				//UBA_ASSERT(section.sh_link == header.e_shstrndx);
 				UBA_ASSERT(section.sh_entsize == sizeof(Elf64Sym));
 				auto symbols = (Elf64Sym*)(m_data + section.sh_offset);
 				u64 symbolCount = section.sh_size / sizeof(Elf64Sym);
@@ -238,84 +186,92 @@ namespace uba
 					if (symbolBinding == STB_GLOBAL)
 					{
 						u32 vis = ELF64_ST_VISIBILITY(symbol.st_other);
-						(void)vis;
 						if (symbolType != STT_NOTYPE)
 						{
+							// This is an ugly hack because some platforms are not exporting vtable even though they should
+							if (vis == STV_HIDDEN && strncmp(symbolName, "_ZTV", 4) == 0)
+								m_exports.emplace(symbolName, ExportInfo{std::string(), 0});//index++});
+#if 0
 							// On ELF, protected visibility indicates that the symbol will be placed in the dynamic symbol table, but that references within the defining module will bind to the local symbol. That is, the symbol cannot be overridden by another module.
 							if (vis == STV_PROTECTED) //|| vis == STV_DEFAULT
 								m_exports.emplace(symbolName, ExportInfo{std::string(), 0});//index++});
+							else if (strncmp(symbolName, "_ZTV", 4) == 0) // This is an ugly hack because some platforms are not exporting vtable even though they should
+								m_exports.emplace(symbolName, ExportInfo{std::string(), 0});//index++});
+							//else if (strncmp(symbolName, "_ZNK14TSoftObjectPtrI", 21) == 0 || strncmp(symbolName, "_ZNK11TSubclassOfI", 18) == 0)
+							//	m_exports.emplace(symbolName, ExportInfo{std::string(), 0});//index++});
+#endif
 						}
 						else
 							m_imports.emplace(symbolName);//index++});
 					}
+					else if (symbolBinding == STB_WEAK)
+					{
+						u32 vis = ELF64_ST_VISIBILITY(symbol.st_other);
+						if (symbolType != STT_NOTYPE && vis == STV_HIDDEN)
+						{
+							// Another ugly hack because of dll exported template specialization
+							if (strncmp(symbolName, "_ZNK14TSoftObjectPtrI", 21) == 0)
+								m_exports.emplace(symbolName, ExportInfo{std::string(), 0});//index++});
+						}
+					}
 				}
 			}
-			else if (section.sh_type == SHT_RELA)
+			else if (section.sh_type == SHT_PROGBITS)
 			{
-				// ".rela.*"
-				char* sectionName = sectionNamesTable + section.sh_name;
-				(void)sectionName;
-				auto& data = *(Elf64Rela*)(m_data + section.sh_offset);
-				ELF64_R_SYM(data.r_info);
-				ELF64_R_TYPE(data.r_info);
-					(void)data;
-				//printf("");
+				if (strcmp(sectionName, ".linker_cmd") == 0)
+				{
+					u8* it = m_data + section.sh_offset;
+					u8* end = it + section.sh_size;
+					while (it < end)
+					{
+						it += 4;
+						char* symbolName = (char*)it;
+						m_exports.emplace(symbolName, ExportInfo{std::string(), 0});
+						it += strlen(symbolName) + 1;
+					}
+				}
 			}
-			//else
-			//	printf("TYPE: %u\n", section.sh_type);
 		}
+
 		return true;
 	}
-	/*
-	bool ObjectFileElf::CreateStripped(Logger& logger, const tchar* newFilename, const UnorderedSymbols& allNeededImports, u32& outKeptExportCount)
-	{
-		FileAccessor newFile(logger, newFilename);
-		if (!newFile.CreateMemoryWrite(false, DefaultAttributes(), m_dataSize))
-			return false;
 
-		u8* newData = newFile.GetData();
-		memcpy(newData, m_data, m_dataSize);
-
-		auto& header = *(Elf64Header*)newData;
-		auto sections = (Elf64SectionHeader*)(newData + header.e_shoff);
+	bool ObjectFileElf::StripExports(Logger& logger, u8* newData, const UnorderedSymbols& allExternalImports)
+	{
+		auto& header = *(Elf64Header*)m_data;
+		auto sections = (Elf64SectionHeader*)(m_data + header.e_shoff);
+		char* sectionNamesTable = (char*)m_data + sections[header.e_shstrndx].sh_offset;
 
-		auto& namesSection = sections[header.e_shstrndx];
-		u8* stringTableData = newData + namesSection.sh_offset;
+		u64 sectionCount = header.e_shnum;
+		if (sectionCount == 0)
+			sectionCount = sections[0].sh_size;
 
-		for (u16 i=0; i!=header.e_shnum; ++i)
+		for (u64 i=0; i!=sectionCount; ++i)
 		{
 			auto& section = sections[i];
+			char* sectionName = sectionNamesTable + section.sh_name;
 
-			if (section.sh_type == SHT_SYMTAB || section.sh_type == SHT_DYNSYM)
+			// Remove all but first symbol.. 
+			// TODO Ideally we want to remove all symbols but then it seems like no stub file is created, maybe we can put a dummy symbol in there?
+			if (section.sh_type == SHT_PROGBITS && strcmp(sectionName, ".linker_cmd") == 0)
 			{
-				UBA_ASSERT(section.sh_link == header.e_shstrndx);
-				UBA_ASSERT(section.sh_entsize == sizeof(Elf64Sym));
-				auto symbols = (Elf64Sym*)(newData + section.sh_offset);
-				u64 symbolCount = section.sh_size / sizeof(Elf64Sym);
-				for (u64 j=0; j!=symbolCount; ++j)
-				{
-					auto& symbol = symbols[j];
-					u8 symbolType = ELF64_ST_TYPE(symbol.st_info);
-					if (symbolType != STT_FUNC && symbolType != STT_OBJECT)
-						continue;
-					auto name = (char*)stringTableData + symbol.st_name;
+				u64 oldSize = section.sh_size;
+				if (!oldSize)
+					return true;
+				u8* begin = m_data + section.sh_offset;
+				u8* it = begin;
 
-					if (m_toRemove.find(name) == m_toRemove.end())
-						continue;
-					symbol.st_info = ELF64_ST_INFO(STB_WEAK, symbolType);
-				}
+				it += 4;
+				it += strlen((char*)it) + 1;
+
+				section.sh_size = it - begin;
+				memset(it, 0, oldSize - section.sh_size);
 			}
 		}
-
-		return newFile.Close();
-	}
-	*/
-	bool ObjectFileElf::StripExports(Logger& logger, u8* newData, const UnorderedSymbols& allNeededImports)
-	{
 		return true;
 	}
 
-	bool ObjectFileElf::CreateExtraFile(Logger& logger, const StringView& platform, MemoryBlock& memoryBlock, const UnorderedSymbols& allNeededImports, const UnorderedSymbols& allSharedImports, const UnorderedExports& allSharedExports, bool includeExportsInFile)
+	bool ObjectFileElf::CreateExtraFile(Logger& logger, const StringView& platform, MemoryBlock& memoryBlock, const UnorderedSymbols& allExternalImports, const UnorderedSymbols& allInternalImports, const UnorderedExports& allExports, bool includeExportsInFile)
 	{
 		u8* data = (u8*)memoryBlock.Allocate(sizeof(Elf64Header) + sizeof(Elf64SectionHeader) + 1, 1, TC(""));
 		auto& header = *(Elf64Header*)data;
diff --git a/Engine/Source/Programs/UnrealBuildAccelerator/Common/Private/UbaObjectFileElf.h b/Engine/Source/Programs/UnrealBuildAccelerator/Common/Private/UbaObjectFileElf.h
--- a/Engine/Source/Programs/UnrealBuildAccelerator/Common/Private/UbaObjectFileElf.h
+++ b/Engine/Source/Programs/UnrealBuildAccelerator/Common/Private/UbaObjectFileElf.h
@@ -12,14 +12,10 @@ namespace uba
 		ObjectFileElf();
 		virtual bool Parse(Logger& logger, const tchar* hint) override;
 
-		static bool CreateExtraFile(Logger& logger, const StringView& platform, MemoryBlock& memoryBlock, const UnorderedSymbols& allNeededImports, const UnorderedSymbols& allSharedImports, const UnorderedExports& allSharedExports, bool includeExportsInFile);
+		static bool CreateExtraFile(Logger& logger, const StringView& platform, MemoryBlock& memoryBlock, const UnorderedSymbols& allExternalImports, const UnorderedSymbols& allInternalImports, const UnorderedExports& allExports, bool includeExportsInFile);
 
 	private:
-		virtual bool StripExports(Logger& logger, u8* newData, const UnorderedSymbols& allNeededImports) override;
-
-		UnorderedSymbols m_toRemove;
-
-		static UnorderedSymbols PotentiallyDuplicatedSymbols;
+		virtual bool StripExports(Logger& logger, u8* newData, const UnorderedSymbols& allExternalImports) override;
 	};
 
 	bool IsElfFile(const u8* data, u64 dataSize);
diff --git a/Engine/Source/Programs/UnrealBuildAccelerator/Common/Public/UbaObjectFile.h b/Engine/Source/Programs/UnrealBuildAccelerator/Common/Public/UbaObjectFile.h
--- a/Engine/Source/Programs/UnrealBuildAccelerator/Common/Public/UbaObjectFile.h
+++ b/Engine/Source/Programs/UnrealBuildAccelerator/Common/Public/UbaObjectFile.h
@@ -39,7 +39,7 @@ namespace uba
 		const UnorderedExports& GetExports() const;
 		const UnorderedSymbols& GetPotentialDuplicates() const;
 
-		static bool CreateExtraFile(Logger& logger, const StringView& extraObjFilename, const StringView& platform, const UnorderedSymbols& allNeededImports, const UnorderedSymbols& allSharedImports, const UnorderedExports& allSharedExports, bool includeExportsInFile);
+		static bool CreateExtraFile(Logger& logger, const StringView& extraObjFilename, const StringView& platform, const UnorderedSymbols& allExternalImports, const UnorderedSymbols& allInternalImports, const UnorderedExports& allExports, bool includeExportsInFile);
 
 		virtual ~ObjectFile();
 
@@ -49,7 +49,9 @@ namespace uba
 
 	protected:
 		virtual bool Parse(Logger& logger, const tchar* hint) = 0;
-		virtual bool StripExports(Logger& logger, u8* newData, const UnorderedSymbols& allNeededImports) = 0;
+		virtual bool StripExports(Logger& logger, u8* newData, const UnorderedSymbols& allExternalImports) = 0;
+
+		static bool CreateDynamicListFile(Logger& logger, MemoryBlock& memoryBlock, const UnorderedSymbols& allExternalImports, const UnorderedSymbols& allInternalImports, const UnorderedExports& allExports, bool includeExportsInFile);
 
 		FileAccessor* m_file = nullptr;
 		u8* m_data = nullptr;
diff --git a/Engine/Source/Programs/UnrealBuildAccelerator/ObjTool/Private/UbaObjTool.cpp b/Engine/Source/Programs/UnrealBuildAccelerator/ObjTool/Private/UbaObjTool.cpp
--- a/Engine/Source/Programs/UnrealBuildAccelerator/ObjTool/Private/UbaObjTool.cpp
+++ b/Engine/Source/Programs/UnrealBuildAccelerator/ObjTool/Private/UbaObjTool.cpp
@@ -248,10 +248,10 @@ namespace uba
 		{
 			CriticalSection cs;
 			Atomic<bool> success = true;
-			UnorderedSymbols allNeededImports; // Imports needed from the outside of the stripped obj files
+			UnorderedSymbols allExternalImports; // Imports needed from the outside of the stripped obj files
 
 			for (auto imp : NeededImports)
-				allNeededImports.insert(imp);
+				allExternalImports.insert(imp);
 
 			u32 workerCount = DefaultProcessorCount;
 			WorkManagerImpl workManager(workerCount);
@@ -266,13 +266,13 @@ namespace uba
 						return;
 					}
 					ScopedCriticalSection _(cs);
-					allNeededImports.insert(symbolFile.imports.begin(), symbolFile.imports.end());
+					allExternalImports.insert(symbolFile.imports.begin(), symbolFile.imports.end());
 				});
 			if (!success)
 				return -1;
 
-			UnorderedSymbols allSharedImports; // Imports from all the obj files about to be stripped
-			UnorderedExports allSharedExports; // Exports from all the obj files about to be stripped
+			UnorderedSymbols allInternalImports; // Imports that the obj files has. Note these could be existing in the obj files and then we might need to create loopbacks
+			UnorderedExports allExports; // Exports from all the obj files.
 
 			Map<TString, ObjectFile*> objectFiles;
 			auto g = MakeGuard([&]() { for (auto& kv : objectFiles) delete kv.second; });
@@ -287,16 +287,34 @@ namespace uba
 						return;
 					}
 					ScopedCriticalSection _(cs);
-					allSharedImports.insert(symbolFile.imports.begin(), symbolFile.imports.end());
-					allSharedExports.insert(symbolFile.exports.begin(), symbolFile.exports.end());
+					allInternalImports.insert(symbolFile.imports.begin(), symbolFile.imports.end());
+					allExports.insert(symbolFile.exports.begin(), symbolFile.exports.end());
 				});
 			if (!success)
 				return -1;
 
 			if (!extraObjFile.empty())
-				if (!ObjectFile::CreateExtraFile(logger, extraObjFile, platform, allNeededImports, allSharedImports, allSharedExports, true))
-					return -1;
+			{
+				{
+					// Ugly hack. These symbols are exported from PosixShim and used by libcrypto.a but since we don't have an .exi file for libcrypto
+					// we do it this way for now
+					const char* symbolsToNeverStrip[] = 
+					{
+						"read_system_certificates_NP",
+						"inet_ntoa",
+						"gethostbyname",
+						"h_errno",
+						"getservbyname",
+						"ioctl",
+						"fcntl_shim",
+					};
+					for (auto symbol : symbolsToNeverStrip)
+						allExternalImports.insert(symbol);
+				}
 
+				if (!ObjectFile::CreateExtraFile(logger, extraObjFile, platform, allExternalImports, allInternalImports, allExports, true))
+					return -1;
+			}
 			//logger.Info(TC("Reduced export count from %llu to %llu"), totalExportCount.load(), totalKeptExportCount.size());
 		}
 		else if (writeImpLib)
