[UBA]
* Fixed crash in elf object file parser when there are symbols that has section index 65535
* Fixed so platforms using .elf files but don't have a .linker_cmd section use visibility as a way to control exports

This is very much VIP right now

diff --git a/Engine/Source/Programs/UnrealBuildAccelerator/Common/Private/UbaObjectFileElf.cpp b/Engine/Source/Programs/UnrealBuildAccelerator/Common/Private/UbaObjectFileElf.cpp
--- a/Engine/Source/Programs/UnrealBuildAccelerator/Common/Private/UbaObjectFileElf.cpp
+++ b/Engine/Source/Programs/UnrealBuildAccelerator/Common/Private/UbaObjectFileElf.cpp
@@ -114,6 +114,18 @@ namespace uba
 	}
 
 
+	const char* GetSymbolName(const Elf64SectionHeader* sections, const Elf64Sym& symbol, const char* symTableNames, const char* dynTableNames, const char* sectionNamesTable)
+	{
+		u8 symbolType = ELF64_ST_TYPE(symbol.st_info);
+		if (symbolType == STT_FUNC || symbolType == STT_NOTYPE || symbolType == STT_OBJECT)
+			return symTableNames + symbol.st_name;
+		if (symbolType == SHT_DYNSYM)
+			return dynTableNames + symbol.st_name;
+		if (symbolType == STT_SECTION && symbol.st_shndx != 65535)
+			return sectionNamesTable + sections[symbol.st_shndx].sh_name;
+		return "";
+	}
+
 	bool ObjectFileElf::Parse(Logger& logger, const tchar* hint)
 	{
 		auto& header = *(Elf64Header*)m_data;
@@ -133,33 +145,37 @@ namespace uba
 		auto& namesSection = sections[header.e_shstrndx];
 		UBA_ASSERT(namesSection.sh_type == SHT_STRTAB);
 		char* sectionNamesTable = (char*)m_data + namesSection.sh_offset;
-		char* symTableNames = nullptr;
-		char* dynTableNames = nullptr;
 
-		u64 sectionCount = header.e_shnum;
-		if (sectionCount == 0)
-			sectionCount = sections[0].sh_size;
+		u64 sectionCount = header.e_shnum ? header.e_shnum : sections[0].sh_size;
 
 		for (u64 i=0; i!=sectionCount; ++i)
 		{
 			auto& section = sections[i];
-			if (section.sh_type != SHT_STRTAB)
-				continue;
 			char* sectionName = sectionNamesTable + section.sh_name;
-			if (strcmp(sectionName, ".strtab") == 0)
-				symTableNames = (char*)m_data + section.sh_offset;
-			else if (strcmp(sectionName, ".dynstr") == 0)
-				dynTableNames = (char*)m_data + section.sh_offset;
-			//else if (strcmp(sectionName, ".shstrtab") == 0)
-				;//printf("");
+
+			if (section.sh_type == SHT_STRTAB)
+			{
+				if (strcmp(sectionName, ".strtab") == 0)
+					m_symTableNamesOffset = section.sh_offset;
+				else if (strcmp(sectionName, ".dynstr") == 0)
+					m_dynTableNamesOffset = section.sh_offset;
+			}
+			else if (section.sh_type == SHT_PROGBITS)
+			{
+				if (strcmp(sectionName, ".linker_cmd") == 0)
+					m_useVisibilityForExports = false;
+			}
 		}
 
+		char* symTableNames = (char*)m_data + m_symTableNamesOffset;
+		char* dynTableNames = (char*)m_data + m_dynTableNamesOffset;
+
 		for (u64 i=0; i!=sectionCount; ++i)
 		{
 			auto& section = sections[i]; 
 
 			char* sectionName = sectionNamesTable + section.sh_name;
-			//printf("TYPE: %3i - %14u %-30s  Offset: 0x%llx Size: %llu\n", i, section.sh_type, sectionName, section.sh_offset, section.sh_size);
+			//printf("TYPE: %3llu - %14u %-30s  Offset: 0x%llx Size: %llu\n", i, section.sh_type, sectionName, section.sh_offset, section.sh_size);
 
 			if (section.sh_type == SHT_SYMTAB || section.sh_type == SHT_DYNSYM)
 			{
@@ -169,54 +185,51 @@ namespace uba
 				for (u64 j=0; j!=symbolCount; ++j)
 				{
 					auto& symbol = symbols[j];
-					u8 symbolType = ELF64_ST_TYPE(symbol.st_info);
-					const char* symbolName;
-
-					if (symbolType == STT_FUNC || symbolType == STT_OBJECT)
-						symbolName = (section.sh_type == SHT_DYNSYM ? dynTableNames : symTableNames) + symbol.st_name;
-					else if (symbolType == STT_SECTION)
-						symbolName = sectionNamesTable + sections[symbol.st_shndx].sh_name;
-					else if (symbolType == STT_NOTYPE)
-						symbolName = symTableNames + symbol.st_name;
-					else
+					const char* symbolName = GetSymbolName(sections, symbol, symTableNames, dynTableNames, sectionNamesTable);
+					if (!*symbolName)
 						continue;
 
+					u8 symbolType = ELF64_ST_TYPE(symbol.st_info);
 					u32 symbolBinding = ELF64_ST_BIND(symbol.st_info);
+					u32 symbolVisibility = ELF64_ST_VISIBILITY(symbol.st_other);
+
+					//if (strstr(symbolName, "_ZN11TSubclassOfI13UAnimInstance"))
+					//	printf("");
 
 					if (symbolBinding == STB_GLOBAL)
 					{
-						u32 vis = ELF64_ST_VISIBILITY(symbol.st_other);
 						if (symbolType != STT_NOTYPE)
 						{
 							// This is an ugly hack because some platforms are not exporting vtable even though they should
-							if (vis == STV_HIDDEN && strncmp(symbolName, "_ZTV", 4) == 0)
-								m_exports.emplace(symbolName, ExportInfo{std::string(), 0});//index++});
-#if 0
-							// On ELF, protected visibility indicates that the symbol will be placed in the dynamic symbol table, but that references within the defining module will bind to the local symbol. That is, the symbol cannot be overridden by another module.
-							if (vis == STV_PROTECTED) //|| vis == STV_DEFAULT
+							if (symbolVisibility == STV_HIDDEN && strncmp(symbolName, "_ZTV", 4) == 0)
 								m_exports.emplace(symbolName, ExportInfo{std::string(), 0});//index++});
-							else if (strncmp(symbolName, "_ZTV", 4) == 0) // This is an ugly hack because some platforms are not exporting vtable even though they should
+
+							if (m_useVisibilityForExports && symbolVisibility != STV_HIDDEN)
 								m_exports.emplace(symbolName, ExportInfo{std::string(), 0});//index++});
-							//else if (strncmp(symbolName, "_ZNK14TSoftObjectPtrI", 21) == 0 || strncmp(symbolName, "_ZNK11TSubclassOfI", 18) == 0)
-							//	m_exports.emplace(symbolName, ExportInfo{std::string(), 0});//index++});
-#endif
 						}
 						else
 							m_imports.emplace(symbolName);//index++});
 					}
 					else if (symbolBinding == STB_WEAK)
 					{
-						u32 vis = ELF64_ST_VISIBILITY(symbol.st_other);
-						if (symbolType != STT_NOTYPE && vis == STV_HIDDEN)
+						if (symbolType != STT_NOTYPE)
 						{
-							// Another ugly hack because of dll exported template specialization
-							if (strncmp(symbolName, "_ZNK14TSoftObjectPtrI", 21) == 0)
-								m_exports.emplace(symbolName, ExportInfo{std::string(), 0});//index++});
+							if (symbolVisibility == STV_DEFAULT || symbolVisibility == STV_PROTECTED)
+							{
+								// Another ugly hack because of dll exported template specialization
+							}
+							else
+							{
+								if (strcmp(symbolName, "_ZNK5Chaos24TBoundingVolumeHierarchyI6TArrayIPNS_7TSphereIdLi3EEE22TSizedDefaultAllocatorILi32EEES1_IiS6_EdLi3EE26FindAllIntersectionsHelperERKNS_8TBVHNodeIdLi3EEERKNS_7TVectorIdLi3EEE") == 0)
+									m_exports.emplace(symbolName, ExportInfo{std::string(), 0});//index++});
+								if (strcmp(symbolName, "_ZN5Chaos24TBoundingVolumeHierarchyI6TArrayIPNS_7TSphereIdLi3EEE22TSizedDefaultAllocatorILi32EEES1_IiS6_EdLi3EEC1ERKS7_ibd") == 0)
+									m_exports.emplace(symbolName, ExportInfo{std::string(), 0});//index++});
+							}
 						}
 					}
 				}
 			}
-			else if (section.sh_type == SHT_PROGBITS)
+			else if (!m_useVisibilityForExports && section.sh_type == SHT_PROGBITS)
 			{
 				if (strcmp(sectionName, ".linker_cmd") == 0)
 				{
@@ -238,27 +251,79 @@ namespace uba
 
 	bool ObjectFileElf::StripExports(Logger& logger, u8* newData, const UnorderedSymbols& allExternalImports)
 	{
-		auto& header = *(Elf64Header*)m_data;
-		auto sections = (Elf64SectionHeader*)(m_data + header.e_shoff);
-		char* sectionNamesTable = (char*)m_data + sections[header.e_shstrndx].sh_offset;
-
-		u64 sectionCount = header.e_shnum;
-		if (sectionCount == 0)
-			sectionCount = sections[0].sh_size;
+		auto& header = *(Elf64Header*)newData;
+		auto sections = (Elf64SectionHeader*)(newData + header.e_shoff);
+		char* sectionNamesTable = (char*)newData + sections[header.e_shstrndx].sh_offset;
+		u64 sectionCount = header.e_shnum ? header.e_shnum : sections[0].sh_size;
 
-		for (u64 i=0; i!=sectionCount; ++i)
+		if (m_useVisibilityForExports)
 		{
-			auto& section = sections[i];
-			char* sectionName = sectionNamesTable + section.sh_name;
+			char* symTableNames = (char*)newData + m_symTableNamesOffset;
+			char* dynTableNames = (char*)newData + m_dynTableNamesOffset;
 
-			// Remove all but first symbol.. 
-			// TODO Ideally we want to remove all symbols but then it seems like no stub file is created, maybe we can put a dummy symbol in there?
-			if (section.sh_type == SHT_PROGBITS && strcmp(sectionName, ".linker_cmd") == 0)
+			for (u64 i=0; i!=sectionCount; ++i)
 			{
+				auto& section = sections[i];
+
+				if (section.sh_type != SHT_SYMTAB && section.sh_type != SHT_DYNSYM)
+					continue;
+
+				auto symbols = (Elf64Sym*)(newData + section.sh_offset);
+				u64 symbolCount = section.sh_size / sizeof(Elf64Sym);
+				for (u64 j=0; j!=symbolCount; ++j)
+				{
+					auto& symbol = symbols[j];
+					const char* symbolName = GetSymbolName(sections, symbol, symTableNames, dynTableNames, sectionNamesTable);
+					if (!*symbolName)
+						continue;
+
+					u8 symbolType = ELF64_ST_TYPE(symbol.st_info);
+					u32 symbolBinding = ELF64_ST_BIND(symbol.st_info);
+					u32 symbolVisibility = ELF64_ST_VISIBILITY(symbol.st_other);
+
+					if (symbolBinding == STB_GLOBAL)
+					{
+						if (symbolType != STT_NOTYPE)
+						{
+							if (symbolVisibility == STV_HIDDEN && strncmp(symbolName, "_ZTV", 4) == 0)
+								symbol.st_other = STV_PROTECTED;
+							//else if (symbolVisibility != STV_HIDDEN)
+							//	symbol.st_other = STV_HIDDEN;
+						}
+					}
+					else if (symbolBinding == STB_WEAK && symbolType != STT_NOTYPE)
+					{
+						if (symbolVisibility == STV_DEFAULT)
+						{
+						}
+						else if (symbolVisibility == STV_HIDDEN)
+						{
+							if (strcmp(symbolName, "_ZNK5Chaos24TBoundingVolumeHierarchyI6TArrayIPNS_7TSphereIdLi3EEE22TSizedDefaultAllocatorILi32EEES1_IiS6_EdLi3EE26FindAllIntersectionsHelperERKNS_8TBVHNodeIdLi3EEERKNS_7TVectorIdLi3EEE") == 0)
+								symbol.st_other = STV_PROTECTED;
+							if (strcmp(symbolName, "_ZN5Chaos24TBoundingVolumeHierarchyI6TArrayIPNS_7TSphereIdLi3EEE22TSizedDefaultAllocatorILi32EEES1_IiS6_EdLi3EEC1ERKS7_ibd") == 0)
+								symbol.st_other = STV_PROTECTED;
+						}
+
+					}
+				}
+			}
+		}
+		else
+		{
+			for (u64 i=0; i!=sectionCount; ++i)
+			{
+				auto& section = sections[i];
+				char* sectionName = sectionNamesTable + section.sh_name;
+
+				// Remove all but first symbol.. 
+				// TODO Ideally we want to remove all symbols but then it seems like no stub file is created, maybe we can put a dummy symbol in there?
+				if (section.sh_type != SHT_PROGBITS || strcmp(sectionName, ".linker_cmd") != 0)
+					continue;
+
 				u64 oldSize = section.sh_size;
 				if (!oldSize)
 					return true;
-				u8* begin = m_data + section.sh_offset;
+				u8* begin = newData + section.sh_offset;
 				u8* it = begin;
 
 				it += 4;
@@ -266,6 +331,8 @@ namespace uba
 
 				section.sh_size = it - begin;
 				memset(it, 0, oldSize - section.sh_size);
+
+				break;
 			}
 		}
 		return true;
diff --git a/Engine/Source/Programs/UnrealBuildAccelerator/Common/Private/UbaObjectFileElf.h b/Engine/Source/Programs/UnrealBuildAccelerator/Common/Private/UbaObjectFileElf.h
--- a/Engine/Source/Programs/UnrealBuildAccelerator/Common/Private/UbaObjectFileElf.h
+++ b/Engine/Source/Programs/UnrealBuildAccelerator/Common/Private/UbaObjectFileElf.h
@@ -16,6 +16,10 @@ namespace uba
 
 	private:
 		virtual bool StripExports(Logger& logger, u8* newData, const UnorderedSymbols& allExternalImports) override;
+
+		u64 m_symTableNamesOffset = 0;
+		u64 m_dynTableNamesOffset = 0;
+		bool m_useVisibilityForExports = true;
 	};
 
 	bool IsElfFile(const u8* data, u64 dataSize);
