[UBA]
* Added support for emd files in UbaObjTool

diff --git a/Engine/Source/Programs/UnrealBuildAccelerator/Common/Private/UbaObjectFile.cpp b/Engine/Source/Programs/UnrealBuildAccelerator/Common/Private/UbaObjectFile.cpp
--- a/Engine/Source/Programs/UnrealBuildAccelerator/Common/Private/UbaObjectFile.cpp
+++ b/Engine/Source/Programs/UnrealBuildAccelerator/Common/Private/UbaObjectFile.cpp
@@ -218,7 +218,7 @@ namespace uba
 		return m_potentialDuplicates;
 	}
 
-	bool ObjectFile::CreateExtraFile(Logger& logger, const StringView& extraObjFilename, const StringView& platform, const UnorderedSymbols& allExternalImports, const UnorderedSymbols& allInternalImports, const UnorderedExports& allExports, bool includeExportsInFile)
+	bool ObjectFile::CreateExtraFile(Logger& logger, const StringView& extraObjFilename, const StringView& moduleName, const StringView& platform, const UnorderedSymbols& allExternalImports, const UnorderedSymbols& allInternalImports, const UnorderedExports& allExports, bool includeExportsInFile)
 	{
 		ObjectFileCoff objectFileCoff;
 		ObjectFileElf objectFileElf;
@@ -230,6 +230,8 @@ namespace uba
 			res = ObjectFileCoff::CreateExtraFile(logger, platform, memoryBlock, allExternalImports, allInternalImports, allExports, includeExportsInFile);
 		else if (extraObjFilename.EndsWith(TC("ldscript")))
 			res = CreateDynamicListFile(logger, memoryBlock, allExternalImports, allInternalImports, allExports, includeExportsInFile);
+		else if (extraObjFilename.EndsWith(TC("emd")))
+			res = CreateEmdFile(logger, memoryBlock, moduleName, allExternalImports, allInternalImports, allExports, includeExportsInFile);
 		else
 			res = ObjectFileElf::CreateExtraFile(logger, platform, memoryBlock, allExternalImports, allInternalImports, allExports, includeExportsInFile);
 
@@ -284,10 +286,7 @@ namespace uba
 
 	bool ObjectFile::CreateDynamicListFile(Logger& logger, MemoryBlock& memoryBlock, const UnorderedSymbols& allExternalImports, const UnorderedSymbols& allInternalImports, const UnorderedExports& allExports, bool includeExportsInFile)
 	{
-		auto WriteString = [&](const char* str, u64 strLen)
-			{
-				memcpy(memoryBlock.Allocate(strLen, 1, TC("")), str, strLen);
-			};
+		auto WriteString = [&](const char* str, u64 strLen) { memcpy(memoryBlock.Allocate(strLen, 1, TC("")), str, strLen); };
 
 		//WriteString("VERSION ", 8);
 		WriteString("{", 1);
@@ -310,4 +309,32 @@ namespace uba
 
 		return true;
 	}
+
+	bool ObjectFile::CreateEmdFile(Logger& logger, MemoryBlock& memoryBlock, const StringView& moduleName, const UnorderedSymbols& allExternalImports, const UnorderedSymbols& allInternalImports, const UnorderedExports& allExports, bool includeExportsInFile)
+	{
+		auto WriteString = [&](const char* str, u64 strLen) { memcpy(memoryBlock.Allocate(strLen, 1, TC("")), str, strLen); };
+
+		char moduleName2[256];
+		u32 moduleNameLen = StringBuffer<>(moduleName.data).Parse(moduleName2, 256) - 1;
+
+		WriteString("Library: ", 9);
+		WriteString(moduleName2, moduleNameLen);
+		WriteString(" { export: {\r\n", 14);
+
+		bool symbolAdded = false;
+		for (auto& symbol : allExports)
+			if (allExternalImports.find(symbol.first) != allExternalImports.end())
+			{
+				WriteString(symbol.first.c_str(), symbol.first.size());
+				WriteString("\r\n", 2);
+				symbolAdded = true;
+			}
+
+		if (!symbolAdded)
+			WriteString("ThisIsAnUnrealEngineModule\r\n", 28); // Workaround for tool not liking empty lists
+
+		WriteString("}}", 2);
+
+		return true;
+	}
 }
diff --git a/Engine/Source/Programs/UnrealBuildAccelerator/Common/Public/UbaObjectFile.h b/Engine/Source/Programs/UnrealBuildAccelerator/Common/Public/UbaObjectFile.h
--- a/Engine/Source/Programs/UnrealBuildAccelerator/Common/Public/UbaObjectFile.h
+++ b/Engine/Source/Programs/UnrealBuildAccelerator/Common/Public/UbaObjectFile.h
@@ -39,7 +39,7 @@ namespace uba
 		const UnorderedExports& GetExports() const;
 		const UnorderedSymbols& GetPotentialDuplicates() const;
 
-		static bool CreateExtraFile(Logger& logger, const StringView& extraObjFilename, const StringView& platform, const UnorderedSymbols& allExternalImports, const UnorderedSymbols& allInternalImports, const UnorderedExports& allExports, bool includeExportsInFile);
+		static bool CreateExtraFile(Logger& logger, const StringView& extraObjFilename, const StringView& moduleName, const StringView& platform, const UnorderedSymbols& allExternalImports, const UnorderedSymbols& allInternalImports, const UnorderedExports& allExports, bool includeExportsInFile);
 
 		virtual ~ObjectFile();
 
@@ -52,6 +52,7 @@ namespace uba
 		virtual bool StripExports(Logger& logger, u8* newData, const UnorderedSymbols& allExternalImports) = 0;
 
 		static bool CreateDynamicListFile(Logger& logger, MemoryBlock& memoryBlock, const UnorderedSymbols& allExternalImports, const UnorderedSymbols& allInternalImports, const UnorderedExports& allExports, bool includeExportsInFile);
+		static bool CreateEmdFile(Logger& logger, MemoryBlock& memoryBlock, const StringView& moduleName, const UnorderedSymbols& allExternalImports, const UnorderedSymbols& allInternalImports, const UnorderedExports& allExports, bool includeExportsInFile);
 
 		FileAccessor* m_file = nullptr;
 		u8* m_data = nullptr;
diff --git a/Engine/Source/Programs/UnrealBuildAccelerator/ObjTool/Private/UbaObjTool.cpp b/Engine/Source/Programs/UnrealBuildAccelerator/ObjTool/Private/UbaObjTool.cpp
--- a/Engine/Source/Programs/UnrealBuildAccelerator/ObjTool/Private/UbaObjTool.cpp
+++ b/Engine/Source/Programs/UnrealBuildAccelerator/ObjTool/Private/UbaObjTool.cpp
@@ -48,6 +48,8 @@ namespace uba
 		logger.Info(TC("    /S:<objfile>             Obj file to strip. Will produce a .strip.obj file. Multiple allowed"));
 		logger.Info(TC("    /D:<objfile>             Obj file depending on obj files to strip. Multiple allowed"));
 		logger.Info(TC("    /O:<objfile>             Obj file to output containing exports and loopbacks"));
+		logger.Info(TC("    /T:<platform>            Target platform"));
+		logger.Info(TC("    /M:<module>              Name of module. Needed in emd files"));
 		logger.Info(TC("    /COMPRESS                Write '/O' file compressed"));
 		logger.Info(TC(""));
 		return -1;
@@ -80,6 +82,7 @@ namespace uba
 		std::string impLibName;
 		TString impLibFile;
 		TString platform;
+		TString moduleName;
 
 		auto parseArg = [&](const tchar* arg, bool isRsp)
 			{
@@ -174,6 +177,10 @@ namespace uba
 				{
 					platform = value.data;
 				}
+				else if (name.StartsWith(TC("/M")))
+				{
+					moduleName = value.data;
+				}
 				else if (name.Equals(TC("-printsymbols")))
 				{
 					printSymbols = true;
@@ -312,7 +319,7 @@ namespace uba
 						allExternalImports.insert(symbol);
 				}
 
-				if (!ObjectFile::CreateExtraFile(logger, extraObjFile, platform, allExternalImports, allInternalImports, allExports, true))
+				if (!ObjectFile::CreateExtraFile(logger, extraObjFile, moduleName, platform, allExternalImports, allInternalImports, allExports, true))
 					return -1;
 			}
 			//logger.Info(TC("Reduced export count from %llu to %llu"), totalExportCount.load(), totalKeptExportCount.size());
