[UBA]
* Added comment function's to non-obvious return value

diff --git a/Engine/Source/Programs/UnrealBuildAccelerator/Core/Public/UbaStringBuffer.h b/Engine/Source/Programs/UnrealBuildAccelerator/Core/Public/UbaStringBuffer.h
--- a/Engine/Source/Programs/UnrealBuildAccelerator/Core/Public/UbaStringBuffer.h
+++ b/Engine/Source/Programs/UnrealBuildAccelerator/Core/Public/UbaStringBuffer.h
@@ -66,7 +66,7 @@ namespace uba
 		bool Parse(u32& out);
 		bool Parse(u16& out);
 		bool Parse(float& out);
-		u32 Parse(char* out, u64 outCapacity);
+		u32 Parse(char* out, u64 outCapacity); // Note, return value contains length + null termination.. zero if failed
 
 		u32 count;
 		u32 capacity;
