[UBA]
* Renamed extension for dynamic-list to dynlist from ldscript

diff --git a/Engine/Source/Programs/UnrealBuildAccelerator/Common/Private/UbaObjectFile.cpp b/Engine/Source/Programs/UnrealBuildAccelerator/Common/Private/UbaObjectFile.cpp
--- a/Engine/Source/Programs/UnrealBuildAccelerator/Common/Private/UbaObjectFile.cpp
+++ b/Engine/Source/Programs/UnrealBuildAccelerator/Common/Private/UbaObjectFile.cpp
@@ -228,7 +228,7 @@ namespace uba
 		bool res;
 		if (platform.Equals(TC("win64")) || platform.Equals(TC("wingdk")) || platform.Equals(TC("xb1")) || platform.Equals(TC("xsx"))) 
 			res = ObjectFileCoff::CreateExtraFile(logger, platform, memoryBlock, allExternalImports, allInternalImports, allExports, includeExportsInFile);
-		else if (extraObjFilename.EndsWith(TC("ldscript")))
+		else if (extraObjFilename.EndsWith(TC("dynlist")))
 			res = CreateDynamicListFile(logger, memoryBlock, allExternalImports, allInternalImports, allExports, includeExportsInFile);
 		else if (extraObjFilename.EndsWith(TC("emd")))
 			res = CreateEmdFile(logger, memoryBlock, moduleName, allExternalImports, allInternalImports, allExports, includeExportsInFile);
