[UBA]
* Fixed forever sleeps on macos UbaAgent

diff --git a/Engine/Source/Programs/UnrealBuildAccelerator/Core/Private/UbaPlatform.cpp b/Engine/Source/Programs/UnrealBuildAccelerator/Core/Private/UbaPlatform.cpp
--- a/Engine/Source/Programs/UnrealBuildAccelerator/Core/Private/UbaPlatform.cpp
+++ b/Engine/Source/Programs/UnrealBuildAccelerator/Core/Private/UbaPlatform.cpp
@@ -78,7 +78,7 @@ namespace uba
 		struct timespec ts;
 		ts.tv_sec = milliseconds / 1000;
 		ts.tv_nsec = (milliseconds % 1000) * 1000000;
-		nanosleep(&ts, NULL);
+		nanosleep(&ts, &ts);
 		#else
 		if (milliseconds >= 1000)
 		{
