[UBA]
* Improved posix sleep to retry if interrupted
* Changed so posix directory iterator ignores directories it is not allowed to parse
* Improved ctrl-c experience in UbaAgent.. a double ctrl-c becomes a termination

diff --git a/Engine/Source/Programs/UnrealBuildAccelerator/Agent/Private/UbaAgent.cpp b/Engine/Source/Programs/UnrealBuildAccelerator/Agent/Private/UbaAgent.cpp
--- a/Engine/Source/Programs/UnrealBuildAccelerator/Agent/Private/UbaAgent.cpp
+++ b/Engine/Source/Programs/UnrealBuildAccelerator/Agent/Private/UbaAgent.cpp
@@ -119,6 +119,7 @@ namespace uba
 	LoggerWithWriter* g_logger;
 	SessionClient* g_sessionClient;
 	Atomic<bool> g_shouldExit;
+	Atomic<bool> g_ctrlPressed;
 
 	bool ShouldExit()
 	{
@@ -127,7 +128,11 @@ namespace uba
 
 	void CtrlBreakPressed()
 	{
+		if (g_ctrlPressed)
+			FatalError(13, TC("Force terminate"));
+
 		g_shouldExit = true;
+		g_ctrlPressed = true;
 
 		g_exitLock->EnterWrite();
 		if (g_logger)
diff --git a/Engine/Source/Programs/UnrealBuildAccelerator/Common/Public/UbaDirectoryIterator.h b/Engine/Source/Programs/UnrealBuildAccelerator/Common/Public/UbaDirectoryIterator.h
--- a/Engine/Source/Programs/UnrealBuildAccelerator/Common/Public/UbaDirectoryIterator.h
+++ b/Engine/Source/Programs/UnrealBuildAccelerator/Common/Public/UbaDirectoryIterator.h
@@ -261,7 +261,7 @@ namespace uba
 					res = lstat(dirPath2, &attr);
 				}
 
-				if (errno == ENODEV || errno == EACCES)
+				if (errno == ENODEV || errno == EACCES || errno == EPERM)
 					continue;
 				UBA_ASSERTF(res == 0, "TraverseDir:stat error handling not added for file %s (error %s)", dirPath2, strerror(errno));
 			}
diff --git a/Engine/Source/Programs/UnrealBuildAccelerator/Core/Private/UbaPlatform.cpp b/Engine/Source/Programs/UnrealBuildAccelerator/Core/Private/UbaPlatform.cpp
--- a/Engine/Source/Programs/UnrealBuildAccelerator/Core/Private/UbaPlatform.cpp
+++ b/Engine/Source/Programs/UnrealBuildAccelerator/Core/Private/UbaPlatform.cpp
@@ -78,12 +78,11 @@ namespace uba
 		struct timespec ts;
 		ts.tv_sec = milliseconds / 1000;
 		ts.tv_nsec = (milliseconds % 1000) * 1000000;
-		nanosleep(&ts, &ts);
+		int res;
+		do { res = nanosleep(&ts, &ts); } while (res && errno == EINTR);
 		#else
 		if (milliseconds >= 1000)
-		{
 			sleep(milliseconds / 1000);
-		}
 		usleep((milliseconds % 1000) * 1000);
 		#endif
 	}
