[UBA]
* Implemented detoured mkdir and rmdir for posix
* Added unit tests

diff --git a/Engine/Source/Programs/UnrealBuildAccelerator/Common/Private/UbaSessionClient.cpp b/Engine/Source/Programs/UnrealBuildAccelerator/Common/Private/UbaSessionClient.cpp
--- a/Engine/Source/Programs/UnrealBuildAccelerator/Common/Private/UbaSessionClient.cpp
+++ b/Engine/Source/Programs/UnrealBuildAccelerator/Common/Private/UbaSessionClient.cpp
@@ -677,11 +677,8 @@ namespace uba
 		out.result = reader.ReadBool();
 		out.errorCode = reader.ReadU32();
 		if (out.result)
-		{
-			reader.Reset();
-			if (!SendUpdateDirectoryTable(reader))
+			if (!SendUpdateDirectoryTable(reader.Reset()))
 				return false;
-		}
 		out.directoryTableSize = GetDirectoryTableSize();
 		return true;
 	}
@@ -708,11 +705,8 @@ namespace uba
 			out.closeId = ~0u;
 			out.errorCode = reader.ReadU32();
 			if (!out.errorCode)
-			{
-				reader.Reset();
-				if (!SendUpdateDirectoryTable(reader))
+				if (!SendUpdateDirectoryTable(reader.Reset()))
 					return false;
-			}
 			out.directoryTableSize = GetDirectoryTableSize();
 			return true;
 		}
@@ -811,11 +805,17 @@ namespace uba
 		StackBinaryWriter<1024> writer;
 		NetworkMessage networkMsg(m_client, ServiceId, SessionMessageType_CreateDirectory, writer);
 		writer.WriteString(msg.name);
-		StackBinaryReader<1024> reader;
+		StackBinaryReader<SendMaxSize> reader;
 		if (!networkMsg.Send(reader, Stats().createDirMsg))
 			return false;
 		out.result = reader.ReadBool();
 		out.errorCode = reader.ReadU32();
+
+		if (out.result)
+			if (!SendUpdateDirectoryTable(reader.Reset()))
+				return false;
+
+		out.directoryTableSize = GetDirectoryTableSize();
 		return true;
 	}
 
@@ -824,11 +824,17 @@ namespace uba
 		StackBinaryWriter<1024> writer;
 		NetworkMessage networkMsg(m_client, ServiceId, SessionMessageType_RemoveDirectory, writer);
 		writer.WriteString(msg.name);
-		StackBinaryReader<1024> reader;
+		StackBinaryReader<SendMaxSize> reader;
 		if (!networkMsg.Send(reader, Stats().deleteFileMsg)) // Wrong message
 			return false;
 		out.result = reader.ReadBool();
 		out.errorCode = reader.ReadU32();
+
+		if (out.result)
+			if (!SendUpdateDirectoryTable(reader.Reset()))
+				return false;
+
+		out.directoryTableSize = GetDirectoryTableSize();
 		return true;
 	}
 
@@ -935,7 +941,7 @@ namespace uba
 			item.done.Create(true);
 
 			lock.Leave();
-			bool res = item.done.IsSet(5*60*60);
+			bool res = item.done.IsSet(5*60*1000);
 			lock.Enter();
 
 			if (item.prev)
@@ -1049,8 +1055,7 @@ namespace uba
 			NetworkMessage msg(m_client, ServiceId, SessionMessageType_GetDirectoriesFromServer, writer);
 			writer.WriteU32(m_sessionId);
 
-			reader.Reset();
-			if (msg.Send(reader, Stats().getDirsMsg))
+			if (msg.Send(reader.Reset(), Stats().getDirsMsg))
 				continue;
 
 			// Let's signal waiters to exit faster since we will not get out of this situation (most likely a disconnect)
@@ -1081,9 +1086,7 @@ namespace uba
 				NetworkMessage msg(m_client, ServiceId, SessionMessageType_GetNameToHashFromServer, writer);
 				writer.WriteU32(serverTableSize);
 				writer.WriteU32(localTableSize);
-
-				reader.Reset();
-				if (!msg.Send(reader, Stats().getHashesMsg))
+				if (!msg.Send(reader.Reset(), Stats().getHashesMsg))
 					return false;
 			}
 			serverTime = reader.ReadU64();
@@ -1315,14 +1318,9 @@ namespace uba
 		}
 
 		if (!out.empty())
-		{
 			if (neededDirectoryTableSize > GetDirectoryTableSize())
-			{
-				reader.Reset();
-				if (!SendUpdateDirectoryTable(reader))
+				if (!SendUpdateDirectoryTable(reader.Reset()))
 					return false;
-			}
-		}
 
 		// Always nice to update name-to-hash table since it can reduce number of messages while building.
 		u32 hashTableMemSize;
@@ -1331,11 +1329,8 @@ namespace uba
 			hashTableMemSize = u32(m_nameToHashTableMem.writtenSize);
 		}
 		if (neededHashTableSize > hashTableMemSize)
-		{
-			reader.Reset();
-			if (!SendUpdateNameToHashTable(reader))
+			if (!SendUpdateNameToHashTable(reader.Reset()))
 				return false;
-		}
 
 		return true;
 	}
@@ -1977,8 +1972,7 @@ namespace uba
 			}
 		}
 
-		reader.Reset();
-		return SendUpdateDirectoryTable(reader);
+		return SendUpdateDirectoryTable(reader.Reset());
 	}
 
 	bool SessionClient::CustomMessage(Process& process, BinaryReader& reader, BinaryWriter& writer)
diff --git a/Engine/Source/Programs/UnrealBuildAccelerator/Core/Public/UbaBinaryReaderWriter.h b/Engine/Source/Programs/UnrealBuildAccelerator/Core/Public/UbaBinaryReaderWriter.h
--- a/Engine/Source/Programs/UnrealBuildAccelerator/Core/Public/UbaBinaryReaderWriter.h
+++ b/Engine/Source/Programs/UnrealBuildAccelerator/Core/Public/UbaBinaryReaderWriter.h
@@ -92,7 +92,7 @@ namespace uba
 	struct StackBinaryReader : BinaryReader
 	{
 		StackBinaryReader() : BinaryReader(buffer, 0, Capacity) { *buffer = 0; }
-		void Reset() { m_pos = m_begin; m_end = buffer + Capacity; }
+		StackBinaryReader& Reset() { m_pos = m_begin; m_end = buffer + Capacity; return *this; }
 		u8 buffer[Capacity];
 	};
 
diff --git a/Engine/Source/Programs/UnrealBuildAccelerator/Detours/Private/PosixOS/UbaDetoursFunctionsPosix.cpp b/Engine/Source/Programs/UnrealBuildAccelerator/Detours/Private/PosixOS/UbaDetoursFunctionsPosix.cpp
--- a/Engine/Source/Programs/UnrealBuildAccelerator/Detours/Private/PosixOS/UbaDetoursFunctionsPosix.cpp
+++ b/Engine/Source/Programs/UnrealBuildAccelerator/Detours/Private/PosixOS/UbaDetoursFunctionsPosix.cpp
@@ -40,6 +40,8 @@ using namespace uba;
 #define DETOURED_FUNCTIONS \
 	DETOURED_FUNCTION(chdir) \
 	DETOURED_FUNCTION(fchdir) \
+	DETOURED_FUNCTION(mkdir) \
+	DETOURED_FUNCTION(rmdir) \
 	DETOURED_FUNCTION(chroot) \
 	DETOURED_FUNCTION(getcwd) \
 	DETOURED_FUNCTION(getenv) \
@@ -463,7 +465,16 @@ FILE* Shared_fopen(const char* funcName, const char* path, const char* mode, con
 	bool p = strchr(mode, '+');
 	//bool b = strchr(mode, 'b');
 	//bool t = strchr(mode, 't');
-	UBA_ASSERTF(!a, "%s with append not implemented (%s)", funcName, mode);
+	if (a)
+	{
+		//if (StartsWith(path, g_systemTemp.data))
+		{
+			FILE* res = TRUE_WRAPPER(fopen)(path, mode);
+			DEBUG_LOG_TRUE(funcName, "(%s  %s) -> %p", path, mode, res);
+			return res;
+		}
+		//UBA_ASSERTF(false, "%s with append (%s) not implemented (%s)", funcName, mode, path);
+	}
 
 	int flags = 0;
 	if (r)
@@ -704,8 +715,16 @@ UBA_EXPORT int UBA_WRAPPER(_NSGetExecutablePath)(char* buf, uint32_t* bufsize)
 UBA_EXPORT int UBA_WRAPPER(chdir)(const char* path)
 {
 	UBA_INIT_DETOUR(chdir, path);
-	UBA_ASSERTF(false, "chdir not implemented");
-	return TRUE_WRAPPER(chdir)(path);
+	FixPath(g_virtualWorkingDir.Clear(), path);
+	g_virtualWorkingDir.EnsureEndsWithSlash();
+	if (g_runningRemote)
+	{
+		DEBUG_LOG_DETOURED("chdir", "%s -> 0", path);
+		return 0;
+	}
+	int res = TRUE_WRAPPER(chdir)(path);
+	DEBUG_LOG_TRUE("chdir", "%s -> %i", path, res);
+	return res;
 }
 
 UBA_EXPORT int UBA_WRAPPER(fchdir)(int fd)
@@ -715,6 +734,85 @@ UBA_EXPORT int UBA_WRAPPER(fchdir)(int fd)
 	return TRUE_WRAPPER(fchdir)(fd);
 }
 
+UBA_EXPORT int UBA_WRAPPER(mkdir)(const char* path, mode_t mode)
+{
+	UBA_INIT_DETOUR(mkdir, path, mode);
+
+	StringBuffer<> pathName;
+	FixPath(pathName, path);
+	if (pathName.StartsWith(g_systemTemp.data))
+	{
+		int res = TRUE_WRAPPER(mkdir)(path, mode);
+		DEBUG_LOG_TRUE("mkdir", "%s -> %i", path, res);
+		return res;
+	}
+
+	u32 directoryTableSize;
+	bool res;
+	u32 errorCode = 0;
+	StringKey pathNameKey = ToStringKeyLower(pathName);
+
+	{
+		TimerScope ts(g_stats.createFile);
+		SCOPED_WRITE_LOCK(g_communicationLock, pcs);
+		BinaryWriter writer;
+		writer.WriteByte(MessageType_CreateDirectory);
+		writer.WriteStringKey(pathNameKey);
+		writer.WriteString(pathName);
+		writer.Flush();
+		BinaryReader reader;
+		res = reader.ReadBool();
+		errorCode = reader.ReadU32();
+		directoryTableSize = reader.ReadU32();
+	}
+
+	g_directoryTable.ParseDirectoryTable(directoryTableSize);
+
+	errno = errorCode;
+	DEBUG_LOG_DETOURED("mkdir", "%ls -> %i (%u)", path, res, errorCode);
+	return res ? 0 : -1;
+}
+
+UBA_EXPORT int UBA_WRAPPER(rmdir)(const char* path)
+{
+	UBA_INIT_DETOUR(rmdir, path);
+
+	StringBuffer<> pathName;
+	FixPath(pathName, path);
+	if (pathName.StartsWith(g_systemTemp.data))
+	{
+		//SuppressCreateFileDetourScope s; // TODO: Revisit this.. will not work remotely
+		int res = TRUE_WRAPPER(rmdir)(path);
+		DEBUG_LOG_TRUE("rmdir", "%s -> %i", path, res);
+		return res;
+	}
+
+	u32 directoryTableSize;
+	bool res;
+	u32 errorCode = 0;
+	StringKey pathNameKey = ToStringKeyLower(pathName);
+
+	{
+		TimerScope ts(g_stats.deleteFile);
+		SCOPED_WRITE_LOCK(g_communicationLock, pcs);
+		BinaryWriter writer;
+		writer.WriteByte(MessageType_RemoveDirectory);
+		writer.WriteStringKey(pathNameKey);
+		writer.WriteString(pathName);
+		writer.Flush();
+		BinaryReader reader;
+		res = reader.ReadBool();
+		errorCode = reader.ReadU32();
+		directoryTableSize = reader.ReadU32();
+	}
+
+	g_directoryTable.ParseDirectoryTable(directoryTableSize);
+
+	errno = errorCode;
+	DEBUG_LOG_DETOURED("rmdir", "%ls -> %i (%u)", path, res, errorCode);
+	return res ? 0 : -1;
+}
+
 UBA_EXPORT int UBA_WRAPPER(chroot)(const char* path)
 {
 	UBA_INIT_DETOUR(chroot, path);
diff --git a/Engine/Source/Programs/UnrealBuildAccelerator/TestApp/Private/PosixOS/UbaTestAppPosix.cpp b/Engine/Source/Programs/UnrealBuildAccelerator/TestApp/Private/PosixOS/UbaTestAppPosix.cpp
--- a/Engine/Source/Programs/UnrealBuildAccelerator/TestApp/Private/PosixOS/UbaTestAppPosix.cpp
+++ b/Engine/Source/Programs/UnrealBuildAccelerator/TestApp/Private/PosixOS/UbaTestAppPosix.cpp
@@ -155,6 +155,21 @@ int main(int argc, char* argv[])
 		struct stat attrRoot;
 		if (stat("/", &attrRoot) != 0)
 			return LogError("stat for '/' failed");
+
+		if (mkdir("FooDir", S_IRUSR | S_IWUSR) != 0)
+			return LogError("mkdir 'FooDir' failed");
+
+		struct stat attrFoo;
+		if (stat("FooDir", &attrFoo) != 0)
+			return LogError("stat for 'FooDir' failed");
+		if (!S_ISDIR(attrFoo.st_mode))
+			return LogError("stat for dir 'FooDir' returned wrong type");
+
+		if (rmdir("FooDir") != 0)
+			return LogError("rmdir 'FooDir' failed");
+
+		if (stat("FooDir", &attrFoo) == 0)
+			return LogError("stat for 'FooDir' failed to not find removed directory");
 		return 0;
 	}
 	else if (strncmp(argv[1], "-file=", 6) == 0)
