[UBA]
* Added stubs and code to parse mac binaries for dependencies

diff --git a/Engine/Source/Programs/UnrealBuildAccelerator/Common/Private/UbaSession.cpp b/Engine/Source/Programs/UnrealBuildAccelerator/Common/Private/UbaSession.cpp
--- a/Engine/Source/Programs/UnrealBuildAccelerator/Common/Private/UbaSession.cpp
+++ b/Engine/Source/Programs/UnrealBuildAccelerator/Common/Private/UbaSession.cpp
@@ -23,6 +23,7 @@
 
 // Used to get Memory Information
 #if PLATFORM_MAC
+#include "UbaMacBinDependencyParser.h"
 #include <mach/vm_statistics.h>
 #include <mach/mach_types.h>
 #include <mach/mach_init.h>
@@ -930,14 +931,16 @@ namespace uba
 
 	bool Session::CopyImports(Vector<BinaryModule>& out, const tchar* library, tchar* applicationDir, tchar* applicationDirEnd, UnorderedSet<TString>& handledImports)
 	{
-		#if PLATFORM_WINDOWS
 		if (!handledImports.insert(library).second)
 			return true;
-		swprintf_s(applicationDirEnd, 512 - (applicationDirEnd - applicationDir), TC("%s"), library);
+		TSprintf_s(applicationDirEnd, 512 - (applicationDirEnd - applicationDir), TC("%s"), library);
 		const tchar* applicationName = applicationDir;
 		u32 attr = GetFileAttributesW(applicationName); // TODO: Use attributes table
 		tchar temp[512];
 		tchar temp2[512];
+		bool result = true;
+
+		#if PLATFORM_WINDOWS
 		if (attr == INVALID_FILE_ATTRIBUTES)
 		{
 			if (!SearchPathW(NULL, library, NULL, 512, temp, NULL))
@@ -963,16 +966,38 @@ namespace uba
 
 		out.push_back({ library, temp3.data, attr, isSystem });
 
-		bool result = true;
 		FindImports(applicationName, [&](const tchar* importName, bool isKnown)
 			{
 				if (result && !isKnown)
 					result = CopyImports(out, importName, applicationDir, applicationDirEnd, handledImports);
 			});
-		return result;
+
 		#else
-		return true;
+
+		if (attr == INVALID_FILE_ATTRIBUTES)
+		{
+			UBA_ASSERTF(false, TC("DIR NOT FOUND: %s"), applicationName);
+			return false;
+		}
+
+		StringBuffer<512> temp3;
+		FixPath(applicationName, nullptr, 0, temp3);
+
+		out.push_back({ TString(library), TString(temp3.data), S_IRUSR | S_IWUSR | S_IXUSR });
+
+		#if PLATFORM_MAC
+		FindImportsMac(applicationName, [&](const tchar* importName, bool isKnown)
+			{
+				if (result && !isKnown)
+					result = CopyImports(out, importName, applicationDir, applicationDirEnd, handledImports);
+			});
+		#endif
+
+		// This code is needed if application is compiled with tsan
+		//strcpy(applicationDirEnd, "libclang_rt.tsan.so");
+		//out.push_back({ "libclang_rt.tsan.so", applicationDir, S_IRUSR | S_IWUSR });
 		#endif
+		return result;
 	}
 
 	Session::Session(const SessionCreateInfo& info, const tchar* logPrefix, bool runningRemote, WorkManager* workManager)
@@ -1628,19 +1653,8 @@ namespace uba
 		memcpy(applicationDir, application, applicationDirLen * sizeof(tchar));
 		tchar* applicationDirEnd = applicationDir + applicationDirLen;
 
-#if PLATFORM_WINDOWS
-
 		UnorderedSet<TString> handledImports;
-		CopyImports(out, applicationName, applicationDir, applicationDirEnd, handledImports);
-#else
-		// TODO: This. Does non-windows have dlls that needs to be downloaded here?
-		out.push_back({ TString(applicationName), TString(application), S_IRUSR | S_IWUSR | S_IXUSR });
-		
-		// This code is needed if application is compiled with tsan
-		//strcpy(applicationDirEnd, "libclang_rt.tsan.so");
-		//out.push_back({ "libclang_rt.tsan.so", applicationDir, S_IRUSR | S_IWUSR });
-#endif
-		return true;
+		return CopyImports(out, applicationName, applicationDir, applicationDirEnd, handledImports);
 	}
 
 	void Session::Free(Vector<BinaryModule>& v)
diff --git a/Engine/Source/Programs/UnrealBuildAccelerator/Core/Public/UbaMacBinDependencyParser.h b/Engine/Source/Programs/UnrealBuildAccelerator/Core/Public/UbaMacBinDependencyParser.h
new file mode 100644
--- /dev/null
+++ b/Engine/Source/Programs/UnrealBuildAccelerator/Core/Public/UbaMacBinDependencyParser.h
@@ -0,0 +1,56 @@
+// Copyright Epic Games, Inc. All Rights Reserved.
+
+#pragma once
+
+#include "UbaStringBuffer.h"
+
+#if !defined(UBA_IS_DETOURED_INCLUDE)
+#define TRUE_WRAPPER(func) func
+#endif
+
+namespace uba
+{
+	inline bool FindImportsMac(const tchar* fileName, const Function<void(const tchar* import, bool isKnown)>& func)
+	{
+#if 0
+		HANDLE fileHandle = ::CreateFileW(fileName, GENERIC_READ, FILE_SHARE_READ, 0, OPEN_EXISTING, FILE_ATTRIBUTE_NORMAL | FILE_FLAG_SEQUENTIAL_SCAN, 0);
+		if (fileHandle == INVALID_HANDLE_VALUE)
+			return true;
+		auto closeFileHandle = MakeGuard([&]() { CloseHandle(fileHandle); });
+		HANDLE fileMapping = ::CreateFileMappingW(fileHandle, NULL, PAGE_READONLY, 0, 0, NULL);
+		if (!fileMapping)
+			return false;
+		auto closeMappingHandle = MakeGuard([&]() { ::CloseHandle(fileMapping); });
+		void* mem = ::MapViewOfFile(fileMapping, FILE_MAP_READ, 0, 0, 0);
+		if (!mem)
+			return false;
+		auto unmap = MakeGuard([&]() { ::UnmapViewOfFile(mem); });
+
+		auto it = (const u8*)mem;
+		while (true)
+		{
+			if (*it++ != '@')
+				continue;
+
+			if (memcmp(it, "executable_path", 14) == 0)
+				break;
+
+			if (memcmp(it, "rpath/", 6) != 0)
+				continue;
+			it += 6;
+			auto importFile = (const char*)it;
+			if (!strstr(importFile, ".dylib"))
+				return false;
+
+#if PLATFORM_WINDOWS
+			StringBuffer<> temp;
+			temp.Append(importFile);
+			func(temp.data, false);
+#else
+			func(importFile, false);
+#endif
+		}
+#endif
+		return true;
+	}
+}
\ No newline at end of file
diff --git a/Engine/Source/Programs/UnrealBuildAccelerator/Core/Public/UbaWinBinDependencyParser.h b/Engine/Source/Programs/UnrealBuildAccelerator/Core/Public/UbaWinBinDependencyParser.h
--- a/Engine/Source/Programs/UnrealBuildAccelerator/Core/Public/UbaWinBinDependencyParser.h
+++ b/Engine/Source/Programs/UnrealBuildAccelerator/Core/Public/UbaWinBinDependencyParser.h
@@ -97,7 +97,7 @@ namespace uba
 		return rva - seh->VirtualAddress + seh->PointerToRawData;
 	} 
 
-	inline bool FindImportsInMem(LPCWSTR fileName, void* mem, const Function<void(const tchar* import, bool isKnown)>& func)
+	inline bool FindImportsInMem(const tchar* fileName, void* mem, const Function<void(const tchar* import, bool isKnown)>& func)
 	{
 		auto hdrs = (PIMAGE_NT_HEADERS)(PCHAR(mem) + PIMAGE_DOS_HEADER(mem)->e_lfanew);   
 		auto pSech=IMAGE_FIRST_SECTION(hdrs);
@@ -115,7 +115,6 @@ namespace uba
 				wname.Append(name);
 				func(wname.data, IsKnownSystemFile(wname.data));
 				++importDesc;
-
 			}
 		}
 		__except(GetExceptionCode() == EXCEPTION_ACCESS_VIOLATION ? EXCEPTION_EXECUTE_HANDLER : EXCEPTION_CONTINUE_SEARCH)
@@ -128,7 +127,7 @@ namespace uba
 		return true;
 	}
 
-	inline bool FindImports(LPCWSTR fileName, const Function<void(const tchar* import, bool isKnown)>& func)
+	inline bool FindImports(const tchar* fileName, const Function<void(const tchar* import, bool isKnown)>& func)
 	{
 		HANDLE fileHandle = True_CreateFileW(fileName, GENERIC_READ, FILE_SHARE_READ, 0, OPEN_EXISTING, FILE_ATTRIBUTE_NORMAL | FILE_FLAG_SEQUENTIAL_SCAN, 0);
 		if (fileHandle == INVALID_HANDLE_VALUE)
diff --git a/Engine/Source/Programs/UnrealBuildAccelerator/Test/Private/UbaTestAll.h b/Engine/Source/Programs/UnrealBuildAccelerator/Test/Private/UbaTestAll.h
--- a/Engine/Source/Programs/UnrealBuildAccelerator/Test/Private/UbaTestAll.h
+++ b/Engine/Source/Programs/UnrealBuildAccelerator/Test/Private/UbaTestAll.h
@@ -35,6 +35,7 @@ namespace uba
 		UBA_TEST(TestCacheEntry) \
 		UBA_TEST(TestConfig) \
 		UBA_TEST(TestCrypto) \
+		UBA_TEST(TestBinDependencies) \
 
 
 #define UBA_POSIX_TESTS \
diff --git a/Engine/Source/Programs/UnrealBuildAccelerator/Test/Private/UbaTestBasics.h b/Engine/Source/Programs/UnrealBuildAccelerator/Test/Private/UbaTestBasics.h
--- a/Engine/Source/Programs/UnrealBuildAccelerator/Test/Private/UbaTestBasics.h
+++ b/Engine/Source/Programs/UnrealBuildAccelerator/Test/Private/UbaTestBasics.h
@@ -20,6 +20,8 @@
 
 #if PLATFORM_WINDOWS
 #include "UbaWinBinDependencyParser.h"
+#elif PLATFORM_MAC
+#include "UbaMacBinDependencyParser.h"
 #endif
 
 #define VA_ARGS(...) , ##__VA_ARGS__
@@ -457,4 +459,27 @@ namespace uba
 
 		return true;
 	}
+
+	bool TestBinDependencies(Logger& logger, const StringBufferBase& rootDir)
+	{
+
+#if PLATFORM_WINDOWS
+		StringBuffer<> path;
+		GetDirectoryOfCurrentModule(logger, path);
+		path.EnsureEndsWithSlash().Append(TC("UbaTestApp.exe"));
+		bool importKernel = false;
+		FindImports(path.data, [&](const tchar* import, bool isKnown)
+		{
+			importKernel = isKnown && Contains(import, TC("KERNEL32.dll"));
+		});
+		if (!importKernel)
+			return logger.Error(TC("Failed to find Kernel32 as import"));
+#elif PLATFORM_MAC
+		//FindImportsMac(TC(""), [&](const tchar* import, bool isKnown)
+		//{
+		//	LoggerWithWriter(g_consoleLogWriter).Info(TC("IMPORT: %s"), import);
+		//});
+#endif
+		return true;
+	}
 }
