[UBA]
* Added so it is possible to configure number of workers when spawning UbaAgent (mostly for debugging)

diff --git a/Engine/Source/Programs/UnrealBuildAccelerator/Agent/Private/UbaAgent.cpp b/Engine/Source/Programs/UnrealBuildAccelerator/Agent/Private/UbaAgent.cpp
--- a/Engine/Source/Programs/UnrealBuildAccelerator/Agent/Private/UbaAgent.cpp
+++ b/Engine/Source/Programs/UnrealBuildAccelerator/Agent/Private/UbaAgent.cpp
@@ -80,6 +80,7 @@ namespace uba
 		logger.Info(TC("  -maxcpu=<number>        Max number of processes that can be started. Defaults to \"%u\" on this machine"), DefaultProcessorCount);
 		logger.Info(TC("  -mulcpu=<number>        This value multiplies with number of cpu to figure out max cpu. Defaults to 1.0"));
 		logger.Info(TC("  -maxcon=<number>        Max number of connections that can be started by agent. Defaults to \"%u\" (amount up to max will depend on ping)"), DefaultMaxConnectionCount);
+		logger.Info(TC("  -maxworkers=<number>    Max number of workers is started by agent. Defaults to \"%u\""), DefaultProcessorCount);
 		logger.Info(TC("  -capacity=<gigaby>      Capacity of local store. Defaults to %u gigabytes"), DefaultCapacityGb);
 		logger.Info(TC("  -config=<file>          Config file that contains options for various systems"));
 		logger.Info(TC("  -quic                   Use Quic instead of tcp backend."));
@@ -401,6 +402,7 @@ namespace uba
 		#endif
 
 		u32 maxProcessCount = DefaultProcessorCount;
+		u32 maxWorkerCount = DefaultProcessorCount;
 		float mulProcessValue = 1.0f;
 		u32 maxConnectionCount = DefaultMaxConnectionCount;
 		u32 outputStatsThresholdMs = 0;
@@ -487,6 +489,11 @@ namespace uba
 				if (!value.Parse(maxConnectionCount) || maxConnectionCount == 0)
 					return PrintHelp(TC("Invalid value for -maxcon"));
 			}
+			else if (name.Equals(TC("-maxworkers")))
+			{
+				if (!value.Parse(maxWorkerCount))
+					return PrintHelp(TC("Invalid value for -maxworkers"));
+			}
 			else if (name.Equals(TC("-capacity")))
 			{
 				if (!value.Parse(storageCapacityGb))
@@ -999,6 +1006,7 @@ namespace uba
 			//ncci.Apply(config);
 			ncci.sendSize = sendSize;
 			ncci.receiveTimeoutSeconds = receiveTimeoutSeconds;
+			ncci.workerCount = maxWorkerCount;
 			if (hasCrypto)
 				ncci.cryptoKey128 = crypto;
 			bool ctorSuccess = true;
diff --git a/Engine/Source/Programs/UnrealBuildAccelerator/Common/Private/UbaNetworkClient.cpp b/Engine/Source/Programs/UnrealBuildAccelerator/Common/Private/UbaNetworkClient.cpp
--- a/Engine/Source/Programs/UnrealBuildAccelerator/Common/Private/UbaNetworkClient.cpp
+++ b/Engine/Source/Programs/UnrealBuildAccelerator/Common/Private/UbaNetworkClient.cpp
@@ -11,7 +11,7 @@
 namespace uba
 {
 	NetworkClient::NetworkClient(bool& outCtorSuccess, const NetworkClientCreateInfo& info, const tchar* name)
-	: WorkManagerImpl(GetLogicalProcessorCount())
+	: WorkManagerImpl(info.workerCount == 0 ? GetLogicalProcessorCount() : info.workerCount)
 	,	m_logWriter(info.logWriter)
 	,	m_logger(info.logWriter, SetGetPrefix(name))
 	,	m_isConnected(true)
diff --git a/Engine/Source/Programs/UnrealBuildAccelerator/Common/Public/UbaNetworkClient.h b/Engine/Source/Programs/UnrealBuildAccelerator/Common/Public/UbaNetworkClient.h
--- a/Engine/Source/Programs/UnrealBuildAccelerator/Common/Public/UbaNetworkClient.h
+++ b/Engine/Source/Programs/UnrealBuildAccelerator/Common/Public/UbaNetworkClient.h
@@ -24,6 +24,7 @@ namespace uba
 		LogWriter& logWriter;
 		u32 sendSize = SendDefaultSize;
 		u32 receiveTimeoutSeconds = 0;
+		u32 workerCount = 0; // Zero means logical processor count will be used
 		const u8* cryptoKey128 = nullptr;
 	};
 
