[UBA]
* Fixed so locale is ignored when writing logs. This to prevent crashes when logging deep down in core functions

diff --git a/Engine/Source/Programs/UnrealBuildAccelerator/Detours/Private/UbaDetoursShared.cpp b/Engine/Source/Programs/UnrealBuildAccelerator/Detours/Private/UbaDetoursShared.cpp
--- a/Engine/Source/Programs/UnrealBuildAccelerator/Detours/Private/UbaDetoursShared.cpp
+++ b/Engine/Source/Programs/UnrealBuildAccelerator/Detours/Private/UbaDetoursShared.cpp
@@ -59,10 +59,6 @@ namespace uba
 	thread_local char t_b[LogBufSize];
 	thread_local u32 t_b_size;
 
-	StringBufferBase& GetLogTlsBuffer()
-	{
-		return t_a;
-	}
 	void GetPrefixExtra(StringBufferBase& out)
 	{
 		#if 0
@@ -75,8 +71,23 @@ namespace uba
 		#endif
 		//out.Appendf(TC("[%7u]"), GetCurrentThreadId());
 	}
-	void WriteDebugLogWithPrefix(const char* prefix, LogScope& scope)
+	void WriteDebugLogWithPrefix(const char* prefix, LogScope& scope, const tchar* command, const tchar* format, ...)
 	{
+		#if PLATFORM_MAC
+		static locale_t safeLocale = newlocale(LC_NUMERIC_MASK, "C", duplocale(LC_GLOBAL_LOCALE));
+		locale_t oldLocale = uselocale(safeLocale);
+		#endif
+
+		t_a.Clear().Append(command).Append(' ');
+		if (*format)
+		{
+			va_list arg;
+			va_start(arg, format);
+			t_a.Append(format, arg);
+			va_end(arg);
+		}
+		t_a.Append(TC("\n"));
+
 		u32 size__ = t_b_size;
 		StringBuffer<128> extra;
 		GetPrefixExtra(extra);
@@ -89,9 +100,24 @@ namespace uba
 		if (res__ != -1)
 			t_b_size += res__;
 		scope.Flush();
+
+		#if PLATFORM_MAC
+		uselocale(oldLocale);
+		#endif
 	}
-	void WriteDebugLog()
+
+	void WriteDebugLog(const tchar* format, ...)
 	{
+		t_a.Clear();
+		if (*format)
+		{
+			va_list arg;
+			va_start(arg, format);
+			t_a.Append(format, arg);
+			va_end(arg);
+		}
+		t_a.Append(TC("\n"));
+
 		#if PLATFORM_WINDOWS
 		t_b_size = sprintf_s(t_b, LogBufSize, "%S", t_a.data);
 		WriteDebug(t_b, t_b_size);
diff --git a/Engine/Source/Programs/UnrealBuildAccelerator/Detours/Private/UbaDetoursShared.h b/Engine/Source/Programs/UnrealBuildAccelerator/Detours/Private/UbaDetoursShared.h
--- a/Engine/Source/Programs/UnrealBuildAccelerator/Detours/Private/UbaDetoursShared.h
+++ b/Engine/Source/Programs/UnrealBuildAccelerator/Detours/Private/UbaDetoursShared.h
@@ -26,11 +26,7 @@ namespace uba
 	#if UBA_DEBUG_LOG_ENABLED
 		#define DEBUG_LOG_PREFIX(Prefix, Command, ...) \
 			LogScope STRING_JOIN(ls, __LINE__); \
-			if (isLogging()) \
-			{ \
-				GetLogTlsBuffer().Clear().Append(Command).Append(' ').Appendf(__VA_ARGS__).Append(TC("\n")); \
-				WriteDebugLogWithPrefix(#Prefix, STRING_JOIN(ls, __LINE__)); \
-			}
+			if (isLogging()) WriteDebugLogWithPrefix(#Prefix, STRING_JOIN(ls, __LINE__), Command, __VA_ARGS__); \
 
 		//#define DEBUG_LOG_DETOURED(Command, ...) 
 		#define DEBUG_LOG_DETOURED(Command, ...) DEBUG_LOG_PREFIX(D, Command, __VA_ARGS__)
@@ -39,7 +35,7 @@ namespace uba
 		//#define DEBUG_LOG_PIPE(Command, ...) ts.leave(); DEBUG_LOG_PREFIX(P, Command, __VA_ARGS__)
 		#define DEBUG_LOG_PIPE(Command, ...) ts.Leave();
 		//#define DEBUG_LOG(...)
-		#define DEBUG_LOG(...) { if (isLogging()) { GetLogTlsBuffer().Clear().Appendf(__VA_ARGS__).Append(TC("\n")); WriteDebugLog(); }}
+		#define DEBUG_LOG(...) { if (isLogging()) WriteDebugLog(__VA_ARGS__); }
 	#else
 		#define DEBUG_LOG(...)
 		#define DEBUG_LOG_DETOURED(Command, ...)
@@ -75,9 +71,8 @@ namespace uba
 	extern FileHandle g_debugFile;
 	inline bool isLogging() { return g_debugFile != InvalidFileHandle; }
 	struct LogScope { LogScope(); ~LogScope(); void Flush(); };
-	StringBufferBase& GetLogTlsBuffer();
-	void WriteDebugLogWithPrefix(const char* prefix, LogScope& scope);
-	void WriteDebugLog();
+	void WriteDebugLogWithPrefix(const char* prefix, LogScope& scope, const tchar* command, const tchar* format, ...);
+	void WriteDebugLog(const tchar* format, ...);
 	void FlushDebugLog();
 	#endif
 	#if UBA_DEBUG_VALIDATE
