[UBA]
* Fixed code parsing binary dylib dependencies. still WIP

diff --git a/Engine/Source/Programs/UnrealBuildAccelerator/Common/Private/UbaSession.cpp b/Engine/Source/Programs/UnrealBuildAccelerator/Common/Private/UbaSession.cpp
--- a/Engine/Source/Programs/UnrealBuildAccelerator/Common/Private/UbaSession.cpp
+++ b/Engine/Source/Programs/UnrealBuildAccelerator/Common/Private/UbaSession.cpp
@@ -986,11 +986,13 @@ namespace uba
 		out.push_back({ TString(library), TString(temp3.data), S_IRUSR | S_IWUSR | S_IXUSR });
 
 		#if PLATFORM_MAC
-		FindImportsMac(applicationName, [&](const tchar* importName, bool isKnown)
+		StringBuffer<> errorStr;
+		if (!FindImportsMac(applicationName, [&](const tchar* importName, bool isKnown)
 			{
 				if (result && !isKnown)
 					result = CopyImports(out, importName, applicationDir, applicationDirEnd, handledImports);
-			});
+			}, errorStr))
+			m_logger.Error(errorStr.data);
 		#endif
 
 		// This code is needed if application is compiled with tsan
diff --git a/Engine/Source/Programs/UnrealBuildAccelerator/Core/Public/UbaMacBinDependencyParser.h b/Engine/Source/Programs/UnrealBuildAccelerator/Core/Public/UbaMacBinDependencyParser.h
--- a/Engine/Source/Programs/UnrealBuildAccelerator/Core/Public/UbaMacBinDependencyParser.h
+++ b/Engine/Source/Programs/UnrealBuildAccelerator/Core/Public/UbaMacBinDependencyParser.h
@@ -10,24 +10,34 @@
 
 namespace uba
 {
-	inline bool FindImportsMac(const tchar* fileName, const Function<void(const tchar* import, bool isKnown)>& func)
+	inline bool FindImportsMac(const tchar* fileName, const Function<void(const tchar* import, bool isKnown)>& func, StringBufferBase& outError)
 	{
-#if 0
-		HANDLE fileHandle = ::CreateFileW(fileName, GENERIC_READ, FILE_SHARE_READ, 0, OPEN_EXISTING, FILE_ATTRIBUTE_NORMAL | FILE_FLAG_SEQUENTIAL_SCAN, 0);
-		if (fileHandle == INVALID_HANDLE_VALUE)
-			return true;
-		auto closeFileHandle = MakeGuard([&]() { CloseHandle(fileHandle); });
-		HANDLE fileMapping = ::CreateFileMappingW(fileHandle, NULL, PAGE_READONLY, 0, 0, NULL);
-		if (!fileMapping)
+		int fd = open(fileName, O_RDONLY);
+		if (fd == -1)
+		{
+			outError.Appendf("Open failed for file %s", fileName);
 			return false;
-		auto closeMappingHandle = MakeGuard([&]() { ::CloseHandle(fileMapping); });
-		void* mem = ::MapViewOfFile(fileMapping, FILE_MAP_READ, 0, 0, 0);
-		if (!mem)
+		}
+		auto closeFileHandle = MakeGuard([&]() { close(fd); });
+		struct stat sb;
+		if (fstat(fd, &sb) == -1)
+		{
+			outError.Appendf("Stat failed for file %s", fileName);
+			return false;
+		}
+		u32 size = Min(u32(sb.st_size), 8048u);
+
+		void* mem = mmap(nullptr, size, PROT_READ, MAP_PRIVATE, fd, 0);
+		if (mem == MAP_FAILED)
+		{
+			outError.Appendf("Mmap failed for file %s", fileName);
 			return false;
-		auto unmap = MakeGuard([&]() { ::UnmapViewOfFile(mem); });
+		}
+		auto unmap = MakeGuard([&]() { munmap(mem, size); });
 
 		auto it = (const u8*)mem;
-		while (true)
+		auto end = it + size - 14;
+		while (it != end)
 		{
 			if (*it++ != '@')
 				continue;
@@ -40,17 +50,14 @@ namespace uba
 			it += 6;
 			auto importFile = (const char*)it;
 			if (!strstr(importFile, ".dylib"))
+			if (mem == MAP_FAILED)
+			{
+				outError.Appendf("Found @rpath in binary %s that did not end with .dylib (%s)", fileName, importFile);
 				return false;
+			}
 
-#if PLATFORM_WINDOWS
-			StringBuffer<> temp;
-			temp.Append(importFile);
-			func(temp.data, false);
-#else
 			func(importFile, false);
-#endif
 		}
-#endif
 		return true;
 	}
 }
\ No newline at end of file
