[UBA]
* Added logging to track processes (disabled by default)

diff --git a/Engine/Source/Programs/UnrealBuildAccelerator/Common/Private/UbaLogger.cpp b/Engine/Source/Programs/UnrealBuildAccelerator/Common/Private/UbaLogger.cpp
--- a/Engine/Source/Programs/UnrealBuildAccelerator/Common/Private/UbaLogger.cpp
+++ b/Engine/Source/Programs/UnrealBuildAccelerator/Common/Private/UbaLogger.cpp
@@ -372,12 +372,19 @@ namespace uba
 		void LogNoLock(LogEntryType type, const tchar* str, u32 strLen, const tchar* prefix, u32 prefixLen)
 		{
 			#if PLATFORM_WINDOWS
-			u8 buffer[2048];
-			BinaryWriter writer(buffer);
-			writer.WriteString(str, strLen);
-			BinaryReader reader(buffer);
-			u64 charLen = reader.Read7BitEncoded();
-			m_file->Write(reader.GetPositionData(), charLen);
+			char buffer[512];
+			size_t destLen;
+			u32 maxLen = Min(strLen, 256u);
+			if (wcstombs_s(&destLen, buffer, sizeof(buffer), str, maxLen) != 0)
+			{
+				strcpy_s(buffer, sizeof(buffer), "BAD_STRING\n");
+				destLen = 11;
+			}
+			else
+			{
+				buffer[destLen] = '\n';
+			}
+			m_file->Write(buffer, destLen);
 			#else
 			m_file->Write(str, strLen);
 			#endif
diff --git a/Engine/Source/Programs/UnrealBuildAccelerator/Common/Private/UbaProcess.cpp b/Engine/Source/Programs/UnrealBuildAccelerator/Common/Private/UbaProcess.cpp
--- a/Engine/Source/Programs/UnrealBuildAccelerator/Common/Private/UbaProcess.cpp
+++ b/Engine/Source/Programs/UnrealBuildAccelerator/Common/Private/UbaProcess.cpp
@@ -33,6 +33,9 @@
 extern char **environ;
 #endif
 
+#define UBA_DEBUG_TRACK_PROCESS UBA_DEBUG_LOGGER
+
+
 //////////////////////////////////////////////////////////////////////////////
 
 #define UBA_EXIT_CODE(x) (9000 + x)
@@ -525,6 +528,10 @@ namespace uba
 
 		KernelStats::GetGlobal().Add(m_kernelStats);
 
+		#if UBA_DEBUG_TRACK_PROCESS
+		g_debugLogger.Info(TC("ProcessExitedStart (%u)\n"), m_id);
+		#endif
+
 		// For some reason a parent can exit before a child. Need to figure out repro for this but I've seen it happen on ClangEditor win64
 		for (auto& child : m_childProcesses)
 			while (!((ProcessImpl*)child.m_process)->m_hasExited)
@@ -567,6 +574,10 @@ namespace uba
 
 		// Must be done last to make sure shutdown is not racing
 		m_session.ProcessExited(*this, m_processStats.wallTime);
+
+		#if UBA_DEBUG_TRACK_PROCESS
+		g_debugLogger.Info(TC("ProcessExitedDone  (%u)\n"), m_id);
+		#endif
 	}
 
 	bool ProcessImpl::HandleSpecialApplication()
@@ -961,6 +972,9 @@ namespace uba
 		writer.WriteString(info.logFile);
 		#endif
 
+		#if UBA_DEBUG_TRACK_PROCESS
+		g_debugLogger.Info(TC("CreateChildProcess (%u creating child %u at index %u) %s %s\n"), m_id, process.m_id, childProcessId-1, process.m_realApplication.c_str(), commandLine);
+		#endif
 		return true;
 	}
 
@@ -1002,9 +1016,17 @@ namespace uba
 		process.m_nativeProcessHandle = (ProcHandle)nativeProcessHandle;
 		process.m_nativeProcessId = nativeProcessId;
 #endif
+		
+		#if UBA_DEBUG_TRACK_PROCESS
+		g_debugLogger.Info(TC("WaitForChildProcessReady (%u waiting for child %u at index %u)\n"), m_id, process.m_id, processId-1);
+		#endif
 
 		setWaitForParent.Execute();
-		
+
+		#if UBA_DEBUG_TRACK_PROCESS
+		g_debugLogger.Info(TC("WaitForChildProcessReadyDone (%u waiting for child %u at index %u)\n"), m_id, process.m_id, processId-1);
+		#endif
+
 		#if PLATFORM_WINDOWS
 		// TOOD: This is ugly, should use event or something instead.. right now we make sure to wait and not return until we know payload has been uploaded etc
 		// It is a very uncommon usecase that processes starts suspended.. some process in ninja/cmake does it when building clang/llvm
@@ -1015,6 +1037,7 @@ namespace uba
 
 		return true;
 	}
+
 	bool ProcessImpl::HandleExitChildProcess(BinaryReader& reader, BinaryWriter& writer)
 	{
 		u32 nativeProcessId = reader.ReadU32();
@@ -1519,12 +1542,25 @@ namespace uba
 			m_nativeProcessHandle = (ProcHandle)(u64)processInfo.hProcess;
 			m_nativeProcessId = processInfo.dwProcessId;
 			m_nativeThreadHandle = processInfo.hThread;
+
+			#if UBA_DEBUG_TRACK_PROCESS
+			g_debugLogger.Info(TC("CreateRealProcess  (%u) %s %s\n"), m_id, m_realApplication.c_str(), m_startInfo.arguments);
+			#endif
 		}
 		else
 		{
+			#if UBA_DEBUG_TRACK_PROCESS
+			g_debugLogger.Info(TC("WaitingForParentReady (%u)\n"), m_id);
+			#endif
+
 			WaitForParent();
 			if (m_nativeProcessHandle == InvalidProcHandle) // Failed to create the child process
 				return UBA_EXIT_CODE(7);
+
+			#if UBA_DEBUG_TRACK_PROCESS
+			g_debugLogger.Info(TC("WaitingForParentReadyDone (%u)\n"), m_id);
+			#endif
+
 			m_extractExports = m_parentProcess->m_extractExports;
 		}
 
@@ -1819,9 +1855,17 @@ namespace uba
 			m_stdOutPipe = outPipe[0];
 			m_stdErrPipe = errPipe[0];
 			pipeGuard0.Cancel();
+
+			#if UBA_DEBUG_TRACK_PROCESS
+			g_debugLogger.Info(TC("CreateRealProcess  (%u) %s %s\n"), m_id, m_realApplication.c_str(), m_startInfo.arguments);
+			#endif
 		}
 		else
 		{
+			#if UBA_DEBUG_TRACK_PROCESS
+			g_debugLogger.Info(TC("WaitingForParent (%u) %s\n"), m_id, m_realApplication.c_str());
+			#endif
+
 			//logger.Info("Waiting for parent");
 			WaitForParent();
 
diff --git a/Engine/Source/Programs/UnrealBuildAccelerator/Common/Private/UbaSession.cpp b/Engine/Source/Programs/UnrealBuildAccelerator/Common/Private/UbaSession.cpp
--- a/Engine/Source/Programs/UnrealBuildAccelerator/Common/Private/UbaSession.cpp
+++ b/Engine/Source/Programs/UnrealBuildAccelerator/Common/Private/UbaSession.cpp
@@ -32,6 +32,8 @@
 extern char **environ;
 #endif
 
+#define UBA_DEBUG_TRACK_DIR 0 // UBA_DEBUG_LOGGER
+
 //////////////////////////////////////////////////////////////////////////////
 
 #if PLATFORM_WINDOWS
@@ -260,7 +262,7 @@ namespace uba
 			hasher.Update(forHash.data, forHash.count);
 		}
 
-		#if UBA_DEBUG_LOGGER
+		#if UBA_DEBUG_TRACK_DIR
 		g_debugLogger.BeginScope();
 		auto dg = MakeGuard([]() { g_debugLogger.EndScope(); });
 		g_debugLogger.Info(TC("TRACKDIR %s\n"), dirPath);
@@ -284,7 +286,7 @@ namespace uba
 				UBA_ASSERT(e.attributes);
 				memoryWriter.WriteString(e.name, e.nameLen);
 
-				#if UBA_DEBUG_LOGGER
+				#if UBA_DEBUG_TRACK_DIR
 				g_debugLogger.Info(TC("    %s (Size: %llu, Key: %s, Id: %llu)\n"), e.name, e.size, KeyToString(fileKey).data, e.id);
 				#endif
 
@@ -836,7 +838,7 @@ namespace uba
 			written = writer.GetPosition();
 		}
 
-		#if UBA_DEBUG_LOGGER
+		#if UBA_DEBUG_TRACK_DIR
 		g_debugLogger.Info(TC("TRACKADD    %s (Size: %llu, Key: %s, Id: %llu)\n"), fileName.data, fileSize, KeyToString(fileNameKey).data, fileIndex);
 		#endif
 
@@ -913,7 +915,7 @@ namespace uba
 			written = writer.GetPosition();
 		}
 
-		#if UBA_DEBUG_LOGGER
+		#if UBA_DEBUG_TRACK_DIR
 		g_debugLogger.Info(TC("TRACKDEL    %s (Key: %s)\n"), fileName.data, KeyToString(fileNameKey).data);
 		#endif
 
