[UBA]
* Changed CreateHardLinkW to use "link" instead of symlink on macos and linux. Reason is because when binary has explicit dynamic library dependencies they are often next to the binary and if it is a symlink it will look in the wrong place

diff --git a/Engine/Source/Programs/UnrealBuildAccelerator/Common/Private/UbaFile.cpp b/Engine/Source/Programs/UnrealBuildAccelerator/Common/Private/UbaFile.cpp
--- a/Engine/Source/Programs/UnrealBuildAccelerator/Common/Private/UbaFile.cpp
+++ b/Engine/Source/Programs/UnrealBuildAccelerator/Common/Private/UbaFile.cpp
@@ -715,10 +715,8 @@ namespace uba
 		MAKE_LONG_FILENAME(existingFileName);
 		return ::CreateHardLinkW(newFileName, existingFileName, NULL);
 #else
-#if 1//PLATFORM_MAC
-		int res = symlink(existingFileName, newFileName);
-#else
-		int res = link(existingFileName, newFileName);
+		int res = link(existingFileName, newFileName); // We need to use links in order for explicit dynamic library dependencies  to be found at the same path.
+		//int res = symlink(existingFileName, newFileName);
 #endif
 		if (res == 0)
 			return true;
