[UBA]
* Fixed so GetApplicationShortName handles forward slashes

diff --git a/Engine/Source/Programs/UnrealBuildAccelerator/Detours/Private/UbaDetoursShared.cpp b/Engine/Source/Programs/UnrealBuildAccelerator/Detours/Private/UbaDetoursShared.cpp
--- a/Engine/Source/Programs/UnrealBuildAccelerator/Detours/Private/UbaDetoursShared.cpp
+++ b/Engine/Source/Programs/UnrealBuildAccelerator/Detours/Private/UbaDetoursShared.cpp
@@ -159,8 +159,10 @@ namespace uba
 
 	const tchar* GetApplicationShortName()
 	{
-		if (const tchar* lastBackslash = TStrrchr(g_virtualApplication.data, '\\'))
-			return lastBackslash + 1;
+		const tchar* lastBackslash = TStrrchr(g_virtualApplication.data, '\\');
+		const tchar* lastSlash = TStrrchr(g_virtualApplication.data, '/');
+		if (lastBackslash || lastSlash)
+			return (lastBackslash > lastSlash ? lastBackslash : lastSlash) + 1;
 		return g_virtualApplication.data;
 	}
 
