[UBA]
* Added error handling when failed to spawn a child process
* Improved some logging when enabled

diff --git a/Engine/Source/Programs/UnrealBuildAccelerator/Common/Private/UbaProcess.cpp b/Engine/Source/Programs/UnrealBuildAccelerator/Common/Private/UbaProcess.cpp
--- a/Engine/Source/Programs/UnrealBuildAccelerator/Common/Private/UbaProcess.cpp
+++ b/Engine/Source/Programs/UnrealBuildAccelerator/Common/Private/UbaProcess.cpp
@@ -939,6 +939,12 @@ namespace uba
 
 		ProcessImpl* parent = isChild ? this : nullptr;
 		ProcessHandle h = m_session.InternalRunProcess(info, true, parent, true);
+		if (!h.m_process)
+		{
+			// TODO: Do we need to log something here? This failure is most likely a cause of something else so error logging might show somewhere else
+			writer.WriteU32(0); // childProcessId
+			return true;
+		}
 
 		u32 childProcessId = ~0u;
 		if (isChild)
@@ -973,7 +979,7 @@ namespace uba
 		#endif
 
 		#if UBA_DEBUG_TRACK_PROCESS
-		g_debugLogger.Info(TC("CreateChildProcess (%u creating child %u at index %u) %s %s\n"), m_id, process.m_id, childProcessId-1, process.m_realApplication.c_str(), commandLine);
+		g_debugLogger.Info(TC("CreateChildProcess (%u creating child %u at index %u) %s %s (%s)\n"), m_id, process.m_id, childProcessId-1, process.m_realApplication.c_str(), commandLine, info.logFile);
 		#endif
 		return true;
 	}
