[UBA]
* Fixed so stat is not writing anything to output struct if failing
* Stubbed out dlopen detour just to get it in the log (this will need implementation)
* Fixes for _NSGetExecutablePath - Implemented by [REDACTED_PERSON]
* Fixes for chdir  - Implemented by [REDACTED_PERSON]
* Improved logging in posix wait detours
* Changed so execve closes the connection. This might be wrong but is needed for metal compiler to work. Will need to revisit

diff --git a/Engine/Source/Programs/UnrealBuildAccelerator/Detours/Private/PosixOS/UbaDetoursFunctionsPosix.cpp b/Engine/Source/Programs/UnrealBuildAccelerator/Detours/Private/PosixOS/UbaDetoursFunctionsPosix.cpp
--- a/Engine/Source/Programs/UnrealBuildAccelerator/Detours/Private/PosixOS/UbaDetoursFunctionsPosix.cpp
+++ b/Engine/Source/Programs/UnrealBuildAccelerator/Detours/Private/PosixOS/UbaDetoursFunctionsPosix.cpp
@@ -90,6 +90,7 @@ using namespace uba;
 	DETOURED_FUNCTION(wait3) \
 	DETOURED_FUNCTION(wait4) \
 	DETOURED_FUNCTION(system) \
+	DETOURED_FUNCTION(dlopen) \
 	DETOURED_FUNCTION(execv) \
 	DETOURED_FUNCTION(execve) \
 	DETOURED_FUNCTION(execvp) \
@@ -199,6 +200,8 @@ using namespace uba;
 #define	ERANGE		34	/* Math result not representable */
 #endif
 
+void CloseCom();
+
 namespace uba
 {
 	template<typename TrueFunc>
@@ -570,7 +573,7 @@ int Shared_fstat(const char* funcName, int fd, struct stat* attr, const True_fst
 		return res;
 	}
 
-	FileAttributes fileAttr;
+	FileAttributes fileAttr = {};
 	const char* realName = Shared_GetFileAttributes(fileAttr, fi.originalName);
 
 	if (!fileAttr.useCache)
@@ -587,7 +590,8 @@ int Shared_fstat(const char* funcName, int fd, struct stat* attr, const True_fst
 	DEBUG_LOG_DETOURED(funcName, "(%i) (name: %s size: %llu id: %llu dev: %u)-> %i (%s)", fd, fi.originalName, fileAttr.data.st_size, fileAttr.data.st_ino, fileAttr.data.st_dev, res, StrError(res, fileAttr.lastError));
 
 	errno = fileAttr.lastError;
-	memcpy(attr, &fileAttr.data, sizeof(struct stat));
+	if (res == 0)
+		memcpy(attr, &fileAttr.data, sizeof(struct stat));
 
 	#if UBA_DEBUG_VALIDATE
 	if (!g_runningRemote)
@@ -621,7 +625,7 @@ int Shared_stat(const char* funcName, const char* file, struct stat* attr, const
 	StringBuffer<> fixedFile;
 	if (!FixPath(fixedFile, file) || fixedFile.Equals("/") || fixedFile.StartsWith("/etc/") || fixedFile.StartsWith(g_systemTemp.data))//(access(file, F_OK) != 0))
 	{
-		int res =  trueStat(file, attr);
+		int res = trueStat(file, attr);
 		return res;
 	}
 
@@ -649,7 +653,9 @@ int Shared_stat(const char* funcName, const char* file, struct stat* attr, const
 	DEBUG_LOG_DETOURED(funcName, "%s (%s size: %llu id: %llu dev: %u)-> %i (%s)", file, realName, fileAttr.data.st_size, fileAttr.data.st_ino, fileAttr.data.st_dev, res, StrError(res, fileAttr.lastError));
 
 	errno = fileAttr.lastError;
-	memcpy(attr, &fileAttr.data, sizeof(fileAttr.data));
+
+	if (res == 0) // This check is important.. dont write to attr unless success. (some programs rely on that)
+		memcpy(attr, &fileAttr.data, sizeof(fileAttr.data));
 
 	#if UBA_DEBUG_VALIDATE
 	if (!g_runningRemote)
@@ -688,10 +694,16 @@ int Shared_stat(const char* funcName, const char* file, struct stat* attr, const
 UBA_EXPORT int UBA_WRAPPER(_NSGetExecutablePath)(char* buf, uint32_t* bufsize)
 {
 	if (!g_isDetouring)
-	{
 		return TRUE_WRAPPER(_NSGetExecutablePath)(buf, bufsize);
-	}
-	memcpy(buf, g_virtualApplication.data, g_virtualApplication.count + 1);
+	if (bufsize == nullptr)
+		return -1;
+	const uint32_t requiredBufsize = g_virtualApplication.count + 1;
+	const uint32_t initialBufsize = *bufsize;
+	*bufsize = requiredBufsize;
+	if (initialBufsize < requiredBufsize)
+		return -1;
+	if (buf != nullptr)
+		memcpy(buf, g_virtualApplication.data, requiredBufsize);
 	return 0;
 }
 
@@ -715,16 +727,21 @@ UBA_EXPORT int UBA_WRAPPER(_NSGetExecutablePath)(char* buf, uint32_t* bufsize)
 UBA_EXPORT int UBA_WRAPPER(chdir)(const char* path)
 {
 	UBA_INIT_DETOUR(chdir, path);
-	FixPath(g_virtualWorkingDir.Clear(), path);
-	g_virtualWorkingDir.EnsureEndsWithSlash();
-	if (g_runningRemote)
+	if (path == nullptr || *path == '\0')
 	{
-		DEBUG_LOG_DETOURED("chdir", "%s -> 0", path);
-		return 0;
+		errno = ENOENT;
+		return -1;
 	}
-	int res = TRUE_WRAPPER(chdir)(path);
-	DEBUG_LOG_TRUE("chdir", "%s -> %i", path, res);
-	return res;
+	size_t pathlen = strlen(path);
+	if (pathlen >= g_virtualWorkingDir.capacity)
+	{
+		errno = ENAMETOOLONG;
+		return -1;
+	}
+	memcpy(g_virtualWorkingDir.data, path, pathlen + 1);
+	setenv("PWD", g_virtualWorkingDir.data, 1);
+	g_virtualWorkingDir.EnsureEndsWithSlash();
+	return 0;
 }
 
 UBA_EXPORT int UBA_WRAPPER(fchdir)(int fd)
@@ -833,7 +850,7 @@ UBA_EXPORT char* UBA_WRAPPER(getenv)(const char* name)
 {
 	UBA_INIT_DETOUR(getenv, name);
 	auto res = TRUE_WRAPPER(getenv)(name);
-	DEBUG_LOG_TRUE("getenv", "(%s) -> %s", name, res ? res : "<null>");
+	DEBUG_LOG_TRUE("getenv", "(%s) -> %s", name ? name : "<null>", res ? res : "<null>");
 	return res;
 }
 
@@ -1608,27 +1625,33 @@ UBA_EXPORT int UBA_WRAPPER(remove)(const char* pathname)
 
 thread_local int t_inVfork;
 
-UBA_EXPORT int UBA_WRAPPER(posix_spawn)(pid_t* pid, const char* path, const posix_spawn_file_actions_t* file_actions, const posix_spawnattr_t* attrp, char* const argv[], char* const envp[])
+int shared_posix_spawn(pid_t* pid, const char* path, const posix_spawn_file_actions_t* file_actions, const posix_spawnattr_t* attrp, char* const argv[], char* const envp[])
 {
 	UBA_INIT_DETOUR(posix_spawn, pid, path, file_actions, attrp, argv, envp);
 
 	t_inVfork = 0;
 
 	TString cmdLine;
+	u32 argc = 0;
 	for (u32 i = 0; argv[i]; ++i)
 	{
 		if (i != 0)
 			cmdLine.append(" ");
 		cmdLine.append(argv[i]);
+		++argc;
 	}
-	
+
+	Vector<const char*> argv2;
+	argv2.resize(argc + 1);
+	memcpy(argv2.data(), argv, sizeof(char*) * (argc+1));
+
 	TString realApplication;
 	TString commandLine;
 	u32 processId = 0;
 	StringBuffer<512> currentDir;
 	StringBuffer<256> comIdVar;
 	StringBuffer<32> rulesStr;
-	StringBuffer<256> logFile;
+	StringBuffer<512> logFile;
 
 	{
 		TimerScope ts(g_stats.createProcess);
@@ -1656,8 +1679,13 @@ UBA_EXPORT int UBA_WRAPPER(posix_spawn)(pid_t* pid, const char* path, const posi
 
 		BinaryReader reader;
 		processId = reader.ReadU32();
-		UBA_ASSERTF(processId > 0, "Process id was zero");
-
+		UBA_ASSERTF(processId > 0, "Process id was zero: path=%s", path);
+		if (!processId)
+		{
+			errno = EINVAL; // This is not really correct but there is no errno for this failure
+			return -1;
+		}
+		
 		rulesStr.Append("UBA_RULES=").AppendValue(reader.ReadU32());
 
 		u32 dllNameSize = reader.ReadU32();
@@ -1673,6 +1701,9 @@ UBA_EXPORT int UBA_WRAPPER(posix_spawn)(pid_t* pid, const char* path, const posi
 
 		logFile.Append("UBA_LOGFILE=");
 		reader.ReadString(logFile);
+
+		// TODO: Recalculate argv2 right now it is wrong, it should be what is in command line
+		//argv2[0] = realApplication.c_str();
 	}
 	
 	std::vector<const char*> envvars;
@@ -1705,7 +1736,7 @@ UBA_EXPORT int UBA_WRAPPER(posix_spawn)(pid_t* pid, const char* path, const posi
 		DEBUG_LOG("            %s", argv[i]);
 	#endif
 
-	int res = TRUE_WRAPPER(posix_spawn)(pid, realApplication.data(), file_actions, attrp, argv, (char**)envvars.data());
+	int res = TRUE_WRAPPER(posix_spawn)(pid, realApplication.data(), file_actions, attrp, (char* const*)argv2.data(), (char**)envvars.data());
 	bool success = res == 0;
 
 	{
@@ -1727,6 +1758,11 @@ UBA_EXPORT int UBA_WRAPPER(posix_spawn)(pid_t* pid, const char* path, const posi
 	return res;
 }
 
+UBA_EXPORT int UBA_WRAPPER(posix_spawn)(pid_t* pid, const char* path, const posix_spawn_file_actions_t* file_actions, const posix_spawnattr_t* attrp, char* const argv[], char* const envp[])
+{
+	return shared_posix_spawn(pid, path, file_actions, attrp, argv, envp);
+}
+
 UBA_EXPORT int UBA_WRAPPER(posix_spawnp)(pid_t* pid, const char* file, const posix_spawn_file_actions_t* file_actions, const posix_spawnattr_t* attrp, char* const argv[], char* const envp[])
 {
 	UBA_INIT_DETOUR(posix_spawnp, pid, file, file_actions, attrp, argv, envp);
@@ -1737,8 +1773,9 @@ UBA_EXPORT int UBA_WRAPPER(posix_spawnp)(pid_t* pid, const char* file, const pos
 UBA_EXPORT pid_t UBA_WRAPPER(wait)(int* status)
 {
 	UBA_INIT_DETOUR(wait, status);
-	DEBUG_LOG_TRUE("wait", "");
-	return TRUE_WRAPPER(wait)(status);
+	pid_t res = TRUE_WRAPPER(wait)(status);
+	DEBUG_LOG_TRUE("wait", "%i -> %i", status ? *status : 0, res);
+	return res;
 }
 
 UBA_EXPORT pid_t UBA_WRAPPER(waitpid)(pid_t pid, int* status, int options)
@@ -1746,17 +1783,33 @@ UBA_EXPORT pid_t UBA_WRAPPER(waitpid)(pid_t pid, int* status, int options)
 	UBA_INIT_DETOUR(waitpid, pid, status, options);
 	// TODO: Should probably report id to session
 	auto res = TRUE_WRAPPER(waitpid)(pid, status, options);
-	DEBUG_LOG_TRUE("waitpid", "(%i) -> ", pid, res);
+	DEBUG_LOG_TRUE("waitpid", "(%i) -> %i", pid, res);
 	return res;
 }
 
+const char* GetResult(siginfo_t* info)
+{
+	if (!info)
+		return "null";
+	if (WIFEXITED(info))
+		return "Exited";
+	if (WIFSIGNALED(info))
+		return "Signaled";
+	if (WIFSTOPPED(info))
+		return "Stopped";
+	if (WIFCONTINUED(info))
+		return "Continued";
+	return "Running";
+}
+
 UBA_EXPORT int UBA_WRAPPER(waitid)(idtype_t idtype, id_t id, siginfo_t* infop, int options)
 {
 	UBA_INIT_DETOUR(waitid, idtype, id, infop, options);
-	DEBUG_LOG_TRUE("waitid", "");
 	UBA_ASSERTF(!t_inVfork, "waitid: is in fork");
 	// TODO: Should probably report id to session
-	return TRUE_WRAPPER(waitid)(idtype, id, infop, options);
+	auto res = TRUE_WRAPPER(waitid)(idtype, id, infop, options);
+	DEBUG_LOG_TRUE("waitid", "%i -> %i (%s)", id, res, GetResult(infop));
+	return res;
 }
 
 UBA_EXPORT pid_t UBA_WRAPPER(wait3)(int* status, int options, struct rusage* rusage)
@@ -1778,10 +1831,19 @@ UBA_EXPORT pid_t UBA_WRAPPER(wait4)(pid_t pid, int* status, int options, struct
 	return res;
 }
 
+UBA_EXPORT void* UBA_WRAPPER(dlopen)(const char* path, int mode)
+{
+	UBA_INIT_DETOUR(dlopen, path, mode);
+	DEBUG_LOG_TRUE("dlopen", "%s", path);
+	void* res = TRUE_WRAPPER(dlopen)(path, mode);
+	DEBUG_LOG_TRUE("dlopen DONE", "%s", path);
+	return res;
+}
+
 UBA_EXPORT int UBA_WRAPPER(execv)(const char* path, char* const argv[])
 {
 	UBA_INIT_DETOUR(execv, path, argv);
-	DEBUG_LOG_TRUE("execv", "");
+	DEBUG_LOG_TRUE("execv", "%s", path);
 	return TRUE_WRAPPER(execv)(path, argv);
 }
 
@@ -1789,11 +1851,11 @@ int Internal_execve(const char* pathname, char* const _Nullable argv[], char* co
 {
 	UBA_INIT_DETOUR(execve, pathname, argv, envp);
 	// We are most likely in a vfork/fork here, which means that we won't see the exit call since _exit will drop back out to the entrance of fork.
-	DEBUG_LOG_TRUE("execve", "");
+	DEBUG_LOG_TRUE("execve", "%s", pathname);
 
 	bool inVfork = t_inVfork != 0;
 	pid_t pid;
-	int res = posix_spawn(&pid, pathname, nullptr, nullptr, argv, envp);
+	int res = shared_posix_spawn(&pid, pathname, nullptr, nullptr, argv, envp);
 	if (inVfork)
 		t_inVfork = pid;
 
@@ -1825,6 +1887,8 @@ int Internal_execve(const char* pathname, char* const _Nullable argv[], char* co
 UBA_EXPORT int UBA_WRAPPER(execve)(const char* pathname, char* const _Nullable argv[], char* const _Nullable envp[])
 {
 	int res = Internal_execve(pathname, argv, envp);
+	Deinit();
+	CloseCom();
 	TRUE_WRAPPER(_exit)(res);
 	return 0;
 }
@@ -1870,8 +1934,6 @@ UBA_EXPORT FILE* UBA_WRAPPER(popen)(const char* command, const char* type)
 	return TRUE_WRAPPER(popen)(command, type);
 }
 
-void CloseCom();
-
 UBA_EXPORT void UBA_WRAPPER(exit)(int status)
 {
 	//UBA_INIT_DETOUR(exit, status);
