[UBA]
* Added disabled code for testing code that get macos binary dylib dependency list

diff --git a/Engine/Source/Programs/UnrealBuildAccelerator/Test/Private/UbaTestBasics.h b/Engine/Source/Programs/UnrealBuildAccelerator/Test/Private/UbaTestBasics.h
--- a/Engine/Source/Programs/UnrealBuildAccelerator/Test/Private/UbaTestBasics.h
+++ b/Engine/Source/Programs/UnrealBuildAccelerator/Test/Private/UbaTestBasics.h
@@ -460,6 +460,22 @@ namespace uba
 		return true;
 	}
 
+
+#if PLATFORM_MAC
+	Set<TString> g_visited;
+	void LogImports(const tchar* import, bool isKnown)
+	{
+		if (!g_visited.insert(import).second)
+			return;
+		LoggerWithWriter logger(g_consoleLogWriter);
+		logger.Info(TC("IMPORT: %s"), import);
+		StringBuffer<> path(TC("[REDACTED_USER_HOME]/p4/fn/Engine/Binaries/Mac"));
+		path.EnsureEndsWithSlash().Append(TC(import));
+		StringBuffer<> error;
+		FindImportsMac(path.data, LogImports, error);
+	}
+#endif
+
 	bool TestBinDependencies(Logger& logger, const StringBufferBase& rootDir)
 	{
 
@@ -475,10 +491,8 @@ namespace uba
 		if (!importKernel)
 			return logger.Error(TC("Failed to find Kernel32 as import"));
 #elif PLATFORM_MAC
-		//FindImportsMac(TC(""), [&](const tchar* import, bool isKnown)
-		//{
-		//	LoggerWithWriter(g_consoleLogWriter).Info(TC("IMPORT: %s"), import);
-		//});
+		//StringBuffer<> error;
+		//FindImportsMac(TC("[REDACTED_USER_HOME]/p4/fn/Engine/Binaries/Mac/ShaderCompileWorker"), LogImports, error);
 #endif
 		return true;
 	}
