[UBA]
* Removed #endif that shouldn't be there

diff --git a/Engine/Source/Programs/UnrealBuildAccelerator/Common/Private/UbaFile.cpp b/Engine/Source/Programs/UnrealBuildAccelerator/Common/Private/UbaFile.cpp
--- a/Engine/Source/Programs/UnrealBuildAccelerator/Common/Private/UbaFile.cpp
+++ b/Engine/Source/Programs/UnrealBuildAccelerator/Common/Private/UbaFile.cpp
@@ -717,7 +717,6 @@ namespace uba
 #else
 		int res = link(existingFileName, newFileName); // We need to use links in order for explicit dynamic library dependencies  to be found at the same path.
 		//int res = symlink(existingFileName, newFileName);
-#endif
 		if (res == 0)
 			return true;
 
