[UBA]
* Added wip code for figuring out search paths for explicitly loaded dylibs

diff --git a/Engine/Source/Programs/UnrealBuildAccelerator/Common/Private/UbaSession.cpp b/Engine/Source/Programs/UnrealBuildAccelerator/Common/Private/UbaSession.cpp
--- a/Engine/Source/Programs/UnrealBuildAccelerator/Common/Private/UbaSession.cpp
+++ b/Engine/Source/Programs/UnrealBuildAccelerator/Common/Private/UbaSession.cpp
@@ -19,10 +19,7 @@
 #include "UbaWinBinDependencyParser.h"
 #include <powerbase.h>
 #pragma comment(lib, "Powrprof.lib")
-#endif
-
-// Used to get Memory Information
-#if PLATFORM_MAC
+#elif PLATFORM_MAC
 #include "UbaMacBinDependencyParser.h"
 #include <mach/vm_statistics.h>
 #include <mach/mach_types.h>
@@ -30,6 +27,8 @@
 #include <mach/mach_host.h>
 #include <mach/mach.h>
 extern char **environ;
+#else // PLATFORM_LINUX
+#include "UbaLinuxBinDependencyParser.h"
 #endif
 
 #define UBA_DEBUG_TRACK_DIR 0 // UBA_DEBUG_LOGGER
@@ -931,7 +930,7 @@ namespace uba
 		return dirTable.m_memorySize;
 	}
 
-	bool Session::CopyImports(Vector<BinaryModule>& out, const tchar* library, tchar* applicationDir, tchar* applicationDirEnd, UnorderedSet<TString>& handledImports)
+	bool Session::CopyImports(Vector<BinaryModule>& out, const tchar* library, tchar* applicationDir, tchar* applicationDirEnd, UnorderedSet<TString>& handledImports, const char* const* loaderPaths)
 	{
 		if (!handledImports.insert(library).second)
 			return true;
@@ -942,16 +941,35 @@ namespace uba
 		tchar temp2[512];
 		bool result = true;
 
-		#if PLATFORM_WINDOWS
 		if (attr == INVALID_FILE_ATTRIBUTES)
 		{
+			#if PLATFORM_WINDOWS
 			if (!SearchPathW(NULL, library, NULL, 512, temp, NULL))
 				return true; // TODO: We have to return true here because there are scenarios where failing is actually ok (it seems it can return false on crt shim libraries such as api-ms-win-crt*)
+			#elif PLATFORM_MAC
+			if (!loaderPaths)
+				return m_logger.Error("Failed to find file %s", applicationName);
+			auto it = loaderPaths;
+			for (auto it = loaderPaths; *it; ++it)
+			{
+				StringBuffer<> absolutePath;
+				absolutePath.Append(applicationDir, applicationDirEnd - applicationDir).Append(*it).EnsureEndsWithSlash().Append(library);
+				m_logger.Info("SEARCH_PATH: %s", absolutePath.data);
+				attr = GetFileAttributesW(applicationName);
+				if (attr == INVALID_FILE_ATTRIBUTES)
+					continue;
+				memcpy(temp, absolutePath.data, absolutePath.count + 1);
+				break;
+			}
+			if (attr == INVALID_FILE_ATTRIBUTES)
+				return m_logger.Error("Failed to find file %s", applicationName);
+			//UBA_ASSERTF(false, TC("DIR NOT FOUND: %s"), applicationName);
+			#endif
 
 			applicationName = temp;
 			attr = DefaultAttributes();
 
-			tchar* lastSlash = TStrrchr(temp, '\\');
+			tchar* lastSlash = TStrrchr(temp, PathSeparator);
 			UBA_ASSERT(lastSlash);
 			u64 applicationDirLen = u64(lastSlash + 1 - temp);
 			memcpy(temp2, temp, applicationDirLen * sizeof(tchar));
@@ -968,39 +986,18 @@ namespace uba
 
 		out.push_back({ library, temp3.data, attr, isSystem });
 
-		FindImports(applicationName, [&](const tchar* importName, bool isKnown)
-			{
-				if (result && !isKnown)
-					result = CopyImports(out, importName, applicationDir, applicationDirEnd, handledImports);
-			});
-
-		#else
-
-		if (attr == INVALID_FILE_ATTRIBUTES)
-		{
-			UBA_ASSERTF(false, TC("DIR NOT FOUND: %s"), applicationName);
-			return false;
-		}
-
-		StringBuffer<512> temp3;
-		FixPath(applicationName, nullptr, 0, temp3);
-
-		out.push_back({ TString(library), TString(temp3.data), S_IRUSR | S_IWUSR | S_IXUSR });
-
-		#if PLATFORM_MAC
 		StringBuffer<> errorStr;
-		if (!FindImportsMac(applicationName, [&](const tchar* importName, bool isKnown)
+		FindImports(applicationName, [&](const tchar* importName, bool isKnown, const char* const* importLoaderPaths)
 			{
 				if (result && !isKnown)
-					result = CopyImports(out, importName, applicationDir, applicationDirEnd, handledImports);
-			}, errorStr))
-			m_logger.Error(errorStr.data);
-		#endif
+					result = CopyImports(out, importName, applicationDir, applicationDirEnd, handledImports, importLoaderPaths);
+			}, errorStr);
+		if (errorStr.count)
+			return m_logger.Error(errorStr.data);
 
 		// This code is needed if application is compiled with tsan
 		//strcpy(applicationDirEnd, "libclang_rt.tsan.so");
 		//out.push_back({ "libclang_rt.tsan.so", applicationDir, S_IRUSR | S_IWUSR });
-		#endif
 		return result;
 	}
 
@@ -1658,7 +1655,7 @@ namespace uba
 		tchar* applicationDirEnd = applicationDir + applicationDirLen;
 
 		UnorderedSet<TString> handledImports;
-		return CopyImports(out, applicationName, applicationDir, applicationDirEnd, handledImports);
+		return CopyImports(out, applicationName, applicationDir, applicationDirEnd, handledImports, nullptr);
 	}
 
 	void Session::Free(Vector<BinaryModule>& v)
diff --git a/Engine/Source/Programs/UnrealBuildAccelerator/Common/Public/UbaSession.h b/Engine/Source/Programs/UnrealBuildAccelerator/Common/Public/UbaSession.h
--- a/Engine/Source/Programs/UnrealBuildAccelerator/Common/Public/UbaSession.h
+++ b/Engine/Source/Programs/UnrealBuildAccelerator/Common/Public/UbaSession.h
@@ -150,7 +150,7 @@ namespace uba
 		void AddEnvironmentVariableNoLock(const tchar* key, const tchar* value);
 		bool WriteDirectoryEntriesInternal(DirectoryTable::Directory& dir, const StringKey& dirKey, const tchar* dirPath, bool isRefresh, u32& outTableOffset);
 		void WriteDirectoryEntriesRecursive(const StringKey& dirKey, tchar* dirPath, u32& outTableOffset);
-		bool CopyImports(Vector<BinaryModule>& out, const tchar* library, tchar* applicationDir, tchar* applicationDirEnd, UnorderedSet<TString>& handledImports);
+		bool CopyImports(Vector<BinaryModule>& out, const tchar* library, tchar* applicationDir, tchar* applicationDirEnd, UnorderedSet<TString>& handledImports, const char* const* loaderPaths);
 		bool CreateProcessJobObject();
 		void GetSystemInfo(StringBufferBase& out);
 		bool GetMemoryInfo(u64& outAvailable, u64& outTotal);
diff --git a/Engine/Source/Programs/UnrealBuildAccelerator/Core/Public/UbaLinuxBinDependencyParser.h b/Engine/Source/Programs/UnrealBuildAccelerator/Core/Public/UbaLinuxBinDependencyParser.h
new file mode 100644
--- /dev/null
+++ b/Engine/Source/Programs/UnrealBuildAccelerator/Core/Public/UbaLinuxBinDependencyParser.h
@@ -0,0 +1,22 @@
+// Copyright Epic Games, Inc. All Rights Reserved.
+
+#pragma once
+
+#include "UbaStringBuffer.h"
+
+#if !defined(UBA_IS_DETOURED_INCLUDE)
+#define TRUE_WRAPPER(func) func
+#endif
+
+namespace uba
+{
+	inline bool IsKnownSystemFile(const tchar* fileName)
+	{
+		return false;
+	}
+
+	inline bool FindImports(const tchar* fileName, const Function<void(const tchar* import, bool isKnown, const char* const* loaderPaths)>& func, StringBufferBase& outError)
+	{
+		return true;
+	}
+}
\ No newline at end of file
diff --git a/Engine/Source/Programs/UnrealBuildAccelerator/Core/Public/UbaMacBinDependencyParser.h b/Engine/Source/Programs/UnrealBuildAccelerator/Core/Public/UbaMacBinDependencyParser.h
--- a/Engine/Source/Programs/UnrealBuildAccelerator/Core/Public/UbaMacBinDependencyParser.h
+++ b/Engine/Source/Programs/UnrealBuildAccelerator/Core/Public/UbaMacBinDependencyParser.h
@@ -10,7 +10,12 @@
 
 namespace uba
 {
-	inline bool FindImportsMac(const tchar* fileName, const Function<void(const tchar* import, bool isKnown)>& func, StringBufferBase& outError)
+	inline bool IsKnownSystemFile(const tchar* fileName)
+	{
+		return false;
+	}
+
+	inline bool FindImports(const tchar* fileName, const Function<void(const tchar* import, bool isKnown, const char* const* loaderPaths)>& func, StringBufferBase& outError)
 	{
 		int fd = open(fileName, O_RDONLY);
 		if (fd == -1)
@@ -35,6 +40,13 @@ namespace uba
 		}
 		auto unmap = MakeGuard([&]() { munmap(mem, size); });
 
+
+		const char* libs[1024];
+		u32 libsCount = 0;
+
+		const char* loaderPaths[256];
+		u32 loaderPathsCount = 0;
+
 		auto it = (const u8*)mem;
 		auto end = it + size - 14;
 		while (it != end)
@@ -42,22 +54,35 @@ namespace uba
 			if (*it++ != '@')
 				continue;
 
-			if (memcmp(it, "executable_path", 14) == 0)
-				break;
+			if (memcmp(it, "rpath/", 6) == 0)
+			{
+				it += 6;
+				auto importFile = (const char*)it;
+				if (!strstr(importFile, ".dylib"))
+				{
+					outError.Appendf("Found @rpath in binary %s that did not end with .dylib (%s)", fileName, importFile);
+					return false;
+				}
 
-			if (memcmp(it, "rpath/", 6) != 0)
-				continue;
-			it += 6;
-			auto importFile = (const char*)it;
-			if (!strstr(importFile, ".dylib"))
-			if (mem == MAP_FAILED)
+				libs[libsCount++] = importFile;
+			}
+			else if (memcmp(it, "executable_path/", 16) == 0)
 			{
-				outError.Appendf("Found @rpath in binary %s that did not end with .dylib (%s)", fileName, importFile);
-				return false;
+				it += 16;
+				auto executablePath = (const char*)it;
+				UBA_ASSERT(*executablePath == 0);
+			}
+			else if (memcmp(it, "loader_path/", 12) == 0)
+			{
+				it += 12;
+				loaderPaths[loaderPathsCount++] = (const char*)it;
 			}
-
-			func(importFile, false);
 		}
+
+		loaderPaths[loaderPathsCount] = nullptr;
+
+		for (u32 i=0; i!=libsCount; ++i)
+			func(libs[i], false, loaderPaths);
 		return true;
 	}
 }
\ No newline at end of file
diff --git a/Engine/Source/Programs/UnrealBuildAccelerator/Core/Public/UbaWinBinDependencyParser.h b/Engine/Source/Programs/UnrealBuildAccelerator/Core/Public/UbaWinBinDependencyParser.h
--- a/Engine/Source/Programs/UnrealBuildAccelerator/Core/Public/UbaWinBinDependencyParser.h
+++ b/Engine/Source/Programs/UnrealBuildAccelerator/Core/Public/UbaWinBinDependencyParser.h
@@ -97,7 +97,7 @@ namespace uba
 		return rva - seh->VirtualAddress + seh->PointerToRawData;
 	} 
 
-	inline bool FindImportsInMem(const tchar* fileName, void* mem, const Function<void(const tchar* import, bool isKnown)>& func)
+	inline bool FindImportsInMem(const tchar* fileName, void* mem, const Function<void(const tchar* import, bool isKnown, const char* const* importLoaderPaths)>& func)
 	{
 		auto hdrs = (PIMAGE_NT_HEADERS)(PCHAR(mem) + PIMAGE_DOS_HEADER(mem)->e_lfanew);   
 		auto pSech=IMAGE_FIRST_SECTION(hdrs);
@@ -113,7 +113,7 @@ namespace uba
 				const char* name = (PCHAR)((DWORD_PTR)mem + GetOffset(importDesc->Name, pSech, hdrs));
 				StringBuffer<256> wname;
 				wname.Append(name);
-				func(wname.data, IsKnownSystemFile(wname.data));
+				func(wname.data, IsKnownSystemFile(wname.data), nullptr);
 				++importDesc;
 			}
 		}
@@ -127,7 +127,7 @@ namespace uba
 		return true;
 	}
 
-	inline bool FindImports(const tchar* fileName, const Function<void(const tchar* import, bool isKnown)>& func)
+	inline bool FindImports(const tchar* fileName, const Function<void(const tchar* import, bool isKnown, const char* const* importLoaderPaths)>& func, StringBufferBase& outError)
 	{
 		HANDLE fileHandle = True_CreateFileW(fileName, GENERIC_READ, FILE_SHARE_READ, 0, OPEN_EXISTING, FILE_ATTRIBUTE_NORMAL | FILE_FLAG_SEQUENTIAL_SCAN, 0);
 		if (fileHandle == INVALID_HANDLE_VALUE)
diff --git a/Engine/Source/Programs/UnrealBuildAccelerator/Host/Private/UbaExports.cpp b/Engine/Source/Programs/UnrealBuildAccelerator/Host/Private/UbaExports.cpp
--- a/Engine/Source/Programs/UnrealBuildAccelerator/Host/Private/UbaExports.cpp
+++ b/Engine/Source/Programs/UnrealBuildAccelerator/Host/Private/UbaExports.cpp
@@ -682,7 +682,8 @@ extern "C"
 	void Uba_FindImports(const uba::tchar* binary, ImportFunc* func, void* userData)
 	{
 #if PLATFORM_WINDOWS
-		uba::FindImports(binary, [&](const uba::tchar* importName, bool isKnown) { func(importName, userData); });
+		uba::StringBuffer<> errors;
+		uba::FindImports(binary, [&](const uba::tchar* importName, bool isKnown, const char* const* importLoaderPaths) { func(importName, userData); }, errors);
 #endif
 	}
 
