[UBA]
* Fixed so code parsing dylib dependencies work properly on mac. dylibs in other folders are not properly transferred to the remote helper

diff --git a/Engine/Source/Programs/UnrealBuildAccelerator/Common/Private/UbaSession.cpp b/Engine/Source/Programs/UnrealBuildAccelerator/Common/Private/UbaSession.cpp
--- a/Engine/Source/Programs/UnrealBuildAccelerator/Common/Private/UbaSession.cpp
+++ b/Engine/Source/Programs/UnrealBuildAccelerator/Common/Private/UbaSession.cpp
@@ -939,6 +939,7 @@ namespace uba
 		u32 attr = GetFileAttributesW(applicationName); // TODO: Use attributes table
 		tchar temp[512];
 		tchar temp2[512];
+		StringBuffer<512> temp3;
 		bool result = true;
 
 		if (attr == INVALID_FILE_ATTRIBUTES)
@@ -949,16 +950,15 @@ namespace uba
 			#elif PLATFORM_MAC
 			if (!loaderPaths)
 				return m_logger.Error("Failed to find file %s", applicationName);
-			auto it = loaderPaths;
 			for (auto it = loaderPaths; *it; ++it)
 			{
 				StringBuffer<> absolutePath;
 				absolutePath.Append(applicationDir, applicationDirEnd - applicationDir).Append(*it).EnsureEndsWithSlash().Append(library);
-				m_logger.Info("SEARCH_PATH: %s", absolutePath.data);
-				attr = GetFileAttributesW(applicationName);
+				FixPath(absolutePath.data, nullptr, 0, temp3.Clear());
+				attr = GetFileAttributesW(temp3.data);
 				if (attr == INVALID_FILE_ATTRIBUTES)
 					continue;
-				memcpy(temp, absolutePath.data, absolutePath.count + 1);
+				memcpy(temp, temp3.data, temp3.count+1);
 				break;
 			}
 			if (attr == INVALID_FILE_ATTRIBUTES)
@@ -977,8 +977,7 @@ namespace uba
 			applicationDirEnd = temp2 + applicationDirLen;
 		}
 
-		StringBuffer<512> temp3;
-		FixPath(applicationName, nullptr, 0, temp3);
+		FixPath(applicationName, nullptr, 0, temp3.Clear());
 
 		bool isSystem = StartsWith(applicationName, m_systemPath.data);
 		if (isSystem && IsKnownSystemFile(applicationName))
