[UBA]
* Removed lock in HandleGetTraceInformation that should not be there

diff --git a/Engine/Source/Programs/UnrealBuildAccelerator/Common/Private/UbaSessionServer.cpp b/Engine/Source/Programs/UnrealBuildAccelerator/Common/Private/UbaSessionServer.cpp
--- a/Engine/Source/Programs/UnrealBuildAccelerator/Common/Private/UbaSessionServer.cpp
+++ b/Engine/Source/Programs/UnrealBuildAccelerator/Common/Private/UbaSessionServer.cpp
@@ -1286,7 +1286,6 @@ namespace uba
 	bool SessionServer::HandleGetTraceInformation(const ConnectionInfo& connectionInfo, BinaryReader& reader, BinaryWriter& writer)
 	{
 		u32 remotePos = reader.ReadU32();
-		SCOPED_READ_LOCK(m_trace.m_memoryLock, lock);
 		u32 localPos;
 		{
 			SCOPED_READ_LOCK(m_trace.m_memoryLock, l);
