[UBA]
* Fixed so CreateHardLink is allowed to fail on macos without throwing asserts
* Added CopyFile implementation for macos

diff --git a/Engine/Source/Programs/UnrealBuildAccelerator/Common/Private/UbaFile.cpp b/Engine/Source/Programs/UnrealBuildAccelerator/Common/Private/UbaFile.cpp
--- a/Engine/Source/Programs/UnrealBuildAccelerator/Common/Private/UbaFile.cpp
+++ b/Engine/Source/Programs/UnrealBuildAccelerator/Common/Private/UbaFile.cpp
@@ -22,6 +22,7 @@
 
 #if PLATFORM_MAC
 #include <mach-o/dyld.h>
+#include <copyfile.h>
 #endif
 
 namespace uba
@@ -571,6 +572,11 @@ namespace uba
 		MAKE_LONG_FILENAME(existingFileName);
 		MAKE_LONG_FILENAME(newFileName);
 		return ::CopyFileW(existingFileName, newFileName, bFailIfExists);
+#elif PLATFORM_MACOS
+		if (copyfile(existingFileName, newFileName, 0, COPYFILE_ALL) == 0)
+			return true;
+		UBA_ASSERTF(false, TC("CopyFileW failed on %s - Error handling not implemented (%s)"), existingFileName, strerror(errno));
+		return false;
 #else
 
 		UBA_ASSERTF(false, TC("CopyFileW not implemented"));
@@ -720,6 +726,11 @@ namespace uba
 		if (res == 0)
 			return true;
 
+		#if PLATFORM_MACOS
+		if (errno == EPERM) // Because of System Integrity Protection we might not be allowed to link this file, fallback to copy
+			return false;
+		#endif
+
 		UBA_ASSERTF(false, TC("CreateHardLinkW %s to %s error handling not implemented (%s)"), existingFileName, newFileName, strerror(errno));
 		return false;
 #endif
