[UBA]
* Fixed logging for linux

diff --git a/Engine/Source/Programs/UnrealBuildAccelerator/Detours/Private/PosixOS/UbaDetoursFunctionsPosix.cpp b/Engine/Source/Programs/UnrealBuildAccelerator/Detours/Private/PosixOS/UbaDetoursFunctionsPosix.cpp
--- a/Engine/Source/Programs/UnrealBuildAccelerator/Detours/Private/PosixOS/UbaDetoursFunctionsPosix.cpp
+++ b/Engine/Source/Programs/UnrealBuildAccelerator/Detours/Private/PosixOS/UbaDetoursFunctionsPosix.cpp
@@ -850,7 +850,7 @@ UBA_EXPORT char* UBA_WRAPPER(getenv)(const char* name)
 {
 	UBA_INIT_DETOUR(getenv, name);
 	auto res = TRUE_WRAPPER(getenv)(name);
-	DEBUG_LOG_TRUE("getenv", "(%s) -> %s", name ? name : "<null>", res ? res : "<null>");
+	DEBUG_LOG_TRUE("getenv", "(%s) -> %s", name, res ? res : "<null>");
 	return res;
 }
 
@@ -1791,14 +1791,17 @@ const char* GetResult(siginfo_t* info)
 {
 	if (!info)
 		return "null";
-	if (WIFEXITED(info))
+	int code = info->si_code;
+	if (code == CLD_EXITED)
 		return "Exited";
-	if (WIFSIGNALED(info))
-		return "Signaled";
-	if (WIFSTOPPED(info))
+	if (code == CLD_KILLED)
+		return "Killed";
+	if (code == CLD_STOPPED)
 		return "Stopped";
-	if (WIFCONTINUED(info))
+	if (code == CLD_CONTINUED)
 		return "Continued";
+	if (code == CLD_TRAPPED)
+		return "Trapped";
 	return "Running";
 }
 
