[UBA]
* Added -noproxy option to UbaAgent to be able to spawn agents that are not allowed to be proxy of others

diff --git a/Engine/Source/Programs/UnrealBuildAccelerator/Agent/Private/UbaAgent.cpp b/Engine/Source/Programs/UnrealBuildAccelerator/Agent/Private/UbaAgent.cpp
--- a/Engine/Source/Programs/UnrealBuildAccelerator/Agent/Private/UbaAgent.cpp
+++ b/Engine/Source/Programs/UnrealBuildAccelerator/Agent/Private/UbaAgent.cpp
@@ -103,6 +103,7 @@ namespace uba
 		logger.Info(TC("  -eventfile=<file>       File containing external events to agent. Things like machine is about to be terminated etc"));
 		logger.Info(TC("  -sentry                 Enable sentry"));
 		logger.Info(TC("  -zone                   Set the zone this machine exists in. This info is used to figure out if proxies should be created."));
+		logger.Info(TC("  -noproxy                Does not allow this agent to be a storage proxy for other agents"));
 		logger.Info(TC("  -killrandom             Kills random process and exit session"));
 		logger.Info(TC("  -memwait=<percent>      The amount of memory needed to spawn a process. Set this to 100 to disable. Defaults to 80%%"));
 		logger.Info(TC("  -memkill=<percent>      The amount of memory needed before processes starts to be killed. Set this to 100 to disable. Defaults to 90%%"));
@@ -424,6 +425,7 @@ namespace uba
 		bool useBinariesAsVersion = false;
 		bool useQuic = false;
 		bool poll = true;
+		bool allowProxy = true;
 		bool useStorage = true;
 		bool resetStore = false;
 		bool quiet = false;
@@ -683,6 +685,10 @@ namespace uba
 					return PrintHelp(TC("-zone needs a value"));
 				zone.Append(value);
 			}
+			else if (name.Equals(TC("-noproxy")))
+			{
+				allowProxy = false;
+			}
 			else if (name.Equals(TC("-command")))
 			{
 				if (value.IsEmpty())
@@ -1147,6 +1153,7 @@ namespace uba
 			storageInfo.workManager = client;
 			storageInfo.getProxyBackendCallback = getProxyBackend;
 			storageInfo.getProxyBackendUserData = &backends;
+			storageInfo.allowProxy = allowProxy;
 			storageInfo.startProxyCallback = startProxy;
 			storageInfo.startProxyUserData = &proxy;
 			storageInfo.zone = zone.data;
