[UBA]
* Fixed so casentries collected through populating from disk have their verified and exists set to true
* Changed so copyorlink the cas data to temporary folders for execution are never linked on macos.. for some reason this does not always work which is strange

diff --git a/Engine/Source/Programs/UnrealBuildAccelerator/Common/Private/UbaStorage.cpp b/Engine/Source/Programs/UnrealBuildAccelerator/Common/Private/UbaStorage.cpp
--- a/Engine/Source/Programs/UnrealBuildAccelerator/Common/Private/UbaStorage.cpp
+++ b/Engine/Source/Programs/UnrealBuildAccelerator/Common/Private/UbaStorage.cpp
@@ -2527,9 +2527,15 @@ namespace uba
 			bool firstTry = true;
 			while (true)
 			{
-				bool success = true;
-				if (CreateHardLinkW(destination, casFile.data) == 0)
+				bool success = false;
+
+				#if !PLATFORM_MAC // For some reason creating links on macos causes trouble when they are exec/dylibs and being executed.. sometimes it is like the link behaves like a symlink.. but not always
+				success = CreateHardLinkW(destination, casFile.data);
+				#endif
+
+				if (!success)
 					success = uba::CopyFileW(casFile.data, destination, true) != 0;
+
 				if (success)
 				{
 					#if !PLATFORM_WINDOWS
diff --git a/Engine/Source/Programs/UnrealBuildAccelerator/Common/Private/UbaStorageClient.cpp b/Engine/Source/Programs/UnrealBuildAccelerator/Common/Private/UbaStorageClient.cpp
--- a/Engine/Source/Programs/UnrealBuildAccelerator/Common/Private/UbaStorageClient.cpp
+++ b/Engine/Source/Programs/UnrealBuildAccelerator/Common/Private/UbaStorageClient.cpp
@@ -876,6 +876,8 @@ namespace uba
 					if (insres.second)
 					{
 						localFile.casEntry.size = e.size;
+						localFile.casEntry.verified = true;
+						localFile.casEntry.exists = true;
 						localFile.fileName = fullPath.data;
 					}
 					return;
@@ -906,6 +908,8 @@ namespace uba
 						if (insres.second)
 						{
 							localFile.casEntry.size = s;
+							localFile.casEntry.verified = true;
+							localFile.casEntry.exists = true;
 							localFile.fileName = filePath;
 						}
 
