[UBA]
* Added support for sending loader paths through GetFullFilename logic

diff --git a/Engine/Source/Programs/UnrealBuildAccelerator/Common/Private/UbaProcess.cpp b/Engine/Source/Programs/UnrealBuildAccelerator/Common/Private/UbaProcess.cpp
--- a/Engine/Source/Programs/UnrealBuildAccelerator/Common/Private/UbaProcess.cpp
+++ b/Engine/Source/Programs/UnrealBuildAccelerator/Common/Private/UbaProcess.cpp
@@ -326,6 +326,7 @@ namespace uba
 					break;
 				}
 
+				u32 nativeProcessId = m_nativeProcessId;
 				m_nativeProcessId = 0;
 				m_nativeProcessExitCode = signalInfo.si_status;
 				
@@ -333,7 +334,7 @@ namespace uba
 					break;
 					
 				StringBuffer<> err;
-				err.Appendf(TC("Process %u (%s) %s by signal %i. Received %u messages. Execution time: %s."), m_nativeProcessId, m_startInfo.GetDescription(), codeType, signalInfo.si_status, m_messageCount, TimeToText(GetTime() - m_startTime).str);
+				err.Appendf(TC("Process %u (%s) %s by signal %i. Received %u messages. Execution time: %s."), nativeProcessId, m_startInfo.GetDescription(), codeType, signalInfo.si_status, m_messageCount, TimeToText(GetTime() - m_startTime).str);
 				LogLine(false, err.data, LogEntryType_Error);
 				m_nativeProcessExitCode = UBA_EXIT_CODE(666); // We do exit code 666 to trigger non-uba retry on the outside
 				return false;
@@ -713,6 +714,9 @@ namespace uba
 		GetFullFileNameMessage msg { *this };
 		reader.ReadString(msg.fileName);
 		msg.fileNameKey = reader.ReadStringKey();
+		msg.loaderPaths = reader.GetPositionData();
+		msg.loaderPathsSize = u32(reader.GetLeft());
+		
 		GetFullFileNameResponse response;
 		m_messageSuccess = m_session.GetFullFileName(response, msg) && m_messageSuccess;
 		writer.WriteString(response.fileName);
diff --git a/Engine/Source/Programs/UnrealBuildAccelerator/Common/Private/UbaSessionClient.cpp b/Engine/Source/Programs/UnrealBuildAccelerator/Common/Private/UbaSessionClient.cpp
--- a/Engine/Source/Programs/UnrealBuildAccelerator/Common/Private/UbaSessionClient.cpp
+++ b/Engine/Source/Programs/UnrealBuildAccelerator/Common/Private/UbaSessionClient.cpp
@@ -192,7 +192,7 @@ namespace uba
 		return true;
 	}
 
-	bool SessionClient::EnsureBinaryFile(StringBufferBase& out, StringBufferBase& outVirtual, u32 processId, const StringBufferBase& fileName, const StringKey& fileNameKey, const tchar* applicationDir)
+	bool SessionClient::EnsureBinaryFile(StringBufferBase& out, StringBufferBase& outVirtual, u32 processId, const StringBufferBase& fileName, const StringKey& fileNameKey, const tchar* applicationDir, const u8* loaderPaths, u32 loaderPathsSize)
 	{
 		CasKey casKey;
 		u32 fileAttributes = DefaultAttributes(); // TODO: This is wrong.. need to retrieve from server if this is executable or not
@@ -217,6 +217,8 @@ namespace uba
 			writer.WriteString(fileName);
 			writer.WriteStringKey(fileNameKey);
 			writer.WriteString(applicationDir);
+			if (loaderPathsSize)
+				writer.WriteBytes(loaderPaths, loaderPathsSize);
 
 			StackBinaryReader<1024> reader;
 			if (!msg.Send(reader, Stats().getBinaryMsg))
@@ -859,7 +861,7 @@ namespace uba
 
 		StringBuffer<> dir;
 		dir.AppendDir(msg.process.m_startInfo.application);
-		if (!EnsureBinaryFile(out.fileName, out.virtualFileName, msg.process.m_id, msg.fileName, msg.fileNameKey, dir.data))
+		if (!EnsureBinaryFile(out.fileName, out.virtualFileName, msg.process.m_id, msg.fileName, msg.fileNameKey, dir.data, msg.loaderPaths, msg.loaderPathsSize))
 			return false;
 
 		StringKey fileNameKey = msg.fileNameKey;
diff --git a/Engine/Source/Programs/UnrealBuildAccelerator/Common/Private/UbaSessionServer.cpp b/Engine/Source/Programs/UnrealBuildAccelerator/Common/Private/UbaSessionServer.cpp
--- a/Engine/Source/Programs/UnrealBuildAccelerator/Common/Private/UbaSessionServer.cpp
+++ b/Engine/Source/Programs/UnrealBuildAccelerator/Common/Private/UbaSessionServer.cpp
@@ -716,15 +716,40 @@ namespace uba
 			return true;
 		}
 
+		Vector<TString> loaderPaths;
+		while (reader.GetLeft())
+			loaderPaths.push_back(reader.ReadString());
+
 		CasKey casKey = CasKeyZero;
 		StringBuffer<> absoluteFile;
-		if (SearchPathForFile(m_logger, absoluteFile, fileName.data, applicationDir.data))
-			if (!absoluteFile.StartsWith(m_systemPath.data) || !IsKnownSystemFile(absoluteFile.data))
+
+		if (!loaderPaths.empty())
+		{
+			for (auto& loaderPath : loaderPaths)
 			{
+				StringBuffer<> fullPath;
+				fullPath.Append(applicationDir).Append(loaderPath).Append(fileName);
+				if (GetFileAttributesW(fullPath.data) == INVALID_FILE_ATTRIBUTES)
+					continue;
+				StringBuffer<> out;
+				FixPath(fullPath.data, nullptr, 0, absoluteFile);
 				fileNameKey = ToStringKeyLower(absoluteFile);
 				if (!StoreCasFile(casKey, fileNameKey, absoluteFile.data))
 					return false;
+
+				break;
 			}
+		}
+		else
+		{
+			if (SearchPathForFile(m_logger, absoluteFile, fileName.data, applicationDir.data))
+				if (!absoluteFile.StartsWith(m_systemPath.data) || !IsKnownSystemFile(absoluteFile.data))
+				{
+					fileNameKey = ToStringKeyLower(absoluteFile);
+					if (!StoreCasFile(casKey, fileNameKey, absoluteFile.data))
+						return false;
+				}
+		}
 
 		u64 startPos = writer.GetPosition();
 		writer.WriteCasKey(casKey);
diff --git a/Engine/Source/Programs/UnrealBuildAccelerator/Common/Public/UbaSession.h b/Engine/Source/Programs/UnrealBuildAccelerator/Common/Public/UbaSession.h
--- a/Engine/Source/Programs/UnrealBuildAccelerator/Common/Public/UbaSession.h
+++ b/Engine/Source/Programs/UnrealBuildAccelerator/Common/Public/UbaSession.h
@@ -398,6 +398,8 @@ namespace uba
 		ProcessImpl& process;
 		StringBuffer<> fileName;
 		StringKey fileNameKey;
+		const u8* loaderPaths = nullptr;
+		u32 loaderPathsSize = 0;
 	};
 
 	struct GetFullFileNameResponse
diff --git a/Engine/Source/Programs/UnrealBuildAccelerator/Common/Public/UbaSessionClient.h b/Engine/Source/Programs/UnrealBuildAccelerator/Common/Public/UbaSessionClient.h
--- a/Engine/Source/Programs/UnrealBuildAccelerator/Common/Public/UbaSessionClient.h
+++ b/Engine/Source/Programs/UnrealBuildAccelerator/Common/Public/UbaSessionClient.h
@@ -76,7 +76,7 @@ namespace uba
 
 		bool GetCasKeyForFile(CasKey& out, u32 processId, const StringBufferBase& fileName, const StringKey& fileNameKey);
 		bool ReadModules(List<ModuleInfo>& outModules, u32 processId, const tchar* application);
-		bool EnsureBinaryFile(StringBufferBase& out, StringBufferBase& outVirtual, u32 processId, const StringBufferBase& fileName, const StringKey& fileNameKey, const tchar* applicationDir);
+		bool EnsureBinaryFile(StringBufferBase& out, StringBufferBase& outVirtual, u32 processId, const StringBufferBase& fileName, const StringKey& fileNameKey, const tchar* applicationDir, const u8* loaderPaths, u32 loaderPathsSize);
 		bool WriteBinFile(StringBufferBase& out, const tchar* binaryName, const CasKey& casKey, const KeyToString& applicationDir, u32 fileAttributes);
 		bool SendFiles(ProcessImpl& process, Timer& sendFiles);
 		bool SendFile(WrittenFile& source, const tchar* destination, u32 processId, bool keepMappingInMemory);
diff --git a/Engine/Source/Programs/UnrealBuildAccelerator/Detours/Private/UbaDetoursFileMappingTable.cpp b/Engine/Source/Programs/UnrealBuildAccelerator/Detours/Private/UbaDetoursFileMappingTable.cpp
--- a/Engine/Source/Programs/UnrealBuildAccelerator/Detours/Private/UbaDetoursFileMappingTable.cpp
+++ b/Engine/Source/Programs/UnrealBuildAccelerator/Detours/Private/UbaDetoursFileMappingTable.cpp
@@ -248,7 +248,7 @@ namespace uba
 		return findIt->second;
 	}
 
-	void Rpc_GetFullFileName(const tchar*& path, u64& pathLen, StringBufferBase& tempBuf, bool useVirtualName)
+	void Rpc_GetFullFileName(const tchar*& path, u64& pathLen, StringBufferBase& tempBuf, bool useVirtualName, const tchar* const* loaderPaths)
 	{
 		StringKey fileNameKey;
 		if (IsAbsolutePath(path))
@@ -269,6 +269,9 @@ namespace uba
 			writer.WriteByte(MessageType_GetFullFileName);
 			writer.WriteString(path);
 			writer.WriteStringKey(fileNameKey);
+			if (loaderPaths)
+				for (auto i=loaderPaths; *i; ++i)
+					writer.WriteString(*i);
 			writer.Flush();
 			BinaryReader reader;
 			reader.ReadString(tempBuf);
diff --git a/Engine/Source/Programs/UnrealBuildAccelerator/Detours/Private/UbaDetoursFileMappingTable.h b/Engine/Source/Programs/UnrealBuildAccelerator/Detours/Private/UbaDetoursFileMappingTable.h
--- a/Engine/Source/Programs/UnrealBuildAccelerator/Detours/Private/UbaDetoursFileMappingTable.h
+++ b/Engine/Source/Programs/UnrealBuildAccelerator/Detours/Private/UbaDetoursFileMappingTable.h
@@ -64,7 +64,7 @@ namespace uba
 	void Rpc_UpdateCloseHandle(const tchar* handleName, u32 closeId, bool deleteOnClose, const tchar* newName, u64 mappingHandle, u64 mappingWritten, bool success);
 	void Rpc_UpdateTables();
 	u32  Rpc_GetEntryOffset(const StringKey& entryNameKey, const tchar* entryName, u64 entryNameLen, bool checkIfDir = false);
-	void Rpc_GetFullFileName(const tchar*& path, u64& pathLen, StringBufferBase& tempBuf, bool useVirtualName);
+	void Rpc_GetFullFileName(const tchar*& path, u64& pathLen, StringBufferBase& tempBuf, bool useVirtualName, const tchar* const* loaderPaths = nullptr);
 
 	struct DirHash
 	{
