[UBA]
* Fixed bug with how loader paths are sent from detour to session

diff --git a/Engine/Source/Programs/UnrealBuildAccelerator/Common/Private/UbaProcess.cpp b/Engine/Source/Programs/UnrealBuildAccelerator/Common/Private/UbaProcess.cpp
--- a/Engine/Source/Programs/UnrealBuildAccelerator/Common/Private/UbaProcess.cpp
+++ b/Engine/Source/Programs/UnrealBuildAccelerator/Common/Private/UbaProcess.cpp
@@ -714,8 +714,8 @@ namespace uba
 		GetFullFileNameMessage msg { *this };
 		reader.ReadString(msg.fileName);
 		msg.fileNameKey = reader.ReadStringKey();
+		msg.loaderPathsSize = reader.ReadU16();
 		msg.loaderPaths = reader.GetPositionData();
-		msg.loaderPathsSize = u32(reader.GetLeft());
 		
 		GetFullFileNameResponse response;
 		m_messageSuccess = m_session.GetFullFileName(response, msg) && m_messageSuccess;
diff --git a/Engine/Source/Programs/UnrealBuildAccelerator/Detours/Private/UbaDetoursFileMappingTable.cpp b/Engine/Source/Programs/UnrealBuildAccelerator/Detours/Private/UbaDetoursFileMappingTable.cpp
--- a/Engine/Source/Programs/UnrealBuildAccelerator/Detours/Private/UbaDetoursFileMappingTable.cpp
+++ b/Engine/Source/Programs/UnrealBuildAccelerator/Detours/Private/UbaDetoursFileMappingTable.cpp
@@ -269,9 +269,12 @@ namespace uba
 			writer.WriteByte(MessageType_GetFullFileName);
 			writer.WriteString(path);
 			writer.WriteStringKey(fileNameKey);
+			u16& bytes = *(u16*)writer.AllocWrite(2);
+			auto pos = writer.GetPosition();
 			if (loaderPaths)
 				for (auto i=loaderPaths; *i; ++i)
 					writer.WriteString(*i);
+			bytes = u16(writer.GetPosition() - pos);
 			writer.Flush();
 			BinaryReader reader;
 			reader.ReadString(tempBuf);
diff --git a/Engine/Source/Programs/UnrealBuildAccelerator/Detours/Private/Windows/UbaDetoursFunctionsKernelBase.inl b/Engine/Source/Programs/UnrealBuildAccelerator/Detours/Private/Windows/UbaDetoursFunctionsKernelBase.inl
--- a/Engine/Source/Programs/UnrealBuildAccelerator/Detours/Private/Windows/UbaDetoursFunctionsKernelBase.inl
+++ b/Engine/Source/Programs/UnrealBuildAccelerator/Detours/Private/Windows/UbaDetoursFunctionsKernelBase.inl
@@ -2437,11 +2437,12 @@ HMODULE Recursive_LoadLibraryExW(LPCWSTR lpLibFileName, LPCWSTR originalName, DW
 	std::vector<Import, GrowingAllocator<Import>> importedModules(&g_memoryBlock);
 	{
 		SuppressCreateFileDetourScope cfs;
-		if (!FindImports(lpLibFileName, [&](const wchar_t* import, bool isKnown)
+		StringBuffer<256> error;
+		if (!FindImports(lpLibFileName, [&](const wchar_t* import, bool isKnown, const char* const* importLoaderPaths)
 			{
 				if (!GetModuleHandleW(import))
 					importedModules.emplace_back(import, isKnown);
-			}))
+			}, error))
 		{
 			UBA_ASSERTF(false, L"Failed to find imports for binary %ls (%ls)", lpLibFileName, originalName);
 		}
