[UBA]
* Fixed so dlopen work properly remotely and downloads dylib plus depending dylib before trying to load them

diff --git a/Engine/Source/Programs/UnrealBuildAccelerator/Common/Private/UbaSessionServer.cpp b/Engine/Source/Programs/UnrealBuildAccelerator/Common/Private/UbaSessionServer.cpp
--- a/Engine/Source/Programs/UnrealBuildAccelerator/Common/Private/UbaSessionServer.cpp
+++ b/Engine/Source/Programs/UnrealBuildAccelerator/Common/Private/UbaSessionServer.cpp
@@ -728,7 +728,7 @@ namespace uba
 			for (auto& loaderPath : loaderPaths)
 			{
 				StringBuffer<> fullPath;
-				fullPath.Append(applicationDir).Append(loaderPath).Append(fileName);
+				fullPath.Append(applicationDir).EnsureEndsWithSlash().Append(loaderPath).EnsureEndsWithSlash().Append(fileName);
 				if (GetFileAttributesW(fullPath.data) == INVALID_FILE_ATTRIBUTES)
 					continue;
 				StringBuffer<> out;
@@ -736,7 +736,6 @@ namespace uba
 				fileNameKey = ToStringKeyLower(absoluteFile);
 				if (!StoreCasFile(casKey, fileNameKey, absoluteFile.data))
 					return false;
-
 				break;
 			}
 		}
diff --git a/Engine/Source/Programs/UnrealBuildAccelerator/Core/Public/UbaMacBinDependencyParser.h b/Engine/Source/Programs/UnrealBuildAccelerator/Core/Public/UbaMacBinDependencyParser.h
--- a/Engine/Source/Programs/UnrealBuildAccelerator/Core/Public/UbaMacBinDependencyParser.h
+++ b/Engine/Source/Programs/UnrealBuildAccelerator/Core/Public/UbaMacBinDependencyParser.h
@@ -17,28 +17,28 @@ namespace uba
 
 	inline bool FindImports(const tchar* fileName, const Function<void(const tchar* import, bool isKnown, const char* const* loaderPaths)>& func, StringBufferBase& outError)
 	{
-		int fd = open(fileName, O_RDONLY);
+		int fd = TRUE_WRAPPER(open)(fileName, O_RDONLY);
 		if (fd == -1)
 		{
 			outError.Appendf("Open failed for file %s", fileName);
 			return false;
 		}
-		auto closeFileHandle = MakeGuard([&]() { close(fd); });
+		auto closeFileHandle = MakeGuard([&]() { TRUE_WRAPPER(close)(fd); });
 		struct stat sb;
-		if (fstat(fd, &sb) == -1)
+		if (TRUE_WRAPPER(fstat)(fd, &sb) == -1)
 		{
 			outError.Appendf("Stat failed for file %s", fileName);
 			return false;
 		}
 		u32 size = Min(u32(sb.st_size), 8048u);
 
-		void* mem = mmap(nullptr, size, PROT_READ, MAP_PRIVATE, fd, 0);
+		void* mem = TRUE_WRAPPER(mmap)(nullptr, size, PROT_READ, MAP_PRIVATE, fd, 0);
 		if (mem == MAP_FAILED)
 		{
 			outError.Appendf("Mmap failed for file %s", fileName);
 			return false;
 		}
-		auto unmap = MakeGuard([&]() { munmap(mem, size); });
+		auto unmap = MakeGuard([&]() { TRUE_WRAPPER(munmap)(mem, size); });
 
 
 		const char* libs[1024];
diff --git a/Engine/Source/Programs/UnrealBuildAccelerator/Detours/Private/PosixOS/UbaDetoursFunctionsPosix.cpp b/Engine/Source/Programs/UnrealBuildAccelerator/Detours/Private/PosixOS/UbaDetoursFunctionsPosix.cpp
--- a/Engine/Source/Programs/UnrealBuildAccelerator/Detours/Private/PosixOS/UbaDetoursFunctionsPosix.cpp
+++ b/Engine/Source/Programs/UnrealBuildAccelerator/Detours/Private/PosixOS/UbaDetoursFunctionsPosix.cpp
@@ -80,6 +80,8 @@ using namespace uba;
 	DETOURED_FUNCTION(renameat) \
 	DETOURED_FUNCTION(utimensat) \
 	DETOURED_FUNCTION(remove) \
+	DETOURED_FUNCTION(link) \
+	DETOURED_FUNCTION(unlink) \
 	DETOURED_FUNCTION(symlink) \
 	DETOURED_FUNCTION(access) \
 	DETOURED_FUNCTION(posix_spawn) \
@@ -151,6 +153,8 @@ using namespace uba;
 			__attribute__ ((section ("__DATA,__interpose"))) = { (const void*)(unsigned long)&_replacement, (const void*)(unsigned long)&_replacee };
 	#endif // DYLD_INTERPOSE
 
+	#include "UbaMacBinDependencyParser.h"
+
 #else
 	#define UBA_WRAPPER(func) func
 	#define TRUE_WRAPPER(func) True_##func
@@ -161,6 +165,8 @@ using namespace uba;
 	DETOURED_FUNCTIONS
 	#undef DETOURED_FUNCTION
 
+	#include "UbaLinuxBinDependencyParser.h"
+
 #endif
 
 #if 0
@@ -1569,26 +1575,10 @@ UBA_EXPORT ssize_t UBA_WRAPPER(read)(int fd, void *buf, size_t nbyte)
 	return TRUE_WRAPPER(read)(fd, buf, nbyte);
 }
 
-
-UBA_EXPORT int UBA_WRAPPER(remove)(const char* pathname)
+int Shared_DeleteFile(const char* funcName, const char* pathname)
 {
-	UBA_INIT_DETOUR(remove, pathname);
 	StringBuffer<> fixedName;
 	FixPath(fixedName, pathname);
-
-	//if (!CanDetour(fixedName.data))
-	//{
-	//	DEBUG_LOG_TRUE(L"DeleteFile", L"(%ls)", original);
-	//	return TRUE_WRAPPER(remove)(pathname);
-	//}
-
-	//if (KeepInMemory(fixedName.data, fixedName.count))
-	//{
-	//	DEBUG_LOG_DETOURED(L"DeleteFile", L"(INMEMORY) (%ls) -> Success", lpFileName);
-	//	SetLastError(ERROR_SUCCESS);
-	//	return TRUE;
-	//}
-
 	StringKey fileNameKey = ToFilenameKey(fixedName);
 
 	u32 directoryTableSize;
@@ -1617,12 +1607,32 @@ UBA_EXPORT int UBA_WRAPPER(remove)(const char* pathname)
 	g_mappedFileTable.SetDeleted(fileNameKey, fixedName.data, true);
 
 	int res = result != 0 ? 0 : -1;
-	DEBUG_LOG_DETOURED("remove", "(%s) -> %i (%s)", pathname, res, StrError(res, errorCode));
+	DEBUG_LOG_DETOURED(funcName, "(%s) -> %i (%s)", pathname, res, StrError(res, errorCode));
 	errno = errorCode;
 	return res;
 	//return TRUE_WRAPPER(remove)(pathname);
 }
 
+UBA_EXPORT int UBA_WRAPPER(remove)(const char* pathname)
+{
+	UBA_INIT_DETOUR(remove, pathname);
+	// TODO: Should check if pathname is a dir, in that case call rmdir
+	return Shared_DeleteFile("remove", pathname);
+}
+
+UBA_EXPORT int UBA_WRAPPER(link)(const char* oldpath, const char* newpath)
+{
+	UBA_INIT_DETOUR(link, oldpath, newpath);
+	int res = TRUE_WRAPPER(link)(oldpath, newpath);
+	return res;
+}
+
+UBA_EXPORT int UBA_WRAPPER(unlink)(const char* pathname)
+{
+	UBA_INIT_DETOUR(unlink, pathname);
+	return Shared_DeleteFile("unlink", pathname);
+}
+
 thread_local int t_inVfork;
 
 int shared_posix_spawn(pid_t* pid, const char* path, const posix_spawn_file_actions_t* file_actions, const posix_spawnattr_t* attrp, char* const argv[], char* const envp[])
@@ -1703,7 +1713,6 @@ int shared_posix_spawn(pid_t* pid, const char* path, const posix_spawn_file_acti
 		reader.ReadString(logFile);
 
 		// TODO: Recalculate argv2 right now it is wrong, it should be what is in command line
-		//argv2[0] = realApplication.c_str();
 	}
 	
 	std::vector<const char*> envvars;
@@ -1834,12 +1843,52 @@ UBA_EXPORT pid_t UBA_WRAPPER(wait4)(pid_t pid, int* status, int options, struct
 	return res;
 }
 
+bool Shared_LoadLibrary(Set<TString>& handled, const char*& path, const char* const* loaderPaths, StringBufferBase& tempBuf)
+{
+#if PLATFORM_MAC
+	u64 nameLen = 0;
+	Rpc_GetFullFileName(path, nameLen, tempBuf, false, loaderPaths);
+
+	StringBuffer<> error;
+	FindImports(path, [&](const tchar* import, bool isKnown, const char* const* importLoaderPaths)
+	{
+		if (!handled.insert(import).second)
+			return;
+		StringBuffer<> temp;
+		DEBUG_LOG("IMPORT: %s", import);
+		Shared_LoadLibrary(handled, import, importLoaderPaths, temp);
+	}, error);
+	if (error.count)
+		DEBUG_LOG(error.data)
+	//DEBUG_LOG("Rpc_GetFullFileName: %s", tempBuf.data);
+#endif
+	return true;
+}
+
 UBA_EXPORT void* UBA_WRAPPER(dlopen)(const char* path, int mode)
 {
 	UBA_INIT_DETOUR(dlopen, path, mode);
-	DEBUG_LOG_TRUE("dlopen", "%s", path);
+
+#if PLATFORM_MAC
+	StringBuffer<> tempBuf;
+	
+	auto originalPath = path;
+	if (StartsWith(path, "@rpath/"))
+	{
+		path += 7;
+		Set<TString> handled;
+		handled.insert(path);
+		const char* loaderPaths[] = { "/", 0 };
+		Shared_LoadLibrary(handled, path, loaderPaths, tempBuf);
+	}
+	else if (!StartsWith(path, "/System") && !StartsWith(path, "/usr/lib"))
+	{
+		u64 nameLen = 0;
+		Rpc_GetFullFileName(path, nameLen, tempBuf, false);
+	}
+#endif
 	void* res = TRUE_WRAPPER(dlopen)(path, mode);
-	DEBUG_LOG_TRUE("dlopen DONE", "%s", path);
+	DEBUG_LOG_TRUE("dlopen", "%s (%i) -> 0x%x", path, mode, res);
 	return res;
 }
 
