[UBA]
* Fixed code using wrong platform define

diff --git a/Engine/Source/Programs/UnrealBuildAccelerator/Common/Private/UbaFile.cpp b/Engine/Source/Programs/UnrealBuildAccelerator/Common/Private/UbaFile.cpp
--- a/Engine/Source/Programs/UnrealBuildAccelerator/Common/Private/UbaFile.cpp
+++ b/Engine/Source/Programs/UnrealBuildAccelerator/Common/Private/UbaFile.cpp
@@ -572,7 +572,7 @@ namespace uba
 		MAKE_LONG_FILENAME(existingFileName);
 		MAKE_LONG_FILENAME(newFileName);
 		return ::CopyFileW(existingFileName, newFileName, bFailIfExists);
-#elif PLATFORM_MACOS
+#elif PLATFORM_MAC
 		if (copyfile(existingFileName, newFileName, 0, COPYFILE_ALL) == 0)
 			return true;
 		UBA_ASSERTF(false, TC("CopyFileW failed on %s - Error handling not implemented (%s)"), existingFileName, strerror(errno));
@@ -726,7 +726,7 @@ namespace uba
 		if (res == 0)
 			return true;
 
-		#if PLATFORM_MACOS
+		#if PLATFORM_MAC
 		if (errno == EPERM) // Because of System Integrity Protection we might not be allowed to link this file, fallback to copy
 			return false;
 		#endif
