[UBA]
* Removed logging that should not be there

diff --git a/Engine/Source/Programs/UnrealBuildAccelerator/Detours/Private/PosixOS/UbaDetoursFunctionsPosix.cpp b/Engine/Source/Programs/UnrealBuildAccelerator/Detours/Private/PosixOS/UbaDetoursFunctionsPosix.cpp
--- a/Engine/Source/Programs/UnrealBuildAccelerator/Detours/Private/PosixOS/UbaDetoursFunctionsPosix.cpp
+++ b/Engine/Source/Programs/UnrealBuildAccelerator/Detours/Private/PosixOS/UbaDetoursFunctionsPosix.cpp
@@ -1845,7 +1845,7 @@ UBA_EXPORT pid_t UBA_WRAPPER(wait4)(pid_t pid, int* status, int options, struct
 	return res;
 }
 
-bool Shared_LoadLibrary(Set<TString>& handled, const char*& path, const char* const* loaderPaths, StringBufferBase& tempBuf)
+void Shared_LoadLibrary(Set<TString>& handled, const char*& path, const char* const* loaderPaths, StringBufferBase& tempBuf)
 {
 #if PLATFORM_MAC
 	u64 nameLen = 0;
@@ -1857,14 +1857,11 @@ bool Shared_LoadLibrary(Set<TString>& handled, const char*& path, const char* co
 		if (!handled.insert(import).second)
 			return;
 		StringBuffer<> temp;
-		DEBUG_LOG("IMPORT: %s", import);
 		Shared_LoadLibrary(handled, import, importLoaderPaths, temp);
 	}, error);
 	if (error.count)
 		DEBUG_LOG(error.data)
-	//DEBUG_LOG("Rpc_GetFullFileName: %s", tempBuf.data);
 #endif
-	return true;
 }
 
 UBA_EXPORT void* UBA_WRAPPER(dlopen)(const char* path, int mode)
