[UBA]
* Fixed macos shipping compile error

diff --git a/Engine/Source/Programs/UnrealBuildAccelerator/Detours/Private/PosixOS/UbaDetoursFunctionsPosix.cpp b/Engine/Source/Programs/UnrealBuildAccelerator/Detours/Private/PosixOS/UbaDetoursFunctionsPosix.cpp
--- a/Engine/Source/Programs/UnrealBuildAccelerator/Detours/Private/PosixOS/UbaDetoursFunctionsPosix.cpp
+++ b/Engine/Source/Programs/UnrealBuildAccelerator/Detours/Private/PosixOS/UbaDetoursFunctionsPosix.cpp
@@ -1860,7 +1860,7 @@ void Shared_LoadLibrary(Set<TString>& handled, const char*& path, const char* co
 		Shared_LoadLibrary(handled, import, importLoaderPaths, temp);
 	}, error);
 	if (error.count)
-		DEBUG_LOG(error.data)
+		DEBUG_LOG(error.data);
 #endif
 }
 
diff --git a/Engine/Source/Programs/UnrealBuildAccelerator/Detours/Private/UbaDetoursShared.h b/Engine/Source/Programs/UnrealBuildAccelerator/Detours/Private/UbaDetoursShared.h
--- a/Engine/Source/Programs/UnrealBuildAccelerator/Detours/Private/UbaDetoursShared.h
+++ b/Engine/Source/Programs/UnrealBuildAccelerator/Detours/Private/UbaDetoursShared.h
@@ -37,9 +37,9 @@ namespace uba
 		//#define DEBUG_LOG(...)
 		#define DEBUG_LOG(...) { if (isLogging()) WriteDebugLog(__VA_ARGS__); }
 	#else
-		#define DEBUG_LOG(...)
-		#define DEBUG_LOG_DETOURED(Command, ...)
-		#define DEBUG_LOG_TRUE(Command, ...)
+		#define DEBUG_LOG(...) {}
+		#define DEBUG_LOG_DETOURED(Command, ...) {}
+		#define DEBUG_LOG_TRUE(Command, ...) {}
 		#define DEBUG_LOG_PIPE(...) ts.Leave();
 	#endif
 
