[UBA]
* Fixed mac unit test compile error

diff --git a/Engine/Source/Programs/UnrealBuildAccelerator/Test/Private/UbaTestBasics.h b/Engine/Source/Programs/UnrealBuildAccelerator/Test/Private/UbaTestBasics.h
--- a/Engine/Source/Programs/UnrealBuildAccelerator/Test/Private/UbaTestBasics.h
+++ b/Engine/Source/Programs/UnrealBuildAccelerator/Test/Private/UbaTestBasics.h
@@ -461,7 +461,7 @@ namespace uba
 	}
 
 
-#if PLATFORM_MAC
+#if 0//PLATFORM_MAC
 	Set<TString> g_visited;
 	void LogImports(const tchar* import, bool isKnown)
 	{
