[UBA]
* Added StringView::StartsWith
* Changed so CasKeyBuilder::Update returns itself

diff --git a/Engine/Source/Programs/UnrealBuildAccelerator/Core/Private/UbaHash.cpp b/Engine/Source/Programs/UnrealBuildAccelerator/Core/Private/UbaHash.cpp
--- a/Engine/Source/Programs/UnrealBuildAccelerator/Core/Private/UbaHash.cpp
+++ b/Engine/Source/Programs/UnrealBuildAccelerator/Core/Private/UbaHash.cpp
@@ -79,9 +79,10 @@ namespace uba
 		blake3_hasher_init((blake3_hasher*)&hasher);
 	}
 
-	void CasKeyHasher::Update(const void* data, u64 bytes)
+	CasKeyHasher& CasKeyHasher::Update(const void* data, u64 bytes)
 	{
 		blake3_hasher_update((blake3_hasher*)&hasher, data, bytes);
+		return *this;
 	}
 
 	CasKey ToCasKey(const CasKeyHasher& hasher, bool compressed)
diff --git a/Engine/Source/Programs/UnrealBuildAccelerator/Core/Public/UbaHash.h b/Engine/Source/Programs/UnrealBuildAccelerator/Core/Public/UbaHash.h
--- a/Engine/Source/Programs/UnrealBuildAccelerator/Core/Public/UbaHash.h
+++ b/Engine/Source/Programs/UnrealBuildAccelerator/Core/Public/UbaHash.h
@@ -131,7 +131,7 @@ namespace uba
 	struct CasKeyHasher
 	{
 		CasKeyHasher();
-		void Update(const void* data, u64 bytes);
+		CasKeyHasher& Update(const void* data, u64 bytes);
 		u64 hasher[1912/sizeof(u64)];
 	};
 
diff --git a/Engine/Source/Programs/UnrealBuildAccelerator/Core/Public/UbaStringBuffer.h b/Engine/Source/Programs/UnrealBuildAccelerator/Core/Public/UbaStringBuffer.h
--- a/Engine/Source/Programs/UnrealBuildAccelerator/Core/Public/UbaStringBuffer.h
+++ b/Engine/Source/Programs/UnrealBuildAccelerator/Core/Public/UbaStringBuffer.h
@@ -96,6 +96,7 @@ namespace uba
 		StringView(const StringBufferBase& sb) : data(sb.data), count(sb.count) {}
 		StringView(const TString& str) : data(str.data()), count(u32(str.size())) {}
 		bool IsEmpty() const { return count == 0; }
+		bool StartsWith(const tchar* str, bool ignoreCase = true) const { return uba::StartsWith(data, str, ignoreCase); }
 		bool EndsWith(const tchar* value, bool ignoreCase = true) const { return uba::EndsWith(data, count, value, ignoreCase); }
 		bool Contains(const tchar* str, bool ignoreCase = true) const { return uba::Contains(data, str, ignoreCase); }
 		bool Equals(const tchar* str, bool ignoreCase = true) const { return uba::Equals(data, str, ignoreCase); }
