[UBA]
* Added HostRun command which can be used to send commands from remotes to host. command will return stdout

diff --git a/Engine/Source/Programs/UnrealBuildAccelerator/Common/Private/UbaProcess.cpp b/Engine/Source/Programs/UnrealBuildAccelerator/Common/Private/UbaProcess.cpp
--- a/Engine/Source/Programs/UnrealBuildAccelerator/Common/Private/UbaProcess.cpp
+++ b/Engine/Source/Programs/UnrealBuildAccelerator/Common/Private/UbaProcess.cpp
@@ -1254,6 +1254,14 @@ namespace uba
 		return true;
 	}
 
+	bool ProcessImpl::HandleHostRun(BinaryReader& reader, BinaryWriter& writer)
+	{
+		u16 size = reader.ReadU16();
+		BinaryReader reader2(reader.GetPositionData(), 0, size);
+		m_session.HostRun(reader2, writer);
+		return true;
+	}
+
 	bool ProcessImpl::CreateTempFile(BinaryReader& reader, ProcHandle nativeProcessHandle, const tchar* application)
 	{
 		StringKey key = reader.ReadStringKey();
diff --git a/Engine/Source/Programs/UnrealBuildAccelerator/Common/Private/UbaSession.cpp b/Engine/Source/Programs/UnrealBuildAccelerator/Common/Private/UbaSession.cpp
--- a/Engine/Source/Programs/UnrealBuildAccelerator/Common/Private/UbaSession.cpp
+++ b/Engine/Source/Programs/UnrealBuildAccelerator/Common/Private/UbaSession.cpp
@@ -2381,6 +2381,49 @@ namespace uba
 		return false;
 	}
 
+	bool Session::HostRun(BinaryReader& reader, BinaryWriter& writer)
+	{
+		#if !PLATFORM_WINDOWS
+
+		Vector<TString> args;
+		while (reader.GetLeft())
+			args.push_back(reader.ReadString());
+		bool success = false;
+
+		StringBuffer<> command;
+		for (auto& arg : args)
+		{
+			if (command.count)
+				command.Append(' ');
+			command.Append(arg);
+		}
+
+		char result[4096];
+		if (FILE* fp = popen(command.data, "r"))
+		{
+			char* dest = result;
+			errno = 0;
+			while (true)
+			{
+				if (!fgets(dest, sizeof(result) - (dest - result), fp))
+				{
+					success = errno == 0;
+					if (!success)
+						snprintf(result, sizeof(result), "fgets failed with command: %s", command.data);
+					break;
+				}
+				dest += strlen(dest);
+			}
+			pclose(fp);
+		}
+		else
+			snprintf(result, sizeof(result), "popen failed with command: %s", command.data);
+		writer.WriteBool(success);
+		writer.WriteString(result);
+		#endif
+		return true;
+	}
+
 	void Session::FileEntryAdded(StringKey fileNameKey, u64 lastWritten, u64 size)
 	{
 	}
diff --git a/Engine/Source/Programs/UnrealBuildAccelerator/Common/Private/UbaSessionClient.cpp b/Engine/Source/Programs/UnrealBuildAccelerator/Common/Private/UbaSessionClient.cpp
--- a/Engine/Source/Programs/UnrealBuildAccelerator/Common/Private/UbaSessionClient.cpp
+++ b/Engine/Source/Programs/UnrealBuildAccelerator/Common/Private/UbaSessionClient.cpp
@@ -741,7 +741,7 @@ namespace uba
 			if (findIt != writtenFiles.end())
 			{
 				auto insres = writtenFiles.try_emplace(toName);
-				UBA_ASSERTF(insres.second, TC("Moving written file %s to other written file %s."), fromName, toName);
+				UBA_ASSERTF(insres.second, TC("Moving written file %s to other written file %s. (%s)"), fromName, toName, msg.process.m_startInfo.description);
 				insres.first->second = findIt->second;
 				insres.first->second.owner = &msg.process;
 				writtenFiles.erase(findIt);
@@ -2013,6 +2013,35 @@ namespace uba
 		return true;
 	}
 
+	bool SessionClient::HostRun(BinaryReader& reader, BinaryWriter& writer)
+	{
+		const void* data = reader.GetPositionData();
+		u64 size = reader.GetLeft();
+
+		CasKey key = ToCasKey(CasKeyHasher().Update(data, size), false);
+
+		SCOPED_WRITE_LOCK(m_hostRunCacheLock, l);
+		auto insres = m_hostRunCache.try_emplace(key);
+		auto& buffer = insres.first->second;
+		if (!insres.second)
+		{
+			writer.WriteBytes(buffer.data(), buffer.size());
+			return true;
+		}
+
+		StackBinaryWriter<SendMaxSize> msgWriter;
+		NetworkMessage msg(m_client, ServiceId, SessionMessageType_HostRun, msgWriter);
+		msgWriter.WriteBytes(data, size);
+		BinaryReader msgReader(writer.GetData(), 0);
+		if (!msg.Send(msgReader, m_stats.customMsg))
+			return false;
+		writer.AllocWrite(msgReader.GetLeft());
+
+		buffer.resize(msgReader.GetLeft());
+		memcpy(buffer.data(), msgReader.GetPositionData(), buffer.size());
+		return true;
+	}
+
 	bool SessionClient::FlushWrittenFiles(ProcessImpl& process)
 	{
 		SCOPED_WRITE_LOCK(process.m_writtenFilesLock, lock);
diff --git a/Engine/Source/Programs/UnrealBuildAccelerator/Common/Private/UbaSessionServer.cpp b/Engine/Source/Programs/UnrealBuildAccelerator/Common/Private/UbaSessionServer.cpp
--- a/Engine/Source/Programs/UnrealBuildAccelerator/Common/Private/UbaSessionServer.cpp
+++ b/Engine/Source/Programs/UnrealBuildAccelerator/Common/Private/UbaSessionServer.cpp
@@ -2110,4 +2110,9 @@ namespace uba
 #endif
 		return true;
 	}
+
+	bool SessionServer::HandleHostRun(const ConnectionInfo& connectionInfo, BinaryReader& reader, BinaryWriter& writer)
+	{
+		return HostRun(reader, writer);
+	}
 }
diff --git a/Engine/Source/Programs/UnrealBuildAccelerator/Common/Public/UbaNetwork.h b/Engine/Source/Programs/UnrealBuildAccelerator/Common/Public/UbaNetwork.h
--- a/Engine/Source/Programs/UnrealBuildAccelerator/Common/Public/UbaNetwork.h
+++ b/Engine/Source/Programs/UnrealBuildAccelerator/Common/Public/UbaNetwork.h
@@ -88,6 +88,7 @@ namespace uba
 		UBA_SESSION_MESSAGE(Command) \
 		UBA_SESSION_MESSAGE(SHGetKnownFolderPath) \
 		UBA_SESSION_MESSAGE(DebugFileNotFoundError) \
+		UBA_SESSION_MESSAGE(HostRun) \
 
 	enum SessionMessageType : u8 
 	{
diff --git a/Engine/Source/Programs/UnrealBuildAccelerator/Common/Public/UbaProcess.h b/Engine/Source/Programs/UnrealBuildAccelerator/Common/Public/UbaProcess.h
--- a/Engine/Source/Programs/UnrealBuildAccelerator/Common/Public/UbaProcess.h
+++ b/Engine/Source/Programs/UnrealBuildAccelerator/Common/Public/UbaProcess.h
@@ -3,6 +3,7 @@
 #pragma once
 
 #include "UbaProcessStartInfoHolder.h"
+#include "UbaProtocol.h"
 #include "UbaSession.h"
 #include "UbaStats.h"
 
@@ -71,37 +72,13 @@ namespace uba
 
 		void ThreadRun(bool runningRemote, void* environment);
 		void ThreadExit();
-		bool HandleSpecialApplication();
+
 		bool HandleMessage(BinaryReader& reader, BinaryWriter& writer);
-		bool HandleInit(BinaryReader& reader, BinaryWriter& writer);
-		bool HandleCreateFile(BinaryReader& reader, BinaryWriter& writer);
-		bool HandleGetFullFileName(BinaryReader& reader, BinaryWriter& writer);
-		bool HandleGetLongPathName(BinaryReader& reader, BinaryWriter& writer);
-		bool HandleCloseFile(BinaryReader& reader, BinaryWriter& writer);
-		bool HandleDeleteFile(BinaryReader& reader, BinaryWriter& writer);
-		bool HandleCopyFile(BinaryReader& reader, BinaryWriter& writer);
-		bool HandleMoveFile(BinaryReader& reader, BinaryWriter& writer);
-		bool HandleChmod(BinaryReader& reader, BinaryWriter& writer);
-		bool HandleCreateDirectory(BinaryReader& reader, BinaryWriter& writer);
-		bool HandleRemoveDirectory(BinaryReader& reader, BinaryWriter& writer);
-		bool HandleListDirectory(BinaryReader& reader, BinaryWriter& writer);
-		bool HandleUpdateTables(BinaryReader& reader, BinaryWriter& writer);
-		bool HandleCreateProcess(BinaryReader& reader, BinaryWriter& writer);
-		bool HandleStartProcess(BinaryReader& reader, BinaryWriter& writer);
-		bool HandleExitChildProcess(BinaryReader& reader, BinaryWriter& writer);
-		bool HandleCreateTempFile(BinaryReader& reader, BinaryWriter& writer);
-		bool HandleOpenTempFile(BinaryReader& reader, BinaryWriter& writer);
-		bool HandleVirtualAllocFailed(BinaryReader& reader, BinaryWriter& writer);
-		bool HandleLog(BinaryReader& reader, BinaryWriter& writer);
-		bool HandleEchoOn(BinaryReader& reader, BinaryWriter& writer);
-		bool HandleInputDependencies(BinaryReader& reader, BinaryWriter& writer);
-		bool HandleExit(BinaryReader& reader, BinaryWriter& writer);
-		bool HandleFlushWrittenFiles(BinaryReader& reader, BinaryWriter& writer);
-		bool HandleUpdateEnvironment(BinaryReader& reader, BinaryWriter& writer);
-		bool HandleGetNextProcess(BinaryReader& reader, BinaryWriter& writer);
-		bool HandleCustom(BinaryReader& reader, BinaryWriter& writer);
-		bool HandleSHGetKnownFolderPath(BinaryReader& reader, BinaryWriter& writer);
-		bool HandleRpcCommunication(BinaryReader& reader, BinaryWriter& writer);
+		bool HandleSpecialApplication();
+
+		#define UBA_PROCESS_MESSAGE(M) bool Handle##M(BinaryReader& reader, BinaryWriter& writer);
+		UBA_PROCESS_MESSAGES
+		#undef UBA_PROCESS_MESSAGE
 
 		void LogLine(bool printInSession, TString&& line, LogEntryType logType);
 		bool CreateTempFile(BinaryReader& reader, ProcHandle nativeProcessHandle, const tchar* application);
diff --git a/Engine/Source/Programs/UnrealBuildAccelerator/Common/Public/UbaSession.h b/Engine/Source/Programs/UnrealBuildAccelerator/Common/Public/UbaSession.h
--- a/Engine/Source/Programs/UnrealBuildAccelerator/Common/Public/UbaSession.h
+++ b/Engine/Source/Programs/UnrealBuildAccelerator/Common/Public/UbaSession.h
@@ -140,6 +140,7 @@ namespace uba
 		virtual bool GetNextProcess(Process& process, bool& outNewProcess, NextProcessInfo& outNextProcess, u32 prevExitCode, BinaryReader& statsReader);
 		virtual bool CustomMessage(Process& process, BinaryReader& reader, BinaryWriter& writer);
 		virtual bool SHGetKnownFolderPath(Process& process, BinaryReader& reader, BinaryWriter& writer);
+		virtual bool HostRun(BinaryReader& reader, BinaryWriter& writer);
 		virtual void FileEntryAdded(StringKey fileNameKey, u64 lastWritten, u64 size);
 		virtual bool FlushWrittenFiles(ProcessImpl& process);
 		virtual bool UpdateEnvironment(ProcessImpl& process, const tchar* reason, bool resetStats);
diff --git a/Engine/Source/Programs/UnrealBuildAccelerator/Common/Public/UbaSessionClient.h b/Engine/Source/Programs/UnrealBuildAccelerator/Common/Public/UbaSessionClient.h
--- a/Engine/Source/Programs/UnrealBuildAccelerator/Common/Public/UbaSessionClient.h
+++ b/Engine/Source/Programs/UnrealBuildAccelerator/Common/Public/UbaSessionClient.h
@@ -66,6 +66,7 @@ namespace uba
 		virtual bool GetNextProcess(Process& process, bool& outNewProcess, NextProcessInfo& outNextProcess, u32 prevExitCode, BinaryReader& statsReader) override;
 		virtual bool CustomMessage(Process& process, BinaryReader& reader, BinaryWriter& writer) override;
 		virtual bool SHGetKnownFolderPath(Process& process, BinaryReader& reader, BinaryWriter& writer) override;
+		virtual bool HostRun(BinaryReader& reader, BinaryWriter& writer) override;
 		virtual bool FlushWrittenFiles(ProcessImpl& process) override;
 		virtual bool UpdateEnvironment(ProcessImpl& process, const tchar* reason, bool resetStats) override;
 		virtual bool LogLine(ProcessImpl& process, const tchar* line, LogEntryType logType) override;
@@ -160,5 +161,8 @@ namespace uba
 		Atomic<u64> m_bestPing;
 		u64 m_lastPing = 0;
 		u64 m_lastPingSendTime = 0;
+
+		UnorderedMap<CasKey, Vector<u8>> m_hostRunCache;
+		ReaderWriterLock m_hostRunCacheLock;
 	};
 }
diff --git a/Engine/Source/Programs/UnrealBuildAccelerator/Core/Public/UbaProtocol.h b/Engine/Source/Programs/UnrealBuildAccelerator/Core/Public/UbaProtocol.h
--- a/Engine/Source/Programs/UnrealBuildAccelerator/Core/Public/UbaProtocol.h
+++ b/Engine/Source/Programs/UnrealBuildAccelerator/Core/Public/UbaProtocol.h
@@ -36,6 +36,7 @@ namespace uba
 		UBA_PROCESS_MESSAGE(Custom) \
 		UBA_PROCESS_MESSAGE(SHGetKnownFolderPath) \
 		UBA_PROCESS_MESSAGE(RpcCommunication) \
+		UBA_PROCESS_MESSAGE(HostRun) \
 
 	enum MessageType : u8
 	{
diff --git a/Engine/Source/Programs/UnrealBuildAccelerator/Detours/Private/PosixOS/UbaDetoursFunctionsPosix.cpp b/Engine/Source/Programs/UnrealBuildAccelerator/Detours/Private/PosixOS/UbaDetoursFunctionsPosix.cpp
--- a/Engine/Source/Programs/UnrealBuildAccelerator/Detours/Private/PosixOS/UbaDetoursFunctionsPosix.cpp
+++ b/Engine/Source/Programs/UnrealBuildAccelerator/Detours/Private/PosixOS/UbaDetoursFunctionsPosix.cpp
@@ -794,7 +794,7 @@ UBA_EXPORT int UBA_WRAPPER(mkdir)(const char* path, mode_t mode)
 	g_directoryTable.ParseDirectoryTable(directoryTableSize);
 
 	errno = errorCode;
-	DEBUG_LOG_DETOURED("mkdir", "%ls -> %i (%u)", path, res, errorCode);
+	DEBUG_LOG_DETOURED("mkdir", "%s -> %i (%u)", path, res, errorCode);
 	return res ? 0 : -1;
 }
 
@@ -834,7 +834,7 @@ UBA_EXPORT int UBA_WRAPPER(rmdir)(const char* path)
 	g_directoryTable.ParseDirectoryTable(directoryTableSize);
 
 	errno = errorCode;
-	DEBUG_LOG_DETOURED("rmdir", "%ls -> %i (%u)", path, res, errorCode);
+	DEBUG_LOG_DETOURED("rmdir", "%s -> %i (%u)", path, res, errorCode);
 	return res ? 0 : -1;
 }
 
@@ -1338,7 +1338,9 @@ UBA_EXPORT int UBA_WRAPPER(access)(const char* pathname, int mode)
 	StringBuffer<> fixedPath;
 	if (!FixPath(fixedPath, pathname) || fixedPath.StartsWith(g_systemTemp.data) || fixedPath.StartsWith("/proc"))
 	{
-		return TRUE_WRAPPER(access)(pathname, mode);
+		auto res = TRUE_WRAPPER(access)(pathname, mode);
+		DEBUG_LOG_TRUE("access", "%s %i -> %i (%s)", pathname, mode, res, StrError(res, errno));
+		return res;
 	}
 
 	bool checkIfDir = false;
@@ -1601,7 +1603,7 @@ int Shared_DeleteFile(const char* funcName, const char* pathname)
 		errorCode = reader.ReadU32();
 		directoryTableSize = reader.ReadU32();
 		pcs.Leave();
-		DEBUG_LOG_PIPE(L"DeleteFile", L"%ls", lpFileName);
+		DEBUG_LOG_PIPE(L"DeleteFile", L"%s", lpFileName);
 	}
 	//DEBUG_LOG_DETOURED(L"DeleteFile", L"(%ls) -> %ls", pathname, ToString(result));
 
@@ -1637,6 +1639,57 @@ UBA_EXPORT int UBA_WRAPPER(unlink)(const char* pathname)
 
 thread_local int t_inVfork;
 
+bool ExecuteHostRun(StringBufferBase& out, const char* const* argv)
+{
+	TimerScope ts(g_stats.getFullFileName);
+	SCOPED_WRITE_LOCK(g_communicationLock, pcs);
+	BinaryWriter writer;
+	writer.WriteByte(MessageType_HostRun);
+	u16& size = *(u16*)writer.AllocWrite(2);
+	u64 pos = writer.GetPosition();
+	for (u32 i = 0; argv[i]; ++i)
+		writer.WriteString(argv[i]);
+	size = u16(writer.GetPosition() - pos);
+	writer.Flush();
+
+	BinaryReader reader;
+	bool success = reader.ReadBool();
+	reader.ReadString(out);
+	
+	if (out.count && out.data[out.count-1] == '\n')
+		out.Resize(out.count-1);
+
+	if (!success)
+		DEBUG_LOG("HOSTRUN FAILED: %s", out.data)
+	return success;
+}
+
+int SpawnEcho(char* str, pid_t* pid, const posix_spawn_file_actions_t* file_actions, const posix_spawnattr_t* attrp, char* const envp[])
+{
+	char cmd[] = "/bin/echo"; 
+	char* const argv2[] = { str, nullptr };
+	return TRUE_WRAPPER(posix_spawn)(pid, cmd, file_actions, attrp, argv2, envp);
+}
+
+void FlattenArgs(StringBufferBase& out, char* const argv[])
+{
+	if (!argv)
+		return;
+	for (u32 i = 0; argv[i]; ++i)
+	{
+		if (i != 0)
+			out.Append(' ');
+		out.Append(argv[i]);
+	}
+}
+
+void UnsupportedHostRun(char* const argv[], const char* msg)
+{
+	StringBuffer<4096> command;
+	FlattenArgs(command, argv);
+	UbaAssert(msg, __FILE__, __LINE__, command.data, 1999);
+}
+
 int shared_posix_spawn(pid_t* pid, const char* path, const posix_spawn_file_actions_t* file_actions, const posix_spawnattr_t* attrp, char* const argv[], char* const envp[])
 {
 	UBA_INIT_DETOUR(posix_spawn, pid, path, file_actions, attrp, argv, envp);
@@ -1691,8 +1744,8 @@ int shared_posix_spawn(pid_t* pid, const char* path, const posix_spawn_file_acti
 
 		BinaryReader reader;
 		processId = reader.ReadU32();
-		UBA_ASSERTF(processId > 0, "Process id was zero: path=%s", path);
-		if (!processId)
+
+		if (!processId) // Can happen if session client got disconnected from server session
 		{
 			errno = EINVAL; // This is not really correct but there is no errno for this failure
 			return -1;
@@ -1870,20 +1923,22 @@ UBA_EXPORT void* UBA_WRAPPER(dlopen)(const char* path, int mode)
 
 #if PLATFORM_MAC
 	StringBuffer<> tempBuf;
-	
-	auto originalPath = path;
-	if (StartsWith(path, "@rpath/"))
-	{
-		path += 7;
-		Set<TString> handled;
-		handled.insert(path);
-		const char* loaderPaths[] = { "/", 0 };
-		Shared_LoadLibrary(handled, path, loaderPaths, tempBuf);
-	}
-	else if (!StartsWith(path, "/System") && !StartsWith(path, "/usr/lib"))
+	if (g_runningRemote)
 	{
-		u64 nameLen = 0;
-		Rpc_GetFullFileName(path, nameLen, tempBuf, false);
+		auto originalPath = path;
+		if (StartsWith(path, "@rpath/"))
+		{
+			path += 7;
+			Set<TString> handled;
+			handled.insert(path);
+			const char* loaderPaths[] = { "/", 0 };
+			Shared_LoadLibrary(handled, path, loaderPaths, tempBuf);
+		}
+		else if (!StartsWith(path, "/System") && !StartsWith(path, "/usr/lib"))
+		{
+			u64 nameLen = 0;
+			Rpc_GetFullFileName(path, nameLen, tempBuf, false);
+		}
 	}
 #endif
 	void* res = TRUE_WRAPPER(dlopen)(path, mode);
