[UBA]
* Ignore invalid values for -maxcpu and -mulcpu arguments

Rather than aborting the default value is used instead, only printing a warning.

diff --git a/Engine/Source/Programs/UnrealBuildAccelerator/Agent/Private/UbaAgent.cpp b/Engine/Source/Programs/UnrealBuildAccelerator/Agent/Private/UbaAgent.cpp
--- a/Engine/Source/Programs/UnrealBuildAccelerator/Agent/Private/UbaAgent.cpp
+++ b/Engine/Source/Programs/UnrealBuildAccelerator/Agent/Private/UbaAgent.cpp
@@ -478,13 +478,21 @@ namespace uba
 			}
 			else if (name.Equals(TC("-maxcpu")))
 			{
+				u32 defaultValue = maxProcessCount;
 				if (!value.Parse(maxProcessCount))
-					return PrintHelp(TC("Invalid value for -maxcpu"));
+				{
+					LoggerWithWriter(g_consoleLogWriter, TC("")).Warning(TC("Invalid value for -maxcpu, ignoring!"));
+					maxProcessCount = defaultValue;
+				}
 			}
 			else if (name.Equals(TC("-mulcpu")))
 			{
+				float defaultValue = mulProcessValue;
 				if (!value.Parse(mulProcessValue))
-					return PrintHelp(TC("Invalid value for -mulcpu"));
+				{
+					LoggerWithWriter(g_consoleLogWriter, TC("")).Warning(TC("Invalid value for -mulcpu, ignoring!"));
+					mulProcessValue = defaultValue;
+				}
 			}
 			else if (name.Equals(TC("-maxcon")) || name.Equals(TC("-maxtcp")))
 			{
