[UBA]
* Fixed so MoveFile/CopyFile works with child processes. Now parent's "writtenfiles" is read from child process upon start
* Fixed some bugs related to spawning child processes with poorly formatted paths
* Fixed hang if parent process crashes after starting child process

diff --git a/Engine/Source/Programs/UnrealBuildAccelerator/Common/Private/UbaProcess.cpp b/Engine/Source/Programs/UnrealBuildAccelerator/Common/Private/UbaProcess.cpp
--- a/Engine/Source/Programs/UnrealBuildAccelerator/Common/Private/UbaProcess.cpp
+++ b/Engine/Source/Programs/UnrealBuildAccelerator/Common/Private/UbaProcess.cpp
@@ -153,10 +153,20 @@ namespace uba
 
 		m_startInfo = startInfo;
 
-		FixPathSeparators(m_startInfo.applicationStr.data());
 		FixPathSeparators(m_startInfo.workingDirStr.data());
 		FixPathSeparators(m_startInfo.logFileStr.data());
 
+		if (IsAbsolutePath(m_startInfo.application))
+		{
+			StringBuffer<256> temp2;
+			FixPath(m_startInfo.applicationStr.data(), nullptr, 0, temp2);
+			m_startInfo.applicationStr = temp2.data;
+			m_startInfo.application = m_startInfo.applicationStr.data();
+		}
+		else
+			FixPathSeparators(m_startInfo.applicationStr.data());
+
+
 		m_startInfo.Expand();
 
 		m_extractExports = Contains(m_startInfo.arguments, TC("/extractexports"), true);
@@ -535,8 +545,12 @@ namespace uba
 
 		// For some reason a parent can exit before a child. Need to figure out repro for this but I've seen it happen on ClangEditor win64
 		for (auto& child : m_childProcesses)
-			while (!((ProcessImpl*)child.m_process)->m_hasExited)
+		{
+			auto& childProcess = *(ProcessImpl*)child.m_process;
+			childProcess.m_waitForParent.Set();
+			while (!childProcess.m_hasExited)
 				Sleep(10);
+		}
 
 		UBA_ASSERT(!m_parentProcess || !m_parentProcess->m_hasExited);
 
@@ -681,6 +695,7 @@ namespace uba
 		InitResponse response;
 		m_messageSuccess = m_session.GetInitResponse(response, msg) && m_messageSuccess;
 		writer.WriteBool(m_echoOn);
+		writer.WriteBool(m_parentProcess != nullptr);
 		writer.WriteString(m_startInfo.application);
 		writer.WriteString(m_startInfo.workingDir);
 		writer.WriteU64(response.directoryTableHandle);
@@ -875,6 +890,31 @@ namespace uba
 		return true;
 	}
 
+	bool ProcessImpl::HandleGetParentWrittenFiles(BinaryReader& reader, BinaryWriter& writer)
+	{
+		// TODO: Recurse
+		UBA_ASSERT(m_parentProcess);
+		u32& writtenCount = *(u32*)writer.AllocWrite(4);
+		
+		//writer.WriteU32(u32(m_parentProcess->m_writtenFiles.size()));
+		u32 count = 0;
+		for (auto& kv : m_parentProcess->m_writtenFiles)
+		{
+			WrittenFile& wf = kv.second;
+			if (wf.mappingHandle.IsValid())
+				continue;
+			u64 fileSize;
+			if (!FileExists(m_session.m_logger, wf.name.c_str(), &fileSize))
+				continue;
+			writer.WriteStringKey(wf.key);
+			writer.WriteString(wf.name);
+			writer.WriteU64(fileSize);
+			++count;
+		}
+		writtenCount = count;
+		return true;
+	}
+
 	bool ProcessImpl::HandleCreateProcess(BinaryReader& reader, BinaryWriter& writer)
 	{
 		#if PLATFORM_LINUX
@@ -924,8 +964,6 @@ namespace uba
 				applicationStr.assign(fullCommandLine.data, u64(secondParamStart - fullCommandLine.data));
 		}
 
-		FixPathSeparators(applicationStr.data());
-
 		while (*commandLine == ' ')
 			++commandLine;
 
diff --git a/Engine/Source/Programs/UnrealBuildAccelerator/Common/Private/UbaSession.cpp b/Engine/Source/Programs/UnrealBuildAccelerator/Common/Private/UbaSession.cpp
--- a/Engine/Source/Programs/UnrealBuildAccelerator/Common/Private/UbaSession.cpp
+++ b/Engine/Source/Programs/UnrealBuildAccelerator/Common/Private/UbaSession.cpp
@@ -964,13 +964,15 @@ namespace uba
 			if (attr == INVALID_FILE_ATTRIBUTES)
 				return m_logger.Error("Failed to find file %s", applicationName);
 			//UBA_ASSERTF(false, TC("DIR NOT FOUND: %s"), applicationName);
+			#else
+			return m_logger.Error("Code path not implemented for linux!");
 			#endif
 
 			applicationName = temp;
 			attr = DefaultAttributes();
 
 			tchar* lastSlash = TStrrchr(temp, PathSeparator);
-			UBA_ASSERT(lastSlash);
+			UBA_ASSERTF(lastSlash, TC("No slash found in path %s"), temp);
 			u64 applicationDirLen = u64(lastSlash + 1 - temp);
 			memcpy(temp2, temp, applicationDirLen * sizeof(tchar));
 			applicationDir = temp2;
@@ -2042,13 +2044,6 @@ namespace uba
 			if (msg.process.m_extractExports && msg.process.m_startInfo.rules->ShouldExtractSymbols(file.name))
 				if (!ExtractSymbolsFromObjectFile(msg, name, fileSize))
 					return false;
-
-			// It might be that child processes need this file so we need to add the mapping
-			if (m_runningRemote && registerRealFile) // Note, if file is in file mapping it never touches the disk on remotes. no point adding it
-			{
-				FileExists(m_logger, writtenFile.name.c_str(), &fileSize); // Need to get file size from disk
-				AddFileMapping(writtenFile.key, name, writtenFile.name.c_str(), fileSize);
-			}
 		}
 
 		if (!msg.newName.IsEmpty())
diff --git a/Engine/Source/Programs/UnrealBuildAccelerator/Common/Private/UbaSessionClient.cpp b/Engine/Source/Programs/UnrealBuildAccelerator/Common/Private/UbaSessionClient.cpp
--- a/Engine/Source/Programs/UnrealBuildAccelerator/Common/Private/UbaSessionClient.cpp
+++ b/Engine/Source/Programs/UnrealBuildAccelerator/Common/Private/UbaSessionClient.cpp
@@ -278,7 +278,6 @@ namespace uba
 				startInfo.application = startInfo.applicationStr.c_str();
 			}
 			return true;
-
 		}
 
 		List<ModuleInfo> modules;
@@ -733,19 +732,25 @@ namespace uba
 	{
 		const tchar* fromName = msg.fromName.data;
 		const tchar* toName = msg.toName.data;
+		auto& process = msg.process;
 
 		{
-			SCOPED_WRITE_LOCK(msg.process.m_writtenFilesLock, lock);
-			auto& writtenFiles = msg.process.m_writtenFiles;
+			SCOPED_WRITE_LOCK(process.m_writtenFilesLock, lock);
+			auto& writtenFiles = process.m_writtenFiles;
 			auto findIt = writtenFiles.find(fromName);
 			if (findIt != writtenFiles.end())
 			{
 				auto insres = writtenFiles.try_emplace(toName);
-				UBA_ASSERTF(insres.second, TC("Moving written file %s to other written file %s. (%s)"), fromName, toName, msg.process.m_startInfo.description);
+				UBA_ASSERTF(insres.second, TC("Moving written file %s to other written file %s. (%s)"), fromName, toName, process.m_startInfo.description);
 				insres.first->second = findIt->second;
-				insres.first->second.owner = &msg.process;
+				insres.first->second.key = msg.toKey;
+				insres.first->second.owner = &process;
 				writtenFiles.erase(findIt);
 			}
+			else
+			{
+				// TODO: Need to tell server 
+			}
 		}
 
 		bool sendMove = true;
@@ -769,6 +774,8 @@ namespace uba
 			return true;
 		}
 
+		// TODO: This should be done by server?
+
 		out.result = uba::MoveFileExW(fromName, toName, 0);
 		out.errorCode = GetLastError();
 
diff --git a/Engine/Source/Programs/UnrealBuildAccelerator/Core/Public/UbaProtocol.h b/Engine/Source/Programs/UnrealBuildAccelerator/Core/Public/UbaProtocol.h
--- a/Engine/Source/Programs/UnrealBuildAccelerator/Core/Public/UbaProtocol.h
+++ b/Engine/Source/Programs/UnrealBuildAccelerator/Core/Public/UbaProtocol.h
@@ -20,6 +20,7 @@ namespace uba
 		UBA_PROCESS_MESSAGE(RemoveDirectory) \
 		UBA_PROCESS_MESSAGE(ListDirectory) \
 		UBA_PROCESS_MESSAGE(UpdateTables) \
+		UBA_PROCESS_MESSAGE(GetParentWrittenFiles) \
 		UBA_PROCESS_MESSAGE(CreateProcess) \
 		UBA_PROCESS_MESSAGE(StartProcess) \
 		UBA_PROCESS_MESSAGE(ExitChildProcess) \
diff --git a/Engine/Source/Programs/UnrealBuildAccelerator/Detours/Private/PosixOS/UbaDetoursFunctionsPosix.cpp b/Engine/Source/Programs/UnrealBuildAccelerator/Detours/Private/PosixOS/UbaDetoursFunctionsPosix.cpp
--- a/Engine/Source/Programs/UnrealBuildAccelerator/Detours/Private/PosixOS/UbaDetoursFunctionsPosix.cpp
+++ b/Engine/Source/Programs/UnrealBuildAccelerator/Detours/Private/PosixOS/UbaDetoursFunctionsPosix.cpp
@@ -2212,6 +2212,10 @@ namespace uba
 			}
 			DEBUG_LOG("");
 #endif
+			DEBUG_LOG("SystemTemp: %s", g_systemTemp.data);
+			if (g_runningRemote)
+				DEBUG_LOG("RunningRemote: true");
+			DEBUG_LOG("");
 		}
 		#endif
 
@@ -2246,6 +2250,7 @@ namespace uba
 			BinaryReader reader;
 
 			bool echoOn = reader.ReadBool();
+			g_isChild = reader.ReadBool();
 
 			reader.ReadString(g_virtualApplication);
 			reader.ReadString(g_virtualWorkingDir);
@@ -2284,6 +2289,9 @@ namespace uba
 		UBA_ASSERTF(dirTableMem != MAP_FAILED, "mmap for dirtable mem failed (%s)", strerror(errno));
 		g_directoryTable.Init(dirTableMem, directoryTableCount, directoryTableSize);
 
+		if (g_runningRemote && g_isChild)
+			Rpc_GetParentWrittenFiles();
+
 		g_isDetouring = true;
 		DEBUG_LOG("Detouring enabled");
 
diff --git a/Engine/Source/Programs/UnrealBuildAccelerator/Detours/Private/UbaDetoursFileMappingTable.cpp b/Engine/Source/Programs/UnrealBuildAccelerator/Detours/Private/UbaDetoursFileMappingTable.cpp
--- a/Engine/Source/Programs/UnrealBuildAccelerator/Detours/Private/UbaDetoursFileMappingTable.cpp
+++ b/Engine/Source/Programs/UnrealBuildAccelerator/Detours/Private/UbaDetoursFileMappingTable.cpp
@@ -198,6 +198,29 @@ namespace uba
 		g_mappedFileTable.Parse(fileMappingTableSize);
 	}
 
+	void Rpc_GetParentWrittenFiles()
+	{
+		TimerScope ts(g_stats.updateTables);
+		SCOPED_WRITE_LOCK(g_communicationLock, pcs);
+		BinaryWriter writer;
+		writer.WriteByte(MessageType_GetParentWrittenFiles);
+		writer.Flush();
+		BinaryReader reader;
+
+		u32 count = reader.ReadU32();
+		while (count--)
+		{
+			StringKey key = reader.ReadStringKey();
+			auto insres = g_mappedFileTable.m_lookup.try_emplace(key);
+			FileInfo& info = insres.first->second;
+			StringBuffer<> name;
+			reader.ReadString(name);
+			u64 fileSize = reader.ReadU64();
+			info.name = g_mappedFileTable.m_memoryBlock.Strdup(name.data);
+			info.size = fileSize;
+		}
+	}
+
 	u32 Rpc_GetEntryOffset(const StringKey& entryNameKey, const tchar* entryName, u64 entryNameLen, bool checkIfDir)
 	{
 		u32 dirTableOffset = ~u32(0);
@@ -251,9 +274,13 @@ namespace uba
 	void Rpc_GetFullFileName(const tchar*& path, u64& pathLen, StringBufferBase& tempBuf, bool useVirtualName, const tchar* const* loaderPaths)
 	{
 		StringKey fileNameKey;
+		StringBuffer<> temp2;
 		if (IsAbsolutePath(path))
 		{
 			FixPath(tempBuf, path);
+			temp2.Append(tempBuf);
+			path = temp2.data;
+
 			if (CaseInsensitiveFs)
 				tempBuf.MakeLower();
 			fileNameKey = ToStringKey(tempBuf);
diff --git a/Engine/Source/Programs/UnrealBuildAccelerator/Detours/Private/UbaDetoursFileMappingTable.h b/Engine/Source/Programs/UnrealBuildAccelerator/Detours/Private/UbaDetoursFileMappingTable.h
--- a/Engine/Source/Programs/UnrealBuildAccelerator/Detours/Private/UbaDetoursFileMappingTable.h
+++ b/Engine/Source/Programs/UnrealBuildAccelerator/Detours/Private/UbaDetoursFileMappingTable.h
@@ -65,6 +65,7 @@ namespace uba
 	void Rpc_UpdateTables();
 	u32  Rpc_GetEntryOffset(const StringKey& entryNameKey, const tchar* entryName, u64 entryNameLen, bool checkIfDir = false);
 	void Rpc_GetFullFileName(const tchar*& path, u64& pathLen, StringBufferBase& tempBuf, bool useVirtualName, const tchar* const* loaderPaths = nullptr);
+	void Rpc_GetParentWrittenFiles();
 
 	struct DirHash
 	{
diff --git a/Engine/Source/Programs/UnrealBuildAccelerator/Detours/Private/Windows/UbaDetoursFunctionsWin.cpp b/Engine/Source/Programs/UnrealBuildAccelerator/Detours/Private/Windows/UbaDetoursFunctionsWin.cpp
--- a/Engine/Source/Programs/UnrealBuildAccelerator/Detours/Private/Windows/UbaDetoursFunctionsWin.cpp
+++ b/Engine/Source/Programs/UnrealBuildAccelerator/Detours/Private/Windows/UbaDetoursFunctionsWin.cpp
@@ -836,6 +836,7 @@ void Init(const DetoursPayload& payload, u64 startTime)
 		BinaryReader reader;
 
 		g_echoOn = reader.ReadBool();
+		g_isChild = reader.ReadBool();
 
 		reader.ReadString(applicationBuffer);
 		reader.ReadString(workingDirBuffer);
@@ -932,6 +933,9 @@ void Init(const DetoursPayload& payload, u64 startTime)
 		g_directoryTable.Init(directoryTableMem, directoryTableCount, directoryTableSize);
 	}
 
+	if (g_runningRemote && g_isChild)
+		Rpc_GetParentWrittenFiles();
+
 	g_stats.attach.time += GetTime() - startTime;
 	g_stats.attach.count = 1;
 
diff --git a/Engine/Source/Programs/UnrealBuildAccelerator/Test/Private/UbaTestBasics.h b/Engine/Source/Programs/UnrealBuildAccelerator/Test/Private/UbaTestBasics.h
--- a/Engine/Source/Programs/UnrealBuildAccelerator/Test/Private/UbaTestBasics.h
+++ b/Engine/Source/Programs/UnrealBuildAccelerator/Test/Private/UbaTestBasics.h
@@ -484,10 +484,11 @@ namespace uba
 		GetDirectoryOfCurrentModule(logger, path);
 		path.EnsureEndsWithSlash().Append(TC("UbaTestApp.exe"));
 		bool importKernel = false;
-		FindImports(path.data, [&](const tchar* import, bool isKnown)
+		StringBuffer<> error;
+		FindImports(path.data, [&](const tchar* import, bool isKnown, const char* const* importLoaderPaths)
 		{
 			importKernel = isKnown && Contains(import, TC("KERNEL32.dll"));
-		});
+		}, error);
 		if (!importKernel)
 			return logger.Error(TC("Failed to find Kernel32 as import"));
 #elif PLATFORM_MAC
diff --git a/Engine/Source/Programs/UnrealBuildAccelerator/Test/Private/UbaTestSession.h b/Engine/Source/Programs/UnrealBuildAccelerator/Test/Private/UbaTestSession.h
--- a/Engine/Source/Programs/UnrealBuildAccelerator/Test/Private/UbaTestSession.h
+++ b/Engine/Source/Programs/UnrealBuildAccelerator/Test/Private/UbaTestSession.h
@@ -137,15 +137,33 @@ namespace uba
 			return false;
 
 		{
-			StringBuffer<MaxPath> dir1;
-			dir1.Append(workingDir).Append(TC("Dir1"));
-			if (!CreateDirectoryW(dir1.data))
-				return logger.Error(TC("Failed to create dir %s"), dir1.data);
+			StringBuffer<MaxPath> dir;
+			dir.Append(workingDir).Append(TC("Dir1"));
+			if (!CreateDirectoryW(dir.data))
+				return logger.Error(TC("Failed to create dir %s"), dir.data);
+
+			dir.Clear().Append(workingDir).Append(TC("Dir2"));
+			if (!CreateDirectoryW(dir.data))
+				return logger.Error(TC("Failed to create dir %s"), dir.data);
+			dir.EnsureEndsWithSlash().Append(TC("Dir3"));
+			if (!CreateDirectoryW(dir.data))
+				return logger.Error(TC("Failed to create dir %s"), dir.data);
+			dir.EnsureEndsWithSlash().Append(TC("Dir4"));
+			if (!CreateDirectoryW(dir.data))
+				return logger.Error(TC("Failed to create dir %s"), dir.data);
+			dir.EnsureEndsWithSlash().Append(TC("Dir5"));
+			if (!CreateDirectoryW(dir.data))
+				return logger.Error(TC("Failed to create dir %s"), dir.data);
 		}
 
 		ProcessStartInfo processInfo;
 		processInfo.application = testApp.data;
 		processInfo.workingDir = workingDir;
+		processInfo.logLineFunc = [](void* userData, const tchar* line, u32 length, LogEntryType type)
+			{
+				LoggerWithWriter(g_consoleLogWriter, TC("")).Info(line);
+			};
+
 		ProcessHandle process = runProcess(processInfo);
 		if (!process.WaitForExit(100000))
 			return logger.Error(TC("UbaTestApp did not exit in 10 seconds"));
diff --git a/Engine/Source/Programs/UnrealBuildAccelerator/TestApp/Private/PosixOS/UbaTestAppPosix.cpp b/Engine/Source/Programs/UnrealBuildAccelerator/TestApp/Private/PosixOS/UbaTestAppPosix.cpp
--- a/Engine/Source/Programs/UnrealBuildAccelerator/TestApp/Private/PosixOS/UbaTestAppPosix.cpp
+++ b/Engine/Source/Programs/UnrealBuildAccelerator/TestApp/Private/PosixOS/UbaTestAppPosix.cpp
@@ -6,6 +6,8 @@
 #include <sys/stat.h>
 #include <stdio.h>
 #include <dlfcn.h>
+#include <spawn.h>
+#include <wait.h>
 
 #if PLATFORM_LINUX
 #include <link.h>
@@ -127,9 +129,18 @@ int main(int argc, char* argv[])
 
 		struct stat attrD1;
 		if (stat("Dir1", &attrD1) == -1)
-			return LogError("stat for FileW failed");
+			return LogError("stat for Dir1 failed");
 		if (S_ISREG(attrD1.st_mode) != 0)
-			return LogError("stat for FileW did not return directory");
+			return LogError("stat for Dir1 did not return directory");
+		
+		if (mkdir("Dir2/Dir3", S_IRUSR | S_IWUSR) == 0)
+			return LogError("mkdir for dir2 did not fail even though it exists");
+		if (errno != EEXIST)
+			return LogError("mkdir for dir2 did not return error that it exists");
+
+		struct stat attrD2;
+		if (stat("Dir2/Dir3/Dir4/Dir5", &attrD2) == -1)
+			return LogError("stat for Dir2/Dir3 failed");
 
 		struct stat attr3;
 		if (stat("/usr", &attr3) == -1)
@@ -170,6 +181,41 @@ int main(int argc, char* argv[])
 
 		if (stat("FooDir", &attrFoo) == 0)
 			return LogError("stat for 'FooDir' failed to not find removed directory");
+
+		{
+			char execPath[1024];
+			if (readlink("/proc/self/exe", execPath, sizeof(execPath)) == -1)
+				return LogError("readlink failed while getting executable path");
+
+			pid_t childPid = 0;
+			const char* args[] { execPath, "-child", nullptr };
+			if (posix_spawn(&childPid, execPath, nullptr, nullptr, (char**)args, environ) != 0)
+				return LogError("posix_spawn failed");
+			int status = 0;
+			do
+			{
+				if (waitpid(childPid, &status, WUNTRACED | WCONTINUED) == -1)
+					return LogError("waitpid on child process failed (pid %i)", childPid);
+				if (WIFSIGNALED(status))
+					return LogError("Child process killed by signal %d", WTERMSIG(status));
+				if (WIFSTOPPED(status))
+					return LogError("Child process stopped by signal %d", WSTOPSIG(status));
+				if (WIFCONTINUED(status))
+					return LogError("Child process continued");
+			} while (!WIFEXITED(status));
+			if (WEXITSTATUS(status) != 0)
+				return LogError("Child process failed");
+		}
+
+		return 0;
+	}
+	else if (strcmp(argv[1], "-child") == 0)
+	{
+		struct stat attr;
+		if (stat("FileW2", &attr) != 0)
+			return LogError("stat for 'FileW' in child process failed");
+		if (stat("FileW", &attr) != -1)
+			return LogError("stat for 'FileW' in child process failed");
 		return 0;
 	}
 	else if (strncmp(argv[1], "-file=", 6) == 0)
diff --git a/Engine/Source/Programs/UnrealBuildAccelerator/TestApp/Private/Windows/UbaTestAppWin.cpp b/Engine/Source/Programs/UnrealBuildAccelerator/TestApp/Private/Windows/UbaTestAppWin.cpp
--- a/Engine/Source/Programs/UnrealBuildAccelerator/TestApp/Private/Windows/UbaTestAppWin.cpp
+++ b/Engine/Source/Programs/UnrealBuildAccelerator/TestApp/Private/Windows/UbaTestAppWin.cpp
@@ -150,7 +150,46 @@ int wmain(int argc, wchar_t* argv[])
 
 			if (GetFileAttributesW(L"DirA") != INVALID_FILE_ATTRIBUTES)
 				return LogError(L"Found attributes of deleted directory");
+
+			if (CreateDirectoryW(L"Dir2\\Dir3", NULL))
+				return LogError(L"Should not succeed creation directory that exists");
+			if (GetLastError() != ERROR_ALREADY_EXISTS)
+				return LogError(L"Did not get correct error when failing to create existing directory");
+			if (GetFileAttributesW(L"Dir2\\Dir3\\Dir4\\Dir5") == INVALID_FILE_ATTRIBUTES)
+				return LogError(L"Failed to get attributes of directory");
+		}
+
+
+		{
+			STARTUPINFOW si;
+			memset(&si, 0, sizeof(si));
+			PROCESS_INFORMATION pi;
+			memset(&pi, 0, sizeof(pi));
+			wchar_t arg[1024];
+			wcscpy_s(arg, 1024, argv[0]);
+			wcscat_s(arg, 1024, L" -child");
+			if (!CreateProcessW(nullptr, arg, nullptr, nullptr, TRUE, 0, nullptr, nullptr, &si, &pi))
+				return LogError(L"Failed to create child process");
+			CloseHandle(pi.hThread);
+			
+			if (WaitForSingleObject(pi.hProcess, 10000) != WAIT_OBJECT_0)
+				return LogError(L"Failed waiting for child process");
+
+			DWORD exitCode;
+			if (!GetExitCodeProcess(pi.hProcess, &exitCode) || exitCode)
+				return LogError(L"Child process failed");
+			CloseHandle(pi.hProcess);
 		}
+
+	}
+	else if (wcscmp(argv[1], L"-child") == 0)
+	{
+		if (GetFileAttributes(L"FileW2") == INVALID_FILE_ATTRIBUTES)
+			return LogError(L"Child process could not get attributes of FileW2");
+		if (GetFileAttributes(L"FileWF") == INVALID_FILE_ATTRIBUTES)
+			return LogError(L"Child process could not get attributes of FileWF");
+		if (GetFileAttributes(L"FileW") != INVALID_FILE_ATTRIBUTES)
+			return LogError(L"Child process found FileW which should not exist anymore");
 	}
 	else if (wcscmp(argv[1], L"-reuse") == 0)
 	{
