[UBA]
* Added so agent can use environment variables in parameters in to maxcpu and mulcpu

diff --git a/Engine/Source/Programs/UnrealBuildAccelerator/Agent/Private/UbaAgent.cpp b/Engine/Source/Programs/UnrealBuildAccelerator/Agent/Private/UbaAgent.cpp
--- a/Engine/Source/Programs/UnrealBuildAccelerator/Agent/Private/UbaAgent.cpp
+++ b/Engine/Source/Programs/UnrealBuildAccelerator/Agent/Private/UbaAgent.cpp
@@ -478,6 +478,8 @@ namespace uba
 			}
 			else if (name.Equals(TC("-maxcpu")))
 			{
+				if (!ExpandEnvironmentVariables(value))
+					return false;
 				u32 defaultValue = maxProcessCount;
 				if (!value.Parse(maxProcessCount))
 				{
@@ -487,6 +489,8 @@ namespace uba
 			}
 			else if (name.Equals(TC("-mulcpu")))
 			{
+				if (!ExpandEnvironmentVariables(value))
+					return false;
 				float defaultValue = mulProcessValue;
 				if (!value.Parse(mulProcessValue))
 				{
