[UBA]
* Fixed unit test compile error on mac

diff --git a/Engine/Source/Programs/UnrealBuildAccelerator/TestApp/Private/PosixOS/UbaTestAppPosix.cpp b/Engine/Source/Programs/UnrealBuildAccelerator/TestApp/Private/PosixOS/UbaTestAppPosix.cpp
--- a/Engine/Source/Programs/UnrealBuildAccelerator/TestApp/Private/PosixOS/UbaTestAppPosix.cpp
+++ b/Engine/Source/Programs/UnrealBuildAccelerator/TestApp/Private/PosixOS/UbaTestAppPosix.cpp
@@ -5,15 +5,17 @@
 #include <unistd.h>
 #include <sys/stat.h>
 #include <stdio.h>
+#include <dirent.h>
 #include <dlfcn.h>
 #include <spawn.h>
-#include <wait.h>
 
 #if PLATFORM_LINUX
 #include <link.h>
 #include <dlfcn.h>
+#include <wait.h>
 #else
 #include <mach-o/dyld.h>
+extern char** environ;
 #endif
 
 
@@ -182,6 +184,17 @@ int main(int argc, char* argv[])
 		if (stat("FooDir", &attrFoo) == 0)
 			return LogError("stat for 'FooDir' failed to not find removed directory");
 
+		auto dir = opendir(".");
+		if (!dir)
+			return LogError("opendir failed");
+		while (true)
+		{
+			dirent* ent = readdir(dir);
+			if (!ent)
+				break;
+		}
+		closedir(dir);
+
 		{
 			char execPath[1024];
 			if (readlink("/proc/self/exe", execPath, sizeof(execPath)) == -1)
