[UBA]
* Fixed debug log entry that could cause bufferrun.. now capped to 100 characters
* Changed so FatalAssert does a signal kill instead of exit

diff --git a/Engine/Source/Programs/UnrealBuildAccelerator/Common/Private/UbaLogger.cpp b/Engine/Source/Programs/UnrealBuildAccelerator/Common/Private/UbaLogger.cpp
--- a/Engine/Source/Programs/UnrealBuildAccelerator/Common/Private/UbaLogger.cpp
+++ b/Engine/Source/Programs/UnrealBuildAccelerator/Common/Private/UbaLogger.cpp
@@ -56,7 +56,7 @@ namespace uba
 
 		ExitProcess(terminateCode);
 #else
-		exit(-1);
+		_Exit(-1);
 #endif
 	}
 
@@ -83,7 +83,7 @@ namespace uba
 #else
 		printf(TC("FATAL ERROR %u: %s\n"), code, buffer);
 		fflush(stdout);
-		exit(int(code));
+		kill(getpid(), SIGKILL);
 #endif
 	}
 
diff --git a/Engine/Source/Programs/UnrealBuildAccelerator/Common/Private/UbaProcess.cpp b/Engine/Source/Programs/UnrealBuildAccelerator/Common/Private/UbaProcess.cpp
--- a/Engine/Source/Programs/UnrealBuildAccelerator/Common/Private/UbaProcess.cpp
+++ b/Engine/Source/Programs/UnrealBuildAccelerator/Common/Private/UbaProcess.cpp
@@ -1913,13 +1913,13 @@ namespace uba
 			pipeGuard0.Cancel();
 
 			#if UBA_DEBUG_TRACK_PROCESS
-			g_debugLogger.Info(TC("CreateRealProcess  (%u) %s %s\n"), m_id, m_realApplication.c_str(), m_startInfo.arguments);
+			g_debugLogger.Info(TC("CreateRealProcess  (%u) %s %.100s\n"), m_id, m_realApplication.c_str(), m_startInfo.arguments);
 			#endif
 		}
 		else
 		{
 			#if UBA_DEBUG_TRACK_PROCESS
-			g_debugLogger.Info(TC("WaitingForParent (%u) %s\n"), m_id, m_realApplication.c_str());
+			g_debugLogger.Info(TC("WaitingForParent (%u) %.100s\n"), m_id, m_realApplication.c_str());
 			#endif
 
 			//logger.Info("Waiting for parent");
