[UBA]
* Fixed so visualizer shows *bps (bits) for networking to reduce confusing of if it is byte or bit

diff --git a/Engine/Source/Programs/UnrealBuildAccelerator/Visualizer/Private/UbaVisualizer.cpp b/Engine/Source/Programs/UnrealBuildAccelerator/Visualizer/Private/UbaVisualizer.cpp
--- a/Engine/Source/Programs/UnrealBuildAccelerator/Visualizer/Private/UbaVisualizer.cpp
+++ b/Engine/Source/Programs/UnrealBuildAccelerator/Visualizer/Private/UbaVisualizer.cpp
@@ -1926,8 +1926,8 @@ namespace uba
 			logger.AddSpace(3);
 			logger.SetColor(m_cpuColor).Info(L"  Cpu: %.1f%%", m_stats.cpuLoad * 100.0f);
 			logger.SetColor(m_memColor).Info(L"  Mem: %ls/%ls", BytesToText(m_stats.memTotal - m_stats.memAvail).str, BytesToText(m_stats.memTotal).str);
-			logger.SetColor(m_recvColor).Info(L"  Recv: %ls/s", BytesToText(m_stats.recvBytesPerSecond).str);
-			logger.SetColor(m_sendColor).Info(L"  Send: %ls/s", BytesToText(m_stats.sendBytesPerSecond).str);
+			logger.SetColor(m_recvColor).Info(L"  Recv: %lsps", BytesToText(m_stats.recvBytesPerSecond*8).str);
+			logger.SetColor(m_sendColor).Info(L"  Send: %lsps", BytesToText(m_stats.sendBytesPerSecond*8).str);
 			if (m_stats.ping)
 				logger.Info(L"  Ping: %ls", TimeToText(m_stats.ping, false, m_traceView.frequency).str);
 			logger.AddSpace(3);
@@ -2297,8 +2297,8 @@ namespace uba
 					recvPerS = u64((u.recv - session.prevRecv) / duration);
 				}
 				drawText(L"ClientId: %u  TcpCount: %u", session.clientUid.data1, u.connectionCount);
-				drawText(L"Recv: %ls (%s/s)", BytesToText(u.recv), BytesToText(recvPerS));
-				drawText(L"Send: %ls (%s/s)", BytesToText(u.send), BytesToText(sendPerS));
+				drawText(L"Recv: %ls (%sps)", BytesToText(u.recv), BytesToText(recvPerS*8));
+				drawText(L"Send: %ls (%sps)", BytesToText(u.send), BytesToText(sendPerS*8));
 			}
 
 			if (session.disconnectTime == ~u64(0))
@@ -2363,8 +2363,8 @@ namespace uba
 						sendPerS = u64((u.send - session.prevSend) / duration);
 						recvPerS = u64((u.recv - session.prevRecv) / duration);
 					}
-					drawText(L"Recv: %ls (%s/s)", BytesToText(u.recv), BytesToText(recvPerS));
-					drawText(L"Send: %ls (%s/s)", BytesToText(u.send), BytesToText(sendPerS));
+					drawText(L"Recv: %ls (%sps)", BytesToText(u.recv), BytesToText(recvPerS));
+					drawText(L"Send: %ls (%sps)", BytesToText(u.send), BytesToText(sendPerS));
 				}
 			}
 		}
