[UBA]
* LOTS of fixes around posix detouring path. With this code it is now possible to get through compile shaders remotely on mac (even though the output is causing some errors).

diff --git a/Engine/Source/Programs/UnrealBuildAccelerator/Detours/Private/PosixOS/UbaDetoursFunctionsPosix.cpp b/Engine/Source/Programs/UnrealBuildAccelerator/Detours/Private/PosixOS/UbaDetoursFunctionsPosix.cpp
--- a/Engine/Source/Programs/UnrealBuildAccelerator/Detours/Private/PosixOS/UbaDetoursFunctionsPosix.cpp
+++ b/Engine/Source/Programs/UnrealBuildAccelerator/Detours/Private/PosixOS/UbaDetoursFunctionsPosix.cpp
@@ -48,6 +48,7 @@ using namespace uba;
 	DETOURED_FUNCTION(getcwd) \
 	DETOURED_FUNCTION(getenv) \
 	DETOURED_FUNCTION(setenv) \
+	DETOURED_FUNCTION(unsetenv) \
 	DETOURED_FUNCTION(realpath) \
 	DETOURED_FUNCTION(readlink) \
 	DETOURED_FUNCTION(readlinkat) \
@@ -95,10 +96,15 @@ using namespace uba;
 	DETOURED_FUNCTION(wait4) \
 	DETOURED_FUNCTION(system) \
 	DETOURED_FUNCTION(dlopen) \
+	DETOURED_FUNCTION(dladdr) \
 	DETOURED_FUNCTION(execv) \
 	DETOURED_FUNCTION(execve) \
 	DETOURED_FUNCTION(execvp) \
+	DETOURED_FUNCTION(execl) \
+	DETOURED_FUNCTION(execle) \
+	DETOURED_FUNCTION(execlp) \
 	DETOURED_FUNCTION(fork) \
+	DETOURED_FUNCTION(vfork) \
 	DETOURED_FUNCTION(popen) \
 	DETOURED_FUNCTION(exit) \
 	DETOURED_FUNCTION(_exit) \
@@ -122,6 +128,7 @@ using namespace uba;
 #define DETOURED_FUNCTION_LINUX \
 	DETOURED_FUNCTION(get_current_dir_name) \
 	DETOURED_FUNCTION(fopen64) \
+	DETOURED_FUNCTION(secure_getenv) \
 	DETOURED_FUNCTION(fcntl) \
 	DETOURED_FUNCTION(__xstat) \
 	DETOURED_FUNCTION(__xstat64) \
@@ -131,12 +138,12 @@ using namespace uba;
 	DETOURED_FUNCTION(__fxstatat64) \
 	DETOURED_FUNCTION(open64) \
 	DETOURED_FUNCTION(fstatat64) \
-	DETOURED_FUNCTION(vfork) \
 
 #else
 #define DETOURED_FUNCTION_LINUX
 #define DETOURED_FUNCTION_MACOS \
 	DETOURED_FUNCTION(_NSGetExecutablePath) \
+	DETOURED_FUNCTION(execvP) \
 
 #endif
 
@@ -286,24 +293,34 @@ namespace uba
 	}
 }
 
+bool CanDetour2(const StringView& file)
+{
+	return g_isDetouring && !t_disallowDetour
+		&& !file.StartsWith("/dev/")
+		&& !file.StartsWith("/etc/")
+		&& !file.StartsWith(g_systemTemp.data)
+		&& !file.Contains("/.cache/clang/")
+		;
+}
+
 // Shared functions
 
 template<typename TrueOpen>
 int Shared_open(const char* funcName, const char* file, int flags, int mode, const TrueOpen& trueOpen)
 {
-	if (!g_isDetouring || t_disallowDetour || strncmp(file, "/dev/", 5) == 0 || strncmp(file, "/etc/", 5) == 0)
+	StringBuffer<> fileName;
+	FixPath(fileName, file);
+
+	if (!CanDetour2(fileName))
 	{
 		int res = trueOpen(file, flags, mode);
 		DEBUG_LOG_TRUE(funcName, "NODETOUR (%s) -> %i", file, res);
 		return res;
 	}
-	StringBuffer<> fileName;
-	FixPath(fileName, file);
 
 	StringKey fileNameKey = ToFilenameKey(fileName);
 
 	bool keepInMemory = false;
-	bool isSystemOrTempFile = StartsWith(file, g_systemTemp.data);
 
 	const char* realFileName = fileName.data;
 	u64 size = InvalidValue;
@@ -318,7 +335,7 @@ int Shared_open(const char* funcName, const char* file, int flags, int mode, con
 
 	if (g_allowDirectoryCache)
 	{
-		if (!isWrite && !isSystemOrTempFile) // We need to skip SystemTemp.. lots of stuff going on there.
+		if (!isWrite) // We need to skip SystemTemp.. lots of stuff going on there.
 		{
 			dirTableOffset = Rpc_GetEntryOffset(fileNameKey, fileName.data, fileName.count, false);
 			bool allowEarlyOut = true;
@@ -346,6 +363,7 @@ int Shared_open(const char* funcName, const char* file, int flags, int mode, con
 				g_directoryTable.GetEntryInformation(entryInfo, dirTableOffset);
 				if (entryInfo.attributes == 0)
 				{
+					DEBUG_LOG_DETOURED(funcName, "DELETED (%s) (%s) -> -1", fileName.data, KeyToString(fileNameKey).data);
 					errno = ENOENT;
 					return -1;
 				}
@@ -363,7 +381,7 @@ int Shared_open(const char* funcName, const char* file, int flags, int mode, con
 	{
 		info.originalName = g_memoryBlock.Strdup(fileName.data);
 		info.name = info.originalName;
-		if (!keepInMemory && !isSystemOrTempFile)
+		if (!keepInMemory)
 		{
 			char newFileName[512];
 			Rpc_CreateFileW(realFileName, fileNameKey, u8(desiredAccess), newFileName, sizeof_array(newFileName), size, closeId, false);
@@ -631,7 +649,7 @@ template<typename True_stat>
 int Shared_stat(const char* funcName, const char* file, struct stat* attr, const True_stat& trueStat)
 {
 	StringBuffer<> fixedFile;
-	if (!FixPath(fixedFile, file) || fixedFile.Equals("/") || fixedFile.StartsWith("/etc/") || fixedFile.StartsWith(g_systemTemp.data))//(access(file, F_OK) != 0))
+	if (!FixPath(fixedFile, file) || fixedFile.Equals("/") || !CanDetour2(fixedFile))
 	{
 		int res = trueStat(file, attr);
 		return res;
@@ -639,9 +657,9 @@ int Shared_stat(const char* funcName, const char* file, struct stat* attr, const
 
 	UBA_ASSERTF(fixedFile.count, "FixPath failed with %s", file);
 
-	StringBuffer<> temp;
 	if (g_runningRemote && fixedFile.StartsWith(g_exeDir.data))
 	{
+		StringBuffer<> temp;
 		temp.Append(g_virtualApplicationDir).Append(fixedFile.data + g_exeDir.count);
 		fixedFile.Clear().Append(temp);
 	}
@@ -666,12 +684,12 @@ int Shared_stat(const char* funcName, const char* file, struct stat* attr, const
 		memcpy(attr, &fileAttr.data, sizeof(fileAttr.data));
 
 	#if UBA_DEBUG_VALIDATE
-	if (!g_runningRemote)
+	if (!g_runningRemote && *file)
 	{
 		struct stat attr2;
 		int res2 = trueStat(file, &attr2);
-		UBA_ASSERTF(res == res2, "stat: return value differs for %s (cached %i vs actual %i) [fixed: %s]", file, res, res2, fixedFile.data);
-		if (res != -1)
+		UBA_ASSERTF(res == res2, "%s: return value differs for %s (cached %i vs actual %i) [fixed: %s]", funcName, file, res, res2, fixedFile.data, StrError(res2, errno));
+		if (res == 0)
 		{
 			bool isDir = S_ISDIR(attr->st_mode);
 			UBA_ASSERT(isDir == S_ISDIR(attr2.st_mode));
@@ -764,8 +782,7 @@ UBA_EXPORT int UBA_WRAPPER(mkdir)(const char* path, mode_t mode)
 	UBA_INIT_DETOUR(mkdir, path, mode);
 
 	StringBuffer<> pathName;
-	FixPath(pathName, path);
-	if (pathName.StartsWith(g_systemTemp.data))
+	if (!FixPath(pathName, path) || !CanDetour2(pathName))
 	{
 		int res = TRUE_WRAPPER(mkdir)(path, mode);
 		DEBUG_LOG_TRUE("mkdir", "%s -> %i", path, res);
@@ -773,9 +790,9 @@ UBA_EXPORT int UBA_WRAPPER(mkdir)(const char* path, mode_t mode)
 	}
 
 	u32 directoryTableSize;
-	bool res;
+	int res;
 	u32 errorCode = 0;
-	StringKey pathNameKey = ToStringKeyLower(pathName);
+	StringKey pathNameKey = ToFilenameKey(pathName);
 
 	{
 		TimerScope ts(g_stats.createFile);
@@ -786,7 +803,7 @@ UBA_EXPORT int UBA_WRAPPER(mkdir)(const char* path, mode_t mode)
 		writer.WriteString(pathName);
 		writer.Flush();
 		BinaryReader reader;
-		res = reader.ReadBool();
+		res = reader.ReadBool() ? 0 : -1;
 		errorCode = reader.ReadU32();
 		directoryTableSize = reader.ReadU32();
 	}
@@ -795,7 +812,7 @@ UBA_EXPORT int UBA_WRAPPER(mkdir)(const char* path, mode_t mode)
 
 	errno = errorCode;
 	DEBUG_LOG_DETOURED("mkdir", "%s -> %i (%u)", path, res, errorCode);
-	return res ? 0 : -1;
+	return res;
 }
 
 UBA_EXPORT int UBA_WRAPPER(rmdir)(const char* path)
@@ -803,10 +820,8 @@ UBA_EXPORT int UBA_WRAPPER(rmdir)(const char* path)
 	UBA_INIT_DETOUR(rmdir, path);
 
 	StringBuffer<> pathName;
-	FixPath(pathName, path);
-	if (pathName.StartsWith(g_systemTemp.data))
+	if (!FixPath(pathName, path) || !CanDetour2(pathName))
 	{
-		//SuppressCreateFileDetourScope s; // TODO: Revisit this.. will not work remotely
 		int res = TRUE_WRAPPER(rmdir)(path);
 		DEBUG_LOG_TRUE("rmdir", "%s -> %i", path, res);
 		return res;
@@ -815,7 +830,7 @@ UBA_EXPORT int UBA_WRAPPER(rmdir)(const char* path)
 	u32 directoryTableSize;
 	bool res;
 	u32 errorCode = 0;
-	StringKey pathNameKey = ToStringKeyLower(pathName);
+	StringKey pathNameKey = ToFilenameKey(pathName);
 
 	{
 		TimerScope ts(g_stats.deleteFile);
@@ -870,6 +885,14 @@ UBA_EXPORT int UBA_WRAPPER(setenv)(const char* name, const char* value, int repl
 	return res;
 }
 
+UBA_EXPORT int UBA_WRAPPER(unsetenv)(const char* name)
+{
+	UBA_INIT_DETOUR(unsetenv, name);
+	auto res = TRUE_WRAPPER(unsetenv)(name);
+	DEBUG_LOG_TRUE("unsetenv", "(%s) -> (%i)", name, res);
+	return res;
+}
+
 UBA_EXPORT char* UBA_WRAPPER(realpath)(const char* path, char* resolved_path)
 {
 	UBA_INIT_DETOUR(realpath, path, resolved_path);
@@ -973,6 +996,14 @@ UBA_EXPORT int UBA_WRAPPER(open64)(const char* file, int flags, ...)
 	UBA_INIT_DETOUR(open64, file, flags, mode);
 	return Shared_open("open64", file, flags, mode, [](const char* realFile, int flags, int mode) { return TRUE_WRAPPER(open64)(realFile, flags, mode); });
 }
+UBA_EXPORT char* UBA_WRAPPER(secure_getenv)(const char* name)
+{
+	UBA_INIT_DETOUR(secure_getenv, name);
+	auto res = TRUE_WRAPPER(secure_getenv)(name);
+	DEBUG_LOG_TRUE("secure_getenv", "(%s) -> %s", name, res ? res : "<null>");
+	return res;
+}
+
 #endif
 
 UBA_EXPORT int UBA_WRAPPER(open)(const char* file, int flags, ...)
@@ -1128,32 +1159,26 @@ UBA_EXPORT DIR* UBA_WRAPPER(opendir)(const char* name)
 {
 	UBA_INIT_DETOUR(opendir, name);
 	StringBuffer<> dirName;
-	FixPath(dirName, name);
 
-	bool useDetour = true;
-
-	// TODO: Someone needs to debug this.. it fails on the farm.
-	#if PLATFORM_MAC
-	useDetour = false;
-	#endif
-
-	if (!useDetour)
+	if (!FixPath(dirName, name) || !CanDetour2(dirName))
 	{
 		DIR* res = TRUE_WRAPPER(opendir)(dirName.data);
 		DEBUG_LOG_TRUE("opendir", "(%s) -> %p", dirName.data, res);
 		return res;
 	}
 
-	DirHash hash(dirName.data, dirName.count);
+	StringBuffer<> forHash(dirName);
+	if (CaseInsensitiveFs)
+		forHash.MakeLower();
+	DirHash hash(forHash.data,forHash.count);
 
 	SCOPED_WRITE_LOCK(g_directoryTable.m_lookupLock, lookLock);
 	auto insres = g_directoryTable.m_lookup.try_emplace(hash.key, &g_memoryBlock);
 	DirectoryTable::Directory& dir = insres.first->second;
 	if (insres.second)
-	{
-		if (g_directoryTable.EntryExistsNoLock(hash.key, dirName.data, dirName.count) != DirectoryTable::Exists_No)
+		if (g_directoryTable.EntryExistsNoLock(hash.key, forHash.data, forHash.count) != DirectoryTable::Exists_No)
 			Rpc_UpdateDirectory(hash.key, dirName.data, dirName.count, false);
-	}
+\
 	bool exists = false;
 	if (dir.tableOffset != InvalidTableOffset)
 	{
@@ -1166,6 +1191,7 @@ UBA_EXPORT DIR* UBA_WRAPPER(opendir)(const char* name)
 	if (!exists)
 	{
 		errno = ENOENT;
+		DEBUG_LOG_DETOURED("opendir", "(%s) -> %p", dirName.data, nullptr);
 		return nullptr;
 	}
 
@@ -1336,7 +1362,7 @@ UBA_EXPORT int UBA_WRAPPER(access)(const char* pathname, int mode)
 	UBA_INIT_DETOUR(access, pathname, mode);
 
 	StringBuffer<> fixedPath;
-	if (!FixPath(fixedPath, pathname) || fixedPath.StartsWith(g_systemTemp.data) || fixedPath.StartsWith("/proc"))
+	if (!FixPath(fixedPath, pathname) || fixedPath.StartsWith("/proc") || !CanDetour2(fixedPath))
 	{
 		auto res = TRUE_WRAPPER(access)(pathname, mode);
 		DEBUG_LOG_TRUE("access", "%s %i -> %i (%s)", pathname, mode, res, StrError(res, errno));
@@ -1471,6 +1497,23 @@ UBA_EXPORT int UBA_WRAPPER(rename)(const char* oldpath, const char* newpath)
 		}
 	}
 
+	bool canDetourOld = CanDetour2(fixedOldPath);
+	bool canDetourNew = CanDetour2(fixedNewPath);
+	if (!canDetourOld)
+	{
+		if (!canDetourNew)
+		{
+			int res = TRUE_WRAPPER(rename)(oldpath, newpath);
+			DEBUG_LOG_TRUE("rename", "(from %s to %s) -> %i (%s)", fixedOldPath.data, fixedNewPath.data, res, StrError(res, errno));
+			return res;
+		}
+	}
+	else
+	{
+		//UBA_ASSERT(!canDetourNew);
+	}
+
+
 	u32 directoryTableSize;
 	u32 errorCode;
 	bool result;
@@ -1494,6 +1537,20 @@ UBA_EXPORT int UBA_WRAPPER(rename)(const char* oldpath, const char* newpath)
 
 	//DEBUG_LOG_DETOURED(L"rena", L"(PIPE) (%ls to %ls) -> %ls", lpExistingFileName, lpNewFileName, ToString(result));
 
+	// We need to add new into g_mappingFileTable since it will not be part of directory table (SHOULD THIS BE FIXED?)
+	if (g_runningRemote)
+	{
+		SCOPED_WRITE_LOCK(g_mappedFileTable.m_lookupLock, _);
+		auto insres = g_mappedFileTable.m_lookup.try_emplace(newKey);
+		FileInfo& info = insres.first->second;
+
+		auto findIt = g_mappedFileTable.m_lookup.find(oldKey);
+		if (findIt != g_mappedFileTable.m_lookup.end())
+			info = findIt->second;
+		info.originalName = g_memoryBlock.Strdup(fixedNewPath.data);
+		info.name = info.originalName;
+	}
+
 	g_directoryTable.ParseDirectoryTable(directoryTableSize);
 	g_mappedFileTable.SetDeleted(oldKey, fixedOldPath.data, true);
 	g_mappedFileTable.SetDeleted(newKey, fixedNewPath.data, false);
@@ -1512,6 +1569,14 @@ UBA_EXPORT int UBA_WRAPPER(chmod)(const char* pathname, mode_t mode)
 
 	StringBuffer<> fixedName;
 	FixPath(fixedName, pathname);
+
+	if (!CanDetour2(fixedName))
+	{
+		int res = TRUE_WRAPPER(chmod)(pathname, mode);
+		DEBUG_LOG_TRUE("chmod", "%s %i -> %i (%s)", pathname, mode, res, StrError(res, errno));
+		return res;
+	}
+
 	StringKey key = ToFilenameKey(fixedName);
 	u32 errorCode;
 	{
@@ -1533,13 +1598,13 @@ UBA_EXPORT int UBA_WRAPPER(chmod)(const char* pathname, mode_t mode)
 
 	errno = errorCode;
 	return res;
-	//return TRUE_WRAPPER(chmod)(pathname, mode);
 }
 
 UBA_EXPORT int UBA_WRAPPER(renameat)(int olddirfd, const char* oldpath, int newdirfd, const char* newpath)
 {
 	UBA_INIT_DETOUR(renameat, olddirfd, oldpath, newdirfd, newpath);
 	DEBUG_LOG_TRUE("renameat", "(from %s to %s)", oldpath, newpath);
+	UBA_ASSERTF(false, "Not implemented");
 	return TRUE_WRAPPER(renameat)(olddirfd, oldpath, newdirfd, newpath);
 }
 
@@ -1547,6 +1612,7 @@ UBA_EXPORT int UBA_WRAPPER(utimensat)(int dirfd, const char* pathname, const str
 {
 	UBA_INIT_DETOUR(utimensat, dirfd, pathname, times, flags);
 	DEBUG_LOG_TRUE("utimensat", "(%s)", pathname);
+	UBA_ASSERTF(false, "Not implemented");
 	return TRUE_WRAPPER(utimensat)(dirfd, pathname, times, flags);
 }
 
@@ -1554,6 +1620,7 @@ UBA_EXPORT int UBA_WRAPPER(symlink)(const char* path1, const char* path2)
 {
 	UBA_INIT_DETOUR(symlink, path1, path2);
 	DEBUG_LOG_TRUE("symlink", "(from %s to %s)", path1, path2);
+	UBA_ASSERTF(false, "Not implemented");
 	return TRUE_WRAPPER(symlink)(path1, path2);
 }
 
@@ -1583,6 +1650,15 @@ int Shared_DeleteFile(const char* funcName, const char* pathname)
 {
 	StringBuffer<> fixedName;
 	FixPath(fixedName, pathname);
+
+	if (!CanDetour2(fixedName))
+	{
+		int res = TRUE_WRAPPER(unlink)(pathname);
+		DEBUG_LOG_TRUE(funcName, "(%s) -> %i (%s)", pathname, res, StrError(res, errno));
+		return res;
+	}
+
+
 	StringKey fileNameKey = ToFilenameKey(fixedName);
 
 	u32 directoryTableSize;
@@ -1614,7 +1690,6 @@ int Shared_DeleteFile(const char* funcName, const char* pathname)
 	DEBUG_LOG_DETOURED(funcName, "(%s) -> %i (%s)", pathname, res, StrError(res, errorCode));
 	errno = errorCode;
 	return res;
-	//return TRUE_WRAPPER(remove)(pathname);
 }
 
 UBA_EXPORT int UBA_WRAPPER(remove)(const char* pathname)
@@ -1627,6 +1702,21 @@ UBA_EXPORT int UBA_WRAPPER(remove)(const char* pathname)
 UBA_EXPORT int UBA_WRAPPER(link)(const char* oldpath, const char* newpath)
 {
 	UBA_INIT_DETOUR(link, oldpath, newpath);
+
+	StringBuffer<> fixedOldPath;
+	FixPath(fixedOldPath, oldpath);
+	StringBuffer<> fixedNewPath;
+	FixPath(fixedNewPath, newpath);
+
+	if (!CanDetour2(fixedNewPath))
+	{
+		UBA_ASSERT(!CanDetour2(fixedOldPath));
+		int res = TRUE_WRAPPER(link)(oldpath, newpath);
+		DEBUG_LOG_TRUE("link", "(%s -> %s) -> %i (%s)", oldpath, newpath, res, StrError(res, errno));
+		return res;
+	}
+
+	UBA_ASSERT(false);
 	int res = TRUE_WRAPPER(link)(oldpath, newpath);
 	return res;
 }
@@ -1639,8 +1729,25 @@ UBA_EXPORT int UBA_WRAPPER(unlink)(const char* pathname)
 
 thread_local int t_inVfork;
 
+void FlattenArgs(StringBufferBase& out, const char* const argv[])
+{
+	if (!argv)
+		return;
+	for (u32 i = 0; argv[i]; ++i)
+	{
+		if (i != 0)
+			out.Append(' ');
+		out.Append(argv[i]);
+	}
+}
+
 bool ExecuteHostRun(StringBufferBase& out, const char* const* argv)
 {
+	{
+		StringBuffer<4096> command;
+		FlattenArgs(command, argv);
+		DEBUG_LOG_DETOURED("HostRun", "%s", command.data)
+	}
 	TimerScope ts(g_stats.getFullFileName);
 	SCOPED_WRITE_LOCK(g_communicationLock, pcs);
 	BinaryWriter writer;
@@ -1666,21 +1773,27 @@ bool ExecuteHostRun(StringBufferBase& out, const char* const* argv)
 
 int SpawnEcho(char* str, pid_t* pid, const posix_spawn_file_actions_t* file_actions, const posix_spawnattr_t* attrp, char* const envp[])
 {
+#if 0
+	static Atomic<u32> counter;
+	StringBuffer<> tempFile;
+	tempFile.Append(g_systemTemp).EnsureEndsWithSlash().Append("UbaTempFile").AppendValue(getpid()).Append('_').AppendValue(counter++);
+	int fd = TRUE_WRAPPER(open)(tempFile.data, O_CREAT|O_TRUNC|O_WRONLY, S_IRUSR | S_IWUSR);
+	TRUE_WRAPPER(write)(fd, str, strlen(str));
+	TRUE_WRAPPER(close)(fd);
+
+	DEBUG_LOG("Created %s for cat", tempFile.data);
+	char cmd[] = "/bin/cat"; 
+	char* const argv2[] = { cmd, tempFile.data, nullptr };
+	int res = TRUE_WRAPPER(posix_spawn)(pid, cmd, file_actions, attrp, argv2, envp);
+	DEBUG_LOG_TRUE("posix_spawn", "(CAT) %s (pid: %u) -> %i", str, *pid, res);
+#else
+	char* const env[] = { nullptr };
 	char cmd[] = "/bin/echo"; 
 	char* const argv2[] = { str, nullptr };
-	return TRUE_WRAPPER(posix_spawn)(pid, cmd, file_actions, attrp, argv2, envp);
-}
-
-void FlattenArgs(StringBufferBase& out, char* const argv[])
-{
-	if (!argv)
-		return;
-	for (u32 i = 0; argv[i]; ++i)
-	{
-		if (i != 0)
-			out.Append(' ');
-		out.Append(argv[i]);
-	}
+	int res = TRUE_WRAPPER(posix_spawn)(pid, cmd, file_actions, attrp, argv2, env);
+	DEBUG_LOG_TRUE("posix_spawn", "(ECHO) %s (pid: %u) -> %i", str, *pid, res);
+#endif
+	return res;
 }
 
 void UnsupportedHostRun(char* const argv[], const char* msg)
@@ -1696,20 +1809,87 @@ int shared_posix_spawn(pid_t* pid, const char* path, const posix_spawn_file_acti
 
 	t_inVfork = 0;
 
+	{
+		//StringBuffer<4096> command;
+		//FlattenArgs(command, argv);
+		//DEBUG_LOG("RUNNING: %s", command.data);
+	}
+
+	const char* realArgv[1024];
+	StringBuffer<> result;
+
+	if (strstr(path, "xcode-select"))
+	{
+		if (strcmp(argv[1], "--print-path") != 0)
+			UbaAssert("xcode-select only supported with --print-path", __FILE__, __LINE__, "", 1999);
+		if (!ExecuteHostRun(result, argv))
+			return -1;
+		return SpawnEcho(result.data, pid,  file_actions, attrp, envp);
+	}
+	else if(strstr(path, "xcrun"))
+	{
+		if (strcmp(argv[1], "--sdk") != 0)
+			UnsupportedHostRun(argv, "xcrun unsupported first param");
+
+		if (strcmp(argv[3], "--find") == 0)
+		{
+			if (!ExecuteHostRun(result, argv))
+				return -1;
+			return SpawnEcho(result.data, pid,  file_actions, attrp, envp);
+		}
+		
+		if (strcmp(argv[3], "metal") != 0 && strcmp(argv[3], "metallib") != 0)
+			UnsupportedHostRun(argv, "xcrun unsupported third param");
+
+		const char* argv2[] = { argv[0], argv[1], argv[2], "--find", argv[3], 0 };
+		if (!ExecuteHostRun(result, argv2))
+			return -1;
+		// This will return the trampoline metal which we don't want
+		if (const char* usrbin = strstr(result.data, "/bin/metal"))
+		{
+			result.Resize(usrbin - result.data);
+			result.Append("/metal/");
+			if (strcmp(argv[2], "macosx") == 0)
+				result.Append("macos");
+			else
+				result.Append("ios");
+			result.Append("/bin/").Append(argv[3]);
+		}
+
+		path = result.data;
+
+		u32 argc3 = 0;
+		realArgv[argc3++] = result.data;
+		for (int i=4;argv[i]; ++i)
+			realArgv[argc3++] = argv[i];
+		realArgv[argc3] = 0;
+
+		//StringBuffer<4096> command;
+		//FlattenArgs(command, (char*const*)realArgv);
+		//DEBUG_LOG("FIXEDUP: %s", command.data);
+		//return TRUE_WRAPPER(posix_spawn)(pid, path, file_actions, attrp, (char*const*)realArgv, envp);
+	}
+	else
+	{
+		int i = 0;
+		for (;argv[i]; ++i)
+			realArgv[i] = argv[i];
+		realArgv[i] = 0;
+	}
+
+	//DEBUG_LOG("LIBRARY_SEARCH_PATHS: %s", getenv("LIBRARY_SEARCH_PATHS"));
+	//DEBUG_LOG("PATH: %s", getenv("PATH"));
+
 	TString cmdLine;
 	u32 argc = 0;
-	for (u32 i = 0; argv[i]; ++i)
+	for (u32 i = 0; realArgv[i]; ++i)
 	{
 		if (i != 0)
 			cmdLine.append(" ");
-		cmdLine.append(argv[i]);
+		cmdLine.append(realArgv[i]);
 		++argc;
 	}
 
-	Vector<const char*> argv2;
-	argv2.resize(argc + 1);
-	memcpy(argv2.data(), argv, sizeof(char*) * (argc+1));
-
 	TString realApplication;
 	TString commandLine;
 	u32 processId = 0;
@@ -1766,8 +1946,6 @@ int shared_posix_spawn(pid_t* pid, const char* path, const posix_spawn_file_acti
 
 		logFile.Append("UBA_LOGFILE=");
 		reader.ReadString(logFile);
-
-		// TODO: Recalculate argv2 right now it is wrong, it should be what is in command line
 	}
 	
 	std::vector<const char*> envvars;
@@ -1780,8 +1958,7 @@ int shared_posix_spawn(pid_t* pid, const char* path, const posix_spawn_file_acti
 		envvars.push_back(comIdVar.data);
 		envvars.push_back(currentDir.data);
 		envvars.push_back(rulesStr.data);
-		if (logFile.count > 12)
-			envvars.push_back(logFile.data);
+		envvars.push_back(logFile.data);
 		//envvars.push_back("LD_DEBUG=bindings");
 		envvars.push_back(nullptr);
 	}
@@ -1795,12 +1972,12 @@ int shared_posix_spawn(pid_t* pid, const char* path, const posix_spawn_file_acti
 //	}
 	
 	#if UBA_DEBUG_LOG_ENABLED
-	DEBUG_LOG_TRUE("posix_spawn", "%s", realApplication.data());
-	for (u32 i = 0; argv[i]; ++i)
-		DEBUG_LOG("            %s", argv[i]);
+	DEBUG_LOG_TRUE("posix_spawn", "%s (%s)", realApplication.data(), logFile.data);
+	for (u32 i = 0; realArgv[i]; ++i)
+		DEBUG_LOG("            %s", realArgv[i]);
 	#endif
 
-	int res = TRUE_WRAPPER(posix_spawn)(pid, realApplication.data(), file_actions, attrp, (char* const*)argv2.data(), (char**)envvars.data());
+	int res = TRUE_WRAPPER(posix_spawn)(pid, realApplication.data(), file_actions, attrp, (char*const*)realArgv, (char**)envvars.data());
 	bool success = res == 0;
 
 	{
@@ -1848,6 +2025,8 @@ UBA_EXPORT pid_t UBA_WRAPPER(waitpid)(pid_t pid, int* status, int options)
 	// TODO: Should probably report id to session
 	auto res = TRUE_WRAPPER(waitpid)(pid, status, options);
 	DEBUG_LOG_TRUE("waitpid", "(%i) -> %i", pid, res);
+	if (WIFEXITED(*status))
+		Rpc_UpdateTables();
 	return res;
 }
 
@@ -1876,6 +2055,8 @@ UBA_EXPORT int UBA_WRAPPER(waitid)(idtype_t idtype, id_t id, siginfo_t* infop, i
 	// TODO: Should probably report id to session
 	auto res = TRUE_WRAPPER(waitid)(idtype, id, infop, options);
 	DEBUG_LOG_TRUE("waitid", "%i -> %i (%s)", id, res, GetResult(infop));
+	if (infop->si_code == CLD_EXITED)
+		Rpc_UpdateTables();
 	return res;
 }
 
@@ -1898,7 +2079,9 @@ UBA_EXPORT pid_t UBA_WRAPPER(wait4)(pid_t pid, int* status, int options, struct
 	return res;
 }
 
-void Shared_LoadLibrary(Set<TString>& handled, const char*& path, const char* const* loaderPaths, StringBufferBase& tempBuf)
+Set<TString> g_handledLibraries;
+
+void Shared_LoadLibrary(const char*& path, const char* const* loaderPaths, StringBufferBase& tempBuf)
 {
 #if PLATFORM_MAC
 	u64 nameLen = 0;
@@ -1907,10 +2090,10 @@ void Shared_LoadLibrary(Set<TString>& handled, const char*& path, const char* co
 	StringBuffer<> error;
 	FindImports(path, [&](const tchar* import, bool isKnown, const char* const* importLoaderPaths)
 	{
-		if (!handled.insert(import).second)
+		if (!g_handledLibraries.insert(import).second)
 			return;
 		StringBuffer<> temp;
-		Shared_LoadLibrary(handled, import, importLoaderPaths, temp);
+		Shared_LoadLibrary(import, importLoaderPaths, temp);
 	}, error);
 	if (error.count)
 		DEBUG_LOG(error.data);
@@ -1929,10 +2112,9 @@ UBA_EXPORT void* UBA_WRAPPER(dlopen)(const char* path, int mode)
 		if (StartsWith(path, "@rpath/"))
 		{
 			path += 7;
-			Set<TString> handled;
-			handled.insert(path);
 			const char* loaderPaths[] = { "/", 0 };
-			Shared_LoadLibrary(handled, path, loaderPaths, tempBuf);
+			if (g_handledLibraries.insert(path).second)
+				Shared_LoadLibrary(path, loaderPaths, tempBuf);
 		}
 		else if (!StartsWith(path, "/System") && !StartsWith(path, "/usr/lib"))
 		{
@@ -1946,6 +2128,20 @@ UBA_EXPORT void* UBA_WRAPPER(dlopen)(const char* path, int mode)
 	return res;
 }
 
+UBA_EXPORT int UBA_WRAPPER(dladdr)(const void *addr, Dl_info *info)
+{
+	UBA_INIT_DETOUR(dladdr, addr, info);
+	int res = TRUE_WRAPPER(dladdr)(addr, info);
+	if (StartsWith(info->dli_fname, g_exeDir.data))
+	{
+		StringBuffer<> newPath;
+		newPath.Append(g_virtualApplication.data).Append(info->dli_fname + g_exeDir.count - 1); // TODO: Make g_handledLibraries a map and store the real path
+		info->dli_fname = g_memoryBlock.Strdup(newPath.data);
+	}
+	DEBUG_LOG_TRUE("dladdr", "%s -> %i", info->dli_fname, res);
+	return res;
+}
+
 UBA_EXPORT int UBA_WRAPPER(execv)(const char* path, char* const argv[])
 {
 	UBA_INIT_DETOUR(execv, path, argv);
@@ -2006,7 +2202,35 @@ UBA_EXPORT int UBA_WRAPPER(execvp)(const char* file, char* const argv[])
 	return TRUE_WRAPPER(execvp)(file , argv);
 }
 
+UBA_EXPORT int UBA_WRAPPER(execl)(const char *path, const char *arg0, ...)
+{
+	DEBUG_LOG_TRUE("execl", "");
+	UBA_ASSERT(false);
+	return -1;
+}
+
+UBA_EXPORT int UBA_WRAPPER(execle)(const char *path, const char *arg0, ...)
+{
+	DEBUG_LOG_TRUE("execle", "");
+	UBA_ASSERT(false);
+	return -1;
+}
+
+UBA_EXPORT int UBA_WRAPPER(execlp)(const char *file, const char *arg0, ...)
+{
+	DEBUG_LOG_TRUE("execlp", "");
+	UBA_ASSERT(false);
+	return -1;
+}
+
 #if PLATFORM_MAC
+UBA_EXPORT int UBA_WRAPPER(execvP)(const char *file, const char *search_path, char *const argv[])
+{
+	DEBUG_LOG_TRUE("execvP", "");
+	UBA_ASSERT(false);
+	return -1;
+}
+
 UBA_EXPORT pid_t UBA_WRAPPER(fork)(void)
 {
 	UBA_INIT_DETOUR(fork);
@@ -2016,7 +2240,10 @@ UBA_EXPORT pid_t UBA_WRAPPER(fork)(void)
 #endif
 
 
-#if PLATFORM_LINUX
+#if PLATFORM_MAC
+#pragma clang diagnostic ignored "-Wdeprecated-declarations"
+#endif
+
 UBA_EXPORT pid_t UBA_WRAPPER(vfork)(void)
 {
 	UBA_INIT_DETOUR(vfork);
@@ -2031,7 +2258,6 @@ UBA_EXPORT pid_t UBA_WRAPPER(vfork)(void)
 	}
 	return pid;
 }
-#endif
 
 UBA_EXPORT FILE* UBA_WRAPPER(popen)(const char* command, const char* type)
 {
@@ -2044,8 +2270,8 @@ UBA_EXPORT void UBA_WRAPPER(exit)(int status)
 {
 	//UBA_INIT_DETOUR(exit, status);
 	DEBUG_LOG_TRUE("exit", "(%i)", status);
-	Deinit();
-	CloseCom();
+	//Deinit();
+	//CloseCom();
 	TRUE_WRAPPER(exit)(status);
 }
 
@@ -2149,7 +2375,6 @@ namespace uba
 #endif
 	}
 
-
 	void PreInit(const char* logFile)
 	{
 		g_exeDirMem.Create();
@@ -2292,6 +2517,12 @@ namespace uba
 		if (g_runningRemote && g_isChild)
 			Rpc_GetParentWrittenFiles();
 
+		//pthread_atfork([]()
+		//{
+		//	DEBUG_LOG("FORKING!!");
+		//},[](){},[](){});
+
+
 		g_isDetouring = true;
 		DEBUG_LOG("Detouring enabled");
 
