[UBA]
* Fixed linux/macos crash when sub process with more than 1024 arguments were spawned

diff --git a/Engine/Source/Programs/UnrealBuildAccelerator/Detours/Private/PosixOS/UbaDetoursFunctionsPosix.cpp b/Engine/Source/Programs/UnrealBuildAccelerator/Detours/Private/PosixOS/UbaDetoursFunctionsPosix.cpp
--- a/Engine/Source/Programs/UnrealBuildAccelerator/Detours/Private/PosixOS/UbaDetoursFunctionsPosix.cpp
+++ b/Engine/Source/Programs/UnrealBuildAccelerator/Detours/Private/PosixOS/UbaDetoursFunctionsPosix.cpp
@@ -1815,7 +1815,7 @@ int shared_posix_spawn(pid_t* pid, const char* path, const posix_spawn_file_acti
 		//DEBUG_LOG("RUNNING: %s", command.data);
 	}
 
-	const char* realArgv[1024];
+	const char* tempArgv[1024];
 	StringBuffer<> result;
 
 	if (strstr(path, "xcode-select"))
@@ -1859,22 +1859,11 @@ int shared_posix_spawn(pid_t* pid, const char* path, const posix_spawn_file_acti
 		path = result.data;
 
 		u32 argc3 = 0;
-		realArgv[argc3++] = result.data;
+		tempArgv[argc3++] = result.data;
 		for (int i=4;argv[i]; ++i)
-			realArgv[argc3++] = argv[i];
-		realArgv[argc3] = 0;
-
-		//StringBuffer<4096> command;
-		//FlattenArgs(command, (char*const*)realArgv);
-		//DEBUG_LOG("FIXEDUP: %s", command.data);
-		//return TRUE_WRAPPER(posix_spawn)(pid, path, file_actions, attrp, (char*const*)realArgv, envp);
-	}
-	else
-	{
-		int i = 0;
-		for (;argv[i]; ++i)
-			realArgv[i] = argv[i];
-		realArgv[i] = 0;
+			tempArgv[argc3++] = argv[i];
+		tempArgv[argc3] = 0;
+		argv = (char*const*)tempArgv;
 	}
 
 	//DEBUG_LOG("LIBRARY_SEARCH_PATHS: %s", getenv("LIBRARY_SEARCH_PATHS"));
@@ -1882,11 +1871,11 @@ int shared_posix_spawn(pid_t* pid, const char* path, const posix_spawn_file_acti
 
 	TString cmdLine;
 	u32 argc = 0;
-	for (u32 i = 0; realArgv[i]; ++i)
+	for (u32 i = 0; argv[i]; ++i)
 	{
 		if (i != 0)
 			cmdLine.append(" ");
-		cmdLine.append(realArgv[i]);
+		cmdLine.append(argv[i]);
 		++argc;
 	}
 
@@ -1973,11 +1962,11 @@ int shared_posix_spawn(pid_t* pid, const char* path, const posix_spawn_file_acti
 	
 	#if UBA_DEBUG_LOG_ENABLED
 	DEBUG_LOG_TRUE("posix_spawn", "%s (%s)", realApplication.data(), logFile.data);
-	for (u32 i = 0; realArgv[i]; ++i)
-		DEBUG_LOG("            %s", realArgv[i]);
+	for (u32 i = 0; argv[i]; ++i)
+		DEBUG_LOG("            %s", argv[i]);
 	#endif
 
-	int res = TRUE_WRAPPER(posix_spawn)(pid, realApplication.data(), file_actions, attrp, (char*const*)realArgv, (char**)envvars.data());
+	int res = TRUE_WRAPPER(posix_spawn)(pid, realApplication.data(), file_actions, attrp, argv, (char**)envvars.data());
 	bool success = res == 0;
 
 	{
