[UBA]
* Fixed some application rules for non windows

diff --git a/Engine/Source/Programs/UnrealBuildAccelerator/Core/Private/UbaApplicationRules.cpp b/Engine/Source/Programs/UnrealBuildAccelerator/Core/Private/UbaApplicationRules.cpp
--- a/Engine/Source/Programs/UnrealBuildAccelerator/Core/Private/UbaApplicationRules.cpp
+++ b/Engine/Source/Programs/UnrealBuildAccelerator/Core/Private/UbaApplicationRules.cpp
@@ -513,6 +513,7 @@ namespace uba
 		static RulesRec rules[]
 		{
 			{ TC(""),							new ApplicationRules() },		// Must be index 0
+#if PLATFORM_WINDOWS
 			{ TC("cl.exe"),						new ApplicationRulesClExe() },	// Must be index 1
 			{ TC("link.exe"),					new ApplicationRulesLinkExe() }, // Must be index 2
 			{ TC("lib.exe"),					new ApplicationRulesVcLink() },
@@ -534,11 +535,16 @@ namespace uba
 			{ TC("llvm-objcopy.exe"),			new ApplicationRulesLlvmObjCopyExe() },
 			{ TC("UnrealBuildTool.dll"),		new ApplicationRulesUBTDll() },
 			{ TC("PVS-Studio.exe"),				new ApplicationRulesPVSStudio() },
-			{ TC("ShaderCompileWorker.exe"),	new ApplicationRulesShaderCompileWorker() },
 			{ TC("UbaObjTool.exe"),				new ApplicationRulesUbaObjTool() },
+			{ TC("ShaderCompileWorker.exe"),	new ApplicationRulesShaderCompileWorker() },
 			//{ L"MSBuild.dll"),				new ApplicationRules() },
 			//{ L"BreakpadSymbolEncoder.exe"),	new ApplicationRulesClang() },
 			//{ L"cmd.exe"),		new ApplicationRules() },
+#else
+			{ TC("clang++"),					new ApplicationRulesClangPlusPlusExe() },
+			{ TC("ld.lld"),						new ApplicationRulesLdLLdExe() },
+			{ TC("ShaderCompileWorker"),		new ApplicationRulesShaderCompileWorker() },
+#endif
 			{ nullptr, nullptr }
 		};
 
