[UBA]
* Fixed so shaders can now be compiled remotely on macos

diff --git a/Engine/Source/Programs/UnrealBuildAccelerator/Detours/Private/PosixOS/UbaDetoursFunctionsPosix.cpp b/Engine/Source/Programs/UnrealBuildAccelerator/Detours/Private/PosixOS/UbaDetoursFunctionsPosix.cpp
--- a/Engine/Source/Programs/UnrealBuildAccelerator/Detours/Private/PosixOS/UbaDetoursFunctionsPosix.cpp
+++ b/Engine/Source/Programs/UnrealBuildAccelerator/Detours/Private/PosixOS/UbaDetoursFunctionsPosix.cpp
@@ -299,7 +299,6 @@ bool CanDetour2(const StringView& file)
 		&& !file.StartsWith("/dev/")
 		&& !file.StartsWith("/etc/")
 		&& !file.StartsWith(g_systemTemp.data)
-		&& !file.Contains("/.cache/clang/")
 		;
 }
 
@@ -1619,8 +1618,20 @@ UBA_EXPORT int UBA_WRAPPER(utimensat)(int dirfd, const char* pathname, const str
 UBA_EXPORT int UBA_WRAPPER(symlink)(const char* path1, const char* path2)
 {
 	UBA_INIT_DETOUR(symlink, path1, path2);
-	DEBUG_LOG_TRUE("symlink", "(from %s to %s)", path1, path2);
-	UBA_ASSERTF(false, "Not implemented");
+
+	StringBuffer<> fixedPath1;
+	FixPath(fixedPath1, path1);
+	StringBuffer<> fixedPath2;
+	FixPath(fixedPath2, path2);
+
+	if (!CanDetour2(fixedPath1) && !CanDetour2(fixedPath2))
+	{
+		DEBUG_LOG_TRUE("symlink", "(from %s to %s)", path1, path2);
+		return TRUE_WRAPPER(symlink)(path1, path2);
+	}
+
+	UBA_ASSERTF(false, "symlink not implemented (from %s to %s)", path1, path2);
+	DEBUG_LOG_DETOURED("symlink", "(from %s to %s)", path1, path2);
 	return TRUE_WRAPPER(symlink)(path1, path2);
 }
 
@@ -1817,6 +1828,7 @@ int shared_posix_spawn(pid_t* pid, const char* path, const posix_spawn_file_acti
 
 	const char* tempArgv[1024];
 	StringBuffer<> result;
+	StringBuffer<> additionalArg;
 
 	if (strstr(path, "xcode-select"))
 	{
@@ -1862,6 +1874,12 @@ int shared_posix_spawn(pid_t* pid, const char* path, const posix_spawn_file_acti
 		tempArgv[argc3++] = result.data;
 		for (int i=4;argv[i]; ++i)
 			tempArgv[argc3++] = argv[i];
+		if (Equals(argv[3], "metal"))
+		{
+			// This is needed because we want clang cache to be local to machine and not be under the host machine's temp (which might not match remote machine's temp)
+			additionalArg.Append("-fmodules-cache-path=").Append(g_systemTemp).EnsureEndsWithSlash().Append("clangcache");
+			tempArgv[argc3++] = additionalArg.data;
+		}
 		tempArgv[argc3] = 0;
 		argv = (char*const*)tempArgv;
 	}
