[UBA]
* Improved UbaVisualizer logic when connected via network

diff --git a/Engine/Source/Programs/UnrealBuildAccelerator/Visualizer/Private/UbaVisualizer.cpp b/Engine/Source/Programs/UnrealBuildAccelerator/Visualizer/Private/UbaVisualizer.cpp
--- a/Engine/Source/Programs/UnrealBuildAccelerator/Visualizer/Private/UbaVisualizer.cpp
+++ b/Engine/Source/Programs/UnrealBuildAccelerator/Visualizer/Private/UbaVisualizer.cpp
@@ -317,6 +317,8 @@ namespace uba
 		if (!StartHwndThread())
 			return true;
 
+		m_clientDisconnect.Create(true);
+
 		wchar_t dots[] = TC("....");
 		u32 dotsCounter = 0;
 
@@ -326,7 +328,9 @@ namespace uba
 			if (!m_client)
 			{
 				bool ctorSuccess = true;
-				m_client = new NetworkClient(ctorSuccess, {});
+				NetworkClientCreateInfo ncci;
+				ncci.workerCount = 1;
+				m_client = new NetworkClient(ctorSuccess, ncci);
 				if (!ctorSuccess)
 					return false;
 			}
@@ -341,13 +345,20 @@ namespace uba
 
 			PostNewTrace(0, false);
 
-			while (m_hwnd && m_client->IsConnected())
-				Sleep(1000);
+			title.Clear().Appendf(L"Connected to %s:%u", host, port);
+			PostNewTitle(title);
+
+			while (m_hwnd && m_client->IsConnected() && !m_clientDisconnect.IsSet(1000))
+				;
+
+			title.Clear().Appendf(L"Disconnected...");
+			PostNewTitle(title);
 
 			m_client->Disconnect();
 			delete m_client;
 			m_client = nullptr;
-			Sleep(2000); // To prevent it from reconnecting to the same thing again and get thrown out (since it will post a WM_NEWTRACE and clean everything
+			m_clientDisconnect.Reset();
+			Sleep(4000); // To prevent it from reconnecting to the same thing again and get thrown out (since it will post a WM_NEWTRACE and clean everything
 		}
 		return true;
 	}
@@ -2349,7 +2360,7 @@ namespace uba
 		{
 			drawText(L"Finished Processes: %u (local: %u)", m_traceView.totalProcessExitedCount, session.processExitedCount);
 			drawText(L"Active Processes: %u (local: %u)", m_traceView.totalProcessActiveCount, session.processActiveCount);
-			drawText(L"Active Helpers: %u", m_traceView.activeSessionCount - 1);
+			drawText(L"Active Helpers: %u", Max(1u, m_traceView.activeSessionCount) - 1);
 
 			if (!session.updates.empty())
 			{
@@ -2962,9 +2973,11 @@ namespace uba
 			if (m_client)
 			{
 				if (!m_trace.StartReadClient(m_traceView, *m_client))
+				{
+					m_clientDisconnect.Set();
 					return false;
+				}
 				m_namedTrace.Clear().Append(m_newTraceName);
-				title.Appendf(L"Connected to host");
 				m_traceView.finished = false;
 			}
 			else if (!m_fileName.IsEmpty())
@@ -3039,7 +3052,10 @@ namespace uba
 						m_trace.UpdateReadFile(m_traceView, timeOffset, changed);
 				}
 				else if (m_client)
-					m_trace.UpdateReadClient(m_traceView, *m_client, changed);
+				{
+					if (!m_trace.UpdateReadClient(m_traceView, *m_client, changed))
+						m_clientDisconnect.Set();
+				}
 				else
 				{
 					if (!m_trace.UpdateReadNamed(m_traceView, m_replay ? timeOffset : ~u64(0), changed))
diff --git a/Engine/Source/Programs/UnrealBuildAccelerator/Visualizer/Private/UbaVisualizer.h b/Engine/Source/Programs/UnrealBuildAccelerator/Visualizer/Private/UbaVisualizer.h
--- a/Engine/Source/Programs/UnrealBuildAccelerator/Visualizer/Private/UbaVisualizer.h
+++ b/Engine/Source/Programs/UnrealBuildAccelerator/Visualizer/Private/UbaVisualizer.h
@@ -212,10 +212,12 @@ namespace uba
 
 		Logger& m_logger;
 		VisualizerConfig m_config;
-		NetworkClient* m_client = nullptr;
 		TraceReader m_trace;
 		TraceView m_traceView;
-		
+
+		NetworkClient* m_client = nullptr;
+		Event m_clientDisconnect;
+
 		StringBuffer<256>m_listenChannel;
 		StringBuffer<256> m_newTraceName;
 		Event m_listenTimeout;
