[UBA]
* Fixed visualizer title text when using network path

diff --git a/Engine/Source/Programs/UnrealBuildAccelerator/Visualizer/Private/UbaVisualizer.cpp b/Engine/Source/Programs/UnrealBuildAccelerator/Visualizer/Private/UbaVisualizer.cpp
--- a/Engine/Source/Programs/UnrealBuildAccelerator/Visualizer/Private/UbaVisualizer.cpp
+++ b/Engine/Source/Programs/UnrealBuildAccelerator/Visualizer/Private/UbaVisualizer.cpp
@@ -250,9 +250,7 @@ namespace uba
 
 		{
 			StringBuffer<> title;
-			GetTitlePrefix(title);
-			title.Appendf(L"Listening for new sessions on channel '%s'", m_listenChannel.data);
-			PostNewTitle(title);
+			PostNewTitle(GetTitlePrefix(title).Appendf(L"Listening for new sessions on channel '%s'", m_listenChannel.data));
 		}
 
 		StringBuffer<256> traceName;
@@ -336,23 +334,18 @@ namespace uba
 			}
 
 			StringBuffer<> title;
-			GetTitlePrefix(title);
-			title.Appendf(L"Trying to connect to %s:%u%s", host, port, dots + ((dotsCounter--) % 4));
-			PostNewTitle(title);
+			PostNewTitle(GetTitlePrefix(title).Appendf(L"Trying to connect to %s:%u%s", host, port, dots + ((dotsCounter--) % 4)));
 
 			if (!m_client->Connect(backend, host, port))
 				continue;
 
+			PostNewTitle(GetTitlePrefix(title).Appendf(L"Connected to %s:%u", host, port));
 			PostNewTrace(0, false);
 
-			title.Clear().Appendf(L"Connected to %s:%u", host, port);
-			PostNewTitle(title);
-
 			while (m_hwnd && m_client->IsConnected() && !m_clientDisconnect.IsSet(1000))
 				;
 
-			title.Clear().Appendf(L"Disconnected...");
-			PostNewTitle(title);
+			PostNewTitle(GetTitlePrefix(title).Appendf(L"Disconnected..."));
 
 			m_client->Disconnect();
 			delete m_client;
@@ -398,13 +391,15 @@ namespace uba
 		m_locked = lock;
 	}
 
-	void Visualizer::GetTitlePrefix(StringBufferBase& out)
+	StringBufferBase& Visualizer::GetTitlePrefix(StringBufferBase& out)
 	{
+		out.Clear();
 		out.Append(L"UbaVisualizer");
 		#if UBA_DEBUG
 		out.Append(L" (DEBUG)");
 		#endif
 		out.Append(L" - ");
+		return out;
 	}
 
 	bool Visualizer::Unselect()
@@ -591,8 +586,7 @@ namespace uba
 			windowStyle = WS_POPUP | WS_VSCROLL | WS_HSCROLL;// | WS_VISIBLE;
 
 		StringBuffer<> title;
-		GetTitlePrefix(title);
-		title.Append(L"Initializing...");
+		GetTitlePrefix(title).Append(L"Initializing...");
 
 		HWND hwnd = CreateWindowEx(exStyle, windowClassName, title.data, windowStyle, winPosX, winPosY, winWidth, winHeight, NULL, NULL, hInstance, this);
 		SetWindowLongPtr(hwnd, GWLP_USERDATA, (LONG_PTR)this);
@@ -2983,8 +2977,8 @@ namespace uba
 			else if (!m_fileName.IsEmpty())
 			{
 				m_trace.ReadFile(m_traceView, m_fileName.data, m_replay != 0);
-				title.Append(m_fileName);
 				m_traceView.finished = m_replay == 0;
+				PostNewTitle(GetTitlePrefix(title).Append(m_fileName));
 			}
 			else
 			{
@@ -2992,10 +2986,9 @@ namespace uba
 					return false;
 				m_namedTrace.Clear().Append(m_newTraceName);
 				m_traceView.finished = false;
-				title.Appendf(L"%s (Listening for new sessions on channel '%s')", m_namedTrace.data, m_listenChannel.data);
+				PostNewTitle(GetTitlePrefix(title).Appendf(L"%s (Listening for new sessions on channel '%s')", m_namedTrace.data, m_listenChannel.data));
 			}
 
-			PostNewTitle(title);
 			SetTimer(m_hwnd, 0, 200, NULL);
 			return 0;
 		}
diff --git a/Engine/Source/Programs/UnrealBuildAccelerator/Visualizer/Private/UbaVisualizer.h b/Engine/Source/Programs/UnrealBuildAccelerator/Visualizer/Private/UbaVisualizer.h
--- a/Engine/Source/Programs/UnrealBuildAccelerator/Visualizer/Private/UbaVisualizer.h
+++ b/Engine/Source/Programs/UnrealBuildAccelerator/Visualizer/Private/UbaVisualizer.h
@@ -127,7 +127,7 @@ namespace uba
 		bool UpdateAutoscroll();
 		bool UpdateSelection();
 		void UpdateScrollbars(bool redraw);
-		void GetTitlePrefix(StringBufferBase& out);
+		StringBufferBase& GetTitlePrefix(StringBufferBase& out);
 		void InitBrushes();
 		void ThreadLoop();
 		void Pause(bool pause);
