[UBA]
* Fixed so uba obj tool does not strip ThisIsAnUnrealEngineModule to make sure there is always at least one export

diff --git a/Engine/Source/Programs/UnrealBuildAccelerator/Common/Private/UbaObjectFileCoff.cpp b/Engine/Source/Programs/UnrealBuildAccelerator/Common/Private/UbaObjectFileCoff.cpp
--- a/Engine/Source/Programs/UnrealBuildAccelerator/Common/Private/UbaObjectFileCoff.cpp
+++ b/Engine/Source/Programs/UnrealBuildAccelerator/Common/Private/UbaObjectFileCoff.cpp
@@ -494,10 +494,17 @@ namespace uba
 			directiveSection.Characteristics = ImageScnAlign1Bytes|ImageScnLnkInfo|ImageScnLnkRemove;
 			++header.NumberOfSections;
 
+			auto writeExport = [&](const std::string& symbol, const std::string& extra)
+				{
+					write("/EXPORT:", 8);
+					write(symbol.data(), symbol.size());
+					write(extra.data(), extra.size());
+					write(" ", 1);
+				};
+
 			// Directive raw data
 			u32 directiveRawDataStart = u32(memoryBlock.writtenSize);
 			directiveSection.PointerToRawData = directiveRawDataStart;
-			char slashExport[] = "/EXPORT:";
 			for (auto& kv : allExports)
 			{
 				auto& symbol = kv.first;
@@ -509,11 +516,13 @@ namespace uba
 						continue;
 				}
 
-				write(slashExport, sizeof(slashExport) - 1);
-				write(symbol.data(), symbol.size());
-				write(kv.second.extra.data(), kv.second.extra.size());
-				write(" ", 1);
+				writeExport(symbol, kv.second.extra);
 			}
+
+			tmp = "ThisIsAnUnrealEngineModule";
+			if (allExports.find(tmp) != allExports.end())
+				writeExport(tmp, ""); // Workaround for tool not liking empty lists
+
 			write("", 1);
 			directiveSection.SizeOfRawData = u32(memoryBlock.writtenSize) - directiveRawDataStart;
 		}
