[UBA]
* Changed so dllimport storage is cleared too with the exception of a couple symbols which I don't know why they need to stay "imported". Need to revisit why those symbols are special.

diff --git a/Engine/Source/Programs/UnrealBuildAccelerator/Common/Private/UbaObjectFileLLVMIR.cpp b/Engine/Source/Programs/UnrealBuildAccelerator/Common/Private/UbaObjectFileLLVMIR.cpp
--- a/Engine/Source/Programs/UnrealBuildAccelerator/Common/Private/UbaObjectFileLLVMIR.cpp
+++ b/Engine/Source/Programs/UnrealBuildAccelerator/Common/Private/UbaObjectFileLLVMIR.cpp
@@ -664,7 +664,7 @@ namespace uba
 				u32 code = ReadVBR(6);
 
 				if (CanBeExported(blockId, code))
-					m_owner.m_globalVarOrFunctionRecords.push_back(BitStreamEntry{u64(m_pos - m_begin), m_word, m_wordBits, code, false});
+					m_owner.m_globalVarOrFunctionRecords.push_back(BitStreamEntry{u64(m_pos - m_begin), m_word, m_wordBits, code, DllStorage_None, false});
 
 				u32 numElts = ReadVBR(6);
 
@@ -700,7 +700,7 @@ namespace uba
 			if (CanBeExported(blockId, code))
 			{
 				UBA_ASSERT(!abbv.operands.empty());
-				m_owner.m_globalVarOrFunctionRecords.push_back(BitStreamEntry{u64(m_pos - m_begin), m_word, m_wordBits, code, false, abbv.operands });
+				m_owner.m_globalVarOrFunctionRecords.push_back(BitStreamEntry{u64(m_pos - m_begin), m_word, m_wordBits, code, DllStorage_None, false, abbv.operands });
 			}
 
 #if UBA_LOG_STREAM
@@ -900,11 +900,14 @@ namespace uba
 			UBA_ASSERT(!record.isExport);
 
 			if (dllStorage == 1)
+			{
 				m_imports.emplace(name);
+				record.dllStorage = DllStorage_Import;
+			}
 			else if (dllStorage == 2)
 			{
 				UBA_ASSERT(!name.empty());
-				record.isExport = true;
+				record.dllStorage = DllStorage_Export;
 				m_exports.emplace(name, ExportInfo{"", index++});
 			}
 			else // Check visibility... might be a non-windows based llvm ir stream file
@@ -932,13 +935,32 @@ namespace uba
 
 					if (linkage == 5 || linkage == 6 || linkage == 15)
 					{
-						record.isExport = true;
+						record.dllStorage = DllStorage_Export;
 						m_exports.emplace(name, ExportInfo{"", index++});
 					}
 					else if (linkage == 1 || linkage == 16)
+					{
+						record.dllStorage = DllStorage_Import;
 						m_imports.emplace(name);
+					}
 				}
 			}
+
+			if (record.dllStorage == DllStorage_Import)
+			{
+				// microsoft console
+				if (name == "?id@?$codecvt@DDU_Mbstatet@@@std@@2V0locale@2@A"
+				 || name == "?_Id_cnt@id@locale@std@@0HA"
+				 || name == "?cerr@std@@3V?$basic_ostream@DU?$char_traits@D@std@@@1@A"
+					)
+					record.keepAsIs = true;
+
+				// win64
+				if (name == "?_Raise_handler@std@@3P6AXAEBVexception@stdext@@@ZEA"
+				 || name == "?ROLE_INTERCHANGE_SCENE@OpenColorIO_v2_3@@3PEBDEB"
+					)
+					record.keepAsIs = true;
+			}
 		}
 		return true;
 	}
@@ -1018,7 +1040,7 @@ namespace uba
 
 		for (auto& record : m_globalVarOrFunctionRecords)
 		{
-			if (!record.isExport)
+			if (record.dllStorage == DllStorage_None || record.keepAsIs)
 				continue;
 	
 			BitStreamReader reader(*this, logger, newData, m_dataSize);
diff --git a/Engine/Source/Programs/UnrealBuildAccelerator/Common/Private/UbaObjectFileLLVMIR.h b/Engine/Source/Programs/UnrealBuildAccelerator/Common/Private/UbaObjectFileLLVMIR.h
--- a/Engine/Source/Programs/UnrealBuildAccelerator/Common/Private/UbaObjectFileLLVMIR.h
+++ b/Engine/Source/Programs/UnrealBuildAccelerator/Common/Private/UbaObjectFileLLVMIR.h
@@ -30,13 +30,21 @@ namespace uba
 		class BitStreamReader;
 		bool ParseBlock(Logger& logger, BitStreamReader& reader, BlockInfo& blockInfo, u32 blockId, u32 indent);
 
+		enum DllStorage : u8
+		{
+			DllStorage_Export,
+			DllStorage_Import,
+			DllStorage_None,
+		};
+
 		struct BitStreamEntry
 		{
 			u64 pos;
 			u32 word;
 			u32 wordBits;
 			u32 code;
-			bool isExport;
+			DllStorage dllStorage;
+			u8 keepAsIs;
 			Vector<AbbrevOp> operands;
 		};
 
