[UBA]
* Fixed uninitialized variable in UbaCacheServer that I suspect can be the cause of crashes on the farm. When a new bucket is created there is a chance that a pointer is not null causing a crash on next maintenance run

diff --git a/Engine/Source/Programs/UnrealBuildAccelerator/Common/Private/UbaCacheServer.cpp b/Engine/Source/Programs/UnrealBuildAccelerator/Common/Private/UbaCacheServer.cpp
--- a/Engine/Source/Programs/UnrealBuildAccelerator/Common/Private/UbaCacheServer.cpp
+++ b/Engine/Source/Programs/UnrealBuildAccelerator/Common/Private/UbaCacheServer.cpp
@@ -158,7 +158,7 @@ namespace uba
 			BitArray deletedOffsets;
 			bool isInitialized = false;
 			bool shouldTest = false;
-		}* m_maintenanceContext;
+		}* m_maintenanceContext = nullptr;
 	};
 
 	const tchar* ToString(CacheMessageType type)
@@ -1101,7 +1101,7 @@ namespace uba
 		u64 oldestTime = oldest ? GetFileTimeAsTime(now - (m_creationTime + oldest)) : 0;
 		u64 longestUnusedTime = longestUnused ? GetFileTimeAsTime(now - (m_creationTime + longestUnused)) : 0;
 		u64 duration = GetTime() - startTime;
-		m_logger.Info(TC("Maintenance done! (%s) CasFiles: %s (%s) Entries: %s Oldest: %s LongestUnused: %s MaintenanceMem: %s"), TimeToText(duration).str, CountToText(totalCasCount - deletedCasCount).str, BytesToText(totalCasSize).str, CountToText(totalEntryCount.load()).str, TimeToText(oldestTime, true).str, TimeToText(longestUnusedTime, true).str, BytesToText(maxCommittedMemory).str);
+		m_logger.Info(TC("Maintenance done! (%s) CasFiles: %s (%s) Entries: %s Oldest: %s LongestUnused: %s MaintenanceMem: %s/%s"), TimeToText(duration).str, CountToText(totalCasCount - deletedCasCount).str, BytesToText(totalCasSize).str, CountToText(totalEntryCount.load()).str, TimeToText(oldestTime, true).str, TimeToText(longestUnusedTime, true).str, BytesToText(maxCommittedMemory).str, BytesToText(m_maintenanceReserveSize).str);
 		
 		m_longestMaintenance = Max(m_longestMaintenance, duration);
 
