[UBA]
* Added todo comments to uba cache server in order to make maintenance faster

diff --git a/Engine/Source/Programs/UnrealBuildAccelerator/Common/Private/UbaCacheServer.cpp b/Engine/Source/Programs/UnrealBuildAccelerator/Common/Private/UbaCacheServer.cpp
--- a/Engine/Source/Programs/UnrealBuildAccelerator/Common/Private/UbaCacheServer.cpp
+++ b/Engine/Source/Programs/UnrealBuildAccelerator/Common/Private/UbaCacheServer.cpp
@@ -10,7 +10,11 @@
 #include "UbaStorageServer.h"
 //#include <oodle2.h>
 
-// MAKE IT RECOVER TIME WHEN GETTING UNDER CERTAIN SIZE!
+// TODO
+// - Fix so expiration time is set to oldest if overflowing and decreasing time didn't cause any deletes. That way we can make sure next maintenance will delete entries
+// - Sort buckets by last maintenance time to make sure the long ones always get a slot first
+// - Change so save happens when bucket is done in same work to minimize latency for the long ones
+
 
 #if PLATFORM_WINDOWS
 #define UBA_FORCEINLINE __forceinline
