[UBA]
* Changed so existingCas lookup table is using memory block too

diff --git a/Engine/Source/Programs/UnrealBuildAccelerator/Common/Private/UbaCacheServer.cpp b/Engine/Source/Programs/UnrealBuildAccelerator/Common/Private/UbaCacheServer.cpp
--- a/Engine/Source/Programs/UnrealBuildAccelerator/Common/Private/UbaCacheServer.cpp
+++ b/Engine/Source/Programs/UnrealBuildAccelerator/Common/Private/UbaCacheServer.cpp
@@ -601,7 +601,14 @@ namespace uba
 		u64 totalCasSize = 0;
 
 		struct CasFileInfo { CasFileInfo(u64 s = 0) : size(s) {} u64 size; Atomic<bool> isUsed; };
-		UnorderedMap<CasKey, CasFileInfo> existingCas;
+
+		// Existing cas entries can be more than 2 million entries.. which uses a lot of memory
+		constexpr u64 existingCasMemoryReserveSize = 192*1024*1024;
+		MemoryBlock existingCasMemoryBlock;
+		if (!existingCasMemoryBlock.Init(existingCasMemoryReserveSize, nullptr, true))
+			existingCasMemoryBlock.Init(existingCasMemoryReserveSize);
+
+		GrowingNoLockUnorderedMap<CasKey, CasFileInfo> existingCas(&existingCasMemoryBlock);
 
 		m_storage.WaitForActiveWork();
 
diff --git a/Engine/Source/Programs/UnrealBuildAccelerator/Common/Private/UbaObjectFileLLVMIR.cpp b/Engine/Source/Programs/UnrealBuildAccelerator/Common/Private/UbaObjectFileLLVMIR.cpp
--- a/Engine/Source/Programs/UnrealBuildAccelerator/Common/Private/UbaObjectFileLLVMIR.cpp
+++ b/Engine/Source/Programs/UnrealBuildAccelerator/Common/Private/UbaObjectFileLLVMIR.cpp
@@ -897,7 +897,7 @@ namespace uba
 			if (recSize > dllStorageIndex)
 				dllStorage = recIt[dllStorageIndex];
 
-			UBA_ASSERT(!record.isExport);
+			UBA_ASSERT(record.dllStorage != DllStorage_Export);
 
 			if (dllStorage == 1)
 			{
