[UBA]
* Changed UbaCacheServer maintenance code to use custom written hash map in order to improve performance

diff --git a/Engine/Source/Programs/UnrealBuildAccelerator/Common/Private/UbaCacheEntry.cpp b/Engine/Source/Programs/UnrealBuildAccelerator/Common/Private/UbaCacheEntry.cpp
--- a/Engine/Source/Programs/UnrealBuildAccelerator/Common/Private/UbaCacheEntry.cpp
+++ b/Engine/Source/Programs/UnrealBuildAccelerator/Common/Private/UbaCacheEntry.cpp
@@ -2,6 +2,7 @@
 
 #include "UbaCacheEntry.h"
 #include "UbaFile.h"
+#include "UbaHashMap.h"
 #include <algorithm>
 
 namespace uba
@@ -435,7 +436,7 @@ namespace uba
 		BuildInputsT(entry, inputs, entries.empty());
 	}
 
-	void CacheEntries::UpdateEntries(Logger& logger, const GrowingNoLockUnorderedMap<u32, u32>& oldToNewCasKeyOffset, Vector<u32>& temp, Vector<u8>& temp2)
+	void CacheEntries::UpdateEntries(Logger& logger, const HashMap2<u32, u32>& oldToNewCasKeyOffset, Vector<u32>& temp, Vector<u8>& temp2)
 	{
 		if (entries.empty())
 			return;
@@ -448,9 +449,8 @@ namespace uba
 				while (reader.GetLeft())
 				{
 					u32 newOffset = u32(reader.Read7BitEncoded());
-					auto findIt = oldToNewCasKeyOffset.find(newOffset);
-					if (findIt != oldToNewCasKeyOffset.end())
-						newOffset = findIt->second;
+					if (auto o = oldToNewCasKeyOffset.Find(newOffset))
+						newOffset = *o;
 					temp.push_back(newOffset);
 					newOffsetsSize += Get7BitEncodedCount(newOffset);
 				}
@@ -496,9 +496,8 @@ namespace uba
 					u64 newSize = 0;
 					for (auto& offset : temp)
 					{
-						auto findIt = oldToNewCasKeyOffset.find(offset);
-						if (findIt != oldToNewCasKeyOffset.end())
-							offset = findIt->second;
+						if (auto o = oldToNewCasKeyOffset.Find(offset))
+							offset = *o;
 						newSize += Get7BitEncodedCount(offset);
 					}
 
@@ -521,11 +520,8 @@ namespace uba
 					// Flatten using old shared and rebuild it with new shared
 					Flatten(temp, entry, oldShared);
 					for (auto& offset : temp)
-					{
-						auto findIt = oldToNewCasKeyOffset.find(offset);
-						if (findIt != oldToNewCasKeyOffset.end())
-							offset = findIt->second;
-					}
+						if (auto o = oldToNewCasKeyOffset.Find(offset))
+							offset = *o;
 
 					// Sort temp now when it likely is out of order
 					std::sort(temp.begin(), temp.end());
@@ -556,9 +552,8 @@ namespace uba
 						while (excludedReader.GetLeft())
 						{
 							u32 offset = u32(excludedReader.Read7BitEncoded());
-							auto findIt = oldToNewCasKeyOffset.find(offset);
-							if (findIt != oldToNewCasKeyOffset.end())
-								offset = findIt->second;
+							if (auto o = oldToNewCasKeyOffset.Find(offset))
+								offset = *o;
 							out.push_back(offset);
 						}
 					};
diff --git a/Engine/Source/Programs/UnrealBuildAccelerator/Common/Private/UbaCacheServer.cpp b/Engine/Source/Programs/UnrealBuildAccelerator/Common/Private/UbaCacheServer.cpp
--- a/Engine/Source/Programs/UnrealBuildAccelerator/Common/Private/UbaCacheServer.cpp
+++ b/Engine/Source/Programs/UnrealBuildAccelerator/Common/Private/UbaCacheServer.cpp
@@ -6,6 +6,7 @@
 #include "UbaBinaryReaderWriter.h"
 #include "UbaDirectoryIterator.h"
 #include "UbaFileAccessor.h"
+#include "UbaHashMap.h"
 #include "UbaNetworkServer.h"
 #include "UbaStorageServer.h"
 //#include <oodle2.h>
@@ -600,7 +601,7 @@ namespace uba
 
 		u64 totalCasSize = 0;
 
-		struct CasFileInfo { CasFileInfo(u64 s = 0) : size(s) {} u64 size; Atomic<bool> isUsed; };
+		struct CasFileInfo { CasFileInfo(u32 s = 0) : size(s) {} u32 size; Atomic<bool> isUsed; }; // These are compressed cas, should never be over 4gb
 
 		// Existing cas entries can be more than 2 million entries.. which uses a lot of memory
 		constexpr u64 existingCasMemoryReserveSize = 192*1024*1024;
@@ -608,10 +609,12 @@ namespace uba
 		if (!existingCasMemoryBlock.Init(existingCasMemoryReserveSize, nullptr, true))
 			existingCasMemoryBlock.Init(existingCasMemoryReserveSize);
 
-		GrowingNoLockUnorderedMap<CasKey, CasFileInfo> existingCas(&existingCasMemoryBlock);
+		HashMap<CasKey, CasFileInfo> existingCas;
+
 
 		m_storage.WaitForActiveWork();
 
+		u64 totalCasCount;
 		{
 			u64 collectCasStartTime = GetTime();
 
@@ -620,7 +623,8 @@ namespace uba
 			// TODO: Make this cleaner... (inside UbaStorage instead)
 			SCOPED_WRITE_LOCK(m_storage.m_casLookupLock, lookupLock);
 			
-			existingCas.reserve(m_storage.m_casLookup.size());
+			totalCasCount = m_storage.m_casLookup.size();
+			existingCas.Init(existingCasMemoryBlock, totalCasCount);
 
 			for (auto i=m_storage.m_casLookup.begin(), e=m_storage.m_casLookup.end(); i!=e;)
 			{
@@ -633,7 +637,8 @@ namespace uba
 					continue;
 				}
 				totalCasSize += i->second.size;
-				existingCas.try_emplace(i->first, i->second.size);
+				UBA_ASSERT(i->second.size < ~0u);
+				existingCas.Insert(i->first).size = u32(i->second.size);
 				++i;
 			}
 			lookupLock.Leave();
@@ -641,11 +646,9 @@ namespace uba
 			if (removedNonExisting)
 				m_logger.Detail(TC("  Removed %s cas entries (marked as not existing)"), CountToText(removedNonExisting).str);
 
-			m_logger.Detail(TC("  Found %s (%s) cas files and %s deleted by overflow (%s)"), CountToText(existingCas.size()).str, BytesToText(totalCasSize).str, CountToText(deletedCasFiles.size()).str, TimeToText(GetTime() - collectCasStartTime).str);
+			m_logger.Detail(TC("  Found %s (%s) cas files and %s deleted by overflow (%s)"), CountToText(existingCas.Size()).str, BytesToText(totalCasSize).str, CountToText(deletedCasFiles.size()).str, TimeToText(GetTime() - collectCasStartTime).str);
 		}
 
-		u64 totalCasCount = existingCas.size() + deletedCasCount;
-
 		if (shouldExit())
 			return true;
 
@@ -832,10 +835,9 @@ namespace uba
 								CasKey casKey;
 								bucket.m_casKeyTable.GetKey(casKey, offset);
 								UBA_ASSERT(IsCompressed(casKey));
-								auto findIt = existingCas.find(casKey);
-								if (findIt != existingCas.end())
+								if (auto value = existingCas.Find(casKey))
 								{
-									touchedCas.push_back(&findIt->second.isUsed);
+									touchedCas.push_back(&value->isUsed);
 									continue;
 								}
 								deleteEntry = true;
@@ -893,19 +895,20 @@ namespace uba
 			// Reset deleted cas files and update it again..
 			deletedCasFiles.clear();
 
-			for (auto i=existingCas.begin(), e=existingCas.end(); i!=e;)
+			for (auto i=existingCas.ValuesBegin(), e=existingCas.ValuesEnd(); i!=e; ++i)
 			{
-				if (i->second.isUsed)
+				if (i->isUsed)
 				{
-					i->second.isUsed = false;
-					++i;
+					i->isUsed = false;
 					continue;
 				}
-				deletedCasFiles.insert(i->first);
+				const CasKey* key = existingCas.GetKey(i);
+				if (!key)
+					continue;
+				deletedCasFiles.insert(*key);
 				++deletedCasCount;
-				totalCasSize -= i->second.size;
-				i = existingCas.erase(i);
-				e = existingCas.end();
+				totalCasSize -= i->size;
+				existingCas.Erase(*key);
 			}
 
 			// Add drop cas as work so it can run in the background
@@ -995,11 +998,11 @@ namespace uba
 				});
 
 			// Build new path table based on used offsets
-			GrowingNoLockUnorderedMap<u32, u32> oldToNewPathOffset(&memoryBlock);
+			HashMap2<u32, u32> oldToNewPathOffset;
 			u32 oldSize = bucket.m_pathTable.GetSize();
 			{
 				CompactPathTable newPathTable(CachePathTableMaxSize, CompactPathTable::V1, bucket.m_pathTable.GetPathCount(), bucket.m_pathTable.GetSegmentCount());
-				oldToNewPathOffset.reserve(usedPathOffsets.CountSetBits());
+				oldToNewPathOffset.Init(memoryBlock, usedPathOffsets.CountSetBits());
 
 				StringBuffer<> temp;
 				usedPathOffsets.Traverse([&](u32 pathOffset)
@@ -1014,10 +1017,7 @@ namespace uba
 					#endif
 
 					if (pathOffset != newPathOffset)
-					{
-						auto res = oldToNewPathOffset.try_emplace(pathOffset, newPathOffset);
-						UBA_ASSERT(res.second);(void)res;
-					}
+						oldToNewPathOffset.Insert(pathOffset) = newPathOffset;
 				});
 				bucket.m_pathTable.Swap(newPathTable);
 			}
@@ -1026,10 +1026,10 @@ namespace uba
 
 			// Build new caskey table based on used offsets
 			u64 recreateCasKeyTableStart = GetTime();
-			GrowingNoLockUnorderedMap<u32, u32> oldToNewCasKeyOffset(&memoryBlock);
+			HashMap2<u32, u32> oldToNewCasKeyOffset;
 			oldSize = bucket.m_casKeyTable.GetSize();
 			{
-				oldToNewCasKeyOffset.reserve(usedCasKeyOffsetsCount);
+				oldToNewCasKeyOffset.Init(memoryBlock, usedCasKeyOffsetsCount);
 				CompactCasKeyTable newCasKeyTable(CacheCasKeyTableMaxSize, usedCasKeyOffsetsCount);
 				BinaryReader reader2(bucket.m_casKeyTable.GetMemory(), 0, bucket.m_casKeyTable.GetSize());
 				usedCasKeyOffsets.Traverse([&](u32 casKeyOffset)
@@ -1038,22 +1038,18 @@ namespace uba
 					u32 oldPathOffset = u32(reader2.Read7BitEncoded());
 					CasKey casKey = reader2.ReadCasKey();
 					u32 newPathOffset = oldPathOffset;
-					auto findIt = oldToNewPathOffset.find(oldPathOffset);
-					if (findIt != oldToNewPathOffset.end())
-						newPathOffset = findIt->second;
+					if (auto value = oldToNewPathOffset.Find(oldPathOffset))
+						newPathOffset = *value;
 					u32 newCasKeyOffset = newCasKeyTable.Add(casKey, newPathOffset);
 					if (casKeyOffset != newCasKeyOffset)
-					{
-						auto res = oldToNewCasKeyOffset.try_emplace(casKeyOffset, newCasKeyOffset);
-						UBA_ASSERT(res.second);(void)res;
-					}
+						oldToNewCasKeyOffset.Insert(casKeyOffset) = newCasKeyOffset;
 				});
 				bucket.m_casKeyTable.Swap(newCasKeyTable);
 			}
 			m_logger.Detail(TC("    Bucket %u Recreated caskey table. %s -> %s (%s)"), bucket.index, BytesToText(oldSize).str, BytesToText(bucket.m_casKeyTable.GetSize()).str, TimeToText(GetTime() - recreateCasKeyTableStart).str);
 
 
-			if (!oldToNewCasKeyOffset.empty())
+			if (oldToNewCasKeyOffset.Size() > 0)
 			{
 				// Update all casKeyOffsets
 				u64 updateEntriesStart = GetTime();
diff --git a/Engine/Source/Programs/UnrealBuildAccelerator/Common/Private/UbaHashMap.h b/Engine/Source/Programs/UnrealBuildAccelerator/Common/Private/UbaHashMap.h
new file mode 100644
--- /dev/null
+++ b/Engine/Source/Programs/UnrealBuildAccelerator/Common/Private/UbaHashMap.h
@@ -0,0 +1,245 @@
+// Copyright Epic Games, Inc. All Rights Reserved.
+
+#pragma once
+
+namespace uba
+{
+	template<typename Key, typename Value>
+	struct HashMap
+	{
+		struct Entry
+		{
+			Key key;
+			u32 next; // If this is 0 it means that it is unused, ~0u means it is the end of a chain, otherwise offset into entry array
+		};
+
+		void Init(MemoryBlock& memory, u64 maxSize)
+		{
+			u32 v = u32(maxSize);
+			v--;
+			v |= v >> 1;
+			v |= v >> 2;
+			v |= v >> 4;
+			v |= v >> 8;
+			v |= v >> 16;
+			v++;
+			u32 lookupSize = v;// << 1u;
+			mask = lookupSize - 1;
+			nextAvailableEntry = 1;
+
+			// This code relies on that allocated memory is zeroed out
+			lookup = (u32*)memory.Allocate(lookupSize*sizeof(u32), 1, TC(""));
+			entries = (Entry*)memory.Allocate((maxSize+1)*sizeof(Entry), alignof(Entry), TC(""));
+			values = (Value*)memory.Allocate((maxSize+1)*sizeof(Value), alignof(Value), TC(""));
+		}
+
+		Value& Insert(const Key& key)
+		{
+			u32 index = std::hash<Key>()(key) & mask;
+			u32 entryIndex = lookup[index];
+
+			if (entryIndex == 0)
+			{
+				entryIndex = nextAvailableEntry++;
+				lookup[index] = entryIndex;
+				Entry& entry = entries[entryIndex];
+				entry.key = key;
+				entry.next = ~0u;
+				return values[entryIndex];
+			}
+
+			while (true)
+			{
+				Entry& entry = entries[entryIndex];
+				if (entry.key == key)
+					return values[entryIndex];
+
+				if (entry.next != ~0u)
+				{
+					entryIndex = entry.next;
+					//static u32 counter = 0;
+					//printf("MULTIPLE: %u\n", counter++);
+					continue;
+				}
+
+				u32 newEntryIndex = entry.next = nextAvailableEntry++;
+				Entry& newEntry = entries[newEntryIndex];
+				newEntry.key = key;
+				newEntry.next = ~0u;
+				return values[newEntryIndex];
+			}
+		}
+
+		Value* Find(const Key& key) const
+		{
+			u32 index = std::hash<Key>()(key) & mask;
+			u32 entryIndex = lookup[index];
+			if (entryIndex == 0)
+				return nullptr;
+			Entry* entry = entries + entryIndex;
+			while (true)
+			{
+				if (entry->key == key)
+					return values + entryIndex;
+				entryIndex = entry->next;
+				if (entryIndex == ~0u)
+					return nullptr;
+				entry = entries + entryIndex;
+			}
+		}
+
+		const Key* GetKey(Value* value)
+		{
+			u32 pos = u32(value - values);
+			auto& entry = entries[pos];
+			if (entry.next == 0)
+				return nullptr;
+			return &entry.key;
+		}
+
+		u32 Size()
+		{
+			return nextAvailableEntry - 1;
+		}
+
+		void Erase(const Key& key)
+		{
+			u32 index = std::hash<Key>()(key) & mask;
+			u32 entryIndex = lookup[index];
+			if (entryIndex == 0)
+				return;
+			Entry* entry = entries + entryIndex;
+			u32* prevNext = nullptr;
+			while (true)
+			{
+				if (entry->key == key)
+				{
+					if (!prevNext)
+					{
+						if (entry->next == ~0u)
+							lookup[index] = 0;
+						else
+							lookup[index] = entry->next;
+					}
+					else
+						*prevNext = entry->next;
+					entry->next = 0;
+
+					// TODO: We could swap with last and reduce size
+
+					return;
+				}
+				entryIndex = entry->next;
+				if (entryIndex == ~0u)
+					return;
+				prevNext = &entry->next;
+				entry = entries + entryIndex;
+			}
+		}
+
+		Value* ValuesBegin() { return values + 1; }
+		Value* ValuesEnd() { return values + nextAvailableEntry; }
+
+		Entry* entries;
+		Value* values;
+
+		u32* lookup;
+		u32 mask;
+		u32 nextAvailableEntry;
+	};
+
+	template<typename Key, typename Value>
+	struct HashMap2
+	{
+		struct Entry
+		{
+			Key key;
+			Value value;
+			u32 next; // If this is 0 it means that it is unused, ~0u means it is the end of a chain, otherwise offset into entry array
+		};
+
+		void Init(MemoryBlock& memory, u64 maxSize)
+		{
+			u32 v = u32(maxSize);
+			v--;
+			v |= v >> 1;
+			v |= v >> 2;
+			v |= v >> 4;
+			v |= v >> 8;
+			v |= v >> 16;
+			v++;
+			u32 lookupSize = v;// << 1u;
+			mask = lookupSize - 1;
+			nextAvailableEntry = 1;
+
+			// This code relies on that allocated memory is zeroed out
+			lookup = (u32*)memory.Allocate(lookupSize*sizeof(u32), 1, TC(""));
+			entries = (Entry*)memory.Allocate((maxSize+1)*sizeof(Entry), alignof(Entry), TC(""));
+		}
+
+		Value& Insert(const Key& key)
+		{
+			u32 index = std::hash<Key>()(key) & mask;
+			u32 entryIndex = lookup[index];
+
+			if (entryIndex == 0)
+			{
+				entryIndex = nextAvailableEntry++;
+				lookup[index] = entryIndex;
+				Entry& entry = entries[entryIndex];
+				entry.key = key;
+				entry.next = ~0u;
+				return entry.value;
+			}
+
+			while (true)
+			{
+				Entry& entry = entries[entryIndex];
+				if (entry.key == key)
+					return entry.value;
+
+				if (entry.next != ~0u)
+				{
+					entryIndex = entry.next;
+					//static u32 counter = 0;
+					//printf("MULTIPLE: %u\n", counter++);
+					continue;
+				}
+
+				u32 newEntryIndex = entry.next = nextAvailableEntry++;
+				Entry& newEntry = entries[newEntryIndex];
+				newEntry.key = key;
+				newEntry.next = ~0u;
+				return newEntry.value;
+			}
+		}
+
+		Value* Find(const Key& key) const
+		{
+			u32 index = std::hash<Key>()(key) & mask;
+			u32 entryIndex = lookup[index];
+			if (entryIndex == 0)
+				return nullptr;
+			Entry* entry = entries + entryIndex;
+			while (true)
+			{
+				if (entry->key == key)
+					return &entry->value;
+				entryIndex = entry->next;
+				if (entryIndex == ~0u)
+					return nullptr;
+				entry = entries + entryIndex;
+			}
+		}
+
+		u32 Size()
+		{
+			return nextAvailableEntry - 1;
+		}
+
+		Entry* entries;
+		u32* lookup;
+		u32 mask;
+		u32 nextAvailableEntry;
+	};
+}
\ No newline at end of file
diff --git a/Engine/Source/Programs/UnrealBuildAccelerator/Common/Public/UbaCacheEntry.h b/Engine/Source/Programs/UnrealBuildAccelerator/Common/Public/UbaCacheEntry.h
--- a/Engine/Source/Programs/UnrealBuildAccelerator/Common/Public/UbaCacheEntry.h
+++ b/Engine/Source/Programs/UnrealBuildAccelerator/Common/Public/UbaCacheEntry.h
@@ -8,6 +8,8 @@
 
 namespace uba
 {
+	template<typename Key, typename Value> struct HashMap2;
+
 	enum LogLinesType : u8
 	{
 		LogLinesType_Empty,
@@ -43,7 +45,7 @@ namespace uba
 		bool Write(BinaryWriter& writer, u32 clientVersion, bool toDisk);
 		bool Read(Logger& logger, BinaryReader& reader, u32 databaseVersion);
 		void BuildInputs(CacheEntry& entry, const Set<u32>& inputs);
-		void UpdateEntries(Logger& logger, const GrowingNoLockUnorderedMap<u32, u32>& oldToNewCasKeyOffset, Vector<u32>& temp, Vector<u8>& temp2);
+		void UpdateEntries(Logger& logger, const HashMap2<u32, u32>& oldToNewCasKeyOffset, Vector<u32>& temp, Vector<u8>& temp2);
 
 		void Flatten(Vector<u8>& out, const CacheEntry& entry);
 		void Flatten(Vector<u32>& out, const CacheEntry& entry, const Vector<u8>& sharedOffsets);
diff --git a/Engine/Source/Programs/UnrealBuildAccelerator/Core/Public/UbaHash.h b/Engine/Source/Programs/UnrealBuildAccelerator/Core/Public/UbaHash.h
--- a/Engine/Source/Programs/UnrealBuildAccelerator/Core/Public/UbaHash.h
+++ b/Engine/Source/Programs/UnrealBuildAccelerator/Core/Public/UbaHash.h
@@ -203,7 +203,7 @@ namespace uba
 
 //template<> struct std::hash<uba::Guid> { size_t operator()(const uba::Guid& g) const { std::hash<uba::u64> hash; return hash(((uba::u64*)&g)[0]) ^ hash(((uba::u64*)&g)[1]); } };
 template<> struct std::hash<uba::StringKey> { size_t operator()(const uba::StringKey& g) const { return g.a; } };
-template<> struct std::hash<uba::CasKey> { size_t operator()(const uba::CasKey& g) const { return g.a ^ g.b ^ g.c; } };
+template<> struct std::hash<uba::CasKey> { size_t operator()(const uba::CasKey& g) const { return g.a; } };
 
 template <class Map> 
 void PrintMapInfo(const char* name, const Map& map)
diff --git a/Engine/Source/Programs/UnrealBuildAccelerator/Test/Private/UbaTestAll.h b/Engine/Source/Programs/UnrealBuildAccelerator/Test/Private/UbaTestAll.h
--- a/Engine/Source/Programs/UnrealBuildAccelerator/Test/Private/UbaTestAll.h
+++ b/Engine/Source/Programs/UnrealBuildAccelerator/Test/Private/UbaTestAll.h
@@ -33,6 +33,7 @@ namespace uba
 		UBA_TEST(TestRemoteDetouredTestApp) \
 		UBA_TEST(TestCompactPathTable) \
 		UBA_TEST(TestCacheEntry) \
+		UBA_TEST(TestHashTable) \
 		UBA_TEST(TestConfig) \
 		UBA_TEST(TestCrypto) \
 		UBA_TEST(TestBinDependencies) \
diff --git a/Engine/Source/Programs/UnrealBuildAccelerator/Test/Private/UbaTestCache.h b/Engine/Source/Programs/UnrealBuildAccelerator/Test/Private/UbaTestCache.h
--- a/Engine/Source/Programs/UnrealBuildAccelerator/Test/Private/UbaTestCache.h
+++ b/Engine/Source/Programs/UnrealBuildAccelerator/Test/Private/UbaTestCache.h
@@ -10,6 +10,7 @@
 #include "UbaNetworkServer.h"
 #include "UbaSessionServer.h"
 #include "UbaStorageServer.h"
+#include "../Private/UbaHashMap.h"
 
 namespace uba
 {
@@ -206,4 +207,88 @@ namespace uba
 		}
 		return true;
 	}
+
+	bool TestHashTable(Logger& logger, const StringBufferBase& rootDir)
+	{
+
+		struct CasFileInfo { CasFileInfo(u32 s = 0) : size(s) {} u32 size; Atomic<bool> isUsed; }; // These are compressed cas, should never be over 4gb
+
+		{
+			MemoryBlock memoryBlock(1024*1024);
+			HashMap<u32, u32> casMap;
+			if (casMap.Find(1))
+				return false;
+			casMap.Insert(1) = 2;
+			if (*casMap.Find(1) != 2)
+				return false;
+			casMap.Insert(1) = 3;
+			if (*casMap.Find(1) != 3)
+				return false;
+		}
+
+
+
+#if 0
+		constexpr u64 memoryReserveSize = 192*1024*1024;
+		struct ProfileScope
+		{
+			ProfileScope() : startTime(GetTime()) {}
+			~ProfileScope() { u64 duration = GetTime() - startTime; LoggerWithWriter(g_consoleLogWriter, TC("")).Info(TC("Time: %s"), TimeToText(duration).str); }
+			u64 startTime;
+		};
+
+
+		{
+			ProfileScope _;
+			MemoryBlock memoryBlock;
+			if (!memoryBlock.Init(memoryReserveSize, nullptr, true))
+				memoryBlock.Init(memoryReserveSize);
+
+			u32 totalCasCount = 1'800'000;
+			HashMap<CasKey, CasFileInfo> casMap;
+			casMap.Init(memoryBlock, totalCasCount);
+
+			CasKey key;
+
+			for (u32 i=0; i!=totalCasCount; ++i)
+			{
+				key.a = i;
+				casMap.Insert(key);
+			}
+
+			for (u32 i=0; i!=totalCasCount; ++i)
+			{
+				key.a = i;
+				casMap.Find(key);
+			}
+		}
+
+		{
+			ProfileScope _;
+			MemoryBlock memoryBlock;
+			if (!memoryBlock.Init(casMapMemoryReserveSize, nullptr, true))
+				memoryBlock.Init(casMapMemoryReserveSize);
+
+			u32 totalCasCount = 1'800'000;
+			GrowingNoLockUnorderedMap<CasKey, CasFileInfo> casMap(&memoryBlock);
+			casMap.reserve(totalCasCount);
+
+			CasKey key;
+
+			for (u32 i=0; i!=totalCasCount; ++i)
+			{
+				key.a = i;
+				casMap.try_emplace(key);
+			}
+
+			for (u32 i=0; i!=totalCasCount; ++i)
+			{
+				key.a = i;
+				auto it = casMap.find(key);
+			}
+		}
+#endif
+		return true;
+	}
+
 }
\ No newline at end of file
