[UBA]
* Fixed hash map unit test

diff --git a/Engine/Source/Programs/UnrealBuildAccelerator/Test/Private/UbaTestCache.h b/Engine/Source/Programs/UnrealBuildAccelerator/Test/Private/UbaTestCache.h
--- a/Engine/Source/Programs/UnrealBuildAccelerator/Test/Private/UbaTestCache.h
+++ b/Engine/Source/Programs/UnrealBuildAccelerator/Test/Private/UbaTestCache.h
@@ -210,12 +210,10 @@ namespace uba
 
 	bool TestHashTable(Logger& logger, const StringBufferBase& rootDir)
 	{
-
-		struct CasFileInfo { CasFileInfo(u32 s = 0) : size(s) {} u32 size; Atomic<bool> isUsed; }; // These are compressed cas, should never be over 4gb
-
 		{
 			MemoryBlock memoryBlock(1024*1024);
 			HashMap<u32, u32> casMap;
+			casMap.Init(memoryBlock, 3);
 			if (casMap.Find(1))
 				return false;
 			casMap.Insert(1) = 2;
@@ -226,9 +224,8 @@ namespace uba
 				return false;
 		}
 
-
-
 #if 0
+		struct CasFileInfo { CasFileInfo(u32 s = 0) : size(s) {} u32 size; bool isUsed; }; // These are compressed cas, should never be over 4gb
 		constexpr u64 memoryReserveSize = 192*1024*1024;
 		struct ProfileScope
 		{
@@ -237,6 +234,7 @@ namespace uba
 			u64 startTime;
 		};
 
+		u32 totalCasCount = 1'800'000;
 
 		{
 			ProfileScope _;
@@ -244,7 +242,6 @@ namespace uba
 			if (!memoryBlock.Init(memoryReserveSize, nullptr, true))
 				memoryBlock.Init(memoryReserveSize);
 
-			u32 totalCasCount = 1'800'000;
 			HashMap<CasKey, CasFileInfo> casMap;
 			casMap.Init(memoryBlock, totalCasCount);
 
@@ -266,10 +263,9 @@ namespace uba
 		{
 			ProfileScope _;
 			MemoryBlock memoryBlock;
-			if (!memoryBlock.Init(casMapMemoryReserveSize, nullptr, true))
-				memoryBlock.Init(casMapMemoryReserveSize);
+			if (!memoryBlock.Init(memoryReserveSize, nullptr, true))
+				memoryBlock.Init(memoryReserveSize);
 
-			u32 totalCasCount = 1'800'000;
 			GrowingNoLockUnorderedMap<CasKey, CasFileInfo> casMap(&memoryBlock);
 			casMap.reserve(totalCasCount);
 
@@ -290,5 +286,4 @@ namespace uba
 #endif
 		return true;
 	}
-
 }
\ No newline at end of file
