[UBA]
* Fixed so memory used to calculate existing cas entries are freed when not needed anymore

diff --git a/Engine/Source/Programs/UnrealBuildAccelerator/Common/Private/UbaCacheServer.cpp b/Engine/Source/Programs/UnrealBuildAccelerator/Common/Private/UbaCacheServer.cpp
--- a/Engine/Source/Programs/UnrealBuildAccelerator/Common/Private/UbaCacheServer.cpp
+++ b/Engine/Source/Programs/UnrealBuildAccelerator/Common/Private/UbaCacheServer.cpp
@@ -921,6 +921,8 @@ namespace uba
 		}
 		while (!deletedCasFiles.empty()); // if cas files are deleted we need to do another loop and check cache entry inputs to see if files were inputs
 
+		existingCasMemoryBlock.Deinit();
+
 		if (overflowedEntryCount)
 			m_logger.Detail(TC("  Found %llu overflowed cache entries"), overflowedEntryCount.load());
 		if (expiredEntryCount)
