[UBA] Add missing target platforms to UbaCommon.Target.cs script.

Adds the missing "[SupportedPlatforms(UnrealPlatformClass.Desktop)]" attribute to the target rules build script for the UbaCommon module.

#rnx

diff --git a/Engine/Source/Programs/UnrealBuildAccelerator/UbaCommon.Target.cs b/Engine/Source/Programs/UnrealBuildAccelerator/UbaCommon.Target.cs
--- a/Engine/Source/Programs/UnrealBuildAccelerator/UbaCommon.Target.cs
+++ b/Engine/Source/Programs/UnrealBuildAccelerator/UbaCommon.Target.cs
@@ -2,7 +2,7 @@
 
 using UnrealBuildTool;
 
-[SupportedPlatforms()]
+[SupportedPlatforms(UnrealPlatformClass.Desktop)]
 public class UbaCommonTarget : TargetRules
 {
 	public UbaCommonTarget(TargetInfo Target) : base(Target)
