[UBA] Also add missing target platforms to UbaCore.Target.cs script. Follow up to 36152490.

#rnx

diff --git a/Engine/Source/Programs/UnrealBuildAccelerator/UbaCore.Target.cs b/Engine/Source/Programs/UnrealBuildAccelerator/UbaCore.Target.cs
--- a/Engine/Source/Programs/UnrealBuildAccelerator/UbaCore.Target.cs
+++ b/Engine/Source/Programs/UnrealBuildAccelerator/UbaCore.Target.cs
@@ -2,7 +2,7 @@
 
 using UnrealBuildTool;
 
-[SupportedPlatforms()]
+[SupportedPlatforms(UnrealPlatformClass.Desktop)]
 public class UbaCoreTarget : TargetRules
 {
 	public UbaCoreTarget(TargetInfo Target) : base(Target)
