[UBA]
* Fixed so code can compile with static analyser

diff --git a/Engine/Source/Programs/UnrealBuildAccelerator/Cli/Private/UbaCli.cpp b/Engine/Source/Programs/UnrealBuildAccelerator/Cli/Private/UbaCli.cpp
--- a/Engine/Source/Programs/UnrealBuildAccelerator/Cli/Private/UbaCli.cpp
+++ b/Engine/Source/Programs/UnrealBuildAccelerator/Cli/Private/UbaCli.cpp
@@ -843,9 +843,9 @@ namespace uba
 
 			if (populateCache)
 			{
-				logger.Error(TC("Populating cache not implemented... todo"));
-				RootPaths rootPaths;
-				cacheClient->WriteToCache(rootPaths, 0, pinfo, nullptr, 1, nullptr, 0, nullptr, 0);
+				return logger.Error(TC("Populating cache not implemented... todo"));
+				//RootPaths rootPaths;
+				//cacheClient->WriteToCache(rootPaths, 0, pinfo, nullptr, 1, nullptr, 0, nullptr, 0);
 			}
 			return true;
 		};
diff --git a/Engine/Source/Programs/UnrealBuildAccelerator/Common/Private/UbaCacheClient.cpp b/Engine/Source/Programs/UnrealBuildAccelerator/Common/Private/UbaCacheClient.cpp
--- a/Engine/Source/Programs/UnrealBuildAccelerator/Common/Private/UbaCacheClient.cpp
+++ b/Engine/Source/Programs/UnrealBuildAccelerator/Common/Private/UbaCacheClient.cpp
@@ -380,7 +380,7 @@ namespace uba
 
 	u64 CacheClient::MakeId(u32 bucketId)
 	{
-		return u64(bucketId) | ((u64(!CaseInsensitiveFs) + (RootPathsVersion << 1) + (u8(!m_useRoots) << 2)) << 32);
+		return u64(bucketId) | ((u64(!CaseInsensitiveFs) + (RootPathsVersion << 1u) + (u64(!m_useRoots) << 2u)) << 32u);
 	}
 
 
diff --git a/Engine/Source/Programs/UnrealBuildAccelerator/Common/Private/UbaFile.cpp b/Engine/Source/Programs/UnrealBuildAccelerator/Common/Private/UbaFile.cpp
--- a/Engine/Source/Programs/UnrealBuildAccelerator/Common/Private/UbaFile.cpp
+++ b/Engine/Source/Programs/UnrealBuildAccelerator/Common/Private/UbaFile.cpp
@@ -1,7 +1,5 @@
 // Copyright Epic Games, Inc. All Rights Reserved.
 
-#define Local_GetLongPathNameW uba::GetLongPathNameW
-
 #include "UbaFile.h"
 #include "UbaDirectoryIterator.h"
 #include "UbaEvent.h"
diff --git a/Engine/Source/Programs/UnrealBuildAccelerator/Common/Private/UbaHashMap.h b/Engine/Source/Programs/UnrealBuildAccelerator/Common/Private/UbaHashMap.h
--- a/Engine/Source/Programs/UnrealBuildAccelerator/Common/Private/UbaHashMap.h
+++ b/Engine/Source/Programs/UnrealBuildAccelerator/Common/Private/UbaHashMap.h
@@ -2,6 +2,8 @@
 
 #pragma once
 
+#include "UbaMemory.h"
+
 namespace uba
 {
 	template<typename Key, typename Value>
diff --git a/Engine/Source/Programs/UnrealBuildAccelerator/Common/Private/UbaObjectFile.cpp b/Engine/Source/Programs/UnrealBuildAccelerator/Common/Private/UbaObjectFile.cpp
--- a/Engine/Source/Programs/UnrealBuildAccelerator/Common/Private/UbaObjectFile.cpp
+++ b/Engine/Source/Programs/UnrealBuildAccelerator/Common/Private/UbaObjectFile.cpp
@@ -106,6 +106,9 @@ namespace uba
 	bool ObjectFile::CopyMemoryAndClose()
 	{
 		u8* data = (u8*)malloc(m_dataSize);
+		UBA_ASSERT(false);
+		if (!data)
+			return false;
 		memcpy(data, m_data, m_dataSize);
 		if (m_ownsData)
 			free(m_data);
diff --git a/Engine/Source/Programs/UnrealBuildAccelerator/Common/Private/UbaObjectFileLLVMIR.cpp b/Engine/Source/Programs/UnrealBuildAccelerator/Common/Private/UbaObjectFileLLVMIR.cpp
--- a/Engine/Source/Programs/UnrealBuildAccelerator/Common/Private/UbaObjectFileLLVMIR.cpp
+++ b/Engine/Source/Programs/UnrealBuildAccelerator/Common/Private/UbaObjectFileLLVMIR.cpp
@@ -238,7 +238,11 @@ namespace uba
 		{
 			if (m_wordBits >= bits)
 			{
-				u32 res = m_word & (~0u >> (32 - bits));
+				#if !defined( __clang_analyzer__ )
+				u32 res = m_word & (~0u >> (32u - bits));
+				#else
+				u32 res = 0;
+				#endif
 				m_word >>= (bits & 0x1f);
 				m_wordBits -= bits;
 				return res;
@@ -1099,7 +1103,7 @@ namespace uba
 
 			if (!exportStr)
 			{
-				readPos = readEnd;
+				//readPos = readEnd;
 				break;
 			}
 
diff --git a/Engine/Source/Programs/UnrealBuildAccelerator/Common/Private/UbaObjectFileLLVMIR.h b/Engine/Source/Programs/UnrealBuildAccelerator/Common/Private/UbaObjectFileLLVMIR.h
--- a/Engine/Source/Programs/UnrealBuildAccelerator/Common/Private/UbaObjectFileLLVMIR.h
+++ b/Engine/Source/Programs/UnrealBuildAccelerator/Common/Private/UbaObjectFileLLVMIR.h
@@ -3,6 +3,7 @@
 #pragma once
 
 #include "UbaObjectFile.h"
+#include <memory>
 
 namespace uba
 {
diff --git a/Engine/Source/Programs/UnrealBuildAccelerator/Common/Private/UbaProcess.cpp b/Engine/Source/Programs/UnrealBuildAccelerator/Common/Private/UbaProcess.cpp
--- a/Engine/Source/Programs/UnrealBuildAccelerator/Common/Private/UbaProcess.cpp
+++ b/Engine/Source/Programs/UnrealBuildAccelerator/Common/Private/UbaProcess.cpp
@@ -603,11 +603,17 @@ namespace uba
 		const tchar* fromFileBegin = m_startInfo.arguments;
 		const tchar* fromFileEnd = TStrchr(fromFileBegin, '\"');
 		UBA_ASSERT(fromFileEnd);
+		if (!fromFileEnd)
+			return false;
 		const tchar* toFileBegin = TStrchr(fromFileEnd + 1, '\"');
 		UBA_ASSERT(toFileBegin);
+		if (!toFileBegin)
+			return false;
 		++toFileBegin;
 		const tchar* toFileEnd = TStrchr(toFileBegin, '\"');
 		UBA_ASSERT(toFileEnd);
+		if (!toFileEnd)
+			return false;
 
 		m_processStats.wallTime = GetTime() - m_startTime;
 
diff --git a/Engine/Source/Programs/UnrealBuildAccelerator/Common/Public/UbaProcessStartInfoHolder.h b/Engine/Source/Programs/UnrealBuildAccelerator/Common/Public/UbaProcessStartInfoHolder.h
--- a/Engine/Source/Programs/UnrealBuildAccelerator/Common/Public/UbaProcessStartInfoHolder.h
+++ b/Engine/Source/Programs/UnrealBuildAccelerator/Common/Public/UbaProcessStartInfoHolder.h
@@ -2,8 +2,6 @@
 
 #pragma once
 
-#define Local_GetLongPathNameW uba::GetLongPathNameW
-
 #include "UbaBinaryReaderWriter.h"
 #include "UbaFile.h"
 #include "UbaPathUtils.h"
diff --git a/Engine/Source/Programs/UnrealBuildAccelerator/Common/Public/UbaSession.h b/Engine/Source/Programs/UnrealBuildAccelerator/Common/Public/UbaSession.h
--- a/Engine/Source/Programs/UnrealBuildAccelerator/Common/Public/UbaSession.h
+++ b/Engine/Source/Programs/UnrealBuildAccelerator/Common/Public/UbaSession.h
@@ -2,8 +2,6 @@
 
 #pragma once
 
-#define Local_GetLongPathNameW uba::GetLongPathNameW
-
 #include "UbaFile.h"
 #include "UbaDirectoryTable.h"
 #include "UbaFileMapping.h"
diff --git a/Engine/Source/Programs/UnrealBuildAccelerator/Common/Public/UbaStorage.h b/Engine/Source/Programs/UnrealBuildAccelerator/Common/Public/UbaStorage.h
--- a/Engine/Source/Programs/UnrealBuildAccelerator/Common/Public/UbaStorage.h
+++ b/Engine/Source/Programs/UnrealBuildAccelerator/Common/Public/UbaStorage.h
@@ -2,8 +2,6 @@
 
 #pragma once
 
-#define Local_GetLongPathNameW uba::GetLongPathNameW
-
 #include "UbaFile.h"
 #include "UbaFileMapping.h"
 #include "UbaLogger.h"
diff --git a/Engine/Source/Programs/UnrealBuildAccelerator/Core/Private/UbaMemory.cpp b/Engine/Source/Programs/UnrealBuildAccelerator/Core/Private/UbaMemory.cpp
--- a/Engine/Source/Programs/UnrealBuildAccelerator/Core/Private/UbaMemory.cpp
+++ b/Engine/Source/Programs/UnrealBuildAccelerator/Core/Private/UbaMemory.cpp
@@ -168,7 +168,7 @@ namespace uba
 
 	bool SupportsHugePages()
 	{
-		return UBA_SUPPORTS_HUGE_PAGES != 0;
+		return UBA_SUPPORTS_HUGE_PAGES;
 	}
 	u64 GetHugePageCount()
 	{
diff --git a/Engine/Source/Programs/UnrealBuildAccelerator/Core/Public/UbaBase.h b/Engine/Source/Programs/UnrealBuildAccelerator/Core/Public/UbaBase.h
--- a/Engine/Source/Programs/UnrealBuildAccelerator/Core/Public/UbaBase.h
+++ b/Engine/Source/Programs/UnrealBuildAccelerator/Core/Public/UbaBase.h
@@ -30,6 +30,10 @@ namespace uba
 	#define TC(x) x
 	#define PERCENT_HS "%s"
 	#endif
+
+	#if PLATFORM_WINDOWS && defined(UBA_BUILD)
+	#pragma warning(disable:4100) // This is needed because of single header compiles where AdditionalCompilerArguments is not included
+	#endif
 }
 
 #if !defined(UBA_API)
diff --git a/Engine/Source/Programs/UnrealBuildAccelerator/Core/Public/UbaMacBinDependencyParser.h b/Engine/Source/Programs/UnrealBuildAccelerator/Core/Public/UbaMacBinDependencyParser.h
--- a/Engine/Source/Programs/UnrealBuildAccelerator/Core/Public/UbaMacBinDependencyParser.h
+++ b/Engine/Source/Programs/UnrealBuildAccelerator/Core/Public/UbaMacBinDependencyParser.h
@@ -2,6 +2,8 @@
 
 #pragma once
 
+#if PLATFORM_MAC
+
 #include "UbaStringBuffer.h"
 
 #if !defined(UBA_IS_DETOURED_INCLUDE)
@@ -85,4 +87,6 @@ namespace uba
 			func(libs[i], false, loaderPaths);
 		return true;
 	}
-}
\ No newline at end of file
+}
+
+#endif
\ No newline at end of file
diff --git a/Engine/Source/Programs/UnrealBuildAccelerator/Core/Public/UbaPathUtils.h b/Engine/Source/Programs/UnrealBuildAccelerator/Core/Public/UbaPathUtils.h
--- a/Engine/Source/Programs/UnrealBuildAccelerator/Core/Public/UbaPathUtils.h
+++ b/Engine/Source/Programs/UnrealBuildAccelerator/Core/Public/UbaPathUtils.h
@@ -7,6 +7,12 @@
 
 namespace uba
 {
+	#if !defined(UBA_IS_DETOURED_INCLUDE)
+	u32 GetLongPathNameW(const tchar* lpszShortPath, tchar* lpszLongPath, u32 cchBuffer);
+	#define Local_GetLongPathNameW GetLongPathNameW
+	#endif
+
+
 	inline bool FixPath2(const tchar* fileName, const tchar* workingDir, u64 workingDirCharLen, tchar* buffer, u64 bufferCharCapacity, u32* outBufferCharLen)
 	{
 		buffer[1] = 0;
diff --git a/Engine/Source/Programs/UnrealBuildAccelerator/Core/Public/UbaProcessUtils.h b/Engine/Source/Programs/UnrealBuildAccelerator/Core/Public/UbaProcessUtils.h
--- a/Engine/Source/Programs/UnrealBuildAccelerator/Core/Public/UbaProcessUtils.h
+++ b/Engine/Source/Programs/UnrealBuildAccelerator/Core/Public/UbaProcessUtils.h
@@ -2,6 +2,7 @@
 
 #pragma once
 
+#include "UbaPlatform.h"
 #include "UbaStringBuffer.h"
 
 namespace uba
diff --git a/Engine/Source/Programs/UnrealBuildAccelerator/Core/Public/UbaWinBinDependencyParser.h b/Engine/Source/Programs/UnrealBuildAccelerator/Core/Public/UbaWinBinDependencyParser.h
--- a/Engine/Source/Programs/UnrealBuildAccelerator/Core/Public/UbaWinBinDependencyParser.h
+++ b/Engine/Source/Programs/UnrealBuildAccelerator/Core/Public/UbaWinBinDependencyParser.h
@@ -2,6 +2,7 @@
 
 #pragma once
 
+#include "UbaPlatform.h"
 #include "UbaStringBuffer.h"
 #include <algorithm>
 
diff --git a/Engine/Source/Programs/UnrealBuildAccelerator/Detours/Private/UbaDetoursShared.cpp b/Engine/Source/Programs/UnrealBuildAccelerator/Detours/Private/UbaDetoursShared.cpp
--- a/Engine/Source/Programs/UnrealBuildAccelerator/Detours/Private/UbaDetoursShared.cpp
+++ b/Engine/Source/Programs/UnrealBuildAccelerator/Detours/Private/UbaDetoursShared.cpp
@@ -1,5 +1,7 @@
 // Copyright Epic Games, Inc. All Rights Reserved.
 
+#define UBA_IS_DETOURED_INCLUDE 1
+
 #include "UbaDetoursShared.h"
 #include "UbaDetoursFileMappingTable.h"
 #include "UbaDirectoryTable.h"
diff --git a/Engine/Source/Programs/UnrealBuildAccelerator/Detours/Private/Windows/UbaDetoursFunctionsShell32.inl b/Engine/Source/Programs/UnrealBuildAccelerator/Detours/Private/Windows/UbaDetoursFunctionsShell32.inl
--- a/Engine/Source/Programs/UnrealBuildAccelerator/Detours/Private/Windows/UbaDetoursFunctionsShell32.inl
+++ b/Engine/Source/Programs/UnrealBuildAccelerator/Detours/Private/Windows/UbaDetoursFunctionsShell32.inl
@@ -5,7 +5,7 @@
 using SHGetKnownFolderPathFunc = HRESULT(REFKNOWNFOLDERID rfid, DWORD dwFlags, HANDLE hToken, PWSTR* ppszPath);
 SHGetKnownFolderPathFunc* True_SHGetKnownFolderPath;
 
-BOOL Detoured_SHGetKnownFolderPath(REFKNOWNFOLDERID rfid, DWORD dwFlags, HANDLE hToken, PWSTR* ppszPath)
+HRESULT Detoured_SHGetKnownFolderPath(REFKNOWNFOLDERID rfid, DWORD dwFlags, HANDLE hToken, PWSTR* ppszPath)
 {
 	if (g_runningRemote)
 	{
@@ -26,6 +26,9 @@ BOOL Detoured_SHGetKnownFolderPath(REFKNOWNFOLDERID rfid, DWORD dwFlags, HANDLE
 			reader.ReadString(path);
 			u32 memSize = (path.count+1)*2;
 			void* mem = CoTaskMemAlloc(memSize);
+			UBA_ASSERT(mem);
+			if (!mem)
+				return E_FAIL;
 			memcpy(mem, path.data, memSize);
 			*ppszPath = (PWSTR)mem;
 		}
diff --git a/Engine/Source/Programs/UnrealBuildAccelerator/Detours/Private/Windows/UbaDetoursFunctionsUcrtBase.inl b/Engine/Source/Programs/UnrealBuildAccelerator/Detours/Private/Windows/UbaDetoursFunctionsUcrtBase.inl
--- a/Engine/Source/Programs/UnrealBuildAccelerator/Detours/Private/Windows/UbaDetoursFunctionsUcrtBase.inl
+++ b/Engine/Source/Programs/UnrealBuildAccelerator/Detours/Private/Windows/UbaDetoursFunctionsUcrtBase.inl
@@ -339,6 +339,7 @@ int Detoured__wspawnl(int mode, const wchar_t* cmdname, const wchar_t* arg0, con
 	return exitCode;
 }
 
+#if !defined( __clang_analyzer__ )
 errno_t Detoured__wsplitpath_s(const wchar_t* path, wchar_t* drive, size_t driveNumberOfElements, wchar_t* dir, size_t dirNumberOfElements, wchar_t* fname, size_t nameNumberOfElements, wchar_t* ext, size_t extNumberOfElements)
 {
 	DETOURED_CALL(_wsplitpath_s);
@@ -348,6 +349,7 @@ errno_t Detoured__wsplitpath_s(const wchar_t* path, wchar_t* drive, size_t drive
 	//DEBUG_LOG_TRUE(L"_wsplitpath_s", L"%ls %ls %ls %ls %ls", path, drive, dir, fname, ext);
 	return res;
 }
+#endif
 
 #if defined(DETOURED_INCLUDE_DEBUG)
 int Detoured__wcsnicoll_l(const wchar_t* string1, const wchar_t* string2, size_t count, _locale_t locale)
diff --git a/Engine/Source/Programs/UnrealBuildAccelerator/Detours/Private/Windows/UbaDetoursFunctionsWin.h b/Engine/Source/Programs/UnrealBuildAccelerator/Detours/Private/Windows/UbaDetoursFunctionsWin.h
--- a/Engine/Source/Programs/UnrealBuildAccelerator/Detours/Private/Windows/UbaDetoursFunctionsWin.h
+++ b/Engine/Source/Programs/UnrealBuildAccelerator/Detours/Private/Windows/UbaDetoursFunctionsWin.h
@@ -144,6 +144,12 @@
 #define DETOURED_FUNCTIONS_SHLWAPI \
 	DETOURED_FUNCTIONS_SHLWAPI_DEBUG
 
+#if !defined( __clang_analyzer__ )
+#define DETOURED_WSPLITPATH DETOURED_FUNCTION(_wsplitpath_s)
+#else
+#define DETOURED_WSPLITPATH
+#endif
+
 #define DETOURED_FUNCTIONS_UCRTBASE \
 	DETOURED_FUNCTION(_wgetcwd) \
 	DETOURED_FUNCTION(_wfullpath) \
@@ -155,7 +161,7 @@
 	DETOURED_FUNCTION(_isatty) \
 	DETOURED_FUNCTION(_write) \
 	DETOURED_FUNCTION(fputs) \
-	DETOURED_FUNCTION(_wsplitpath_s) \
+	DETOURED_WSPLITPATH \
 	DETOURED_FUNCTIONS_UCRTBASE_DEBUG \
 
 #if UBA_SUPPORT_MSPDBSRV
diff --git a/Engine/Source/Programs/UnrealBuildAccelerator/Detours/Public/UbaDetoursApi.h b/Engine/Source/Programs/UnrealBuildAccelerator/Detours/Public/UbaDetoursApi.h
--- a/Engine/Source/Programs/UnrealBuildAccelerator/Detours/Public/UbaDetoursApi.h
+++ b/Engine/Source/Programs/UnrealBuildAccelerator/Detours/Public/UbaDetoursApi.h
@@ -1,6 +1,6 @@
 // Copyright Epic Games, Inc. All Rights Reserved.
 
-#if UBA_IS_DETOURED_INCLUDE
+#if defined(UBA_IS_DETOURED_INCLUDE)
 #define UBA_DETOURED_API __declspec(dllexport)
 #else
 #define UBA_DETOURED_API __declspec(dllimport)
diff --git a/Engine/Source/Programs/UnrealBuildAccelerator/ObjTool/Private/UbaImportLibWriter.cpp b/Engine/Source/Programs/UnrealBuildAccelerator/ObjTool/Private/UbaImportLibWriter.cpp
--- a/Engine/Source/Programs/UnrealBuildAccelerator/ObjTool/Private/UbaImportLibWriter.cpp
+++ b/Engine/Source/Programs/UnrealBuildAccelerator/ObjTool/Private/UbaImportLibWriter.cpp
@@ -63,21 +63,6 @@ namespace uba
 			u8 isData;
 		};
 
-		enum CallingConvention
-		{
-			CallingConvention_Cdecl,
-			CallingConvention_StdCall,
-			CallingConvention_FastCall,
-			CallingConvention_VectorCall
-		};
-
-		enum ImportType
-		{
-			ImportType_Ordinal,
-			ImportType_SpecifiedName,
-			ImportType_DecoratedName 
-		};
-
 		MemoryBlock m_memory;
 		u32 m_memoryOffset;
 
@@ -99,7 +84,7 @@ namespace uba
 				func(*s);
 		}
 
-		void AddSymbol(const char* name, const char* extra, u32 value, u32 ArgListSizeInBytes, CallingConvention CallingConvention, ImportType importType)
+		void AddSymbol(const char* name, const char* extra, u32 value)
 		{
 			u32 offsetIndex = u32(m_symbols.size()) - m_extraSymbolCount;
 			SymbolInfo& s0 = *m_symbols.emplace_back(new SymbolInfo());
@@ -116,33 +101,9 @@ namespace uba
 				return;
 			}
 
-
-			char buf[256];
-			if (CallingConvention == CallingConvention_VectorCall)
-			{
-				sprintf_s(buf, sizeof(buf), "%s@@%d", name, ArgListSizeInBytes);
-				name = buf;
-			}
 			s0.name = name;
-
-			switch (importType)
-			{
-			case ImportType_Ordinal:
-				s0.ordinal = value;
-				s0.type = 0;
-				break;
-			case ImportType_SpecifiedName:
-				s0.hint = value;
-				if (CallingConvention == CallingConvention_VectorCall)
-					s0.type = 0xC;
-				else
-					s0.type = 4;
-				break;
-			case ImportType_DecoratedName:
-				s0.hint = value;
-				s0.type = 4;
-				break;
-			}
+			s0.hint = value;
+			s0.type = 4;
 
 			auto& s1 = *m_symbols.emplace_back(new SymbolInfo());
 			s1.name = "__imp_" + s0.name;
@@ -214,10 +175,14 @@ namespace uba
 				FileHeader.Name[stringLength] = '/';
 			}
 
-			sprintf_s((char*)FileHeader.Date, 13, "%-12I64d", m_date);
+			#pragma warning(push)
+			#pragma warning(disable : 6386)
+			sprintf_s((char*)FileHeader.Date, 13, "%-12I64d", s64(m_date));
 			FileHeader.UserID[0] = ' '; // Remove 0 termination from Date
 			sprintf_s((char*)FileHeader.Mode, 9, "%-8ho", a);
-			sprintf_s((char*)FileHeader.Size, 11, "%-10d", size);
+			sprintf_s((char*)FileHeader.Size, 11, "%-10d", int(size));
+			#pragma warning(pop)
+
 			FileHeader.EndHeader[0] = 0x60;
 			FileHeader.EndHeader[1] = 0xA;
 			Write(&FileHeader, IMAGE_SIZEOF_ARCHIVE_MEMBER_HDR);
@@ -627,7 +592,7 @@ namespace uba
 					if (handled.insert(kv.first).second)
 						sortedSymbols.try_emplace(kv.second.index, &kv);
 				for (auto& outerKv : sortedSymbols)
-					AddSymbol(outerKv.second->first.c_str(), outerKv.second->second.extra.c_str(), 0, 4, CallingConvention_StdCall, ImportType_SpecifiedName);
+					AddSymbol(outerKv.second->first.c_str(), outerKv.second->second.extra.c_str(), 0);
 			}
 
 			m_moduleName = std::string(libName) + libExt;
diff --git a/Engine/Source/Programs/UnrealBuildAccelerator/ObjTool/Private/UbaObjTool.cpp b/Engine/Source/Programs/UnrealBuildAccelerator/ObjTool/Private/UbaObjTool.cpp
--- a/Engine/Source/Programs/UnrealBuildAccelerator/ObjTool/Private/UbaObjTool.cpp
+++ b/Engine/Source/Programs/UnrealBuildAccelerator/ObjTool/Private/UbaObjTool.cpp
@@ -1,7 +1,5 @@
 // Copyright Epic Games, Inc. All Rights Reserved.
 
-#define Local_GetLongPathNameW uba::GetLongPathNameW
-
 #include "UbaObjectFile.h"
 #include "UbaDirectoryIterator.h"
 #include "UbaFileAccessor.h"
diff --git a/Engine/Source/Programs/UnrealBuildAccelerator/Test/Private/UbaTestBasics.h b/Engine/Source/Programs/UnrealBuildAccelerator/Test/Private/UbaTestBasics.h
--- a/Engine/Source/Programs/UnrealBuildAccelerator/Test/Private/UbaTestBasics.h
+++ b/Engine/Source/Programs/UnrealBuildAccelerator/Test/Private/UbaTestBasics.h
@@ -2,8 +2,6 @@
 
 #pragma once
 
-#define Local_GetLongPathNameW uba::GetLongPathNameW
-
 #include "UbaBinaryReaderWriter.h"
 #include "UbaCompactTables.h"
 #include "UbaFileAccessor.h"
diff --git a/Engine/Source/Programs/UnrealBuildAccelerator/Test/Private/UbaTestCache.h b/Engine/Source/Programs/UnrealBuildAccelerator/Test/Private/UbaTestCache.h
--- a/Engine/Source/Programs/UnrealBuildAccelerator/Test/Private/UbaTestCache.h
+++ b/Engine/Source/Programs/UnrealBuildAccelerator/Test/Private/UbaTestCache.h
@@ -8,6 +8,7 @@
 #include "UbaNetworkBackendTcp.h"
 #include "UbaNetworkClient.h"
 #include "UbaNetworkServer.h"
+#include "UbaRootPaths.h"
 #include "UbaSessionServer.h"
 #include "UbaStorageServer.h"
 #include "../Private/UbaHashMap.h"
diff --git a/Engine/Source/Programs/UnrealBuildAccelerator/Test/Private/UbaTestConfig.h b/Engine/Source/Programs/UnrealBuildAccelerator/Test/Private/UbaTestConfig.h
--- a/Engine/Source/Programs/UnrealBuildAccelerator/Test/Private/UbaTestConfig.h
+++ b/Engine/Source/Programs/UnrealBuildAccelerator/Test/Private/UbaTestConfig.h
@@ -3,6 +3,7 @@
 #pragma once
 
 #include "UbaConfig.h"
+#include "UbaPlatform.h"
 
 namespace uba
 {
diff --git a/Engine/Source/Programs/UnrealBuildAccelerator/Test/Private/UbaTestCrypto.h b/Engine/Source/Programs/UnrealBuildAccelerator/Test/Private/UbaTestCrypto.h
--- a/Engine/Source/Programs/UnrealBuildAccelerator/Test/Private/UbaTestCrypto.h
+++ b/Engine/Source/Programs/UnrealBuildAccelerator/Test/Private/UbaTestCrypto.h
@@ -3,7 +3,8 @@
 #pragma once
 
 #include "UbaCrypto.h"
-#include "UbaNetwork.h"
+#include "UbaLogger.h"
+#include "UbaStringBuffer.h"
 
 namespace uba
 {
diff --git a/Engine/Source/Programs/UnrealBuildAccelerator/Test/Private/UbaTestSession.h b/Engine/Source/Programs/UnrealBuildAccelerator/Test/Private/UbaTestSession.h
--- a/Engine/Source/Programs/UnrealBuildAccelerator/Test/Private/UbaTestSession.h
+++ b/Engine/Source/Programs/UnrealBuildAccelerator/Test/Private/UbaTestSession.h
@@ -2,6 +2,8 @@
 
 #pragma once
 
+#include "UbaFileAccessor.h"
+#include "UbaNetworkBackendTcp.h"
 #include "UbaSessionServer.h"
 #include "UbaSessionClient.h"
 #include "UbaStorageServer.h"
diff --git a/Engine/Source/Programs/UnrealBuildAccelerator/Test/Private/UbaTestStdOut.h b/Engine/Source/Programs/UnrealBuildAccelerator/Test/Private/UbaTestStdOut.h
--- a/Engine/Source/Programs/UnrealBuildAccelerator/Test/Private/UbaTestStdOut.h
+++ b/Engine/Source/Programs/UnrealBuildAccelerator/Test/Private/UbaTestStdOut.h
@@ -2,11 +2,14 @@
 
 #pragma once
 
+#include "UbaNetworkBackendTcp.h"
 #include "UbaSessionServer.h"
 #include "UbaStorageServer.h"
 
 namespace uba
 {
+	void GetTestAppPath(LoggerWithWriter& logger, StringBufferBase& out);
+
 	bool TestStdOut(LoggerWithWriter& logger, const StringBufferBase& testRootDir, bool remote, const tchar* app = nullptr, const tchar* arg = nullptr, const tchar* expectedOut = nullptr)
 	{
 		LogWriter& logWriter = logger.m_writer;
diff --git a/Engine/Source/Programs/UnrealBuildAccelerator/Test/Private/UbaTestStorage.h b/Engine/Source/Programs/UnrealBuildAccelerator/Test/Private/UbaTestStorage.h
--- a/Engine/Source/Programs/UnrealBuildAccelerator/Test/Private/UbaTestStorage.h
+++ b/Engine/Source/Programs/UnrealBuildAccelerator/Test/Private/UbaTestStorage.h
@@ -2,6 +2,8 @@
 
 #pragma once
 
+#include "UbaFileAccessor.h"
+#include "UbaNetworkBackendTcp.h"
 #include "UbaStorageClient.h"
 #include "UbaStorageServer.h"
 
diff --git a/Engine/Source/Programs/UnrealBuildAccelerator/Test/Private/UbaTestStress.h b/Engine/Source/Programs/UnrealBuildAccelerator/Test/Private/UbaTestStress.h
--- a/Engine/Source/Programs/UnrealBuildAccelerator/Test/Private/UbaTestStress.h
+++ b/Engine/Source/Programs/UnrealBuildAccelerator/Test/Private/UbaTestStress.h
@@ -3,12 +3,16 @@
 #pragma once
 
 #include "UbaCoordinatorWrapper.h"
+#include "UbaFileAccessor.h"
+#include "UbaNetworkBackendTcp.h"
 #include "UbaScheduler.h"
 #include "UbaSessionServer.h"
 #include "UbaStorageServer.h"
 
 namespace uba
 {
+	void GetTestAppPath(LoggerWithWriter& logger, StringBufferBase& out);
+
 	bool TestStress(LoggerWithWriter& logger, const StringBufferBase& testRootDir)
 	{
 		LogWriter& logWriter = logger.m_writer;
diff --git a/Engine/Source/Programs/UnrealBuildAccelerator/UbaAgent.Target.cs b/Engine/Source/Programs/UnrealBuildAccelerator/UbaAgent.Target.cs
--- a/Engine/Source/Programs/UnrealBuildAccelerator/UbaAgent.Target.cs
+++ b/Engine/Source/Programs/UnrealBuildAccelerator/UbaAgent.Target.cs
@@ -32,6 +32,8 @@ public class UbaAgentTarget : TargetRules
 		Rules.bCompileAgainstApplicationCore = false;
 		Rules.bCompileICU = false;
 
+		Rules.GlobalDefinitions.Add("UBA_BUILD=1");
+
 		// to build with automation tests:
 		// bForceCompileDevelopmentAutomationTests = true;
 
@@ -117,7 +119,6 @@ public class UbaAgentTarget : TargetRules
 
 		if (Target.Platform.IsInGroup(UnrealPlatformGroup.Windows))
 		{
-			Rules.AdditionalCompilerArguments = "/wd4100 "; // -- C4100: unreferenced formal parameter
 			if (ShouldExport)
 			{
 				Rules.GlobalDefinitions.Add("UBA_API=__declspec(dllexport)");
diff --git a/Engine/Source/Programs/UnrealBuildAccelerator/Visualizer/Private/UbaVisualizer.cpp b/Engine/Source/Programs/UnrealBuildAccelerator/Visualizer/Private/UbaVisualizer.cpp
--- a/Engine/Source/Programs/UnrealBuildAccelerator/Visualizer/Private/UbaVisualizer.cpp
+++ b/Engine/Source/Programs/UnrealBuildAccelerator/Visualizer/Private/UbaVisualizer.cpp
@@ -2958,7 +2958,9 @@ namespace uba
 
 	void Visualizer::PostNewTitle(const StringView& title)
 	{
+		#if !defined( __clang_analyzer__ )
 		PostMessage(m_hwnd, WM_SETTITLE, 0, (LPARAM)_wcsdup(title.data));
+		#endif
 	}
 
 	void Visualizer::PostQuit()
