[UBA]
* Changed so when cache server bucket overflows, we take the oldest unused entry time and subtract one hour for new expiration time before running maintenance on bucket

diff --git a/Engine/Source/Programs/UnrealBuildAccelerator/Common/Private/UbaCacheServer.cpp b/Engine/Source/Programs/UnrealBuildAccelerator/Common/Private/UbaCacheServer.cpp
--- a/Engine/Source/Programs/UnrealBuildAccelerator/Common/Private/UbaCacheServer.cpp
+++ b/Engine/Source/Programs/UnrealBuildAccelerator/Common/Private/UbaCacheServer.cpp
@@ -9,6 +9,8 @@
 #include "UbaHashMap.h"
 #include "UbaNetworkServer.h"
 #include "UbaStorageServer.h"
+#include <algorithm>
+
 //#include <oodle2.h>
 
 // TODO
@@ -151,8 +153,10 @@ namespace uba
 		Atomic<bool> hasDeletedEntries;
 		Atomic<bool> needsSave;
 
+		// Times are in file time from creation of cache database
 		Atomic<u64> lastSavedTime;
 		Atomic<u64> lastUsedTime;
+		u64 longestUnusedTime = 0; 
 
 		u64 expirationTimeSeconds = 0;
 		u32 index = ~0u;
@@ -652,10 +656,24 @@ namespace uba
 		if (shouldExit())
 			return true;
 
+		// Take biggest buckets first
+		Vector<Bucket*> sortedBuckets;
+		for (auto& kv : m_buckets)
+			sortedBuckets.push_back(&kv.second);
+		std::sort(sortedBuckets.begin(), sortedBuckets.end(), [](Bucket* a, Bucket* b)
+			{
+				auto aSize = a->m_casKeyTable.GetSize();
+				auto bSize = b->m_casKeyTable.GetSize();
+				if (aSize != bSize)
+					return aSize > bSize;
+				return a->index < b->index;
+			});
+
+
 		ReaderWriterLock globalStatsLock;
 		u64 now = GetSystemTimeAsFileTime();
 		u64 oldest = 0;
-		u64 longestUnused = 0;
+		u64 longestUnusedTime = 0;
 
 		u32 workerCount = m_server.GetWorkerCount();
 		u32 workerCountToUse = workerCount > 0 ? workerCount - 1 : 0;
@@ -690,12 +708,12 @@ namespace uba
 			bool checkInputsForDeletes = m_checkInputsForDeletedCas && !deletedCasFiles.empty();
 
 			oldest = 0;
-			longestUnused = 0;
+			longestUnusedTime = 0;
 			totalEntryCount = 0;
 
-			m_server.ParallelFor(workerCountToUseForBuckets, m_buckets, [&](auto& it)
+			m_server.ParallelFor(workerCountToUseForBuckets, sortedBuckets, [&](auto& it)
 			{
-				Bucket& bucket = it->second;
+				Bucket& bucket = **it;
 				auto context = bucket.m_maintenanceContext;
 				if (!context)
 					context = bucket.m_maintenanceContext = new Bucket::MaintenanceContext;
@@ -714,33 +732,41 @@ namespace uba
 
 				auto& deletedOffsets = context->deletedOffsets;
 
-				bucket.totalEntryCount = 0;
-				bucket.totalEntrySize = 0;
-
-
 				ReaderWriterLock keysToEraseLock;
 				Vector<CasKey> keysToErase;
 
+				u64 bucketExpirationTimeSeconds = bucket.expirationTimeSeconds;
+
 				u64 lastUseTimeLimit = 0; // This is the time relative to server startup time
-				if (bucket.expirationTimeSeconds && m_bucketCasTableMaxSize)
+
+				if (deleteIteration == 0 && m_bucketCasTableMaxSize && bucket.longestUnusedTime)
 				{
-					if (deleteIteration == 0)
+					// If bucket is larger than max we will take longestUnused and reduce by one hour
+					u64 bucketCasTableSize = bucket.m_casKeyTable.GetSize();
+					if (bucketCasTableSize >= m_bucketCasTableMaxSize)
 					{
-						u64 bucketCasTableSize = bucket.m_casKeyTable.GetSize();
-						u64 oldExpirationTime = bucket.expirationTimeSeconds;
-						if (bucketCasTableSize >= m_bucketCasTableMaxSize)
-							bucket.expirationTimeSeconds -= 60*60; // Decreased by one hour
-						else if (bucket.expirationTimeSeconds < m_expirationTimeSeconds && (bucketCasTableSize + 1ull*1024*1024) < m_bucketCasTableMaxSize)
-							bucket.expirationTimeSeconds += 60*60; // Increased by one hour
-						if (oldExpirationTime != bucket.expirationTimeSeconds)
-							m_logger.Info(TC("%s expiration time for bucket %u to %s"), oldExpirationTime < bucket.expirationTimeSeconds ? TC("Increased") : TC("Decreased"), bucket.index, TimeToText(MsToTime(bucket.expirationTimeSeconds*1000), true).str);
+						u64 longestUnusedSeconds = GetFileTimeAsSeconds((now - m_creationTime) - bucket.longestUnusedTime);
+						if (longestUnusedSeconds < bucketExpirationTimeSeconds)
+						{
+							if (longestUnusedSeconds > 60*60)
+								longestUnusedSeconds -= 60*60;
+							bucketExpirationTimeSeconds = longestUnusedSeconds;
+							m_logger.Info(TC("    Set temporary expiration time for bucket %u to %s to reduce cas table size"), bucket.index, TimeToText(MsToTime(bucketExpirationTimeSeconds*1000), true).str);
+						}
 					}
+				}
 
-					u64 secondsRunning = GetFileTimeAsSeconds(now - m_creationTime);
-					if (secondsRunning > bucket.expirationTimeSeconds)
-						lastUseTimeLimit = (now - m_creationTime) - GetSecondsAsFileTime(bucket.expirationTimeSeconds);
+				if (bucketExpirationTimeSeconds)
+				{
+					u64 secondsSinceCreation = GetFileTimeAsSeconds(now - m_creationTime);
+					if (secondsSinceCreation > bucketExpirationTimeSeconds)
+						lastUseTimeLimit = GetSecondsAsFileTime(secondsSinceCreation - bucketExpirationTimeSeconds);
 				}
 
+				bucket.longestUnusedTime = 0;
+				bucket.totalEntryCount = 0;
+				bucket.totalEntrySize = 0;
+
 				m_server.ParallelFor(workerCountToUse, bucket.m_cacheEntryLookup, [&, touchedCas = Vector<Atomic<bool>*>()](auto& li) mutable
 				{
 					CacheEntries& entries = li->second;
@@ -863,11 +889,16 @@ namespace uba
 						capacityLeft -= neededSize;
 
 						{
+							u64 lastUsedTime = entry.lastUsedTime;
+							if (!lastUsedTime)
+								lastUsedTime = entry.creationTime;
 							SCOPED_WRITE_LOCK(globalStatsLock, l);
 							if (!oldest || entry.creationTime < oldest)
 								oldest = entry.creationTime;
-							if (!longestUnused || entry.lastUsedTime < longestUnused)
-								longestUnused = entry.lastUsedTime;
+							if (!bucket.longestUnusedTime || lastUsedTime < bucket.longestUnusedTime)
+								bucket.longestUnusedTime = lastUsedTime;
+							if (!longestUnusedTime || lastUsedTime < longestUnusedTime)
+								longestUnusedTime = lastUsedTime;
 						}
 
 						for (auto v : touchedCas)
@@ -939,11 +970,11 @@ namespace uba
 
 		u64 maxCommittedMemory = 0;
 
-		m_server.ParallelFor(workerCountToUseForBuckets, m_buckets, [&](auto& it)
+		m_server.ParallelFor(workerCountToUseForBuckets, sortedBuckets, [&](auto& it)
 		{
 			u64 bucketStartTime = GetTime();
 
-			Bucket& bucket = it->second;
+			Bucket& bucket = **it;
 			auto deleteContext = MakeGuard([&]() { delete bucket.m_maintenanceContext; bucket.m_maintenanceContext = nullptr; });
 
 			if (!bucket.hasDeletedEntries && !forceAllSteps)
@@ -1108,9 +1139,9 @@ namespace uba
 		}
 
 		u64 oldestTime = oldest ? GetFileTimeAsTime(now - (m_creationTime + oldest)) : 0;
-		u64 longestUnusedTime = longestUnused ? GetFileTimeAsTime(now - (m_creationTime + longestUnused)) : 0;
+		u64 longestUnusedTime2 = longestUnusedTime ? GetFileTimeAsTime(now - (m_creationTime + longestUnusedTime)) : 0;
 		u64 duration = GetTime() - startTime;
-		m_logger.Info(TC("Maintenance done! (%s) CasFiles: %s (%s) Entries: %s Oldest: %s LongestUnused: %s MaintenanceMem: %s/%s"), TimeToText(duration).str, CountToText(totalCasCount - deletedCasCount).str, BytesToText(totalCasSize).str, CountToText(totalEntryCount.load()).str, TimeToText(oldestTime, true).str, TimeToText(longestUnusedTime, true).str, BytesToText(maxCommittedMemory).str, BytesToText(m_maintenanceReserveSize).str);
+		m_logger.Info(TC("Maintenance done! (%s) CasFiles: %s (%s) Entries: %s Oldest: %s LongestUnused: %s MaintenanceMem: %s/%s"), TimeToText(duration).str, CountToText(totalCasCount - deletedCasCount).str, BytesToText(totalCasSize).str, CountToText(totalEntryCount.load()).str, TimeToText(oldestTime, true).str, TimeToText(longestUnusedTime2, true).str, BytesToText(maxCommittedMemory).str, BytesToText(m_maintenanceReserveSize).str);
 		
 		m_longestMaintenance = Max(m_longestMaintenance, duration);
 
diff --git a/Engine/Source/Programs/UnrealBuildAccelerator/Common/Public/UbaCacheEntry.h b/Engine/Source/Programs/UnrealBuildAccelerator/Common/Public/UbaCacheEntry.h
--- a/Engine/Source/Programs/UnrealBuildAccelerator/Common/Public/UbaCacheEntry.h
+++ b/Engine/Source/Programs/UnrealBuildAccelerator/Common/Public/UbaCacheEntry.h
@@ -19,8 +19,10 @@ namespace uba
 
 	struct CacheEntry
 	{
+		// Times are in file time from creation of cache database
 		u64 creationTime = 0;
 		u64 lastUsedTime = 0;
+		
 		u32 id = 0;
 		LogLinesType logLinesType = LogLinesType_Empty;
 
