[UBA]
* Improved error reporting and allowed retry when failing to create file mapping on posix

diff --git a/Engine/Source/Programs/UnrealBuildAccelerator/Common/Private/UbaFileMapping.cpp b/Engine/Source/Programs/UnrealBuildAccelerator/Common/Private/UbaFileMapping.cpp
--- a/Engine/Source/Programs/UnrealBuildAccelerator/Common/Private/UbaFileMapping.cpp
+++ b/Engine/Source/Programs/UnrealBuildAccelerator/Common/Private/UbaFileMapping.cpp
@@ -60,7 +60,7 @@ namespace uba
 			if (mkdir(lockDir.data, 0777) == -1)
 				if (errno != EEXIST)
 				{
-					UBA_ASSERTF(false, "Failed to create %s (%s)", lockDir.data, strerror(errno));
+					logger.Error(TC("Failed to create %s for memory mapping (%s)"), lockDir.data, strerror(errno));
 					return {};
 				}
 
@@ -80,7 +80,9 @@ namespace uba
 							g_mappingUidCounter = uid;
 							return;
 						}
+						logger.Warning("Failed to open %s for memory mapping (%s)", lockFile.data, strerror(errno));
 						UBA_ASSERTF(false, "Failed to open %s (%s)", lockFile.data, strerror(errno));
+						return;
 					}
 
 					if (flock(lockFd, LOCK_EX | LOCK_NB) == 0)
@@ -109,6 +111,7 @@ namespace uba
 		int shmFd;
 		int lockFd;
 
+		int retryCount = 4;
 		while (true)
 		{
 			uid = ++g_mappingUidCounter;
@@ -120,6 +123,7 @@ namespace uba
 			{
 				if (errno == EEXIST)
 					continue;
+				logger.Warning(TC("Failed to open/create %s (%s)"), lockFile.data, strerror(errno));
 				UBA_ASSERTF(false, "Failed to open/create %s (%s)", lockFile.data, strerror(errno));
 				continue;
 			}
@@ -138,15 +142,18 @@ namespace uba
 			if (shmFd != -1)
 				break;
 
-			shm_unlink(uidName.data);
-			shmFd = shm_open(uidName.data, oflags, S_IRUSR | S_IWUSR);
-			if (shmFd != -1)
-				break;
+			bool retry = retryCount-- > 0;
+
+			int err = errno;
+			auto logType = retry ? LogEntryType_Warning : LogEntryType_Error;
+			logger.Logf(logType, TC("Failed to create shm %s after getting lock-file %s (%s)"), uidName.data, lockFile.data, strerror(err));
+
+			if (retry)
+				continue;
 
 			remove(lockFile.data);
 			close(lockFd);
-			SetLastError(errno);
-			UBA_ASSERTF(false, "Failed to create shm %s after getting lock-file %s (%s)", uidName.data, lockFile.data, strerror(errno));
+			SetLastError(err);
 			return {};
 		}
 
@@ -615,7 +622,7 @@ namespace uba
 	{
 		m_mappingHandle = uba::CreateMemoryMappingW(m_logger, PAGE_READWRITE|SEC_RESERVE, capacity);
 		if (!m_mappingHandle.IsValid())
-			return m_logger.Error(TC("%s - Failed to create memory map (%s)"), m_name, LastErrorToText().data);
+			return m_logger.Error(TC("%s - Failed to create memory map with capacity %llu (%s)"), m_name, capacity, LastErrorToText().data);
 
 		m_blockSize = blockSize;
 		m_pageSize = 64*1024;
