[UBA]
* Removed version 0 of compact tables
* Changed compact tables to use uba hash table

diff --git a/Engine/Source/Programs/UnrealBuildAccelerator/Common/Private/UbaCacheClient.cpp b/Engine/Source/Programs/UnrealBuildAccelerator/Common/Private/UbaCacheClient.cpp
--- a/Engine/Source/Programs/UnrealBuildAccelerator/Common/Private/UbaCacheClient.cpp
+++ b/Engine/Source/Programs/UnrealBuildAccelerator/Common/Private/UbaCacheClient.cpp
@@ -39,9 +39,9 @@ namespace uba
 	{
 		Bucket(u32 id_)
 		:	id(id_)
-		,	serverPathTable(CachePathTableMaxSize, CompactPathTable::V1, CaseInsensitiveFs)
+		,	serverPathTable(CachePathTableMaxSize, CaseInsensitiveFs)
 		,	serverCasKeyTable(CacheCasKeyTableMaxSize)
-		,	sendPathTable(CachePathTableMaxSize, CompactPathTable::V1, CaseInsensitiveFs)
+		,	sendPathTable(CachePathTableMaxSize, CaseInsensitiveFs)
 		,	sendCasKeyTable(CacheCasKeyTableMaxSize)
 		{
 		}
diff --git a/Engine/Source/Programs/UnrealBuildAccelerator/Common/Private/UbaCacheServer.cpp b/Engine/Source/Programs/UnrealBuildAccelerator/Common/Private/UbaCacheServer.cpp
--- a/Engine/Source/Programs/UnrealBuildAccelerator/Common/Private/UbaCacheServer.cpp
+++ b/Engine/Source/Programs/UnrealBuildAccelerator/Common/Private/UbaCacheServer.cpp
@@ -124,7 +124,7 @@ namespace uba
 
 	struct CacheServer::ConnectionBucket
 	{
-		ConnectionBucket(u64 i) : pathTable(CachePathTableMaxSize, CompactPathTable::V1, IsCaseInsensitive(i)), casKeyTable(CacheCasKeyTableMaxSize), id(i) {}
+		ConnectionBucket(u64 i) : pathTable(CachePathTableMaxSize, IsCaseInsensitive(i)), casKeyTable(CacheCasKeyTableMaxSize), id(i) {}
 		CompactPathTable pathTable;
 		CompactCasKeyTable casKeyTable;
 
@@ -141,7 +141,7 @@ namespace uba
 
 	struct CacheServer::Bucket
 	{
-		Bucket(u64 id) : m_pathTable(CachePathTableMaxSize, CompactPathTable::V1, IsCaseInsensitive(id)), m_casKeyTable(CacheCasKeyTableMaxSize) {}
+		Bucket(u64 id) : m_pathTable(CachePathTableMaxSize, IsCaseInsensitive(id)), m_casKeyTable(CacheCasKeyTableMaxSize) {}
 		ReaderWriterLock m_cacheEntryLookupLock;
 		UnorderedMap<CasKey, CacheEntries> m_cacheEntryLookup;
 
@@ -1034,7 +1034,7 @@ namespace uba
 			HashMap2<u32, u32> oldToNewPathOffset;
 			u32 oldSize = bucket.m_pathTable.GetSize();
 			{
-				CompactPathTable newPathTable(CachePathTableMaxSize, CompactPathTable::V1, bucket.m_pathTable.GetPathCount(), bucket.m_pathTable.GetSegmentCount());
+				CompactPathTable newPathTable(CachePathTableMaxSize, bucket.m_pathTable.GetPathCount(), bucket.m_pathTable.GetSegmentCount());
 				oldToNewPathOffset.Init(memoryBlock, usedPathOffsets.CountSetBits());
 
 				StringBuffer<> temp;
diff --git a/Engine/Source/Programs/UnrealBuildAccelerator/Common/Private/UbaCompactTables.cpp b/Engine/Source/Programs/UnrealBuildAccelerator/Common/Private/UbaCompactTables.cpp
--- a/Engine/Source/Programs/UnrealBuildAccelerator/Common/Private/UbaCompactTables.cpp
+++ b/Engine/Source/Programs/UnrealBuildAccelerator/Common/Private/UbaCompactTables.cpp
@@ -5,9 +5,8 @@
 
 namespace uba
 {
-	CompactPathTable::CompactPathTable(u64 reserveSize, Version version, bool caseInsensitive, u64 reservePathCount, u64 reserveSegmentCount)
+	CompactPathTable::CompactPathTable(u64 reserveSize, bool caseInsensitive, u64 reservePathCount, u64 reserveSegmentCount)
 	:	m_reserveSize(reserveSize)
-	,	m_version(version)
 	,	m_caseInsensitive(caseInsensitive)
 	{
 		if (reservePathCount)
@@ -63,54 +62,39 @@ namespace uba
 		u64 segLen = strLen - (seg - str);
 		u8 bytesForParent = Get7BitEncodedCount(parentOffset);
 
-		if (m_version == V0)
+		StringKey segmentKey = ToStringKeyNoCheck(seg, segLen);
+		auto insres2 = m_segmentOffsets.try_emplace(segmentKey);
+		if (insres2.second)
 		{
+			// Put string directly after current element and set segment offset to 0
 			u64 bytesForString = GetStringWriteSize(seg, segLen);
-
-			u64 memSize = bytesForParent + bytesForString;
+			u64 memSize = bytesForParent + 1 + bytesForString;
 			u8* mem = (u8*)m_mem.AllocateNoLock(memSize, 1, TC(""));
 			BinaryWriter writer(mem, 0, memSize);
 			writer.Write7BitEncoded(parentOffset);
+			writer.Write7BitEncoded(0);
 			writer.WriteString(seg, segLen);
-			insres.first->second = u32(mem - m_mem.memory);
-			return insres.first->second;
+			u32 offset = u32(mem - m_mem.memory);
+			insres.first->second = offset;
+			insres2.first->second = offset + bytesForParent + 1;
+			return offset;
 		}
-		else
-		{
-			StringKey segmentKey = ToStringKeyNoCheck(seg, segLen);
-			auto insres2 = m_segmentOffsets.try_emplace(segmentKey);
-			if (insres2.second)
-			{
-				// Put string directly after current element and set segment offset to 0
-				u64 bytesForString = GetStringWriteSize(seg, segLen);
-				u64 memSize = bytesForParent + 1 + bytesForString;
-				u8* mem = (u8*)m_mem.AllocateNoLock(memSize, 1, TC(""));
-				BinaryWriter writer(mem, 0, memSize);
-				writer.Write7BitEncoded(parentOffset);
-				writer.Write7BitEncoded(0);
-				writer.WriteString(seg, segLen);
-				u32 offset = u32(mem - m_mem.memory);
-				insres.first->second = offset;
-				insres2.first->second = offset + bytesForParent + 1;
-				return offset;
-			}
 
-			#if 0
-			StringBuffer<> temp;
-			BinaryReader reader(m_mem.memory, insres2.first->second, 1000);
-			reader.ReadString(temp);
-			UBA_ASSERT(temp.count == segLen && wcsncmp(temp.data, seg, segLen) == 0);
-			#endif
+		#if 0
+		StringBuffer<> temp;
+		BinaryReader reader(m_mem.memory, insres2.first->second, 1000);
+		reader.ReadString(temp);
+		UBA_ASSERT(temp.count == segLen && wcsncmp(temp.data, seg, segLen) == 0);
+		#endif
 
-			u32 strOffset = insres2.first->second;
-			u64 memSize = bytesForParent + Get7BitEncodedCount(strOffset);
-			u8* mem = (u8*)m_mem.AllocateNoLock(memSize, 1, TC(""));
-			BinaryWriter writer(mem, 0, memSize);
-			writer.Write7BitEncoded(parentOffset);
-			writer.Write7BitEncoded(strOffset);
-			insres.first->second = u32(mem - m_mem.memory);
-			return insres.first->second;
-		}
+		u32 strOffset = insres2.first->second;
+		u64 memSize = bytesForParent + Get7BitEncodedCount(strOffset);
+		u8* mem = (u8*)m_mem.AllocateNoLock(memSize, 1, TC(""));
+		BinaryWriter writer(mem, 0, memSize);
+		writer.Write7BitEncoded(parentOffset);
+		writer.Write7BitEncoded(strOffset);
+		insres.first->second = u32(mem - m_mem.memory);
+		return insres.first->second;
 	}
 
 	void CompactPathTable::GetString(StringBufferBase& out, u64 offset) const
@@ -137,21 +121,6 @@ namespace uba
 			offsets[offsetCount] = u32(offset);
 		}
 
-		if (m_version == V0)
-		{
-			bool isFirst = true;
-			for (u32 i=offsetCount;i; --i)
-			{
-				reader.SetPosition(offsets[i-1]);
-				reader.Read7BitEncoded();
-
-				if (!isFirst)
-					out.Append(PathSeparator);
-				isFirst = false;
-				reader.ReadString(out);
-			}
-		}
-		else
 		{
 			bool isFirst = true;
 			for (u32 i=offsetCount;i; --i)
@@ -198,40 +167,23 @@ namespace uba
 		if (!writtenSize)
 			reader2.Skip(1);
 
-		if (m_version == V0)
-		{
-			while (reader2.GetLeft())
-			{
-				u32 offset = u32(reader2.GetPosition());
-				reader2.Read7BitEncoded();
-				reader2.SkipString();
-				StringBuffer<> str;
-				GetString(str, offset);
-				if (m_caseInsensitive)
-					str.MakeLower();
-				m_offsets.try_emplace(ToStringKeyNoCheck(str.data, str.count), offset);
-			}
-		}
-		else
+		while (reader2.GetLeft())
 		{
-			while (reader2.GetLeft())
+			u32 offset = u32(reader2.GetPosition());
+			reader2.Read7BitEncoded();
+			u64 stringOffset = reader2.Read7BitEncoded();
+			if (!stringOffset)
 			{
-				u32 offset = u32(reader2.GetPosition());
-				reader2.Read7BitEncoded();
-				u64 stringOffset = reader2.Read7BitEncoded();
-				if (!stringOffset)
-				{
-					u32 strOffset = u32(reader2.GetPosition());
-					StringBuffer<> seg;
-					reader2.ReadString(seg);
-					m_segmentOffsets.try_emplace(ToStringKeyNoCheck(seg.data, seg.count), strOffset);
-				}
-				StringBuffer<> str;
-				GetString(str, offset);
-				if (m_caseInsensitive)
-					str.MakeLower();
-				m_offsets.try_emplace(ToStringKeyNoCheck(str.data, str.count), offset);
+				u32 strOffset = u32(reader2.GetPosition());
+				StringBuffer<> seg;
+				reader2.ReadString(seg);
+				m_segmentOffsets.try_emplace(ToStringKeyNoCheck(seg.data, seg.count), strOffset);
 			}
+			StringBuffer<> str;
+			GetString(str, offset);
+			if (m_caseInsensitive)
+				str.MakeLower();
+			m_offsets.try_emplace(ToStringKeyNoCheck(str.data, str.count), offset);
 		}
 	}
 
@@ -243,9 +195,6 @@ namespace uba
 		u64 rs = m_reserveSize;
 		m_reserveSize = other.m_reserveSize;
 		other.m_reserveSize = rs;
-		Version v = m_version;
-		m_version = other.m_version;
-		other.m_version = v;
 	}
 
 	CompactCasKeyTable::CompactCasKeyTable(u64 reserveSize, u64 reserveOffsetsCount)
diff --git a/Engine/Source/Programs/UnrealBuildAccelerator/Common/Private/UbaHashMap.h b/Engine/Source/Programs/UnrealBuildAccelerator/Common/Private/UbaHashMap.h
deleted file mode 100644
diff --git a/Engine/Source/Programs/UnrealBuildAccelerator/Common/Public/UbaCompactTables.h b/Engine/Source/Programs/UnrealBuildAccelerator/Common/Public/UbaCompactTables.h
--- a/Engine/Source/Programs/UnrealBuildAccelerator/Common/Public/UbaCompactTables.h
+++ b/Engine/Source/Programs/UnrealBuildAccelerator/Common/Public/UbaCompactTables.h
@@ -3,6 +3,7 @@
 #pragma once
 
 #include "UbaHash.h"
+#include "UbaHashMap.h"
 #include "UbaMemory.h"
 
 namespace uba
@@ -14,9 +15,7 @@ namespace uba
 	class CompactPathTable
 	{
 	public:
-		enum Version : u8 { V0, V1 };
-
-		CompactPathTable(u64 reserveSize, Version version, bool caseSensitive, u64 reservePathCount = 0, u64 reserveSegmentCount = 0);
+		CompactPathTable(u64 reserveSize, bool caseSensitive, u64 reservePathCount = 0, u64 reserveSegmentCount = 0);
 
 		u32 Add(const tchar* str, u64 strLen, u32* outRequiredCasTableSize = nullptr);
 		u32 AddNoLock(const tchar* str, u64 strLen);
@@ -39,7 +38,6 @@ namespace uba
 		UnorderedMap<StringKey, u32> m_offsets;
 		UnorderedMap<StringKey, u32> m_segmentOffsets;
 		u64 m_reserveSize;
-		Version m_version;
 		bool m_caseInsensitive;
 	};
 
diff --git a/Engine/Source/Programs/UnrealBuildAccelerator/Common/Public/UbaHashMap.h b/Engine/Source/Programs/UnrealBuildAccelerator/Common/Public/UbaHashMap.h
new file mode 100644
--- /dev/null
+++ b/Engine/Source/Programs/UnrealBuildAccelerator/Common/Public/UbaHashMap.h
@@ -0,0 +1,247 @@
+// Copyright Epic Games, Inc. All Rights Reserved.
+
+#pragma once
+
+#include "UbaMemory.h"
+
+namespace uba
+{
+	template<typename Key, typename Value>
+	struct HashMap
+	{
+		struct Entry
+		{
+			Key key;
+			u32 next; // If this is 0 it means that it is unused, ~0u means it is the end of a chain, otherwise offset into entry array
+		};
+
+		void Init(MemoryBlock& memory, u64 maxSize)
+		{
+			u32 v = u32(maxSize);
+			v--;
+			v |= v >> 1;
+			v |= v >> 2;
+			v |= v >> 4;
+			v |= v >> 8;
+			v |= v >> 16;
+			v++;
+			u32 lookupSize = v;// << 1u;
+			mask = lookupSize - 1;
+			nextAvailableEntry = 1;
+
+			// This code relies on that allocated memory is zeroed out
+			lookup = (u32*)memory.Allocate(lookupSize*sizeof(u32), 1, TC(""));
+			entries = (Entry*)memory.Allocate((maxSize+1)*sizeof(Entry), alignof(Entry), TC(""));
+			values = (Value*)memory.Allocate((maxSize+1)*sizeof(Value), alignof(Value), TC(""));
+		}
+
+		Value& Insert(const Key& key)
+		{
+			u32 index = std::hash<Key>()(key) & mask;
+			u32 entryIndex = lookup[index];
+
+			if (entryIndex == 0)
+			{
+				entryIndex = nextAvailableEntry++;
+				lookup[index] = entryIndex;
+				Entry& entry = entries[entryIndex];
+				entry.key = key;
+				entry.next = ~0u;
+				return values[entryIndex];
+			}
+
+			while (true)
+			{
+				Entry& entry = entries[entryIndex];
+				if (entry.key == key)
+					return values[entryIndex];
+
+				if (entry.next != ~0u)
+				{
+					entryIndex = entry.next;
+					//static u32 counter = 0;
+					//printf("MULTIPLE: %u\n", counter++);
+					continue;
+				}
+
+				u32 newEntryIndex = entry.next = nextAvailableEntry++;
+				Entry& newEntry = entries[newEntryIndex];
+				newEntry.key = key;
+				newEntry.next = ~0u;
+				return values[newEntryIndex];
+			}
+		}
+
+		Value* Find(const Key& key) const
+		{
+			u32 index = std::hash<Key>()(key) & mask;
+			u32 entryIndex = lookup[index];
+			if (entryIndex == 0)
+				return nullptr;
+			Entry* entry = entries + entryIndex;
+			while (true)
+			{
+				if (entry->key == key)
+					return values + entryIndex;
+				entryIndex = entry->next;
+				if (entryIndex == ~0u)
+					return nullptr;
+				entry = entries + entryIndex;
+			}
+		}
+
+		const Key* GetKey(Value* value)
+		{
+			u32 pos = u32(value - values);
+			auto& entry = entries[pos];
+			if (entry.next == 0)
+				return nullptr;
+			return &entry.key;
+		}
+
+		u32 Size()
+		{
+			return nextAvailableEntry - 1;
+		}
+
+		void Erase(const Key& key)
+		{
+			u32 index = std::hash<Key>()(key) & mask;
+			u32 entryIndex = lookup[index];
+			if (entryIndex == 0)
+				return;
+			Entry* entry = entries + entryIndex;
+			u32* prevNext = nullptr;
+			while (true)
+			{
+				if (entry->key == key)
+				{
+					if (!prevNext)
+					{
+						if (entry->next == ~0u)
+							lookup[index] = 0;
+						else
+							lookup[index] = entry->next;
+					}
+					else
+						*prevNext = entry->next;
+					entry->next = 0;
+
+					// TODO: We could swap with last and reduce size
+
+					return;
+				}
+				entryIndex = entry->next;
+				if (entryIndex == ~0u)
+					return;
+				prevNext = &entry->next;
+				entry = entries + entryIndex;
+			}
+		}
+
+		Value* ValuesBegin() { return values + 1; }
+		Value* ValuesEnd() { return values + nextAvailableEntry; }
+
+		Entry* entries;
+		Value* values;
+
+		u32* lookup;
+		u32 mask;
+		u32 nextAvailableEntry;
+	};
+
+	template<typename Key, typename Value>
+	struct HashMap2
+	{
+		struct Entry
+		{
+			Key key;
+			Value value;
+			u32 next; // If this is 0 it means that it is unused, ~0u means it is the end of a chain, otherwise offset into entry array
+		};
+
+		void Init(MemoryBlock& memory, u64 maxSize)
+		{
+			u32 v = u32(maxSize);
+			v--;
+			v |= v >> 1;
+			v |= v >> 2;
+			v |= v >> 4;
+			v |= v >> 8;
+			v |= v >> 16;
+			v++;
+			u32 lookupSize = v;// << 1u;
+			mask = lookupSize - 1;
+			nextAvailableEntry = 1;
+
+			// This code relies on that allocated memory is zeroed out
+			lookup = (u32*)memory.Allocate(lookupSize*sizeof(u32), 1, TC(""));
+			entries = (Entry*)memory.Allocate((maxSize+1)*sizeof(Entry), alignof(Entry), TC(""));
+		}
+
+		Value& Insert(const Key& key)
+		{
+			u32 index = std::hash<Key>()(key) & mask;
+			u32 entryIndex = lookup[index];
+
+			if (entryIndex == 0)
+			{
+				entryIndex = nextAvailableEntry++;
+				lookup[index] = entryIndex;
+				Entry& entry = entries[entryIndex];
+				entry.key = key;
+				entry.next = ~0u;
+				return entry.value;
+			}
+
+			while (true)
+			{
+				Entry& entry = entries[entryIndex];
+				if (entry.key == key)
+					return entry.value;
+
+				if (entry.next != ~0u)
+				{
+					entryIndex = entry.next;
+					//static u32 counter = 0;
+					//printf("MULTIPLE: %u\n", counter++);
+					continue;
+				}
+
+				u32 newEntryIndex = entry.next = nextAvailableEntry++;
+				Entry& newEntry = entries[newEntryIndex];
+				newEntry.key = key;
+				newEntry.next = ~0u;
+				return newEntry.value;
+			}
+		}
+
+		Value* Find(const Key& key) const
+		{
+			u32 index = std::hash<Key>()(key) & mask;
+			u32 entryIndex = lookup[index];
+			if (entryIndex == 0)
+				return nullptr;
+			Entry* entry = entries + entryIndex;
+			while (true)
+			{
+				if (entry->key == key)
+					return &entry->value;
+				entryIndex = entry->next;
+				if (entryIndex == ~0u)
+					return nullptr;
+				entry = entries + entryIndex;
+			}
+		}
+
+		u32 Size()
+		{
+			return nextAvailableEntry - 1;
+		}
+
+		Entry* entries;
+		u32* lookup;
+		u32 mask;
+		u32 nextAvailableEntry;
+	};
+}
\ No newline at end of file
