[UBA]
* Changes in preparation of moving compact tables over to uba hash maps.

diff --git a/Engine/Source/Programs/UnrealBuildAccelerator/Common/Private/UbaCacheClient.cpp b/Engine/Source/Programs/UnrealBuildAccelerator/Common/Private/UbaCacheClient.cpp
--- a/Engine/Source/Programs/UnrealBuildAccelerator/Common/Private/UbaCacheClient.cpp
+++ b/Engine/Source/Programs/UnrealBuildAccelerator/Common/Private/UbaCacheClient.cpp
@@ -39,9 +39,9 @@ namespace uba
 	{
 		Bucket(u32 id_)
 		:	id(id_)
-		,	serverPathTable(CachePathTableMaxSize, CaseInsensitiveFs)
+		,	serverPathTable(CaseInsensitiveFs, 0, 0)
 		,	serverCasKeyTable(CacheCasKeyTableMaxSize)
-		,	sendPathTable(CachePathTableMaxSize, CaseInsensitiveFs)
+		,	sendPathTable(CaseInsensitiveFs, 0, 0)
 		,	sendCasKeyTable(CacheCasKeyTableMaxSize)
 		{
 		}
@@ -237,7 +237,7 @@ namespace uba
 			}
 
 
-			u32 pathOffset = bucket.sendPathTable.Add(qualifiedPath.c_str(), u32(qualifiedPath.size()), &requiredPathTableSize);
+			u32 pathOffset = bucket.sendPathTable.Add(qualifiedPath.c_str(), u32(qualifiedPath.size()), requiredPathTableSize);
 
 			if (!isOutput) // Output files should be removed from input files.. For example when cl.exe compiles pch it reads previous pch file and we don't want it to be input
 			{
diff --git a/Engine/Source/Programs/UnrealBuildAccelerator/Common/Private/UbaCacheServer.cpp b/Engine/Source/Programs/UnrealBuildAccelerator/Common/Private/UbaCacheServer.cpp
--- a/Engine/Source/Programs/UnrealBuildAccelerator/Common/Private/UbaCacheServer.cpp
+++ b/Engine/Source/Programs/UnrealBuildAccelerator/Common/Private/UbaCacheServer.cpp
@@ -124,7 +124,7 @@ namespace uba
 
 	struct CacheServer::ConnectionBucket
 	{
-		ConnectionBucket(u64 i) : pathTable(CachePathTableMaxSize, IsCaseInsensitive(i)), casKeyTable(CacheCasKeyTableMaxSize), id(i) {}
+		ConnectionBucket(u64 i) : pathTable(IsCaseInsensitive(i), 0, 0), casKeyTable(CacheCasKeyTableMaxSize), id(i) {}
 		CompactPathTable pathTable;
 		CompactCasKeyTable casKeyTable;
 
@@ -141,7 +141,7 @@ namespace uba
 
 	struct CacheServer::Bucket
 	{
-		Bucket(u64 id) : m_pathTable(CachePathTableMaxSize, IsCaseInsensitive(id)), m_casKeyTable(CacheCasKeyTableMaxSize) {}
+		Bucket(u64 id) : m_pathTable(IsCaseInsensitive(id), 0, 0), m_casKeyTable(CacheCasKeyTableMaxSize) {}
 		ReaderWriterLock m_cacheEntryLookupLock;
 		UnorderedMap<CasKey, CacheEntries> m_cacheEntryLookup;
 
@@ -628,7 +628,7 @@ namespace uba
 			SCOPED_WRITE_LOCK(m_storage.m_casLookupLock, lookupLock);
 			
 			totalCasCount = m_storage.m_casLookup.size();
-			existingCas.Init(existingCasMemoryBlock, totalCasCount);
+			existingCas.Init(existingCasMemoryBlock, existingCasMemoryBlock, totalCasCount);
 
 			for (auto i=m_storage.m_casLookup.begin(), e=m_storage.m_casLookup.end(); i!=e;)
 			{
@@ -1032,15 +1032,16 @@ namespace uba
 
 			// Build new path table based on used offsets
 			HashMap2<u32, u32> oldToNewPathOffset;
-			u32 oldSize = bucket.m_pathTable.GetSize();
+			auto& oldPathTable = bucket.m_pathTable;
+			u32 oldSize = oldPathTable.GetSize();
 			{
-				CompactPathTable newPathTable(CachePathTableMaxSize, bucket.m_pathTable.GetPathCount(), bucket.m_pathTable.GetSegmentCount());
+				CompactPathTable newPathTable(oldPathTable.GetCaseInsensitive(), oldPathTable.GetPathCount(), oldPathTable.GetSegmentCount());
 				oldToNewPathOffset.Init(memoryBlock, usedPathOffsets.CountSetBits());
 
 				StringBuffer<> temp;
 				usedPathOffsets.Traverse([&](u32 pathOffset)
 				{
-					bucket.m_pathTable.GetString(temp.Clear(), pathOffset);
+					oldPathTable.GetString(temp.Clear(), pathOffset);
 					u32 newPathOffset = newPathTable.AddNoLock(temp.data, temp.count);
 
 					#if 0
@@ -1052,9 +1053,9 @@ namespace uba
 					if (pathOffset != newPathOffset)
 						oldToNewPathOffset.Insert(pathOffset) = newPathOffset;
 				});
-				bucket.m_pathTable.Swap(newPathTable);
+				oldPathTable.Swap(newPathTable);
 			}
-			m_logger.Detail(TC("    Bucket %u Recreated path table. %s -> %s (%s)"), bucket.index, BytesToText(oldSize).str, BytesToText(bucket.m_pathTable.GetSize()).str, TimeToText(GetTime() - recreatePathTableStart).str);
+			m_logger.Detail(TC("    Bucket %u Recreated path table. %s -> %s (%s)"), bucket.index, BytesToText(oldSize).str, BytesToText(oldPathTable.GetSize()).str, TimeToText(GetTime() - recreatePathTableStart).str);
 
 
 			// Build new caskey table based on used offsets
diff --git a/Engine/Source/Programs/UnrealBuildAccelerator/Common/Private/UbaCompactTables.cpp b/Engine/Source/Programs/UnrealBuildAccelerator/Common/Private/UbaCompactTables.cpp
--- a/Engine/Source/Programs/UnrealBuildAccelerator/Common/Private/UbaCompactTables.cpp
+++ b/Engine/Source/Programs/UnrealBuildAccelerator/Common/Private/UbaCompactTables.cpp
@@ -5,9 +5,11 @@
 
 namespace uba
 {
-	CompactPathTable::CompactPathTable(u64 reserveSize, bool caseInsensitive, u64 reservePathCount, u64 reserveSegmentCount)
-	:	m_reserveSize(reserveSize)
-	,	m_caseInsensitive(caseInsensitive)
+	static constexpr u32 CachePathTableMaxSize = 20*1024*1024; // Max size for string memory
+
+
+	CompactPathTable::CompactPathTable(bool caseInsensitive, u64 reservePathCount, u64 reserveSegmentCount)
+	:	m_caseInsensitive(caseInsensitive)
 	{
 		if (reservePathCount)
 			m_offsets.reserve(reservePathCount);
@@ -15,21 +17,27 @@ namespace uba
 			m_segmentOffsets.reserve(reserveSegmentCount);
 	}
 
-	u32 CompactPathTable::Add(const tchar* str, u64 strLen, u32* outRequiredCasTableSize)
+	u32 CompactPathTable::Add(const tchar* str, u64 strLen)
+	{
+		SCOPED_WRITE_LOCK(m_lock, lock);
+		u32 res = AddNoLock(str, strLen);
+		return res;
+	}
+
+	u32 CompactPathTable::Add(const tchar* str, u64 strLen, u32& outRequiredCasTableSize)
 	{
 		SCOPED_WRITE_LOCK(m_lock, lock);
 		u32 res = AddNoLock(str, strLen);
-		if (outRequiredCasTableSize)
-			*outRequiredCasTableSize = u32(m_mem.writtenSize);
+		outRequiredCasTableSize = u32(m_stringMem.writtenSize);
 		return res;
 	}
 
 	u32 CompactPathTable::AddNoLock(const tchar* str, u64 strLen)
 	{
-		if (!m_mem.memory)
-			m_mem.Init(m_reserveSize);
-		if (!m_mem.writtenSize)
-			m_mem.AllocateNoLock(1, 1, TC(""));
+		if (!m_stringMem.memory)
+			m_stringMem.Init(CachePathTableMaxSize);
+		if (!m_stringMem.writtenSize)
+			m_stringMem.AllocateNoLock(1, 1, TC(""));
 
 		const tchar* stringKeyString = str;
 		StringBuffer<MaxPath> tempStringKeyStr;
@@ -69,12 +77,12 @@ namespace uba
 			// Put string directly after current element and set segment offset to 0
 			u64 bytesForString = GetStringWriteSize(seg, segLen);
 			u64 memSize = bytesForParent + 1 + bytesForString;
-			u8* mem = (u8*)m_mem.AllocateNoLock(memSize, 1, TC(""));
+			u8* mem = (u8*)m_stringMem.AllocateNoLock(memSize, 1, TC(""));
 			BinaryWriter writer(mem, 0, memSize);
 			writer.Write7BitEncoded(parentOffset);
 			writer.Write7BitEncoded(0);
 			writer.WriteString(seg, segLen);
-			u32 offset = u32(mem - m_mem.memory);
+			u32 offset = u32(mem - m_stringMem.memory);
 			insres.first->second = offset;
 			insres2.first->second = offset + bytesForParent + 1;
 			return offset;
@@ -82,18 +90,18 @@ namespace uba
 
 		#if 0
 		StringBuffer<> temp;
-		BinaryReader reader(m_mem.memory, insres2.first->second, 1000);
+		BinaryReader reader(m_stringMem.memory, insres2.first->second, 1000);
 		reader.ReadString(temp);
 		UBA_ASSERT(temp.count == segLen && wcsncmp(temp.data, seg, segLen) == 0);
 		#endif
 
 		u32 strOffset = insres2.first->second;
 		u64 memSize = bytesForParent + Get7BitEncodedCount(strOffset);
-		u8* mem = (u8*)m_mem.AllocateNoLock(memSize, 1, TC(""));
+		u8* mem = (u8*)m_stringMem.AllocateNoLock(memSize, 1, TC(""));
 		BinaryWriter writer(mem, 0, memSize);
 		writer.Write7BitEncoded(parentOffset);
 		writer.Write7BitEncoded(strOffset);
-		insres.first->second = u32(mem - m_mem.memory);
+		insres.first->second = u32(mem - m_stringMem.memory);
 		return insres.first->second;
 	}
 
@@ -102,7 +110,7 @@ namespace uba
 		#if UBA_DEBUG
 		{
 			SCOPED_READ_LOCK(const_cast<CompactPathTable*>(this)->m_lock, lock)
-			UBA_ASSERTF(offset < m_mem.writtenSize, TC("Reading path key from offset %llu which is out of bounds (Max %llu)"), offset, m_mem.writtenSize);
+			UBA_ASSERTF(offset < m_stringMem.writtenSize, TC("Reading path key from offset %llu which is out of bounds (Max %llu)"), offset, m_stringMem.writtenSize);
 		}
 		#endif
 
@@ -110,7 +118,7 @@ namespace uba
 		offsets[0] = u32(offset);
 		u32 offsetCount = 0;
 
-		BinaryReader reader(m_mem.memory, offset, m_mem.writtenSize);
+		BinaryReader reader(m_stringMem.memory, offset, m_stringMem.writtenSize);
 
 		while (offset)
 		{
@@ -141,29 +149,34 @@ namespace uba
 
 	u8* CompactPathTable::GetMemory()
 	{
-		return m_mem.memory;
+		return m_stringMem.memory;
 	}
 
 	u32 CompactPathTable::GetSize()
 	{
 		SCOPED_READ_LOCK(m_lock, lock2)
-		return u32(m_mem.writtenSize);
+		return u32(m_stringMem.writtenSize);
+	}
+
+	bool CompactPathTable::GetCaseInsensitive()
+	{
+		return m_caseInsensitive;
 	}
 
 	void CompactPathTable::ReadMem(BinaryReader& reader, bool populateLookup)
 	{
-		if (!m_mem.memory)
-			m_mem.Init(m_reserveSize);
+		if (!m_stringMem.memory)
+			m_stringMem.Init(CachePathTableMaxSize);
 
-		u64 writtenSize = m_mem.writtenSize;
+		u64 writtenSize = m_stringMem.writtenSize;
 		u64 left = reader.GetLeft();
-		void* mem = m_mem.AllocateNoLock(left, 1, TC(""));
+		void* mem = m_stringMem.AllocateNoLock(left, 1, TC(""));
 		reader.ReadBytes(mem, left);
 
 		if (!populateLookup)
 			return;
 
-		BinaryReader reader2(m_mem.memory, writtenSize, m_mem.writtenSize);
+		BinaryReader reader2(m_stringMem.memory, writtenSize, m_stringMem.writtenSize);
 		if (!writtenSize)
 			reader2.Skip(1);
 
@@ -191,10 +204,10 @@ namespace uba
 	{
 		m_offsets.swap(other.m_offsets);
 		m_segmentOffsets.swap(other.m_segmentOffsets);
-		m_mem.Swap(other.m_mem);
-		u64 rs = m_reserveSize;
-		m_reserveSize = other.m_reserveSize;
-		other.m_reserveSize = rs;
+		m_stringMem.Swap(other.m_stringMem);
+		bool ci = m_caseInsensitive;
+		m_caseInsensitive = other.m_caseInsensitive;
+		other.m_caseInsensitive = ci;
 	}
 
 	CompactCasKeyTable::CompactCasKeyTable(u64 reserveSize, u64 reserveOffsetsCount)
diff --git a/Engine/Source/Programs/UnrealBuildAccelerator/Common/Public/UbaCompactTables.h b/Engine/Source/Programs/UnrealBuildAccelerator/Common/Public/UbaCompactTables.h
--- a/Engine/Source/Programs/UnrealBuildAccelerator/Common/Public/UbaCompactTables.h
+++ b/Engine/Source/Programs/UnrealBuildAccelerator/Common/Public/UbaCompactTables.h
@@ -15,15 +15,17 @@ namespace uba
 	class CompactPathTable
 	{
 	public:
-		CompactPathTable(u64 reserveSize, bool caseSensitive, u64 reservePathCount = 0, u64 reserveSegmentCount = 0);
+		CompactPathTable(bool caseSensitive, u64 reservePathCount, u64 reserveSegmentCount);
 
-		u32 Add(const tchar* str, u64 strLen, u32* outRequiredCasTableSize = nullptr);
+		u32 Add(const tchar* str, u64 strLen);
+		u32 Add(const tchar* str, u64 strLen, u32& outRequiredCasTableSize);
 		u32 AddNoLock(const tchar* str, u64 strLen);
 
 		void GetString(StringBufferBase& out, u64 offset) const;
 
 		u8* GetMemory();
 		u32 GetSize();
+		bool GetCaseInsensitive();
 
 		void ReadMem(BinaryReader& reader, bool populateLookup);
 		void Swap(CompactPathTable& other);
@@ -34,10 +36,9 @@ namespace uba
 	private:
 		u32 InternalAdd(const tchar* str, const tchar* stringKeyString, u64 strLen);
 		ReaderWriterLock m_lock;
-		MemoryBlock m_mem;
+		MemoryBlock m_stringMem;
 		UnorderedMap<StringKey, u32> m_offsets;
 		UnorderedMap<StringKey, u32> m_segmentOffsets;
-		u64 m_reserveSize;
 		bool m_caseInsensitive;
 	};
 
diff --git a/Engine/Source/Programs/UnrealBuildAccelerator/Common/Public/UbaHashMap.h b/Engine/Source/Programs/UnrealBuildAccelerator/Common/Public/UbaHashMap.h
--- a/Engine/Source/Programs/UnrealBuildAccelerator/Common/Public/UbaHashMap.h
+++ b/Engine/Source/Programs/UnrealBuildAccelerator/Common/Public/UbaHashMap.h
@@ -6,7 +6,7 @@
 
 namespace uba
 {
-	template<typename Key, typename Value>
+	template<typename Key, typename Value, bool AllowGrow = false>
 	struct HashMap
 	{
 		struct Entry
@@ -15,7 +15,7 @@ namespace uba
 			u32 next; // If this is 0 it means that it is unused, ~0u means it is the end of a chain, otherwise offset into entry array
 		};
 
-		void Init(MemoryBlock& memory, u64 maxSize)
+		void Init(MemoryBlock& entriesMem, MemoryBlock& valuesMem, u64 maxSize)
 		{
 			u32 v = u32(maxSize);
 			v--;
@@ -29,10 +29,15 @@ namespace uba
 			mask = lookupSize - 1;
 			nextAvailableEntry = 1;
 
+			maxEntries = u32(maxSize) + 1;
 			// This code relies on that allocated memory is zeroed out
-			lookup = (u32*)memory.Allocate(lookupSize*sizeof(u32), 1, TC(""));
-			entries = (Entry*)memory.Allocate((maxSize+1)*sizeof(Entry), alignof(Entry), TC(""));
-			values = (Value*)memory.Allocate((maxSize+1)*sizeof(Value), alignof(Value), TC(""));
+			lookup = (u32*)entriesMem.AllocateNoLock(lookupSize*sizeof(u32), 1, TC(""));
+			entries = (Entry*)entriesMem.AllocateNoLock(maxEntries*sizeof(Entry), alignof(Entry), TC(""));
+			values = (Value*)valuesMem.AllocateNoLock(maxEntries*sizeof(Value), alignof(Value), TC(""));
+			UBA_ASSERT(&entriesMem != &valuesMem || !AllowGrow);
+
+			entriesMemory = &entriesMem;
+			valuesMemory = &valuesMem;
 		}
 
 		Value& Insert(const Key& key)
@@ -40,6 +45,15 @@ namespace uba
 			u32 index = std::hash<Key>()(key) & mask;
 			u32 entryIndex = lookup[index];
 
+			UBA_ASSERT(nextAvailableEntry < maxEntries || AllowGrow);
+			if constexpr (AllowGrow)
+				if (nextAvailableEntry == maxEntries)
+				{
+					++maxEntries;
+					entriesMemory->AllocateNoLock(sizeof(Value), alignof(Value), TC(""));
+					valuesMemory->AllocateNoLock(sizeof(Value), alignof(Value), TC(""));
+				}
+
 			if (entryIndex == 0)
 			{
 				entryIndex = nextAvailableEntry++;
@@ -142,12 +156,16 @@ namespace uba
 		Value* ValuesBegin() { return values + 1; }
 		Value* ValuesEnd() { return values + nextAvailableEntry; }
 
+		MemoryBlock* entriesMemory;
+		MemoryBlock* valuesMemory;
+
 		Entry* entries;
 		Value* values;
 
 		u32* lookup;
 		u32 mask;
 		u32 nextAvailableEntry;
+		u32 maxEntries;
 	};
 
 	template<typename Key, typename Value>
@@ -174,9 +192,9 @@ namespace uba
 			mask = lookupSize - 1;
 			nextAvailableEntry = 1;
 
-			// This code relies on that allocated memory is zeroed out
-			lookup = (u32*)memory.Allocate(lookupSize*sizeof(u32), 1, TC(""));
-			entries = (Entry*)memory.Allocate((maxSize+1)*sizeof(Entry), alignof(Entry), TC(""));
+			// This code relies on that AllocateNoLockd memory is zeroed out
+			lookup = (u32*)memory.AllocateNoLock(lookupSize*sizeof(u32), 1, TC(""));
+			entries = (Entry*)memory.AllocateNoLock((maxSize+1)*sizeof(Entry), alignof(Entry), TC(""));
 		}
 
 		Value& Insert(const Key& key)
diff --git a/Engine/Source/Programs/UnrealBuildAccelerator/Common/Public/UbaNetwork.h b/Engine/Source/Programs/UnrealBuildAccelerator/Common/Public/UbaNetwork.h
--- a/Engine/Source/Programs/UnrealBuildAccelerator/Common/Public/UbaNetwork.h
+++ b/Engine/Source/Programs/UnrealBuildAccelerator/Common/Public/UbaNetwork.h
@@ -18,7 +18,6 @@ namespace uba
 	static constexpr u32 SessionNetworkVersion = 34;
 	static constexpr u32 CacheNetworkVersion = 5;
 
-	static constexpr u32 CachePathTableMaxSize = 32*1024*1024;
 	static constexpr u32 CacheCasKeyTableMaxSize = 64*1024*1024;
 
 	// Messages used over network between client/server (system, storage and session)
diff --git a/Engine/Source/Programs/UnrealBuildAccelerator/Test/Private/UbaTestBasics.h b/Engine/Source/Programs/UnrealBuildAccelerator/Test/Private/UbaTestBasics.h
--- a/Engine/Source/Programs/UnrealBuildAccelerator/Test/Private/UbaTestBasics.h
+++ b/Engine/Source/Programs/UnrealBuildAccelerator/Test/Private/UbaTestBasics.h
@@ -3,7 +3,6 @@
 #pragma once
 
 #include "UbaBinaryReaderWriter.h"
-#include "UbaCompactTables.h"
 #include "UbaFileAccessor.h"
 #include "UbaPathUtils.h"
 #include "UbaPlatform.h"
@@ -376,37 +375,6 @@ namespace uba
 	}
 	#endif
 
-	bool TestCompactPathTable(Logger& logger, const StringBufferBase& rootDir)
-	{
-		for (u32 i=0;i!=2; ++i)
-		{
-			CompactPathTable table(64*1024, CompactPathTable::Version(i), CaseInsensitiveFs);
-
-			StringBuffer<> str;
-			str.Append("foo").EnsureEndsWithSlash().Append("bar");
-			u32 offset = table.Add(str.data, str.count);
-
-			StringBuffer<> str2;
-			table.GetString(str2, offset);
-			if (!str.Equals(str2.data))
-				return false;
-
-			str.Clear().Append(PathSeparator).Append("foo").Append(PathSeparator).Append("bar");
-			offset = table.Add(str.data, str.count);
-			table.GetString(str2.Clear(), offset);
-			if (!str.Equals(str2.data))
-				return false;
-
-			CompactPathTable table2(64*1024, CompactPathTable::Version(i), CaseInsensitiveFs);
-			BinaryReader reader(table.GetMemory(), 0, table.GetSize());
-			table2.ReadMem(reader, true);
-			u32 offset2 = table2.Add(str.data, str.count);
-			if (offset != offset2)
-				return false;
-		}
-		return true;
-	}
-
 	bool TestRootPaths(Logger& logger, const StringBufferBase& rootDir)
 	{
 		#if PLATFORM_WINDOWS
diff --git a/Engine/Source/Programs/UnrealBuildAccelerator/Test/Private/UbaTestCache.h b/Engine/Source/Programs/UnrealBuildAccelerator/Test/Private/UbaTestCache.h
--- a/Engine/Source/Programs/UnrealBuildAccelerator/Test/Private/UbaTestCache.h
+++ b/Engine/Source/Programs/UnrealBuildAccelerator/Test/Private/UbaTestCache.h
@@ -5,13 +5,14 @@
 #include "UbaCacheClient.h"
 #include "UbaCacheEntry.h"
 #include "UbaCacheServer.h"
+#include "UbaCompactTables.h"
 #include "UbaNetworkBackendTcp.h"
 #include "UbaNetworkClient.h"
 #include "UbaNetworkServer.h"
 #include "UbaRootPaths.h"
 #include "UbaSessionServer.h"
 #include "UbaStorageServer.h"
-#include "../Private/UbaHashMap.h"
+#include "UbaHashMap.h"
 
 namespace uba
 {
@@ -209,12 +210,40 @@ namespace uba
 		return true;
 	}
 
+	bool TestCompactPathTable(Logger& logger, const StringBufferBase& rootDir)
+	{
+		CompactPathTable table(CaseInsensitiveFs, 0, 0);
+
+		StringBuffer<> str;
+		str.Append("foo").EnsureEndsWithSlash().Append("bar");
+		u32 offset = table.Add(str.data, str.count);
+
+		StringBuffer<> str2;
+		table.GetString(str2, offset);
+		if (!str.Equals(str2.data))
+			return false;
+
+		str.Clear().Append(PathSeparator).Append("foo").Append(PathSeparator).Append("bar");
+		offset = table.Add(str.data, str.count);
+		table.GetString(str2.Clear(), offset);
+		if (!str.Equals(str2.data))
+			return false;
+
+		CompactPathTable table2(CaseInsensitiveFs, 0, 0);
+		BinaryReader reader(table.GetMemory(), 0, table.GetSize());
+		table2.ReadMem(reader, true);
+		u32 offset2 = table2.Add(str.data, str.count);
+		if (offset != offset2)
+			return false;
+		return true;
+	}
+
 	bool TestHashTable(Logger& logger, const StringBufferBase& rootDir)
 	{
 		{
 			MemoryBlock memoryBlock(1024*1024);
 			HashMap<u32, u32> casMap;
-			casMap.Init(memoryBlock, 3);
+			casMap.Init(memoryBlock, memoryBlock, 3);
 			if (casMap.Find(1))
 				return false;
 			casMap.Insert(1) = 2;
@@ -223,6 +252,16 @@ namespace uba
 			casMap.Insert(1) = 3;
 			if (*casMap.Find(1) != 3)
 				return false;
+
+			MemoryBlock valuesMemoryBlock(1024*1024);
+			HashMap<u32, u32, true> casMap2;
+			casMap2.Init(memoryBlock, valuesMemoryBlock, 1);
+			casMap2.Insert(1) = 1;
+			casMap2.Insert(2) = 2;
+			if (*casMap2.Find(1) != 1)
+				return false;
+			if (*casMap2.Find(2) != 2)
+				return false;
 		}
 
 #if 0
