[UBA]
* Changed compact tables to use uba hash map

diff --git a/Engine/Source/Programs/UnrealBuildAccelerator/Common/Private/UbaCacheClient.cpp b/Engine/Source/Programs/UnrealBuildAccelerator/Common/Private/UbaCacheClient.cpp
--- a/Engine/Source/Programs/UnrealBuildAccelerator/Common/Private/UbaCacheClient.cpp
+++ b/Engine/Source/Programs/UnrealBuildAccelerator/Common/Private/UbaCacheClient.cpp
@@ -40,9 +40,7 @@ namespace uba
 		Bucket(u32 id_)
 		:	id(id_)
 		,	serverPathTable(CaseInsensitiveFs, 0, 0)
-		,	serverCasKeyTable(CacheCasKeyTableMaxSize)
 		,	sendPathTable(CaseInsensitiveFs, 0, 0)
-		,	sendCasKeyTable(CacheCasKeyTableMaxSize)
 		{
 		}
 
@@ -330,7 +328,7 @@ namespace uba
 			}
 
 			UBA_ASSERT(IsCompressed(casKey));
-			insres.first->second = bucket.sendCasKeyTable.Add(casKey, pathOffset, &requiredCasTableSize);
+			insres.first->second = bucket.sendCasKeyTable.Add(casKey, pathOffset, requiredCasTableSize);
 		}
 
 		if (!success)
diff --git a/Engine/Source/Programs/UnrealBuildAccelerator/Common/Private/UbaCacheServer.cpp b/Engine/Source/Programs/UnrealBuildAccelerator/Common/Private/UbaCacheServer.cpp
--- a/Engine/Source/Programs/UnrealBuildAccelerator/Common/Private/UbaCacheServer.cpp
+++ b/Engine/Source/Programs/UnrealBuildAccelerator/Common/Private/UbaCacheServer.cpp
@@ -124,7 +124,7 @@ namespace uba
 
 	struct CacheServer::ConnectionBucket
 	{
-		ConnectionBucket(u64 i) : pathTable(IsCaseInsensitive(i), 0, 0), casKeyTable(CacheCasKeyTableMaxSize), id(i) {}
+		ConnectionBucket(u64 i) : pathTable(IsCaseInsensitive(i), 0, 0), id(i) {}
 		CompactPathTable pathTable;
 		CompactCasKeyTable casKeyTable;
 
@@ -141,7 +141,7 @@ namespace uba
 
 	struct CacheServer::Bucket
 	{
-		Bucket(u64 id) : m_pathTable(IsCaseInsensitive(id), 0, 0), m_casKeyTable(CacheCasKeyTableMaxSize) {}
+		Bucket(u64 id) : m_pathTable(IsCaseInsensitive(id), 0, 0) {}
 		ReaderWriterLock m_cacheEntryLookupLock;
 		UnorderedMap<CasKey, CacheEntries> m_cacheEntryLookup;
 
@@ -628,7 +628,7 @@ namespace uba
 			SCOPED_WRITE_LOCK(m_storage.m_casLookupLock, lookupLock);
 			
 			totalCasCount = m_storage.m_casLookup.size();
-			existingCas.Init(existingCasMemoryBlock, existingCasMemoryBlock, totalCasCount);
+			existingCas.Init(existingCasMemoryBlock, totalCasCount);
 
 			for (auto i=m_storage.m_casLookup.begin(), e=m_storage.m_casLookup.end(); i!=e;)
 			{
@@ -1064,7 +1064,7 @@ namespace uba
 			oldSize = bucket.m_casKeyTable.GetSize();
 			{
 				oldToNewCasKeyOffset.Init(memoryBlock, usedCasKeyOffsetsCount);
-				CompactCasKeyTable newCasKeyTable(CacheCasKeyTableMaxSize, usedCasKeyOffsetsCount);
+				CompactCasKeyTable newCasKeyTable(usedCasKeyOffsetsCount);
 				BinaryReader reader2(bucket.m_casKeyTable.GetMemory(), 0, bucket.m_casKeyTable.GetSize());
 				usedCasKeyOffsets.Traverse([&](u32 casKeyOffset)
 				{
@@ -1074,7 +1074,7 @@ namespace uba
 					u32 newPathOffset = oldPathOffset;
 					if (auto value = oldToNewPathOffset.Find(oldPathOffset))
 						newPathOffset = *value;
-					u32 newCasKeyOffset = newCasKeyTable.Add(casKey, newPathOffset);
+					u32 newCasKeyOffset = newCasKeyTable.AddNoLock(casKey, newPathOffset);
 					if (casKeyOffset != newCasKeyOffset)
 						oldToNewCasKeyOffset.Insert(casKeyOffset) = newCasKeyOffset;
 				});
@@ -1352,7 +1352,8 @@ namespace uba
 			UBA_ASSERT(test.Equals(path.data));
 			#endif
 
-			u32 casKeyOffset = bucket.m_casKeyTable.Add(casKey, pathOffset);
+			u32 requiredCasTableSize;
+			u32 casKeyOffset = bucket.m_casKeyTable.Add(casKey, pathOffset, requiredCasTableSize);
 			auto insres = inputs.insert(casKeyOffset);
 			if (!insres.second)
 			{
@@ -1466,7 +1467,8 @@ namespace uba
 			UBA_ASSERT(test.Equals(path.data));
 			#endif
 
-			u32 casKeyOffset = bucket.m_casKeyTable.Add(casKey, pathOffset);
+			u32 requiredCasTableSize;
+			u32 casKeyOffset = bucket.m_casKeyTable.Add(casKey, pathOffset, requiredCasTableSize);
 			outputs.insert(casKeyOffset);
 			bytesForOutput += Get7BitEncodedCount(casKeyOffset);
 
diff --git a/Engine/Source/Programs/UnrealBuildAccelerator/Common/Private/UbaCompactTables.cpp b/Engine/Source/Programs/UnrealBuildAccelerator/Common/Private/UbaCompactTables.cpp
--- a/Engine/Source/Programs/UnrealBuildAccelerator/Common/Private/UbaCompactTables.cpp
+++ b/Engine/Source/Programs/UnrealBuildAccelerator/Common/Private/UbaCompactTables.cpp
@@ -5,7 +5,9 @@
 
 namespace uba
 {
-	static constexpr u32 CachePathTableMaxSize = 20*1024*1024; // Max size for string memory
+	static constexpr u32 PathTableMaxSize = 20*1024*1024; // Max size for string memory
+	static constexpr u32 CasKeyTableMaxSize = 64*1024*1024;
+	static constexpr u32 CasKeyLookupInitialCount = 1024;
 
 
 	CompactPathTable::CompactPathTable(bool caseInsensitive, u64 reservePathCount, u64 reserveSegmentCount)
@@ -24,20 +26,20 @@ namespace uba
 		return res;
 	}
 
-	u32 CompactPathTable::Add(const tchar* str, u64 strLen, u32& outRequiredCasTableSize)
+	u32 CompactPathTable::Add(const tchar* str, u64 strLen, u32& outRequiredPathTableSize)
 	{
 		SCOPED_WRITE_LOCK(m_lock, lock);
 		u32 res = AddNoLock(str, strLen);
-		outRequiredCasTableSize = u32(m_stringMem.writtenSize);
+		outRequiredPathTableSize = u32(m_pathTableMem.writtenSize);
 		return res;
 	}
 
 	u32 CompactPathTable::AddNoLock(const tchar* str, u64 strLen)
 	{
-		if (!m_stringMem.memory)
-			m_stringMem.Init(CachePathTableMaxSize);
-		if (!m_stringMem.writtenSize)
-			m_stringMem.AllocateNoLock(1, 1, TC(""));
+		if (!m_pathTableMem.memory)
+			m_pathTableMem.Init(PathTableMaxSize);
+		if (!m_pathTableMem.writtenSize)
+			m_pathTableMem.AllocateNoLock(1, 1, TC(""));
 
 		const tchar* stringKeyString = str;
 		StringBuffer<MaxPath> tempStringKeyStr;
@@ -77,12 +79,12 @@ namespace uba
 			// Put string directly after current element and set segment offset to 0
 			u64 bytesForString = GetStringWriteSize(seg, segLen);
 			u64 memSize = bytesForParent + 1 + bytesForString;
-			u8* mem = (u8*)m_stringMem.AllocateNoLock(memSize, 1, TC(""));
+			u8* mem = (u8*)m_pathTableMem.AllocateNoLock(memSize, 1, TC(""));
 			BinaryWriter writer(mem, 0, memSize);
 			writer.Write7BitEncoded(parentOffset);
 			writer.Write7BitEncoded(0);
 			writer.WriteString(seg, segLen);
-			u32 offset = u32(mem - m_stringMem.memory);
+			u32 offset = u32(mem - m_pathTableMem.memory);
 			insres.first->second = offset;
 			insres2.first->second = offset + bytesForParent + 1;
 			return offset;
@@ -90,18 +92,18 @@ namespace uba
 
 		#if 0
 		StringBuffer<> temp;
-		BinaryReader reader(m_stringMem.memory, insres2.first->second, 1000);
+		BinaryReader reader(m_pathTableMem.memory, insres2.first->second, 1000);
 		reader.ReadString(temp);
 		UBA_ASSERT(temp.count == segLen && wcsncmp(temp.data, seg, segLen) == 0);
 		#endif
 
 		u32 strOffset = insres2.first->second;
 		u64 memSize = bytesForParent + Get7BitEncodedCount(strOffset);
-		u8* mem = (u8*)m_stringMem.AllocateNoLock(memSize, 1, TC(""));
+		u8* mem = (u8*)m_pathTableMem.AllocateNoLock(memSize, 1, TC(""));
 		BinaryWriter writer(mem, 0, memSize);
 		writer.Write7BitEncoded(parentOffset);
 		writer.Write7BitEncoded(strOffset);
-		insres.first->second = u32(mem - m_stringMem.memory);
+		insres.first->second = u32(mem - m_pathTableMem.memory);
 		return insres.first->second;
 	}
 
@@ -110,7 +112,7 @@ namespace uba
 		#if UBA_DEBUG
 		{
 			SCOPED_READ_LOCK(const_cast<CompactPathTable*>(this)->m_lock, lock)
-			UBA_ASSERTF(offset < m_stringMem.writtenSize, TC("Reading path key from offset %llu which is out of bounds (Max %llu)"), offset, m_stringMem.writtenSize);
+			UBA_ASSERTF(offset < m_pathTableMem.writtenSize, TC("Reading path key from offset %llu which is out of bounds (Max %llu)"), offset, m_pathTableMem.writtenSize);
 		}
 		#endif
 
@@ -118,7 +120,7 @@ namespace uba
 		offsets[0] = u32(offset);
 		u32 offsetCount = 0;
 
-		BinaryReader reader(m_stringMem.memory, offset, m_stringMem.writtenSize);
+		BinaryReader reader(m_pathTableMem.memory, offset, m_pathTableMem.writtenSize);
 
 		while (offset)
 		{
@@ -149,13 +151,13 @@ namespace uba
 
 	u8* CompactPathTable::GetMemory()
 	{
-		return m_stringMem.memory;
+		return m_pathTableMem.memory;
 	}
 
 	u32 CompactPathTable::GetSize()
 	{
 		SCOPED_READ_LOCK(m_lock, lock2)
-		return u32(m_stringMem.writtenSize);
+		return u32(m_pathTableMem.writtenSize);
 	}
 
 	bool CompactPathTable::GetCaseInsensitive()
@@ -165,18 +167,18 @@ namespace uba
 
 	void CompactPathTable::ReadMem(BinaryReader& reader, bool populateLookup)
 	{
-		if (!m_stringMem.memory)
-			m_stringMem.Init(CachePathTableMaxSize);
+		if (!m_pathTableMem.memory)
+			m_pathTableMem.Init(PathTableMaxSize);
 
-		u64 writtenSize = m_stringMem.writtenSize;
+		u64 writtenSize = m_pathTableMem.writtenSize;
 		u64 left = reader.GetLeft();
-		void* mem = m_stringMem.AllocateNoLock(left, 1, TC(""));
+		void* mem = m_pathTableMem.AllocateNoLock(left, 1, TC(""));
 		reader.ReadBytes(mem, left);
 
 		if (!populateLookup)
 			return;
 
-		BinaryReader reader2(m_stringMem.memory, writtenSize, m_stringMem.writtenSize);
+		BinaryReader reader2(m_pathTableMem.memory, writtenSize, m_pathTableMem.writtenSize);
 		if (!writtenSize)
 			reader2.Skip(1);
 
@@ -204,32 +206,33 @@ namespace uba
 	{
 		m_offsets.swap(other.m_offsets);
 		m_segmentOffsets.swap(other.m_segmentOffsets);
-		m_stringMem.Swap(other.m_stringMem);
+		m_pathTableMem.Swap(other.m_pathTableMem);
 		bool ci = m_caseInsensitive;
 		m_caseInsensitive = other.m_caseInsensitive;
 		other.m_caseInsensitive = ci;
 	}
 
-	CompactCasKeyTable::CompactCasKeyTable(u64 reserveSize, u64 reserveOffsetsCount)
+	CompactCasKeyTable::CompactCasKeyTable(u64 reserveOffsetsCount)
 	{
-		m_reserveSize = reserveSize;
 		if (reserveOffsetsCount)
-			m_offsets.reserve(reserveOffsetsCount);
+			m_offsets.Init(reserveOffsetsCount);
 	}
 
 	CompactCasKeyTable::~CompactCasKeyTable()
 	{
-		for (auto& kv : m_offsets)
-			if (kv.second.count > 1)
-				delete[] kv.second.stringAndCasKeyOffsets;
-
+		if (m_offsets.IsInitialized())
+			for (auto i = m_offsets.ValuesBegin(), e = m_offsets.ValuesEnd(); i!=e; ++i)
+				if (i->count > 1)
+					delete[] i->stringAndCasKeyOffsets;
 	}
 
-	u32 CompactCasKeyTable::Add(const CasKey& casKey, u64 stringOffset, u32* outRequiredCasTableSize)
+	u32 CompactCasKeyTable::Add(const CasKey& casKey, u64 stringOffset, u32& outRequiredCasTableSize)
 	{
 		SCOPED_WRITE_LOCK(m_lock, lock2)
-		if (!m_mem.memory)
-			m_mem.Init(m_reserveSize);
+		if (!m_casKeyTableMem.memory)
+			m_casKeyTableMem.Init(CasKeyTableMaxSize);
+		if (!m_offsets.IsInitialized())
+			m_offsets.Init(CasKeyLookupInitialCount);
 
 		bool added = false;
 		u32* casKeyOffset = InternalAdd(casKey, stringOffset, added);
@@ -237,28 +240,47 @@ namespace uba
 		if (added)
 		{
 			u8 bytesForStringOffset = Get7BitEncodedCount(stringOffset);
-			u8* mem = (u8*)m_mem.AllocateNoLock(bytesForStringOffset + sizeof(CasKey), 1, TC(""));
+			u8* mem = (u8*)m_casKeyTableMem.AllocateNoLock(bytesForStringOffset + sizeof(CasKey), 1, TC(""));
 			BinaryWriter writer(mem, 0, 1000);
 			writer.Write7BitEncoded(stringOffset);
 			writer.WriteCasKey(casKey);
-			*casKeyOffset = u32(mem - m_mem.memory);
-			if (outRequiredCasTableSize)
-				*outRequiredCasTableSize = (u32)m_mem.writtenSize;
+			*casKeyOffset = u32(mem - m_casKeyTableMem.memory);
+			outRequiredCasTableSize = (u32)m_casKeyTableMem.writtenSize;
 		}
-		else if (outRequiredCasTableSize)
+		else
 		{
-			BinaryReader reader(m_mem.memory, *casKeyOffset, ~0u);
+			BinaryReader reader(m_casKeyTableMem.memory, *casKeyOffset, ~0u);
 			reader.Read7BitEncoded();
-			*outRequiredCasTableSize = Max(*outRequiredCasTableSize, u32(reader.GetPosition() + sizeof(CasKey)));
+			outRequiredCasTableSize = Max(outRequiredCasTableSize, u32(reader.GetPosition() + sizeof(CasKey)));
+		}
+		return *casKeyOffset;
+	}
+
+	u32 CompactCasKeyTable::AddNoLock(const CasKey& casKey, u64 stringOffset)
+	{
+		if (!m_casKeyTableMem.memory)
+			m_casKeyTableMem.Init(CasKeyTableMaxSize);
+		if (!m_offsets.IsInitialized())
+			m_offsets.Init(CasKeyLookupInitialCount);
+
+		bool added = false;
+		u32* casKeyOffset = InternalAdd(casKey, stringOffset, added);
+		if (added)
+		{
+			u8 bytesForStringOffset = Get7BitEncodedCount(stringOffset);
+			u8* mem = (u8*)m_casKeyTableMem.AllocateNoLock(bytesForStringOffset + sizeof(CasKey), 1, TC(""));
+			BinaryWriter writer(mem, 0, 1000);
+			writer.Write7BitEncoded(stringOffset);
+			writer.WriteCasKey(casKey);
+			*casKeyOffset = u32(mem - m_casKeyTableMem.memory);
 		}
 		return *casKeyOffset;
 	}
 
 	u32* CompactCasKeyTable::InternalAdd(const CasKey& casKey, u64 stringOffset, bool& outAdded)
 	{
-		auto insres = m_offsets.try_emplace(casKey);
-		Value& value = insres.first->second;
-		if (insres.second)
+		Value& value = m_offsets.Insert(casKey);
+		if (value.count == 0)
 		{
 			value.count = 1;
 			value.single.stringOffset = u32(stringOffset);
@@ -297,7 +319,7 @@ namespace uba
 
 	void CompactCasKeyTable::GetKey(CasKey& outKey, u64 offset) const
 	{
-		BinaryReader reader(m_mem.memory, offset, ~0u);
+		BinaryReader reader(m_casKeyTableMem.memory, offset, ~0u);
 		reader.Read7BitEncoded();
 		outKey = reader.ReadCasKey();
 	}
@@ -307,11 +329,11 @@ namespace uba
 		#if UBA_DEBUG
 		{
 			SCOPED_READ_LOCK(const_cast<CompactCasKeyTable*>(this)->m_lock, lock)
-			UBA_ASSERTF(offset + sizeof(CasKey) < m_mem.writtenSize, TC("Reading cas key from offset %llu which is out of bounds (Max %llu)"), offset + sizeof(CasKey), m_mem.writtenSize);
+			UBA_ASSERTF(offset + sizeof(CasKey) < m_casKeyTableMem.writtenSize, TC("Reading cas key from offset %llu which is out of bounds (Max %llu)"), offset + sizeof(CasKey), m_casKeyTableMem.writtenSize);
 		}
 		#endif
 
-		BinaryReader reader(m_mem.memory, offset, ~0u);
+		BinaryReader reader(m_casKeyTableMem.memory, offset, ~0u);
 		u32 stringOffset = (u32)reader.Read7BitEncoded();
 		outKey = reader.ReadCasKey();
 		pathTable.GetString(outPath, stringOffset);
@@ -319,30 +341,32 @@ namespace uba
 
 	u8* CompactCasKeyTable::GetMemory()
 	{
-		return m_mem.memory;
+		return m_casKeyTableMem.memory;
 	}
 
 	u32 CompactCasKeyTable::GetSize()
 	{
 		SCOPED_READ_LOCK(m_lock, lock2)
-		return u32(m_mem.writtenSize);
+		return u32(m_casKeyTableMem.writtenSize);
 	}
 
 	void CompactCasKeyTable::ReadMem(BinaryReader& reader, bool populateLookup)
 	{
-		if (!m_mem.memory)
-			m_mem.Init(m_reserveSize);
+		if (!m_casKeyTableMem.memory)
+			m_casKeyTableMem.Init(CasKeyTableMaxSize);
+		if (!m_offsets.IsInitialized())
+			m_offsets.Init(CasKeyLookupInitialCount);
 
-		u64 writtenSize = m_mem.writtenSize;
+		u64 writtenSize = m_casKeyTableMem.writtenSize;
 
 		u64 left = reader.GetLeft();
-		void* mem = m_mem.AllocateNoLock(left, 1, TC(""));
+		void* mem = m_casKeyTableMem.AllocateNoLock(left, 1, TC(""));
 		reader.ReadBytes(mem, left);
 
 		if (!populateLookup)
 			return;
 
-		BinaryReader reader2(m_mem.memory, writtenSize, m_mem.writtenSize);
+		BinaryReader reader2(m_casKeyTableMem.memory, writtenSize, m_casKeyTableMem.writtenSize);
 		while (reader2.GetLeft())
 		{
 			u32 offset = u32(reader2.GetPosition());
@@ -357,10 +381,7 @@ namespace uba
 
 	void CompactCasKeyTable::Swap(CompactCasKeyTable& other)
 	{
-		m_offsets.swap(other.m_offsets);
-		m_mem.Swap(other.m_mem);
-		u64 rs = m_reserveSize;
-		m_reserveSize = other.m_reserveSize;
-		other.m_reserveSize = rs;
+		m_offsets.Swap(other.m_offsets);
+		m_casKeyTableMem.Swap(other.m_casKeyTableMem);
 	}
 }
diff --git a/Engine/Source/Programs/UnrealBuildAccelerator/Common/Public/UbaCompactTables.h b/Engine/Source/Programs/UnrealBuildAccelerator/Common/Public/UbaCompactTables.h
--- a/Engine/Source/Programs/UnrealBuildAccelerator/Common/Public/UbaCompactTables.h
+++ b/Engine/Source/Programs/UnrealBuildAccelerator/Common/Public/UbaCompactTables.h
@@ -18,7 +18,7 @@ namespace uba
 		CompactPathTable(bool caseSensitive, u64 reservePathCount, u64 reserveSegmentCount);
 
 		u32 Add(const tchar* str, u64 strLen);
-		u32 Add(const tchar* str, u64 strLen, u32& outRequiredCasTableSize);
+		u32 Add(const tchar* str, u64 strLen, u32& outRequiredPathTableSize);
 		u32 AddNoLock(const tchar* str, u64 strLen);
 
 		void GetString(StringBufferBase& out, u64 offset) const;
@@ -36,7 +36,7 @@ namespace uba
 	private:
 		u32 InternalAdd(const tchar* str, const tchar* stringKeyString, u64 strLen);
 		ReaderWriterLock m_lock;
-		MemoryBlock m_stringMem;
+		MemoryBlock m_pathTableMem;
 		UnorderedMap<StringKey, u32> m_offsets;
 		UnorderedMap<StringKey, u32> m_segmentOffsets;
 		bool m_caseInsensitive;
@@ -45,10 +45,11 @@ namespace uba
 	class CompactCasKeyTable
 	{
 	public:
-		CompactCasKeyTable(u64 reserveSize, u64 reserveOffsetsCount = 0);
+		CompactCasKeyTable(u64 reserveOffsetsCount = 1024);
 		~CompactCasKeyTable();
 
-		u32 Add(const CasKey& casKey, u64 stringOffset, u32* outRequiredCasTableSize = nullptr);
+		u32 Add(const CasKey& casKey, u64 stringOffset, u32& outRequiredCasTableSize);
+		u32 AddNoLock(const CasKey& casKey, u64 stringOffset);
 
 		template<typename Func>
 		void TraverseOffsets(const CasKey& casKey, const Func& func) const;
@@ -63,13 +64,13 @@ namespace uba
 		void ReadMem(BinaryReader& reader, bool populateLookup);
 		void Swap(CompactCasKeyTable& other);
 
-		u64 GetKeyCount() { return m_offsets.size(); }
+		u64 GetKeyCount() { return m_offsets.Size(); }
 
 	private:
 		u32* InternalAdd(const CasKey& casKey, u64 stringOffset, bool& outAdded);
 
 		ReaderWriterLock m_lock;
-		MemoryBlock m_mem;
+		MemoryBlock m_casKeyTableMem;
 		struct Value
 		{
 			union
@@ -84,18 +85,17 @@ namespace uba
 			u32 count; // If count is 1, single is set, otherwise stringAndCasKeyOffsets is allocated and contains two offsets per entry
 		};
 
-		UnorderedMap<CasKey, Value> m_offsets;
-		u64 m_reserveSize;
+		HashMap<CasKey, Value, true> m_offsets;
 	};
 
 
 	template<typename Func>
 	void CompactCasKeyTable::TraverseOffsets(const CasKey& casKey, const Func& func) const
 	{
-		auto findIt = m_offsets.find(casKey);
-		if (findIt == m_offsets.end())
+		const Value* valuePtr = m_offsets.Find(casKey);
+		if (!valuePtr)
 			return;
-		const Value& value = findIt->second;
+		const Value& value = *valuePtr;
 		if (value.count == 1)
 			func(value.single.casKeyOffset);
 		else
diff --git a/Engine/Source/Programs/UnrealBuildAccelerator/Common/Public/UbaHashMap.h b/Engine/Source/Programs/UnrealBuildAccelerator/Common/Public/UbaHashMap.h
--- a/Engine/Source/Programs/UnrealBuildAccelerator/Common/Public/UbaHashMap.h
+++ b/Engine/Source/Programs/UnrealBuildAccelerator/Common/Public/UbaHashMap.h
@@ -15,60 +15,81 @@ namespace uba
 			u32 next; // If this is 0 it means that it is unused, ~0u means it is the end of a chain, otherwise offset into entry array
 		};
 
-		void Init(MemoryBlock& entriesMem, MemoryBlock& valuesMem, u64 maxSize)
+		~HashMap()
 		{
-			u32 v = u32(maxSize);
-			v--;
-			v |= v >> 1;
-			v |= v >> 2;
-			v |= v >> 4;
-			v |= v >> 8;
-			v |= v >> 16;
-			v++;
-			u32 lookupSize = v;// << 1u;
-			mask = lookupSize - 1;
-			nextAvailableEntry = 1;
+			if (m_ownsMemory)
+			{
+				delete m_entriesMemory;
+				delete m_valuesMemory;
+			}
+		}
 
-			maxEntries = u32(maxSize) + 1;
-			// This code relies on that allocated memory is zeroed out
-			lookup = (u32*)entriesMem.AllocateNoLock(lookupSize*sizeof(u32), 1, TC(""));
-			entries = (Entry*)entriesMem.AllocateNoLock(maxEntries*sizeof(Entry), alignof(Entry), TC(""));
-			values = (Value*)valuesMem.AllocateNoLock(maxEntries*sizeof(Value), alignof(Value), TC(""));
-			UBA_ASSERT(&entriesMem != &valuesMem || !AllowGrow);
+		void Init(MemoryBlock& mem, u64 reserveCount)
+		{
+			reserveCount = Max(reserveCount, 4ull);
+			u64 lookupCount = GetLookupCount(reserveCount);
+			m_entriesMemory = &mem;
+			m_valuesMemory = &mem;
+			UBA_ASSERT(!AllowGrow);
+			InternalInit(reserveCount, lookupCount);
+		}
 
-			entriesMemory = &entriesMem;
-			valuesMemory = &valuesMem;
+		void Init(u64 reserveCount = 4)
+		{
+			reserveCount = Max(reserveCount, 4ull);
+			u64 lookupCount = GetLookupCount(reserveCount);
+
+			u64 lookupSize = lookupCount*sizeof(u32);
+			u64 entriesSize = sizeof(Entry)*(reserveCount+1);
+			m_entriesMemory = new MemoryBlock();
+			m_entriesMemory->Init(lookupSize + entriesSize);
+
+			u64 valuesSize = sizeof(Value)*(reserveCount+1);
+			m_valuesMemory = new MemoryBlock();
+			m_valuesMemory->Init(valuesSize);
+
+			m_ownsMemory = true;
+
+			InternalInit(reserveCount, lookupCount);
 		}
 
-		Value& Insert(const Key& key)
+		bool IsInitialized()
 		{
-			u32 index = std::hash<Key>()(key) & mask;
-			u32 entryIndex = lookup[index];
+			return m_entriesMemory != nullptr;
+		}
 
-			UBA_ASSERT(nextAvailableEntry < maxEntries || AllowGrow);
+		Value& Insert(const Key& key)
+		{
+			UBA_ASSERT(m_nextAvailableEntry < m_reserveCount || AllowGrow);
 			if constexpr (AllowGrow)
-				if (nextAvailableEntry == maxEntries)
+				if (m_nextAvailableEntry == m_reserveCount)
 				{
-					++maxEntries;
-					entriesMemory->AllocateNoLock(sizeof(Value), alignof(Value), TC(""));
-					valuesMemory->AllocateNoLock(sizeof(Value), alignof(Value), TC(""));
+					HashMap newMap;
+					newMap.Init(GetNextReserveCount(m_reserveCount));
+					for (u32 i=1, e=m_nextAvailableEntry; i!=e; ++i)
+						if (m_entries[i].next)
+							newMap.Insert(m_entries[i].key) = m_values[i];
+					Swap(newMap);
 				}
 
+			u32 index = std::hash<Key>()(key) & m_mask;
+			u32 entryIndex = m_lookup[index];
+
 			if (entryIndex == 0)
 			{
-				entryIndex = nextAvailableEntry++;
-				lookup[index] = entryIndex;
-				Entry& entry = entries[entryIndex];
+				entryIndex = m_nextAvailableEntry++;
+				m_lookup[index] = entryIndex;
+				Entry& entry = m_entries[entryIndex];
 				entry.key = key;
 				entry.next = ~0u;
-				return values[entryIndex];
+				return m_values[entryIndex];
 			}
 
 			while (true)
 			{
-				Entry& entry = entries[entryIndex];
+				Entry& entry = m_entries[entryIndex];
 				if (entry.key == key)
-					return values[entryIndex];
+					return m_values[entryIndex];
 
 				if (entry.next != ~0u)
 				{
@@ -78,36 +99,36 @@ namespace uba
 					continue;
 				}
 
-				u32 newEntryIndex = entry.next = nextAvailableEntry++;
-				Entry& newEntry = entries[newEntryIndex];
+				u32 newEntryIndex = entry.next = m_nextAvailableEntry++;
+				Entry& newEntry = m_entries[newEntryIndex];
 				newEntry.key = key;
 				newEntry.next = ~0u;
-				return values[newEntryIndex];
+				return m_values[newEntryIndex];
 			}
 		}
 
 		Value* Find(const Key& key) const
 		{
-			u32 index = std::hash<Key>()(key) & mask;
-			u32 entryIndex = lookup[index];
+			u32 index = std::hash<Key>()(key) & m_mask;
+			u32 entryIndex = m_lookup[index];
 			if (entryIndex == 0)
 				return nullptr;
-			Entry* entry = entries + entryIndex;
+			Entry* entry = m_entries + entryIndex;
 			while (true)
 			{
 				if (entry->key == key)
-					return values + entryIndex;
+					return m_values + entryIndex;
 				entryIndex = entry->next;
 				if (entryIndex == ~0u)
 					return nullptr;
-				entry = entries + entryIndex;
+				entry = m_entries + entryIndex;
 			}
 		}
 
 		const Key* GetKey(Value* value)
 		{
-			u32 pos = u32(value - values);
-			auto& entry = entries[pos];
+			u32 pos = u32(value - m_values);
+			auto& entry = m_entries[pos];
 			if (entry.next == 0)
 				return nullptr;
 			return &entry.key;
@@ -115,16 +136,16 @@ namespace uba
 
 		u32 Size()
 		{
-			return nextAvailableEntry - 1;
+			return m_nextAvailableEntry - 1;
 		}
 
 		void Erase(const Key& key)
 		{
-			u32 index = std::hash<Key>()(key) & mask;
-			u32 entryIndex = lookup[index];
+			u32 index = std::hash<Key>()(key) & m_mask;
+			u32 entryIndex = m_lookup[index];
 			if (entryIndex == 0)
 				return;
-			Entry* entry = entries + entryIndex;
+			Entry* entry = m_entries + entryIndex;
 			u32* prevNext = nullptr;
 			while (true)
 			{
@@ -133,9 +154,9 @@ namespace uba
 					if (!prevNext)
 					{
 						if (entry->next == ~0u)
-							lookup[index] = 0;
+							m_lookup[index] = 0;
 						else
-							lookup[index] = entry->next;
+							m_lookup[index] = entry->next;
 					}
 					else
 						*prevNext = entry->next;
@@ -149,23 +170,64 @@ namespace uba
 				if (entryIndex == ~0u)
 					return;
 				prevNext = &entry->next;
-				entry = entries + entryIndex;
+				entry = m_entries + entryIndex;
 			}
 		}
 
-		Value* ValuesBegin() { return values + 1; }
-		Value* ValuesEnd() { return values + nextAvailableEntry; }
+		Value* ValuesBegin() { return m_values + 1; }
+		Value* ValuesEnd() { return m_values + m_nextAvailableEntry; }
 
-		MemoryBlock* entriesMemory;
-		MemoryBlock* valuesMemory;
+		void Swap(HashMap& o)
+		{
+			u8 temp[sizeof(HashMap)];
+			memcpy(temp, this, sizeof(HashMap));
+			memcpy(this, &o, sizeof(HashMap));
+			memcpy(&o, temp, sizeof(HashMap));
+		}
 
-		Entry* entries;
-		Value* values;
+private:
+		void InternalInit(u64 reserveCount, u64 lookupCount)
+		{
+			m_mask = u32(lookupCount - 1);
+			m_nextAvailableEntry = 1;
+			m_reserveCount = u32(reserveCount) + 1;
+			// This code relies on that allocated memory is zeroed out
+			m_lookup = (u32*)m_entriesMemory->AllocateNoLock(lookupCount*sizeof(u32), 1, TC(""));
+			m_entries = (Entry*)m_entriesMemory->AllocateNoLock(m_reserveCount*sizeof(Entry), alignof(Entry), TC(""));
+			m_values = (Value*)m_valuesMemory->AllocateNoLock(m_reserveCount*sizeof(Value), alignof(Value), TC(""));
+		}
 
-		u32* lookup;
-		u32 mask;
-		u32 nextAvailableEntry;
-		u32 maxEntries;
+		u32 GetLookupCount(u64 reserveCount)
+		{
+			u32 v = u32(reserveCount);
+			v--;
+			v |= v >> 1;
+			v |= v >> 2;
+			v |= v >> 4;
+			v |= v >> 8;
+			v |= v >> 16;
+			v++;
+			return v;
+		}
+
+		u64 GetNextReserveCount(u64 reserveCount)
+		{
+			if (reserveCount < 50'000)
+				return reserveCount*2;
+			return reserveCount + 50'000;
+		}
+
+		MemoryBlock* m_entriesMemory = nullptr;
+		MemoryBlock* m_valuesMemory = nullptr;
+
+		Entry* m_entries = nullptr;
+		Value* m_values = nullptr;
+		u32* m_lookup = nullptr;
+
+		u32 m_mask = 0;
+		u32 m_nextAvailableEntry = 0;
+		u32 m_reserveCount = 0;
+		bool m_ownsMemory = false;
 	};
 
 	template<typename Key, typename Value>
diff --git a/Engine/Source/Programs/UnrealBuildAccelerator/Common/Public/UbaNetwork.h b/Engine/Source/Programs/UnrealBuildAccelerator/Common/Public/UbaNetwork.h
--- a/Engine/Source/Programs/UnrealBuildAccelerator/Common/Public/UbaNetwork.h
+++ b/Engine/Source/Programs/UnrealBuildAccelerator/Common/Public/UbaNetwork.h
@@ -18,8 +18,6 @@ namespace uba
 	static constexpr u32 SessionNetworkVersion = 34;
 	static constexpr u32 CacheNetworkVersion = 5;
 
-	static constexpr u32 CacheCasKeyTableMaxSize = 64*1024*1024;
-
 	// Messages used over network between client/server (system, storage and session)
 
 	enum SystemMessageType : u8
diff --git a/Engine/Source/Programs/UnrealBuildAccelerator/Test/Private/UbaTestAll.h b/Engine/Source/Programs/UnrealBuildAccelerator/Test/Private/UbaTestAll.h
--- a/Engine/Source/Programs/UnrealBuildAccelerator/Test/Private/UbaTestAll.h
+++ b/Engine/Source/Programs/UnrealBuildAccelerator/Test/Private/UbaTestAll.h
@@ -32,6 +32,7 @@ namespace uba
 		UBA_TEST(TestDetouredTestApp) \
 		UBA_TEST(TestRemoteDetouredTestApp) \
 		UBA_TEST(TestCompactPathTable) \
+		UBA_TEST(TestCompactCasKeyTable) \
 		UBA_TEST(TestCacheEntry) \
 		UBA_TEST(TestHashTable) \
 		UBA_TEST(TestConfig) \
diff --git a/Engine/Source/Programs/UnrealBuildAccelerator/Test/Private/UbaTestCache.h b/Engine/Source/Programs/UnrealBuildAccelerator/Test/Private/UbaTestCache.h
--- a/Engine/Source/Programs/UnrealBuildAccelerator/Test/Private/UbaTestCache.h
+++ b/Engine/Source/Programs/UnrealBuildAccelerator/Test/Private/UbaTestCache.h
@@ -238,12 +238,19 @@ namespace uba
 		return true;
 	}
 
+	bool TestCompactCasKeyTable(Logger& logger, const StringBufferBase& rootDir)
+	{
+		CompactCasKeyTable table(0);
+		table.AddNoLock(CasKeyZero, 0);
+		return true;
+	}
+
 	bool TestHashTable(Logger& logger, const StringBufferBase& rootDir)
 	{
 		{
 			MemoryBlock memoryBlock(1024*1024);
 			HashMap<u32, u32> casMap;
-			casMap.Init(memoryBlock, memoryBlock, 3);
+			casMap.Init(memoryBlock, 3);
 			if (casMap.Find(1))
 				return false;
 			casMap.Insert(1) = 2;
@@ -253,15 +260,31 @@ namespace uba
 			if (*casMap.Find(1) != 3)
 				return false;
 
-			MemoryBlock valuesMemoryBlock(1024*1024);
 			HashMap<u32, u32, true> casMap2;
-			casMap2.Init(memoryBlock, valuesMemoryBlock, 1);
-			casMap2.Insert(1) = 1;
-			casMap2.Insert(2) = 2;
-			if (*casMap2.Find(1) != 1)
-				return false;
-			if (*casMap2.Find(2) != 2)
-				return false;
+			casMap2.Init(4);
+			for (u32 i=0;i!=4; ++i)
+			{
+				casMap2.Insert(i) = i;
+				if (*casMap2.Find(i) != i)
+					return false;
+			}
+			for (u32 i=0;i!=4; ++i)
+				if (*casMap2.Find(i) != i)
+					return false;
+
+			casMap2.Insert(4) = 4;
+			for (u32 i=0;i!=5; ++i)
+				if (*casMap2.Find(i) != i)
+					return false;
+
+			for (u32 i=5;i!=1000; ++i)
+			{
+				if (casMap2.Find(i))
+					return false;
+				casMap2.Insert(i) = i;
+				if (*casMap2.Find(i) != i)
+					return false;
+			}
 		}
 
 #if 0
