[UBA]
* Fixed bug in ubaagent where terminating agent because it couldn't find xcode did not terminate fast enough so it started fetching processes from host and error out on them

diff --git a/Engine/Source/Programs/UnrealBuildAccelerator/Agent/Private/UbaAgent.cpp b/Engine/Source/Programs/UnrealBuildAccelerator/Agent/Private/UbaAgent.cpp
--- a/Engine/Source/Programs/UnrealBuildAccelerator/Agent/Private/UbaAgent.cpp
+++ b/Engine/Source/Programs/UnrealBuildAccelerator/Agent/Private/UbaAgent.cpp
@@ -905,6 +905,9 @@ namespace uba
 					}
 				}
 			});
+
+			if (xcodeDirectories.empty())
+				terminationReason.Append("Unable to populate from any Xcodes. Agent is unusable.");
 		}
 		// if we didn't want a single version, or all xcodes, then use active xcode (useful for user running their own agents)
 		else
@@ -922,9 +925,6 @@ namespace uba
 			}
 		}
 
-		if (xcodeDirectories.size() == 0)
-			terminationReason.Append("Unable to populate from any Xcodes. Agent is unusable.");
-
 		for (TString& xcodeDir : xcodeDirectories)
 		{
 			logger.Info("Populating cas with %s", xcodeDir.data());
@@ -996,14 +996,11 @@ namespace uba
 		u64 terminationTimeMs = 0;
 
 		#if UBA_USE_AWS
-		if (aws.IsTerminating(logger, terminationReason, terminationTimeMs))
-		{
-			LoggerWithWriter(g_consoleLogWriter, TC("")).Info(TC("%s. Exiting UbaAgent before starting session"), terminationReason.data);
-			return true;
-		}
+		if (terminationReason.IsEmpty())
+			aws.IsTerminating(logger, terminationReason, terminationTimeMs);
 		#endif
 
-		bool isTerminating = false;
+		bool isTerminating = !terminationReason.IsEmpty();
 		do
 		{
 			// TODO: Revisit.. it could be that something stalls during startup so we can't have this timeout
@@ -1287,19 +1284,19 @@ namespace uba
 			//#endif
 
 			bool needPrepopulate = !populateCasDirs.empty();
-			if (needPrepopulate)
+			if (needPrepopulate || isTerminating)
 				sessionClient->SetAllowSpawn(false);
 
 			storageClient->Start();
 			sessionClient->Start();
 
 			// We do population here to make sure session thread is running which will send pings to host (to prevent timeouts
-			if (needPrepopulate)
+			if (needPrepopulate && !isTerminating)
 			{
 				if (storageClient->PopulateCasFromDirs(populateCasDirs, maxProcessCount, [&]() { return isDisconnected.load(); }))
 					sessionClient->SetAllowSpawn(true);
 				else
-					terminationReason.Append(TC("Failed to prepopulate cas from local directory"));
+					terminationReason.Append(TC("Failed to prepopulate cas from local directories"));
 			}
 
 			if (terminationReason.count)
