[UBA]
* Moved CompactPathTable to uba hash map

diff --git a/Engine/Source/Programs/UnrealBuildAccelerator/Common/Private/UbaCacheServer.cpp b/Engine/Source/Programs/UnrealBuildAccelerator/Common/Private/UbaCacheServer.cpp
--- a/Engine/Source/Programs/UnrealBuildAccelerator/Common/Private/UbaCacheServer.cpp
+++ b/Engine/Source/Programs/UnrealBuildAccelerator/Common/Private/UbaCacheServer.cpp
@@ -326,6 +326,8 @@ namespace uba
 			BinaryReader casKeyTableReader(reader.GetPositionData(), 0, casKeyTableSize);
 			bucket.m_casKeyTable.ReadMem(casKeyTableReader, true);
 			reader.Skip(casKeyTableSize);
+
+			bucket.m_casKeyTable.Debug(bucket.m_pathTable);
 		}
 		outStats.totalCasKeyTableSize += casKeyTableSize;
 
@@ -1035,7 +1037,9 @@ namespace uba
 			auto& oldPathTable = bucket.m_pathTable;
 			u32 oldSize = oldPathTable.GetSize();
 			{
-				CompactPathTable newPathTable(oldPathTable.GetCaseInsensitive(), oldPathTable.GetPathCount(), oldPathTable.GetSegmentCount());
+				u64 pathReserveCount = oldPathTable.GetPathCount() + 1000; // Add extra reservation space for growth
+				u64 segmentReserveCount = oldPathTable.GetSegmentCount() + 1000;  // Add extra reservation space for growth
+				CompactPathTable newPathTable(oldPathTable.GetCaseInsensitive(), pathReserveCount, segmentReserveCount);
 				oldToNewPathOffset.Init(memoryBlock, usedPathOffsets.CountSetBits());
 
 				StringBuffer<> temp;
@@ -1061,11 +1065,12 @@ namespace uba
 			// Build new caskey table based on used offsets
 			u64 recreateCasKeyTableStart = GetTime();
 			HashMap2<u32, u32> oldToNewCasKeyOffset;
-			oldSize = bucket.m_casKeyTable.GetSize();
+			auto& oldCasKeyTable = bucket.m_casKeyTable;
+			oldSize = oldCasKeyTable.GetSize();
 			{
 				oldToNewCasKeyOffset.Init(memoryBlock, usedCasKeyOffsetsCount);
-				CompactCasKeyTable newCasKeyTable(usedCasKeyOffsetsCount);
-				BinaryReader reader2(bucket.m_casKeyTable.GetMemory(), 0, bucket.m_casKeyTable.GetSize());
+				CompactCasKeyTable newCasKeyTable(usedCasKeyOffsetsCount + Min(usedCasKeyOffsetsCount/2, 10'000ull)); // Add extra reservation space for growth
+				BinaryReader reader2(oldCasKeyTable.GetMemory(), 0, oldSize);
 				usedCasKeyOffsets.Traverse([&](u32 casKeyOffset)
 				{
 					reader2.SetPosition(casKeyOffset);
@@ -1078,7 +1083,7 @@ namespace uba
 					if (casKeyOffset != newCasKeyOffset)
 						oldToNewCasKeyOffset.Insert(casKeyOffset) = newCasKeyOffset;
 				});
-				bucket.m_casKeyTable.Swap(newCasKeyTable);
+				oldCasKeyTable.Swap(newCasKeyTable);
 			}
 			m_logger.Detail(TC("    Bucket %u Recreated caskey table. %s -> %s (%s)"), bucket.index, BytesToText(oldSize).str, BytesToText(bucket.m_casKeyTable.GetSize()).str, TimeToText(GetTime() - recreateCasKeyTableStart).str);
 
diff --git a/Engine/Source/Programs/UnrealBuildAccelerator/Common/Private/UbaCompactTables.cpp b/Engine/Source/Programs/UnrealBuildAccelerator/Common/Private/UbaCompactTables.cpp
--- a/Engine/Source/Programs/UnrealBuildAccelerator/Common/Private/UbaCompactTables.cpp
+++ b/Engine/Source/Programs/UnrealBuildAccelerator/Common/Private/UbaCompactTables.cpp
@@ -8,15 +8,14 @@ namespace uba
 	static constexpr u32 PathTableMaxSize = 20*1024*1024; // Max size for string memory
 	static constexpr u32 CasKeyTableMaxSize = 64*1024*1024;
 	static constexpr u32 CasKeyLookupInitialCount = 1024;
+	static constexpr u32 CasKeyArrayMaxSize = 16;
 
 
 	CompactPathTable::CompactPathTable(bool caseInsensitive, u64 reservePathCount, u64 reserveSegmentCount)
 	:	m_caseInsensitive(caseInsensitive)
 	{
-		if (reservePathCount)
-			m_offsets.reserve(reservePathCount);
-		if (reserveSegmentCount)
-			m_segmentOffsets.reserve(reserveSegmentCount);
+		m_offsets.Init(reservePathCount);
+		m_segmentOffsets.Init(reserveSegmentCount);
 	}
 
 	u32 CompactPathTable::Add(const tchar* str, u64 strLen)
@@ -52,9 +51,10 @@ namespace uba
 	u32 CompactPathTable::InternalAdd(const tchar* str, const tchar* stringKeyString, u64 strLen)
 	{
 		StringKey key = ToStringKeyNoCheck(stringKeyString, strLen);
-		auto insres = m_offsets.try_emplace(key);
-		if (!insres.second)
-			return insres.first->second;
+		bool added = false;
+		u32 offsetIndex = m_offsets.InsertIndex(key, added); // m_offsets might re-allocate so we need to track index instead
+		if (!added)
+			return m_offsets.GetValueFromIndex(offsetIndex);
 			
 		const tchar* seg = str;
 		u32 parentOffset = 0;
@@ -72,9 +72,12 @@ namespace uba
 		u64 segLen = strLen - (seg - str);
 		u8 bytesForParent = Get7BitEncodedCount(parentOffset);
 
+		u32& offset = m_offsets.GetValueFromIndex(offsetIndex);
+
 		StringKey segmentKey = ToStringKeyNoCheck(seg, segLen);
-		auto insres2 = m_segmentOffsets.try_emplace(segmentKey);
-		if (insres2.second)
+		added = false;
+		u32& strOffset = m_segmentOffsets.Insert(segmentKey, added);
+		if (added)
 		{
 			// Put string directly after current element and set segment offset to 0
 			u64 bytesForString = GetStringWriteSize(seg, segLen);
@@ -84,9 +87,8 @@ namespace uba
 			writer.Write7BitEncoded(parentOffset);
 			writer.Write7BitEncoded(0);
 			writer.WriteString(seg, segLen);
-			u32 offset = u32(mem - m_pathTableMem.memory);
-			insres.first->second = offset;
-			insres2.first->second = offset + bytesForParent + 1;
+			offset = u32(mem - m_pathTableMem.memory);
+			strOffset = offset + bytesForParent + 1;
 			return offset;
 		}
 
@@ -97,14 +99,13 @@ namespace uba
 		UBA_ASSERT(temp.count == segLen && wcsncmp(temp.data, seg, segLen) == 0);
 		#endif
 
-		u32 strOffset = insres2.first->second;
 		u64 memSize = bytesForParent + Get7BitEncodedCount(strOffset);
 		u8* mem = (u8*)m_pathTableMem.AllocateNoLock(memSize, 1, TC(""));
 		BinaryWriter writer(mem, 0, memSize);
 		writer.Write7BitEncoded(parentOffset);
 		writer.Write7BitEncoded(strOffset);
-		insres.first->second = u32(mem - m_pathTableMem.memory);
-		return insres.first->second;
+		offset = u32(mem - m_pathTableMem.memory);
+		return offset;
 	}
 
 	void CompactPathTable::GetString(StringBufferBase& out, u64 offset) const
@@ -192,20 +193,20 @@ namespace uba
 				u32 strOffset = u32(reader2.GetPosition());
 				StringBuffer<> seg;
 				reader2.ReadString(seg);
-				m_segmentOffsets.try_emplace(ToStringKeyNoCheck(seg.data, seg.count), strOffset);
+				m_segmentOffsets.Insert(ToStringKeyNoCheck(seg.data, seg.count)) = strOffset;
 			}
 			StringBuffer<> str;
 			GetString(str, offset);
 			if (m_caseInsensitive)
 				str.MakeLower();
-			m_offsets.try_emplace(ToStringKeyNoCheck(str.data, str.count), offset);
+			m_offsets.Insert(ToStringKeyNoCheck(str.data, str.count)) = offset;
 		}
 	}
 
 	void CompactPathTable::Swap(CompactPathTable& other)
 	{
-		m_offsets.swap(other.m_offsets);
-		m_segmentOffsets.swap(other.m_segmentOffsets);
+		m_offsets.Swap(other.m_offsets);
+		m_segmentOffsets.Swap(other.m_segmentOffsets);
 		m_pathTableMem.Swap(other.m_pathTableMem);
 		bool ci = m_caseInsensitive;
 		m_caseInsensitive = other.m_caseInsensitive;
@@ -222,8 +223,12 @@ namespace uba
 	{
 		if (m_offsets.IsInitialized())
 			for (auto i = m_offsets.ValuesBegin(), e = m_offsets.ValuesEnd(); i!=e; ++i)
-				if (i->count > 1)
-					delete[] i->stringAndCasKeyOffsets;
+			{
+				if (i->count > CasKeyArrayMaxSize)
+					delete i->stringLookup;
+				else if (i->count > 1)
+					delete[] i->array;
+			}
 	}
 
 	u32 CompactCasKeyTable::Add(const CasKey& casKey, u64 stringOffset, u32& outRequiredCasTableSize)
@@ -277,10 +282,24 @@ namespace uba
 		return *casKeyOffset;
 	}
 
+	u32 NextPow2(u32 v)
+	{
+		v--;
+		v |= v >> 1;
+		v |= v >> 2;
+		v |= v >> 4;
+		v |= v >> 8;
+		v |= v >> 16;
+		v++;
+		return v;
+	}
+
+
 	u32* CompactCasKeyTable::InternalAdd(const CasKey& casKey, u64 stringOffset, bool& outAdded)
 	{
-		Value& value = m_offsets.Insert(casKey);
-		if (value.count == 0)
+		bool added = false;
+		Value& value = m_offsets.Insert(casKey, added);
+		if (added)
 		{
 			value.count = 1;
 			value.single.stringOffset = u32(stringOffset);
@@ -293,28 +312,73 @@ namespace uba
 			if (value.single.stringOffset == stringOffset)
 				return &value.single.casKeyOffset;
 
-			u32* newOffsets = new u32[4];
-			newOffsets[0] = value.single.stringOffset;
-			newOffsets[1] = value.single.casKeyOffset;
-			newOffsets[2] = u32(stringOffset);
-			value.stringAndCasKeyOffsets = newOffsets;
+			auto newOffsets = new StringAndKey[2];
+			newOffsets[0].stringOffset = value.single.stringOffset;
+			newOffsets[0].casKeyOffset = value.single.casKeyOffset;
+			newOffsets[1].stringOffset = u32(stringOffset);
+			value.array = newOffsets;
 			value.count = 2;
 			outAdded = true;
-			return &newOffsets[3];
+			return &newOffsets[1].casKeyOffset;
 		}
 
-		for (u32 i=0, e=value.count*2; i!=e; i+=2)
-			if (value.stringAndCasKeyOffsets[i] == stringOffset)
-				return &value.stringAndCasKeyOffsets[i + 1];
-
-		u32* newOffsets = new u32[(value.count+1)*2];
-		memcpy(newOffsets, value.stringAndCasKeyOffsets, value.count*2*sizeof(u32));
-		newOffsets[value.count*2] = u32(stringOffset);
-		delete[] value.stringAndCasKeyOffsets;
-		value.stringAndCasKeyOffsets = newOffsets;
-		++value.count;
-		outAdded = true;
-		return &newOffsets[value.count*2 - 1];
+		if (value.count <= CasKeyArrayMaxSize)
+		{
+			for (u32 i=0, e=value.count; i!=e; ++i)
+				if (value.array[i].stringOffset == stringOffset)
+					return &value.array[i].casKeyOffset;
+
+			outAdded = true;
+
+			u32 index = value.count;
+			u32 newCount = index + 1;
+
+			if (newCount <= CasKeyArrayMaxSize)
+			{
+				u32 capacity = NextPow2(index);
+
+				if (newCount <= capacity)
+				{
+					value.array[index].stringOffset = u32(stringOffset);
+					value.count = newCount;
+					return &value.array[index].casKeyOffset;
+				}
+				u32 newCapacity = NextPow2(newCount);
+				auto newArray = new StringAndKey[newCapacity];
+				memcpy(newArray, value.array, index*sizeof(StringAndKey));
+
+				#if PLATFORM_WINDOWS
+				#pragma warning(push)
+				#pragma warning(disable:6385) // analyzer claims that newArray can have buffer overrun.. but can't see how that can happen
+				#pragma warning(disable:6386) // analyzer claims that newArray can have buffer overrun.. but can't see how that can happen
+				#endif
+
+				newArray[index].stringOffset = u32(stringOffset);
+				delete[] value.array;
+				value.array = newArray;
+				value.count = newCount;
+				return &newArray[index].casKeyOffset;
+
+				#if PLATFORM_WINDOWS
+				#pragma warning(pop)
+				#endif
+			}
+
+			auto& lookup = *new UnorderedMap<u32, u32>();
+			for (u32 i=0, e=value.count; i!=e; ++i)
+				lookup[value.array[i].stringOffset] = value.array[i].casKeyOffset;
+			delete[] value.array;
+
+			value.stringLookup = &lookup;
+			value.count = newCount;
+
+			return &lookup[u32(stringOffset)];
+		}
+
+		auto& lookup = *value.stringLookup;
+		auto insres = lookup.try_emplace(u32(stringOffset));
+		outAdded = insres.second;
+		return &insres.first->second;
 	}
 
 	void CompactCasKeyTable::GetKey(CasKey& outKey, u64 offset) const
@@ -379,6 +443,39 @@ namespace uba
 		}
 	}
 
+	void CompactCasKeyTable::Debug(const CompactPathTable& pathTable)
+	{
+		#if 0
+		for (auto vi=m_offsets.ValuesBegin(), ve=m_offsets.ValuesEnd(); vi!=ve; ++vi)
+		{
+			auto& value = *vi;
+			if (value.count < 2)
+				continue;
+
+			if (value.count <= CasKeyArrayMaxSize)
+			{
+				printf("COUNT: %u\n", value.count);
+				for (u32 i=0, e=value.count; i!=e; ++i)
+				{
+					StringBuffer<> str;
+					pathTable.GetString(str, value.array[i].stringOffset);
+					wprintf(L"   %s\n", str.data);
+				}
+			}
+			else
+			{
+				printf("COUNT: %u\n", u32(value.stringLookup->size()));
+				for (auto& kv : *value.stringLookup)
+				{
+					StringBuffer<> str;
+					pathTable.GetString(str, kv.first);
+					wprintf(L"   %s\n", str.data);
+				}
+			}
+		}
+		#endif
+	}
+
 	void CompactCasKeyTable::Swap(CompactCasKeyTable& other)
 	{
 		m_offsets.Swap(other.m_offsets);
diff --git a/Engine/Source/Programs/UnrealBuildAccelerator/Common/Public/UbaCompactTables.h b/Engine/Source/Programs/UnrealBuildAccelerator/Common/Public/UbaCompactTables.h
--- a/Engine/Source/Programs/UnrealBuildAccelerator/Common/Public/UbaCompactTables.h
+++ b/Engine/Source/Programs/UnrealBuildAccelerator/Common/Public/UbaCompactTables.h
@@ -30,15 +30,15 @@ namespace uba
 		void ReadMem(BinaryReader& reader, bool populateLookup);
 		void Swap(CompactPathTable& other);
 
-		u64 GetPathCount() { return m_offsets.size(); }
-		u64 GetSegmentCount() { return m_segmentOffsets.size(); }
+		u64 GetPathCount() { return m_offsets.Size(); }
+		u64 GetSegmentCount() { return m_segmentOffsets.Size(); }
 
 	private:
 		u32 InternalAdd(const tchar* str, const tchar* stringKeyString, u64 strLen);
 		ReaderWriterLock m_lock;
 		MemoryBlock m_pathTableMem;
-		UnorderedMap<StringKey, u32> m_offsets;
-		UnorderedMap<StringKey, u32> m_segmentOffsets;
+		HashMap<StringKey, u32, true> m_offsets;
+		HashMap<StringKey, u32, true> m_segmentOffsets;
 		bool m_caseInsensitive;
 	};
 
@@ -64,6 +64,8 @@ namespace uba
 		void ReadMem(BinaryReader& reader, bool populateLookup);
 		void Swap(CompactCasKeyTable& other);
 
+		void Debug(const CompactPathTable& pathTable);
+
 		u64 GetKeyCount() { return m_offsets.Size(); }
 
 	private:
@@ -71,18 +73,22 @@ namespace uba
 
 		ReaderWriterLock m_lock;
 		MemoryBlock m_casKeyTableMem;
+
+		struct StringAndKey
+		{
+			u32 stringOffset;
+			u32 casKeyOffset;
+		};
+
 		struct Value
 		{
 			union
 			{
-				u32* stringAndCasKeyOffsets;
-				struct
-				{
-					u32 stringOffset;
-					u32 casKeyOffset;
-				} single;
+				StringAndKey single;
+				StringAndKey* array;
+				UnorderedMap<u32, u32>* stringLookup;
 			};
-			u32 count; // If count is 1, single is set, otherwise stringAndCasKeyOffsets is allocated and contains two offsets per entry
+			u32 count; // If count is 1, single is set, otherwise array is allocated up to 8 elements and contains two offsets per entry. After 8 it turns to a stringlookup
 		};
 
 		HashMap<CasKey, Value, true> m_offsets;
@@ -98,8 +104,15 @@ namespace uba
 		const Value& value = *valuePtr;
 		if (value.count == 1)
 			func(value.single.casKeyOffset);
+		else if (value.count <= 16)
+		{
+			for (u32 i=0, e=value.count; i!=e; ++i)
+				func(value.array[i].casKeyOffset);
+		}
 		else
-			for (u32 i=1, e=value.count*2+1; i!=e; i+=2)
-				func(value.stringAndCasKeyOffsets[i]);
+		{
+			for (auto& kv : *value.stringLookup)
+				func(kv.second);
+		}
 	}
 }
diff --git a/Engine/Source/Programs/UnrealBuildAccelerator/Common/Public/UbaHashMap.h b/Engine/Source/Programs/UnrealBuildAccelerator/Common/Public/UbaHashMap.h
--- a/Engine/Source/Programs/UnrealBuildAccelerator/Common/Public/UbaHashMap.h
+++ b/Engine/Source/Programs/UnrealBuildAccelerator/Common/Public/UbaHashMap.h
@@ -18,36 +18,27 @@ namespace uba
 		~HashMap()
 		{
 			if (m_ownsMemory)
-			{
-				delete m_entriesMemory;
-				delete m_valuesMemory;
-			}
+				delete m_memory;
 		}
 
 		void Init(MemoryBlock& mem, u64 reserveCount)
 		{
-			reserveCount = Max(reserveCount, 4ull);
-			u64 lookupCount = GetLookupCount(reserveCount);
-			m_entriesMemory = &mem;
-			m_valuesMemory = &mem;
+			reserveCount = Max(reserveCount + 1, 4ull);
+			m_memory = &mem;
 			UBA_ASSERT(!AllowGrow);
-			InternalInit(reserveCount, lookupCount);
+			InternalInit(reserveCount, GetLookupCount(reserveCount));
 		}
 
 		void Init(u64 reserveCount = 4)
 		{
-			reserveCount = Max(reserveCount, 4ull);
+			reserveCount = Max(reserveCount + 1, 4ull);
 			u64 lookupCount = GetLookupCount(reserveCount);
 
 			u64 lookupSize = lookupCount*sizeof(u32);
 			u64 entriesSize = sizeof(Entry)*(reserveCount+1);
-			m_entriesMemory = new MemoryBlock();
-			m_entriesMemory->Init(lookupSize + entriesSize);
-
 			u64 valuesSize = sizeof(Value)*(reserveCount+1);
-			m_valuesMemory = new MemoryBlock();
-			m_valuesMemory->Init(valuesSize);
-
+			m_memory = new MemoryBlock();
+			m_memory->Init(lookupSize + entriesSize + valuesSize);
 			m_ownsMemory = true;
 
 			InternalInit(reserveCount, lookupCount);
@@ -55,20 +46,32 @@ namespace uba
 
 		bool IsInitialized()
 		{
-			return m_entriesMemory != nullptr;
+			return m_memory != nullptr;
 		}
 
 		Value& Insert(const Key& key)
+		{
+			bool added = false;
+			return Insert(key, added);
+		}
+
+		Value& Insert(const Key& key, bool& outAdded)
+		{
+			u32 index = InsertIndex(key, outAdded);
+			return m_values[index];
+		}
+
+		u32 InsertIndex(const Key& key, bool& outAdded)
 		{
 			UBA_ASSERT(m_nextAvailableEntry < m_reserveCount || AllowGrow);
 			if constexpr (AllowGrow)
 				if (m_nextAvailableEntry == m_reserveCount)
 				{
+					// Create a new map, populate and swap. Note that we need the entries and values to be at the same index as before
 					HashMap newMap;
 					newMap.Init(GetNextReserveCount(m_reserveCount));
 					for (u32 i=1, e=m_nextAvailableEntry; i!=e; ++i)
-						if (m_entries[i].next)
-							newMap.Insert(m_entries[i].key) = m_values[i];
+						newMap.Insert(m_entries[i].key) = m_values[i];
 					Swap(newMap);
 				}
 
@@ -82,14 +85,15 @@ namespace uba
 				Entry& entry = m_entries[entryIndex];
 				entry.key = key;
 				entry.next = ~0u;
-				return m_values[entryIndex];
+				outAdded = true;
+				return entryIndex;
 			}
 
 			while (true)
 			{
 				Entry& entry = m_entries[entryIndex];
 				if (entry.key == key)
-					return m_values[entryIndex];
+					return entryIndex;
 
 				if (entry.next != ~0u)
 				{
@@ -103,7 +107,8 @@ namespace uba
 				Entry& newEntry = m_entries[newEntryIndex];
 				newEntry.key = key;
 				newEntry.next = ~0u;
-				return m_values[newEntryIndex];
+				outAdded = true;
+				return newEntryIndex;
 			}
 		}
 
@@ -134,6 +139,11 @@ namespace uba
 			return &entry.key;
 		}
 
+		Value& GetValueFromIndex(u32 index)
+		{
+			return m_values[index];
+		}
+
 		u32 Size()
 		{
 			return m_nextAvailableEntry - 1;
@@ -141,6 +151,7 @@ namespace uba
 
 		void Erase(const Key& key)
 		{
+			UBA_ASSERT(!AllowGrow);
 			u32 index = std::hash<Key>()(key) & m_mask;
 			u32 entryIndex = m_lookup[index];
 			if (entryIndex == 0)
@@ -190,11 +201,11 @@ private:
 		{
 			m_mask = u32(lookupCount - 1);
 			m_nextAvailableEntry = 1;
-			m_reserveCount = u32(reserveCount) + 1;
+			m_reserveCount = u32(reserveCount);
 			// This code relies on that allocated memory is zeroed out
-			m_lookup = (u32*)m_entriesMemory->AllocateNoLock(lookupCount*sizeof(u32), 1, TC(""));
-			m_entries = (Entry*)m_entriesMemory->AllocateNoLock(m_reserveCount*sizeof(Entry), alignof(Entry), TC(""));
-			m_values = (Value*)m_valuesMemory->AllocateNoLock(m_reserveCount*sizeof(Value), alignof(Value), TC(""));
+			m_lookup = (u32*)m_memory->AllocateNoLock(lookupCount*sizeof(u32), 1, TC(""));
+			m_entries = (Entry*)m_memory->AllocateNoLock(m_reserveCount*sizeof(Entry), alignof(Entry), TC(""));
+			m_values = (Value*)m_memory->AllocateNoLock(m_reserveCount*sizeof(Value), alignof(Value), TC(""));
 		}
 
 		u32 GetLookupCount(u64 reserveCount)
@@ -217,8 +228,7 @@ private:
 			return reserveCount + 50'000;
 		}
 
-		MemoryBlock* m_entriesMemory = nullptr;
-		MemoryBlock* m_valuesMemory = nullptr;
+		MemoryBlock* m_memory = nullptr;
 
 		Entry* m_entries = nullptr;
 		Value* m_values = nullptr;
diff --git a/Engine/Source/Programs/UnrealBuildAccelerator/Test/Private/UbaTestCache.h b/Engine/Source/Programs/UnrealBuildAccelerator/Test/Private/UbaTestCache.h
--- a/Engine/Source/Programs/UnrealBuildAccelerator/Test/Private/UbaTestCache.h
+++ b/Engine/Source/Programs/UnrealBuildAccelerator/Test/Private/UbaTestCache.h
@@ -240,8 +240,27 @@ namespace uba
 
 	bool TestCompactCasKeyTable(Logger& logger, const StringBufferBase& rootDir)
 	{
-		CompactCasKeyTable table(0);
-		table.AddNoLock(CasKeyZero, 0);
+		CompactCasKeyTable table;
+		u32 offset0 = table.AddNoLock(CasKeyZero, 0);
+		u32 offset1 = table.AddNoLock(CasKeyZero, 1);
+		u32 offset2 = table.AddNoLock(CasKeyZero, 2);
+
+		CompactCasKeyTable table2;
+		BinaryReader reader(table.GetMemory(), 0, table.GetSize());
+		table2.ReadMem(reader, true);
+		if (table2.AddNoLock(CasKeyZero, 0) != offset0)
+			return false;
+		if (table2.AddNoLock(CasKeyZero, 1) != offset1)
+			return false;
+		if (table2.AddNoLock(CasKeyZero, 2) != offset2)
+			return false;
+
+		CompactCasKeyTable table3(table.GetKeyCount());
+		table3.AddNoLock(CasKeyZero, 0);
+		table3.AddNoLock(CasKeyZero, 1);
+		table3.AddNoLock(CasKeyZero, 2);
+		table3.AddNoLock(CasKeyZero, 3);
+		table3.AddNoLock(CasKeyZero, 4);
 		return true;
 	}
 
