[UBA]
* Added support for uba not failing if failing to create file mappings (To work around some odd linux issue Brandon saw on his machine)

diff --git a/Engine/Source/Programs/UnrealBuildAccelerator/Common/Private/UbaProcess.cpp b/Engine/Source/Programs/UnrealBuildAccelerator/Common/Private/UbaProcess.cpp
--- a/Engine/Source/Programs/UnrealBuildAccelerator/Common/Private/UbaProcess.cpp
+++ b/Engine/Source/Programs/UnrealBuildAccelerator/Common/Private/UbaProcess.cpp
@@ -88,24 +88,28 @@ namespace uba
 		TString currentString;
 	};
 
-	ProcessImpl::ProcessImpl(Session& session, u32 id, ProcessImpl* parent)
+	ProcessImpl::ProcessImpl(Session& session, u32 id, ProcessImpl* parent, bool detourEnabled)
 	:	m_session(session)
 	,	m_parentProcess(parent)
 	,	m_id(id)
-	,	m_comMemory(m_session.m_processCommunicationAllocator.Alloc(TC("")))
+	,	m_comMemory(detourEnabled ? m_session.m_processCommunicationAllocator.Alloc(TC("")) : FileMappingAllocator::Allocation())
 		#if !PLATFORM_WINDOWS
-	,	m_cancelEvent(*new (m_comMemory.memory) Event)
-	,	m_writeEvent(*new (m_comMemory.memory + sizeof(Event)) Event)
-	,	m_readEvent(*new (m_comMemory.memory + sizeof(Event)*2) Event)
+	,	m_cancelEvent(*(m_comMemory.memory ? (new (m_comMemory.memory) Event) : (Event*)nullptr))
+	,	m_writeEvent(*(m_comMemory.memory ? (new (m_comMemory.memory + sizeof(Event)) Event) : (Event*)nullptr))
+	,	m_readEvent(*(m_comMemory.memory ? (new (m_comMemory.memory + sizeof(Event)*2) Event) : (Event*)nullptr))
 		#endif
+	,	m_detourEnabled(detourEnabled)
 	,	m_writtenFilesLock(parent ? parent->m_writtenFilesLock : *new ReaderWriterLock())
 	,	m_writtenFiles(parent ? parent->m_writtenFiles : *new UnorderedMap<TString, WrittenFile>())
 	,	m_tempFilesLock(parent ? parent->m_tempFilesLock : *new ReaderWriterLock())
 	,	m_tempFiles(parent ? parent->m_tempFiles : *new UnorderedMap<StringKey, WrittenFile>)
 	{
-		m_cancelEvent.Create(true, true);
-		m_writeEvent.Create(false, true);
-		m_readEvent.Create(false, true);
+		if (m_comMemory.memory)
+		{
+			m_cancelEvent.Create(true, true);
+			m_writeEvent.Create(false, true);
+			m_readEvent.Create(false, true);
+		}
 	}
 
 	ProcessImpl::~ProcessImpl()
@@ -145,10 +149,8 @@ namespace uba
 		}
 	}
 
-	bool ProcessImpl::Start(const ProcessStartInfo& startInfo, bool runningRemote, void* environment, bool async, bool enableDetour)
+	bool ProcessImpl::Start(const ProcessStartInfo& startInfo, bool runningRemote, void* environment, bool async)
 	{
-		m_detourEnabled = enableDetour;
-
 		m_startTime = GetTime();
 
 		m_startInfo = startInfo;
@@ -556,6 +558,7 @@ namespace uba
 
 		m_hasExited = true;
 
+		if (m_comMemory.memory)
 		{
 			#if !PLATFORM_WINDOWS
 			SCOPED_WRITE_LOCK(m_comMemoryLock, lock);
diff --git a/Engine/Source/Programs/UnrealBuildAccelerator/Common/Private/UbaSession.cpp b/Engine/Source/Programs/UnrealBuildAccelerator/Common/Private/UbaSession.cpp
--- a/Engine/Source/Programs/UnrealBuildAccelerator/Common/Private/UbaSession.cpp
+++ b/Engine/Source/Programs/UnrealBuildAccelerator/Common/Private/UbaSession.cpp
@@ -1043,6 +1043,7 @@ namespace uba
 		UBA_ASSERTF(m_shouldWriteToDisk || m_allowMemoryMaps, TC("Can't disable both should write to disk and allow memory maps"));
 
 		m_storeObjFilesCompressed = info.storeObjFilesCompressed;
+		m_allowLocalDetour = true;
 
 		m_detailedTrace = info.detailedTrace;
 		m_traceChildProcesses = info.traceChildProcesses;
@@ -1126,10 +1127,11 @@ namespace uba
 
 		m_fileMappingBuffer.AddTransient(TC("FileMappings"));
 
-		u64 reserveSize = CommunicationMemSize * 512;
-		if (!m_processCommunicationAllocator.Init(CommunicationMemSize, reserveSize))
-			return false;
-
+		if (!m_processCommunicationAllocator.Init(CommunicationMemSize, CommunicationMemSize * 512))
+		{
+			m_allowLocalDetour = false;
+			m_logger.Warning(TC("Failed to create process communication allocator.. local detouring will be disabled."));
+		}
 		if (!CreateProcessJobObject())
 			return false;
 
@@ -1257,6 +1259,7 @@ namespace uba
 	{
 		FlushDeadProcesses();
 		ValidateStartInfo(startInfo);
+		enableDetour &= m_allowLocalDetour;
 		return InternalRunProcess(startInfo, async, nullptr, enableDetour);
 	}
 
@@ -1295,9 +1298,9 @@ namespace uba
 			si.rules = GetRules(si);
 
 		void* env = GetProcessEnvironmentVariables();
-		auto process = new ProcessImpl(*this, processId, parent);
+		auto process = new ProcessImpl(*this, processId, parent, enableDetour);
 		ProcessHandle h(process);
-		if (!process->Start(startInfo, m_runningRemote, env, async, enableDetour))
+		if (!process->Start(startInfo, m_runningRemote, env, async))
 			return {};
 
 		si.logFile = originalLogFile;
diff --git a/Engine/Source/Programs/UnrealBuildAccelerator/Common/Private/UbaSessionClient.cpp b/Engine/Source/Programs/UnrealBuildAccelerator/Common/Private/UbaSessionClient.cpp
--- a/Engine/Source/Programs/UnrealBuildAccelerator/Common/Private/UbaSessionClient.cpp
+++ b/Engine/Source/Programs/UnrealBuildAccelerator/Common/Private/UbaSessionClient.cpp
@@ -1695,7 +1695,7 @@ namespace uba
 
 					void* env = GetProcessEnvironmentVariables();
 
-					auto process = new ProcessImpl(*this, startInfo.processId, nullptr);
+					auto process = new ProcessImpl(*this, startInfo.processId, nullptr, true);
 
 					activeProcesses.emplace_back(process);
 					ProcessRec* rec = &activeProcesses.back();
@@ -1810,7 +1810,7 @@ namespace uba
 							session.m_processFinished(&process);
 					};
 
-					process->Start(startInfo, true, env, true, true);
+					process->Start(startInfo, true, env, true);
 				}
 
 				RemoveInactiveProcesses();
diff --git a/Engine/Source/Programs/UnrealBuildAccelerator/Common/Public/UbaProcess.h b/Engine/Source/Programs/UnrealBuildAccelerator/Common/Public/UbaProcess.h
--- a/Engine/Source/Programs/UnrealBuildAccelerator/Common/Public/UbaProcess.h
+++ b/Engine/Source/Programs/UnrealBuildAccelerator/Common/Public/UbaProcess.h
@@ -56,12 +56,12 @@ namespace uba
 		virtual ProcessExecutionType GetExecutionType() const override { return m_detourEnabled ? ProcessExecutionType_Detoured : ProcessExecutionType_Native; }
 		virtual bool IsChild() override { return m_parentProcess != nullptr; }
 
-		ProcessImpl(Session& session, u32 id, ProcessImpl* parent);
+		ProcessImpl(Session& session, u32 id, ProcessImpl* parent, bool detourEnabled);
 		~ProcessImpl();
 
 		struct PipeReader;
 
-		bool Start(const ProcessStartInfo& startInfo, bool runningRemote, void* environment, bool async, bool enableDetour);
+		bool Start(const ProcessStartInfo& startInfo, bool runningRemote, void* environment, bool async);
 
 		bool IsActive();
 		bool IsCancelled();
diff --git a/Engine/Source/Programs/UnrealBuildAccelerator/Common/Public/UbaSession.h b/Engine/Source/Programs/UnrealBuildAccelerator/Common/Public/UbaSession.h
--- a/Engine/Source/Programs/UnrealBuildAccelerator/Common/Public/UbaSession.h
+++ b/Engine/Source/Programs/UnrealBuildAccelerator/Common/Public/UbaSession.h
@@ -188,6 +188,7 @@ namespace uba
 		bool m_traceChildProcesses;
 		bool m_logToFile;
 		bool m_storeObjFilesCompressed;
+		bool m_allowLocalDetour;
 
 		u64 m_keepOutputFileMemoryMapsThreshold;
 
