[UBA]
* Fixed so a missing rsp file is not causing errors when searching the cache for matching cache entries

diff --git a/Engine/Source/Programs/UnrealBuildAccelerator/Common/Private/UbaCacheClient.cpp b/Engine/Source/Programs/UnrealBuildAccelerator/Common/Private/UbaCacheClient.cpp
--- a/Engine/Source/Programs/UnrealBuildAccelerator/Common/Private/UbaCacheClient.cpp
+++ b/Engine/Source/Programs/UnrealBuildAccelerator/Common/Private/UbaCacheClient.cpp
@@ -1176,6 +1176,8 @@ namespace uba
 					StringBuffer<> fullPath;
 					FixPath(rsp.data, workingDir.data, workingDir.count, fullPath);
 					CasKey rspCasKey = rootPaths.NormalizeAndHashFile(m_logger, rsp.data);
+					if (rspCasKey == CasKeyZero)
+						return CasKeyZero;
 					hasher.Update(&rspCasKey, sizeof(CasKey));
 				}
 			}
diff --git a/Engine/Source/Programs/UnrealBuildAccelerator/Common/Private/UbaRootPaths.cpp b/Engine/Source/Programs/UnrealBuildAccelerator/Common/Private/UbaRootPaths.cpp
--- a/Engine/Source/Programs/UnrealBuildAccelerator/Common/Private/UbaRootPaths.cpp
+++ b/Engine/Source/Programs/UnrealBuildAccelerator/Common/Private/UbaRootPaths.cpp
@@ -122,7 +122,7 @@ namespace uba
 	CasKey RootPaths::NormalizeAndHashFile(Logger& logger, const tchar* filename) const
 	{
 		FileAccessor file(logger, filename);
-		if (!file.OpenMemoryRead())
+		if (!file.OpenMemoryRead(0, false))
 			return CasKeyZero;
 
 		bool wasNormalized = false;
