[UBA]
* Fixed temporary expiration time logic.. it could end up in a state where it does not lower expiration time even though bucket cas table is too big. Also renamed variable to match what it is

diff --git a/Engine/Source/Programs/UnrealBuildAccelerator/Common/Private/UbaCacheServer.cpp b/Engine/Source/Programs/UnrealBuildAccelerator/Common/Private/UbaCacheServer.cpp
--- a/Engine/Source/Programs/UnrealBuildAccelerator/Common/Private/UbaCacheServer.cpp
+++ b/Engine/Source/Programs/UnrealBuildAccelerator/Common/Private/UbaCacheServer.cpp
@@ -156,7 +156,7 @@ namespace uba
 		// Times are in file time from creation of cache database
 		Atomic<u64> lastSavedTime;
 		Atomic<u64> lastUsedTime;
-		u64 longestUnusedTime = 0; 
+		u64 oldestUsedTime = 0; 
 
 		u64 expirationTimeSeconds = 0;
 		u32 index = ~0u;
@@ -675,7 +675,7 @@ namespace uba
 		ReaderWriterLock globalStatsLock;
 		u64 now = GetSystemTimeAsFileTime();
 		u64 oldest = 0;
-		u64 longestUnusedTime = 0;
+		u64 oldestUsedTime = 0;
 
 		u32 workerCount = m_server.GetWorkerCount();
 		u32 workerCountToUse = workerCount > 0 ? workerCount - 1 : 0;
@@ -710,7 +710,7 @@ namespace uba
 			bool checkInputsForDeletes = m_checkInputsForDeletedCas && !deletedCasFiles.empty();
 
 			oldest = 0;
-			longestUnusedTime = 0;
+			oldestUsedTime = 0;
 			totalEntryCount = 0;
 
 			m_server.ParallelFor(workerCountToUseForBuckets, sortedBuckets, [&](auto& it)
@@ -741,20 +741,18 @@ namespace uba
 
 				u64 lastUseTimeLimit = 0; // This is the time relative to server startup time
 
-				if (deleteIteration == 0 && m_bucketCasTableMaxSize && bucket.longestUnusedTime)
+				if (deleteIteration == 0 && m_bucketCasTableMaxSize && bucket.oldestUsedTime)
 				{
 					// If bucket is larger than max we will take longestUnused and reduce by one hour
 					u64 bucketCasTableSize = bucket.m_casKeyTable.GetSize();
 					if (bucketCasTableSize >= m_bucketCasTableMaxSize)
 					{
-						u64 longestUnusedSeconds = GetFileTimeAsSeconds((now - m_creationTime) - bucket.longestUnusedTime);
-						if (longestUnusedSeconds < bucketExpirationTimeSeconds)
-						{
-							if (longestUnusedSeconds > 60*60)
-								longestUnusedSeconds -= 60*60;
-							bucketExpirationTimeSeconds = longestUnusedSeconds;
-							m_logger.Info(TC("    Set temporary expiration time for bucket %u to %s to reduce cas table size"), bucket.index, TimeToText(MsToTime(bucketExpirationTimeSeconds*1000), true).str);
-						}
+						u64 longestUnusedSeconds = GetFileTimeAsSeconds((now - m_creationTime) - bucket.oldestUsedTime);
+
+						bucketExpirationTimeSeconds = Min(longestUnusedSeconds, bucketExpirationTimeSeconds);
+						if (bucketExpirationTimeSeconds > 60*60)
+							bucketExpirationTimeSeconds -= 60*60;
+						m_logger.Detail(TC("    Set temporary expiration time for bucket %u to %s to reduce cas table size"), bucket.index, TimeToText(MsToTime(bucketExpirationTimeSeconds*1000), true).str);
 					}
 				}
 
@@ -765,7 +763,7 @@ namespace uba
 						lastUseTimeLimit = GetSecondsAsFileTime(secondsSinceCreation - bucketExpirationTimeSeconds);
 				}
 
-				bucket.longestUnusedTime = 0;
+				bucket.oldestUsedTime = 0;
 				bucket.totalEntryCount = 0;
 				bucket.totalEntrySize = 0;
 
@@ -897,10 +895,10 @@ namespace uba
 							SCOPED_WRITE_LOCK(globalStatsLock, l);
 							if (!oldest || entry.creationTime < oldest)
 								oldest = entry.creationTime;
-							if (!bucket.longestUnusedTime || lastUsedTime < bucket.longestUnusedTime)
-								bucket.longestUnusedTime = lastUsedTime;
-							if (!longestUnusedTime || lastUsedTime < longestUnusedTime)
-								longestUnusedTime = lastUsedTime;
+							if (!bucket.oldestUsedTime || lastUsedTime < bucket.oldestUsedTime)
+								bucket.oldestUsedTime = lastUsedTime;
+							if (!oldestUsedTime || lastUsedTime < oldestUsedTime)
+								oldestUsedTime = lastUsedTime;
 						}
 
 						for (auto v : touchedCas)
@@ -1145,9 +1143,9 @@ namespace uba
 		}
 
 		u64 oldestTime = oldest ? GetFileTimeAsTime(now - (m_creationTime + oldest)) : 0;
-		u64 longestUnusedTime2 = longestUnusedTime ? GetFileTimeAsTime(now - (m_creationTime + longestUnusedTime)) : 0;
+		u64 longestUnused = oldestUsedTime ? GetFileTimeAsTime(now - (m_creationTime + oldestUsedTime)) : 0;
 		u64 duration = GetTime() - startTime;
-		m_logger.Info(TC("Maintenance done! (%s) CasFiles: %s (%s) Entries: %s Oldest: %s LongestUnused: %s MaintenanceMem: %s/%s"), TimeToText(duration).str, CountToText(totalCasCount - deletedCasCount).str, BytesToText(totalCasSize).str, CountToText(totalEntryCount.load()).str, TimeToText(oldestTime, true).str, TimeToText(longestUnusedTime2, true).str, BytesToText(maxCommittedMemory).str, BytesToText(m_maintenanceReserveSize).str);
+		m_logger.Info(TC("Maintenance done! (%s) CasFiles: %s (%s) Entries: %s Oldest: %s LongestUnused: %s MaintenanceMem: %s/%s"), TimeToText(duration).str, CountToText(totalCasCount - deletedCasCount).str, BytesToText(totalCasSize).str, CountToText(totalEntryCount.load()).str, TimeToText(oldestTime, true).str, TimeToText(longestUnused, true).str, BytesToText(maxCommittedMemory).str, BytesToText(m_maintenanceReserveSize).str);
 		
 		m_longestMaintenance = Max(m_longestMaintenance, duration);
 
