[UBA]
* Changed so uba cache server maintenance loops if bucket is still overflowing (will stop looping once bucket is not overflowing or a new connection is made)

diff --git a/Engine/Source/Programs/UnrealBuildAccelerator/Common/Private/UbaCacheServer.cpp b/Engine/Source/Programs/UnrealBuildAccelerator/Common/Private/UbaCacheServer.cpp
--- a/Engine/Source/Programs/UnrealBuildAccelerator/Common/Private/UbaCacheServer.cpp
+++ b/Engine/Source/Programs/UnrealBuildAccelerator/Common/Private/UbaCacheServer.cpp
@@ -546,22 +546,34 @@ namespace uba
 
 	bool CacheServer::RunMaintenance(bool force, const Function<bool()>& shouldExit)
 	{
-		if (m_addsSinceMaintenance == 0 && !force)
+		if (m_addsSinceMaintenance == 0 && !force && !m_bucketIsOverflowing)
 			return true;
 
-		SCOPED_WRITE_LOCK(m_connectionsLock, lock2);
-		if (!m_connections.empty())
-			return true;
-		m_isRunningMaintenance = true;
-		lock2.Leave();
+		do
+		{
+			SCOPED_WRITE_LOCK(m_connectionsLock, lock2);
+			if (!m_connections.empty())
+				return true;
+			m_isRunningMaintenance = true;
+			lock2.Leave();
 
-		auto g = MakeGuard([&]()
-			{
-				SCOPED_WRITE_LOCK(m_connectionsLock, lock2);
-				m_isRunningMaintenance = false;
-			});
+			auto g = MakeGuard([&]()
+				{
+					SCOPED_WRITE_LOCK(m_connectionsLock, lock2);
+					m_isRunningMaintenance = false;
+				});
 
+			if(!RunMaintenanceInternal(shouldExit))
+				return false;
+
+		}
+		while (m_bucketIsOverflowing);
+
+		return true;
+	}
 
+	bool CacheServer::RunMaintenanceInternal(const Function<bool()>& shouldExit)
+	{
 		//m_forceAllSteps = true;
 		bool forceAllSteps = m_forceAllSteps;
 		m_forceAllSteps = false;
@@ -965,6 +977,8 @@ namespace uba
 
 		m_logger.Detail(TC("  Deleted %llu cas files and %llu cache entries over %u buckets (%s)"), deletedCasCount, deleteEntryCount.load(), u32(m_buckets.size()), TimeToText(GetTime() - deleteCacheEntriesStartTime).str);
 
+		m_bucketIsOverflowing = false;
+
 		if (shouldExit())
 			return true;
 
@@ -1032,18 +1046,18 @@ namespace uba
 
 			// Build new path table based on used offsets
 			HashMap2<u32, u32> oldToNewPathOffset;
-			auto& oldPathTable = bucket.m_pathTable;
-			u32 oldSize = oldPathTable.GetSize();
+			auto& pathTable = bucket.m_pathTable;
+			u32 oldSize = pathTable.GetSize();
 			{
-				u64 pathReserveCount = oldPathTable.GetPathCount() + 1000; // Add extra reservation space for growth
-				u64 segmentReserveCount = oldPathTable.GetSegmentCount() + 1000;  // Add extra reservation space for growth
-				CompactPathTable newPathTable(oldPathTable.GetCaseInsensitive(), pathReserveCount, segmentReserveCount);
+				u64 pathReserveCount = pathTable.GetPathCount() + 1000; // Add extra reservation space for growth
+				u64 segmentReserveCount = pathTable.GetSegmentCount() + 1000;  // Add extra reservation space for growth
+				CompactPathTable newPathTable(pathTable.GetCaseInsensitive(), pathReserveCount, segmentReserveCount);
 				oldToNewPathOffset.Init(memoryBlock, usedPathOffsets.CountSetBits());
 
 				StringBuffer<> temp;
 				usedPathOffsets.Traverse([&](u32 pathOffset)
 				{
-					oldPathTable.GetString(temp.Clear(), pathOffset);
+					pathTable.GetString(temp.Clear(), pathOffset);
 					u32 newPathOffset = newPathTable.AddNoLock(temp.data, temp.count);
 
 					#if 0
@@ -1055,10 +1069,12 @@ namespace uba
 					if (pathOffset != newPathOffset)
 						oldToNewPathOffset.Insert(pathOffset) = newPathOffset;
 				});
-				oldPathTable.Swap(newPathTable);
+				pathTable.Swap(newPathTable);
 			}
-			m_logger.Detail(TC("    Bucket %u Recreated path table. %s -> %s (%s)"), bucket.index, BytesToText(oldSize).str, BytesToText(oldPathTable.GetSize()).str, TimeToText(GetTime() - recreatePathTableStart).str);
+			m_logger.Detail(TC("    Bucket %u Recreated path table. %s -> %s (%s)"), bucket.index, BytesToText(oldSize).str, BytesToText(pathTable.GetSize()).str, TimeToText(GetTime() - recreatePathTableStart).str);
 
+			if (pathTable.GetSize() >= m_bucketCasTableMaxSize) // Still overflowing?
+				m_bucketIsOverflowing = true;
 
 			// Build new caskey table based on used offsets
 			u64 recreateCasKeyTableStart = GetTime();
diff --git a/Engine/Source/Programs/UnrealBuildAccelerator/Common/Public/UbaCacheServer.h b/Engine/Source/Programs/UnrealBuildAccelerator/Common/Public/UbaCacheServer.h
--- a/Engine/Source/Programs/UnrealBuildAccelerator/Common/Public/UbaCacheServer.h
+++ b/Engine/Source/Programs/UnrealBuildAccelerator/Common/Public/UbaCacheServer.h
@@ -60,6 +60,7 @@ namespace uba
 		struct Bucket;
 		struct LoadStats;
 
+		bool RunMaintenanceInternal(const Function<bool()>& shouldExit);
 		bool LoadBucket(Bucket& bucket, BinaryReader& reader, u32 databaseVersion, LoadStats& outStats);
 		bool SaveBucket(u64 bucketId, Bucket& bucket);
 		bool SaveNoLock();
@@ -91,6 +92,7 @@ namespace uba
 		Atomic<u64> m_cacheKeyFetchCount;
 		Atomic<u64> m_cacheKeyHitCount;
 		Atomic<bool> m_isRunningMaintenance;
+		Atomic<bool> m_bucketIsOverflowing;
 
 		ReaderWriterLock m_bucketsLock;
 		Map<u64, Bucket> m_buckets;
