[UBA]
* Changed code to use new[]/delete[] instead of malloc/free just to make sure it ends up in mimalloc. There is a shutdown bug somewhere related to double deletes, memory overwrites or something and hopefully by changing behavior maybe we can catch this earlier if it ever happens again

diff --git a/Engine/Source/Programs/UnrealBuildAccelerator/Common/Private/UbaStorage.cpp b/Engine/Source/Programs/UnrealBuildAccelerator/Common/Private/UbaStorage.cpp
--- a/Engine/Source/Programs/UnrealBuildAccelerator/Common/Private/UbaStorage.cpp
+++ b/Engine/Source/Programs/UnrealBuildAccelerator/Common/Private/UbaStorage.cpp
@@ -11,6 +11,8 @@
 #include "UbaWorkManager.h"
 #include <oodle2.h>
 
+#define UBA_CHECK_BUFFER_SLOTS 0
+
 namespace uba
 {
 	constexpr u32 CasTableVersion = IsWindows ? 32 : 34;
@@ -36,21 +38,40 @@ namespace uba
 			m_slots.pop_back();
 			return back;
 		}
-		return (u8*)malloc(BufferSlotSize);
+#if UBA_CHECK_BUFFER_SLOTS
+		u8* res = new u8[BufferSlotSize + 12];
+		res += 8;
+		*(u32*)(res - 4) = 0xdeadbeef;
+		*(u32*)(res + BufferSlotSize) = 0xdeadbeef;
+		return res;
+#else
+		return new u8[BufferSlotSize];
+#endif
 	}
 
 	void BufferSlots::Push(u8* slot)
 	{
 		if (!slot)
 			return;
+
+#if UBA_CHECK_BUFFER_SLOTS
+		UBA_ASSERT(*(u32*)(slot - 4) == 0xdeadbeef);
+		UBA_ASSERT(*(u32*)(slot + BufferSlotSize) == 0xdeadbeef);
+#endif
+
 		SCOPED_WRITE_LOCK(m_slotsLock, lock);
 		m_slots.push_back(slot);
 	}
 
 	BufferSlots::~BufferSlots()
 	{
+#if UBA_CHECK_BUFFER_SLOTS
 		for (u8* slot : m_slots)
-			free(slot);
+			delete[] (slot-8);
+#else
+		for (u8* slot : m_slots)
+			delete[] slot;
+#endif
 	}
 
 	void StorageCreateInfo::Apply(Config& config)
