[UBA]
* Added so network traffic from cache client is tracked in visualizer

diff --git a/Engine/Source/Programs/UnrealBuildAccelerator/Common/Private/UbaSessionServer.cpp b/Engine/Source/Programs/UnrealBuildAccelerator/Common/Private/UbaSessionServer.cpp
--- a/Engine/Source/Programs/UnrealBuildAccelerator/Common/Private/UbaSessionServer.cpp
+++ b/Engine/Source/Programs/UnrealBuildAccelerator/Common/Private/UbaSessionServer.cpp
@@ -494,6 +494,12 @@ namespace uba
 		return m_server;
 	}
 
+	void SessionServer::RegisterNetworkTrafficProvider(const NetworkTrafficProvider& provider)
+	{
+		ScopedCriticalSection lock(m_remoteProcessAndSessionLock);
+		m_provider = provider;
+	}
+
 	void SessionServer::OnDisconnected(const Guid& clientUid, u32 clientId)
 	{
 		u32 returnCount = 0;
@@ -1908,6 +1914,10 @@ namespace uba
 	void SessionServer::TraceSessionUpdate()
 	{
 		u32 sessionIndex = 1;
+
+		u64 serverSend = m_server.GetTotalSentBytes();
+		u64 serverRecv = m_server.GetTotalRecvBytes();
+
 		ScopedCriticalSection lock(m_remoteProcessAndSessionLock);
 		for (auto sptr : m_clientSessions)
 		{
@@ -1918,11 +1928,17 @@ namespace uba
 				m_trace.SessionUpdate(sessionIndex, stats.connectionCount, stats.send, stats.recv, s.lastPing, s.memAvail, s.memTotal, s.cpuLoad);
 			++sessionIndex;
 		}
+		if (m_provider)
+		{
+			u64 send;
+			u64 recv;
+			m_provider(send, recv);
+			serverSend += send;
+			serverRecv += recv;
+		}
 		lock.Leave();
 
 		float cpuLoad = UpdateCpuLoad();
-		u64 serverSend = m_server.GetTotalSentBytes();
-		u64 serverRecv = m_server.GetTotalRecvBytes();
 		u64 memAvail = m_memAvail;
 		u64 memTotal = m_memTotal;
 		m_trace.SessionUpdate(0, 0, serverSend, serverRecv, 0, memAvail, memTotal, cpuLoad);
diff --git a/Engine/Source/Programs/UnrealBuildAccelerator/Common/Public/UbaCacheClient.h b/Engine/Source/Programs/UnrealBuildAccelerator/Common/Public/UbaCacheClient.h
--- a/Engine/Source/Programs/UnrealBuildAccelerator/Common/Public/UbaCacheClient.h
+++ b/Engine/Source/Programs/UnrealBuildAccelerator/Common/Public/UbaCacheClient.h
@@ -59,6 +59,7 @@ namespace uba
 		inline MutableLogger& GetLogger() { return m_logger; }
 		inline NetworkClient& GetClient() { return m_client; }
 		inline StorageImpl& GetStorage() { return m_storage; }
+		inline Session& GetSession() { return m_session; }
 
 	private:
 		struct Bucket;
diff --git a/Engine/Source/Programs/UnrealBuildAccelerator/Common/Public/UbaSessionServer.h b/Engine/Source/Programs/UnrealBuildAccelerator/Common/Public/UbaSessionServer.h
--- a/Engine/Source/Programs/UnrealBuildAccelerator/Common/Public/UbaSessionServer.h
+++ b/Engine/Source/Programs/UnrealBuildAccelerator/Common/Public/UbaSessionServer.h
@@ -65,6 +65,9 @@ namespace uba
 		// Get the network server used by this session
 		NetworkServer& GetServer();
 
+		using NetworkTrafficProvider = Function<void(u64& outSent, u64& outReceive)>;
+		void RegisterNetworkTrafficProvider(const NetworkTrafficProvider& provider);
+
 	protected:
 		struct ClientSession;
 
@@ -169,6 +172,8 @@ namespace uba
 		struct ApplicationData { ReaderWriterLock lock; Vector<u8> bytes; };
 		UnorderedMap<StringKey, ApplicationData> m_applicationData;
 
+		NetworkTrafficProvider m_provider;
+
 		Event m_memoryThreadEvent;
 		Thread m_memoryThread;
 		Atomic<u64> m_memAvail;
diff --git a/Engine/Source/Programs/UnrealBuildAccelerator/Host/Private/UbaExports.cpp b/Engine/Source/Programs/UnrealBuildAccelerator/Host/Private/UbaExports.cpp
--- a/Engine/Source/Programs/UnrealBuildAccelerator/Host/Private/UbaExports.cpp
+++ b/Engine/Source/Programs/UnrealBuildAccelerator/Host/Private/UbaExports.cpp
@@ -581,6 +581,13 @@ extern "C"
 			delete networkClient;
 			return nullptr;
 		}
+
+		session->RegisterNetworkTrafficProvider([networkClient](u64& outSent, u64& outReceive)
+			{
+				outSent = networkClient->GetTotalSentBytes();
+				outReceive = networkClient->GetTotalRecvBytes();
+			});
+
 		CacheClientCreateInfo info{writer, storage, *networkClient, *session};
 		info.Apply(GetConfig());
 
@@ -647,6 +654,11 @@ extern "C"
 	void CacheClient_Destroy(uba::CacheClient* cacheClient)
 	{
 		using namespace uba;
+
+		auto& session = (SessionServer&)cacheClient->GetSession();
+		session.RegisterNetworkTrafficProvider({});
+
+
 		auto& networkClient = (NetworkClientWithBackend&)cacheClient->GetClient();
 		networkClient.Disconnect();
 		delete cacheClient;
