[UBA]
* Added more stats in console output for cache server

diff --git a/Engine/Source/Programs/UnrealBuildAccelerator/Common/Private/UbaCacheServer.cpp b/Engine/Source/Programs/UnrealBuildAccelerator/Common/Private/UbaCacheServer.cpp
--- a/Engine/Source/Programs/UnrealBuildAccelerator/Common/Private/UbaCacheServer.cpp
+++ b/Engine/Source/Programs/UnrealBuildAccelerator/Common/Private/UbaCacheServer.cpp
@@ -549,6 +549,32 @@ namespace uba
 		if (m_addsSinceMaintenance == 0 && !force && !m_bucketIsOverflowing)
 			return true;
 
+		KernelStats& kernelStats = KernelStats::GetGlobal();
+		auto kernelReset = MakeGuard([&]() { kernelStats = {}; });
+
+		{
+			u32 peakConnectionCount = m_peakConnectionCount;
+			m_peakConnectionCount = 0;
+
+			u64 startTime = GetTime();
+			auto& storageStats = m_storage.Stats();
+			u64 hits = m_cacheKeyHitCount;
+			u64 miss = m_cacheKeyFetchCount - hits;
+			m_logger.Info(TC("Stats since boot (%s ago)"), TimeToText(startTime - m_bootTime, true).str);
+			m_logger.Info(TC("  CacheServer %s hits, %s misses"), CountToText(hits).str, CountToText(miss).str);
+			u64 recvCount = storageStats.sendCas.count.load();
+			u64 sendCount = storageStats.recvCas.count.load();
+			m_logger.Info(TC("  StorageServer cas %s (%s) sent, %s (%s) received"), CountToText(recvCount).str, BytesToText(storageStats.sendCasBytesComp).str, CountToText(sendCount).str, BytesToText(storageStats.recvCasBytesComp).str);
+			
+			if (m_lastMaintenance)
+				m_logger.Info(TC("Stats since last maintenance (%s ago)"), TimeToText(startTime - m_lastMaintenance, true).str);
+			m_logger.Info(TC("  Peak connection count: %4u"), peakConnectionCount);
+			auto& sentTimer = m_server.GetTotalSentTimer();
+			m_logger.Info(TC("  SendSocket         %8u %9s"), sentTimer.count.load(), TimeToText(sentTimer.time).str);
+			m_logger.Info(TC("     Bytes            %17s"), BytesToText(m_server.GetTotalSentBytes()).str);
+			kernelStats.Print(m_logger, false);
+		}
+
 		do
 		{
 			SCOPED_WRITE_LOCK(m_connectionsLock, lock2);
@@ -565,7 +591,6 @@ namespace uba
 
 			if(!RunMaintenanceInternal(shouldExit))
 				return false;
-
 		}
 		while (m_bucketIsOverflowing);
 
@@ -584,18 +609,6 @@ namespace uba
 
 		u64 startTime = GetTime();
 
-		if (entriesAdded)
-		{
-			auto& storageStats = m_storage.Stats();
-			u64 hits = m_cacheKeyHitCount;
-			u64 miss = m_cacheKeyFetchCount - hits;
-			m_logger.Info(TC("Stats since boot (%s ago)"), TimeToText(startTime - m_bootTime, true).str);
-			m_logger.Info(TC("  CacheServer %s hits, %s misses"), CountToText(hits).str, CountToText(miss).str);
-			u64 recvCount = storageStats.sendCas.count.load();
-			u64 sendCount = storageStats.recvCas.count.load();
-			m_logger.Info(TC("  StorageServer cas %s (%s) sent, %s (%s) received"), CountToText(recvCount).str, BytesToText(storageStats.sendCasBytesComp).str, CountToText(sendCount).str, BytesToText(storageStats.recvCasBytesComp).str);
-		}
-
 		if (m_shouldWipe)
 		{
 			m_shouldWipe = false;
@@ -607,8 +620,7 @@ namespace uba
 		}
 		else
 		{
-			TimeToText lastTime(startTime - m_lastMaintenance, true);
-			m_logger.Info(TC("Maintenance started after %u added cache entries (Ran last time %s ago)"), addsSinceMaintenance, (m_lastMaintenance ? lastTime.str : TC("<never>")));
+			m_logger.Info(TC("Maintenance starting after %u added cache entries"), addsSinceMaintenance);
 		}
 
 		m_lastMaintenance = startTime;
@@ -993,7 +1005,7 @@ namespace uba
 
 			if (!bucket.hasDeletedEntries && !forceAllSteps)
 			{
-				m_logger.Detail(TC("    Bucket %u skipped updating. (%s entries)"), bucket.index, CountToText(bucket.totalEntryCount.load()).str);
+				//m_logger.Detail(TC("    Bucket %u skipped updating. (%s entries)"), bucket.index, CountToText(bucket.totalEntryCount.load()).str);
 				return;
 			}
 			bucket.hasDeletedEntries = false;
@@ -1245,6 +1257,7 @@ namespace uba
 			auto insres = m_connections.try_emplace(connectionInfo.GetId());
 			auto& connection = insres.first->second;
 			connection.clientVersion = clientVersion;
+			m_peakConnectionCount = Max(m_peakConnectionCount, u32(m_connections.size()));
 			return true;
 		}
 		case CacheMessageType_StorePathTable:
diff --git a/Engine/Source/Programs/UnrealBuildAccelerator/Common/Private/UbaNetworkServer.cpp b/Engine/Source/Programs/UnrealBuildAccelerator/Common/Private/UbaNetworkServer.cpp
--- a/Engine/Source/Programs/UnrealBuildAccelerator/Common/Private/UbaNetworkServer.cpp
+++ b/Engine/Source/Programs/UnrealBuildAccelerator/Common/Private/UbaNetworkServer.cpp
@@ -890,6 +890,11 @@ namespace uba
 		return m_recvBytes;
 	}
 
+	Timer& NetworkServer::GetTotalSentTimer()
+	{
+		return m_sendTimer;
+	}
+
 	u32 NetworkServer::GetConnectionCount()
 	{
 		SCOPED_READ_LOCK(m_connectionsLock, lock);
diff --git a/Engine/Source/Programs/UnrealBuildAccelerator/Common/Public/UbaCacheServer.h b/Engine/Source/Programs/UnrealBuildAccelerator/Common/Public/UbaCacheServer.h
--- a/Engine/Source/Programs/UnrealBuildAccelerator/Common/Public/UbaCacheServer.h
+++ b/Engine/Source/Programs/UnrealBuildAccelerator/Common/Public/UbaCacheServer.h
@@ -109,6 +109,7 @@ namespace uba
 		u64 m_lastMaintenance = 0;
 		u64 m_longestMaintenance = 0;
 		u64 m_expirationTimeSeconds = 0;
+		u32 m_peakConnectionCount = 0;
 		bool m_dbfileDirty = false;
 
 		bool m_checkInputsForDeletedCas = true;
diff --git a/Engine/Source/Programs/UnrealBuildAccelerator/Common/Public/UbaNetworkServer.h b/Engine/Source/Programs/UnrealBuildAccelerator/Common/Public/UbaNetworkServer.h
--- a/Engine/Source/Programs/UnrealBuildAccelerator/Common/Public/UbaNetworkServer.h
+++ b/Engine/Source/Programs/UnrealBuildAccelerator/Common/Public/UbaNetworkServer.h
@@ -92,6 +92,7 @@ namespace uba
 		MutableLogger& GetLogger();
 		u64 GetTotalSentBytes();
 		u64 GetTotalRecvBytes();
+		Timer& GetTotalSentTimer();
 		u32 GetConnectionCount();
 		void GetClientStats(ClientStats& out, u32 clientId);
 
@@ -175,7 +176,6 @@ namespace uba
 		UnorderedMap<u32, Client> m_clients;
 
 		Timer m_sendTimer;
-		Timer m_sendRawTimer;
 		Timer m_encryptTimer;
 		Timer m_decryptTimer;
 		Atomic<u64> m_sendBytes;
