[UBA]
* Fixed so breadcrumbs are visible in entries where processes are being reused
* Fixed bug where session stats was added for local processes when doing process reuse which caused visualizer to show wrong data
* Fixed so CreateProcess stat is not counted twice

diff --git a/Engine/Source/Programs/UnrealBuildAccelerator/Common/Private/UbaProcess.cpp b/Engine/Source/Programs/UnrealBuildAccelerator/Common/Private/UbaProcess.cpp
--- a/Engine/Source/Programs/UnrealBuildAccelerator/Common/Private/UbaProcess.cpp
+++ b/Engine/Source/Programs/UnrealBuildAccelerator/Common/Private/UbaProcess.cpp
@@ -154,6 +154,7 @@ namespace uba
 		m_startTime = GetTime();
 
 		m_startInfo = startInfo;
+		m_runningRemote = runningRemote;
 
 		FixPathSeparators(m_startInfo.workingDirStr.data());
 		FixPathSeparators(m_startInfo.logFileStr.data());
@@ -203,15 +204,15 @@ namespace uba
 		m_session.ProcessAdded(*this, 0);
 
 		if (async)
-			m_messageThread.Start([this, runningRemote, environment]()
+			m_messageThread.Start([this, environment]()
 				{
-					ThreadRun(runningRemote, environment);
+					ThreadRun(environment);
 					ThreadExit();
 					return 0;
 				});
 		else
 		{
-			ThreadRun(runningRemote, environment);
+			ThreadRun(environment);
 			ThreadExit();
 		}
 		return true;
@@ -440,7 +441,7 @@ namespace uba
 		m_writeEvent.Set();
 	}
 
-	void ProcessImpl::ThreadRun(bool runningRemote, void* environment)
+	void ProcessImpl::ThreadRun(void* environment)
 	{
 		KernelStatsScope kernelStatsScope(m_kernelStats);
 		StorageStatsScope storageStatsScope(m_storageStats);
@@ -462,7 +463,7 @@ namespace uba
 
 		while (!IsCancelled())
 		{
-			exitCode = InternalCreateProcess(runningRemote, environment, m_comMemory.handle, m_comMemory.offset);
+			exitCode = InternalCreateProcess(environment, m_comMemory.handle, m_comMemory.offset);
 
 			if (exitCode == 0)
 			{
@@ -1241,7 +1242,8 @@ namespace uba
 		kernelStats.Add(m_kernelStats);
 
 		processStats.Write(statsWriter);
-		m_sessionStats.Write(statsWriter);
+		if (m_runningRemote)
+			m_sessionStats.Write(statsWriter);
 		m_storageStats.Write(statsWriter);
 		kernelStats.Write(statsWriter);
 		BinaryReader statsReader(statsWriter.GetData(), 0, statsWriter.GetPosition());
@@ -1422,7 +1424,7 @@ namespace uba
 #endif
 
 
-	u32 ProcessImpl::InternalCreateProcess(bool runningRemote, void* environment, FileMappingHandle communicationHandle, u64 communicationOffset)
+	u32 ProcessImpl::InternalCreateProcess(void* environment, FileMappingHandle communicationHandle, u64 communicationOffset)
 	{
 		SCOPED_WRITE_LOCK(m_initLock, initLock);
 		Logger& logger = m_session.m_logger;
@@ -1652,7 +1654,7 @@ namespace uba
 			payload.communicationOffset = communicationOffset;
 			payload.rulesIndex = m_startInfo.rules->index;
 			payload.version = ProcessMessageVersion;
-			payload.runningRemote = runningRemote;
+			payload.runningRemote = m_runningRemote;
 			payload.allowKeepFilesInMemory = m_session.m_allowKeepFilesInMemory;
 			payload.allowOutputFiles = m_session.m_allowOutputFiles;
 			payload.suppressLogging = m_session.m_suppressLogging;
@@ -1844,7 +1846,7 @@ namespace uba
 				envvars.push_back(comIdVar.data);
 				envvars.push_back(workingDir.data);
 				envvars.push_back(rulesStr.data);
-				if (runningRemote)
+				if (m_runningRemote)
 					envvars.push_back("UBA_REMOTE=1");
 				if (!logFile.IsEmpty())
 					envvars.push_back(logFile.data);
diff --git a/Engine/Source/Programs/UnrealBuildAccelerator/Common/Private/UbaSession.cpp b/Engine/Source/Programs/UnrealBuildAccelerator/Common/Private/UbaSession.cpp
--- a/Engine/Source/Programs/UnrealBuildAccelerator/Common/Private/UbaSession.cpp
+++ b/Engine/Source/Programs/UnrealBuildAccelerator/Common/Private/UbaSession.cpp
@@ -2355,7 +2355,7 @@ namespace uba
 		if (!outNewProcess)
 			return true;
 
-		m_trace.ProcessEnvironmentUpdated(process.GetId(), outNextProcess.description.c_str(), statsReader.GetPositionData(), statsReader.GetLeft());
+		m_trace.ProcessEnvironmentUpdated(process.GetId(), outNextProcess.description.c_str(), statsReader.GetPositionData(), statsReader.GetLeft(), process.GetStartInfo().breadcrumbs);
 
 		return true;
 	}
@@ -2435,14 +2435,13 @@ namespace uba
 	{
 		if (!resetStats)
 			return true;
+		UBA_ASSERT(!m_runningRemote); // local do not write session stats
 		StackBinaryWriter<16 * 1024> writer;
 		process.m_processStats.Write(writer);
-		process.m_sessionStats.Write(writer);
 		process.m_storageStats.Write(writer);
 		process.m_kernelStats.Write(writer);
-		m_trace.ProcessEnvironmentUpdated(process.GetId(), reason, writer.GetData(), writer.GetPosition());
+		m_trace.ProcessEnvironmentUpdated(process.GetId(), reason, writer.GetData(), writer.GetPosition(), process.GetStartInfo().breadcrumbs);
 		process.m_processStats = {};
-		process.m_sessionStats = {};
 		process.m_storageStats = {};
 		process.m_kernelStats = {};
 		return true;
diff --git a/Engine/Source/Programs/UnrealBuildAccelerator/Common/Private/UbaSessionServer.cpp b/Engine/Source/Programs/UnrealBuildAccelerator/Common/Private/UbaSessionServer.cpp
--- a/Engine/Source/Programs/UnrealBuildAccelerator/Common/Private/UbaSessionServer.cpp
+++ b/Engine/Source/Programs/UnrealBuildAccelerator/Common/Private/UbaSessionServer.cpp
@@ -1414,9 +1414,16 @@ namespace uba
 	bool SessionServer::HandleUpdateEnvironment(const ConnectionInfo& connectionInfo, BinaryReader& reader, BinaryWriter& writer)
 	{
 		u32 processId = reader.ReadU32();
+		SCOPED_WRITE_LOCK(m_processesLock, lock);
+		auto findIt = m_processes.find(processId);
+		if (findIt == m_processes.end())
+			return m_logger.Error(TC("Failed to find process for id %u when receiving custom message"), processId);
+		ProcessHandle h(findIt->second);
+		lock.Leave();
+
 		StringBuffer<> reason;
 		reader.ReadString(reason);
-		m_trace.ProcessEnvironmentUpdated(processId, reason.data, reader.GetPositionData(), reader.GetLeft());
+		m_trace.ProcessEnvironmentUpdated(processId, reason.data, reader.GetPositionData(), reader.GetLeft(), h.GetStartInfo().breadcrumbs);
 		return true;
 	}
 
diff --git a/Engine/Source/Programs/UnrealBuildAccelerator/Common/Private/UbaTrace.cpp b/Engine/Source/Programs/UnrealBuildAccelerator/Common/Private/UbaTrace.cpp
--- a/Engine/Source/Programs/UnrealBuildAccelerator/Common/Private/UbaTrace.cpp
+++ b/Engine/Source/Programs/UnrealBuildAccelerator/Common/Private/UbaTrace.cpp
@@ -256,12 +256,13 @@ namespace uba
 		writer.WriteString(description);
 	}
 
-	void Trace::ProcessEnvironmentUpdated(u32 processId, const tchar* reason, const u8* data, u64 dataSize)
+	void Trace::ProcessEnvironmentUpdated(u32 processId, const tchar* reason, const u8* data, u64 dataSize, const tchar* breadcrumbs)
 	{
 		BEGIN_TRACE_ENTRY(ProcessEnvironmentUpdated);
 		writer.WriteU32(processId);
 		writer.WriteString(reason);
 		writer.WriteBytes(data, dataSize);
+		writer.WriteString(breadcrumbs);
 	}
 
 	void Trace::ProcessExited(u32 processId, u32 exitCode, const u8* data, u64 dataSize, const Vector<ProcessLogLine>& logLines, const tchar* breadcrumbs)
diff --git a/Engine/Source/Programs/UnrealBuildAccelerator/Common/Public/UbaProcess.h b/Engine/Source/Programs/UnrealBuildAccelerator/Common/Public/UbaProcess.h
--- a/Engine/Source/Programs/UnrealBuildAccelerator/Common/Public/UbaProcess.h
+++ b/Engine/Source/Programs/UnrealBuildAccelerator/Common/Public/UbaProcess.h
@@ -36,7 +36,7 @@ namespace uba
 		friend ProcessHandle;
 	};
 
-	class ProcessImpl : public Process
+	class ProcessImpl final : public Process
 	{
 	public:
 
@@ -70,7 +70,7 @@ namespace uba
 		bool WaitForRead(PipeReader& outReader, PipeReader& errReader);
 		void SetWritten();
 
-		void ThreadRun(bool runningRemote, void* environment);
+		void ThreadRun(void* environment);
 		void ThreadExit();
 
 		bool HandleMessage(BinaryReader& reader, BinaryWriter& writer);
@@ -91,7 +91,7 @@ namespace uba
 		bool PollStdPipes(PipeReader& outReader, PipeReader& errReader, int timeoutMs = -1);
 		#endif
 
-		u32 InternalCreateProcess(bool runningRemote, void* environment, FileMappingHandle communicationHandle, u64 communicationOffset);
+		u32 InternalCreateProcess(void* environment, FileMappingHandle communicationHandle, u64 communicationOffset);
 		u32	InternalExitProcess(bool cancel);
 		void ClearTempFiles();
 		void WaitForParent();
@@ -131,6 +131,7 @@ namespace uba
 		u32 m_exitCode = ~0u;
 		u32 m_messageCount = 0;
 		Atomic<bool> m_hasExited;
+		bool m_runningRemote = false;
 		bool m_messageSuccess = true;
 		bool m_echoOn = true;
 		bool m_gotExitMessage = false;
diff --git a/Engine/Source/Programs/UnrealBuildAccelerator/Common/Public/UbaTrace.h b/Engine/Source/Programs/UnrealBuildAccelerator/Common/Public/UbaTrace.h
--- a/Engine/Source/Programs/UnrealBuildAccelerator/Common/Public/UbaTrace.h
+++ b/Engine/Source/Programs/UnrealBuildAccelerator/Common/Public/UbaTrace.h
@@ -68,7 +68,7 @@ namespace uba
 	using Color = u32;
 	inline Color toColor(u8 r, u8 g, u8 b) { return (r << 16) + (g << 8) + b; }
 
-	static constexpr u32 TraceVersion = 34;
+	static constexpr u32 TraceVersion = 35;
 	static constexpr u32 TraceReadCompatibilityVersion = 6;
 
 	class Trace : public WorkTracker
@@ -86,7 +86,7 @@ namespace uba
 		void SessionSummary(u32 sessionId, const u8* data, u64 dataSize);
 		void SessionDisconnect(u32 sessionId);
 		void ProcessAdded(u32 sessionId, u32 processId, const tchar* description);
-		void ProcessEnvironmentUpdated(u32 processId, const tchar* reason, const u8* data, u64 dataSize);
+		void ProcessEnvironmentUpdated(u32 processId, const tchar* reason, const u8* data, u64 dataSize, const tchar* breadcrumbs);
 		void ProcessExited(u32 processId, u32 exitCode, const u8* data, u64 dataSize, const Vector<ProcessLogLine>& logLines, const tchar* breadcrumbs);
 		void ProcessReturned(u32 processId, const StringView& reason);
 		void ProxyCreated(u32 clientId, const tchar* proxyName);
diff --git a/Engine/Source/Programs/UnrealBuildAccelerator/Detours/Private/Windows/UbaDetoursFunctionsKernelBase.inl b/Engine/Source/Programs/UnrealBuildAccelerator/Detours/Private/Windows/UbaDetoursFunctionsKernelBase.inl
--- a/Engine/Source/Programs/UnrealBuildAccelerator/Detours/Private/Windows/UbaDetoursFunctionsKernelBase.inl
+++ b/Engine/Source/Programs/UnrealBuildAccelerator/Detours/Private/Windows/UbaDetoursFunctionsKernelBase.inl
@@ -2808,6 +2808,7 @@ BOOL Detoured_CreateProcessW(LPCWSTR lpApplicationName, LPWSTR lpCommandLine, LP
 	if (isChild)
 	{
 		TimerScope ts(g_stats.createProcess);
+		--g_stats.createProcess.count; // Don't want double count for one process
 		SCOPED_WRITE_LOCK(g_communicationLock, pcs);
 		BinaryWriter writer;
 		writer.WriteByte(MessageType_StartProcess);
diff --git a/Engine/Source/Programs/UnrealBuildAccelerator/Visualizer/Private/UbaTraceReader.cpp b/Engine/Source/Programs/UnrealBuildAccelerator/Visualizer/Private/UbaTraceReader.cpp
--- a/Engine/Source/Programs/UnrealBuildAccelerator/Visualizer/Private/UbaTraceReader.cpp
+++ b/Engine/Source/Programs/UnrealBuildAccelerator/Visualizer/Private/UbaTraceReader.cpp
@@ -602,9 +602,7 @@ namespace uba
 			memcpy(process.stats.data(), dataStart, dataEnd - dataStart);
 
 			if (out.version >= 34)
-			{
 				process.breadcrumbs = reader.ReadString();
-			}
 
 			process.createFilesTime = processStats.createFile.time;
 			process.writeFilesTime = Max(processStats.writeFiles.time, processStats.sendFiles.time);
@@ -657,7 +655,8 @@ namespace uba
 			StorageStats storageStats;
 			KernelStats kernelStats;
 			processStats.Read(reader, out.version);
-			sessionStats.Read(reader, out.version);
+			if (process.isRemote || out.version < 35)
+				sessionStats.Read(reader, out.version);
 			storageStats.Read(reader, out.version);
 			kernelStats.Read(reader, out.version);
 
@@ -665,6 +664,9 @@ namespace uba
 			process.stats.resize(dataEnd - dataStart);
 			memcpy(process.stats.data(), dataStart, dataEnd - dataStart);
 
+			if (out.version >= 34)
+				process.breadcrumbs = reader.ReadString();
+
 			process.exitCode = 0u;
 			process.stop = time;
 			process.bitmapDirty = true;
diff --git a/Engine/Source/Programs/UnrealBuildAccelerator/Visualizer/Private/UbaVisualizer.cpp b/Engine/Source/Programs/UnrealBuildAccelerator/Visualizer/Private/UbaVisualizer.cpp
--- a/Engine/Source/Programs/UnrealBuildAccelerator/Visualizer/Private/UbaVisualizer.cpp
+++ b/Engine/Source/Programs/UnrealBuildAccelerator/Visualizer/Private/UbaVisualizer.cpp
@@ -1770,7 +1770,7 @@ namespace uba
 			const auto& breadcrumbs = process.breadcrumbs;
 			if (!breadcrumbs.empty())
 			{
-				constexpr TString::size_type maxLineLen = 37;
+				constexpr TString::size_type maxLineLen = 50;
 				logger.Info(L"");
 				logger.Info(L"  ------------ Breadcrumbs ------------");
 				for (TString::size_type lineStart = 0, lineEnd = 0; lineEnd < breadcrumbs.size(); lineStart = lineEnd + 1)
