[UBA]
* Fixed bug that caused crashes when trying to load uba files of trace version 34

diff --git a/Engine/Source/Programs/UnrealBuildAccelerator/Visualizer/Private/UbaTraceReader.cpp b/Engine/Source/Programs/UnrealBuildAccelerator/Visualizer/Private/UbaTraceReader.cpp
--- a/Engine/Source/Programs/UnrealBuildAccelerator/Visualizer/Private/UbaTraceReader.cpp
+++ b/Engine/Source/Programs/UnrealBuildAccelerator/Visualizer/Private/UbaTraceReader.cpp
@@ -664,7 +664,7 @@ namespace uba
 			process.stats.resize(dataEnd - dataStart);
 			memcpy(process.stats.data(), dataStart, dataEnd - dataStart);
 
-			if (out.version >= 34)
+			if (out.version >= 35)
 				process.breadcrumbs = reader.ReadString();
 
 			process.exitCode = 0u;
