[UBA]
* Changed so breadcrumbs are sent in process added instead of exit to more easily handle process reuse

diff --git a/Engine/Source/Programs/UnrealBuildAccelerator/Common/Private/UbaProcess.cpp b/Engine/Source/Programs/UnrealBuildAccelerator/Common/Private/UbaProcess.cpp
--- a/Engine/Source/Programs/UnrealBuildAccelerator/Common/Private/UbaProcess.cpp
+++ b/Engine/Source/Programs/UnrealBuildAccelerator/Common/Private/UbaProcess.cpp
@@ -1217,7 +1217,7 @@ namespace uba
 		StringBuffer<> reason;
 		reader.ReadString(reason);
 		bool resetStats = reader.ReadBool();
-		bool result = m_session.UpdateEnvironment(*this, reason.data, resetStats);
+		bool result = m_session.UpdateEnvironment(*this, reason, resetStats);
 		writer.WriteBool(result);
 		return true;
 	}
diff --git a/Engine/Source/Programs/UnrealBuildAccelerator/Common/Private/UbaScheduler.cpp b/Engine/Source/Programs/UnrealBuildAccelerator/Common/Private/UbaScheduler.cpp
--- a/Engine/Source/Programs/UnrealBuildAccelerator/Common/Private/UbaScheduler.cpp
+++ b/Engine/Source/Programs/UnrealBuildAccelerator/Common/Private/UbaScheduler.cpp
@@ -529,6 +529,7 @@ namespace uba
 			outNextProcess.workingDir = si.workingDir;
 			outNextProcess.description = si.description;
 			outNextProcess.logFile = si.logFile;
+			outNextProcess.breadcrumbs = si.breadcrumbs;
 
 			#if UBA_DEBUG
 			auto PrepPath = [this](StringBufferBase& out, const ProcessStartInfo& psi)
diff --git a/Engine/Source/Programs/UnrealBuildAccelerator/Common/Private/UbaSession.cpp b/Engine/Source/Programs/UnrealBuildAccelerator/Common/Private/UbaSession.cpp
--- a/Engine/Source/Programs/UnrealBuildAccelerator/Common/Private/UbaSession.cpp
+++ b/Engine/Source/Programs/UnrealBuildAccelerator/Common/Private/UbaSession.cpp
@@ -1485,8 +1485,9 @@ namespace uba
 	{
 		u32 processId = process.GetId();
 
+		auto& startInfo = process.GetStartInfo();
 		if (!process.IsChild() || m_traceChildProcesses)
-			m_trace.ProcessAdded(sessionId, processId, process.GetStartInfo().GetDescription());
+			m_trace.ProcessAdded(sessionId, processId, ToView(startInfo.GetDescription()), ToView(startInfo.breadcrumbs));
 
 		SCOPED_WRITE_LOCK(m_processesLock, lock);
 		bool success = m_processes.try_emplace(processId, ProcessHandle(&process)).second;
@@ -1512,7 +1513,7 @@ namespace uba
 			u32 exitCode = process.GetExitCode();
 			Vector<ProcessLogLine> emptyLines;
 			auto& logLines = (exitCode != 0 || m_detailedTrace) ? process.m_logLines : emptyLines;
-			m_trace.ProcessExited(id, exitCode, writer.GetData(), writer.GetPosition(), logLines, process.GetStartInfo().breadcrumbs);
+			m_trace.ProcessExited(id, exitCode, writer.GetData(), writer.GetPosition(), logLines);
 			SCOPED_WRITE_LOCK(m_processStatsLock, lock);
 			m_processStats.Add(process.m_processStats);
 			m_stats.Add(process.m_sessionStats);
@@ -2355,7 +2356,7 @@ namespace uba
 		if (!outNewProcess)
 			return true;
 
-		m_trace.ProcessEnvironmentUpdated(process.GetId(), outNextProcess.description.c_str(), statsReader.GetPositionData(), statsReader.GetLeft(), process.GetStartInfo().breadcrumbs);
+		m_trace.ProcessEnvironmentUpdated(process.GetId(), outNextProcess.description, statsReader.GetPositionData(), statsReader.GetLeft(), outNextProcess.breadcrumbs);
 
 		return true;
 	}
@@ -2431,7 +2432,7 @@ namespace uba
 		return true;
 	}
 
-	bool Session::UpdateEnvironment(ProcessImpl& process, const tchar* reason, bool resetStats)
+	bool Session::UpdateEnvironment(ProcessImpl& process, const StringView& reason, bool resetStats)
 	{
 		if (!resetStats)
 			return true;
@@ -2440,7 +2441,7 @@ namespace uba
 		process.m_processStats.Write(writer);
 		process.m_storageStats.Write(writer);
 		process.m_kernelStats.Write(writer);
-		m_trace.ProcessEnvironmentUpdated(process.GetId(), reason, writer.GetData(), writer.GetPosition(), process.GetStartInfo().breadcrumbs);
+		m_trace.ProcessEnvironmentUpdated(process.GetId(), reason, writer.GetData(), writer.GetPosition(), ToView(process.GetStartInfo().breadcrumbs));
 		process.m_processStats = {};
 		process.m_storageStats = {};
 		process.m_kernelStats = {};
diff --git a/Engine/Source/Programs/UnrealBuildAccelerator/Common/Private/UbaSessionClient.cpp b/Engine/Source/Programs/UnrealBuildAccelerator/Common/Private/UbaSessionClient.cpp
--- a/Engine/Source/Programs/UnrealBuildAccelerator/Common/Private/UbaSessionClient.cpp
+++ b/Engine/Source/Programs/UnrealBuildAccelerator/Common/Private/UbaSessionClient.cpp
@@ -2063,7 +2063,7 @@ namespace uba
 		return success;
 	}
 
-	bool SessionClient::UpdateEnvironment(ProcessImpl& process, const tchar* reason, bool resetStats)
+	bool SessionClient::UpdateEnvironment(ProcessImpl& process, const StringView& reason, bool resetStats)
 	{
 		StackBinaryReader<SendMaxSize> reader;
 
diff --git a/Engine/Source/Programs/UnrealBuildAccelerator/Common/Private/UbaSessionServer.cpp b/Engine/Source/Programs/UnrealBuildAccelerator/Common/Private/UbaSessionServer.cpp
--- a/Engine/Source/Programs/UnrealBuildAccelerator/Common/Private/UbaSessionServer.cpp
+++ b/Engine/Source/Programs/UnrealBuildAccelerator/Common/Private/UbaSessionServer.cpp
@@ -464,10 +464,10 @@ namespace uba
 		m_maxRemoteProcessCount.exchange(count);
 	}
 
-	u32 SessionServer::BeginExternalProcess(const tchar* description)
+	u32 SessionServer::BeginExternalProcess(const tchar* description, const tchar* breadcrumbs)
 	{
 		u32 processId = CreateProcessId();
-		m_trace.ProcessAdded(0, processId, description);
+		m_trace.ProcessAdded(0, processId, ToView(description), ToView(breadcrumbs));
 		return processId;
 	}
 
@@ -476,7 +476,7 @@ namespace uba
 		StackBinaryWriter<1024> statsWriter;
 		ProcessStats processStats;
 		processStats.Write(statsWriter);
-		m_trace.ProcessExited(id, exitCode, statsWriter.GetData(), statsWriter.GetPosition(), Vector<ProcessLogLine>(), TC(""));
+		m_trace.ProcessExited(id, exitCode, statsWriter.GetData(), statsWriter.GetPosition(), Vector<ProcessLogLine>());
 	}
 
 	void SessionServer::UpdateProgress(u32 processesTotal, u32 processesDone, u32 errorCount)
@@ -1245,7 +1245,7 @@ namespace uba
 		u32 id = process.m_processId;
 		Vector<ProcessLogLine> emptyLines;
 		auto& logLines = (exitCode != 0 || m_detailedTrace) ? process.m_logLines : emptyLines;
-		m_trace.ProcessExited(id, exitCode, reader.GetPositionData(), reader.GetLeft(), logLines, process.GetStartInfo().breadcrumbs);
+		m_trace.ProcessExited(id, exitCode, reader.GetPositionData(), reader.GetLeft(), logLines);
 
 		ProcessStats processStats;
 		processStats.Read(reader, ~0u);
@@ -1423,7 +1423,7 @@ namespace uba
 
 		StringBuffer<> reason;
 		reader.ReadString(reason);
-		m_trace.ProcessEnvironmentUpdated(processId, reason.data, reader.GetPositionData(), reader.GetLeft(), h.GetStartInfo().breadcrumbs);
+		m_trace.ProcessEnvironmentUpdated(processId, reason, reader.GetPositionData(), reader.GetLeft(), ToView(h.GetStartInfo().breadcrumbs));
 		return true;
 	}
 
@@ -1794,7 +1794,7 @@ namespace uba
 			SessionStats().Write(writer);
 			StorageStats().Write(writer);
 			KernelStats().Write(writer);
-			m_trace.ProcessExited(process->m_processId, process->m_exitCode, writer.GetData(), writer.GetPosition(), Vector<ProcessLogLine>(), process->GetStartInfo().breadcrumbs);
+			m_trace.ProcessExited(process->m_processId, process->m_exitCode, writer.GetData(), writer.GetPosition(), Vector<ProcessLogLine>());
 
 			m_logger.Warning(TC("Cancelling remote active processes has not been tested. Notify devs"));
 
diff --git a/Engine/Source/Programs/UnrealBuildAccelerator/Common/Private/UbaTrace.cpp b/Engine/Source/Programs/UnrealBuildAccelerator/Common/Private/UbaTrace.cpp
--- a/Engine/Source/Programs/UnrealBuildAccelerator/Common/Private/UbaTrace.cpp
+++ b/Engine/Source/Programs/UnrealBuildAccelerator/Common/Private/UbaTrace.cpp
@@ -248,15 +248,16 @@ namespace uba
 		writer.WriteU32(sessionId);
 	}
 
-	void Trace::ProcessAdded(u32 sessionId, u32 processId, const tchar* description)
+	void Trace::ProcessAdded(u32 sessionId, u32 processId, const StringView& description, const StringView& breadcrumbs)
 	{
 		BEGIN_TRACE_ENTRY(ProcessAdded);
 		writer.WriteU32(sessionId);
 		writer.WriteU32(processId);
 		writer.WriteString(description);
+		writer.WriteString(breadcrumbs);
 	}
 
-	void Trace::ProcessEnvironmentUpdated(u32 processId, const tchar* reason, const u8* data, u64 dataSize, const tchar* breadcrumbs)
+	void Trace::ProcessEnvironmentUpdated(u32 processId, const StringView& reason, const u8* data, u64 dataSize, const StringView& breadcrumbs)
 	{
 		BEGIN_TRACE_ENTRY(ProcessEnvironmentUpdated);
 		writer.WriteU32(processId);
@@ -265,13 +266,13 @@ namespace uba
 		writer.WriteString(breadcrumbs);
 	}
 
-	void Trace::ProcessExited(u32 processId, u32 exitCode, const u8* data, u64 dataSize, const Vector<ProcessLogLine>& logLines, const tchar* breadcrumbs)
+	void Trace::ProcessExited(u32 processId, u32 exitCode, const u8* data, u64 dataSize, const Vector<ProcessLogLine>& logLines)
 	{
 		BEGIN_TRACE_ENTRY(ProcessExited);
 		writer.WriteU32(processId);
 		writer.WriteU32(exitCode);
 		writer.WriteBytes(data, dataSize);
-		writer.WriteString(breadcrumbs);
+
 		u32 lineCounter = 0;
 		for (auto& line : logLines)
 		{
diff --git a/Engine/Source/Programs/UnrealBuildAccelerator/Common/Public/UbaSession.h b/Engine/Source/Programs/UnrealBuildAccelerator/Common/Public/UbaSession.h
--- a/Engine/Source/Programs/UnrealBuildAccelerator/Common/Public/UbaSession.h
+++ b/Engine/Source/Programs/UnrealBuildAccelerator/Common/Public/UbaSession.h
@@ -141,7 +141,7 @@ namespace uba
 		virtual bool HostRun(BinaryReader& reader, BinaryWriter& writer);
 		virtual void FileEntryAdded(StringKey fileNameKey, u64 lastWritten, u64 size);
 		virtual bool FlushWrittenFiles(ProcessImpl& process);
-		virtual bool UpdateEnvironment(ProcessImpl& process, const tchar* reason, bool resetStats);
+		virtual bool UpdateEnvironment(ProcessImpl& process, const StringView& reason, bool resetStats);
 		virtual bool LogLine(ProcessImpl& process, const tchar* line, LogEntryType logType);
 
 		static constexpr CasKey CasKeyIsDirectory = { ~u64(0), ~u64(0), ~u32(0) };
@@ -477,6 +477,7 @@ namespace uba
 		TString workingDir;
 		TString description;
 		TString logFile;
+		TString breadcrumbs;
 	};
 }
 
diff --git a/Engine/Source/Programs/UnrealBuildAccelerator/Common/Public/UbaSessionClient.h b/Engine/Source/Programs/UnrealBuildAccelerator/Common/Public/UbaSessionClient.h
--- a/Engine/Source/Programs/UnrealBuildAccelerator/Common/Public/UbaSessionClient.h
+++ b/Engine/Source/Programs/UnrealBuildAccelerator/Common/Public/UbaSessionClient.h
@@ -68,7 +68,7 @@ namespace uba
 		virtual bool SHGetKnownFolderPath(Process& process, BinaryReader& reader, BinaryWriter& writer) override;
 		virtual bool HostRun(BinaryReader& reader, BinaryWriter& writer) override;
 		virtual bool FlushWrittenFiles(ProcessImpl& process) override;
-		virtual bool UpdateEnvironment(ProcessImpl& process, const tchar* reason, bool resetStats) override;
+		virtual bool UpdateEnvironment(ProcessImpl& process, const StringView& reason, bool resetStats) override;
 		virtual bool LogLine(ProcessImpl& process, const tchar* line, LogEntryType logType) override;
 		virtual void TraceSessionUpdate() override;
 
diff --git a/Engine/Source/Programs/UnrealBuildAccelerator/Common/Public/UbaSessionServer.h b/Engine/Source/Programs/UnrealBuildAccelerator/Common/Public/UbaSessionServer.h
--- a/Engine/Source/Programs/UnrealBuildAccelerator/Common/Public/UbaSessionServer.h
+++ b/Engine/Source/Programs/UnrealBuildAccelerator/Common/Public/UbaSessionServer.h
@@ -51,7 +51,7 @@ namespace uba
 		void SetMaxRemoteProcessCount(u32 count);
 
 		// Report an external process just to get it visible in the trace stream/visualizer. Returns an unique id that should be sent in to EndExternalProcess
-		u32 BeginExternalProcess(const tchar* description);
+		u32 BeginExternalProcess(const tchar* description, const tchar* breadcrumbs = TC(""));
 
 		// End external process.
 		void EndExternalProcess(u32 id, u32 exitCode);
diff --git a/Engine/Source/Programs/UnrealBuildAccelerator/Common/Public/UbaTrace.h b/Engine/Source/Programs/UnrealBuildAccelerator/Common/Public/UbaTrace.h
--- a/Engine/Source/Programs/UnrealBuildAccelerator/Common/Public/UbaTrace.h
+++ b/Engine/Source/Programs/UnrealBuildAccelerator/Common/Public/UbaTrace.h
@@ -85,9 +85,9 @@ namespace uba
 		void SessionNotification(u32 sessionId, const tchar* text);
 		void SessionSummary(u32 sessionId, const u8* data, u64 dataSize);
 		void SessionDisconnect(u32 sessionId);
-		void ProcessAdded(u32 sessionId, u32 processId, const tchar* description);
-		void ProcessEnvironmentUpdated(u32 processId, const tchar* reason, const u8* data, u64 dataSize, const tchar* breadcrumbs);
-		void ProcessExited(u32 processId, u32 exitCode, const u8* data, u64 dataSize, const Vector<ProcessLogLine>& logLines, const tchar* breadcrumbs);
+		void ProcessAdded(u32 sessionId, u32 processId, const StringView& description, const StringView& breadcrumbs);
+		void ProcessEnvironmentUpdated(u32 processId, const StringView& reason, const u8* data, u64 dataSize, const StringView& breadcrumbs);
+		void ProcessExited(u32 processId, u32 exitCode, const u8* data, u64 dataSize, const Vector<ProcessLogLine>& logLines);
 		void ProcessReturned(u32 processId, const StringView& reason);
 		void ProxyCreated(u32 clientId, const tchar* proxyName);
 		void ProxyUsed(u32 clientId, const tchar* proxyName);
diff --git a/Engine/Source/Programs/UnrealBuildAccelerator/Visualizer/Private/UbaTraceReader.cpp b/Engine/Source/Programs/UnrealBuildAccelerator/Visualizer/Private/UbaTraceReader.cpp
--- a/Engine/Source/Programs/UnrealBuildAccelerator/Visualizer/Private/UbaTraceReader.cpp
+++ b/Engine/Source/Programs/UnrealBuildAccelerator/Visualizer/Private/UbaTraceReader.cpp
@@ -541,6 +541,10 @@ namespace uba
 			StringBuffer<> desc;
 			reader.ReadString(desc);
 
+			TString breadcrumbs;
+			if (out.version >= 35)
+				breadcrumbs = reader.ReadString();
+
 			if (out.version < 15)
 			{
 				time = reader.Read7BitEncoded();
@@ -551,7 +555,7 @@ namespace uba
 				}
 			}
 
-			ProcessBegin(out, sessionIndex, id, time, desc.data);
+			ProcessBegin(out, sessionIndex, id, time, desc, breadcrumbs);
 			break;
 		}
 		case TraceType_ProcessExited:
@@ -601,7 +605,7 @@ namespace uba
 			process.stats.resize(dataEnd - dataStart);
 			memcpy(process.stats.data(), dataStart, dataEnd - dataStart);
 
-			if (out.version >= 34)
+			if (out.version >= 34 && out.version < 35)
 				process.breadcrumbs = reader.ReadString();
 
 			process.createFilesTime = processStats.createFile.time;
@@ -664,8 +668,9 @@ namespace uba
 			process.stats.resize(dataEnd - dataStart);
 			memcpy(process.stats.data(), dataStart, dataEnd - dataStart);
 
+			TString breadcrumbs;
 			if (out.version >= 35)
-				process.breadcrumbs = reader.ReadString();
+				breadcrumbs = reader.ReadString();
 
 			process.exitCode = 0u;
 			process.stop = time;
@@ -677,6 +682,7 @@ namespace uba
 			m_activeProcesses[processId].processIndex = u32(processes.size() - 1);
 			process2.id = processId;
 			process2.description = reason.data;
+			process2.breadcrumbs = breadcrumbs;
 			process2.start = time;
 			process2.stop = ~u64(0);
 			process2.exitCode = ~0u;
@@ -1021,7 +1027,7 @@ namespace uba
 			u32 id = u32(reader.Read7BitEncoded());
 			StringBuffer<> desc;
 			reader.ReadString(desc);
-			ProcessBegin(out, 0, id, time, desc.data)->cacheFetch = true;
+			ProcessBegin(out, 0, id, time, desc, StringView())->cacheFetch = true;
 			break;
 		}
 		case TraceType_CacheEndFetch:
@@ -1145,7 +1151,7 @@ namespace uba
 		return nullptr;
 	}
 
-	TraceView::Process* TraceReader::ProcessBegin(TraceView& out, u32 sessionIndex, u32 id, u64 time, const tchar* description)
+	TraceView::Process* TraceReader::ProcessBegin(TraceView& out, u32 sessionIndex, u32 id, u64 time, const StringView& description, const StringView& breadcrumbs)
 	{
 		TraceView::Processor* processor = nullptr;
 		TraceView::Session& session = GetSession(out, sessionIndex);
@@ -1176,7 +1182,8 @@ namespace uba
 		++out.totalProcessActiveCount;
 
 		process.id = id;
-		process.description = description;
+		process.description = description.data;
+		process.breadcrumbs = breadcrumbs.data;
 		process.start = time;
 		process.stop = ~u64(0);
 		process.exitCode = ~0u;
diff --git a/Engine/Source/Programs/UnrealBuildAccelerator/Visualizer/Private/UbaTraceReader.h b/Engine/Source/Programs/UnrealBuildAccelerator/Visualizer/Private/UbaTraceReader.h
--- a/Engine/Source/Programs/UnrealBuildAccelerator/Visualizer/Private/UbaTraceReader.h
+++ b/Engine/Source/Programs/UnrealBuildAccelerator/Visualizer/Private/UbaTraceReader.h
@@ -183,7 +183,7 @@ namespace uba
 		TraceView::Session& GetSession(TraceView& out, u32 sessionIndex);
 		TraceView::Session* GetSession(TraceView& out, const Guid& clientUid);
 
-		TraceView::Process* ProcessBegin(TraceView& out, u32 sessionIndex, u32 id, u64 time, const tchar* description);
+		TraceView::Process* ProcessBegin(TraceView& out, u32 sessionIndex, u32 id, u64 time, const StringView& description, const StringView& breadcrumbs);
 		TraceView::Process* ProcessEnd(TraceView& out, u32& outSessionIndex, u32 id, u64 time);
 
 		UnorderedMap<u32, TraceView::ProcessLocation> m_activeProcesses;
