[UBA]
* Improved some stats logging for UbaCacheServer

diff --git a/Engine/Source/Programs/UnrealBuildAccelerator/Common/Private/UbaCacheServer.cpp b/Engine/Source/Programs/UnrealBuildAccelerator/Common/Private/UbaCacheServer.cpp
--- a/Engine/Source/Programs/UnrealBuildAccelerator/Common/Private/UbaCacheServer.cpp
+++ b/Engine/Source/Programs/UnrealBuildAccelerator/Common/Private/UbaCacheServer.cpp
@@ -158,7 +158,6 @@ namespace uba
 		Atomic<u64> lastUsedTime;
 		u64 oldestUsedTime = 0; 
 
-		u64 expirationTimeSeconds = 0;
 		u32 index = ~0u;
 
 		struct MaintenanceContext
@@ -309,8 +308,6 @@ namespace uba
 		if (databaseVersion != CacheFileVersion)
 			bucket.needsSave = true;
 
-		bucket.expirationTimeSeconds = m_expirationTimeSeconds;
-
 		u32 pathTableSize = reader.ReadU32();
 		if (pathTableSize)
 		{
@@ -552,28 +549,7 @@ namespace uba
 		KernelStats& kernelStats = KernelStats::GetGlobal();
 		auto kernelReset = MakeGuard([&]() { kernelStats = {}; });
 
-		{
-			u32 peakConnectionCount = m_peakConnectionCount;
-			m_peakConnectionCount = 0;
-
-			u64 startTime = GetTime();
-			auto& storageStats = m_storage.Stats();
-			u64 hits = m_cacheKeyHitCount;
-			u64 miss = m_cacheKeyFetchCount - hits;
-			m_logger.Info(TC("Stats since boot (%s ago)"), TimeToText(startTime - m_bootTime, true).str);
-			m_logger.Info(TC("  CacheServer %s hits, %s misses"), CountToText(hits).str, CountToText(miss).str);
-			u64 recvCount = storageStats.sendCas.count.load();
-			u64 sendCount = storageStats.recvCas.count.load();
-			m_logger.Info(TC("  StorageServer cas %s (%s) sent, %s (%s) received"), CountToText(recvCount).str, BytesToText(storageStats.sendCasBytesComp).str, CountToText(sendCount).str, BytesToText(storageStats.recvCasBytesComp).str);
-			
-			if (m_lastMaintenance)
-				m_logger.Info(TC("Stats since last maintenance (%s ago)"), TimeToText(startTime - m_lastMaintenance, true).str);
-			m_logger.Info(TC("  Peak connection count: %4u"), peakConnectionCount);
-			auto& sentTimer = m_server.GetTotalSentTimer();
-			m_logger.Info(TC("  SendSocket         %8u %9s"), sentTimer.count.load(), TimeToText(sentTimer.time).str);
-			m_logger.Info(TC("     Bytes            %17s"), BytesToText(m_server.GetTotalSentBytes()).str);
-			kernelStats.Print(m_logger, false);
-		}
+		bool isFirst = true;
 
 		do
 		{
@@ -589,6 +565,34 @@ namespace uba
 					m_isRunningMaintenance = false;
 				});
 
+			if (isFirst)
+			{
+				isFirst = false;
+
+				u32 peakConnectionCount = m_peakConnectionCount;
+				m_peakConnectionCount = 0;
+
+				u64 startTime = GetTime();
+				auto& storageStats = m_storage.Stats();
+				u64 hits = m_cacheKeyHitCount;
+				u64 miss = m_cacheKeyFetchCount - hits;
+				m_logger.Info(TC("Stats since boot (%s ago)"), TimeToText(startTime - m_bootTime, true).str);
+				m_logger.Info(TC("  CacheServer %s hits, %s misses"), CountToText(hits).str, CountToText(miss).str);
+				u64 recvCount = storageStats.sendCas.count.load();
+				u64 sendCount = storageStats.recvCas.count.load();
+				m_logger.Info(TC("  StorageServer cas %s (%s) sent, %s (%s) received"), CountToText(recvCount).str, BytesToText(storageStats.sendCasBytesComp).str, CountToText(sendCount).str, BytesToText(storageStats.recvCasBytesComp).str);
+			
+				if (m_lastMaintenance)
+					m_logger.Info(TC("Stats since last maintenance (%s ago)"), TimeToText(startTime - m_lastMaintenance, true).str);
+				m_logger.Info(TC("  Peak connection count: %4u"), peakConnectionCount);
+				auto& sentTimer = m_server.GetTotalSentTimer();
+				m_logger.Info(TC("  SendSocket         %8u %9s"), sentTimer.count.load(), TimeToText(sentTimer.time).str);
+				m_logger.Info(TC("     Bytes            %17s"), BytesToText(m_server.GetTotalSentBytes()).str);
+				kernelStats.Print(m_logger, false);
+				m_server.ResetTotalStats();
+			}
+
+
 			if(!RunMaintenanceInternal(shouldExit))
 				return false;
 		}
@@ -761,7 +765,7 @@ namespace uba
 				ReaderWriterLock keysToEraseLock;
 				Vector<CasKey> keysToErase;
 
-				u64 bucketExpirationTimeSeconds = bucket.expirationTimeSeconds;
+				u64 bucketExpirationTimeSeconds = m_expirationTimeSeconds;
 
 				u64 lastUseTimeLimit = 0; // This is the time relative to server startup time
 
@@ -771,7 +775,7 @@ namespace uba
 					u64 bucketCasTableSize = bucket.m_casKeyTable.GetSize();
 					if (bucketCasTableSize >= m_bucketCasTableMaxSize)
 					{
-						u64 longestUnusedSeconds = GetFileTimeAsSeconds((now - m_creationTime) - bucket.oldestUsedTime);
+						u64 longestUnusedSeconds = GetFileTimeAsSeconds(now - m_creationTime - bucket.oldestUsedTime);
 
 						bucketExpirationTimeSeconds = Min(longestUnusedSeconds, bucketExpirationTimeSeconds);
 						if (bucketExpirationTimeSeconds > 60*60)
@@ -1145,7 +1149,10 @@ namespace uba
 
 			bucket.needsSave = true;
 
-			m_logger.Info(TC("    Bucket %u Done (%s). CacheEntries: %s (%s) PathTable: %s CasTable: %s Expiration: %s"), bucket.index, TimeToText(GetTime() - bucketStartTime).str, CountToText(bucket.totalEntryCount.load()).str, BytesToText(bucket.totalEntrySize.load()).str, BytesToText(bucket.m_pathTable.GetSize()).str, BytesToText(bucket.m_casKeyTable.GetSize()).str, TimeToText(MsToTime(bucket.expirationTimeSeconds*1000), true).str);
+			u64 longestUnusedSeconds = 0;
+			if (bucket.oldestUsedTime)
+				longestUnusedSeconds = GetFileTimeAsSeconds(now - m_creationTime - bucket.oldestUsedTime);
+			m_logger.Info(TC("    Bucket %u Done (%s). CacheEntries: %s (%s) PathTable: %s CasTable: %s LongestUnused: %s"), bucket.index, TimeToText(GetTime() - bucketStartTime).str, CountToText(bucket.totalEntryCount.load()).str, BytesToText(bucket.totalEntrySize.load()).str, BytesToText(bucket.m_pathTable.GetSize()).str, BytesToText(bucket.m_casKeyTable.GetSize()).str, TimeToText(MsToTime(longestUnusedSeconds*1000), true).str);
 
 			SCOPED_WRITE_LOCK(globalStatsLock, l);
 			maxCommittedMemory = Max(maxCommittedMemory, memoryBlock.writtenSize);
@@ -1222,7 +1229,6 @@ namespace uba
 		if (!insres.second)
 			return bucket;
 		bucket.index = u32(m_buckets.size() - 1);
-		bucket.expirationTimeSeconds = m_expirationTimeSeconds;
 		return bucket;
 	}
 
diff --git a/Engine/Source/Programs/UnrealBuildAccelerator/Common/Private/UbaNetworkServer.cpp b/Engine/Source/Programs/UnrealBuildAccelerator/Common/Private/UbaNetworkServer.cpp
--- a/Engine/Source/Programs/UnrealBuildAccelerator/Common/Private/UbaNetworkServer.cpp
+++ b/Engine/Source/Programs/UnrealBuildAccelerator/Common/Private/UbaNetworkServer.cpp
@@ -917,6 +917,13 @@ namespace uba
 		out.connectionCount += c.connectionCount;
 	}
 
+	void NetworkServer::ResetTotalStats()
+	{
+		m_sendTimer = {};
+		m_sendBytes = 0;
+		m_recvBytes = 0;
+	}
+
 	bool NetworkServer::DoAdditionalWork()
 	{
 		AdditionalWork work;
diff --git a/Engine/Source/Programs/UnrealBuildAccelerator/Common/Public/UbaNetworkServer.h b/Engine/Source/Programs/UnrealBuildAccelerator/Common/Public/UbaNetworkServer.h
--- a/Engine/Source/Programs/UnrealBuildAccelerator/Common/Public/UbaNetworkServer.h
+++ b/Engine/Source/Programs/UnrealBuildAccelerator/Common/Public/UbaNetworkServer.h
@@ -95,6 +95,7 @@ namespace uba
 		Timer& GetTotalSentTimer();
 		u32 GetConnectionCount();
 		void GetClientStats(ClientStats& out, u32 clientId);
+		void ResetTotalStats();
 
 		bool DoAdditionalWork();
 
