[UBA]
* Commented out kernel stats from cache server reports.. it is too noisy and hard to read through datadog

diff --git a/Engine/Source/Programs/UnrealBuildAccelerator/Common/Private/UbaCacheServer.cpp b/Engine/Source/Programs/UnrealBuildAccelerator/Common/Private/UbaCacheServer.cpp
--- a/Engine/Source/Programs/UnrealBuildAccelerator/Common/Private/UbaCacheServer.cpp
+++ b/Engine/Source/Programs/UnrealBuildAccelerator/Common/Private/UbaCacheServer.cpp
@@ -546,9 +546,6 @@ namespace uba
 		if (m_addsSinceMaintenance == 0 && !force && !m_bucketIsOverflowing)
 			return true;
 
-		KernelStats& kernelStats = KernelStats::GetGlobal();
-		auto kernelReset = MakeGuard([&]() { kernelStats = {}; });
-
 		bool isFirst = true;
 
 		do
@@ -588,7 +585,9 @@ namespace uba
 				auto& sentTimer = m_server.GetTotalSentTimer();
 				m_logger.Info(TC("  SendSocket         %8u %9s"), sentTimer.count.load(), TimeToText(sentTimer.time).str);
 				m_logger.Info(TC("     Bytes            %17s"), BytesToText(m_server.GetTotalSentBytes()).str);
-				kernelStats.Print(m_logger, false);
+				//KernelStats& kernelStats = KernelStats::GetGlobal();
+				//kernelStats.Print(m_logger, false);
+				//kernelStats = {};
 				m_server.ResetTotalStats();
 			}
 
