[UBA]
* Added debug logging in detouring code for when getting full filenames (mostly for dlls).. to be able to figure out how one of the licensees setup end up loading wrong c1xx.dll

diff --git a/Engine/Source/Programs/UnrealBuildAccelerator/Detours/Private/UbaDetoursFileMappingTable.cpp b/Engine/Source/Programs/UnrealBuildAccelerator/Detours/Private/UbaDetoursFileMappingTable.cpp
--- a/Engine/Source/Programs/UnrealBuildAccelerator/Detours/Private/UbaDetoursFileMappingTable.cpp
+++ b/Engine/Source/Programs/UnrealBuildAccelerator/Detours/Private/UbaDetoursFileMappingTable.cpp
@@ -307,11 +307,21 @@ namespace uba
 			reader.ReadString(tempBuf);
 			if (useVirtualName)
 			{
-				tempBuf.Clear();
-				reader.ReadString(tempBuf);
+				reader.ReadString(tempBuf.Clear());
+				#if UBA_DEBUG
+				DEBUG_LOG_DETOURED(L"Rpc_GetFullFileName", L"%s -> %s", path, tempBuf.data);
+				#endif
 			}
 			else
+			{
+				#if UBA_DEBUG
+				StringBuffer<> virtualName;
+				reader.ReadString(virtualName);
+				DEBUG_LOG_DETOURED(L"Rpc_GetFullFileName", L"%s -> %s (%s)", path, tempBuf.data, virtualName.data);
+				#else
 				reader.SkipString();
+				#endif
+			}
 			mappedFileTableSize = reader.ReadU32();
 			DEBUG_LOG_PIPE(L"GetFileName", L"(%ls)", tempBuf);
 		}
