[UBA]
* Fixed bug in cache server where it didn't rerun maintenance if cas key table is still overflowed

diff --git a/Engine/Source/Programs/UnrealBuildAccelerator/Common/Private/UbaCacheServer.cpp b/Engine/Source/Programs/UnrealBuildAccelerator/Common/Private/UbaCacheServer.cpp
--- a/Engine/Source/Programs/UnrealBuildAccelerator/Common/Private/UbaCacheServer.cpp
+++ b/Engine/Source/Programs/UnrealBuildAccelerator/Common/Private/UbaCacheServer.cpp
@@ -1088,18 +1088,15 @@ namespace uba
 			}
 			m_logger.Detail(TC("    Bucket %u Recreated path table. %s -> %s (%s)"), bucket.index, BytesToText(oldSize).str, BytesToText(pathTable.GetSize()).str, TimeToText(GetTime() - recreatePathTableStart).str);
 
-			if (pathTable.GetSize() >= m_bucketCasTableMaxSize) // Still overflowing?
-				m_bucketIsOverflowing = true;
-
 			// Build new caskey table based on used offsets
 			u64 recreateCasKeyTableStart = GetTime();
 			HashMap2<u32, u32> oldToNewCasKeyOffset;
-			auto& oldCasKeyTable = bucket.m_casKeyTable;
-			oldSize = oldCasKeyTable.GetSize();
+			auto& casKeyTable = bucket.m_casKeyTable;
+			oldSize = casKeyTable.GetSize();
 			{
 				oldToNewCasKeyOffset.Init(memoryBlock, usedCasKeyOffsetsCount);
 				CompactCasKeyTable newCasKeyTable(usedCasKeyOffsetsCount + Min(usedCasKeyOffsetsCount/2, 10'000ull)); // Add extra reservation space for growth
-				BinaryReader reader2(oldCasKeyTable.GetMemory(), 0, oldSize);
+				BinaryReader reader2(casKeyTable.GetMemory(), 0, oldSize);
 				usedCasKeyOffsets.Traverse([&](u32 casKeyOffset)
 				{
 					reader2.SetPosition(casKeyOffset);
@@ -1112,10 +1109,12 @@ namespace uba
 					if (casKeyOffset != newCasKeyOffset)
 						oldToNewCasKeyOffset.Insert(casKeyOffset) = newCasKeyOffset;
 				});
-				oldCasKeyTable.Swap(newCasKeyTable);
+				casKeyTable.Swap(newCasKeyTable);
 			}
 			m_logger.Detail(TC("    Bucket %u Recreated caskey table. %s -> %s (%s)"), bucket.index, BytesToText(oldSize).str, BytesToText(bucket.m_casKeyTable.GetSize()).str, TimeToText(GetTime() - recreateCasKeyTableStart).str);
 
+			if (casKeyTable.GetSize() >= m_bucketCasTableMaxSize) // Still overflowing?
+				m_bucketIsOverflowing = true;
 
 			if (oldToNewCasKeyOffset.Size() > 0)
 			{
