[UBA]
* Removed detouring of _wsplitpath_s when compiling uba using clang.. because it throws errors

diff --git a/Engine/Source/Programs/UnrealBuildAccelerator/Detours/Private/Windows/UbaDetoursFunctionsUcrtBase.inl b/Engine/Source/Programs/UnrealBuildAccelerator/Detours/Private/Windows/UbaDetoursFunctionsUcrtBase.inl
--- a/Engine/Source/Programs/UnrealBuildAccelerator/Detours/Private/Windows/UbaDetoursFunctionsUcrtBase.inl
+++ b/Engine/Source/Programs/UnrealBuildAccelerator/Detours/Private/Windows/UbaDetoursFunctionsUcrtBase.inl
@@ -339,7 +339,7 @@ int Detoured__wspawnl(int mode, const wchar_t* cmdname, const wchar_t* arg0, con
 	return exitCode;
 }
 
-#if !defined( __clang_analyzer__ )
+#if !defined(__clang__)
 errno_t Detoured__wsplitpath_s(const wchar_t* path, wchar_t* drive, size_t driveNumberOfElements, wchar_t* dir, size_t dirNumberOfElements, wchar_t* fname, size_t nameNumberOfElements, wchar_t* ext, size_t extNumberOfElements)
 {
 	DETOURED_CALL(_wsplitpath_s);
diff --git a/Engine/Source/Programs/UnrealBuildAccelerator/Detours/Private/Windows/UbaDetoursFunctionsWin.h b/Engine/Source/Programs/UnrealBuildAccelerator/Detours/Private/Windows/UbaDetoursFunctionsWin.h
--- a/Engine/Source/Programs/UnrealBuildAccelerator/Detours/Private/Windows/UbaDetoursFunctionsWin.h
+++ b/Engine/Source/Programs/UnrealBuildAccelerator/Detours/Private/Windows/UbaDetoursFunctionsWin.h
@@ -144,7 +144,7 @@
 #define DETOURED_FUNCTIONS_SHLWAPI \
 	DETOURED_FUNCTIONS_SHLWAPI_DEBUG
 
-#if !defined( __clang_analyzer__ )
+#if !defined(__clang__)
 #define DETOURED_WSPLITPATH DETOURED_FUNCTION(_wsplitpath_s)
 #else
 #define DETOURED_WSPLITPATH
