[UBA]
* Fixed uba visualizer hittest offsets when worker bars are visible
* Fixed so UbaVisualizer compiles with clang

diff --git a/Engine/Source/Programs/UnrealBuildAccelerator/ObjTool/Private/UbaImportLibWriter.cpp b/Engine/Source/Programs/UnrealBuildAccelerator/ObjTool/Private/UbaImportLibWriter.cpp
--- a/Engine/Source/Programs/UnrealBuildAccelerator/ObjTool/Private/UbaImportLibWriter.cpp
+++ b/Engine/Source/Programs/UnrealBuildAccelerator/ObjTool/Private/UbaImportLibWriter.cpp
@@ -524,7 +524,7 @@ namespace uba
 					header.SizeOfData = sizeOfData;
 					header.Ordinal = u16(i);
 					header.Type = symbol.type & 3;
-					header.NameType = symbol.type >> 2;
+					header.NameType = WORD(symbol.type >> 2);
 					header.Reserved = 0;
 
 					Write(&header, sizeof(header));
diff --git a/Engine/Source/Programs/UnrealBuildAccelerator/Visualizer/Private/UbaVisualizer.cpp b/Engine/Source/Programs/UnrealBuildAccelerator/Visualizer/Private/UbaVisualizer.cpp
--- a/Engine/Source/Programs/UnrealBuildAccelerator/Visualizer/Private/UbaVisualizer.cpp
+++ b/Engine/Source/Programs/UnrealBuildAccelerator/Visualizer/Private/UbaVisualizer.cpp
@@ -404,7 +404,7 @@ namespace uba
 
 	bool Visualizer::Unselect()
 	{
-		if (m_processSelected || m_sessionSelectedIndex != ~0u || m_statsSelected || m_timelineSelected || m_fetchedFilesSelected != ~0u || m_workSelected || !m_hyperLinkSelected.empty())
+		if (m_processSelected || m_sessionSelectedIndex != ~0u || m_statsSelected || m_timelineSelected != 0 || m_fetchedFilesSelected != ~0u || m_workSelected || !m_hyperLinkSelected.empty())
 		{
 			m_processSelected = false;
 			m_sessionSelectedIndex = ~0u;
@@ -735,6 +735,14 @@ namespace uba
 						process.bitmap = 0;
 				}
 
+		for (auto& workTrack : m_traceView.workTracks)
+			for (auto& work : workTrack.records)
+			{
+					//work.bitmapDirty = true;
+					if (full)
+						work.bitmap = 0;
+			}
+
 		if (!full)
 			return;
 
@@ -793,7 +801,7 @@ namespace uba
 
 	void Visualizer::UpdateProcessFont()
 	{
-		m_zoomValue = 1.0f + m_boxHeight / 30.0f;
+		m_zoomValue = 1.0f + float(m_boxHeight) / 30.0f;
 		int fontHeight = Max(m_boxHeight - 2, 1);
 		UpdateFont(m_processFont, fontHeight, false);
 		m_progressRectLeft = int(5 + float(m_processFont.height) * 1.8f);
@@ -937,7 +945,7 @@ namespace uba
 		{
 			drawIndentedText(ToView(L"Progress"), LogEntryType_Info, 1, false);
 
-			float progress = float(m_traceView.progressProcessesDone) / m_traceView.progressProcessesTotal;
+			float progress = float(m_traceView.progressProcessesDone) / float(m_traceView.progressProcessesTotal);
 			u32 width = m_activeFont.height * 18;
 			RECT rect;
 			rect.left = 3 + 6*m_activeFont.height;
@@ -946,7 +954,7 @@ namespace uba
 			rect.bottom = posY + m_activeFont.height;
 			FillRect(hdc, &rect, m_processBrushes[0].inProgress);
 
-			rect.right = rect.left + int(progress*width);//durationMs/100;
+			rect.right = rect.left + int(progress*float(width));//durationMs/100;
 			FillRect(hdc, &rect, m_traceView.progressErrorCount ? m_processBrushes[0].error : m_processBrushes[0].success);
 
 			const tchar* remoteDisabled = TC(""); // TODO: Only show this in "advanced mode" m_traceView.remoteExecutionDisabled ? TC("   (Remote spawn disabled)") : TC("");
@@ -988,8 +996,8 @@ namespace uba
 					bool done = stop != ~u64(0);
 					if (!done)
 						stop = playTime;
-					int left = int(posX + TimeToS(start) * scaleX);
-					int right = int(posX + TimeToS(stop) * scaleX) - 1;
+					int left = posX + int(float(TimeToS(start)) * scaleX);
+					int right = posX + int(float(TimeToS(stop)) * scaleX) - 1;
 
 					RECT rect;
 					rect.left = left;
@@ -1167,11 +1175,11 @@ namespace uba
 						auto updateSend = update.send;
 						auto updateRecv = update.recv;
 
-						int x = int(posX + TimeToS(update.time) * scaleX);
+						int x = posX + int(float(TimeToS(update.time)) * scaleX);
 						int sendY = graphBaseY;
 						int recvY = graphBaseY;
 						int cpuY = graphBaseY - int(cpuLoad*(GraphHeight-2));
-						int memY = graphBaseY - int(double(session.memTotal - update.memAvail)*(GraphHeight - 2)/session.memTotal);
+						int memY = graphBaseY - int(double(session.memTotal - update.memAvail)*double(GraphHeight - 2)/double(session.memTotal));
 
 						double duration = TimeToS(update.time - prevTime);
 						if (update.time == 0)
@@ -1179,9 +1187,11 @@ namespace uba
 						else if (prevSend > updateSend || prevRecv > updateRecv)
 							isFirstUpdate = true;
 
-						if (double sendInvScaleY = duration * sendScale)
+						double sendInvScaleY = duration * sendScale;
+						if (sendInvScaleY != 0)
 							sendY = graphBaseY - int(double(updateSend - prevSend) / sendInvScaleY);
-						if (double recvInvScaleY = duration * recvScale)
+						double recvInvScaleY = duration * recvScale;
+						if (recvInvScaleY != 0)
 							recvY = graphBaseY - int(double(updateRecv - prevRecv) / recvInvScaleY) - 1;
 
 						if (!isFirstUpdate && x > clientRect.left && prevX <= clientRect.right)
@@ -1264,7 +1274,7 @@ namespace uba
 						int posX = int(m_scrollPosX) + progressRect.left;
 						for (auto& process : processor.processes)
 						{
-							int left = int(posX + TimeToS(process.start) * scaleX);
+							int left = posX + int(float(TimeToS(process.start)) * scaleX);
 
 							auto pig = MakeGuard([&]() { ++processLocation.processIndex; });
 
@@ -1282,7 +1292,7 @@ namespace uba
 
 							RECT rect;
 							rect.left = left;
-							rect.right = int(posX + TimeToS(stop) * scaleX) - 1;
+							rect.right = posX + int(float(TimeToS(stop)) * scaleX) - 1;
 							rect.top = posY;
 							rect.bottom = rectBottom - 1;
 
@@ -1460,7 +1470,7 @@ namespace uba
 
 							float startTime = TimeToS(work.start);
 
-							int left = int(posX + startTime * scaleX);
+							int left = posX + int(startTime * scaleX);
 
 							if (left >= progressRect.right)
 							{
@@ -1492,9 +1502,9 @@ namespace uba
 
 							RECT rect;
 							rect.left = left;
-							rect.right = int(posX + stopTime * scaleX) - 1;
+							rect.right = posX + int(stopTime * scaleX) - 1;
 							rect.top = posY;
-							rect.bottom = rectBottom;
+							rect.bottom = rectBottom - 1;
 
 							if (rect.right <= progressRect.left)
 							{
@@ -1582,7 +1592,7 @@ namespace uba
 						}
 					}
 					++trackIndex;
-					posY += stepY;
+					posY += processStepY;
 				}
 			}
 
@@ -1610,13 +1620,13 @@ namespace uba
 			POINT pos;
 			GetCursorPos(&pos);
 			ScreenToClient(m_hwnd, &pos);
-			timelineSelected = startOffset + (pos.x - m_progressRectLeft) / timeScale;
+			timelineSelected = startOffset + float(pos.x - m_progressRectLeft) / timeScale;
 		}
 
-		if (timelineSelected)
+		if (timelineSelected != 0)
 		{
 			int posX = int(m_scrollPosX) + progressRect.left;
-			int left = int(posX + timelineSelected * scaleX);
+			int left = posX + int(timelineSelected * scaleX);
 			int timelineTop = GetTimelineTop(clientRect);
 
 			// TODO: Draw line up
@@ -2206,7 +2216,7 @@ namespace uba
 		SetActiveFont(m_timelineFont);
 		int top = GetTimelineTop(clientRect);
 		float timeScale = (m_horizontalScaleValue*m_zoomValue)*50.0f;
-		float startOffset = ((m_scrollPosX/timeScale) - int(m_scrollPosX/timeScale)) * timeScale;
+		float startOffset = ((m_scrollPosX/timeScale) - float(int(m_scrollPosX/timeScale))) * timeScale;
 		int index = -int(startOffset/timeScale);
 			
 		int number = -int(float(m_scrollPosX)/timeScale);
@@ -2234,7 +2244,7 @@ namespace uba
 
 		while (true)
 		{
-			int pos = progressRect.left + int(startOffset + index*timeScale);
+			int pos = progressRect.left + int(startOffset + float(index)*timeScale);
 			if (pos >= clientRect.right)
 				break;
 
@@ -2327,10 +2337,11 @@ namespace uba
 				auto& u = session.updates.back();
 				u64 sendPerS = 0;
 				u64 recvPerS = 0;
-				if (float duration = TimeToS(u.time - session.prevUpdateTime))
+				float duration = TimeToS(u.time - session.prevUpdateTime);
+				if (duration != 0)
 				{
-					sendPerS = u64((u.send - session.prevSend) / duration);
-					recvPerS = u64((u.recv - session.prevRecv) / duration);
+					sendPerS = u64(float(u.send - session.prevSend) / duration);
+					recvPerS = u64(float(u.recv - session.prevRecv) / duration);
 				}
 				drawText(L"ClientId: %u  TcpCount: %u", session.clientUid.data1, u.connectionCount);
 				drawText(L"Recv: %ls (%sps)", BytesToText(u.recv), BytesToText(recvPerS*8));
@@ -2394,10 +2405,11 @@ namespace uba
 				{
 					u64 sendPerS = 0;
 					u64 recvPerS = 0;
-					if (float duration = TimeToS(u.time - session.prevUpdateTime))
+					float duration = TimeToS(u.time - session.prevUpdateTime);
+					if (duration != 0)
 					{
-						sendPerS = u64((u.send - session.prevSend) / duration);
-						recvPerS = u64((u.recv - session.prevRecv) / duration);
+						sendPerS = u64(float(u.send - session.prevSend) / duration);
+						recvPerS = u64(float(u.recv - session.prevRecv) / duration);
 					}
 					drawText(L"Recv: %ls (%sps)", BytesToText(u.recv), BytesToText(recvPerS));
 					drawText(L"Send: %ls (%sps)", BytesToText(u.send), BytesToText(sendPerS));
@@ -2469,6 +2481,19 @@ namespace uba
 
 		outResult.section = 0;
 
+		if (m_config.showTimeline && !m_traceView.sessions.empty())
+		{
+			int timelineTop = GetTimelineTop(clientRect);
+			if (pos.y >= timelineTop)// && pos.y < timelineTop + 40)
+			{
+				outResult.section = 3;
+				float timeScale = (m_horizontalScaleValue * m_zoomValue)*50.0f;
+				float startOffset = -(m_scrollPosX / timeScale);
+				outResult.timelineSelected = startOffset + float(pos.x - m_progressRectLeft) / timeScale;
+				return;
+			}
+		}
+
 		u64 lastStop = 0;
 
 		if (m_config.showProgress && m_traceView.progressProcessesTotal)
@@ -2523,8 +2548,8 @@ namespace uba
 					bool done = stop != ~u64(0);
 					if (!done)
 						stop = playTime;
-					int left = int(posX + TimeToS(process.start) * scaleX);
-					int right = int(posX + TimeToS(stop) * scaleX) - 1;
+					int left = posX + int(float(TimeToS(process.start)) * scaleX);
+					int right = posX + int(float(TimeToS(stop)) * scaleX) - 1;
 
 					if (pos.x >= left && pos.x <= right)
 					{
@@ -2590,7 +2615,7 @@ namespace uba
 					int prevX = 100000;
 					for (auto& update : session.updates)
 					{
-						int x = int(posX + TimeToS(update.time) * scaleX);
+						int x = posX + int(float(TimeToS(update.time)) * scaleX);
 
 						if (prevSend > update.send || prevRecv > update.recv)
 						{
@@ -2604,8 +2629,8 @@ namespace uba
 						if (pos.x + hitOffset >= prevX && pos.x + hitOffset <= x)
 						{
 							double duration = TimeToS(update.time - prevTime);
-							outResult.stats.recvBytesPerSecond = u64((update.recv - prevRecv) / duration);
-							outResult.stats.sendBytesPerSecond = u64((update.send - prevSend) / duration);
+							outResult.stats.recvBytesPerSecond = u64(float(update.recv - prevRecv) / duration);
+							outResult.stats.sendBytesPerSecond = u64(float(update.send - prevSend) / duration);
 							outResult.stats.ping = update.ping;
 							outResult.stats.memAvail = update.memAvail;
 							outResult.stats.cpuLoad = update.cpuLoad;
@@ -2649,7 +2674,7 @@ namespace uba
 						int posX = int(m_scrollPosX) + progressRect.left;
 						for (auto& process : processor.processes)
 						{
-							int left = int(posX + TimeToS(process.start) * scaleX);
+							int left = posX + int(float(TimeToS(process.start)) * scaleX);
 
 							auto pig = MakeGuard([&]() { ++processIndex; });
 
@@ -2670,7 +2695,7 @@ namespace uba
 
 							RECT rect;
 							rect.left = left;
-							rect.right = int(posX + TimeToS(stopTime) * scaleX);
+							rect.right = posX + int(float(TimeToS(stopTime)) * scaleX);
 
 							if (rect.right <= progressRect.left)
 								continue;
@@ -2717,7 +2742,7 @@ namespace uba
 						int posX = int(m_scrollPosX) + progressRect.left;
 						for (auto& work : workTrack.records)
 						{
-							int left = int(posX + TimeToS(work.start) * scaleX);
+							int left = posX + int(float(TimeToS(work.start)) * scaleX);
 
 							if (left >= progressRect.right)
 							{
@@ -2734,7 +2759,7 @@ namespace uba
 
 							RECT rect;
 							rect.left = left;
-							rect.right = int(posX + TimeToS(stopTime) * scaleX);
+							rect.right = posX + int(float(TimeToS(stopTime)) * scaleX);
 
 							if (rect.right <= progressRect.left)
 							{
@@ -2763,18 +2788,6 @@ namespace uba
 
 		m_contentWidth = m_progressRectLeft + Max(0, int(TimeToS((lastStop != 0 && lastStop != ~u64(0)) ? lastStop : playTime) * scaleX));
 		m_contentHeight = posY - int(m_scrollPosY) + processStepY + 14;
-
-		if (m_config.showTimeline && !m_traceView.sessions.empty())
-		{
-			int timelineTop = GetTimelineTop(clientRect);
-			if (pos.y >= timelineTop)// && pos.y < timelineTop + 40)
-			{
-				outResult.section = 3;
-				float timeScale = (m_horizontalScaleValue * m_zoomValue)*50.0f;
-				float startOffset = -(m_scrollPosX / timeScale);
-				outResult.timelineSelected = startOffset + (pos.x - m_progressRectLeft) / timeScale;
-			}
-		}
 	}
 
 	void Visualizer::WriteProcessStats(Logger& out, const TraceView::Process& process)
@@ -2881,7 +2894,7 @@ namespace uba
 		else
 		{
 			float oldScrollPosX = m_scrollPosX;
-			m_scrollPosX = Min(0.0f, (float)rect.right - timeS*50.0f*m_horizontalScaleValue*m_zoomValue - m_progressRectLeft);
+			m_scrollPosX = Min(0.0f, (float)rect.right - timeS*50.0f*m_horizontalScaleValue*m_zoomValue - float(m_progressRectLeft));
 			return oldScrollPosX != m_scrollPosX;
 		}
 	}
@@ -3058,10 +3071,10 @@ namespace uba
 		case WM_SIZE:
 		{
 			int height = HIWORD(lParam);
-			if (m_contentHeight && m_contentHeight + m_scrollPosY < height)
+			if (m_contentHeight && m_contentHeight + int(m_scrollPosY) < height)
 				m_scrollPosY = float(Min(0, height - m_contentHeight));
 			int width = LOWORD(lParam);
-			if (m_contentWidth && m_contentWidth + m_scrollPosX < width)
+			if (m_contentWidth && m_contentWidth + int(m_scrollPosX) < width)
 				m_scrollPosX = float(Min(0, width - m_contentWidth));
 			UpdateScrollbars(true);
 			break;
@@ -3149,7 +3162,7 @@ namespace uba
 						newScaleValue = Max(m_horizontalScaleValue + m_horizontalScaleValue*float(delta)*0.0006f, 0.001f);
 
 					// TODO: m_progressRectLeft changes with zoom so anchor logic is wrong
-					const float scrollAnchorOffsetX = float(cursorPos.x) - m_progressRectLeft;
+					const float scrollAnchorOffsetX = float(cursorPos.x) - float(m_progressRectLeft);
 					const float scrollAnchorOffsetY = 0;//float(cursorPos.y)*m_zoomValue;// - m_progressRectLeft;
 
 					float oldZoomValue = m_zoomValue;
@@ -3207,7 +3220,7 @@ namespace uba
 				RECT r;
 				GetClientRect(hWnd, &r);
 				float oldScrollY = m_scrollPosY;
-				m_scrollPosY = m_scrollPosY + delta;
+				m_scrollPosY = m_scrollPosY + float(delta);
 				m_scrollPosY = Min(Max(m_scrollPosY, float(r.bottom - m_contentHeight)), 0.0f);
 				if (oldScrollY != m_scrollPosY)
 				{
@@ -3231,15 +3244,15 @@ namespace uba
 				if (m_contentHeight <= r.bottom)
 					m_scrollPosY = 0;
 				else
-					m_scrollPosY = Max(Min(m_scrollAtAnchorY + pos.y - m_mouseAnchor.y, 0.0f), float(r.bottom - m_contentHeight));
+					m_scrollPosY = Max(Min(m_scrollAtAnchorY + float(pos.y - m_mouseAnchor.y), 0.0f), float(r.bottom - m_contentHeight));
 
 				if (m_contentWidth <= r.right)
 					m_scrollPosX = 0;
 				else
 				{
 					int minScroll = r.right - m_contentWidth;
-					m_scrollPosX = Max(Min(m_scrollAtAnchorX + pos.x - m_mouseAnchor.x, 0.0f), float(minScroll));
-					if (!m_traceView.finished && m_scrollPosX <= minScroll)
+					m_scrollPosX = Max(Min(m_scrollAtAnchorX + float(pos.x - m_mouseAnchor.x), 0.0f), float(minScroll));
+					if (!m_traceView.finished && m_scrollPosX <= float(minScroll))
 						m_autoScroll = true;
 				}
 				UpdateScrollbars(true);
@@ -3323,7 +3336,7 @@ namespace uba
 				UpdateScrollbars(true);
 				Redraw(false);
 			}
-			else if (m_timelineSelected)
+			else if (m_timelineSelected != 0)
 			{
 				if (!m_client) // Does not work for network streams
 				{
@@ -3402,7 +3415,7 @@ namespace uba
 		}
 		case WM_LBUTTONUP:
 		{
-			if (!(m_buttonSelected != ~0u || m_timelineSelected))
+			if (!(m_buttonSelected != ~0u || m_timelineSelected != 0))
 			{
 				StopDragToScroll();
 			}
@@ -3651,10 +3664,10 @@ namespace uba
 					m_scrollPosY = -float(scrollInfo.nTrackPos);
 					break;
 				case SB_PAGEDOWN:
-					m_scrollPosY = m_scrollPosY - r.bottom;
+					m_scrollPosY = m_scrollPosY - float(r.bottom);
 					break;
 				case SB_PAGEUP:
-					m_scrollPosY = m_scrollPosY + r.bottom;
+					m_scrollPosY = m_scrollPosY + float(r.bottom);
 					break;
 				case SB_LINEDOWN:
 					m_scrollPosY -= 30;
@@ -3698,10 +3711,10 @@ namespace uba
 					m_scrollPosX = -float(scrollInfo.nTrackPos);
 					break;
 				case SB_PAGEDOWN:
-					m_scrollPosX = m_scrollPosX - r.right;
+					m_scrollPosX = m_scrollPosX - float(r.right);
 					break;
 				case SB_PAGEUP:
-					m_scrollPosX = m_scrollPosX + r.right;
+					m_scrollPosX = m_scrollPosX + float(r.right);
 					break;
 				case SB_LINEDOWN:
 					m_scrollPosX -= 30;
@@ -3714,7 +3727,7 @@ namespace uba
 				}
 
 				int minScroll = r.right - m_contentWidth;
-				m_autoScroll = !m_traceView.finished && (m_scrollPosX <= minScroll || autoScroll);
+				m_autoScroll = !m_traceView.finished && (m_scrollPosX <= float(minScroll) || autoScroll);
 				m_scrollPosX = Min(Max(m_scrollPosX, float(r.right - m_contentWidth)), 0.0f);
 
 				if (oldScrollX != m_scrollPosX)
