[Uba]
* Fixed so trace entry end time is fetched before writing to trace memory since the write itself can stall

diff --git a/Engine/Source/Programs/UnrealBuildAccelerator/Common/Private/UbaTrace.cpp b/Engine/Source/Programs/UnrealBuildAccelerator/Common/Private/UbaTrace.cpp
--- a/Engine/Source/Programs/UnrealBuildAccelerator/Common/Private/UbaTrace.cpp
+++ b/Engine/Source/Programs/UnrealBuildAccelerator/Common/Private/UbaTrace.cpp
@@ -200,11 +200,12 @@ namespace uba
 	#define BEGIN_TRACE_ENTRY(x) \
 		if (!m_memoryBegin) \
 			return; \
+		u64 time = GetTime() - m_startTime; \
 		WriterScope writer(*this); \
 		if (!writer.IsValid()) \
 			return; \
 		writer.WriteByte(TraceType_##x); \
-		writer.Write7BitEncoded(GetTime() - m_startTime);
+		writer.Write7BitEncoded(time);
 
 	void Trace::SessionAdded(u32 sessionId, u32 clientId, const tchar* name, const tchar* info)
 	{
