[UBA]
* Added warning if GetFileAttributes inside SearchPathForFile fails for other reasons than file not existing

diff --git a/Engine/Source/Programs/UnrealBuildAccelerator/Common/Private/UbaFile.cpp b/Engine/Source/Programs/UnrealBuildAccelerator/Common/Private/UbaFile.cpp
--- a/Engine/Source/Programs/UnrealBuildAccelerator/Common/Private/UbaFile.cpp
+++ b/Engine/Source/Programs/UnrealBuildAccelerator/Common/Private/UbaFile.cpp
@@ -378,11 +378,22 @@ namespace uba
 		fullPath.Append(applicationDir);
 		fullPath.EnsureEndsWithSlash();
 		fullPath.Append(file);
-		if (GetFileAttributesW(fullPath.data) != INVALID_FILE_ATTRIBUTES)
-		{
-			FixPath(fullPath.data, nullptr, 0, out);
+
+		auto TestFileExists = [&]()
+			{
+				if (GetFileAttributesW(fullPath.data) != INVALID_FILE_ATTRIBUTES)
+				{
+					FixPath(fullPath.data, nullptr, 0, out);
+					return true;
+				}
+				u32 lastError = GetLastError();
+				if (lastError != ERROR_FILE_NOT_FOUND && lastError != ERROR_PATH_NOT_FOUND)
+					logger.Warning(TC("SearchPathForFile found the file %s but got error when getting attributes ()"), fullPath.data, LastErrorToText(lastError).data);
+				return false;
+			};
+
+		if (TestFileExists())
 			return true;
-		}
 
 		char varSeparator = IsWindows ? ';' : ':';
 		tchar buff[32 * 1024] = { 0 };
@@ -410,11 +421,9 @@ namespace uba
 			if (*lastStart)
 				fullPath.EnsureEndsWithSlash();
 			fullPath.Append(file);
-			if (GetFileAttributesW(fullPath.data) != INVALID_FILE_ATTRIBUTES)
-			{
-				FixPath(fullPath.data, nullptr, 0, out);
+
+			if (TestFileExists())
 				return true;
-			}
 
 			lastStart = ++it;
 		}
