[UBA]
* Added notification message from client to server when idle timeout
* Added "Done" notification to trace when client is disconnected because it is done
* Fixed so "known inputs" honor allowstorageproxy rule
* Fixed so storage proxy stays around waiting for active fetches to finish (with a max time of 60 seconds)
* Fixed so detailed stats are showing correct bytes for fetching files
* Improved detailed trace a bit to be able to see how much time is spent calculating caskey etc
* Changed const tchar* to const StringView& in some places
* Disabled memory mapping of pch file on local machine. Seems to create an annoying stall
* Added session stats to processes executing locally

diff --git a/Engine/Source/Programs/UnrealBuildAccelerator/Agent/Private/UbaAgent.cpp b/Engine/Source/Programs/UnrealBuildAccelerator/Agent/Private/UbaAgent.cpp
--- a/Engine/Source/Programs/UnrealBuildAccelerator/Agent/Private/UbaAgent.cpp
+++ b/Engine/Source/Programs/UnrealBuildAccelerator/Agent/Private/UbaAgent.cpp
@@ -1242,10 +1242,21 @@ namespace uba
 
 				networkBackend->StopListen();
 				storageClient->StopProxy();
+				sessionClient->Stop();
+
 				auto proxyServer = proxy.server.load();
 				if (proxyServer)
+				{
+					// Let's give the active fetches some time (60 seconds)
+					u32 waitCount = 60*10;
+					while (proxy.storage->GetActiveFetchCount())
+					{
+						Sleep(100);
+						if (!waitCount--)
+							break;
+					}
 					proxyServer->DisconnectClients();
-				sessionClient->Stop();
+				}
 				sessionClient->SendSummary([&](Logger& logger) { if (proxyServer) proxyServer->PrintSummary(logger); });
 				client->Disconnect();
 
diff --git a/Engine/Source/Programs/UnrealBuildAccelerator/Common/Private/UbaCacheServer.cpp b/Engine/Source/Programs/UnrealBuildAccelerator/Common/Private/UbaCacheServer.cpp
--- a/Engine/Source/Programs/UnrealBuildAccelerator/Common/Private/UbaCacheServer.cpp
+++ b/Engine/Source/Programs/UnrealBuildAccelerator/Common/Private/UbaCacheServer.cpp
@@ -169,15 +169,15 @@ namespace uba
 		}* m_maintenanceContext = nullptr;
 	};
 
-	const tchar* ToString(CacheMessageType type)
+	StringView ToString(CacheMessageType type)
 	{
 		switch (type)
 		{
-			#define UBA_CACHE_MESSAGE(x) case CacheMessageType_##x: return TC("")#x;
+			#define UBA_CACHE_MESSAGE(x) case CacheMessageType_##x: return AsView(TC("")#x);
 			UBA_CACHE_MESSAGES
 			#undef UBA_CACHE_MESSAGE
 		default:
-			return TC("Unknown"); // Should never happen
+			return AsView(TC("Unknown")); // Should never happen
 		}
 	}
 
diff --git a/Engine/Source/Programs/UnrealBuildAccelerator/Common/Private/UbaNetworkServer.cpp b/Engine/Source/Programs/UnrealBuildAccelerator/Common/Private/UbaNetworkServer.cpp
--- a/Engine/Source/Programs/UnrealBuildAccelerator/Common/Private/UbaNetworkServer.cpp
+++ b/Engine/Source/Programs/UnrealBuildAccelerator/Common/Private/UbaNetworkServer.cpp
@@ -576,7 +576,7 @@ namespace uba
 				server.m_additionalWork.pop_front();
 				lock.Leave();
 
-				u32 workId = server.TrackWorkStart(work.desc.c_str());
+				u32 workId = server.TrackWorkStart(work.desc);
 
 				work.func();
 
@@ -600,11 +600,11 @@ namespace uba
 
 	const tchar* g_typeStr[] = { TC("0"), TC("1"), TC("2"), TC("3"), TC("4"), TC("5"), TC("6"), TC("7"), TC("8"), TC("9"), TC("10"), TC("11"), TC("12") };
 
-	static const tchar* GetMessageTypeToName(u8 type)
+	static StringView GetMessageTypeToName(u8 type)
 	{
 		if (type <= 12)
-			return g_typeStr[type];
-		return TC("NUMBER HIGHER THAN 12");
+			return ToView(g_typeStr[type]);
+		return ToView(TC("NUMBER HIGHER THAN 12"));
 	}
 
 	NetworkServer::NetworkServer(bool& outCtorSuccess, const NetworkServerCreateInfo& info, const tchar* name)
@@ -980,7 +980,7 @@ namespace uba
 		m_additionalWork.pop_front();
 		lock.Leave();
 
-		TrackWorkScope tws(*this, work.desc.c_str());
+		TrackWorkScope tws(*this, work.desc);
 
 		work.func();
 
diff --git a/Engine/Source/Programs/UnrealBuildAccelerator/Common/Private/UbaSession.cpp b/Engine/Source/Programs/UnrealBuildAccelerator/Common/Private/UbaSession.cpp
--- a/Engine/Source/Programs/UnrealBuildAccelerator/Common/Private/UbaSession.cpp
+++ b/Engine/Source/Programs/UnrealBuildAccelerator/Common/Private/UbaSession.cpp
@@ -1508,6 +1508,7 @@ namespace uba
 		{
 			StackBinaryWriter<1024> writer;
 			process.m_processStats.Write(writer);
+			process.m_sessionStats.Write(writer);
 			process.m_storageStats.Write(writer);
 			process.m_kernelStats.Write(writer);
 			u32 exitCode = process.GetExitCode();
@@ -1709,12 +1710,17 @@ namespace uba
 	}
 
 	u32 Session::GetMemoryMapAlignment(const StringView& fileName) const
+	{
+		return GetMemoryMapAlignment(fileName, m_runningRemote);
+	}
+
+	u32 Session::GetMemoryMapAlignment(const StringView& fileName, bool runningRemote) const
 	{
 		// It is not necessarily better to make mem maps of everything.. only things that are read more than once in the build.
 		// Reason is because there is additional overhead to use memory mappings.
 		// Upside is that all things that are memory mapped can be stored compressed in cas storage so it saves space.
 
-		if (fileName.EndsWith(TC(".pch")))
+		if (runningRemote && fileName.EndsWith(TC(".pch")))
 			return 64 * 1024; // pch needs 64k alignment
 		if (fileName.EndsWith(TC(".h")) || fileName.EndsWith(TC(".inl")) || fileName.EndsWith(TC(".gch")))
 			return 4 * 1024; // clang seems to need 4k alignment? Is it a coincidence it works or what is happening inside the code? (msvc works with alignment 1byte here)
diff --git a/Engine/Source/Programs/UnrealBuildAccelerator/Common/Private/UbaSessionClient.cpp b/Engine/Source/Programs/UnrealBuildAccelerator/Common/Private/UbaSessionClient.cpp
--- a/Engine/Source/Programs/UnrealBuildAccelerator/Common/Private/UbaSessionClient.cpp
+++ b/Engine/Source/Programs/UnrealBuildAccelerator/Common/Private/UbaSessionClient.cpp
@@ -106,12 +106,7 @@ namespace uba
 	{
 		m_terminationTime = GetTime() + MsToTime(delayMs);
 		m_terminationReason = reason;
-
-		StackBinaryWriter<1024> writer;
-		NetworkMessage msg(m_client, ServiceId, SessionMessageType_Notification, writer);
-		writer.WriteU32(m_sessionId);
-		writer.WriteString(reason);
-		msg.Send();
+		SendNotification(ToView(reason));
 	}
 
 	void SessionClient::SetMaxProcessCount(u32 count)
@@ -1270,18 +1265,22 @@ namespace uba
 
 	bool SessionClient::SendProcessAvailable(Vector<InternalProcessStartInfo>& out, float availableWeight)
 	{
-		StackBinaryWriter<32> writer;
-		NetworkMessage msg(m_client, ServiceId, SessionMessageType_ProcessAvailable, writer);
-		writer.WriteU32(m_sessionId);
-		writer.WriteU32(*(u32*)&availableWeight);
-
 		StackBinaryReader<SendMaxSize> reader;
-		if (!msg.Send(reader, m_stats.procAvailableMsg))
+
 		{
-			if (m_loop)
-				m_logger.Error(TC("Failed to send ProcessAvailable message"));
-			return false;
+			StackBinaryWriter<32> writer;
+			NetworkMessage msg(m_client, ServiceId, SessionMessageType_ProcessAvailable, writer);
+			writer.WriteU32(m_sessionId);
+			writer.WriteU32(*(u32*)&availableWeight);
+
+			if (!msg.Send(reader, m_stats.procAvailableMsg))
+			{
+				if (m_loop)
+					m_logger.Error(TC("Failed to send ProcessAvailable message"));
+				return false;
+			}
 		}
+
 		while (true)
 		{
 			u32 processId = reader.ReadU32();
@@ -1312,14 +1311,14 @@ namespace uba
 			{
 				CasKey knownInputKey = reader.ReadCasKey();
 				u32 mappingAlignment = reader.ReadU32();
+				bool allowProxy = reader.ReadBool();
 				bool storeUncompressed = !m_allowMemoryMaps || mappingAlignment == 0;
 				if (storeUncompressed)
 					knownInputKey = AsCompressed(knownInputKey, false);
 
-				m_client.AddWork([knownInputKey, this]()
+				m_client.AddWork([knownInputKey, allowProxy, this]()
 					{
 						Storage::RetrieveResult result;
-						bool allowProxy = true;
 						bool res = m_storage.RetrieveCasFile(result, knownInputKey, TC("KnownInput"), nullptr, 1, allowProxy);
 						(void)res;
 					}, 1, TC("KnownInput"));
@@ -1465,6 +1464,15 @@ namespace uba
 			abort();
 	}
 
+	void SessionClient::SendNotification(const StringView& text)
+	{
+		StackBinaryWriter<1024> writer;
+		NetworkMessage msg(m_client, ServiceId, SessionMessageType_Notification, writer);
+		writer.WriteU32(m_sessionId);
+		writer.WriteString(text);
+		msg.Send();
+	}
+
 	void SessionClient::SendSummary(const Function<void(Logger&)>& extraInfo)
 	{
 		StackBinaryWriter<SendMaxSize> writer;
@@ -1584,6 +1592,8 @@ namespace uba
 				if (idleTime > m_maxIdleSeconds)
 				{
 					m_logger.Info(TC("Session has been idle longer than max idle time (%u seconds). Disconnecting (Did %u process requests during idle)"), m_maxIdleSeconds, processRequestCount);
+					SendNotification(AsView(TC("Idle time timeout")));
+
 					m_waitToSendEvent.Set();
 					m_remoteExecutionEnabled = false;
 				}
diff --git a/Engine/Source/Programs/UnrealBuildAccelerator/Common/Private/UbaSessionServer.cpp b/Engine/Source/Programs/UnrealBuildAccelerator/Common/Private/UbaSessionServer.cpp
--- a/Engine/Source/Programs/UnrealBuildAccelerator/Common/Private/UbaSessionServer.cpp
+++ b/Engine/Source/Programs/UnrealBuildAccelerator/Common/Private/UbaSessionServer.cpp
@@ -87,7 +87,7 @@ namespace uba
 		u32 m_sessionId = 0;
 		TString m_executingHost;
 
-		struct KnownInput { CasKey key; u32 mappingAlignment = 0; };
+		struct KnownInput { CasKey key; u32 mappingAlignment = 0; bool allowProxy = true; };
 		KnownInput* m_knownInputs = nullptr;
 		u32 m_knownInputsCount = 0;
 	};
@@ -118,11 +118,11 @@ namespace uba
 			{
 				switch (type)
 				{
-					#define UBA_SESSION_MESSAGE(x) case SessionMessageType_##x: return TC("")#x;
+					#define UBA_SESSION_MESSAGE(x) case SessionMessageType_##x: return AsView(TC("")#x);
 					UBA_SESSION_MESSAGES
 					#undef UBA_SESSION_MESSAGE
 				default:
-					return TC("Unknown");
+					return ToView(TC("Unknown"));
 				}
 			}
 		);
@@ -274,7 +274,11 @@ namespace uba
 		ValidateStartInfo(startInfo);
 		u32 processId = CreateProcessId();
 		RemoteProcess* remoteProcess = new RemoteProcess(this, startInfo, processId, weight);
-		
+
+		auto rules = GetRules(remoteProcess->m_startInfo);
+
+		remoteProcess->m_startInfo.rules = rules;
+
 		if (knownInputsCount)
 		{
 			auto keys = remoteProcess->m_knownInputs = new RemoteProcess::KnownInput[knownInputsCount];
@@ -295,7 +299,8 @@ namespace uba
 
 				auto& ki = keys[keysIndex++];
 				ki.key = casKey;
-				ki.mappingAlignment = GetMemoryMapAlignment(fileName);
+				ki.mappingAlignment = GetMemoryMapAlignment(fileName, true);
+				ki.allowProxy = rules->AllowStorageProxy(fileName);
 
 
 				// Update name to hash table
@@ -316,8 +321,6 @@ namespace uba
 			remoteProcess->m_knownInputsCount = keysIndex;
 		}
 
-		remoteProcess->m_startInfo.rules = GetRules(remoteProcess->m_startInfo);
-
 		ProcessHandle h(remoteProcess); // Keep ref count up even if process is removed by callbacks etc.
 
 		ScopedCriticalSection lock(m_remoteProcessAndSessionLock);
@@ -516,7 +519,7 @@ namespace uba
 			it = m_activeRemoteProcesses.erase(it);
 			remoteProcess->m_executingHost.clear();
 
-			m_trace.ProcessReturned(remoteProcess->m_processId, ToView(TC("Disconnected")));
+			m_trace.ProcessReturned(remoteProcess->m_processId, AsView(TC("Disconnected")));
 
 			ProcessHandle h = ProcessRemoved(remoteProcess->m_processId);
 			if (!h.m_process)
@@ -530,6 +533,7 @@ namespace uba
 			if (m_remoteProcessReturnedEvent)
 				m_remoteProcessReturnedEvent(*remoteProcess);
 		}
+
 		m_returnedRemoteProcessCount += returnCount;
 
 		u32 sessionId = 0;
@@ -540,6 +544,10 @@ namespace uba
 			auto& s = *sptr;
 			if (s.id != clientId)
 				continue;
+
+			if (!returnCount && !s.hasNotification && !m_remoteExecutionEnabled)
+				m_trace.SessionNotification(sessionId, TC("Done"));
+
 			m_trace.SessionDisconnect(sessionId);
 
 			sessionName.Append(s.name);
@@ -1172,6 +1180,7 @@ namespace uba
 				break;
 			writer.WriteCasKey(kv->key);
 			writer.WriteU32(kv->mappingAlignment);
+			writer.WriteBool(kv->allowProxy);
 		}
 		return true;
 	}
@@ -1355,6 +1364,14 @@ namespace uba
 	bool SessionServer::HandleNotification(const ConnectionInfo& connectionInfo, BinaryReader& reader, BinaryWriter& writer)
 	{
 		u32 sessionId = reader.ReadU32();
+
+		u32 sessionIndex = sessionId - 1;
+		{
+			ScopedCriticalSection lock(m_remoteProcessAndSessionLock);
+			if (sessionIndex < m_clientSessions.size())
+				m_clientSessions[sessionIndex]->hasNotification = true;
+		}
+
 		StringBuffer<1024> str;
 		reader.ReadString(str);
 		m_trace.SessionNotification(sessionId, str.data);
diff --git a/Engine/Source/Programs/UnrealBuildAccelerator/Common/Private/UbaStorageClient.cpp b/Engine/Source/Programs/UnrealBuildAccelerator/Common/Private/UbaStorageClient.cpp
--- a/Engine/Source/Programs/UnrealBuildAccelerator/Common/Private/UbaStorageClient.cpp
+++ b/Engine/Source/Programs/UnrealBuildAccelerator/Common/Private/UbaStorageClient.cpp
@@ -9,8 +9,21 @@
 #include "UbaNetworkBackendTcp.h"
 #include "UbaStorageUtils.h"
 
+#define UBA_REPORT_PROXYFETCH 0
+
 namespace uba
 {
+	inline CasKey AsProxyCasKey(const CasKey& key)
+	{
+		CasKey newKey = key;
+		#ifndef __clang_analyzer__
+		u8 flagField = ((u8*)&key)[19];
+		((u8*)&newKey)[19] = (flagField | u8(4));
+		#endif
+		return newKey;
+	}
+
+
 	void StorageClientCreateInfo::Apply(Config& config)
 	{
 		StorageCreateInfo::Apply(config);
@@ -243,6 +256,10 @@ namespace uba
 		u64 actualSize = 0;
 		u64 sizeOnDisk = 0;
 
+		#if UBA_REPORT_PROXYFETCH
+		bool proxyFetchSent = false;
+		#endif
+
 		while (true)
 		{
 			u8* readBuffer = nullptr;
@@ -304,9 +321,22 @@ namespace uba
 				});
 
 			{
+				#if UBA_REPORT_PROXYFETCH
+				if (proxy && !proxyFetchSent)
+				{
+					proxyFetchSent = true;
+					StackBinaryWriter<1024> writer;
+					NetworkMessage msg(m_client, ServiceId, StorageMessageType_ProxyFetchBegin, writer);
+					writer.WriteCasKey(AsProxyCasKey(casKey));
+					writer.WriteString(hint);
+					StackBinaryReader<32> reader;
+					msg.Send(reader); // Need to use reader just to make sure this reaches server before the ProxyFetchEnd
+				}
+				#endif
+
 				StackBinaryWriter<1024> writer;
 				NetworkMessage msg(*client, ServiceId, StorageMessageType_FetchBegin, writer);
-				writer.WriteBool(wantsProxy);
+				writer.WriteByte(wantsProxy ? 1 : 0);
 				writer.WriteCasKey(casKey);
 				writer.WriteString(hint);
 				BinaryReader reader(slot, 0, SendMaxSize);
@@ -660,6 +690,17 @@ namespace uba
 					return false;
 			}
 
+			#if UBA_REPORT_PROXYFETCH
+			if (proxyFetchSent)
+			{
+				StackBinaryWriter<1024> writer;
+				NetworkMessage msg(m_client, ServiceId, StorageMessageType_ProxyFetchEnd, writer);
+				writer.WriteCasKey(AsProxyCasKey(casKey));
+				msg.Send();
+			}
+			#endif
+
+
 			if (!mappingBuffer)
 				if (!destinationFile.Close())
 					return false;
diff --git a/Engine/Source/Programs/UnrealBuildAccelerator/Common/Private/UbaStorageProxy.cpp b/Engine/Source/Programs/UnrealBuildAccelerator/Common/Private/UbaStorageProxy.cpp
--- a/Engine/Source/Programs/UnrealBuildAccelerator/Common/Private/UbaStorageProxy.cpp
+++ b/Engine/Source/Programs/UnrealBuildAccelerator/Common/Private/UbaStorageProxy.cpp
@@ -88,6 +88,12 @@ namespace uba
 		logger.Info(TC(""));
 	}
 
+	u32 StorageProxy::GetActiveFetchCount()
+	{
+		SCOPED_READ_LOCK(m_activeFetchesLock, lock);
+		return u32(m_activeFetches.size());
+	}
+
 	u16 StorageProxy::PopId()
 	{
 		if (m_availableIds.empty())
@@ -132,7 +138,8 @@ namespace uba
 
 		case StorageMessageType_FetchBegin:
 			{
-				reader.ReadBool(); // Wants proxy
+				u8 recvFlags = reader.ReadByte(); // Wants proxy
+				UBA_ASSERT((recvFlags & 2) == 0);(void)recvFlags;
 
 				CasKey casKey = reader.ReadCasKey();
 				StringBuffer<> hint;
@@ -185,10 +192,10 @@ namespace uba
 
 					if (!file.memory)
 					{
-						file.trackId = m_client.TrackWorkStart(StringBuffer<512>().AppendFileName(hint.data).data);
+						file.trackId = m_client.TrackWorkStart(StringBuffer<512>().AppendFileName(hint.data));
 
 						NetworkMessage msg(m_client, ServiceId, StorageMessageType_FetchBegin, writer2);
-						writer2.WriteBool(false); // Wants proxy
+						writer2.WriteByte(2); // Does not want proxy but informs it is proxy
 						writer2.WriteCasKey(casKey);
 						writer2.WriteString(hint);
 						writer2.WriteBytes(reader.GetPositionData(), reader.GetLeft());
diff --git a/Engine/Source/Programs/UnrealBuildAccelerator/Common/Private/UbaStorageServer.cpp b/Engine/Source/Programs/UnrealBuildAccelerator/Common/Private/UbaStorageServer.cpp
--- a/Engine/Source/Programs/UnrealBuildAccelerator/Common/Private/UbaStorageServer.cpp
+++ b/Engine/Source/Programs/UnrealBuildAccelerator/Common/Private/UbaStorageServer.cpp
@@ -175,7 +175,7 @@ namespace uba
 					store.casEntry->verified = false;
 					store.casEntry->beingWritten = false;
 					if (m_traceStore)
-						m_trace->FileEndStore(clientId, store.casEntry->key);
+						m_trace->FileStoreEnd(clientId, store.casEntry->key);
 				}
 
 				if (auto fa = store.fileAccessor)
@@ -205,7 +205,7 @@ namespace uba
 				fetch.Release(*this, TC("OnDisconnected"));
 
 				if (m_traceFetch)
-					m_trace->FileEndFetch(clientId, AsCompressed(fetch.casKey, m_storeCompressed));
+					m_trace->FileFetchEnd(clientId, AsCompressed(fetch.casKey, m_storeCompressed));
 
 				it = m_activeFetches.erase(it);
 			}
@@ -362,7 +362,9 @@ namespace uba
 						}
 					});
 
-				if (reader.ReadBool()) // Wants proxy
+				u8 recvFlags = reader.ReadByte();
+				bool wantsProxy = (recvFlags & 1) != 0;
+				if (wantsProxy) // Wants proxy
 				{
 					todoRemoveMe = 1;
 					SCOPED_READ_LOCK(m_connectionInfoLock, lock);
@@ -440,6 +442,12 @@ namespace uba
 
 				casKey = AsCompressed(casKey, m_storeCompressed);
 
+				bool detailedTrace = m_traceFetch;
+				if (detailedTrace)
+					m_trace->FileFetchBegin(connectionInfo.GetId(), casKey, hint);
+				auto traceFetchGuard = MakeGuard([&](){ if (detailedTrace) m_trace->FileFetchEnd(connectionInfo.GetId(), casKey); });
+
+
 				CasEntry* casEntry = nullptr;
 				bool has = HasCasFile(casKey, &casEntry); // HasCasFile also writes deferred cas entries if in queue
 				if (!has)
@@ -566,10 +574,10 @@ namespace uba
 
 				todoRemoveMe = 6;
 
-				if (m_trace)
-					m_trace->FileBeginFetch(connectionInfo.GetId(), casKey, fileSize, hint.data, m_traceFetch);
-
-				auto cg = MakeGuard([&](){ if (m_traceFetch) m_trace->FileEndFetch(connectionInfo.GetId(), casKey); });
+				if (detailedTrace)
+					m_trace->FileFetchSize(connectionInfo.GetId(), casKey, fileSize);
+				else if (m_trace)
+					m_trace->FileFetchLight(connectionInfo.GetId(), casKey, fileSize);
 
 				u64 left = fileSize;
 
@@ -629,6 +637,8 @@ namespace uba
 
 				left -= toWrite;
 
+				traceFetchGuard.Cancel();
+
 				if (!left)
 				{
 					*fetchId = u16(~0);
@@ -640,7 +650,6 @@ namespace uba
 				todoRemoveMe = 9;
 
 				mvg.Cancel();
-				cg.Cancel();
 				rfg.Cancel();
 
 				*fetchId = PopId();
@@ -713,7 +722,7 @@ namespace uba
 			{
 				CasKey key = reader.ReadCasKey();
 				if (m_traceFetch)
-					m_trace->FileEndFetch(connectionInfo.GetId(), AsCompressed(key, m_storeCompressed));
+					m_trace->FileFetchEnd(connectionInfo.GetId(), AsCompressed(key, m_storeCompressed));
 				return true;
 			}
 
@@ -861,7 +870,7 @@ namespace uba
 				firstStore->recvCasTime = GetTime() - start;
 
 				if (m_trace)
-					m_trace->FileBeginStore(connectionInfo.GetId(), casKey, fileSize, hint.data, m_traceStore);
+					m_trace->FileStoreBegin(connectionInfo.GetId(), casKey, fileSize, hint, m_traceStore);
 
 				[[fallthrough]];
 			}
@@ -981,10 +990,25 @@ namespace uba
 			{
 				CasKey key = reader.ReadCasKey();
 				if (m_traceStore)
-					m_trace->FileEndStore(connectionInfo.GetId(), key);
+					m_trace->FileStoreEnd(connectionInfo.GetId(), key);
+				return true;
+			}
+			case StorageMessageType_ProxyFetchBegin:
+			{
+				CasKey casKey = reader.ReadCasKey();
+				StringBuffer<> hint;
+				reader.ReadString(hint);
+				if (m_trace)
+					m_trace->FileFetchBegin(connectionInfo.GetId(), casKey, hint);
+				return true;
+			}
+			case StorageMessageType_ProxyFetchEnd:
+			{
+				CasKey casKey = reader.ReadCasKey();
+				if (m_trace)
+					m_trace->FileFetchEnd(connectionInfo.GetId(), casKey);
 				return true;
 			}
-
 		}
 		UBA_ASSERT(false);
 		return false;
diff --git a/Engine/Source/Programs/UnrealBuildAccelerator/Common/Private/UbaTrace.cpp b/Engine/Source/Programs/UnrealBuildAccelerator/Common/Private/UbaTrace.cpp
--- a/Engine/Source/Programs/UnrealBuildAccelerator/Common/Private/UbaTrace.cpp
+++ b/Engine/Source/Programs/UnrealBuildAccelerator/Common/Private/UbaTrace.cpp
@@ -129,7 +129,7 @@ namespace uba
 		return true;
 	}
 
-	u32 Trace::TrackWorkStart(const tchar* desc)
+	u32 Trace::TrackWorkStart(const StringView& desc)
 	{
 		u32 workId = m_workCounter++;
 		BeginWork(workId, desc);
@@ -176,15 +176,13 @@ namespace uba
 	}
 
 
-	u32 Trace::AddString(const tchar* string)
+	u32 Trace::AddString(const StringView& string)
 	{
 		if (!m_memoryBegin)
 			return 0;
 
-		u64 stringLen = TStrlen(string);
-
 		SCOPED_WRITE_LOCK(m_stringsLock, lock);
-		auto insres = m_strings.try_emplace(ToStringKeyNoCheck(string, stringLen));
+		auto insres = m_strings.try_emplace(ToStringKeyNoCheck(string.data, string.count));
 		if (insres.second)
 		{
 			insres.first->second = u32(m_strings.size() - 1);
@@ -192,7 +190,7 @@ namespace uba
 			if (!writer.IsValid())
 				return 0;
 			writer.WriteByte(TraceType_String);
-			writer.WriteString(string, stringLen);
+			writer.WriteString(string);
 		}
 		return insres.first->second;
 	}
@@ -308,38 +306,43 @@ namespace uba
 		writer.WriteString(proxyName);
 	}
 
-	void Trace::FileBeginFetch(u32 clientId, const CasKey& key, u64 size, const tchar* hint, bool detailed)
+	void Trace::FileFetchLight(u32 clientId, const CasKey& key, u64 fileSize)
 	{
-		if (detailed)
-		{
-			u32 stringIndex = AddString(hint);
-			BEGIN_TRACE_ENTRY(FileBeginFetch);
-			writer.Write7BitEncoded(clientId);
-			writer.WriteCasKey(key);
-			writer.Write7BitEncoded(size);
-			writer.Write7BitEncoded(stringIndex);
-		}
-		else
-		{
 			BEGIN_TRACE_ENTRY(FileFetchLight);
 			writer.Write7BitEncoded(clientId);
-			writer.Write7BitEncoded(size);
-		}
+			writer.Write7BitEncoded(fileSize);
+	}
+
+	void Trace::FileFetchBegin(u32 clientId, const CasKey& key, const StringView& hint)
+	{
+		u32 stringIndex = AddString(hint);
+		BEGIN_TRACE_ENTRY(FileFetchBegin);
+		writer.Write7BitEncoded(clientId);
+		writer.WriteCasKey(key);
+		writer.Write7BitEncoded(stringIndex);
+	}
+
+	void Trace::FileFetchSize(u32 clientId, const CasKey& key, u64 fileSize)
+	{
+		BEGIN_TRACE_ENTRY(FileFetchSize);
+		writer.Write7BitEncoded(clientId);
+		writer.WriteCasKey(key);
+		writer.Write7BitEncoded(fileSize);
 	}
 
-	void Trace::FileEndFetch(u32 clientId, const CasKey& key)
+	void Trace::FileFetchEnd(u32 clientId, const CasKey& key)
 	{
-		BEGIN_TRACE_ENTRY(FileEndFetch);
+		BEGIN_TRACE_ENTRY(FileFetchEnd);
 		writer.Write7BitEncoded(clientId);
 		writer.WriteCasKey(key);
 	}
 
-	void Trace::FileBeginStore(u32 clientId, const CasKey& key, u64 size, const tchar* hint, bool detailed)
+	void Trace::FileStoreBegin(u32 clientId, const CasKey& key, u64 size, const StringView& hint, bool detailed)
 	{
 		if (detailed)
 		{
 			u32 stringIndex = AddString(hint);
-			BEGIN_TRACE_ENTRY(FileBeginStore);
+			BEGIN_TRACE_ENTRY(FileStoreBegin);
 			writer.Write7BitEncoded(clientId);
 			writer.WriteCasKey(key);
 			writer.Write7BitEncoded(size);
@@ -353,14 +356,14 @@ namespace uba
 		}
 	}
 
-	void Trace::FileEndStore(u32 clientId, const CasKey& key)
+	void Trace::FileStoreEnd(u32 clientId, const CasKey& key)
 	{
-		BEGIN_TRACE_ENTRY(FileEndStore);
+		BEGIN_TRACE_ENTRY(FileStoreEnd);
 		writer.Write7BitEncoded(clientId);
 		writer.WriteCasKey(key);
 	}
 
-	void Trace::BeginWork(u32 workIndex, const tchar* desc)
+	void Trace::BeginWork(u32 workIndex, const StringView& desc)
 	{
 		u32 stringIndex = AddString(desc);
 		BEGIN_TRACE_ENTRY(BeginWork);
diff --git a/Engine/Source/Programs/UnrealBuildAccelerator/Common/Private/UbaWorkManager.cpp b/Engine/Source/Programs/UnrealBuildAccelerator/Common/Private/UbaWorkManager.cpp
--- a/Engine/Source/Programs/UnrealBuildAccelerator/Common/Private/UbaWorkManager.cpp
+++ b/Engine/Source/Programs/UnrealBuildAccelerator/Common/Private/UbaWorkManager.cpp
@@ -7,7 +7,7 @@
 
 namespace uba
 {
-	u32 WorkManager::TrackWorkStart(const tchar* desc)
+	u32 WorkManager::TrackWorkStart(const StringView& desc)
 	{
 		if (auto t = m_workTracker.load())
 			return t->TrackWorkStart(desc);
@@ -60,7 +60,7 @@ namespace uba
 						manager.m_work.pop_front();
 						lock.Leave();
 
-						TrackWorkScope tws(manager, work.desc.c_str());
+						TrackWorkScope tws(manager, work.desc);
 						work.func();
 					}
 
@@ -162,7 +162,7 @@ namespace uba
 			m_work.pop_front();
 			lock.Leave();
 
-			TrackWorkScope tws(*this, work.desc.c_str());
+			TrackWorkScope tws(*this, work.desc);
 			work.func();
 		}
 	}
diff --git a/Engine/Source/Programs/UnrealBuildAccelerator/Common/Public/UbaNetwork.h b/Engine/Source/Programs/UnrealBuildAccelerator/Common/Public/UbaNetwork.h
--- a/Engine/Source/Programs/UnrealBuildAccelerator/Common/Public/UbaNetwork.h
+++ b/Engine/Source/Programs/UnrealBuildAccelerator/Common/Public/UbaNetwork.h
@@ -3,6 +3,7 @@
 #pragma once
 
 #include "UbaDefaultConstants.h"
+#include "UbaStringBuffer.h"
 
 namespace uba
 {
@@ -15,7 +16,7 @@ namespace uba
 	
 	static constexpr u32 SystemNetworkVersion = 1339;
 	static constexpr u32 StorageNetworkVersion = 4;
-	static constexpr u32 SessionNetworkVersion = 34;
+	static constexpr u32 SessionNetworkVersion = 35;
 	static constexpr u32 CacheNetworkVersion = 5;
 
 	// Messages used over network between client/server (system, storage and session)
@@ -35,6 +36,8 @@ namespace uba
 		UBA_STORAGE_MESSAGE(StoreSegment) \
 		UBA_STORAGE_MESSAGE(StoreEnd) \
 		UBA_STORAGE_MESSAGE(Connect) \
+		UBA_STORAGE_MESSAGE(ProxyFetchBegin) \
+		UBA_STORAGE_MESSAGE(ProxyFetchEnd) \
 
 	enum StorageMessageType : u8
 	{
@@ -43,15 +46,15 @@ namespace uba
 		#undef UBA_STORAGE_MESSAGE
 	};
 
-	inline const tchar* ToString(StorageMessageType type)
+	inline StringView ToString(StorageMessageType type)
 	{
 		switch (type)
 		{
-			#define UBA_STORAGE_MESSAGE(x) case StorageMessageType_##x: return TC("")#x;
+			#define UBA_STORAGE_MESSAGE(x) case StorageMessageType_##x: return AsView(TC("")#x);
 			UBA_STORAGE_MESSAGES
 			#undef UBA_STORAGE_MESSAGE
 		default:
-			return TC("Unknown"); // Should never happen
+			return ToView(TC("Unknown")); // Should never happen
 		}
 	}
 
diff --git a/Engine/Source/Programs/UnrealBuildAccelerator/Common/Public/UbaNetworkServer.h b/Engine/Source/Programs/UnrealBuildAccelerator/Common/Public/UbaNetworkServer.h
--- a/Engine/Source/Programs/UnrealBuildAccelerator/Common/Public/UbaNetworkServer.h
+++ b/Engine/Source/Programs/UnrealBuildAccelerator/Common/Public/UbaNetworkServer.h
@@ -61,7 +61,7 @@ namespace uba
 
 		void PrintSummary(Logger& logger);
 
-		using TypeToNameFunction = const tchar*(u8 type);
+		using TypeToNameFunction = StringView(u8 type);
 		using WorkerFunction = Function<bool(const ConnectionInfo& connectionInfo, MessageInfo& messageInfo, BinaryReader& reader, BinaryWriter& writer)>;
 		void RegisterService(u8 serviceId, const WorkerFunction& function, TypeToNameFunction* typeToNameFunc = nullptr);
 		void UnregisterService(u8 serviceId);
diff --git a/Engine/Source/Programs/UnrealBuildAccelerator/Common/Public/UbaSession.h b/Engine/Source/Programs/UnrealBuildAccelerator/Common/Public/UbaSession.h
--- a/Engine/Source/Programs/UnrealBuildAccelerator/Common/Public/UbaSession.h
+++ b/Engine/Source/Programs/UnrealBuildAccelerator/Common/Public/UbaSession.h
@@ -97,6 +97,7 @@ namespace uba
 		u32 GetDirectoryTableSize();
 		u32 GetFileMappingSize();
 		u32 GetMemoryMapAlignment(const StringView& fileName) const;
+		u32 GetMemoryMapAlignment(const StringView& fileName, bool runningRemote) const;
 
 		SessionStats& Stats();
 
diff --git a/Engine/Source/Programs/UnrealBuildAccelerator/Common/Public/UbaSessionClient.h b/Engine/Source/Programs/UnrealBuildAccelerator/Common/Public/UbaSessionClient.h
--- a/Engine/Source/Programs/UnrealBuildAccelerator/Common/Public/UbaSessionClient.h
+++ b/Engine/Source/Programs/UnrealBuildAccelerator/Common/Public/UbaSessionClient.h
@@ -92,6 +92,7 @@ namespace uba
 		bool SendProcessInputs(ProcessImpl& process);
 		bool SendProcessFinished(ProcessImpl& process, u32 exitCode);
 		void SendPing(u64 memAvail, u64 memTotal);
+		void SendNotification(const StringView& text);
 		void SendLogFileToServer(ProcessImpl& pi);
 		void GetLogFileName(StringBufferBase& out, const tchar* logFile, const tchar* arguments, u32 processId);
 		u32 CountLogLines(ProcessImpl& process);
diff --git a/Engine/Source/Programs/UnrealBuildAccelerator/Common/Public/UbaSessionServer.h b/Engine/Source/Programs/UnrealBuildAccelerator/Common/Public/UbaSessionServer.h
--- a/Engine/Source/Programs/UnrealBuildAccelerator/Common/Public/UbaSessionServer.h
+++ b/Engine/Source/Programs/UnrealBuildAccelerator/Common/Public/UbaSessionServer.h
@@ -147,6 +147,7 @@ namespace uba
 			bool enabled = true;
 			bool dedicated = false;
 			bool abort = false;
+			bool hasNotification = false;
 		};
 		Vector<ClientSession*> m_clientSessions;
 
diff --git a/Engine/Source/Programs/UnrealBuildAccelerator/Common/Public/UbaStorageProxy.h b/Engine/Source/Programs/UnrealBuildAccelerator/Common/Public/UbaStorageProxy.h
--- a/Engine/Source/Programs/UnrealBuildAccelerator/Common/Public/UbaStorageProxy.h
+++ b/Engine/Source/Programs/UnrealBuildAccelerator/Common/Public/UbaStorageProxy.h
@@ -23,6 +23,8 @@ namespace uba
 		~StorageProxy();
 		void PrintSummary();
 
+		u32 GetActiveFetchCount();
+
 	protected:
 		u16 PopId();
 		void PushId(u16 id);
diff --git a/Engine/Source/Programs/UnrealBuildAccelerator/Common/Public/UbaTrace.h b/Engine/Source/Programs/UnrealBuildAccelerator/Common/Public/UbaTrace.h
--- a/Engine/Source/Programs/UnrealBuildAccelerator/Common/Public/UbaTrace.h
+++ b/Engine/Source/Programs/UnrealBuildAccelerator/Common/Public/UbaTrace.h
@@ -34,10 +34,10 @@ namespace uba
 		UBA_TRACE_TYPE(ProcessAdded) \
 		UBA_TRACE_TYPE(ProcessExited) \
 		UBA_TRACE_TYPE(ProcessReturned) \
-		UBA_TRACE_TYPE(FileBeginFetch) \
-		UBA_TRACE_TYPE(FileEndFetch) \
-		UBA_TRACE_TYPE(FileBeginStore) \
-		UBA_TRACE_TYPE(FileEndStore) \
+		UBA_TRACE_TYPE(FileFetchBegin) \
+		UBA_TRACE_TYPE(FileFetchEnd) \
+		UBA_TRACE_TYPE(FileStoreBegin) \
+		UBA_TRACE_TYPE(FileStoreEnd) \
 		UBA_TRACE_TYPE(Summary) \
 		UBA_TRACE_TYPE(BeginWork) \
 		UBA_TRACE_TYPE(EndWork) \
@@ -57,6 +57,7 @@ namespace uba
 		UBA_TRACE_TYPE(CacheEndWrite) \
 		UBA_TRACE_TYPE(ProgressUpdate) \
 		UBA_TRACE_TYPE(RemoteExecutionDisabled) \
+		UBA_TRACE_TYPE(FileFetchSize) \
 
 	enum TraceType : u8
 	{
@@ -68,7 +69,7 @@ namespace uba
 	using Color = u32;
 	inline Color toColor(u8 r, u8 g, u8 b) { return (r << 16) + (g << 8) + b; }
 
-	static constexpr u32 TraceVersion = 35;
+	static constexpr u32 TraceVersion = 36;
 	static constexpr u32 TraceReadCompatibilityVersion = 6;
 
 	class Trace : public WorkTracker
@@ -91,11 +92,13 @@ namespace uba
 		void ProcessReturned(u32 processId, const StringView& reason);
 		void ProxyCreated(u32 clientId, const tchar* proxyName);
 		void ProxyUsed(u32 clientId, const tchar* proxyName);
-		void FileBeginFetch(u32 clientId, const CasKey& key, u64 size, const tchar* hint, bool detailed);
-		void FileEndFetch(u32 clientId, const CasKey& key);
-		void FileBeginStore(u32 clientId, const CasKey& key, u64 size, const tchar* hint, bool detailed);
-		void FileEndStore(u32 clientId, const CasKey& key);
-		void BeginWork(u32 workIndex, const tchar* desc);
+		void FileFetchLight(u32 clientId, const CasKey& key, u64 fileSize);
+		void FileFetchBegin(u32 clientId, const CasKey& key, const StringView& hint);
+		void FileFetchSize(u32 clientId, const CasKey& key, u64 size);
+		void FileFetchEnd(u32 clientId, const CasKey& key);
+		void FileStoreBegin(u32 clientId, const CasKey& key, u64 size, const StringView& hint, bool detailed);
+		void FileStoreEnd(u32 clientId, const CasKey& key);
+		void BeginWork(u32 workIndex, const StringView& desc);
 		void EndWork(u32 workIndex);
 		void ProgressUpdate(u32 processesTotal, u32 processesDone, u32 errorCount);
 		void StatusUpdate(u32 statusRow, u32 statusColumn, const tchar* statusText, LogEntryType statusType, const tchar* statusLink);
@@ -108,14 +111,14 @@ namespace uba
 
 		bool StopWrite(const tchar* writeFileName);
 
-		virtual u32 TrackWorkStart(const tchar* desc) final override;
+		virtual u32 TrackWorkStart(const StringView& desc) final override;
 		virtual void TrackWorkEnd(u32 id) final override;
 
 	private:
 		struct WriterScope;
 		void FreeMemory();
 		bool EnsureMemory(u64 size);
-		u32 AddString(const tchar* string);
+		u32 AddString(const StringView& string);
 
 		LoggerWithWriter m_logger;
 		TString m_namedTrace;
diff --git a/Engine/Source/Programs/UnrealBuildAccelerator/Common/Public/UbaWorkManager.h b/Engine/Source/Programs/UnrealBuildAccelerator/Common/Public/UbaWorkManager.h
--- a/Engine/Source/Programs/UnrealBuildAccelerator/Common/Public/UbaWorkManager.h
+++ b/Engine/Source/Programs/UnrealBuildAccelerator/Common/Public/UbaWorkManager.h
@@ -7,10 +7,12 @@
 
 namespace uba
 {
+	struct StringView;
+
 	class WorkTracker
 	{
 	public:
-		virtual u32 TrackWorkStart(const tchar* desc) { return 0; }
+		virtual u32 TrackWorkStart(const StringView& desc) { return 0; }
 		virtual void TrackWorkEnd(u32 id) {}
 	};
 
@@ -21,7 +23,7 @@ namespace uba
 		virtual u32 GetWorkerCount() = 0;
 		virtual void DoWork(u32 count = 1) = 0;
 
-		u32 TrackWorkStart(const tchar* desc);
+		u32 TrackWorkStart(const StringView& desc);
 		void TrackWorkEnd(u32 id);
 
 		void SetWorkTracker(WorkTracker* workTracker) { m_workTracker = workTracker; }
@@ -64,7 +66,7 @@ namespace uba
 
 	struct TrackWorkScope
 	{
-		TrackWorkScope(WorkManager& wm, const tchar* desc) : workManager(wm), workIndex(wm.TrackWorkStart(desc)) {}
+		TrackWorkScope(WorkManager& wm, const StringView& desc) : workManager(wm), workIndex(wm.TrackWorkStart(desc)) {}
 		~TrackWorkScope() { workManager.TrackWorkEnd(workIndex); }
 		WorkManager& workManager;
 		u32 workIndex;
diff --git a/Engine/Source/Programs/UnrealBuildAccelerator/Core/Public/UbaStringBuffer.h b/Engine/Source/Programs/UnrealBuildAccelerator/Core/Public/UbaStringBuffer.h
--- a/Engine/Source/Programs/UnrealBuildAccelerator/Core/Public/UbaStringBuffer.h
+++ b/Engine/Source/Programs/UnrealBuildAccelerator/Core/Public/UbaStringBuffer.h
@@ -95,6 +95,7 @@ namespace uba
 		StringView(const tchar* d, u32 c) : data(d), count(c) {}
 		StringView(const StringBufferBase& sb) : data(sb.data), count(sb.count) {}
 		StringView(const TString& str) : data(str.data()), count(u32(str.size())) {}
+
 		bool IsEmpty() const { return count == 0; }
 		bool StartsWith(const tchar* str, bool ignoreCase = true) const { return uba::StartsWith(data, str, ignoreCase); }
 		bool EndsWith(const tchar* value, bool ignoreCase = true) const { return uba::EndsWith(data, count, value, ignoreCase); }
@@ -107,6 +108,9 @@ namespace uba
 
 	StringView ToView(const tchar* s);
 
+	template<typename T, size_t size>
+	StringView AsView(T (&buffer)[size]) { return StringView(buffer, size-1); }
+
 	struct LastErrorToText : StringBuffer<256>
 	{
 		LastErrorToText();
diff --git a/Engine/Source/Programs/UnrealBuildAccelerator/Visualizer/Private/UbaTraceReader.cpp b/Engine/Source/Programs/UnrealBuildAccelerator/Visualizer/Private/UbaTraceReader.cpp
--- a/Engine/Source/Programs/UnrealBuildAccelerator/Visualizer/Private/UbaTraceReader.cpp
+++ b/Engine/Source/Programs/UnrealBuildAccelerator/Visualizer/Private/UbaTraceReader.cpp
@@ -598,6 +598,8 @@ namespace uba
 			}
 			else if (out.version >= 30)
 			{
+				if (out.version >= 36)
+					sessionStats.Read(reader, out.version);
 				storageStats.Read(reader, out.version);
 				kernelStats.Read(reader, out.version);
 			}
@@ -727,11 +729,14 @@ namespace uba
 			process.bitmapDirty = true;
 			break;
 		}
-		case TraceType_FileBeginFetch:
+		case TraceType_FileFetchBegin:
 		{
 			Guid clientUid = ReadClientId(out, reader);
 			CasKey key = reader.ReadCasKey();
-			u64 size = reader.Read7BitEncoded();
+			u64 size = 0;
+			if (out.version < 36)
+				size = reader.Read7BitEncoded();
+
 			StringBuffer<> temp;
 			const tchar* hint;
 			if (out.version < 14)
@@ -754,8 +759,8 @@ namespace uba
 
 			if (auto session = GetSession(out, clientUid))
 			{
+				session->fetchedFilesActive.try_emplace(key, u32(session->fetchedFiles.size()));
 				session->fetchedFiles.push_back({ key, size, hint, time, ~u64(0) });
-				session->fetchedFilesBytes += size;
 			}
 			break;
 		}
@@ -812,7 +817,23 @@ namespace uba
 				session->proxyName = proxyName.data;
 			break;
 		}
-		case TraceType_FileEndFetch:
+		case TraceType_FileFetchSize:
+		{
+			Guid clientUid = ReadClientId(out, reader);
+			CasKey key = reader.ReadCasKey();
+			u64 fileSize = reader.Read7BitEncoded();
+			if (auto session = GetSession(out, clientUid))
+			{
+				auto findIt = session->fetchedFilesActive.find(key);
+				if (findIt != session->fetchedFilesActive.end())
+				{
+					auto& file = session->fetchedFiles[findIt->second];
+					file.size = fileSize;
+				}
+			}
+			break;
+		}
+		case TraceType_FileFetchEnd:
 		{
 			Guid clientUid = ReadClientId(out, reader);
 			CasKey key = reader.ReadCasKey();
@@ -829,17 +850,19 @@ namespace uba
 
 			if (auto session = GetSession(out, clientUid))
 			{
-				for (auto rit = session->fetchedFiles.rbegin(), rend = session->fetchedFiles.rend(); rit != rend; ++rit)
+				auto findIt = session->fetchedFilesActive.find(key);
+				if (findIt != session->fetchedFilesActive.end())
 				{
-					if (rit->key != key)
-						continue;
-					rit->stop = time;
+					auto& file = session->fetchedFiles[findIt->second];
+					file.stop = time;
+					session->fetchedFilesBytes += file.size;
+					session->fetchedFilesActive.erase(findIt);
 					break;
 				}
 			}
 			break;
 		}
-		case TraceType_FileBeginStore:
+		case TraceType_FileStoreBegin:
 		{
 			Guid clientUid = ReadClientId(out, reader);
 			CasKey key = reader.ReadCasKey();
@@ -881,7 +904,7 @@ namespace uba
 			}
 			break;
 		}
-		case TraceType_FileEndStore:
+		case TraceType_FileStoreEnd:
 		{
 			Guid clientUid = ReadClientId(out, reader);
 			CasKey key = reader.ReadCasKey();
@@ -1015,6 +1038,8 @@ namespace uba
 		case TraceType_RemoteExecutionDisabled:
 		{
 			out.remoteExecutionDisabled = true;
+			if (!out.sessions.empty())
+				out.sessions[0].notification = TC("(Remote execution disabled)");
 			break;
 		}
 		case TraceType_String:
diff --git a/Engine/Source/Programs/UnrealBuildAccelerator/Visualizer/Private/UbaTraceReader.h b/Engine/Source/Programs/UnrealBuildAccelerator/Visualizer/Private/UbaTraceReader.h
--- a/Engine/Source/Programs/UnrealBuildAccelerator/Visualizer/Private/UbaTraceReader.h
+++ b/Engine/Source/Programs/UnrealBuildAccelerator/Visualizer/Private/UbaTraceReader.h
@@ -88,6 +88,7 @@ namespace uba
 			Vector<Processor> processors;
 			Vector<SessionUpdate> updates;
 			Vector<TString> summary;
+			UnorderedMap<CasKey, u32> fetchedFilesActive;
 			Vector<FileTransfer> fetchedFiles;
 			Vector<FileTransfer> storedFiles;
 			TString notification;
diff --git a/Engine/Source/Programs/UnrealBuildAccelerator/Visualizer/Private/UbaVisualizer.cpp b/Engine/Source/Programs/UnrealBuildAccelerator/Visualizer/Private/UbaVisualizer.cpp
--- a/Engine/Source/Programs/UnrealBuildAccelerator/Visualizer/Private/UbaVisualizer.cpp
+++ b/Engine/Source/Programs/UnrealBuildAccelerator/Visualizer/Private/UbaVisualizer.cpp
@@ -70,6 +70,8 @@ namespace uba
 		#define UBA_VISUALIZER_FLAG(name, defaultValue, desc) config.GetValueAsBool(name, TC(#name));
 		UBA_VISUALIZER_FLAGS2
 		#undef UBA_VISUALIZER_FLAG
+
+		fontSize = Min(fontSize, 30u);
 		return true;
 	}
 
@@ -943,7 +945,7 @@ namespace uba
 
 		if (m_config.showProgress && m_traceView.progressProcessesTotal)
 		{
-			drawIndentedText(ToView(L"Progress"), LogEntryType_Info, 1, false);
+			drawIndentedText(AsView(L"Progress"), LogEntryType_Info, 1, false);
 
 			float progress = float(m_traceView.progressProcessesDone) / float(m_traceView.progressProcessesTotal);
 			u32 width = m_activeFont.height * 18;
@@ -1836,7 +1838,7 @@ namespace uba
 
 					if (reader.GetLeft())
 					{
-						if (process.isRemote)
+						if (process.isRemote || m_traceView.version >= 36)
 							sessionStats.Read(reader, m_traceView.version);
 						storageStats.Read(reader, m_traceView.version);
 						kernelStats.Read(reader, m_traceView.version);
@@ -2147,7 +2149,7 @@ namespace uba
 
 		u64 writeFilesTime = process.writeFilesTime;
 
-		if (!done || process.exitCode != 0 || !m_config.ShowReadWriteColors || (TimeToMs(writeFilesTime, m_traceView.frequency) < 300 && TimeToMs(process.createFilesTime, m_traceView.frequency) < 300))
+		if (!done || process.exitCode != 0 || !m_config.ShowReadWriteColors)
 		{
 			if (writingBitmap)
 				rect.right = 256;
@@ -2164,7 +2166,6 @@ namespace uba
 		double recvPart = (double(ConvertTime(m_traceView, process.createFilesTime)) / duration);
 		if (int headSize = int(recvPart * width))
 		{
-			UBA_ASSERT(headSize > 0);
 			main.left += headSize;
 			RECT r2 = rect;
 			r2.right = r2.left + headSize;
@@ -2176,7 +2177,6 @@ namespace uba
 		double sendPart = (double(ConvertTime(m_traceView, writeFilesTime)) / duration);
 		if (int tailSize = int(sendPart * width))
 		{
-			UBA_ASSERT(tailSize > 0);
 			main.right -= tailSize;
 			RECT r2 = rect;
 			r2.left = r2.right - tailSize;
@@ -2361,11 +2361,11 @@ namespace uba
 
 			int fileWidth = 700;
 
-			auto drawFiles = [&](const wchar_t* fileType, Vector<TraceView::FileTransfer>& files, u64 bytes, u32& maxVisibleFiles)
+			auto drawFiles = [&](const wchar_t* fileType, Vector<TraceView::FileTransfer>& files, u64 count, u64 bytes, u32& maxVisibleFiles)
 				{
 					textRect.left = posX;
 					textRect.right = posX + fileWidth;
-					drawText(L"%ls Files: %u (%s)", fileType, u32(files.size()), BytesToText(bytes));
+					drawText(L"%ls Files: %u (%s)", fileType, u32(count), BytesToText(bytes));
 					u32 fileCount = 0;
 					for (auto rit = files.rbegin(), rend = files.rend(); rit != rend; ++rit)
 					{
@@ -2375,7 +2375,14 @@ namespace uba
 						u64 time = 0;
 						if (file.start < playTime)
 							time = playTime - file.start;
-						drawText(L"%s - %s, (%ls)", file.hint.c_str(), BytesToText(file.size).str, TimeToText(time, true).str);
+
+						if ((((u8*)&file.key)[19] & 4) != 0)
+							drawText(L"%s (proxy) %s", file.hint.c_str(), TimeToText(time, true).str);
+						else if (file.size == 0)
+							drawText(L"%s (calc) %s", file.hint.c_str(), TimeToText(time, true).str);
+						else
+							drawText(L"%s (%s) %s", file.hint.c_str(), BytesToText(file.size).str, TimeToText(time, true).str);
+
 						if (fileCount++ > 5)
 							break;
 					}
@@ -2383,13 +2390,18 @@ namespace uba
 					maxVisibleFiles = Max(maxVisibleFiles, fileCount);
 				};
 
+			Vector<TraceView::FileTransfer> fetchedFiles;
+			for (auto& kv : session.fetchedFilesActive)
+				fetchedFiles.push_back(session.fetchedFiles[kv.second]);
+			std::sort(fetchedFiles.begin(), fetchedFiles.end(), [](const TraceView::FileTransfer& a, const TraceView::FileTransfer& b) { return a.start > b.start; });
+
 			posY = startPosY;
 			posX += 150;
-			drawFiles(L"Fetched", session.fetchedFiles, session.fetchedFilesBytes, session.maxVisibleFiles);
+			drawFiles(L"Fetched", fetchedFiles, session.fetchedFiles.size(), session.fetchedFilesBytes, session.maxVisibleFiles);
 			int posY2 = posY;
 			posY = startPosY;
 			posX += fileWidth;
-			drawFiles(L"Stored", session.storedFiles, session.storedFilesBytes, session.maxVisibleFiles);
+			drawFiles(L"Stored", session.storedFiles, session.storedFiles.size(), session.storedFilesBytes, session.maxVisibleFiles);
 			posY = Max(posY, Max(posY1, posY2));
 		}
 		else
@@ -2814,7 +2826,7 @@ namespace uba
 			processStats.Read(reader, m_traceView.version);
 			if (reader.GetLeft())
 			{
-				if (process.isRemote)
+				if (process.isRemote || m_traceView.version >= 36)
 					sessionStats.Read(reader, m_traceView.version);
 				storageStats.Read(reader, m_traceView.version);
 				kernelStats.Read(reader, m_traceView.version);
