[UBA]
* Rewrote code in compact tables code to make static analyzer happy
* Fixed posix debug compile errors
* Extended unit tests for compact tables

diff --git a/Engine/Source/Programs/UnrealBuildAccelerator/Common/Private/UbaCompactTables.cpp b/Engine/Source/Programs/UnrealBuildAccelerator/Common/Private/UbaCompactTables.cpp
--- a/Engine/Source/Programs/UnrealBuildAccelerator/Common/Private/UbaCompactTables.cpp
+++ b/Engine/Source/Programs/UnrealBuildAccelerator/Common/Private/UbaCompactTables.cpp
@@ -336,32 +336,23 @@ namespace uba
 			if (newCount <= CasKeyArrayMaxSize)
 			{
 				u32 capacity = NextPow2(index);
+				u32 newCapacity = NextPow2(newCount);
 
-				if (newCount <= capacity)
+				if (capacity == newCapacity)
 				{
 					value.array[index].stringOffset = u32(stringOffset);
 					value.count = newCount;
 					return &value.array[index].casKeyOffset;
 				}
-				u32 newCapacity = NextPow2(newCount);
+
 				auto newArray = new StringAndKey[newCapacity];
 				memcpy(newArray, value.array, index*sizeof(StringAndKey));
 
-				#if PLATFORM_WINDOWS
-				#pragma warning(push)
-				#pragma warning(disable:6385) // analyzer claims that newArray can have buffer overrun.. but can't see how that can happen
-				#pragma warning(disable:6386) // analyzer claims that newArray can have buffer overrun.. but can't see how that can happen
-				#endif
-
 				newArray[index].stringOffset = u32(stringOffset);
 				delete[] value.array;
 				value.array = newArray;
 				value.count = newCount;
 				return &newArray[index].casKeyOffset;
-
-				#if PLATFORM_WINDOWS
-				#pragma warning(pop)
-				#endif
 			}
 
 			auto& lookup = *new UnorderedMap<u32, u32>();
diff --git a/Engine/Source/Programs/UnrealBuildAccelerator/Detours/Private/UbaDetoursFileMappingTable.cpp b/Engine/Source/Programs/UnrealBuildAccelerator/Detours/Private/UbaDetoursFileMappingTable.cpp
--- a/Engine/Source/Programs/UnrealBuildAccelerator/Detours/Private/UbaDetoursFileMappingTable.cpp
+++ b/Engine/Source/Programs/UnrealBuildAccelerator/Detours/Private/UbaDetoursFileMappingTable.cpp
@@ -309,7 +309,7 @@ namespace uba
 			{
 				reader.ReadString(tempBuf.Clear());
 				#if UBA_DEBUG
-				DEBUG_LOG_DETOURED(L"Rpc_GetFullFileName", L"%s -> %s", path, tempBuf.data);
+				DEBUG_LOG_DETOURED(TC("Rpc_GetFullFileName"), TC("%s -> %s"), path, tempBuf.data);
 				#endif
 			}
 			else
@@ -317,13 +317,13 @@ namespace uba
 				#if UBA_DEBUG
 				StringBuffer<> virtualName;
 				reader.ReadString(virtualName);
-				DEBUG_LOG_DETOURED(L"Rpc_GetFullFileName", L"%s -> %s (%s)", path, tempBuf.data, virtualName.data);
+				DEBUG_LOG_DETOURED(TC("Rpc_GetFullFileName"), TC("%s -> %s (%s)"), path, tempBuf.data, virtualName.data);
 				#else
 				reader.SkipString();
 				#endif
 			}
 			mappedFileTableSize = reader.ReadU32();
-			DEBUG_LOG_PIPE(L"GetFileName", L"(%ls)", tempBuf);
+			DEBUG_LOG_PIPE(TC("GetFileName"), TC("(%ls)"), tempBuf);
 		}
 
 		g_mappedFileTable.Parse(mappedFileTableSize);
diff --git a/Engine/Source/Programs/UnrealBuildAccelerator/Test/Private/UbaTestCache.h b/Engine/Source/Programs/UnrealBuildAccelerator/Test/Private/UbaTestCache.h
--- a/Engine/Source/Programs/UnrealBuildAccelerator/Test/Private/UbaTestCache.h
+++ b/Engine/Source/Programs/UnrealBuildAccelerator/Test/Private/UbaTestCache.h
@@ -255,12 +255,12 @@ namespace uba
 		if (table2.AddNoLock(CasKeyZero, 2) != offset2)
 			return false;
 
-		CompactCasKeyTable table3(table.GetKeyCount());
-		table3.AddNoLock(CasKeyZero, 0);
-		table3.AddNoLock(CasKeyZero, 1);
-		table3.AddNoLock(CasKeyZero, 2);
-		table3.AddNoLock(CasKeyZero, 3);
-		table3.AddNoLock(CasKeyZero, 4);
+		for (u32 i=0;i!=32;++i)
+		{
+			CompactCasKeyTable table3(table.GetKeyCount());
+			for (u32 j=0;j!=i;++j)
+				table3.AddNoLock(CasKeyZero, j);
+		}
 		return true;
 	}
 
