[UBA]
* Added cache server option "-fullmaintenance" which will trigger a full maintenance of the database and tables

diff --git a/Engine/Source/Programs/UnrealBuildAccelerator/CacheService/Private/UbaCacheService.cpp b/Engine/Source/Programs/UnrealBuildAccelerator/CacheService/Private/UbaCacheService.cpp
--- a/Engine/Source/Programs/UnrealBuildAccelerator/CacheService/Private/UbaCacheService.cpp
+++ b/Engine/Source/Programs/UnrealBuildAccelerator/CacheService/Private/UbaCacheService.cpp
@@ -49,6 +49,7 @@ namespace uba
 		logger.Info(TC("  -capacity=<gigaby>      Capacity of local store. Defaults to %u gigabytes"), DefaultCapacityGb);
 		logger.Info(TC("  -expiration=<seconds>   Time until unused cache entries get deleted. Defaults to %s (%u seconds)"), TimeToText(MsToTime(DefaultExpiration*1000)).str, DefaultExpiration);
 		logger.Info(TC("  -http=<port>            If set, a http server will be started and listen on <port>"));
+		logger.Info(TC("  -fullmaintenance        Force a full maintenance"));
 		logger.Info(TC(""));
 		return -1;
 	}
@@ -120,6 +121,7 @@ namespace uba
 		u16 httpPort = 0;
 		bool quiet = false;
 		bool storeCompressed = true;
+		bool fullMaintenance = false;
 		u32 expirationTimeSeconds = DefaultExpiration;
 
 		for (int i=1; i!=argc; ++i)
@@ -174,6 +176,10 @@ namespace uba
 				if (!value.Parse(httpPort))
 					httpPort = 80;
 			}
+			else if (name.Equals(TC("-fullmaintenance")))
+			{
+				fullMaintenance = true;
+			}
 			else if (name.Equals(TC("-?")))
 			{
 				return PrintHelp(TC(""));
@@ -256,6 +262,9 @@ namespace uba
 		if (!cacheServer.Load())
 			return -1;
 
+		if (fullMaintenance)
+			cacheServer.SetForceFullMaintenance();
+
 		if (!cacheServer.RunMaintenance(true, ShouldExit))
 			return -1;
 
diff --git a/Engine/Source/Programs/UnrealBuildAccelerator/Common/Private/UbaCacheServer.cpp b/Engine/Source/Programs/UnrealBuildAccelerator/Common/Private/UbaCacheServer.cpp
--- a/Engine/Source/Programs/UnrealBuildAccelerator/Common/Private/UbaCacheServer.cpp
+++ b/Engine/Source/Programs/UnrealBuildAccelerator/Common/Private/UbaCacheServer.cpp
@@ -541,6 +541,11 @@ namespace uba
 		return success;
 	}
 
+	void CacheServer::SetForceFullMaintenance()
+	{
+		m_forceAllSteps = true;
+	}
+
 	bool CacheServer::RunMaintenance(bool force, const Function<bool()>& shouldExit)
 	{
 		if (m_addsSinceMaintenance == 0 && !force && !m_bucketIsOverflowing)
diff --git a/Engine/Source/Programs/UnrealBuildAccelerator/Common/Public/UbaCacheServer.h b/Engine/Source/Programs/UnrealBuildAccelerator/Common/Public/UbaCacheServer.h
--- a/Engine/Source/Programs/UnrealBuildAccelerator/Common/Public/UbaCacheServer.h
+++ b/Engine/Source/Programs/UnrealBuildAccelerator/Common/Public/UbaCacheServer.h
@@ -52,6 +52,8 @@ namespace uba
 		bool Load();
 		bool Save();
 
+		void SetForceFullMaintenance();
+
 		bool RunMaintenance(bool force, const Function<bool()>& shouldExit);
 
 		bool ShouldShutdown();
