[UBA] Ignore call to ExitProcess if the process hasn't even started to avoid null pointer exception.

#rnx

diff --git a/Engine/Source/Programs/UnrealBuildAccelerator/Common/Private/UbaScheduler.cpp b/Engine/Source/Programs/UnrealBuildAccelerator/Common/Private/UbaScheduler.cpp
--- a/Engine/Source/Programs/UnrealBuildAccelerator/Common/Private/UbaScheduler.cpp
+++ b/Engine/Source/Programs/UnrealBuildAccelerator/Common/Private/UbaScheduler.cpp
@@ -552,10 +552,13 @@ namespace uba
 
 	void Scheduler::ExitProcess(ExitProcessInfo& info, Process& process, u32 exitCode, bool fromCache)
 	{
+		auto si = info.startInfo;
+		if (!si)
+			return;
+
 		ProcessHandle ph;
 		ph.m_process = &process;
 
-		auto si = info.startInfo;
 		if (auto func = si->exitedFunc)
 			func(si->userData, ph);
 
