[UBA]
* Fixed race condition bug in CreateDirectory where if two threads created a dir at the same time the failing one (returning already exists) might return a directory table size that did not contain the creation
* Added some early outs on errors to prevent crashes

diff --git a/Engine/Source/Programs/UnrealBuildAccelerator/Common/Private/UbaSession.cpp b/Engine/Source/Programs/UnrealBuildAccelerator/Common/Private/UbaSession.cpp
--- a/Engine/Source/Programs/UnrealBuildAccelerator/Common/Private/UbaSession.cpp
+++ b/Engine/Source/Programs/UnrealBuildAccelerator/Common/Private/UbaSession.cpp
@@ -551,6 +551,8 @@ namespace uba
 		else
 		{
 			auto mappedView = m_fileMappingBuffer.AllocAndMapView(MappedView_Transient, size, alignment, fileName, false);
+			if (!mappedView.memory)
+				return false;
 			auto unmapGuard = MakeGuard([&](){ m_fileMappingBuffer.UnmapView(mappedView, fileName); });
 
 			if (isCompressed)
@@ -640,6 +642,8 @@ namespace uba
 			else
 			{
 				auto mappedViewWrite = m_fileMappingBuffer.AllocAndMapView(MappedView_Transient, size, alignment, fileName);
+				if (!mappedViewWrite.memory)
+					return false;
 				auto unmapGuard = MakeGuard([&](){ m_fileMappingBuffer.UnmapView(mappedViewWrite, fileName); });
 
 				if (!m_storage.DecompressMemoryToMemory(readMemory, mappedViewWrite.memory, size, fileName, TC("TransientMapping")))
@@ -2136,21 +2140,24 @@ namespace uba
 	bool Session::CreateDirectory(CreateDirectoryResponse& out, const CreateDirectoryMessage& msg)
 	{
 		out.result = uba::CreateDirectoryW(msg.name.data);
-		if (out.result)
+
+		StringKey dirKey;
+		const tchar* lastSlash;
+		StringBuffer<> dirName;
+		if (!GetDirKey(dirKey, dirName, lastSlash, msg.name))
+			return true;
+
+		if (!out.result)
+			out.errorCode = GetLastError();
+
+		// There is a chance that another thread just created the directory and we can't return directoryTableSize until we know it is written
+		// So let's both success and already exists add entries
+		if (out.result || out.errorCode == ERROR_ALREADY_EXISTS)
 		{
-			StringKey dirKey;
-			const tchar* lastSlash;
-			StringBuffer<> dirName;
-			if (!GetDirKey(dirKey, dirName, lastSlash, msg.name))
-				return true;
 			// Both these functions need to be called. otherwise we can get created directories that does not end up in directory table
 			RegisterCreateFileForWrite(msg.nameKey, msg.name, true);
 			WriteDirectoryEntries(dirKey, dirName.data);
 		}
-		else
-		{
-			out.errorCode = GetLastError();
-		}
 
 		out.directoryTableSize = GetDirectoryTableSize();
 		return true;
@@ -2163,6 +2170,8 @@ namespace uba
 			RegisterDeleteFile(msg.nameKey, msg.name);
 		else
 			out.errorCode = GetLastError();
+		// This has a race condition. If same directory is removed at the same time the failing one
+		// might send back a directoryTableSize that does not include the delete
 		out.directoryTableSize = GetDirectoryTableSize();
 		return true;
 	}
