[UBA]
* Fixed so shutdown socket is shutting down the right handle if connect fails early

diff --git a/Engine/Source/Programs/UnrealBuildAccelerator/Common/Private/UbaNetworkBackendTcp.cpp b/Engine/Source/Programs/UnrealBuildAccelerator/Common/Private/UbaNetworkBackendTcp.cpp
--- a/Engine/Source/Programs/UnrealBuildAccelerator/Common/Private/UbaNetworkBackendTcp.cpp
+++ b/Engine/Source/Programs/UnrealBuildAccelerator/Common/Private/UbaNetworkBackendTcp.cpp
@@ -690,7 +690,7 @@ namespace uba
 
 		if (!connectedFunc(&conn, remoteSocketAddr, timedOut))
 		{
-			ShutdownSocket(logger, socketFd, TC("Connect"));
+			ShutdownSocket(logger, conn.socket, TC("Connect"));
 			conn.ready.Set();
 			conn.recvThread.Wait();
 			SCOPED_WRITE_LOCK(m_connectionsLock, lock2);
