[UBA]
* Fixed clipped processor numbers in visualizer at certain font sizes

diff --git a/Engine/Source/Programs/UnrealBuildAccelerator/Visualizer/Private/UbaVisualizer.cpp b/Engine/Source/Programs/UnrealBuildAccelerator/Visualizer/Private/UbaVisualizer.cpp
--- a/Engine/Source/Programs/UnrealBuildAccelerator/Visualizer/Private/UbaVisualizer.cpp
+++ b/Engine/Source/Programs/UnrealBuildAccelerator/Visualizer/Private/UbaVisualizer.cpp
@@ -806,7 +806,7 @@ namespace uba
 		m_zoomValue = 1.0f + float(m_boxHeight) / 30.0f;
 		int fontHeight = Max(m_boxHeight - 2, 1);
 		UpdateFont(m_processFont, fontHeight, false);
-		m_progressRectLeft = int(5 + float(m_processFont.height) * 1.8f);
+		m_progressRectLeft = int(13 + float(m_processFont.height) * 1.5f);
 		DirtyBitmaps(true);
 	}
 
@@ -1402,7 +1402,7 @@ namespace uba
 						{
 							RECT rect;
 							rect.left = 5;
-							rect.right = progressRect.left - 5;
+							rect.right = progressRect.left - 2;
 							rect.top = posY;
 							rect.bottom = rectBottom;
 
