[UBA]

* Fixed GetFinalPath assert from GetFinalPathNameByHandleW for root directory handles such as c:\users (thank you PJ Kack for this fix)

diff --git a/Engine/Source/Programs/UnrealBuildAccelerator/Core/Public/UbaDirectoryTable.h b/Engine/Source/Programs/UnrealBuildAccelerator/Core/Public/UbaDirectoryTable.h
--- a/Engine/Source/Programs/UnrealBuildAccelerator/Core/Public/UbaDirectoryTable.h
+++ b/Engine/Source/Programs/UnrealBuildAccelerator/Core/Public/UbaDirectoryTable.h
@@ -254,7 +254,8 @@ namespace uba
 			const tchar* prevSlash = TStrchr(path+3, PathSeparator);
 			if (!prevSlash)
 			{
-				UBA_ASSERTF(false, TC("GetFinalPath got path \"%s\" which has no backslash"), path);
+				// return root directory as-is
+				out.Append(path);
 				return;
 			}
 			out.Append(path, u64(prevSlash - path));
@@ -363,4 +364,4 @@ namespace uba
 		u32 m_memorySize = 0;
 	};
 
-}
\ No newline at end of file
+}
