[UBA]
* Added option to not use send-only-one-big-file-at-the-time lock

diff --git a/Engine/Source/Programs/UnrealBuildAccelerator/Common/Private/UbaStorageUtils.cpp b/Engine/Source/Programs/UnrealBuildAccelerator/Common/Private/UbaStorageUtils.cpp
--- a/Engine/Source/Programs/UnrealBuildAccelerator/Common/Private/UbaStorageUtils.cpp
+++ b/Engine/Source/Programs/UnrealBuildAccelerator/Common/Private/UbaStorageUtils.cpp
@@ -362,7 +362,7 @@ namespace uba
 
 			bool isDone = sendLeft == 0;
 
-			if (isFirst && !isDone)
+			if (isFirst && !isDone && m_sendOneBigFileAtTheTime)
 			{
 				m_sendOneAtTheTimeLock.EnterWrite();
 				hasSendOneAtTheTimeLock = true;
diff --git a/Engine/Source/Programs/UnrealBuildAccelerator/Common/Public/UbaStorageUtils.h b/Engine/Source/Programs/UnrealBuildAccelerator/Common/Public/UbaStorageUtils.h
--- a/Engine/Source/Programs/UnrealBuildAccelerator/Common/Public/UbaStorageUtils.h
+++ b/Engine/Source/Programs/UnrealBuildAccelerator/Common/Public/UbaStorageUtils.h
@@ -30,6 +30,7 @@ namespace uba
 		ReaderWriterLock& m_sendOneAtTheTimeLock;
 		u8 m_casCompressor = DefaultCompressor;
 		u8 m_casCompressionLevel = DefaultCompressionLevel;
+		bool m_sendOneBigFileAtTheTime = true;
 		u64 m_bytesSent = 0;
 	};
 
