[UBA]
* Fixed some bad logging formatting

diff --git a/Engine/Source/Programs/UnrealBuildAccelerator/Common/Private/UbaNetworkServer.cpp b/Engine/Source/Programs/UnrealBuildAccelerator/Common/Private/UbaNetworkServer.cpp
--- a/Engine/Source/Programs/UnrealBuildAccelerator/Common/Private/UbaNetworkServer.cpp
+++ b/Engine/Source/Programs/UnrealBuildAccelerator/Common/Private/UbaNetworkServer.cpp
@@ -502,14 +502,14 @@ namespace uba
 
 		if (!rec.func)
 		{
-			server.m_logger.Error(TC("WORKER FUNCTION NOT FOUND. id: %u, serviceid: %u type: %s, client: %s"), context.id, context.serviceId, rec.toString(context.messageType), GuidToString(connection.m_client->uid).str);
+			server.m_logger.Error(TC("WORKER FUNCTION NOT FOUND. id: %u, serviceid: %u type: %s, client: %s"), context.id, context.serviceId, rec.toString(context.messageType).data, GuidToString(connection.m_client->uid).str);
 			connection.SetShouldDisconnect();
 			size = ErrorSize;
 		}
 		else if (!rec.func({&connection}, mi, reader, writer))
 		{
 			if (connection.SetShouldDisconnect())
-				server.m_logger.Error(TC("WORKER FUNCTION FAILED. id: %u, serviceid: %u type: %s, client: %s"), context.id, context.serviceId, rec.toString(context.messageType), GuidToString(connection.m_client->uid).str);
+				server.m_logger.Error(TC("WORKER FUNCTION FAILED. id: %u, serviceid: %u type: %s, client: %s"), context.id, context.serviceId, rec.toString(context.messageType).data, GuidToString(connection.m_client->uid).str);
 			size = ErrorSize;
 		}
 		else
