[UBA]
* Fixed bug where compressed files that was already compressed were sent over network. Client could end up in infinite loop

diff --git a/Engine/Source/Programs/UnrealBuildAccelerator/Common/Private/UbaStorageClient.cpp b/Engine/Source/Programs/UnrealBuildAccelerator/Common/Private/UbaStorageClient.cpp
--- a/Engine/Source/Programs/UnrealBuildAccelerator/Common/Private/UbaStorageClient.cpp
+++ b/Engine/Source/Programs/UnrealBuildAccelerator/Common/Private/UbaStorageClient.cpp
@@ -553,6 +553,10 @@ namespace uba
 
 				if (actualSize)
 				{
+					// There is a risk the compressed file we download is larger than half buffer
+					u8* extraBuffer = nullptr;
+					auto ebg = MakeGuard([&]() { delete[] extraBuffer; });
+
 					u64 leftCompressed = fileSize - responseSize;
 					u32 readIndex = 0;
 					bool tryAgain = false;
@@ -570,7 +574,19 @@ namespace uba
 							{
 								if (fetchId == u16(~0))
 									return m_logger.Error(TC("Cas content error (2). Server believes %s was only one segment but client sees more. UncompressedSize: %llu LeftUncompressed: %llu Size: %llu Left to read: %llu ResponseSize: %u. (%s)"), hint, actualSize, leftUncompressed, fileSize, left, responseSize, casFile.data);
-								if (!SendBatchMessages(m_logger, *client, fetchId, readPosition, maxReadSize - u32(readPosition - readBuffer), leftCompressed, sizeOfFirstMessage, readIndex, responseSize))
+
+								u64 capacity = maxReadSize - u64(readPosition - readBuffer);
+								u64 writeCapacity = capacity;
+								u8* writeDest = readPosition;
+								if (capacity < sizeOfFirstMessage) // capacity left is less than the last message
+								{
+									UBA_ASSERT(!extraBuffer);
+									extraBuffer = new u8[sizeOfFirstMessage];
+									writeDest = extraBuffer;
+									writeCapacity = sizeOfFirstMessage;
+								}
+
+								if (!SendBatchMessages(m_logger, *client, fetchId, writeDest, writeCapacity, leftCompressed, sizeOfFirstMessage, readIndex, responseSize))
 								{
 									if (proxy)
 									{
@@ -579,6 +595,14 @@ namespace uba
 									}
 									return m_logger.Error(TC("Failed to send batched messages to server (%s)"), casFile.data);
 								}
+
+								if (extraBuffer)
+								{
+									UBA_ASSERT(responseSize > capacity);
+									memcpy(readPosition, extraBuffer, left);
+									memmove(extraBuffer, extraBuffer + left, u32(responseSize - left));
+								}
+
 								leftCompressed -= responseSize;
 							}
 							else
@@ -589,9 +613,8 @@ namespace uba
 							if (isFirstInBlock)
 							{
 								if ((readPosition - readBuffer) + responseSize < sizeof(u32) * 2)
-								{
 									return m_logger.Error(TC("Received less than minimum amount of data. Most likely corrupt cas file %s (Available: %u UncompressedSize: %llu LeftUncompressed: %llu)"), casFile.data, u32(readPosition - readBuffer), actualSize, leftUncompressed);
-								}
+
 								isFirstInBlock = false;
 								u32* blockSize = (u32*)readBuffer;
 								compressedSize = blockSize[0];
@@ -660,12 +683,23 @@ namespace uba
 							leftUncompressed -= uncompressedSize;
 						}
 
-						// Move overflow back to the beginning of the buffer and start the next block (if there is one)
 						readBuffer = slot;
 						maxReadSize = BufferSlotHalfSize;
-						UBA_ASSERTF(readPosition - overflow >= readBuffer, TC("ReadPosition - overflow is before beginning of buffer (overflow: %u) for file %s"), overflow, hint);
-						UBA_ASSERTF(readPosition <= readBuffer + BufferSlotHalfSize, TC("ReadPosition is outside readBuffer size (pos: %llu, overflow: %u) for file %s"), readPosition - readBuffer, overflow, hint);
-						memmove(readBuffer, readPosition - overflow, overflow);
+
+						if (extraBuffer)
+						{
+							memcpy(readBuffer, extraBuffer, overflow);
+							delete[] extraBuffer;
+							extraBuffer = nullptr;
+						}
+						else
+						{
+							// Move overflow back to the beginning of the buffer and start the next block (if there is one)
+							UBA_ASSERTF(readPosition - overflow >= readBuffer, TC("ReadPosition - overflow is before beginning of buffer (overflow: %u) for file %s"), overflow, hint);
+							UBA_ASSERTF(readPosition <= readBuffer + BufferSlotHalfSize, TC("ReadPosition is outside readBuffer size (pos: %llu, overflow: %u) for file %s"), readPosition - readBuffer, overflow, hint);
+							memmove(readBuffer, readPosition - overflow, overflow);
+						}
+
 						readPosition = readBuffer + overflow;
 						if (overflow)
 						{
