[UBA]
* Fixed code path where cas files can be sent to server uncompressed

diff --git a/Engine/Source/Programs/UnrealBuildAccelerator/Common/Private/UbaStorage.cpp b/Engine/Source/Programs/UnrealBuildAccelerator/Common/Private/UbaStorage.cpp
--- a/Engine/Source/Programs/UnrealBuildAccelerator/Common/Private/UbaStorage.cpp
+++ b/Engine/Source/Programs/UnrealBuildAccelerator/Common/Private/UbaStorage.cpp
@@ -541,11 +541,12 @@ namespace uba
 			if (!WriteCompressed(out, fileName, casFile))
 				return false;
 		}
+		else if (fileIsCompressed && !storeCompressed)
+		{
+			return m_logger.Error(TC("Write compressed file to uncompressed store not implemented"));
+		}
 		else
 		{
-			if (!storeCompressed)
-				return m_logger.Error(TC("Write compressed file to uncompressed store not implemented"));
-
 			FileHandle readHandle;
 			if (!OpenFileSequentialRead(m_logger, fileName, readHandle))
 				return m_logger.Error(TC("Failed to open file %s for read (%s)"), fileName, LastErrorToText().data);
@@ -751,7 +752,7 @@ namespace uba
 				TimerScope ts(stats.decompressToMem);
 				OO_SINTa decompLen = OodleLZ_Decompress(readPos, (OO_SINTa)compressedBlockSize, writePos, (OO_SINTa)decompressedBlockSize); (void)decompLen;
 				if (decompLen != decompressedBlockSize)
-					return m_logger.Error(TC("Expecting to be able to decompress to %u bytes but got %llu (%s -> %s)"), decompressedBlockSize, decompLen, readHint, writeHint);
+					return m_logger.Error(TC("Expecting to be able to decompress %u to %u bytes at pos %llu but got %llu (%s -> %s)"), compressedBlockSize, decompressedBlockSize, decompressedSize - left, decompLen, readHint, writeHint);
 				writePos += decompressedBlockSize;
 				readPos += compressedBlockSize;
 				left -= decompressedBlockSize;
@@ -2121,7 +2122,7 @@ namespace uba
 		return true;
 	}
 
-	bool StorageImpl::StoreCasFile(CasKey& out, StringKey fileNameKey, const tchar* fileName, FileMappingHandle mappingHandle, u64 mappingOffset, u64 fileSize, const tchar* hint, bool deferCreation, bool keepMappingInMemory)
+	bool StorageImpl::StoreCasFile(CasKey& out, StringKey fileNameKey, const tchar* fileName, FileMappingHandle mappingHandle, u64 mappingOffset, u64 fileSize, const tchar* hint, bool deferCreation, bool keepMappingInMemory, bool storeCompressed)
 	{
 		UBA_ASSERT(false);
 		return false;
@@ -2534,6 +2535,34 @@ namespace uba
 				return true;
 			}
 
+			if (casEntry->mappingHandle.IsValid())
+			{
+				SCOPED_WRITE_LOCK(fileEntry.lock, entryLock);
+				fileEntry.verified = false;
+
+				MappedView mappedView = m_casDataBuffer.MapView(casEntry->mappingHandle, casEntry->mappingOffset, casEntry->mappingSize, TC(""));
+				auto mapViewGuard = MakeGuard([&](){ m_casDataBuffer.UnmapView(mappedView, destination); });
+
+				FileAccessor destinationFile(m_logger, destination);
+				if (!destinationFile.CreateMemoryWrite(false, fileAttributes, mappedView.size, m_tempPath.data))
+					return false;
+				u8* writePos = destinationFile.GetData();
+				TimerScope cts(stats.memoryCopy);
+				MapMemoryCopy(writePos, mappedView.memory, mappedView.size);
+				u64 lastWriteTime = 0;
+				if (!destinationFile.Close(&lastWriteTime))
+					return false;
+				UBA_ASSERT(lastWriteTime);
+				if (lastWriteTime)
+				{
+					fileEntry.casKey = casKey;
+					fileEntry.lastWritten = lastWriteTime;
+					fileEntry.size = mappedView.size;
+					fileEntry.verified = true;
+				}
+				return true;
+			}
+
 			StringBuffer<> casFile;
 			#if !UBA_USE_SPARSEFILE
 			if (!GetCasFileName(casFile, actualKey))
diff --git a/Engine/Source/Programs/UnrealBuildAccelerator/Common/Private/UbaStorageClient.cpp b/Engine/Source/Programs/UnrealBuildAccelerator/Common/Private/UbaStorageClient.cpp
--- a/Engine/Source/Programs/UnrealBuildAccelerator/Common/Private/UbaStorageClient.cpp
+++ b/Engine/Source/Programs/UnrealBuildAccelerator/Common/Private/UbaStorageClient.cpp
@@ -34,6 +34,7 @@ namespace uba
 		const ConfigTable& table = *tablePtr;
 		table.GetValueAsBool(sendCompressed, TC("SendCompressed"));
 		table.GetValueAsBool(allowProxy, TC("AllowProxy"));
+		table.GetValueAsBool(sendOneBigFileAtTheTime, TC("SendOneBigAtTheTime"));
 	}
 
 	StorageClient::StorageClient(const StorageClientCreateInfo& info)
@@ -41,6 +42,7 @@ namespace uba
 	,	m_client(info.client)
 	,	m_sendCompressed(info.sendCompressed)
 	,	m_allowProxy(info.allowProxy)
+	,	m_sendOneBigFileAtTheTime(info.sendOneBigFileAtTheTime)
 	,	m_zone(info.zone)
 	,	m_getProxyBackendCallback(info.getProxyBackendCallback)
 	,	m_getProxyBackendUserData(info.getProxyBackendUserData)
@@ -789,7 +791,7 @@ namespace uba
 		return StorageImpl::HasCasFile(casKey, out);
 	}
 
-	bool StorageClient::StoreCasFile(CasKey& out, StringKey fileNameKey, const tchar* fileName, FileMappingHandle mappingHandle, u64 mappingOffset, u64 fileSize, const tchar* hint, bool deferCreation, bool keepMappingInMemory)
+	bool StorageClient::StoreCasFile(CasKey& out, StringKey fileNameKey, const tchar* fileName, FileMappingHandle mappingHandle, u64 mappingOffset, u64 fileSize, const tchar* hint, bool deferCreation, bool keepMappingInMemory, bool storeCompressed)
 	{
 		NetworkClient& client = m_client; // Don't use proxy
 
@@ -822,7 +824,6 @@ namespace uba
 		//UBA_ASSERT(wcsstr(fileName, TC("\\output\\"))); // Just make sure our assumption is correct :)
 		if (isTemporaryStorage)
 		{
-			bool storeCompressed = true;
 			casKey = CalculateCasKey(fileMem, fileSize, storeCompressed);
 			if (casKey == CasKeyZero)
 				return false;
@@ -872,8 +873,19 @@ namespace uba
 			return false;
 		if (!reader.ReadBool())
 		{
-			if (!SendFile(casKey, fileName, fileMem, fileSize, hint))
-				return false;
+			if (storeCompressed)
+			{
+				FileSender sender { m_logger, m_client, m_bufferSlots, Stats(), m_sendOneAtTheTimeLock, m_casCompressor, m_casCompressionLevel };
+				sender.m_sendOneBigFileAtTheTime = m_sendOneBigFileAtTheTime;
+				if (!sender.SendFileCompressed(casKey, fileName, fileMem, fileSize, hint))
+					return false;
+			}
+			else
+			{
+				TimerScope ts(Stats().sendCas);
+				if (!SendFile(m_logger, m_client, casKey, fileMem, fileSize, hint))
+					return false;
+			}
 		}
 		else
 		{
@@ -902,12 +914,6 @@ namespace uba
 			m_proxyClient->client.PrintSummary(logger);
 	}
 
-	bool StorageClient::SendFile(const CasKey& casKey, const tchar* fileName, u8* sourceMem, u64 sourceSize, const tchar* hint)
-	{
-		FileSender sender { m_logger, m_client, m_bufferSlots, Stats(), m_sendOneAtTheTimeLock, m_casCompressor, m_casCompressionLevel };
-		return sender.SendFileCompressed(casKey, fileName, sourceMem, sourceSize, hint);
-	}
-
 	bool StorageClient::PopulateCasFromDirsRecursive(const tchar* dir, WorkManager& workManager, UnorderedSet<u64>& seenIds, ReaderWriterLock& seenIdsLock, const Function<bool()>& shouldExit)
 	{
 		if (shouldExit && shouldExit())
diff --git a/Engine/Source/Programs/UnrealBuildAccelerator/Common/Private/UbaStorageServer.cpp b/Engine/Source/Programs/UnrealBuildAccelerator/Common/Private/UbaStorageServer.cpp
--- a/Engine/Source/Programs/UnrealBuildAccelerator/Common/Private/UbaStorageServer.cpp
+++ b/Engine/Source/Programs/UnrealBuildAccelerator/Common/Private/UbaStorageServer.cpp
@@ -84,8 +84,8 @@ namespace uba
 		out.size = InvalidValue;
 
 		CasKey actualKey = casKey;
-		if (m_storeCompressed)
-			actualKey = AsCompressed(casKey, true);
+		//if (m_storeCompressed)
+		//	actualKey = AsCompressed(casKey, true);
 
 		SCOPED_WRITE_LOCK(m_waitEntriesLock, waitLock);
 		WaitEntry& waitEntry = m_waitEntries[actualKey];
@@ -212,13 +212,12 @@ namespace uba
 		}
 	}
 
-	bool StorageServer::StoreCasFile(CasKey& out, StringKey fileNameKey, const tchar* fileName, FileMappingHandle mappingHandle, u64 mappingOffset, u64 fileSize, const tchar* hint, bool deferCreation, bool keepMappingInMemory)
+	bool StorageServer::StoreCasFile(CasKey& out, StringKey fileNameKey, const tchar* fileName, FileMappingHandle mappingHandle, u64 mappingOffset, u64 fileSize, const tchar* hint, bool deferCreation, bool keepMappingInMemory, bool storeCompressed)
 	{
 		u8* fileMem = MapViewOfFile(mappingHandle, FILE_MAP_READ, mappingOffset, fileSize);
 		UBA_ASSERT(fileMem);
 		auto memClose = MakeGuard([&](){ UnmapViewOfFile(fileMem, fileSize, fileName); });
 
-		bool storeCompressed = true;
 		out = CalculateCasKey(fileMem, fileSize, storeCompressed);
 		if (out == CasKeyZero)
 			return false;
@@ -693,9 +692,13 @@ namespace uba
 				UBA_ASSERT(fetchIndex);
 				const u8* pos = fetch.memoryPos + (fetchIndex-1) * writer.GetCapacityLeft();
 				u64 toWrite = writer.GetCapacityLeft();
-				if ((pos - fetch.memoryBegin) + toWrite > fetch.mappedView.size)
-					toWrite = fetch.mappedView.size - (pos - fetch.memoryBegin);
-
+				u64 readOffset = pos - fetch.memoryBegin;
+				if (readOffset + toWrite > fetch.mappedView.size)
+				{
+					if (readOffset >= fetch.mappedView.size)
+						return m_logger.Error(TC("Client is asking for more data than file contains (fetch index %u, client id %u uid %s)"), fetchIndex, connectionInfo.GetId(), GuidToString(connectionInfo.GetUid()).str);
+					toWrite = fetch.mappedView.size - readOffset;
+				}
 				memcpy(writer.AllocWrite(toWrite), pos, toWrite);
 
 				bool isDone = fetch.left.fetch_sub(toWrite) == toWrite;
@@ -729,7 +732,7 @@ namespace uba
 			case StorageMessageType_ExistsOnServer:
 			{
 				CasKey casKey = reader.ReadCasKey();
-				UBA_ASSERT(IsCompressed(casKey));
+				//UBA_ASSERT(IsCompressed(casKey));
 				SCOPED_WRITE_LOCK(m_casLookupLock, lookupLock);
 				auto casInsres = m_casLookup.try_emplace(casKey);
 				CasEntry& casEntry = casInsres.first->second;
@@ -784,7 +787,7 @@ namespace uba
 				CasKey casKey = reader.ReadCasKey();
 				u64 fileSize = reader.ReadU64();
 				u64 actualSize = reader.ReadU64();
-				UBA_ASSERT(IsCompressed(casKey));
+				//UBA_ASSERT(IsCompressed(casKey));
 				StringBuffer<> hint;
 				reader.ReadString(hint);
 
diff --git a/Engine/Source/Programs/UnrealBuildAccelerator/Common/Public/UbaStorage.h b/Engine/Source/Programs/UnrealBuildAccelerator/Common/Public/UbaStorage.h
--- a/Engine/Source/Programs/UnrealBuildAccelerator/Common/Public/UbaStorage.h
+++ b/Engine/Source/Programs/UnrealBuildAccelerator/Common/Public/UbaStorage.h
@@ -45,7 +45,7 @@ namespace uba
 		virtual bool InvalidateCachedFileInfo(StringKey fileNameKey) = 0;
 
 		virtual bool StoreCasFile(CasKey& out, const tchar* fileName, const CasKey& casKeyOverride, bool deferCreation, bool fileIsCompressed) = 0;
-		virtual bool StoreCasFile(CasKey& out, StringKey fileNameKey, const tchar* fileName, FileMappingHandle mappingHandle, u64 mappingOffset, u64 fileSize, const tchar* hint, bool deferCreation = false, bool keepMappingInMemory = false) = 0;
+		virtual bool StoreCasFile(CasKey& out, StringKey fileNameKey, const tchar* fileName, FileMappingHandle mappingHandle, u64 mappingOffset, u64 fileSize, const tchar* hint, bool deferCreation = false, bool keepMappingInMemory = false, bool storeCompressed = true) = 0;
 		virtual bool DropCasFile(const CasKey& casKey, bool forceDelete, const tchar* hint) = 0;
 		virtual bool ReportBadCasFile(const CasKey& casKey) = 0;
 		virtual bool CalculateCasKey(CasKey& out, const tchar* fileName) = 0;
@@ -127,7 +127,7 @@ namespace uba
 		virtual bool VerifyAndGetCachedFileInfo(CachedFileInfo& out, StringKey fileNameKey, u64 verifiedLastWriteTime, u64 verifiedSize) override;
 		virtual bool InvalidateCachedFileInfo(StringKey fileNameKey) override;
 		virtual bool StoreCasFile(CasKey& out, const tchar* fileName, const CasKey& casKeyOverride, bool deferCreation, bool fileIsCompressed) override;
-		virtual bool StoreCasFile(CasKey& out, StringKey fileNameKey, const tchar* fileName, FileMappingHandle mappingHandle, u64 mappingOffset, u64 fileSize, const tchar* hint, bool deferCreation = false, bool keepMappingInMemory = false) override;
+		virtual bool StoreCasFile(CasKey& out, StringKey fileNameKey, const tchar* fileName, FileMappingHandle mappingHandle, u64 mappingOffset, u64 fileSize, const tchar* hint, bool deferCreation = false, bool keepMappingInMemory = false, bool storeCompressed = true) override;
 		virtual bool DropCasFile(const CasKey& casKey, bool forceDelete, const tchar* hint) override;
 		virtual bool ReportBadCasFile(const CasKey& casKey) override;
 		virtual bool CalculateCasKey(CasKey& out, const tchar* fileName) override;
diff --git a/Engine/Source/Programs/UnrealBuildAccelerator/Common/Public/UbaStorageClient.h b/Engine/Source/Programs/UnrealBuildAccelerator/Common/Public/UbaStorageClient.h
--- a/Engine/Source/Programs/UnrealBuildAccelerator/Common/Public/UbaStorageClient.h
+++ b/Engine/Source/Programs/UnrealBuildAccelerator/Common/Public/UbaStorageClient.h
@@ -23,6 +23,7 @@ namespace uba
 		u16 proxyPort = DefaultStorageProxyPort;
 		bool sendCompressed = true;
 		bool allowProxy = true;
+		bool sendOneBigFileAtTheTime = true;
 		
 		GetProxyBackendCallback* getProxyBackendCallback = nullptr;
 		void* getProxyBackendUserData = nullptr;
@@ -54,18 +55,18 @@ namespace uba
 		virtual bool GetZone(StringBufferBase& out) override;
 		virtual bool RetrieveCasFile(RetrieveResult& out, const CasKey& casKey, const tchar* hint, FileMappingBuffer* mappingBuffer = nullptr, u64 memoryMapAlignment = 1, bool allowProxy = true) override;
 		virtual bool StoreCasFile(CasKey& out, const tchar* fileName, const CasKey& casKeyOverride, bool deferCreation, bool fileIsCompressed) override;
-		virtual bool StoreCasFile(CasKey& out, StringKey fileNameKey, const tchar* fileName, FileMappingHandle mappingFile, u64 mappingOffset, u64 fileSize, const tchar* hint, bool deferCreation = false, bool keepMappingInMemory = false) override;
+		virtual bool StoreCasFile(CasKey& out, StringKey fileNameKey, const tchar* fileName, FileMappingHandle mappingFile, u64 mappingOffset, u64 fileSize, const tchar* hint, bool deferCreation = false, bool keepMappingInMemory = false, bool storeCompressed = true) override;
 		virtual bool HasCasFile(const CasKey& casKey, CasEntry** out = nullptr) override;
 		virtual void Ping() override;
 		virtual void PrintSummary(Logger& logger) override;
 
 	private:
-		bool SendFile(const CasKey& casKey, const tchar* fileName, u8* sourceMem, u64 sourceSize, const tchar* hint);
 		bool PopulateCasFromDirsRecursive(const tchar* dir, WorkManager& workManager, UnorderedSet<u64>& seenIds, ReaderWriterLock& seenIdsLock, const Function<bool()>& shouldExit);
 
 		NetworkClient& m_client;
 		bool m_sendCompressed;
 		bool m_allowProxy;
+		bool m_sendOneBigFileAtTheTime;
 
 		Guid m_storageServerUid;
 
diff --git a/Engine/Source/Programs/UnrealBuildAccelerator/Common/Public/UbaStorageServer.h b/Engine/Source/Programs/UnrealBuildAccelerator/Common/Public/UbaStorageServer.h
--- a/Engine/Source/Programs/UnrealBuildAccelerator/Common/Public/UbaStorageServer.h
+++ b/Engine/Source/Programs/UnrealBuildAccelerator/Common/Public/UbaStorageServer.h
@@ -38,7 +38,7 @@ namespace uba
 
 		virtual bool GetZone(StringBufferBase& out) override;
 		virtual bool RetrieveCasFile(RetrieveResult& out, const CasKey& casKey, const tchar* hint, FileMappingBuffer* mappingBuffer = nullptr, u64 memoryMapAlignment = 1, bool allowProxy = true) override;
-		virtual bool StoreCasFile(CasKey& out, StringKey fileNameKey, const tchar* fileName, FileMappingHandle mappingHandle, u64 mappingOffset, u64 fileSize, const tchar* hint, bool deferCreation = false, bool keepMappingInMemory = false) override;
+		virtual bool StoreCasFile(CasKey& out, StringKey fileNameKey, const tchar* fileName, FileMappingHandle mappingHandle, u64 mappingOffset, u64 fileSize, const tchar* hint, bool deferCreation = false, bool keepMappingInMemory = false, bool storeCompressed = true) override;
 		virtual bool WriteCompressed(WriteResult& out, const tchar* from, const tchar* toFile) override;
 		virtual bool IsDisallowedPath(const tchar* fileName) override;
 		virtual void SetTrace(Trace* trace, bool detailed) override;
diff --git a/Engine/Source/Programs/UnrealBuildAccelerator/Test/Private/UbaTestStorage.h b/Engine/Source/Programs/UnrealBuildAccelerator/Test/Private/UbaTestStorage.h
--- a/Engine/Source/Programs/UnrealBuildAccelerator/Test/Private/UbaTestStorage.h
+++ b/Engine/Source/Programs/UnrealBuildAccelerator/Test/Private/UbaTestStorage.h
@@ -91,7 +91,7 @@ namespace uba
 		return true;
 	}
 
-	bool TestRemoteStorage(LoggerWithWriter& logger, const StringBufferBase& testRootDir)
+	bool TestRemoteStorageStore(LoggerWithWriter& logger, const StringBufferBase& testRootDir)
 	{
 		auto& logWriter = logger.m_writer;
 		NetworkBackendTcp serverTcp(logWriter, TC("ServerTcp"));
@@ -100,6 +100,7 @@ namespace uba
 		bool ctorSuccess = true;
 		NetworkServer server(ctorSuccess, logWriter);
 		NetworkClient client(ctorSuccess, logWriter);
+		auto dg = MakeGuard([&]() { client.Disconnect(); });
 
 		StringBuffer<> rootDir;
 		rootDir.Append(testRootDir).Append(TC("Uba"));
@@ -145,66 +146,85 @@ namespace uba
 		}
 
 		CasKey key;
-		if (!storageClient.StoreCasFile(key, ToStringKeyLower(fileName), fileName.data, FileMappingHandle{}, 0, 0, TC("UbaTestFile")))
+		bool storeCompressed = false;
+		if (!storageClient.StoreCasFile(key, ToStringKeyLower(fileName), fileName.data, FileMappingHandle{}, 0, 0, TC("UbaTestFile"), false, false, storeCompressed))
 			return logger.Error(TC("Failed to store file %s"), fileName.data);
 
 		fileName.Clear().Append(testRootDir).Append(TC("Uba")).Append(PathSeparator).Append(TC("UbaTestFile"));
 		if (!storageServer.CopyOrLink(key, fileName.data, DefaultAttributes()))
 			return logger.Error(TC("Failed to copy cas to file %s"), fileName.data);
+		return true;
+	}
 
-		/*
-		StringBuffer<> detoursLibCopy(detoursLib);
-		detoursLibCopy.Append(TC(".tmp"));
+	bool TestRemoteStorageFetch(LoggerWithWriter& logger, const StringBufferBase& testRootDir)
+	{
+		auto& logWriter = logger.m_writer;
+		NetworkBackendTcp serverTcp(logWriter, TC("ServerTcp"));
+		NetworkBackendTcp clientTcp(logWriter, TC("ClientTcp"));
 
-		auto deleteFile = MakeGuard([&]() { return DeleteFileW(detoursLibCopy.data); });
+		bool ctorSuccess = true;
+		NetworkServer server(ctorSuccess, logWriter);
+		NetworkClient client(ctorSuccess, logWriter);
+		auto dg = MakeGuard([&]() { client.Disconnect(); });
 
-		if (!storage.CopyOrLink(key, detoursLibCopy.data))
-			return logger.Error(TC("Failed to copy cas to file %s"), detoursLibCopy.data);
+		StringBuffer<> rootDir;
+		rootDir.Append(testRootDir).Append(TC("Uba"));
+		if (!DeleteAllFiles(logger, rootDir.data))
+			return false;
 
-		FileHandle original;
-		if (!OpenFileSequentialRead(logger, detoursLib.data, original))
-			return logger.Error(TC("Failed to open %s for read"), detoursLib.data);
-		auto closeOriginal = MakeGuard([&]() { return CloseFile(detoursLib.data, original); });
+		StorageServerCreateInfo storageServerInfo(server, rootDir.data, logger.m_writer);
+		storageServerInfo.casCapacityBytes = 1024ull * 1024 * 1024;
+		storageServerInfo.storeCompressed = false;
+		StorageServer storageServer(storageServerInfo);
 
-		FileHandle copy;
-		if (!OpenFileSequentialRead(logger, detoursLibCopy.data, copy))
-			return logger.Error(TC("Failed to open %s for read"), detoursLibCopy.data);
-		auto closeCopy = MakeGuard([&]() { return CloseFile(detoursLibCopy.data, copy); });
+		auto g = MakeGuard([&]() { server.DisconnectClients(); });
 
-		u64 originalSize;
-		if (!GetFileSizeEx(originalSize, original))
-			return logger.Error(TC("Failed to get size of %s"), detoursLib.data);
-		u64 copySize;
-		if (!GetFileSizeEx(copySize, copy))
-			return logger.Error(TC("Failed to get size of %s"), detoursLibCopy.data);
-		if (originalSize != copySize)
-			return logger.Error(TC("Sizes mismatch between %s and %s"), detoursLib.data, detoursLibCopy.data);
+		rootDir.Append(TC("Client"));
+		if (!DeleteAllFiles(logger, rootDir.data))
+			return false;
 
-		u8 originalBuffer[64 * 1024];
-		u8 copyBuffer[64 * 1024];
-		u64 left = originalSize;
-		while (left)
-		{
-			u64 toRead = Min(left, (u64)sizeof(originalBuffer));
-			if (!ReadFile(logger, detoursLib.data, original, originalBuffer, toRead))
-				return logger.Error(TC("Failed to read %u from %s"), detoursLib.data);
-			if (!ReadFile(logger, detoursLibCopy.data, copy, copyBuffer, toRead))
-				return logger.Error(TC("Failed to read %u from %s"), detoursLibCopy.data);
-			if (memcmp(originalBuffer, copyBuffer, toRead) != 0)
-				return logger.Error(TC("Data mismatch between %s and %s"), detoursLib.data, detoursLibCopy.data);
+		StorageClientCreateInfo storageClientInfo(client, rootDir.data);
+		StorageClient storageClient(storageClientInfo);
+		storageClient.Start();
 
-			left -= toRead;
+		if (!storageClient.LoadCasTable(true))
+			return false;
+
+		rootDir.EnsureEndsWithSlash();
+
+		if (!server.StartListen(serverTcp, 1234))
+			return logger.Error(TC("Failed to listen"));
+		Sleep(100);
+		if (!client.Connect(clientTcp, TC("127.0.0.1"), 1234))
+			return logger.Error(TC("Failed to connect"));
+
+		StringBuffer<> fileName;
+#if 1
+		{
+			fileName.Append(rootDir).Append(TC("UbaTestFile"));
+			FileAccessor fileHandle(logger, fileName.data);
+			if (!fileHandle.CreateWrite())
+				return logger.Error(TC("Failed to create file for write"));
+			u8 byte = 'H';
+			if (!fileHandle.Write(&byte, 1))
+				return false;
+			if (!fileHandle.Close())
+				return false;
 		}
+#else
+		fileName.Append(TC("e:\\dev\\fn\\QAGame\\Saved\\StagedBuilds\\PS5_Temp\\Split\\qagame\\content\\paks\\qagame-ps5.ucas\\.copied"));
+#endif
 
-		if (!closeOriginal.Execute())
-			return logger.Error(TC("Failed to close %s"), detoursLib.data);
+		CasKey casKey;
+		if (!storageServer.CalculateCasKey(casKey, fileName.data))
+			return false;
 
-		if (!closeCopy.Execute())
-			return logger.Error(TC("Failed to close %s"), detoursLibCopy.data);
+		casKey = AsCompressed(casKey, false);
+
+		Storage::RetrieveResult result;
+		if (!storageClient.RetrieveCasFile(result, casKey, fileName.data))
+			return logger.Error(TC("Failed to store file %s"), fileName.data);
 
-		if (!deleteFile.Execute())
-			return logger.Error(TC("Failed to delete %s"), detoursLibCopy.data);
-		*/
 		return true;
 	}
 }
