[UBA]
* Fixed so certain file types are not being compressed when sent back to server from client

diff --git a/Engine/Source/Programs/UnrealBuildAccelerator/Common/Private/UbaSessionClient.cpp b/Engine/Source/Programs/UnrealBuildAccelerator/Common/Private/UbaSessionClient.cpp
--- a/Engine/Source/Programs/UnrealBuildAccelerator/Common/Private/UbaSessionClient.cpp
+++ b/Engine/Source/Programs/UnrealBuildAccelerator/Common/Private/UbaSessionClient.cpp
@@ -587,19 +587,20 @@ namespace uba
 				m_logger.Warning(TC("%s is not using file mapping"), pair.first.c_str());
 #endif
 			bool keepMappingInMemory = IsWindows && !IsRarelyReadAfterWritten(process, pair.first);
-			if (!SendFile(pair.second, pair.first.c_str(), process.GetId(), keepMappingInMemory))
+			bool compressed = process.m_startInfo.rules->SendFileCompressedFromClient(pair.first);
+			if (!SendFile(pair.second, pair.first.c_str(), process.GetId(), keepMappingInMemory, compressed))
 				return false;
 		}
 		return true;
 	}
 
-	bool SessionClient::SendFile(WrittenFile& source, const tchar* destination, u32 processId, bool keepMappingInMemory)
+	bool SessionClient::SendFile(WrittenFile& source, const tchar* destination, u32 processId, bool keepMappingInMemory, bool compressed)
 	{
 		CasKey casKey;
 		{
 			TimerScope ts(m_stats.storageSend);
 			bool deferCreation = false;
-			if (!m_storage.StoreCasFile(casKey, source.key, source.name.c_str(), source.mappingHandle, 0, source.mappingWritten, destination, deferCreation, keepMappingInMemory))
+			if (!m_storage.StoreCasFile(casKey, source.key, source.name.c_str(), source.mappingHandle, 0, source.mappingWritten, destination, deferCreation, keepMappingInMemory, compressed))
 				return false;
 		}
 		UBA_ASSERTF(casKey != CasKeyZero, TC("Failed to store cas file for %s (destination %s)"), source.name.c_str(), destination);
@@ -1506,7 +1507,7 @@ namespace uba
 			logFile = lastSlash + 1;
 		dest.Append(TC("<log>")).Append(logFile);
 		f.key = ToStringKeyLower(dest);
-		SendFile(f, dest.data, pi.GetId(), false);
+		SendFile(f, dest.data, pi.GetId(), false, true);
 		for (auto& child : pi.m_childProcesses)
 			SendLogFileToServer(*(ProcessImpl*)child.m_process);
 	}
@@ -1904,7 +1905,7 @@ namespace uba
 				StringBuffer<> dest(TC("<uba>"));
 				f.key = ToStringKeyLower(dest);
 
-				SendFile(f, dest.data, 0, false);
+				SendFile(f, dest.data, 0, false, true);
 			}
 		}
 	}
diff --git a/Engine/Source/Programs/UnrealBuildAccelerator/Common/Public/UbaSessionClient.h b/Engine/Source/Programs/UnrealBuildAccelerator/Common/Public/UbaSessionClient.h
--- a/Engine/Source/Programs/UnrealBuildAccelerator/Common/Public/UbaSessionClient.h
+++ b/Engine/Source/Programs/UnrealBuildAccelerator/Common/Public/UbaSessionClient.h
@@ -80,7 +80,7 @@ namespace uba
 		bool EnsureBinaryFile(StringBufferBase& out, StringBufferBase& outVirtual, u32 processId, const StringBufferBase& fileName, const StringKey& fileNameKey, const tchar* applicationDir, const u8* loaderPaths, u32 loaderPathsSize);
 		bool WriteBinFile(StringBufferBase& out, const tchar* binaryName, const CasKey& casKey, const KeyToString& applicationDir, u32 fileAttributes);
 		bool SendFiles(ProcessImpl& process, Timer& sendFiles);
-		bool SendFile(WrittenFile& source, const tchar* destination, u32 processId, bool keepMappingInMemory);
+		bool SendFile(WrittenFile& source, const tchar* destination, u32 processId, bool keepMappingInMemory, bool compressed);
 		bool SendUpdateDirectoryTable(StackBinaryReader<SendMaxSize>& reader); // Note, reader is sent in to save stack space.
 		bool UpdateDirectoryTableFromServer(StackBinaryReader<SendMaxSize>& reader);
 		bool SendUpdateNameToHashTable(StackBinaryReader<SendMaxSize>& reader);
diff --git a/Engine/Source/Programs/UnrealBuildAccelerator/Core/Private/UbaApplicationRules.cpp b/Engine/Source/Programs/UnrealBuildAccelerator/Core/Private/UbaApplicationRules.cpp
--- a/Engine/Source/Programs/UnrealBuildAccelerator/Core/Private/UbaApplicationRules.cpp
+++ b/Engine/Source/Programs/UnrealBuildAccelerator/Core/Private/UbaApplicationRules.cpp
@@ -15,7 +15,7 @@ namespace uba
 			return true;
 		}
 
-		virtual u64 FileTypeMaxSize(const StringBufferBase& file, bool isSystemOrTempFile) const override
+		virtual u64 FileTypeMaxSize(const StringView& file, bool isSystemOrTempFile) const override
 		{
 			if (file.EndsWith(TC(".pdb")))
 				return 14ull * 1024 * 1024 * 1024; // This is ridiculous
@@ -62,10 +62,11 @@ namespace uba
 				;// || fileName.EndsWith(TC(".h.pch") // Not tested enough
 		}
 
-		virtual bool IsRarelyRead(const StringBufferBase& file) const override
+		virtual bool IsRarelyRead(const StringView& file) const override
 		{
 			return file.EndsWith(TC(".cpp"))
-				|| file.EndsWith(TC(".obj.rsp"));
+				|| (file.EndsWith(TC(".rsp")) && !file.EndsWith(TC("Shared.rsp")))
+				;
 		}
 
 		virtual bool IsRarelyReadAfterWritten(const StringView& fileName) const override
@@ -127,7 +128,7 @@ namespace uba
 			 	|| EndsWith(file, TStrlen(file), TC(".manifest")); // lld-link.exe is using a different name for files shared with child processes
 		}
 
-		virtual bool IsRarelyRead(const StringBufferBase& file) const override
+		virtual bool IsRarelyRead(const StringView& file) const override
 		{
 			return file.EndsWith(TC(".exp"))
 				|| file.EndsWith(TC(".dll.rsp"))
@@ -136,7 +137,7 @@ namespace uba
 				|| file.EndsWith(TC(".pdb"));
 		}
 
-		virtual bool AllowStorageProxy(const StringBufferBase& file) const override
+		virtual bool AllowStorageProxy(const StringView& file) const override
 		{
 			if (file.EndsWith(TC(".obj")))
 				return false;
@@ -188,7 +189,7 @@ namespace uba
 			return fileName.EndsWith(TC(".manifest")) && fileName.Contains(systemTemp);
 		}
 
-		virtual u64 FileTypeMaxSize(const StringBufferBase& file, bool isSystemOrTempFile) const override
+		virtual u64 FileTypeMaxSize(const StringView& file, bool isSystemOrTempFile) const override
 		{
 			if (file.Contains(TC(".pdb.tmp")))
 				return 14ull * 1024 * 1024 * 1024; // This is ridiculous
@@ -238,10 +239,10 @@ namespace uba
 				;// || fileName.EndsWith(TC(".gch.tmp")); // Need to fix "has been modified since the precompiled header"
 		}
 
-		virtual bool IsRarelyRead(const StringBufferBase& file) const override
+		virtual bool IsRarelyRead(const StringView& file) const override
 		{
 			return file.EndsWith(TC(".cpp"))
-				|| file.EndsWith(TC(".o.rsp"));
+				|| (file.EndsWith(TC(".rsp")) && !file.EndsWith(TC("Shared.rsp")));
 		}
 
 		virtual bool IsRarelyReadAfterWritten(const StringView& fileName) const override
@@ -279,12 +280,12 @@ namespace uba
 			return fileName.Contains(TC(".tmp")); // both .so.tmp and .tmp123456
 		}
 
-		virtual bool IsRarelyRead(const StringBufferBase& file) const override
+		virtual bool IsRarelyRead(const StringView& file) const override
 		{
 			return file.EndsWith(TC(".so.rsp"));
 		}
 
-		virtual u64 FileTypeMaxSize(const StringBufferBase& file, bool isSystemOrTempFile) const override
+		virtual u64 FileTypeMaxSize(const StringView& file, bool isSystemOrTempFile) const override
 		{
 			return 14ull * 1024 * 1024 * 1024; // This is ridiculous (needed for asan targets)
 		}
@@ -304,7 +305,7 @@ namespace uba
 			return fileName.Contains(TC(".temp-stream-"));
 		}
 
-		virtual u64 FileTypeMaxSize(const StringBufferBase& file, bool isSystemOrTempFile) const override
+		virtual u64 FileTypeMaxSize(const StringView& file, bool isSystemOrTempFile) const override
 		{
 			if (IsOutputFile(file))
 				return 14ull * 1024 * 1024 * 1024; // This is ridiculous (needed for asan targets)
@@ -406,6 +407,21 @@ namespace uba
 		}
 	};
 
+	class ApplicationRulesProsperoCompressCmdExe : public ApplicationRules
+	{
+		using Super = ApplicationRules;
+
+		virtual bool AllowStorageProxy(const StringView& file) const override
+		{
+			return !file.EndsWith(TC(".split"));
+		}
+
+		virtual bool SendFileCompressedFromClient(const StringView& fileName) const override
+		{
+			return !fileName.EndsWith(TC(".split"));
+		}
+	};
+
 	// ====
 
 	class ApplicationRulesISPCExe : public ApplicationRules
@@ -462,7 +478,7 @@ namespace uba
 				|| fileName.EndsWith(TC(".stacktrace.txt"));
 		}
 		
-		virtual bool IsRarelyRead(const StringBufferBase& file) const override
+		virtual bool IsRarelyRead(const StringView& file) const override
 		{
 			return file.EndsWith(TC(".i"))
 				|| file.EndsWith(TC(".PVS-Studio.log"))
@@ -485,7 +501,7 @@ namespace uba
 
 	class ApplicationRulesShaderCompileWorker : public ApplicationRules
 	{
-		virtual bool IsRarelyRead(const StringBufferBase& file) const override
+		virtual bool IsRarelyRead(const StringView& file) const override
 		{
 			return file.Contains(TC(".uba."));
 		}
@@ -530,6 +546,7 @@ namespace uba
 			{ TC("orbis-ltop.exe"),				new ApplicationRulesOrbisLdExe() },
 			{ TC("prospero-clang.exe"),			new ApplicationRulesProsperoClangPlusPlusExe() },
 			{ TC("prospero-lld.exe"),			new ApplicationRulesProsperoLldExe() },
+			{ TC("prospero-compress-cmd.exe"),	new ApplicationRulesProsperoCompressCmdExe() },
 			{ TC("dump_syms.exe"),				new ApplicationRulesDumpSymsExe() },
 			{ TC("ld.lld.exe"),					new ApplicationRulesLdLLdExe() },
 			{ TC("llvm-objcopy.exe"),			new ApplicationRulesLlvmObjCopyExe() },
diff --git a/Engine/Source/Programs/UnrealBuildAccelerator/Core/Public/UbaApplicationRules.h b/Engine/Source/Programs/UnrealBuildAccelerator/Core/Public/UbaApplicationRules.h
--- a/Engine/Source/Programs/UnrealBuildAccelerator/Core/Public/UbaApplicationRules.h
+++ b/Engine/Source/Programs/UnrealBuildAccelerator/Core/Public/UbaApplicationRules.h
@@ -54,7 +54,7 @@ namespace uba
 		}
 
 		// Max file size if using memory files.
-		virtual u64 FileTypeMaxSize(const StringBufferBase& file, bool isSystemOrTempFile) const
+		virtual u64 FileTypeMaxSize(const StringView& file, bool isSystemOrTempFile) const
 		{
 			return 8ull * 1024 * 1024 * 1024;
 		}
@@ -73,7 +73,7 @@ namespace uba
 
 		// Return true if the file is only read by this process or very rarely read more than once
 		// This is an optimization to not store the file in the mapping table since it will not be read again and would just take up space
-		virtual bool IsRarelyRead(const StringBufferBase& file) const
+		virtual bool IsRarelyRead(const StringView& file) const
 		{
 			return false;
 		}
@@ -85,7 +85,7 @@ namespace uba
 			return false;
 		}
 
-		virtual bool AllowStorageProxy(const StringBufferBase& file) const
+		virtual bool AllowStorageProxy(const StringView& file) const
 		{
 			return !IsRarelyRead(file);
 		}
@@ -139,6 +139,11 @@ namespace uba
 			return false;
 		}
 
+		virtual bool SendFileCompressedFromClient(const StringView& fileName) const
+		{
+			return true;
+		}
+
 		virtual bool ShouldExtractSymbols(const StringView& fileName) const
 		{
 			return false;
diff --git a/Engine/Source/Programs/UnrealBuildAccelerator/Test/Private/UbaTestAll.h b/Engine/Source/Programs/UnrealBuildAccelerator/Test/Private/UbaTestAll.h
--- a/Engine/Source/Programs/UnrealBuildAccelerator/Test/Private/UbaTestAll.h
+++ b/Engine/Source/Programs/UnrealBuildAccelerator/Test/Private/UbaTestAll.h
@@ -29,6 +29,8 @@ namespace uba
 		UBA_TEST(TestClientServer2) \
 		UBA_TEST(TestClientServerMem) \
 		UBA_TEST(TestStorage) \
+		UBA_TEST(TestRemoteStorageStore) \
+		UBA_TEST(TestRemoteStorageFetch) \
 		UBA_TEST(TestDetouredTestApp) \
 		UBA_TEST(TestRemoteDetouredTestApp) \
 		UBA_TEST(TestCompactPathTable) \
