[UBA]
* Moved extrabuffer to the right scope to prevent memory leak if proxy disconnects
* Fixed static analysis error (scenario should never happen, added error handling)

diff --git a/Engine/Source/Programs/UnrealBuildAccelerator/Common/Private/UbaStorageClient.cpp b/Engine/Source/Programs/UnrealBuildAccelerator/Common/Private/UbaStorageClient.cpp
--- a/Engine/Source/Programs/UnrealBuildAccelerator/Common/Private/UbaStorageClient.cpp
+++ b/Engine/Source/Programs/UnrealBuildAccelerator/Common/Private/UbaStorageClient.cpp
@@ -555,15 +555,15 @@ namespace uba
 
 				if (actualSize)
 				{
-					// There is a risk the compressed file we download is larger than half buffer
-					u8* extraBuffer = nullptr;
-					auto ebg = MakeGuard([&]() { delete[] extraBuffer; });
-
 					u64 leftCompressed = fileSize - responseSize;
 					u32 readIndex = 0;
 					bool tryAgain = false;
 					do
 					{
+						// There is a risk the compressed file we download is larger than half buffer
+						u8* extraBuffer = nullptr;
+						auto ebg = MakeGuard([&]() { delete[] extraBuffer; });
+
 						// First read in a full decompressable block
 						bool isFirstInBlock = true;
 						u32 compressedSize = ~u32(0);
@@ -603,6 +603,8 @@ namespace uba
 									UBA_ASSERT(responseSize > capacity);
 									memcpy(readPosition, extraBuffer, left);
 									memmove(extraBuffer, extraBuffer + left, u32(responseSize - left));
+									if (isFirstInBlock)
+										return m_logger.Error(TC("Make static analysis happy. This should not be possible to happen (%s)"), casFile.data);
 								}
 
 								leftCompressed -= responseSize;
