[UBA]
* Added so the SessionServer_RegisterNewFile and SessionServer_RefreshDirectory returns a bool telling if they succeeded or not since they can fail

diff --git a/Engine/Source/Programs/UnrealBuildAccelerator/Common/Private/UbaSession.cpp b/Engine/Source/Programs/UnrealBuildAccelerator/Common/Private/UbaSession.cpp
--- a/Engine/Source/Programs/UnrealBuildAccelerator/Common/Private/UbaSession.cpp
+++ b/Engine/Source/Programs/UnrealBuildAccelerator/Common/Private/UbaSession.cpp
@@ -1311,7 +1311,7 @@ namespace uba
 		return h;
 	}
 
-	void Session::RefreshDirectory(const tchar* dirName)
+	bool Session::RefreshDirectory(const tchar* dirName)
 	{
 		UBA_ASSERT(!m_runningRemote);
 
@@ -1323,7 +1323,7 @@ namespace uba
 		SCOPED_READ_LOCK(dirTable.m_lookupLock, lookupLock);
 		auto res = dirTable.m_lookup.find(dirKey);
 		if (res == dirTable.m_lookup.end())
-			return;
+			return true;
 		DirectoryTable::Directory& dir = res->second;
 		lookupLock.Leave();
 		SCOPED_WRITE_LOCK(dir.lock, dirLock);
@@ -1341,7 +1341,7 @@ namespace uba
 		m_directoryTable.PopulateDirectory(hasher, dir);
 
 		u32 tableOffset;
-		WriteDirectoryEntriesInternal(dir, dirKey, dirPath.data, true, tableOffset);
+		return WriteDirectoryEntriesInternal(dir, dirKey, dirPath.data, true, tableOffset);
 	}
 
 	StringKey GetKeyAndFixedName(StringBuffer<>& fixedFilePath, const tchar* filePath)
diff --git a/Engine/Source/Programs/UnrealBuildAccelerator/Common/Private/UbaStorage.cpp b/Engine/Source/Programs/UnrealBuildAccelerator/Common/Private/UbaStorage.cpp
--- a/Engine/Source/Programs/UnrealBuildAccelerator/Common/Private/UbaStorage.cpp
+++ b/Engine/Source/Programs/UnrealBuildAccelerator/Common/Private/UbaStorage.cpp
@@ -259,12 +259,12 @@ namespace uba
 
 		FileHandle readHandle;
 		if (!OpenFileSequentialRead(m_logger, from, readHandle))
-			return m_logger.Error(TC("Failed to open file %s for read (%s)"), from, LastErrorToText().data);
+			return m_logger.Error(TC("[WriteCompressed] Failed to open file %s for read (%s)"), from, LastErrorToText().data);
 		auto rsg = MakeGuard([&]() { CloseFile(from, readHandle); });
 
 		u64 fileSize;
 		if (!uba::GetFileSizeEx(fileSize, readHandle))
-			return m_logger.Error(TC("GetFileSize failed for %s (%s)"), from, LastErrorToText().data);
+			return m_logger.Error(TC("[WriteCompressed] GetFileSize failed for %s (%s)"), from, LastErrorToText().data);
 
 		return WriteCompressed(out, from, readHandle, 0, fileSize, to, nullptr, 0);
 	}
@@ -549,12 +549,12 @@ namespace uba
 		{
 			FileHandle readHandle;
 			if (!OpenFileSequentialRead(m_logger, fileName, readHandle))
-				return m_logger.Error(TC("Failed to open file %s for read (%s)"), fileName, LastErrorToText().data);
+				return m_logger.Error(TC("[WriteCasFileNoCheck] Failed to open file %s for read (%s)"), fileName, LastErrorToText().data);
 			auto fileGuard = MakeGuard([&](){ CloseFile(fileName, readHandle); });
 
 			u64 fileSize;
 			if (!uba::GetFileSizeEx(fileSize, readHandle))
-				return m_logger.Error(TC("GetFileSize failed for %s (%s)"), fileName, LastErrorToText().data);
+				return m_logger.Error(TC("[WriteCasFileNoCheck] GetFileSize failed for %s (%s)"), fileName, LastErrorToText().data);
 
 			FileAccessor destinationFile(m_logger, casFile);
 			if (!destinationFile.CreateWrite(false, DefaultAttributes(), 0, m_tempPath.data))
@@ -2343,7 +2343,7 @@ namespace uba
 
 		u64 fileSize;
 		if (!uba::GetFileSizeEx(fileSize, fileHandle))
-			return m_logger.Error(TC("GetFileSize failed for %s (%s)"), fileName, LastErrorToText().data);
+			return m_logger.Error(TC("[CalculateCasKey] GetFileSize failed for %s (%s)"), fileName, LastErrorToText().data);
 
 		out = CalculateCasKey(fileName, fileHandle, fileSize, true);
 		return out != CasKeyZero;
@@ -2397,7 +2397,7 @@ namespace uba
 			if (!HasCasFile(actualKey, &casEntry))
 			{
 				if (!testCompressed)
-					return m_logger.Error(TC("Trying to copy cas %s to %s but can't find neither compressed or uncompressed version"), CasKeyString(actualKey).str, destination);
+					return m_logger.Error(TC("[CopyOrLink] Trying to copy cas %s to %s but can't find neither compressed or uncompressed version"), CasKeyString(actualKey).str, destination);
 
 				actualKey = AsCompressed(casKey, true);
 				testCompressed = false;
@@ -2427,7 +2427,7 @@ namespace uba
 					mappedView = m_casDataBuffer.MapView(casEntry->mappingHandle, casEntry->mappingOffset, casEntry->mappingSize, casFile.data);
 					compressedData = mappedView.memory;
 					if (!compressedData)
-						return m_logger.Error(TC("Failed to map view of mapping %s (%s)"), casFile.data, LastErrorToText().data);
+						return m_logger.Error(TC("[CopyOrLink] Failed to map view of mapping %s (%s)"), casFile.data, LastErrorToText().data);
 
 					decompressedSize = *(u64*)compressedData;
 					readData = compressedData + sizeof(u64);
@@ -2441,10 +2441,10 @@ namespace uba
 					UBA_ASSERT(false);
 #endif
 					if (!OpenFileSequentialRead(m_logger, casFile.data, readHandle))
-						return m_logger.Error(TC("Failed to open file %s for read (%s)"), casFile.data, LastErrorToText().data);
+						return m_logger.Error(TC("[CopyOrLink] Failed to open file %s for read (%s)"), casFile.data, LastErrorToText().data);
 
 					if (!ReadFile(m_logger, casFile.data, readHandle, &decompressedSize, sizeof(u64)))
-						return m_logger.Error(TC("Failed to read first 8 bytes from compressed file %s (%s)"), casFile.data, LastErrorToText().data);
+						return m_logger.Error(TC("[CopyOrLink] Failed to read first 8 bytes from compressed file %s (%s)"), casFile.data, LastErrorToText().data);
 				}
 
 				bool writeDirectlyToFile = casEntry->mappingHandle.IsValid() && false; // Experiment to try to fix bottlenecks on cloud
@@ -2472,7 +2472,7 @@ namespace uba
 					u64 compressedFileSize = mappedView.size;
 					if (!mappedView.memory)
 						if (!GetFileSizeEx(compressedFileSize, readHandle))
-							return m_logger.Error(TC("Failed to get file size of compressed file %s (%s)"), casFile.data, LastErrorToText().data);
+							return m_logger.Error(TC("[CopyOrLink] Failed to get file size of compressed file %s (%s)"), casFile.data, LastErrorToText().data);
 						
 					sizeOnDisk = compressedFileSize + sizeof(CompressedObjFileHeader);
 
@@ -2492,7 +2492,7 @@ namespace uba
 						*(u64*)writePos = decompressedSize;
 						writePos += sizeof(u64);
 						if (!ReadFile(m_logger, casFile.data, readHandle, writePos, compressedFileSize - sizeof(u64)))
-							return m_logger.Error(TC("Failed to read compressed file %s (%s)"), casFile.data, LastErrorToText().data);
+							return m_logger.Error(TC("[CopyOrLink] Failed to read compressed file %s (%s)"), casFile.data, LastErrorToText().data);
 					}
 						
 				}
diff --git a/Engine/Source/Programs/UnrealBuildAccelerator/Common/Public/UbaSession.h b/Engine/Source/Programs/UnrealBuildAccelerator/Common/Public/UbaSession.h
--- a/Engine/Source/Programs/UnrealBuildAccelerator/Common/Public/UbaSession.h
+++ b/Engine/Source/Programs/UnrealBuildAccelerator/Common/Public/UbaSession.h
@@ -56,7 +56,7 @@ namespace uba
 		void CancelAllProcessesAndWait(bool terminate = true); // Cancel all processes and wait for them to go away
 
 		void PrintSummary(Logger& logger); // Print summary stats of session
-		void RefreshDirectory(const tchar* dirName); // Tell uba a directory on disk has been changed by some other system while session is running
+		bool RefreshDirectory(const tchar* dirName); // Tell uba a directory on disk has been changed by some other system while session is running
 		bool RegisterNewFile(const tchar* filePath); // Tell uba a new file on disk has been added by some other system while session is running
 		void RegisterDeleteFile(const tchar* filePath); // Tell uba a file on disk has been deleted by some other system while session is running
 
diff --git a/Engine/Source/Programs/UnrealBuildAccelerator/Host/Private/UbaExports.cpp b/Engine/Source/Programs/UnrealBuildAccelerator/Host/Private/UbaExports.cpp
--- a/Engine/Source/Programs/UnrealBuildAccelerator/Host/Private/UbaExports.cpp
+++ b/Engine/Source/Programs/UnrealBuildAccelerator/Host/Private/UbaExports.cpp
@@ -365,13 +365,13 @@ extern "C"
 	{
 		server->SetRemoteProcessReturnedEvent([returned, userData](uba::Process& process) { returned(process, userData); });
 	}
-	void SessionServer_RefreshDirectory(uba::SessionServer* server, const uba::tchar* directory)
+	bool SessionServer_RefreshDirectory(uba::SessionServer* server, const uba::tchar* directory)
 	{
-		server->RefreshDirectory(directory);
+		return server->RefreshDirectory(directory);
 	}
-	void SessionServer_RegisterNewFile(uba::SessionServer* server, const uba::tchar* filePath)
+	bool SessionServer_RegisterNewFile(uba::SessionServer* server, const uba::tchar* filePath)
 	{
-		server->RegisterNewFile(filePath);
+		return server->RegisterNewFile(filePath);
 	}
 	void SessionServer_RegisterDeleteFile(uba::SessionServer* server, const uba::tchar* filePath)
 	{
diff --git a/Engine/Source/Programs/UnrealBuildAccelerator/Host/Public/UbaExports.h b/Engine/Source/Programs/UnrealBuildAccelerator/Host/Public/UbaExports.h
--- a/Engine/Source/Programs/UnrealBuildAccelerator/Host/Public/UbaExports.h
+++ b/Engine/Source/Programs/UnrealBuildAccelerator/Host/Public/UbaExports.h
@@ -98,8 +98,8 @@ extern "C"
 	UBA_API uba::SessionServer* SessionServer_Create(const uba::SessionServerCreateInfo& info, const uba::u8* environment = nullptr, uba::u32 environmentSize = 0);
 	UBA_API void SessionServer_SetRemoteProcessAvailable(uba::SessionServer* server, SessionServer_RemoteProcessAvailableCallback* available, void* userData);
 	UBA_API void SessionServer_SetRemoteProcessReturned(uba::SessionServer* server, SessionServer_RemoteProcessReturnedCallback* returned, void* userData);
-	UBA_API void SessionServer_RefreshDirectory(uba::SessionServer* server, const uba::tchar* directory);
-	UBA_API void SessionServer_RegisterNewFile(uba::SessionServer* server, const uba::tchar* filePath);
+	UBA_API bool SessionServer_RefreshDirectory(uba::SessionServer* server, const uba::tchar* directory);
+	UBA_API bool SessionServer_RegisterNewFile(uba::SessionServer* server, const uba::tchar* filePath);
 	UBA_API void SessionServer_RegisterDeleteFile(uba::SessionServer* server, const uba::tchar* filePath);
 	UBA_API uba::ProcessHandle* SessionServer_RunProcess(uba::SessionServer* server, uba::ProcessStartInfo& info, bool async, bool enableDetour);
 	UBA_API uba::ProcessHandle* SessionServer_RunProcessRemote(uba::SessionServer* server, uba::ProcessStartInfo& info, float weight, const void* knownInputs = nullptr, uba::u32 knownInputsCount = 0);
