[UBA]
* Fixed bug where computer name is always  included in UbaAgent name even though -name=x was used

diff --git a/Engine/Source/Programs/UnrealBuildAccelerator/Agent/Private/UbaAgent.cpp b/Engine/Source/Programs/UnrealBuildAccelerator/Agent/Private/UbaAgent.cpp
--- a/Engine/Source/Programs/UnrealBuildAccelerator/Agent/Private/UbaAgent.cpp
+++ b/Engine/Source/Programs/UnrealBuildAccelerator/Agent/Private/UbaAgent.cpp
@@ -416,7 +416,7 @@ namespace uba
 		TString command;
 		u16 port = DefaultPort;
 		u16 proxyPort = DefaultStorageProxyPort;
-		StringBuffer<> agentName(DefaultAgentName);
+		StringBuffer<> agentName;
 		bool useListen = false;
 		bool logToFile = false;
 		bool storeCompressed = true;
@@ -787,6 +787,9 @@ namespace uba
 		}
 		#endif
 
+		if (agentName.IsEmpty())
+			agentName.Append(DefaultAgentName);
+
 		if (!zone.count)
 			GetZone(zone);
 
