[UBA]
* Changed application rules to use hashes instead of strings

diff --git a/Engine/Source/Programs/UnrealBuildAccelerator/Common/Private/UbaSession.cpp b/Engine/Source/Programs/UnrealBuildAccelerator/Common/Private/UbaSession.cpp
--- a/Engine/Source/Programs/UnrealBuildAccelerator/Common/Private/UbaSession.cpp
+++ b/Engine/Source/Programs/UnrealBuildAccelerator/Common/Private/UbaSession.cpp
@@ -1422,18 +1422,21 @@ namespace uba
 		
 		while (true)
 		{
+			exeName.MakeLower();
+			u32 appHash = GetApplicationHash(exeName);
+
 			for (u32 i = 1;; ++i)
 			{
-				const tchar* app = rules[i].app;
-				if (!app)
+				u32 hash = rules[i].hash;
+				if (!hash)
 					break;
-				if (!exeName.Equals(app))
+				if (appHash != hash)
 					continue;
-				return GetApplicationRules()[i].rules;
+				return rules[i].rules;
 			}
 
 			if (!exeName.Equals(TC("dotnet.exe")))
-				return GetApplicationRules()[0].rules;
+				return rules[0].rules;
 			
 			u32 firstArgumentStart = 0;
 			u32 firstArgumentEnd = 0;
@@ -1467,7 +1470,7 @@ namespace uba
 			}
 			exeName.Clear().Append(si.arguments + firstArgumentStart, firstArgumentEnd - firstArgumentStart);
 		}
-		return GetApplicationRules()[0].rules;
+		return rules[0].rules;
 	}
 
 	const tchar* Session::GetTempPath()
diff --git a/Engine/Source/Programs/UnrealBuildAccelerator/Core/Private/UbaApplicationRules.cpp b/Engine/Source/Programs/UnrealBuildAccelerator/Core/Private/UbaApplicationRules.cpp
--- a/Engine/Source/Programs/UnrealBuildAccelerator/Core/Private/UbaApplicationRules.cpp
+++ b/Engine/Source/Programs/UnrealBuildAccelerator/Core/Private/UbaApplicationRules.cpp
@@ -321,7 +321,7 @@ namespace uba
 		}
 	};
 
-	class ApplicationRulesOrbisClangPlusPlusExe : public ApplicationRulesClangPlusPlusExe
+	class ApplicationRulesClangPlusPlusExePlatform1 : public ApplicationRulesClangPlusPlusExe
 	{
 		using Super = ApplicationRulesClangPlusPlusExe;
 
@@ -336,7 +336,7 @@ namespace uba
 		//}
 	};
 
-	class ApplicationRulesOrbisLdExe : public ApplicationRules
+	class ApplicationRulesLdExePlatform1 : public ApplicationRules
 	{
 		using Super = ApplicationRules;
 
@@ -360,7 +360,7 @@ namespace uba
 		}
 	};
 
-	class ApplicationRulesProsperoClangPlusPlusExe : public ApplicationRulesClangPlusPlusExe
+	class ApplicationRulesClangPlusPlusExePlatform2 : public ApplicationRulesClangPlusPlusExe
 	{
 		using Super = ApplicationRulesClangPlusPlusExe;
 
@@ -381,7 +381,7 @@ namespace uba
 		//}
 	};
 
-	class ApplicationRulesProsperoLldExe : public ApplicationRules
+	class ApplicationRulesLldExePlatform2 : public ApplicationRules
 	{
 		using Super = ApplicationRules;
 
@@ -407,7 +407,7 @@ namespace uba
 		}
 	};
 
-	class ApplicationRulesProsperoCompressCmdExe : public ApplicationRules
+	class ApplicationRulesStageSplitMergeApplication : public ApplicationRules
 	{
 		using Super = ApplicationRules;
 
@@ -521,54 +521,73 @@ namespace uba
 		}
 	};
 
+	template<typename T, size_t size>
+	constexpr u32 ApplicationHash(T (&buffer)[size])
+	{
+		u32 hash = 5381;
+		for (size_t i=0, e=size-1; i!=e; ++i)
+			hash = ((hash << 5) + hash) + buffer[i];
+		return hash;
+	}
+
+	u32 GetApplicationHash(const StringView& str)
+	{
+		u32 hash = 5381;
+		for (size_t i=0; i!=str.count; ++i)
+			hash = ((hash << 5) + hash) + str.data[i];
+		return hash;
+	}
+
 	const RulesRec* GetApplicationRules()
 	{
 		// TODO: Add support for data driven rules.
 		// Note, they need to be possible to serialize from server to client and then from client to each detoured process
+		// ALL HASHES LOWERCASE
 
 		static RulesRec rules[]
 		{
-			{ TC(""),							new ApplicationRules() },		// Must be index 0
+			{ ApplicationHash(TC("")),							new ApplicationRules() },		 // Must be index 0
 #if PLATFORM_WINDOWS
-			{ TC("cl.exe"),						new ApplicationRulesClExe() },	// Must be index 1
-			{ TC("link.exe"),					new ApplicationRulesLinkExe() }, // Must be index 2
-			{ TC("lib.exe"),					new ApplicationRulesVcLink() },
-			{ TC("cvtres.exe"),					new ApplicationRulesLinkExe() },
-			{ TC("mt.exe"),						new ApplicationRulesVcLink() },
-			{ TC("rc.exe"),						new ApplicationRulesVcLink() },
-			{ TC("lld-link.exe"),				new ApplicationRulesLldLinkExe() },
-			{ TC("clang++.exe"),				new ApplicationRulesClangPlusPlusExe() },
-			{ TC("clang-cl.exe"),				new ApplicationRulesClangPlusPlusExe() },
-			{ TC("verse-clang-cl.exe"),			new ApplicationRulesClangPlusPlusExe() },
-			{ TC("ispc.exe"),					new ApplicationRulesISPCExe() },
-			{ TC("orbis-clang.exe"),			new ApplicationRulesOrbisClangPlusPlusExe() },
-			{ TC("orbis-ld.exe"),				new ApplicationRulesOrbisLdExe() },
-			{ TC("orbis-ltop.exe"),				new ApplicationRulesOrbisLdExe() },
-			{ TC("prospero-clang.exe"),			new ApplicationRulesProsperoClangPlusPlusExe() },
-			{ TC("prospero-lld.exe"),			new ApplicationRulesProsperoLldExe() },
-			{ TC("prospero-compress-cmd.exe"),	new ApplicationRulesProsperoCompressCmdExe() },
-			{ TC("dump_syms.exe"),				new ApplicationRulesDumpSymsExe() },
-			{ TC("ld.lld.exe"),					new ApplicationRulesLdLLdExe() },
-			{ TC("llvm-objcopy.exe"),			new ApplicationRulesLlvmObjCopyExe() },
-			{ TC("UnrealBuildTool.dll"),		new ApplicationRulesUBTDll() },
-			{ TC("PVS-Studio.exe"),				new ApplicationRulesPVSStudio() },
-			{ TC("UbaObjTool.exe"),				new ApplicationRulesUbaObjTool() },
-			{ TC("ShaderCompileWorker.exe"),	new ApplicationRulesShaderCompileWorker() },
+			{ ApplicationHash(TC("cl.exe")),					new ApplicationRulesClExe() },	 // Must be index 1
+			{ ApplicationHash(TC("link.exe")),					new ApplicationRulesLinkExe() }, // Must be index 2
+			{ ApplicationHash(TC("lib.exe")),					new ApplicationRulesVcLink() },
+			{ ApplicationHash(TC("cvtres.exe")),				new ApplicationRulesLinkExe() },
+			{ ApplicationHash(TC("mt.exe")),					new ApplicationRulesVcLink() },
+			{ ApplicationHash(TC("rc.exe")),					new ApplicationRulesVcLink() },
+			{ ApplicationHash(TC("lld-link.exe")),				new ApplicationRulesLldLinkExe() },
+			{ ApplicationHash(TC("clang++.exe")),				new ApplicationRulesClangPlusPlusExe() },
+			{ ApplicationHash(TC("clang-cl.exe")),				new ApplicationRulesClangPlusPlusExe() },
+			{ ApplicationHash(TC("verse-clang-cl.exe")),		new ApplicationRulesClangPlusPlusExe() },
+			{ ApplicationHash(TC("ispc.exe")),					new ApplicationRulesISPCExe() },
+			{ 30341511,											new ApplicationRulesVcLink() },
+			{ 3340509542,										new ApplicationRulesClangPlusPlusExePlatform1() },
+			{ 4113554641,										new ApplicationRulesLdExePlatform1() },
+			{ 1752955744,										new ApplicationRulesLdExePlatform1() },
+			{ 238360161,										new ApplicationRulesClangPlusPlusExePlatform2() },
+			{ 2119756440,										new ApplicationRulesLldExePlatform2() },
+			{ 2898035017,										new ApplicationRulesStageSplitMergeApplication() },
+			{ ApplicationHash(TC("dump_syms.exe")),				new ApplicationRulesDumpSymsExe() },
+			{ ApplicationHash(TC("ld.lld.exe")),				new ApplicationRulesLdLLdExe() },
+			{ ApplicationHash(TC("llvm-objcopy.exe")),			new ApplicationRulesLlvmObjCopyExe() },
+			{ ApplicationHash(TC("unrealbuildtool.dll")),		new ApplicationRulesUBTDll() },
+			{ ApplicationHash(TC("pvs-studio.exe")),			new ApplicationRulesPVSStudio() },
+			{ ApplicationHash(TC("ubaobjtool.exe")),			new ApplicationRulesUbaObjTool() },
+			{ ApplicationHash(TC("shadercompileworker.exe")),	new ApplicationRulesShaderCompileWorker() },
 			//{ L"MSBuild.dll"),				new ApplicationRules() },
 			//{ L"BreakpadSymbolEncoder.exe"),	new ApplicationRulesClang() },
 			//{ L"cmd.exe"),		new ApplicationRules() },
 #else
-			{ TC("clang++"),					new ApplicationRulesClangPlusPlusExe() },
-			{ TC("ld.lld"),						new ApplicationRulesLdLLdExe() },
-			{ TC("ShaderCompileWorker"),		new ApplicationRulesShaderCompileWorker() },
+			{ ApplicationHash(TC("clang++"),					new ApplicationRulesClangPlusPlusExe() },
+			{ ApplicationHash(TC("ld.lld"),						new ApplicationRulesLdLLdExe() },
+			{ ApplicationHash(TC("shadercompileworker"),		new ApplicationRulesShaderCompileWorker() },
 #endif
-			{ nullptr, nullptr }
+			{ 0, nullptr }
 		};
 
 		static bool initRules = []()
 			{
 				u32 index = 0;
-				for (RulesRec* rec = rules; rec->app; ++rec)
+				for (RulesRec* rec = rules; rec->hash; ++rec)
 					rec->rules->index = index++;
 				return true;
 			}();
diff --git a/Engine/Source/Programs/UnrealBuildAccelerator/Core/Public/UbaApplicationRules.h b/Engine/Source/Programs/UnrealBuildAccelerator/Core/Public/UbaApplicationRules.h
--- a/Engine/Source/Programs/UnrealBuildAccelerator/Core/Public/UbaApplicationRules.h
+++ b/Engine/Source/Programs/UnrealBuildAccelerator/Core/Public/UbaApplicationRules.h
@@ -158,7 +158,8 @@ namespace uba
 	};
 
 
-	struct RulesRec { const tchar* app; ApplicationRules* rules; };
+	struct RulesRec { u32 hash; ApplicationRules* rules; };
 
 	const RulesRec* GetApplicationRules();
+	u32 GetApplicationHash(const StringView& str);
 }
