[UBA]
* Added support for registering virtual files. Virtual files can point into segments of existing files. Feature can also be extended to use arbitrary mapped memory if needed
* Fixed so cas files can be created from unaligned mapped memory. This is needed for virtual files
* Fixed so throw-away files are not written to at all and entirely skipped in memory
* Fixed race condition when directory is first created by remote helpers and multiple helpers try to create at the same time
* Fixed missing stats bug when sending back cas uncompressed to host
* Fixed some error reporting

diff --git a/Engine/Source/Programs/UnrealBuildAccelerator/Common/Private/UbaSession.cpp b/Engine/Source/Programs/UnrealBuildAccelerator/Common/Private/UbaSession.cpp
--- a/Engine/Source/Programs/UnrealBuildAccelerator/Common/Private/UbaSession.cpp
+++ b/Engine/Source/Programs/UnrealBuildAccelerator/Common/Private/UbaSession.cpp
@@ -1380,6 +1380,16 @@ namespace uba
 		return RegisterCreateFileForWrite(key, fixedFilePath, true);
 	}
 
+	bool Session::RegisterVirtualFile(const tchar* filePath, const tchar* sourceFile, u64 sourceOffset, u64 sourceSize)
+	{
+		UBA_ASSERT(!m_runningRemote);
+		StringBuffer<> fixedFilePath;
+		auto fileNameKey = GetKeyAndFixedName(fixedFilePath, filePath);
+		if (!RegisterVirtualFileInternal(fileNameKey, fixedFilePath, sourceFile, sourceOffset, sourceSize))
+			return false;
+		return RegisterCreateFileForWrite(fileNameKey, fixedFilePath, false, sourceSize, 0, false);
+	}
+
 	void Session::RegisterDeleteFile(const tchar* filePath)
 	{
 		UBA_ASSERT(!m_runningRemote);
@@ -2483,6 +2493,46 @@ namespace uba
 		logger.Info(TC(""));
 	}
 
+	bool Session::RegisterVirtualFileInternal(const StringKey& fileNameKey, const StringView& filePath, const tchar* sourceFile, u64 sourceOffset, u64 sourceSize)
+	{
+		TimerScope ts(Stats().createMmapFromFile);
+
+		SCOPED_WRITE_LOCK(m_fileMappingTableLookupLock, lookupLock);
+		auto insres = m_fileMappingTableLookup.try_emplace(fileNameKey);
+		FileMappingEntry& entry = insres.first->second;
+		lookupLock.Leave();
+
+		SCOPED_WRITE_LOCK(entry.lock, entryLock);
+		UBA_ASSERT(!entry.handled);
+
+		FileHandle fileHandle = uba::CreateFileW(sourceFile, GENERIC_READ, FILE_SHARE_READ, OPEN_EXISTING, DefaultAttributes());
+		if (fileHandle == InvalidFileHandle)
+			return m_logger.Error(TC("[InternalRegisterVirtualFile] CreateFileW for %s failed (%s)"), sourceFile, LastErrorToText().data);
+		auto fg = MakeGuard([&]() { uba::CloseFile(sourceFile, fileHandle); });
+
+		FileMappingHandle mappingHandle = CreateFileMappingW(fileHandle, PAGE_READONLY, 0, sourceFile);
+		if (!mappingHandle.IsValid())
+			return m_logger.Error(TC("[InternalRegisterVirtualFile] CreateFileMapping for %s failed (%s)"), sourceFile, LastErrorToText().data);
+
+		entry.mapping = mappingHandle; // TODO: This mapping handle will leak until process close.. not good
+		entry.mappingOffset = sourceOffset;
+		entry.size = sourceSize;
+		entry.handled = true;
+
+		StringBuffer<> mappingName;
+		Storage::GetMappingString(mappingName, mappingHandle, sourceOffset);
+		entry.success = true;
+
+		SCOPED_WRITE_LOCK(m_fileMappingTableMemLock, lock);
+		BinaryWriter writer(m_fileMappingTableMem, m_fileMappingTableSize);
+		writer.WriteStringKey(fileNameKey);
+		writer.WriteString(mappingName);
+		writer.Write7BitEncoded(sourceSize);
+		m_fileMappingTableSize = (u32)writer.GetPosition();
+
+		return true;
+	}
+
 	bool Session::CreateProcessJobObject()
 	{
 		#if PLATFORM_WINDOWS
diff --git a/Engine/Source/Programs/UnrealBuildAccelerator/Common/Private/UbaSessionClient.cpp b/Engine/Source/Programs/UnrealBuildAccelerator/Common/Private/UbaSessionClient.cpp
--- a/Engine/Source/Programs/UnrealBuildAccelerator/Common/Private/UbaSessionClient.cpp
+++ b/Engine/Source/Programs/UnrealBuildAccelerator/Common/Private/UbaSessionClient.cpp
@@ -816,7 +816,7 @@ namespace uba
 		out.result = reader.ReadBool();
 		out.errorCode = reader.ReadU32();
 
-		if (out.result)
+		if (out.result || out.errorCode == ERROR_ALREADY_EXISTS) // Even if it exists it might be that this client session is not aware
 			if (!SendUpdateDirectoryTable(reader.Reset()))
 				return false;
 
diff --git a/Engine/Source/Programs/UnrealBuildAccelerator/Common/Private/UbaSessionServer.cpp b/Engine/Source/Programs/UnrealBuildAccelerator/Common/Private/UbaSessionServer.cpp
--- a/Engine/Source/Programs/UnrealBuildAccelerator/Common/Private/UbaSessionServer.cpp
+++ b/Engine/Source/Programs/UnrealBuildAccelerator/Common/Private/UbaSessionServer.cpp
@@ -1559,22 +1559,24 @@ namespace uba
 			}
 		}
 
-		bool fileIsCompressed = false;
-		if (!m_storage.StoreCasFile(out, fileName, casKeyOverride, deferCreation, fileIsCompressed)) // We can defer the creation of the cas file since client might already have it
-			return false;
+		if (casKeyOverride != CasKeyInvalid) // if casKeyOverride is invalid it means it is a virtual file.. look in RegisterVirtualFileInternal
+		{
+			bool fileIsCompressed = false;
+			if (!m_storage.StoreCasFile(out, fileName, casKeyOverride, deferCreation, fileIsCompressed)) // We can defer the creation of the cas file since client might already have it
+				return false;
 
-		if (out != CasKeyZero)
-			return true;
+			if (out != CasKeyZero)
+				return true;
 
-		if (m_shouldWriteToDisk)
-			return true;
+			if (m_shouldWriteToDisk)
+				return true;
+		}
 
 		SCOPED_READ_LOCK(m_fileMappingTableLookupLock, lookupLock);
 		auto findIt = m_fileMappingTableLookup.find(fileNameKey);
 		if (findIt == m_fileMappingTableLookup.end())
 			return true;
 
-		UBA_ASSERT(casKeyOverride == CasKeyZero);
 		FileMappingEntry& entry = findIt->second;
 		lookupLock.Leave();
 		SCOPED_READ_LOCK(entry.lock, entryCs);
@@ -1846,6 +1848,19 @@ namespace uba
 		return str.Appendf(TC("%s"), findIt->second.GetStartInfo().GetDescription()).data;
 	}
 
+	bool SessionServer::RegisterVirtualFileInternal(const StringKey& fileNameKey, const StringView& filePath, const tchar* sourceFile, u64 sourceOffset, u64 sourceSize)
+	{
+		if (!Session::RegisterVirtualFileInternal(fileNameKey, filePath, sourceFile, sourceOffset, sourceSize))
+			return false;
+
+		// Abusing custom cas keys to be able to tell if a file exists only as virtual file (so we can handle StoreCasFile properly)
+		SCOPED_WRITE_LOCK(m_customCasKeysLock, lock);
+		auto insres = m_customCasKeys.try_emplace(fileNameKey);
+		CustomCasKey& customKey = insres.first->second;
+		customKey.casKey = CasKeyInvalid;
+		return true;
+	}
+
 	bool SessionServer::PrepareProcess(ProcessStartInfoHolder& startInfo, bool isChild, StringBufferBase& outRealApplication, const tchar*& outRealWorkingDir)
 	{
 		if (!Session::PrepareProcess(startInfo, isChild, outRealApplication, outRealWorkingDir))
diff --git a/Engine/Source/Programs/UnrealBuildAccelerator/Common/Private/UbaStorageClient.cpp b/Engine/Source/Programs/UnrealBuildAccelerator/Common/Private/UbaStorageClient.cpp
--- a/Engine/Source/Programs/UnrealBuildAccelerator/Common/Private/UbaStorageClient.cpp
+++ b/Engine/Source/Programs/UnrealBuildAccelerator/Common/Private/UbaStorageClient.cpp
@@ -884,9 +884,12 @@ namespace uba
 			}
 			else
 			{
-				TimerScope ts(Stats().sendCas);
+				auto& stats = Stats();
+				TimerScope ts(stats.sendCas);
 				if (!SendFile(m_logger, m_client, casKey, fileMem, fileSize, hint))
 					return false;
+				stats.sendCasBytesRaw += fileSize;
+				stats.sendCasBytesComp += fileSize;
 			}
 		}
 		else
diff --git a/Engine/Source/Programs/UnrealBuildAccelerator/Common/Private/UbaStorageServer.cpp b/Engine/Source/Programs/UnrealBuildAccelerator/Common/Private/UbaStorageServer.cpp
--- a/Engine/Source/Programs/UnrealBuildAccelerator/Common/Private/UbaStorageServer.cpp
+++ b/Engine/Source/Programs/UnrealBuildAccelerator/Common/Private/UbaStorageServer.cpp
@@ -212,11 +212,42 @@ namespace uba
 		}
 	}
 
+	struct FileMappingScope
+	{
+		u8* MapView(Logger& logger, const tchar* func, FileMappingHandle mappingHandle, u64 mappingOffset, u64 fileSize, const tchar* fileName_)
+		{
+			fileName = fileName_;
+			constexpr u64 pageSize = 64*1024;
+			u64 endOffset = mappingOffset + fileSize;
+			u64 alignedOffsetStart = AlignUp(mappingOffset - (pageSize - 1), pageSize);
+			u64 alignedOffsetEnd = AlignUp(endOffset, pageSize);
+			mapSize = alignedOffsetEnd - alignedOffsetStart;
+			mappedMem = MapViewOfFile(mappingHandle, FILE_MAP_READ, alignedOffsetStart, mapSize);
+			if (!mappedMem)
+				mappedMem = MapViewOfFile(mappingHandle, FILE_MAP_READ, alignedOffsetStart, 0);
+			if (mappedMem)
+				return mappedMem + (mappingOffset - alignedOffsetStart);
+			logger.Error(TC("[%s] Failed to map view of file %s at offset %llu and size %llu (%s)"), func, fileName, alignedOffsetStart, mapSize, LastErrorToText().data);
+			return nullptr;
+		}
+
+		~FileMappingScope()
+		{
+			if (mappedMem)
+				UnmapViewOfFile(mappedMem, mapSize, fileName);
+		}
+
+		const tchar* fileName;
+		u8* mappedMem = nullptr;
+		u64 mapSize = 0;
+	};
+
 	bool StorageServer::StoreCasFile(CasKey& out, StringKey fileNameKey, const tchar* fileName, FileMappingHandle mappingHandle, u64 mappingOffset, u64 fileSize, const tchar* hint, bool deferCreation, bool keepMappingInMemory, bool storeCompressed)
 	{
-		u8* fileMem = MapViewOfFile(mappingHandle, FILE_MAP_READ, mappingOffset, fileSize);
-		UBA_ASSERT(fileMem);
-		auto memClose = MakeGuard([&](){ UnmapViewOfFile(fileMem, fileSize, fileName); });
+		FileMappingScope scope;
+		u8* fileMem = scope.MapView(m_logger, TC("StoreCasFile"), mappingHandle, mappingOffset, fileSize, fileName);
+		if (!fileMem)
+			return false;
 
 		out = CalculateCasKey(fileMem, fileSize, storeCompressed);
 		if (out == CasKeyZero)
@@ -257,9 +288,11 @@ namespace uba
 
 		ExternalFileMapping& mapping = findIt->second;
 		lock.Leave();
-		u8* fileMem = MapViewOfFile(mapping.mappingHandle, FILE_MAP_READ, mapping.mappingOffset, mapping.fileSize);
-		UBA_ASSERT(fileMem);
-		auto memClose = MakeGuard([&](){ UnmapViewOfFile(fileMem, mapping.fileSize, from); });
+
+		FileMappingScope scope;
+		u8* fileMem = scope.MapView(m_logger, TC("WriteCompressed"), mapping.mappingHandle, mapping.mappingOffset, mapping.fileSize, from);
+		if (!fileMem)
+			return false;
 		return StorageImpl::WriteCompressed(out, from, InvalidFileHandle, fileMem, mapping.fileSize, toFile, nullptr, 0);
 	}
 
diff --git a/Engine/Source/Programs/UnrealBuildAccelerator/Common/Public/UbaSession.h b/Engine/Source/Programs/UnrealBuildAccelerator/Common/Public/UbaSession.h
--- a/Engine/Source/Programs/UnrealBuildAccelerator/Common/Public/UbaSession.h
+++ b/Engine/Source/Programs/UnrealBuildAccelerator/Common/Public/UbaSession.h
@@ -58,6 +58,7 @@ namespace uba
 		void PrintSummary(Logger& logger); // Print summary stats of session
 		bool RefreshDirectory(const tchar* dirName); // Tell uba a directory on disk has been changed by some other system while session is running
 		bool RegisterNewFile(const tchar* filePath); // Tell uba a new file on disk has been added by some other system while session is running
+		bool RegisterVirtualFile(const tchar* filePath, const tchar* sourceFile, u64 sourceOffset, u64 sourceSize);
 		void RegisterDeleteFile(const tchar* filePath); // Tell uba a file on disk has been deleted by some other system while session is running
 
 		using CustomServiceFunction = Function<u32(Process& handle, const void* recv, u32 recvSize, void* send, u32 sendCapacity)>;
@@ -122,6 +123,8 @@ namespace uba
 		virtual void* GetProcessEnvironmentVariables();
 		virtual void PrintSessionStats(Logger& logger);
 
+		virtual bool RegisterVirtualFileInternal(const StringKey& fileNameKey, const StringView& filePath, const tchar* sourceFile, u64 sourceOffset, u64 sourceSize);
+
 		virtual bool GetInitResponse(InitResponse& out, const InitMessage& msg);
 		virtual bool CreateFile(CreateFileResponse& out, const CreateFileMessage& msg);
 		virtual bool CloseFile(CloseFileResponse& out, const CloseFileMessage& msg);
diff --git a/Engine/Source/Programs/UnrealBuildAccelerator/Common/Public/UbaSessionServer.h b/Engine/Source/Programs/UnrealBuildAccelerator/Common/Public/UbaSessionServer.h
--- a/Engine/Source/Programs/UnrealBuildAccelerator/Common/Public/UbaSessionServer.h
+++ b/Engine/Source/Programs/UnrealBuildAccelerator/Common/Public/UbaSessionServer.h
@@ -91,6 +91,8 @@ namespace uba
 
 		TString GetProcessDescription(u32 processId);
 
+		virtual bool RegisterVirtualFileInternal(const StringKey& fileNameKey, const StringView& filePath, const tchar* sourceFile, u64 sourceOffset, u64 sourceSize) override;
+
 		virtual bool PrepareProcess(ProcessStartInfoHolder& startInfo, bool isChild, StringBufferBase& outRealApplication, const tchar*& outRealWorkingDir) override final;
 		virtual bool CreateFile(CreateFileResponse& out, const CreateFileMessage& msg) override final;
 		virtual void FileEntryAdded(StringKey fileNameKey, u64 lastWritten, u64 size) override final;
diff --git a/Engine/Source/Programs/UnrealBuildAccelerator/Core/Private/UbaApplicationRules.cpp b/Engine/Source/Programs/UnrealBuildAccelerator/Core/Private/UbaApplicationRules.cpp
--- a/Engine/Source/Programs/UnrealBuildAccelerator/Core/Private/UbaApplicationRules.cpp
+++ b/Engine/Source/Programs/UnrealBuildAccelerator/Core/Private/UbaApplicationRules.cpp
@@ -35,13 +35,13 @@ namespace uba
 		{
 			return fileName.Contains(TC("vctip_")) || Super::IsThrowAway(fileName, isRunningRemote);
 		}
-		virtual bool KeepInMemory(const StringView& fileName, const tchar* systemTemp, bool isRunningRemote) const override
+		virtual bool KeepInMemory(const StringView& fileName, const tchar* systemTemp, bool isRunningRemote, bool isWrite) const override
 		{
 			if (fileName.Contains(TC("\\vctip_")))
 				return true;
 			if (fileName.Contains(systemTemp))
 				return true;
-			return Super::KeepInMemory(fileName, systemTemp, isRunningRemote);
+			return Super::KeepInMemory(fileName, systemTemp, isRunningRemote, isWrite);
 		}
 
 		virtual bool IsExitCodeSuccess(u32 exitCode) const override
@@ -184,7 +184,7 @@ namespace uba
 	{
 		using Super = ApplicationRulesVcLink;
 
-		virtual bool KeepInMemory(const StringView& fileName, const tchar* systemTemp, bool isRunningRemote) const override
+		virtual bool KeepInMemory(const StringView& fileName, const tchar* systemTemp, bool isRunningRemote, bool isWrite) const override
 		{
 			return fileName.EndsWith(TC(".manifest")) && fileName.Contains(systemTemp);
 		}
@@ -344,9 +344,9 @@ namespace uba
 		//{
 		//	return fileName.EndsWith(TC(".self")) || Equals(file, TC("Symbols.map"));
 		//}
-		virtual bool KeepInMemory(const StringView& fileName, const tchar* systemTemp, bool isRunningRemote) const override
+		virtual bool KeepInMemory(const StringView& fileName, const tchar* systemTemp, bool isRunningRemote, bool isWrite) const override
 		{
-			return Super::KeepInMemory(fileName, systemTemp, isRunningRemote)
+			return Super::KeepInMemory(fileName, systemTemp, isRunningRemote, isWrite)
 				|| fileName.Contains(TC("thinlto-"));// Used by a clang based platform's link time optimization pass. Shared from lto process back to linker process
 		}
 		virtual bool NeedsSharedMemory(const tchar* file) const override
@@ -420,6 +420,21 @@ namespace uba
 		{
 			return !fileName.EndsWith(TC(".split"));
 		}
+
+		virtual bool KeepInMemory(const StringView& fileName, const tchar* systemTemp, bool isRunningRemote, bool isWrite) const override
+		{
+			if (isWrite && fileName.Contains(TC("\\Split\\")) && fileName.EndsWith(TC(".split")))
+				return true;
+			return false;
+		}
+
+		virtual bool IsThrowAway(const StringView& fileName, bool isRunningRemote) const override
+		{
+			if (fileName.Contains(TC("\\Split\\")) && fileName.EndsWith(TC(".split")))
+				return true;
+			return Super::IsThrowAway(fileName, isRunningRemote);
+		}
+
 	};
 
 	// ====
diff --git a/Engine/Source/Programs/UnrealBuildAccelerator/Core/Public/UbaApplicationRules.h b/Engine/Source/Programs/UnrealBuildAccelerator/Core/Public/UbaApplicationRules.h
--- a/Engine/Source/Programs/UnrealBuildAccelerator/Core/Public/UbaApplicationRules.h
+++ b/Engine/Source/Programs/UnrealBuildAccelerator/Core/Public/UbaApplicationRules.h
@@ -41,7 +41,7 @@ namespace uba
 
 		// Keep file in memory
 		// If this returns true it means that file will be kept in memory and never touch disk.
-		virtual bool KeepInMemory(const StringView& fileName, const tchar* systemTemp, bool isRunningRemote) const
+		virtual bool KeepInMemory(const StringView& fileName, const tchar* systemTemp, bool isRunningRemote, bool isWrite) const
 		{
 			return IsThrowAway(fileName, isRunningRemote);
 		}
diff --git a/Engine/Source/Programs/UnrealBuildAccelerator/Detours/Private/UbaDetoursShared.cpp b/Engine/Source/Programs/UnrealBuildAccelerator/Detours/Private/UbaDetoursShared.cpp
--- a/Engine/Source/Programs/UnrealBuildAccelerator/Detours/Private/UbaDetoursShared.cpp
+++ b/Engine/Source/Programs/UnrealBuildAccelerator/Detours/Private/UbaDetoursShared.cpp
@@ -274,7 +274,7 @@ namespace uba
 
 		memset(&outAttr.data, 0, sizeof(outAttr.data));
 
-		bool keepInMemory = KeepInMemory(StringView(fileName, fileNameForKey.count));
+		bool keepInMemory = KeepInMemory(StringView(fileName, fileNameForKey.count), false);
 		if (keepInMemory)
 		{
 			SCOPED_READ_LOCK(g_mappedFileTable.m_lookupLock, lock);
diff --git a/Engine/Source/Programs/UnrealBuildAccelerator/Detours/Private/UbaDetoursShared.h b/Engine/Source/Programs/UnrealBuildAccelerator/Detours/Private/UbaDetoursShared.h
--- a/Engine/Source/Programs/UnrealBuildAccelerator/Detours/Private/UbaDetoursShared.h
+++ b/Engine/Source/Programs/UnrealBuildAccelerator/Detours/Private/UbaDetoursShared.h
@@ -114,7 +114,7 @@ namespace uba
 	};
 
 	inline bool CanDetour(const tchar* file) { return !t_disallowDetour && g_rules->CanDetour(file); }
-	inline bool KeepInMemory(const StringView& fileName) { return g_allowKeepFilesInMemory && g_rules->KeepInMemory(fileName, g_systemTemp.data, g_runningRemote); }
+	inline bool KeepInMemory(const StringView& fileName, bool isWrite) { return g_allowKeepFilesInMemory && g_rules->KeepInMemory(fileName, g_systemTemp.data, g_runningRemote, isWrite); }
 
 	void Shared_WriteConsole(const char* chars, u32 charCount, bool isError);
 	void Shared_WriteConsole(const wchar_t* chars, u32 charCount, bool isError);
diff --git a/Engine/Source/Programs/UnrealBuildAccelerator/Detours/Private/Windows/UbaDetoursFunctionsKernelBase.inl b/Engine/Source/Programs/UnrealBuildAccelerator/Detours/Private/Windows/UbaDetoursFunctionsKernelBase.inl
--- a/Engine/Source/Programs/UnrealBuildAccelerator/Detours/Private/Windows/UbaDetoursFunctionsKernelBase.inl
+++ b/Engine/Source/Programs/UnrealBuildAccelerator/Detours/Private/Windows/UbaDetoursFunctionsKernelBase.inl
@@ -381,6 +381,8 @@ BOOL Detoured_ReadFile(HANDLE hFile, LPVOID lpBuffer, DWORD nNumberOfBytesToRead
 				DEBUG_LOG_DETOURED(L"ReadFile", L"%llu %u (%ls) -> FAILED TO MAP", uintptr_t(hFile), nNumberOfBytesToRead, HandleToName(hFile));
 				return FALSE;
 			}
+			UBA_ASSERT(fi.fileMapMem || !fi.memoryFile->isThrowAway);
+
 			u8* mem = fi.fileMapMem ? fi.fileMapMem : fi.memoryFile->baseAddress;
 			u64 size = fi.fileMapMem ? fi.fileMapMemEnd - mem : fi.memoryFile->writtenSize;
 			UBA_ASSERTF(dh.pos <= size, L"Filepointer is higher than size of file (pointer: %llu, size: %llu) (%ls)", dh.pos, size, HandleToName(hFile));
@@ -1064,7 +1066,7 @@ DWORD Detoured_GetFileAttributesW(LPCWSTR lpFileName)
 BOOL Detoured_SetFileAttributesW(LPCWSTR lpFileName, DWORD dwFileAttributes)
 {
 	DETOURED_CALL(SetFileAttributesW);
-	if (KeepInMemory(StringView(lpFileName, u32(wcslen(lpFileName)))))
+	if (KeepInMemory(StringView(lpFileName, u32(wcslen(lpFileName))), true))
 	{
 		DEBUG_LOG_DETOURED(L"SetFileAttributesW", L"(%ls) %u", lpFileName, dwFileAttributes);
 		SetLastError(ERROR_SUCCESS);
@@ -1374,7 +1376,7 @@ BOOL Detoured_DeleteFileW(LPCWSTR lpFileName)
 		return True_DeleteFileW(original);
 	}
 
-	if (KeepInMemory(fixedNameLower))
+	if (KeepInMemory(fixedNameLower, true))
 	{
 		DEBUG_LOG_DETOURED(L"DeleteFileW", L"(INMEMORY) (%ls) -> Success", lpFileName);
 		SetLastError(ERROR_SUCCESS);
@@ -1429,7 +1431,7 @@ bool Shared_MoveFile(LPCWSTR lpExistingFileName, LPCWSTR lpNewFileName, DWORD dw
 
 	StringKey sourceKey = ToStringKeyLower(source);
 
-	if (KeepInMemory(source))
+	if (KeepInMemory(source, false))
 	{
 		SCOPED_WRITE_LOCK(g_mappedFileTable.m_lookupLock, lock);
 		auto it = g_mappedFileTable.m_lookup.find(sourceKey);
@@ -1456,7 +1458,7 @@ bool Shared_MoveFile(LPCWSTR lpExistingFileName, LPCWSTR lpNewFileName, DWORD dw
 			return true;
 		}
 
-		UBA_ASSERT(!KeepInMemory(dest));
+		UBA_ASSERT(!KeepInMemory(dest, true));
 
 		DEBUG_LOG_DETOURED(L"MoveFileExW", L"(memfile->file) %ls to %ls", lpExistingFileName, lpNewFileName);
 
diff --git a/Engine/Source/Programs/UnrealBuildAccelerator/Detours/Private/Windows/UbaDetoursFunctionsNtDll.inl b/Engine/Source/Programs/UnrealBuildAccelerator/Detours/Private/Windows/UbaDetoursFunctionsNtDll.inl
--- a/Engine/Source/Programs/UnrealBuildAccelerator/Detours/Private/Windows/UbaDetoursFunctionsNtDll.inl
+++ b/Engine/Source/Programs/UnrealBuildAccelerator/Detours/Private/Windows/UbaDetoursFunctionsNtDll.inl
@@ -689,7 +689,7 @@ NTSTATUS NTAPI Shared_NtCreateFile(bool IsCreateFunc, PHANDLE hFileHandle, ACCES
 	bool useContent = IsContentUse(DesiredAccess, CreateDisposition);
 	bool isWrite = IsWrite(DesiredAccess, CreateDisposition);
 	//bool keepInMemory = (KeepInMemory(fileName) && (useContent || NeedsSharedMemory(fileName.data))) || IsOutputFile(fileName, isWrite, isDeleteOnClose);
-	bool keepInMemory = KeepInMemory(fileName) || IsOutputFile(fileName, isWrite, isDeleteOnClose);
+	bool keepInMemory = KeepInMemory(fileName, isWrite) || IsOutputFile(fileName, isWrite, isDeleteOnClose);
 
 #if UBA_DEBUG_LOG_ENABLED
 	const wchar_t* isWriteStr = isWrite ? L" WRITE" : L""; (void)isWriteStr;
@@ -1100,11 +1100,8 @@ NTSTATUS NTAPI Shared_NtCreateFile(bool IsCreateFunc, PHANDLE hFileHandle, ACCES
 			}
 			else
 			{
-				if (g_rules->IsThrowAway(fileName, g_runningRemote))
-				{
-					//isDeleteOnClose = true;
-				}
-				else if (CreateDisposition == FILE_OPEN)
+				bool isThrowAway = g_rules->IsThrowAway(fileName, g_runningRemote);
+				if (!isThrowAway && CreateDisposition == FILE_OPEN)
 				{
 					*hFileHandle = INVALID_HANDLE_VALUE;
 					DEBUG_LOG_DETOURED(funcName, L"ALREADYEXISTS (%ls) -> Error", fileName.data);
@@ -1113,7 +1110,7 @@ NTSTATUS NTAPI Shared_NtCreateFile(bool IsCreateFunc, PHANDLE hFileHandle, ACCES
 
 				bool isLocal = !IsOutputFile(fileName, isWrite, isDeleteOnClose);
 				//UBA_ASSERTF(CreateDisposition != FILE_OPEN || Contains(fileName.data, L"vctip_"), TC("Unsupported disposition %u for file %s"), CreateDisposition, fileName.data);
-				info.memoryFile = new MemoryFile(isLocal, FileTypeMaxSize(fileName, isSystemOrTempFile));
+				info.memoryFile = new MemoryFile(isLocal, FileTypeMaxSize(fileName, isSystemOrTempFile), isThrowAway);
 			}
 
 			// TODO: Time should be in sync with host machine!
diff --git a/Engine/Source/Programs/UnrealBuildAccelerator/Detours/Private/Windows/UbaDetoursFunctionsWin.cpp b/Engine/Source/Programs/UnrealBuildAccelerator/Detours/Private/Windows/UbaDetoursFunctionsWin.cpp
--- a/Engine/Source/Programs/UnrealBuildAccelerator/Detours/Private/Windows/UbaDetoursFunctionsWin.cpp
+++ b/Engine/Source/Programs/UnrealBuildAccelerator/Detours/Private/Windows/UbaDetoursFunctionsWin.cpp
@@ -247,6 +247,13 @@ BlockAllocator<FileObject> g_fileObjectAllocator(g_memoryBlock);
 void MemoryFile::Write(DetouredHandle& handle, LPCVOID lpBuffer, u64 nNumberOfBytesToWrite)
 {
 	u64 newPos = handle.pos + nNumberOfBytesToWrite;
+
+	if (isThrowAway)
+	{
+		writtenSize = newPos;
+		return;
+	}
+
 	EnsureCommitted(handle, newPos);
 	memcpy(baseAddress + handle.pos, lpBuffer, nNumberOfBytesToWrite);
 	handle.pos += nNumberOfBytesToWrite;
@@ -259,6 +266,9 @@ void MemoryFile::Write(DetouredHandle& handle, LPCVOID lpBuffer, u64 nNumberOfBy
 
 void MemoryFile::EnsureCommitted(const DetouredHandle& handle, u64 size)
 {
+	if (isThrowAway)
+		return;
+
 	if (committedSize >= size)
 		return;
 	if (size > mappedSize)
@@ -398,6 +408,11 @@ bool EnsureMapped(DetouredHandle& handle, DWORD dwFileOffsetHigh = 0, DWORD dwFi
 		u64 mapSize = alignedOffsetEnd - alignedOffsetStart;
 		TimerScope ts(g_kernelStats.mapViewOfFile);
 		info.fileMapMem = (u8*)True_MapViewOfFileEx(info.trueFileMapHandle, info.fileMapViewDesiredAccess, ToHigh(alignedOffsetStart), ToLow(alignedOffsetStart), mapSize, baseAddress);
+		
+		// TODO: This is ugly.. but in some cases we have created virtual files pointing into segments of real file
+		// and real file size is not aligned to page size so we need to set map size to 0 to allow it to map up to size of file
+		if (!info.fileMapMem)
+			info.fileMapMem = (u8*)True_MapViewOfFileEx(info.trueFileMapHandle, info.fileMapViewDesiredAccess, ToHigh(alignedOffsetStart), ToLow(alignedOffsetStart), 0, baseAddress);
 	}
 	else
 	{
diff --git a/Engine/Source/Programs/UnrealBuildAccelerator/Detours/Private/Windows/UbaDetoursUtilsWin.h b/Engine/Source/Programs/UnrealBuildAccelerator/Detours/Private/Windows/UbaDetoursUtilsWin.h
--- a/Engine/Source/Programs/UnrealBuildAccelerator/Detours/Private/Windows/UbaDetoursUtilsWin.h
+++ b/Engine/Source/Programs/UnrealBuildAccelerator/Detours/Private/Windows/UbaDetoursUtilsWin.h
@@ -56,13 +56,15 @@ namespace uba
 	struct MemoryFile
 	{
 		MemoryFile(u8* data = nullptr, bool localOnly = true) : baseAddress(data), isLocalOnly(localOnly) {}
-		MemoryFile(bool localOnly, u64 reserveSize_) : isLocalOnly(localOnly)
+		MemoryFile(bool localOnly, u64 reserveSize_, bool isThrowAway_ = false) : isLocalOnly(localOnly), isThrowAway(isThrowAway_)
 		{
-			Reserve(reserveSize_);
+			if (!isThrowAway)
+				Reserve(reserveSize_);
 		}
 
 		void Reserve(u64 reserveSize_)
 		{
+			UBA_ASSERT(!isThrowAway);
 			reserveSize = reserveSize_;
 			if (isLocalOnly)
 			{
@@ -119,5 +121,6 @@ namespace uba
 		u64 writtenSize = 0;
 		bool isLocalOnly;
 		bool isReported = false;
+		bool isThrowAway = false;
 	};
 }
\ No newline at end of file
diff --git a/Engine/Source/Programs/UnrealBuildAccelerator/Host/Private/UbaExports.cpp b/Engine/Source/Programs/UnrealBuildAccelerator/Host/Private/UbaExports.cpp
--- a/Engine/Source/Programs/UnrealBuildAccelerator/Host/Private/UbaExports.cpp
+++ b/Engine/Source/Programs/UnrealBuildAccelerator/Host/Private/UbaExports.cpp
@@ -377,6 +377,10 @@ extern "C"
 	{
 		server->RegisterDeleteFile(filePath);
 	}
+	bool SessionServer_RegisterVirtualFile(uba::SessionServer* server, const uba::tchar* filename, const uba::tchar* sourceFile, uba::u64 sourceOffset, uba::u64 sourceSize)
+	{
+		return server->RegisterVirtualFile(filename, sourceFile, sourceOffset, sourceSize);
+	}
 	uba::ProcessHandle* SessionServer_RunProcess(uba::SessionServer* server, uba::ProcessStartInfo& info, bool async, bool enableDetour)
 	{
 		return new uba::ProcessHandle(server->RunProcess(info, async, enableDetour));
diff --git a/Engine/Source/Programs/UnrealBuildAccelerator/Host/Public/UbaExports.h b/Engine/Source/Programs/UnrealBuildAccelerator/Host/Public/UbaExports.h
--- a/Engine/Source/Programs/UnrealBuildAccelerator/Host/Public/UbaExports.h
+++ b/Engine/Source/Programs/UnrealBuildAccelerator/Host/Public/UbaExports.h
@@ -101,6 +101,7 @@ extern "C"
 	UBA_API bool SessionServer_RefreshDirectory(uba::SessionServer* server, const uba::tchar* directory);
 	UBA_API bool SessionServer_RegisterNewFile(uba::SessionServer* server, const uba::tchar* filePath);
 	UBA_API void SessionServer_RegisterDeleteFile(uba::SessionServer* server, const uba::tchar* filePath);
+	UBA_API bool SessionServer_RegisterVirtualFile(uba::SessionServer* server, const uba::tchar* filename, const uba::tchar* sourceFile, uba::u64 sourceOffset, uba::u64 sourceSize);
 	UBA_API uba::ProcessHandle* SessionServer_RunProcess(uba::SessionServer* server, uba::ProcessStartInfo& info, bool async, bool enableDetour);
 	UBA_API uba::ProcessHandle* SessionServer_RunProcessRemote(uba::SessionServer* server, uba::ProcessStartInfo& info, float weight, const void* knownInputs = nullptr, uba::u32 knownInputsCount = 0);
 	UBA_API uba::ProcessHandle* SessionServer_RunProcessRacing(uba::SessionServer* server, uba::u32 raceAgainstRemoteProcessId);
