[UBA]
* Fixed posix compile errors

diff --git a/Engine/Source/Programs/UnrealBuildAccelerator/Core/Private/UbaApplicationRules.cpp b/Engine/Source/Programs/UnrealBuildAccelerator/Core/Private/UbaApplicationRules.cpp
--- a/Engine/Source/Programs/UnrealBuildAccelerator/Core/Private/UbaApplicationRules.cpp
+++ b/Engine/Source/Programs/UnrealBuildAccelerator/Core/Private/UbaApplicationRules.cpp
@@ -592,9 +592,9 @@ namespace uba
 			//{ L"BreakpadSymbolEncoder.exe"),	new ApplicationRulesClang() },
 			//{ L"cmd.exe"),		new ApplicationRules() },
 #else
-			{ ApplicationHash(TC("clang++"),					new ApplicationRulesClangPlusPlusExe() },
-			{ ApplicationHash(TC("ld.lld"),						new ApplicationRulesLdLLdExe() },
-			{ ApplicationHash(TC("shadercompileworker"),		new ApplicationRulesShaderCompileWorker() },
+			{ ApplicationHash(TC("clang++")),					new ApplicationRulesClangPlusPlusExe() },
+			{ ApplicationHash(TC("ld.lld")),					new ApplicationRulesLdLLdExe() },
+			{ ApplicationHash(TC("shadercompileworker")),		new ApplicationRulesShaderCompileWorker() },
 #endif
 			{ 0, nullptr }
 		};
