[UBA]
* Bumped directory mapped memory size from 48mb to 64mb because we are hitting the ceiling in one of our targets.. insane!!

diff --git a/Engine/Source/Programs/UnrealBuildAccelerator/Core/Public/UbaProtocol.h b/Engine/Source/Programs/UnrealBuildAccelerator/Core/Public/UbaProtocol.h
--- a/Engine/Source/Programs/UnrealBuildAccelerator/Core/Public/UbaProtocol.h
+++ b/Engine/Source/Programs/UnrealBuildAccelerator/Core/Public/UbaProtocol.h
@@ -52,7 +52,7 @@ namespace uba
 	inline constexpr u32 CommunicationMemSize = IsWindows ? 64*1024 : 64*1024*2; // Macos expands some commandlines to be crazy long
 
 	inline constexpr u32 FileMappingTableMemSize = 16 * 1024 * 1024;
-	inline constexpr u32 DirTableMemSize = 48 * 1024 * 1024;
+	inline constexpr u32 DirTableMemSize = 64 * 1024 * 1024;
 }
 
 // Currently only used for detoured process
