[UBA]
* Changed directory table overflow test code from assert to real error

diff --git a/Engine/Source/Programs/UnrealBuildAccelerator/Common/Private/UbaSession.cpp b/Engine/Source/Programs/UnrealBuildAccelerator/Common/Private/UbaSession.cpp
--- a/Engine/Source/Programs/UnrealBuildAccelerator/Common/Private/UbaSession.cpp
+++ b/Engine/Source/Programs/UnrealBuildAccelerator/Common/Private/UbaSession.cpp
@@ -353,7 +353,7 @@ namespace uba
 		u32 filesOffset = writePos + u32(tableWriter.GetPosition());
 		tableWriter.WriteBytes(memoryBlock.data(), written);
 		dirTable.m_memorySize += u32(tableWriter.GetPosition());
-		UBA_ASSERT(dirTable.m_memorySize < DirTableMemSize);
+		CheckDirectoryTableOverflow(dirTable);
 		
 		memoryLock.Leave();
 
@@ -855,7 +855,7 @@ namespace uba
 		writer.WriteBytes(temp, written);
 		dir.tableOffset = dirTable.m_memorySize + tableOffset;
 		dirTable.m_memorySize += u32(writer.GetPosition());
-		UBA_ASSERT(dirTable.m_memorySize < DirTableMemSize);
+		CheckDirectoryTableOverflow(dirTable);
 		return true;
 	}
 
@@ -930,9 +930,16 @@ namespace uba
 		writer.WriteBytes(temp, written);
 		dir.tableOffset = dirTable.m_memorySize + tableOffset;
 		dirTable.m_memorySize += u32(writer.GetPosition());
-		UBA_ASSERT(dirTable.m_memorySize < DirTableMemSize);
+		CheckDirectoryTableOverflow(dirTable);
 		return dirTable.m_memorySize;
 	}
+	
+	void Session::CheckDirectoryTableOverflow(DirectoryTable& dirTable)
+	{
+		if (dirTable.m_memorySize < DirTableMemSize - 1024*1024)
+			return;
+		static bool called = m_logger.Error(TC("Directory table overflow. DirTableMemSize need to be increased (Size: %llu)"), dirTable.m_memorySize);
+	}
 
 	bool Session::CopyImports(Vector<BinaryModule>& out, const tchar* library, tchar* applicationDir, tchar* applicationDirEnd, UnorderedSet<TString>& handledImports, const char* const* loaderPaths)
 	{
diff --git a/Engine/Source/Programs/UnrealBuildAccelerator/Common/Public/UbaSession.h b/Engine/Source/Programs/UnrealBuildAccelerator/Common/Public/UbaSession.h
--- a/Engine/Source/Programs/UnrealBuildAccelerator/Common/Public/UbaSession.h
+++ b/Engine/Source/Programs/UnrealBuildAccelerator/Common/Public/UbaSession.h
@@ -118,6 +118,7 @@ namespace uba
 
 		bool RegisterCreateFileForWrite(StringKey fileNameKey, const StringView& fileName, bool registerRealFile, u64 fileSize = 0, u64 lastWriteTime = 0, bool invalidateStorage = true);
 		u32 RegisterDeleteFile(StringKey fileNameKey, const StringView& fileName);
+		void CheckDirectoryTableOverflow(DirectoryTable& dirTable);
 
 		virtual bool PrepareProcess(ProcessStartInfoHolder& startInfo, bool isChild, StringBufferBase& outRealApplication, const tchar*& outRealWorkingDir);
 		virtual void* GetProcessEnvironmentVariables();
