[UBA]
* Fixed so dump_syms.exe also preloads bcrypt dlls to prevent deadlocks when mimalloc loads them

diff --git a/Engine/Source/Programs/UnrealBuildAccelerator/Core/Private/UbaApplicationRules.cpp b/Engine/Source/Programs/UnrealBuildAccelerator/Core/Private/UbaApplicationRules.cpp
--- a/Engine/Source/Programs/UnrealBuildAccelerator/Core/Private/UbaApplicationRules.cpp
+++ b/Engine/Source/Programs/UnrealBuildAccelerator/Core/Private/UbaApplicationRules.cpp
@@ -319,6 +319,19 @@ namespace uba
 		{
 			return false; // fileName.EndsWith(TC(".psym")); With psym as output file the BreakpadSymbolEncoder fails to output a .sym file
 		}
+
+		virtual const tchar* const* LibrariesToPreload() const override
+		{
+			// Special handling.. it seems loading bcrypt.dll can deadlock when using mimalloc so we make sure to load it here directly instead
+			// There is a setting to disable bcrypt dll loading inside mimalloc but with that change mimalloc does not work with older versions of windows
+			static constexpr const tchar* preloads[] = 
+			{
+				TC("bcrypt.dll"),
+				TC("bcryptprimitives.dll"),
+				nullptr,
+			};
+			return preloads;
+		}
 	};
 
 	class ApplicationRulesClangPlusPlusExePlatform1 : public ApplicationRulesClangPlusPlusExe
