[UBA]
* Added mapping handle cache for source files used by virtual files. Also added cleanup and closing of file mappings at shutdown

diff --git a/Engine/Source/Programs/UnrealBuildAccelerator/Common/Private/UbaSession.cpp b/Engine/Source/Programs/UnrealBuildAccelerator/Common/Private/UbaSession.cpp
--- a/Engine/Source/Programs/UnrealBuildAccelerator/Common/Private/UbaSession.cpp
+++ b/Engine/Source/Programs/UnrealBuildAccelerator/Common/Private/UbaSession.cpp
@@ -1205,6 +1205,9 @@ namespace uba
 		CancelAllProcessesAndWait();
 		FlushDeadProcesses();
 
+		for (auto& kv : m_virtualSourceFiles)
+			CloseFileMapping(kv.second);
+
 		//for (auto& i : m_fileMappingTableLookup)
 		//	CloseHandle(i.second.mapping);
 
@@ -2504,22 +2507,35 @@ namespace uba
 	{
 		TimerScope ts(Stats().createMmapFromFile);
 
+		FileMappingHandle mappingHandle;
+
+		{
+			StringKey sourceFileKey = CaseInsensitiveFs ? ToStringKeyLower(ToView(sourceFile)) : ToStringKey(ToView(sourceFile));
+			SCOPED_WRITE_LOCK(m_virtualSourceFilesLock, virtualSourceFilesLock);
+			auto insres = m_virtualSourceFiles.try_emplace(sourceFileKey);
+			if (insres.second)
+			{
+				FileHandle fileHandle = uba::CreateFileW(sourceFile, GENERIC_READ, FILE_SHARE_READ, OPEN_EXISTING, DefaultAttributes());
+				if (fileHandle == InvalidFileHandle)
+					return m_logger.Error(TC("[InternalRegisterVirtualFile] CreateFileW for %s failed (%s)"), sourceFile, LastErrorToText().data);
+				auto fg = MakeGuard([&]() { uba::CloseFile(sourceFile, fileHandle); });
+				insres.first->second = CreateFileMappingW(fileHandle, PAGE_READONLY, 0, sourceFile);
+			}
+			mappingHandle = insres.first->second;
+			if (!mappingHandle.IsValid())
+				return m_logger.Error(TC("[InternalRegisterVirtualFile] CreateFileMapping for %s failed (%s)"), sourceFile, LastErrorToText().data);
+		}
+
+
 		SCOPED_WRITE_LOCK(m_fileMappingTableLookupLock, lookupLock);
 		auto insres = m_fileMappingTableLookup.try_emplace(fileNameKey);
 		FileMappingEntry& entry = insres.first->second;
 		lookupLock.Leave();
 
 		SCOPED_WRITE_LOCK(entry.lock, entryLock);
-		UBA_ASSERT(!entry.handled);
-
-		FileHandle fileHandle = uba::CreateFileW(sourceFile, GENERIC_READ, FILE_SHARE_READ, OPEN_EXISTING, DefaultAttributes());
-		if (fileHandle == InvalidFileHandle)
-			return m_logger.Error(TC("[InternalRegisterVirtualFile] CreateFileW for %s failed (%s)"), sourceFile, LastErrorToText().data);
-		auto fg = MakeGuard([&]() { uba::CloseFile(sourceFile, fileHandle); });
-
-		FileMappingHandle mappingHandle = CreateFileMappingW(fileHandle, PAGE_READONLY, 0, sourceFile);
-		if (!mappingHandle.IsValid())
-			return m_logger.Error(TC("[InternalRegisterVirtualFile] CreateFileMapping for %s failed (%s)"), sourceFile, LastErrorToText().data);
+		
+		if (entry.handled)
+			return m_logger.Error(TC("Virtual file %s has already been registered"), filePath.data);
 
 		entry.mapping = mappingHandle; // TODO: This mapping handle will leak until process close.. not good
 		entry.mappingOffset = sourceOffset;
diff --git a/Engine/Source/Programs/UnrealBuildAccelerator/Common/Public/UbaSession.h b/Engine/Source/Programs/UnrealBuildAccelerator/Common/Public/UbaSession.h
--- a/Engine/Source/Programs/UnrealBuildAccelerator/Common/Public/UbaSession.h
+++ b/Engine/Source/Programs/UnrealBuildAccelerator/Common/Public/UbaSession.h
@@ -239,6 +239,9 @@ namespace uba
 		struct ActiveFile { TString name; StringKey nameKey; };
 		UnorderedMap<u32, ActiveFile> m_activeFiles;
 
+		ReaderWriterLock m_virtualSourceFilesLock;
+		UnorderedMap<StringKey, FileMappingHandle> m_virtualSourceFiles;
+
 		u32 m_wantsOnCloseIdCounter = 1;
 
 		SessionStats m_stats;
diff --git a/Engine/Source/Programs/UnrealBuildAccelerator/Core/Private/UbaHash.cpp b/Engine/Source/Programs/UnrealBuildAccelerator/Core/Private/UbaHash.cpp
--- a/Engine/Source/Programs/UnrealBuildAccelerator/Core/Private/UbaHash.cpp
+++ b/Engine/Source/Programs/UnrealBuildAccelerator/Core/Private/UbaHash.cpp
@@ -29,7 +29,7 @@ namespace uba
 		return (StringKey&)output;
 	}
 
-	StringKey ToStringKey(const StringBufferBase& b)
+	StringKey ToStringKey(const StringView& b)
 	{
 		return ToStringKey(b.data, b.count);
 	}
@@ -41,7 +41,7 @@ namespace uba
 		return ToStringKey(temp.data, temp.count);
 	}
 
-	StringKey ToStringKeyLower(const StringBufferBase& b)
+	StringKey ToStringKeyLower(const StringView& b)
 	{
 		return ToStringKeyLower(b.data, b.count);
 	}
diff --git a/Engine/Source/Programs/UnrealBuildAccelerator/Core/Public/UbaHash.h b/Engine/Source/Programs/UnrealBuildAccelerator/Core/Public/UbaHash.h
--- a/Engine/Source/Programs/UnrealBuildAccelerator/Core/Public/UbaHash.h
+++ b/Engine/Source/Programs/UnrealBuildAccelerator/Core/Public/UbaHash.h
@@ -102,8 +102,8 @@ namespace uba
 
 	StringKey ToStringKey(const tchar* str, u64 strLen);
 	StringKey ToStringKeyLower(const tchar* str, u64 strLen);
-	StringKey ToStringKey(const StringBufferBase& b);
-	StringKey ToStringKeyLower(const StringBufferBase& b);
+	StringKey ToStringKey(const StringView& b);
+	StringKey ToStringKeyLower(const StringView& b);
 	StringKey ToStringKey(const StringKeyHasher& hasher, const tchar* str, u64 strLen);
 	StringKey ToStringKey(const StringKeyHasher& hasher);
 	StringKey ToStringKeyNoCheck(const tchar* str, u64 strLen);
