[UBA]
* Improved logging when getting bad paths from PATH environment variable

diff --git a/Engine/Source/Programs/UnrealBuildAccelerator/Common/Private/UbaFile.cpp b/Engine/Source/Programs/UnrealBuildAccelerator/Common/Private/UbaFile.cpp
--- a/Engine/Source/Programs/UnrealBuildAccelerator/Common/Private/UbaFile.cpp
+++ b/Engine/Source/Programs/UnrealBuildAccelerator/Common/Private/UbaFile.cpp
@@ -379,7 +379,7 @@ namespace uba
 		fullPath.EnsureEndsWithSlash();
 		fullPath.Append(file);
 
-		auto TestFileExists = [&]()
+		auto TestFileExists = [&](const tchar* extraInfo)
 			{
 				if (GetFileAttributesW(fullPath.data) != INVALID_FILE_ATTRIBUTES)
 				{
@@ -388,11 +388,11 @@ namespace uba
 				}
 				u32 lastError = GetLastError();
 				if (lastError != ERROR_FILE_NOT_FOUND && lastError != ERROR_PATH_NOT_FOUND)
-					logger.Warning(TC("SearchPathForFile found the file %s but got error when getting attributes ()"), fullPath.data, LastErrorToText(lastError).data);
+					logger.Warning(TC("SearchPathForFile tried to find the file %s%s but got error when getting attributes (%s)"), fullPath.data, extraInfo, LastErrorToText(lastError).data);
 				return false;
 			};
 
-		if (TestFileExists())
+		if (TestFileExists(TC("")))
 			return true;
 
 		char varSeparator = IsWindows ? ';' : ':';
@@ -422,7 +422,7 @@ namespace uba
 				fullPath.EnsureEndsWithSlash();
 			fullPath.Append(file);
 
-			if (TestFileExists())
+			if (TestFileExists(TC(" using PATH environment variable")))
 				return true;
 
 			lastStart = ++it;
