[UBA]
* Removed detouring of _isatty on arm64 because it is failing

diff --git a/Engine/Source/Programs/UnrealBuildAccelerator/Detours/Private/Windows/UbaDetoursFunctionsUcrtBase.inl b/Engine/Source/Programs/UnrealBuildAccelerator/Detours/Private/Windows/UbaDetoursFunctionsUcrtBase.inl
--- a/Engine/Source/Programs/UnrealBuildAccelerator/Detours/Private/Windows/UbaDetoursFunctionsUcrtBase.inl
+++ b/Engine/Source/Programs/UnrealBuildAccelerator/Detours/Private/Windows/UbaDetoursFunctionsUcrtBase.inl
@@ -220,6 +220,7 @@ intptr_t Detoured__get_osfhandle(int fd)
 	return res;
 }
 
+#if !defined(_M_ARM64)
 int Detoured__isatty(int fd)
 {
 	DETOURED_CALL(_isatty);
@@ -229,6 +230,7 @@ int Detoured__isatty(int fd)
 	// DEBUG_LOG_TRUE(L"_isatty", L"(%i) -> %i", fd, res); // Spam
 	return res;
 }
+#endif
 
 int Detoured__write(int fd, const void* buffer, unsigned int count)
 {
diff --git a/Engine/Source/Programs/UnrealBuildAccelerator/Detours/Private/Windows/UbaDetoursFunctionsWin.h b/Engine/Source/Programs/UnrealBuildAccelerator/Detours/Private/Windows/UbaDetoursFunctionsWin.h
--- a/Engine/Source/Programs/UnrealBuildAccelerator/Detours/Private/Windows/UbaDetoursFunctionsWin.h
+++ b/Engine/Source/Programs/UnrealBuildAccelerator/Detours/Private/Windows/UbaDetoursFunctionsWin.h
@@ -150,6 +150,14 @@
 #define DETOURED_WSPLITPATH
 #endif
 
+// TODO: For some reason _isatty can't be detoured on windows arm64
+#if !defined(_M_ARM64)
+#define DETOURED_ISATTY DETOURED_FUNCTION(_isatty)
+#else
+#define DETOURED_ISATTY
+#endif
+
+
 #define DETOURED_FUNCTIONS_UCRTBASE \
 	DETOURED_FUNCTION(_wgetcwd) \
 	DETOURED_FUNCTION(_wfullpath) \
@@ -158,9 +166,9 @@
 	DETOURED_FUNCTION(_waccess_s) \
 	DETOURED_FUNCTION(_wspawnl) \
 	DETOURED_FUNCTION(_get_osfhandle) \
-	DETOURED_FUNCTION(_isatty) \
 	DETOURED_FUNCTION(_write) \
 	DETOURED_FUNCTION(fputs) \
+	DETOURED_ISATTY \
 	DETOURED_WSPLITPATH \
 	DETOURED_FUNCTIONS_UCRTBASE_DEBUG \
 
