[UBA]
* Fixed crash when running uba in debug on arm64

diff --git a/Engine/Source/Programs/UnrealBuildAccelerator/Detours/Private/UbaDetoursShared.cpp b/Engine/Source/Programs/UnrealBuildAccelerator/Detours/Private/UbaDetoursShared.cpp
--- a/Engine/Source/Programs/UnrealBuildAccelerator/Detours/Private/UbaDetoursShared.cpp
+++ b/Engine/Source/Programs/UnrealBuildAccelerator/Detours/Private/UbaDetoursShared.cpp
@@ -54,12 +54,12 @@ namespace uba
 #if UBA_DEBUG_LOG_ENABLED
 	FileHandle g_debugFile = InvalidFileHandle;
 	void WriteDebug(const char* str, u32 strLen);
-	thread_local u64 t_logScopeCount;
 	constexpr const char g_emptyString[] = "                                                     ";
 	constexpr const char* g_emptyStringEnd = ((const char*)g_emptyString) + sizeof_array(g_emptyString) - 1;
 	thread_local StringBuffer<LogBufSize> t_a;
 	thread_local char t_b[LogBufSize];
 	thread_local u32 t_b_size;
+	thread_local u32 t_logScopeCount;
 
 	void GetPrefixExtra(StringBufferBase& out)
 	{
diff --git a/Engine/Source/Programs/UnrealBuildAccelerator/Detours/Private/Windows/UbaDetoursFunctionsNtDll.inl b/Engine/Source/Programs/UnrealBuildAccelerator/Detours/Private/Windows/UbaDetoursFunctionsNtDll.inl
--- a/Engine/Source/Programs/UnrealBuildAccelerator/Detours/Private/Windows/UbaDetoursFunctionsNtDll.inl
+++ b/Engine/Source/Programs/UnrealBuildAccelerator/Detours/Private/Windows/UbaDetoursFunctionsNtDll.inl
@@ -1302,7 +1302,9 @@ NTSTATUS NTAPI Detoured_NtClose(HANDLE handle)
 	{
 		TimerScope ts(g_kernelStats.closeHandle);
 		auto res = True_NtClose(handle);
+#if !defined(_M_ARM64) // For some reason this log line crashes on arm64 with access violation on internal tls variable
 		DEBUG_LOG_TRUE(L"NtClose", L"%llu (%ls) -> %ls", uintptr_t(handle), HandleToName(handle), ToString(res));
+#endif
 		return res;
 	}
 
