[UBA]
* Added bool to connect message saying if client is arm or not

diff --git a/Engine/Source/Programs/UnrealBuildAccelerator/Common/Private/UbaSessionClient.cpp b/Engine/Source/Programs/UnrealBuildAccelerator/Common/Private/UbaSessionClient.cpp
--- a/Engine/Source/Programs/UnrealBuildAccelerator/Common/Private/UbaSessionClient.cpp
+++ b/Engine/Source/Programs/UnrealBuildAccelerator/Common/Private/UbaSessionClient.cpp
@@ -1139,6 +1139,11 @@ namespace uba
 		NetworkMessage msg(m_client, ServiceId, SessionMessageType_Connect, writer);
 		writer.WriteString(m_name.data);
 		writer.WriteU32(SessionNetworkVersion);
+		bool isArm = false;
+		#if defined(_M_ARM64)
+		isArm = true;
+		#endif
+		writer.WriteBool(isArm);
 
 		CasKey keys[2];
 		if (m_useBinariesAsVersion)
diff --git a/Engine/Source/Programs/UnrealBuildAccelerator/Common/Private/UbaSessionServer.cpp b/Engine/Source/Programs/UnrealBuildAccelerator/Common/Private/UbaSessionServer.cpp
--- a/Engine/Source/Programs/UnrealBuildAccelerator/Common/Private/UbaSessionServer.cpp
+++ b/Engine/Source/Programs/UnrealBuildAccelerator/Common/Private/UbaSessionServer.cpp
@@ -610,8 +610,11 @@ namespace uba
 		StringBuffer<128> name;
 		reader.ReadString(name);
 		u32 clientVersion = reader.ReadU32();
+		bool isClientArm = false;
+		if (clientVersion >= 36)
+			isClientArm = reader.ReadBool();
 
-		m_logger.Detail(TC("Client session %s connected (Id: %u, Uid: %s)"), name.data, connectionInfo.GetId(), GuidToString(connectionInfo.GetUid()).str);
+		m_logger.Detail(TC("Client session %s connected (Id: %u, Uid: %s%s)"), name.data, connectionInfo.GetId(), GuidToString(connectionInfo.GetUid()).str, (isClientArm ? TC(", IsArm: true") : TC("")));
 
 		CasKey clientKeys[2];
 		clientKeys[0] = reader.ReadCasKey();
diff --git a/Engine/Source/Programs/UnrealBuildAccelerator/Common/Public/UbaNetwork.h b/Engine/Source/Programs/UnrealBuildAccelerator/Common/Public/UbaNetwork.h
--- a/Engine/Source/Programs/UnrealBuildAccelerator/Common/Public/UbaNetwork.h
+++ b/Engine/Source/Programs/UnrealBuildAccelerator/Common/Public/UbaNetwork.h
@@ -16,7 +16,7 @@ namespace uba
 	
 	static constexpr u32 SystemNetworkVersion = 1339;
 	static constexpr u32 StorageNetworkVersion = 4;
-	static constexpr u32 SessionNetworkVersion = 35;
+	static constexpr u32 SessionNetworkVersion = 36;
 	static constexpr u32 CacheNetworkVersion = 5;
 
 	// Messages used over network between client/server (system, storage and session)
