[UBA]
* Added WIP code to change path for cl.exe to arm if server is not arm and helper is

diff --git a/Engine/Source/Programs/UnrealBuildAccelerator/Common/Private/UbaSessionClient.cpp b/Engine/Source/Programs/UnrealBuildAccelerator/Common/Private/UbaSessionClient.cpp
--- a/Engine/Source/Programs/UnrealBuildAccelerator/Common/Private/UbaSessionClient.cpp
+++ b/Engine/Source/Programs/UnrealBuildAccelerator/Common/Private/UbaSessionClient.cpp
@@ -1139,11 +1139,7 @@ namespace uba
 		NetworkMessage msg(m_client, ServiceId, SessionMessageType_Connect, writer);
 		writer.WriteString(m_name.data);
 		writer.WriteU32(SessionNetworkVersion);
-		bool isArm = false;
-		#if defined(_M_ARM64)
-		isArm = true;
-		#endif
-		writer.WriteBool(isArm);
+		writer.WriteBool(IsArm);
 
 		CasKey keys[2];
 		if (m_useBinariesAsVersion)
@@ -1217,6 +1213,7 @@ namespace uba
 		m_detailedTrace = reader.ReadBool();
 		m_shouldSendLogToServer = reader.ReadBool();
 		m_shouldSendTraceToServer = reader.ReadBool();
+		m_isServerArm = reader.ReadBool();
 
 		if (m_shouldSendLogToServer)
 			m_logToFile = true;
@@ -1702,6 +1699,21 @@ namespace uba
 					startInfo.outputStatsThresholdMs = m_outputStatsThresholdMs != 0 ? m_outputStatsThresholdMs : startInfo.outputStatsThresholdMs;
 					startInfo.rules = GetRules(startInfo);
 
+					if constexpr (IsArm)
+					{
+						if (!m_isServerArm)
+						{
+							const tchar* pos;
+							if (Contains(startInfo.application, TC("Hostx64"), false, &pos))
+							{
+								StringBuffer<> newPath;
+								newPath.Append(startInfo.application, pos - startInfo.application).Append(TC("HostArm64")).Append(pos + 7);
+								startInfo.application = newPath.data;
+								startInfo.application = startInfo.applicationStr.data();
+							}
+						}
+					}
+
 					StringBuffer<> logFile;
 					if (m_logToFile)
 					{
diff --git a/Engine/Source/Programs/UnrealBuildAccelerator/Common/Private/UbaSessionServer.cpp b/Engine/Source/Programs/UnrealBuildAccelerator/Common/Private/UbaSessionServer.cpp
--- a/Engine/Source/Programs/UnrealBuildAccelerator/Common/Private/UbaSessionServer.cpp
+++ b/Engine/Source/Programs/UnrealBuildAccelerator/Common/Private/UbaSessionServer.cpp
@@ -628,6 +628,15 @@ namespace uba
 				StringBuffer<> dir;
 				if (!GetDirectoryOfCurrentModule(m_logger, dir))
 					return false;
+		
+				const tchar* archPath[2] = { TC("x64"), TC("arm64") };
+				if (IsArm != isClientArm)
+				{
+					if (!dir.EndsWith(archPath[IsArm]))
+						return m_logger.Error(TC("Module dir is not under proper folder, can't figure out matching x64/arm64 folder"));
+					dir.Resize(dir.count - TStrlen(archPath[IsArm])).Append(archPath[isClientArm]);
+				}
+
 				u64 dirCount = dir.count;
 				bool deferCreation = true;
 				bool fileIsCompressed = false;
@@ -699,6 +708,7 @@ namespace uba
 		writer.WriteBool(m_detailedTrace);
 		writer.WriteBool(m_remoteLogEnabled);
 		writer.WriteBool(m_remoteTraceEnabled);
+		writer.WriteBool(IsArm);
 		WriteRemoteEnvironmentVariables(writer);
 
 		m_trace.SessionAdded(sessionId, connectionInfo.GetId(), name.data, info.data); // Must be inside lock for TraceSessionUpdate() to not include
diff --git a/Engine/Source/Programs/UnrealBuildAccelerator/Common/Public/UbaSessionClient.h b/Engine/Source/Programs/UnrealBuildAccelerator/Common/Public/UbaSessionClient.h
--- a/Engine/Source/Programs/UnrealBuildAccelerator/Common/Public/UbaSessionClient.h
+++ b/Engine/Source/Programs/UnrealBuildAccelerator/Common/Public/UbaSessionClient.h
@@ -122,6 +122,7 @@ namespace uba
 		bool m_shouldSendLogToServer = false;
 		bool m_shouldSendTraceToServer = false;
 		bool m_remoteExecutionEnabled = true;
+		bool m_isServerArm = false;
 		
 		Atomic<const tchar*> m_terminationReason;
 		Atomic<u64> m_terminationTime;
diff --git a/Engine/Source/Programs/UnrealBuildAccelerator/Core/Public/UbaBase.h b/Engine/Source/Programs/UnrealBuildAccelerator/Core/Public/UbaBase.h
--- a/Engine/Source/Programs/UnrealBuildAccelerator/Core/Public/UbaBase.h
+++ b/Engine/Source/Programs/UnrealBuildAccelerator/Core/Public/UbaBase.h
@@ -31,7 +31,13 @@ namespace uba
 	#define PERCENT_HS "%s"
 	#endif
 
-	#if PLATFORM_WINDOWS && defined(UBA_BUILD)
+	#if defined(_M_ARM64)
+	constexpr bool IsArm = true;
+	#else
+	constexpr bool IsArm = false;
+	#endif
+
+#if PLATFORM_WINDOWS && defined(UBA_BUILD)
 	#pragma warning(disable:4100) // This is needed because of single header compiles where AdditionalCompilerArguments is not included
 	#endif
 }
