[UBA]
* Added support for adding extra breadcrumbs to already created/finished processes

diff --git a/Engine/Source/Programs/UnrealBuildAccelerator/Common/Private/UbaSessionServer.cpp b/Engine/Source/Programs/UnrealBuildAccelerator/Common/Private/UbaSessionServer.cpp
--- a/Engine/Source/Programs/UnrealBuildAccelerator/Common/Private/UbaSessionServer.cpp
+++ b/Engine/Source/Programs/UnrealBuildAccelerator/Common/Private/UbaSessionServer.cpp
@@ -492,6 +492,11 @@ namespace uba
 		m_trace.StatusUpdate(statusRow, statusColumn, statusText, statusType, statusLink);
 	}
 
+	void SessionServer::AddProcessBreadcrumbs(u32 processId, const tchar* breadcrumbs, bool deleteOld)
+	{
+		m_trace.ProcessAddBreadcrumbs(processId, ToView(breadcrumbs), deleteOld);
+	}
+
 	NetworkServer& SessionServer::GetServer()
 	{
 		return m_server;
diff --git a/Engine/Source/Programs/UnrealBuildAccelerator/Common/Private/UbaTrace.cpp b/Engine/Source/Programs/UnrealBuildAccelerator/Common/Private/UbaTrace.cpp
--- a/Engine/Source/Programs/UnrealBuildAccelerator/Common/Private/UbaTrace.cpp
+++ b/Engine/Source/Programs/UnrealBuildAccelerator/Common/Private/UbaTrace.cpp
@@ -292,6 +292,14 @@ namespace uba
 		writer.WriteString(reason);
 	}
 
+	void Trace::ProcessAddBreadcrumbs(u32 processId, const StringView& breadcrumbs, bool deleteOld)
+	{
+		BEGIN_TRACE_ENTRY(ProcessBreadcrumbs);
+		writer.WriteU32(processId);
+		writer.WriteString(breadcrumbs);
+		writer.WriteBool(deleteOld);
+	}
+
 	void Trace::ProxyCreated(u32 clientId, const tchar* proxyName)
 	{
 		BEGIN_TRACE_ENTRY(ProxyCreated);
diff --git a/Engine/Source/Programs/UnrealBuildAccelerator/Common/Public/UbaSessionServer.h b/Engine/Source/Programs/UnrealBuildAccelerator/Common/Public/UbaSessionServer.h
--- a/Engine/Source/Programs/UnrealBuildAccelerator/Common/Public/UbaSessionServer.h
+++ b/Engine/Source/Programs/UnrealBuildAccelerator/Common/Public/UbaSessionServer.h
@@ -62,6 +62,8 @@ namespace uba
 		// Add external status information to trace stream. Will show in visualizer
 		void UpdateStatus(u32 statusRow, u32 statusColumn, const tchar* statusText, LogEntryType statusType, const tchar* statusLink);
 
+		void AddProcessBreadcrumbs(u32 processId, const tchar* breadcrumbs, bool deleteOld = false);
+
 		// Get the network server used by this session
 		NetworkServer& GetServer();
 
diff --git a/Engine/Source/Programs/UnrealBuildAccelerator/Common/Public/UbaTrace.h b/Engine/Source/Programs/UnrealBuildAccelerator/Common/Public/UbaTrace.h
--- a/Engine/Source/Programs/UnrealBuildAccelerator/Common/Public/UbaTrace.h
+++ b/Engine/Source/Programs/UnrealBuildAccelerator/Common/Public/UbaTrace.h
@@ -58,6 +58,7 @@ namespace uba
 		UBA_TRACE_TYPE(ProgressUpdate) \
 		UBA_TRACE_TYPE(RemoteExecutionDisabled) \
 		UBA_TRACE_TYPE(FileFetchSize) \
+		UBA_TRACE_TYPE(ProcessBreadcrumbs) \
 
 	enum TraceType : u8
 	{
@@ -90,6 +91,7 @@ namespace uba
 		void ProcessEnvironmentUpdated(u32 processId, const StringView& reason, const u8* data, u64 dataSize, const StringView& breadcrumbs);
 		void ProcessExited(u32 processId, u32 exitCode, const u8* data, u64 dataSize, const Vector<ProcessLogLine>& logLines);
 		void ProcessReturned(u32 processId, const StringView& reason);
+		void ProcessAddBreadcrumbs(u32 processId, const StringView& breadcrumbs, bool deleteOld);
 		void ProxyCreated(u32 clientId, const tchar* proxyName);
 		void ProxyUsed(u32 clientId, const tchar* proxyName);
 		void FileFetchLight(u32 clientId, const CasKey& key, u64 fileSize);
diff --git a/Engine/Source/Programs/UnrealBuildAccelerator/Host/Private/UbaExports.cpp b/Engine/Source/Programs/UnrealBuildAccelerator/Host/Private/UbaExports.cpp
--- a/Engine/Source/Programs/UnrealBuildAccelerator/Host/Private/UbaExports.cpp
+++ b/Engine/Source/Programs/UnrealBuildAccelerator/Host/Private/UbaExports.cpp
@@ -441,6 +441,11 @@ extern "C"
 		server->UpdateStatus(statusRow, statusColumn, statusText, statusType, statusLink);
 	}
 
+	void SessionServer_AddProcessBreadcrumbs(uba::SessionServer* server, uba::u32 processId, const uba::tchar* breadcrumbs, bool deleteOld)
+	{
+		server->AddProcessBreadcrumbs(processId, breadcrumbs, deleteOld);
+	}
+	
 	void SessionServer_RegisterCustomService(uba::SessionServer* server, SessionServer_CustomServiceFunction* function, void* userData)
 	{
 		server->RegisterCustomService([function, userData](uba::Process& process, const void* recv, uba::u32 recvSize, void* send, uba::u32 sendCapacity)
diff --git a/Engine/Source/Programs/UnrealBuildAccelerator/Host/Public/UbaExports.h b/Engine/Source/Programs/UnrealBuildAccelerator/Host/Public/UbaExports.h
--- a/Engine/Source/Programs/UnrealBuildAccelerator/Host/Public/UbaExports.h
+++ b/Engine/Source/Programs/UnrealBuildAccelerator/Host/Public/UbaExports.h
@@ -115,6 +115,7 @@ extern "C"
 	UBA_API void SessionServer_EndExternalProcess(uba::SessionServer* server, uba::u32 id, uba::u32 exitCode);
 	UBA_API void SessionServer_UpdateProgress(uba::SessionServer* server, uba::u32 processesTotal, uba::u32 processesDone, uba::u32 errorCount);
 	UBA_API void SessionServer_UpdateStatus(uba::SessionServer* server, uba::u32 statusRow, uba::u32 statusColumn, const uba::tchar* statusText, uba::LogEntryType statusType, const uba::tchar* statusLink);
+	UBA_API void SessionServer_AddProcessBreadcrumbs(uba::SessionServer* server, uba::u32 processId, const uba::tchar* breadcrumbs, bool deleteOld = false);
 	UBA_API void SessionServer_RegisterCustomService(uba::SessionServer* server, SessionServer_CustomServiceFunction* function, void* userData = nullptr);
 	UBA_API void SessionServer_Destroy(uba::SessionServer* server);
 
diff --git a/Engine/Source/Programs/UnrealBuildAccelerator/Visualizer/Private/UbaTraceReader.cpp b/Engine/Source/Programs/UnrealBuildAccelerator/Visualizer/Private/UbaTraceReader.cpp
--- a/Engine/Source/Programs/UnrealBuildAccelerator/Visualizer/Private/UbaTraceReader.cpp
+++ b/Engine/Source/Programs/UnrealBuildAccelerator/Visualizer/Private/UbaTraceReader.cpp
@@ -1035,6 +1035,23 @@ namespace uba
 
 			break;
 		}
+		case TraceType_ProcessBreadcrumbs:
+		{
+			u32 processId = reader.ReadU32();
+			TString breadcrumbs = reader.ReadString();
+			bool deleteOld = reader.ReadBool();
+
+			for (auto& session : out.sessions)
+				for (auto& processor : session.processors)
+					for (auto& process : processor.processes)
+						if (process.id == processId)
+						{
+							if (deleteOld)
+								process.breadcrumbs.clear();
+							process.breadcrumbs += breadcrumbs;
+						}
+			break;
+		}
 		case TraceType_RemoteExecutionDisabled:
 		{
 			out.remoteExecutionDisabled = true;
@@ -1094,6 +1111,11 @@ namespace uba
 			write.success = reader.ReadBool();
 			write.bytesSent = reader.Read7BitEncoded();
 			write.end = time;
+			break;
+		}
+		default:
+		{
+			return m_logger.Error(L"Unknown trace type found in stream. UbaVisualizer too old?");
 		}
 		}
 		return true;
