[UBA]
* Fixed bug in cross architecture logic where Hostx64->HostArm64
* Added text to ubaagent description if running arm

diff --git a/Engine/Source/Programs/UnrealBuildAccelerator/Agent/Private/UbaAgent.cpp b/Engine/Source/Programs/UnrealBuildAccelerator/Agent/Private/UbaAgent.cpp
--- a/Engine/Source/Programs/UnrealBuildAccelerator/Agent/Private/UbaAgent.cpp
+++ b/Engine/Source/Programs/UnrealBuildAccelerator/Agent/Private/UbaAgent.cpp
@@ -797,7 +797,9 @@ namespace uba
 			extraInfo.Append(TC(", ")).Append(zone);
 
 		if (IsRunningWine())
-			extraInfo.Append(TC(", Linux/WINE"));
+			extraInfo.Appendf(TC(", Linux%s/WINE"), (IsArm ? TC("Arm64") : TC("")));
+		else if (IsArm)
+			extraInfo.Append(", Arm64");
 
 		if (useQuic)
 			extraInfo.Append(TC(", MsQuic"));
diff --git a/Engine/Source/Programs/UnrealBuildAccelerator/Common/Private/UbaSessionClient.cpp b/Engine/Source/Programs/UnrealBuildAccelerator/Common/Private/UbaSessionClient.cpp
--- a/Engine/Source/Programs/UnrealBuildAccelerator/Common/Private/UbaSessionClient.cpp
+++ b/Engine/Source/Programs/UnrealBuildAccelerator/Common/Private/UbaSessionClient.cpp
@@ -1708,7 +1708,7 @@ namespace uba
 							{
 								StringBuffer<> newPath;
 								newPath.Append(startInfo.application, pos - startInfo.application).Append(TC("HostArm64")).Append(pos + 7);
-								startInfo.application = newPath.data;
+								startInfo.applicationStr = newPath.data;
 								startInfo.application = startInfo.applicationStr.data();
 							}
 						}
@@ -1838,7 +1838,13 @@ namespace uba
 							session.m_processFinished(&process);
 					};
 
-					process->Start(startInfo, true, env, true);
+					if (!process->Start(startInfo, true, env, true)) // If false, exitFunction is not called
+					{
+						SendReturnProcess(rec->handle.GetId(), TC("Failed to find executable"));
+						activeProcesses.pop_back();
+						delete exitedRec;
+						m_remoteExecutionEnabled = false;
+					}
 				}
 
 				RemoveInactiveProcesses();
