[UBA]
* Added support for UBT folder structure for cross architecture execution

diff --git a/Engine/Source/Programs/UnrealBuildAccelerator/Common/Private/UbaSessionServer.cpp b/Engine/Source/Programs/UnrealBuildAccelerator/Common/Private/UbaSessionServer.cpp
--- a/Engine/Source/Programs/UnrealBuildAccelerator/Common/Private/UbaSessionServer.cpp
+++ b/Engine/Source/Programs/UnrealBuildAccelerator/Common/Private/UbaSessionServer.cpp
@@ -634,12 +634,20 @@ namespace uba
 				if (!GetDirectoryOfCurrentModule(m_logger, dir))
 					return false;
 		
-				const tchar* archPath[2] = { TC("x64"), TC("arm64") };
 				if (IsArm != isClientArm)
 				{
+					// UBT has the path win-x64/native or win-arm64/native
+					bool isUbtPath = dir.EndsWith(TC("native"));
+					if (isUbtPath)
+						dir.Resize(dir.count - 7); // Remove native and slash
+
+					const tchar* archPath[2] = { TC("x64"), TC("arm64") };
 					if (!dir.EndsWith(archPath[IsArm]))
-						return m_logger.Error(TC("Module dir is not under proper folder, can't figure out matching x64/arm64 folder"));
+						return m_logger.Error(TC("Module dir is not under supported folder (%s) to be able to run cross architecutres, can't figure out matching x64/arm64 folder"), dir.data);
 					dir.Resize(dir.count - TStrlen(archPath[IsArm])).Append(archPath[isClientArm]);
+
+					if (isUbtPath)
+						dir.Append(PathSeparator).Append(TC("native"));
 				}
 
 				u64 dirCount = dir.count;
