[UBA]
* WIP Hacked in so rc.exe can be run on a remote arm helper

diff --git a/Engine/Source/Programs/UnrealBuildAccelerator/Common/Private/UbaSessionClient.cpp b/Engine/Source/Programs/UnrealBuildAccelerator/Common/Private/UbaSessionClient.cpp
--- a/Engine/Source/Programs/UnrealBuildAccelerator/Common/Private/UbaSessionClient.cpp
+++ b/Engine/Source/Programs/UnrealBuildAccelerator/Common/Private/UbaSessionClient.cpp
@@ -1703,6 +1703,7 @@ namespace uba
 					{
 						if (!m_isServerArm)
 						{
+							// WIP WIP EXTREMELY HACKY... WILL BE REFACTORED
 							const tchar* pos;
 							if (Contains(startInfo.application, TC("Hostx64"), false, &pos))
 							{
@@ -1711,6 +1712,14 @@ namespace uba
 								startInfo.applicationStr = newPath.data;
 								startInfo.application = startInfo.applicationStr.data();
 							}
+							else if (Contains(startInfo.application, TC("rc.exe")) && Contains(startInfo.application, TC("\\x64\\"), false, &pos))
+							{
+								StringBuffer<> newPath;
+								newPath.Append(startInfo.application, pos - startInfo.application).Append(TC("\\arm64\\")).Append(pos + 4);
+								startInfo.applicationStr = newPath.data;
+								startInfo.application = startInfo.applicationStr.data();
+							}
+							// WIP WIP EXTREMELY HACKY... WILL BE REFACTORED
 						}
 					}
 
