[UBA]
* Added better error when architecture mismatch between executable and machine

diff --git a/Engine/Source/Programs/UnrealBuildAccelerator/Common/Private/UbaProcess.cpp b/Engine/Source/Programs/UnrealBuildAccelerator/Common/Private/UbaProcess.cpp
--- a/Engine/Source/Programs/UnrealBuildAccelerator/Common/Private/UbaProcess.cpp
+++ b/Engine/Source/Programs/UnrealBuildAccelerator/Common/Private/UbaProcess.cpp
@@ -251,7 +251,9 @@ namespace uba
 
 			if (m_messageCount == 0) // This is bad.. bad binaries?
 			{
-				bool is64Bit = true;
+				bool imageIs64Bit = true;
+				bool imageIsX64 = false;
+				bool imageIsArm64 = false;
 				u64 fileSize = 0;
 				CasKey key;
 				
@@ -266,18 +268,29 @@ namespace uba
 					key = ToCasKey(hasher, false);
 
 					if (data[0] != 'M' || data[1] != 'Z')
-						is64Bit = false;
+						imageIs64Bit = false;
 					else
 					{
 						u32 offset = *(u32*)(data + 0x3c);
-						is64Bit = *(u32*)(data + offset) == 0x00004550;
+						u32* signaturePos = (u32*)(data + offset);
+						imageIs64Bit = *signaturePos == 0x00004550;
+						u16 machine = *(u16*)(signaturePos + 1);
+						imageIsX64 = machine == 0x8664;
+						imageIsArm64 = machine == 0xaa64;
 					}
 				}
 
-				if (!is64Bit)
-					err.Appendf(TC("ERROR: Process did not start properly. Doesn't seem to be a 64-bit executable (%s Size: %llu, CasKey: %s)"), m_realApplication.c_str(), fileSize, CasKeyString(key).str);
-				else
-					err.Appendf(TC("ERROR: Process did not start properly. GetExitCodeProcess returned 0x%x (%s Size: %llu, CasKey: %s)"), exitCode, m_realApplication.c_str(), fileSize, CasKeyString(key).str);
+				bool machineIsArm = IsArm;
+				err.Append(TC("ERROR: Process did not start properly. "));
+				if (!imageIs64Bit)
+					err.Append(TC("Doesn't seem to be a 64-bit executable"));
+				else if (!imageIsArm64 && !imageIsX64)
+					err.Append(TC("Unknown image architecture"));
+				else if (machineIsArm && imageIsX64)
+					err.Appendf(TC("Machine is arm64 and image is x64"));
+				else if (!machineIsArm && imageIsArm64)
+					err.Appendf(TC("Machine is x64 and image is arm64"));
+				err.Appendf(TC(" (GetExitCodeProcess returned 0x%x (%s Size: %llu, CasKey: %s)"), exitCode, m_realApplication.c_str(), fileSize, CasKeyString(key).str);
 			}
 
 			if (err.IsEmpty())
