[UBA]
* Added more error handling in detouring code

diff --git a/Engine/Source/Programs/UnrealBuildAccelerator/Detours/Private/Windows/UbaDetoursFunctionsWin.cpp b/Engine/Source/Programs/UnrealBuildAccelerator/Detours/Private/Windows/UbaDetoursFunctionsWin.cpp
--- a/Engine/Source/Programs/UnrealBuildAccelerator/Detours/Private/Windows/UbaDetoursFunctionsWin.cpp
+++ b/Engine/Source/Programs/UnrealBuildAccelerator/Detours/Private/Windows/UbaDetoursFunctionsWin.cpp
@@ -550,10 +550,26 @@ void DetourDetachFunction(void** trueFunc, void* detouredFunc, const char* funcN
 	Rpc_WriteLogf(L"Failed to detach detoured %hs", funcName);
 }
 
+void DetourTransactionBegin()
+{
+	LONG error = ::DetourTransactionBegin();
+	if (error != NO_ERROR)
+		FatalError(1357, L"DetourTransactionBegin failed (%ld)", error);
+	error = ::DetourUpdateThread(GetCurrentThread());
+	if (error != NO_ERROR)
+		FatalError(1358, L"DetourUpdateThread failed (%ld)", error);
+}
+
+void DetourTransactionCommit()
+{
+	LONG error = ::DetourTransactionCommit();
+	if (error != NO_ERROR)
+		FatalError(1343, L"DetourTransactionCommit failed (%ld)", error);
+}
+
 int DetourAttachFunctions(bool runningRemote)
 {
 	DetourTransactionBegin();
-	DetourUpdateThread(GetCurrentThread());
 
 	#define DETOURED_FUNCTION(Func) True_##Func = (decltype(True_##Func))GetProcAddress(moduleHandle, #Func);
 
@@ -616,12 +632,7 @@ int DetourAttachFunctions(bool runningRemote)
 	#undef DETOURED_FUNCTION
 
 
-	LONG error = DetourTransactionCommit();
-	if (error != NO_ERROR)
-	{
-		printf("Error detouring: %ld\n", error);
-		ExitProcess(1343);
-	}
+	DetourTransactionCommit();
 
 	#if UBA_SUPPORT_MSPDBSRV
 	True2_NdrClientCall2 = True_NdrClientCall2;
@@ -638,10 +649,8 @@ void OnModuleLoaded(HMODULE moduleHandle, const wchar_t* name)
 		True_SymLoadModuleExW = (SymLoadModuleExWFunc*)GetProcAddress(moduleHandle, "SymLoadModuleExW");
 		UBA_ASSERT(True_SymLoadModuleExW);
 		DetourTransactionBegin();
-		DetourUpdateThread(GetCurrentThread());
 		DetourAttachFunction((PVOID*)&True_SymLoadModuleExW, Detoured_SymLoadModuleExW, "SymLoadModuleExW");
-		LONG error = DetourTransactionCommit(); (void)error;
-		UBA_ASSERT(!error);
+		DetourTransactionCommit();
 	}
 
 	// ImageGetDigestStream is buggy in wine so we have to detour it for ShaderCompileWorker
@@ -650,10 +659,8 @@ void OnModuleLoaded(HMODULE moduleHandle, const wchar_t* name)
 		True_ImageGetDigestStream = (ImageGetDigestStreamFunc*)GetProcAddress(moduleHandle, "ImageGetDigestStream");
 		UBA_ASSERT(True_ImageGetDigestStream);
 		DetourTransactionBegin();
-		DetourUpdateThread(GetCurrentThread());
 		DetourAttachFunction((PVOID*)&True_ImageGetDigestStream, Detoured_ImageGetDigestStream, "ImageGetDigestStream");
-		LONG error = DetourTransactionCommit(); (void)error;
-		UBA_ASSERT(!error);
+		DetourTransactionCommit();
 	}
 
 	// SHGetKnownFolderPath is used by Metal.exe and must always execute on host
@@ -662,10 +669,8 @@ void OnModuleLoaded(HMODULE moduleHandle, const wchar_t* name)
 		True_SHGetKnownFolderPath = (SHGetKnownFolderPathFunc*)GetProcAddress(moduleHandle, "SHGetKnownFolderPath");
 		UBA_ASSERT(True_SHGetKnownFolderPath);
 		DetourTransactionBegin();
-		DetourUpdateThread(GetCurrentThread());
 		DetourAttachFunction((PVOID*)&True_SHGetKnownFolderPath, Detoured_SHGetKnownFolderPath, "SHGetKnownFolderPath");
-		LONG error = DetourTransactionCommit(); (void)error;
-		UBA_ASSERT(!error);
+		DetourTransactionCommit();
 	}
 }
 
@@ -970,9 +975,8 @@ void Deinit(u64 startTime)
 	if (g_isRunningWine) // mt.exe etc fails if detaching is not done during shutdown
 	{
 		DetourTransactionBegin();
-		DetourUpdateThread(GetCurrentThread());
 		DetourDetachFunctions();
-		DetourTransactionCommit();
+		::DetourTransactionCommit(); // Ignore errors
 	}
 
 	#if defined(UBA_PROFILE_DETOURED_CALLS)
@@ -1004,7 +1008,7 @@ void PostDeinit()
 	#endif
 }
 
-}
+} // namespace uba
 
 extern "C"
 {
