[UBA]
* Added formatting validation
* Fixed formatting errors
* Fixed so DEBUG_LOG entries never show up twice

Thank you Lambert Clara for the help

diff --git a/Engine/Source/Programs/UnrealBuildAccelerator/Common/Private/UbaStorageClient.cpp b/Engine/Source/Programs/UnrealBuildAccelerator/Common/Private/UbaStorageClient.cpp
--- a/Engine/Source/Programs/UnrealBuildAccelerator/Common/Private/UbaStorageClient.cpp
+++ b/Engine/Source/Programs/UnrealBuildAccelerator/Common/Private/UbaStorageClient.cpp
@@ -646,7 +646,7 @@ namespace uba
 								if (responseSize > left)
 								{
 									overflow = responseSize - u32(left);
-									UBA_ASSERTF(overflow < BufferSlotHalfSize, TC("Something went wrong. Overflow: %u responseSize: %u, left: %u"), overflow, responseSize, left);
+									UBA_ASSERTF(overflow < BufferSlotHalfSize, TC("Something went wrong. Overflow: %u responseSize: %u, left: %llu"), overflow, responseSize, left);
 									if (overflow >= 8)
 									{
 										responseSize = 0;
diff --git a/Engine/Source/Programs/UnrealBuildAccelerator/Core/Private/UbaStringBuffer.cpp b/Engine/Source/Programs/UnrealBuildAccelerator/Core/Private/UbaStringBuffer.cpp
--- a/Engine/Source/Programs/UnrealBuildAccelerator/Core/Private/UbaStringBuffer.cpp
+++ b/Engine/Source/Programs/UnrealBuildAccelerator/Core/Private/UbaStringBuffer.cpp
@@ -152,7 +152,7 @@ namespace uba
 
 	StringBufferBase& StringBufferBase::Append(const tchar* str, u64 charCount)
 	{
-		UBA_ASSERTF(count + charCount < capacity, TC("Trying to append %u character string to buffer which is %u long and has %u capacity left"), charCount, count, capacity - count);
+		UBA_ASSERTF(count + charCount < capacity, TC("Trying to append %llu character string to buffer which is %u long and has %u capacity left"), charCount, count, capacity - count);
 		memcpy(data + count, str, charCount * sizeof(tchar));
 		count += u32(charCount);
 		data[count] = 0;
diff --git a/Engine/Source/Programs/UnrealBuildAccelerator/Core/Public/UbaDirectoryTable.h b/Engine/Source/Programs/UnrealBuildAccelerator/Core/Public/UbaDirectoryTable.h
--- a/Engine/Source/Programs/UnrealBuildAccelerator/Core/Public/UbaDirectoryTable.h
+++ b/Engine/Source/Programs/UnrealBuildAccelerator/Core/Public/UbaDirectoryTable.h
@@ -40,7 +40,7 @@ namespace uba
 				u64 pos = reader.GetPosition();
 				if (pos == size)
 					break;
-				UBA_ASSERTF(pos < size, TC("Should never read past size (pos: %u, size: %u)"), pos, size);
+				UBA_ASSERTF(pos < size, TC("Should never read past size (pos: %llu, size: %u)"), pos, size);
 				u64 storageSize = reader.Read7BitEncoded();
 				StringKey dirKey = reader.ReadStringKey();
 				auto insres = m_lookup.try_emplace(dirKey, m_memoryBlock); // Note that this is allowed to overwrite
diff --git a/Engine/Source/Programs/UnrealBuildAccelerator/Core/Public/UbaPathUtils.h b/Engine/Source/Programs/UnrealBuildAccelerator/Core/Public/UbaPathUtils.h
--- a/Engine/Source/Programs/UnrealBuildAccelerator/Core/Public/UbaPathUtils.h
+++ b/Engine/Source/Programs/UnrealBuildAccelerator/Core/Public/UbaPathUtils.h
@@ -133,7 +133,7 @@ namespace uba
 
 			UBA_ASSERTF(workingDir && *workingDir, TC("No working dir provided but path is relative (%s)"), buffer);
 			tchar temp2[1024];
-			UBA_ASSERTF(workingDirCharLen + charLen < sizeof_array(temp2), TC("%llu + %llu < %llu"), workingDirCharLen, charLen, sizeof_array(temp2));
+			UBA_ASSERTF(workingDirCharLen + charLen < sizeof_array(temp2), TC("%llu + %llu < %i"), workingDirCharLen, charLen, sizeof_array(temp2));
 			memcpy(temp2, workingDir, workingDirCharLen*sizeof(tchar));
 			memcpy(temp2 + workingDirCharLen, copyFrom, charLen*sizeof(tchar));
 			charLen += workingDirCharLen;
@@ -223,7 +223,7 @@ namespace uba
 			++charLen;
 		}
 
-		UBA_ASSERTF(charLen <= bufferCharCapacity, TC("Buffer overflow (Capacity %llu) fixing path %s"), fileName);
+		UBA_ASSERTF(charLen <= bufferCharCapacity, TC("Buffer overflow (Capacity %llu) fixing path %s"), bufferCharCapacity, fileName);
 
 		if (outBufferCharLen)
 			*outBufferCharLen = u32(charLen - 1); // Remove terminator
diff --git a/Engine/Source/Programs/UnrealBuildAccelerator/Core/Public/UbaPlatform.h b/Engine/Source/Programs/UnrealBuildAccelerator/Core/Public/UbaPlatform.h
--- a/Engine/Source/Programs/UnrealBuildAccelerator/Core/Public/UbaPlatform.h
+++ b/Engine/Source/Programs/UnrealBuildAccelerator/Core/Public/UbaPlatform.h
@@ -51,6 +51,12 @@
 
 #define UBA_EXPERIMENTAL 0
 
+#if UBA_DEBUG && defined(_MSC_VER)
+	#define UBA_FMT_CHECK(Format, ...) do { if (false) { wprintf(Format, ##__VA_ARGS__); } } while(false)
+#else
+	#define UBA_FMT_CHECK(Format, ...)
+#endif
+
 namespace uba
 {
 	#ifndef sizeof_array
@@ -59,7 +65,7 @@ namespace uba
 
 	#if UBA_DEBUG
 	#define UBA_ASSERT(x) do { if (x) break; uba::UbaAssert(TC(""), __FILE__, __LINE__, #x, 543221); } while(false)
-	#define UBA_ASSERTF(x, ...) do { if (x) break; uba::StringBuffer<> _buf; _buf.Appendf(__VA_ARGS__); uba::UbaAssert(_buf.data, __FILE__, __LINE__, #x, 543221); } while(false)
+	#define UBA_ASSERTF(x, fmt, ...) do { UBA_FMT_CHECK(fmt, ##__VA_ARGS__); if (x) break; uba::StringBuffer<> _buf; _buf.Appendf(fmt, ##__VA_ARGS__); uba::UbaAssert(_buf.data, __FILE__, __LINE__, #x, 543221); } while(false)
 	#define UBA_ASSERT_MESSAGEBOX 0
 	#else
 	#define UBA_ASSERT(x) do { } while(false)
diff --git a/Engine/Source/Programs/UnrealBuildAccelerator/Detours/Private/UbaDetoursShared.cpp b/Engine/Source/Programs/UnrealBuildAccelerator/Detours/Private/UbaDetoursShared.cpp
--- a/Engine/Source/Programs/UnrealBuildAccelerator/Detours/Private/UbaDetoursShared.cpp
+++ b/Engine/Source/Programs/UnrealBuildAccelerator/Detours/Private/UbaDetoursShared.cpp
@@ -73,6 +73,12 @@ namespace uba
 		#endif
 		//out.Appendf(TC("[%7u]"), GetCurrentThreadId());
 	}
+	void FlushDebug()
+	{
+		WriteDebug(t_b, t_b_size);
+		t_b_size = 0;
+		t_b[0] = 0;
+	}
 	void WriteDebugLogWithPrefix(const char* prefix, LogScope& scope, const tchar* command, const tchar* format, ...)
 	{
 		#if PLATFORM_MAC
@@ -121,8 +127,10 @@ namespace uba
 		t_a.Append(TC("\n"));
 
 		#if PLATFORM_WINDOWS
+		if (t_b_size)
+			FlushDebug();
 		t_b_size = sprintf_s(t_b, LogBufSize, "%S", t_a.data);
-		WriteDebug(t_b, t_b_size);
+		FlushDebug();
 		#else
 		WriteDebug(t_a.data, t_a.count);
 		#endif
@@ -139,9 +147,7 @@ namespace uba
 	}
 	void LogScope::Flush()
 	{
-		WriteDebug(t_b, t_b_size);
-		t_b_size = 0;
-		t_b[0] = 0;
+		FlushDebug();
 	}
 #endif
 
diff --git a/Engine/Source/Programs/UnrealBuildAccelerator/Detours/Private/UbaDetoursShared.h b/Engine/Source/Programs/UnrealBuildAccelerator/Detours/Private/UbaDetoursShared.h
--- a/Engine/Source/Programs/UnrealBuildAccelerator/Detours/Private/UbaDetoursShared.h
+++ b/Engine/Source/Programs/UnrealBuildAccelerator/Detours/Private/UbaDetoursShared.h
@@ -24,9 +24,10 @@ namespace uba
 	bool FixPath(StringBufferBase& out, const tchar* path);
 
 	#if UBA_DEBUG_LOG_ENABLED
-		#define DEBUG_LOG_PREFIX(Prefix, Command, ...) \
+		#define DEBUG_LOG_PREFIX(Prefix, Command, fmt, ...) \
 			LogScope STRING_JOIN(ls, __LINE__); \
-			if (isLogging()) WriteDebugLogWithPrefix(#Prefix, STRING_JOIN(ls, __LINE__), Command, __VA_ARGS__); \
+			UBA_FMT_CHECK(fmt, ##__VA_ARGS__); \
+			if (isLogging()) WriteDebugLogWithPrefix(#Prefix, STRING_JOIN(ls, __LINE__), Command, fmt, ##__VA_ARGS__); \
 
 		//#define DEBUG_LOG_DETOURED(Command, ...) 
 		#define DEBUG_LOG_DETOURED(Command, ...) DEBUG_LOG_PREFIX(D, Command, __VA_ARGS__)
@@ -35,7 +36,7 @@ namespace uba
 		//#define DEBUG_LOG_PIPE(Command, ...) ts.leave(); DEBUG_LOG_PREFIX(P, Command, __VA_ARGS__)
 		#define DEBUG_LOG_PIPE(Command, ...) ts.Leave();
 		//#define DEBUG_LOG(...)
-		#define DEBUG_LOG(...) { if (isLogging()) WriteDebugLog(__VA_ARGS__); }
+		#define DEBUG_LOG(fmt, ...) { UBA_FMT_CHECK(fmt, ##__VA_ARGS__); if (isLogging()) WriteDebugLog(fmt, ##__VA_ARGS__); }
 	#else
 		#define DEBUG_LOG(...) {}
 		#define DEBUG_LOG_DETOURED(Command, ...) {}
diff --git a/Engine/Source/Programs/UnrealBuildAccelerator/Detours/Private/Windows/UbaDetoursFunctionsKernelBase.inl b/Engine/Source/Programs/UnrealBuildAccelerator/Detours/Private/Windows/UbaDetoursFunctionsKernelBase.inl
--- a/Engine/Source/Programs/UnrealBuildAccelerator/Detours/Private/Windows/UbaDetoursFunctionsKernelBase.inl
+++ b/Engine/Source/Programs/UnrealBuildAccelerator/Detours/Private/Windows/UbaDetoursFunctionsKernelBase.inl
@@ -124,7 +124,7 @@ DWORD Detoured_GetCurrentDirectoryA(DWORD nBufferLength, LPSTR lpBuffer)
 	size_t res;
 	errno_t err = wcstombs_s(&res, lpBuffer, nBufferLength, g_virtualWorkingDir.data, length);
 	if (err)
-		UBA_ASSERTF(false, L"wcstombs_s failed for string '%s' with error code: %u", err);
+		UBA_ASSERTF(false, L"wcstombs_s failed for string '%s' with error code: %u", g_virtualWorkingDir.data, err);
 	DEBUG_LOG_DETOURED(L"GetCurrentDirectoryA", L"(%hs)", lpBuffer);
 	return DWORD(length);
 }
@@ -1971,7 +1971,7 @@ BOOL Detoured_GetFileInformationByHandle(HANDLE hFile, LPBY_HANDLE_FILE_INFORMAT
 					//UBA_ASSERTF(entryInfo.attributes == bhfi.dwFileAttributes, L"Attributes used: 0x%x Actual: 0x%x (%s)", entryInfo.attributes, bhfi.dwFileAttributes, fi.originalName);
 					UBA_ASSERTF(entryInfo.volumeSerial == bhfi.dwVolumeSerialNumber, L"VolumeSerial used: %u Actual: %u (%s)", entryInfo.volumeSerial, bhfi.dwVolumeSerialNumber, fi.originalName);
 					UBA_ASSERTF(entryInfo.fileIndex == fileIndex, L"FileIndex used: %llu Actual: %llu (%s)", entryInfo.fileIndex, fileIndex, fi.originalName);
-					UBA_ASSERTF(bhfi.nNumberOfLinks == 1, L"Links used: %llu Actual: %llu (%s)", 1, bhfi.nNumberOfLinks, fi.originalName);
+					UBA_ASSERTF(bhfi.nNumberOfLinks == 1, L"Links used: %i Actual: %u (%s)", 1, bhfi.nNumberOfLinks, fi.originalName);
 				}
 				else if (True_GetFileAttributesExW(fi.originalName, GetFileExInfoStandard, &data))
 				{
@@ -2837,7 +2837,7 @@ BOOL Detoured_CreateProcessW(LPCWSTR lpApplicationName, LPWSTR lpCommandLine, LP
 	detouredHandle->trueHandle = trueHandle;
 	lpProcessInformation->hProcess = makeDetouredHandle(detouredHandle);
 
-	DEBUG_LOG_DETOURED(L"CreateProcessW", L"%llu (0x%llx)", lpProcessInformation->hProcess, trueHandle);
+	DEBUG_LOG_DETOURED(L"CreateProcessW", L"%llu (0x%llx)", u64(lpProcessInformation->hProcess), u64(trueHandle));
 	return TRUE;
 }
 
@@ -2976,7 +2976,7 @@ DWORD Detoured_WaitForSingleObjectEx(HANDLE hHandle, DWORD dwMilliseconds, BOOL
 		auto lastError = GetLastError();
 		DWORD exitCode;
 		True_GetExitCodeProcess(trueHandle, &exitCode);
-		DEBUG_LOG_DETOURED(L"WaitForSingleObjectEx", L"for process %llu (0x%llx). Exit code: %u", hHandle, trueHandle, exitCode);
+		DEBUG_LOG_DETOURED(L"WaitForSingleObjectEx", L"for process %llu (0x%llx). Exit code: %u", u64(hHandle), u64(trueHandle), exitCode);
 		SetLastError(lastError);
 	}
 #endif
@@ -3150,18 +3150,18 @@ DWORD Shared_GetModuleFileNameA(HMODULE hModule, const wchar_t* moduleName, u32
 			size_t res;
 			errno_t err = wcstombs_s(&res, lpFilename, nSize, moduleName, nSize);
 			if (err)
-				UBA_ASSERTF(false, L"wcstombs_s failed for string '%s' with error code: %u", err);
+				UBA_ASSERTF(false, L"wcstombs_s failed for string '%s' with error code: %u", moduleName, err);
 			lpFilename[nSize - 1] = 0;
 		}
 		SetLastError(ERROR_INSUFFICIENT_BUFFER);
-		DEBUG_LOG_DETOURED(L"GetModuleFileNameExA", L"%llu  %u INSUFFICIENT BUFFER (%hs) -> %u", uintptr_t(hModule), nSize, moduleName, moduleNameLen + 1);
+		DEBUG_LOG_DETOURED(L"GetModuleFileNameExA", L"%llu  %u INSUFFICIENT BUFFER (%s) -> %u", uintptr_t(hModule), nSize, moduleName, moduleNameLen + 1);
 		return nSize;
 	}
 
 	size_t res;
 	errno_t err = wcstombs_s(&res, lpFilename, nSize, moduleName, moduleNameLen);
 	if (err)
-		UBA_ASSERTF(false, L"wcstombs_s failed for string '%s' with error code: %u", err);
+		UBA_ASSERTF(false, L"wcstombs_s failed for string '%s' with error code: %u", moduleName, err);
 	lpFilename[moduleNameLen] = 0;
 	memset(lpFilename + moduleNameLen, 0, nSize - moduleNameLen);
 	DEBUG_LOG_DETOURED(L"GetModuleFileNameExA", L"%llu  %u (%hs) -> %u", uintptr_t(hModule), nSize, lpFilename, moduleNameLen);
@@ -3250,7 +3250,7 @@ DWORD Detoured_GetModuleBaseNameW(HANDLE hProcess, HMODULE hModule, LPWSTR lpBas
 	if (wchar_t* lastSlash = wcsrchr(temp, '\\'))
 		moduleName = lastSlash + 1;
 	DWORD len = (DWORD)wcslen(moduleName);
-	UBA_ASSERTF(len < nSize, L"Module name %hs does not fit in buffer size (is %u, needs %u)", moduleName, nSize, len);
+	UBA_ASSERTF(len < nSize, L"Module name %s does not fit in buffer size (is %u, needs %u)", moduleName, nSize, len);
 	wcscpy_s(lpBaseName, nSize, moduleName);
 	return len;
 }
@@ -3895,14 +3895,14 @@ BOOL Detoured_PathIsRelativeW(LPCWSTR pszPath)
 {
 	//UBA_ASSERTF(!g_runningRemote, L"%ls", pszPath); // intel compiler uses PathIsRelativeW.. don't know if this function touches file system but will comment out this assert for now
 	auto res = True_PathIsRelativeW(pszPath);
-	DEBUG_LOG_TRUE(L"PathIsRelativeW", L"(%ls) -> %ls", pszPath, res);
+	DEBUG_LOG_TRUE(L"PathIsRelativeW", L"(%ls) -> %u", pszPath, res);
 	return res;
 }
 BOOL Detoured_PathIsDirectoryEmptyW(LPCWSTR pszPath)
 {
 	UBA_ASSERTF(!g_runningRemote, L"%ls", pszPath);
 	auto res = True_PathIsDirectoryEmptyW(pszPath);
-	DEBUG_LOG_TRUE(L"PathIsDirectoryEmptyW", L"(%ls) -> %ls", pszPath, res);
+	DEBUG_LOG_TRUE(L"PathIsDirectoryEmptyW", L"(%ls) -> %u", pszPath, res);
 	return res;
 }
 
