[UBA]
* Fixed bug where directory table might not be populated with files when building the full path. And since the code had a fallback to use the provided string it could get wrong with casing

diff --git a/Engine/Source/Programs/UnrealBuildAccelerator/Common/Private/UbaSession.cpp b/Engine/Source/Programs/UnrealBuildAccelerator/Common/Private/UbaSession.cpp
--- a/Engine/Source/Programs/UnrealBuildAccelerator/Common/Private/UbaSession.cpp
+++ b/Engine/Source/Programs/UnrealBuildAccelerator/Common/Private/UbaSession.cpp
@@ -1348,7 +1348,7 @@ namespace uba
 
 		StringKeyHasher hasher;
 		hasher.Update(dirPath.data, TStrlen(dirPath.data));
-		m_directoryTable.PopulateDirectory(hasher, dir);
+		m_directoryTable.PopulateDirectoryNoLock(hasher, dir);
 
 		u32 tableOffset;
 		return WriteDirectoryEntriesInternal(dir, dirKey, dirPath.data, true, tableOffset);
diff --git a/Engine/Source/Programs/UnrealBuildAccelerator/Core/Public/UbaDirectoryTable.h b/Engine/Source/Programs/UnrealBuildAccelerator/Core/Public/UbaDirectoryTable.h
--- a/Engine/Source/Programs/UnrealBuildAccelerator/Core/Public/UbaDirectoryTable.h
+++ b/Engine/Source/Programs/UnrealBuildAccelerator/Core/Public/UbaDirectoryTable.h
@@ -51,10 +51,15 @@ namespace uba
 		}
 
 		void PopulateDirectory(const StringKeyHasher& hasher, Directory& dir)
+		{
+			SCOPED_WRITE_LOCK(dir.lock, lock);
+			PopulateDirectoryNoLock(hasher, dir);
+		}
+
+		void PopulateDirectoryNoLock(const StringKeyHasher& hasher, Directory& dir)
 		{
 			if (dir.parseOffset == dir.tableOffset)
 				return;
-			SCOPED_WRITE_LOCK(dir.lock, lock);
 			PopulateDirectoryRecursive(hasher, dir.tableOffset, dir.parseOffset, dir.files);
 			dir.parseOffset = dir.tableOffset;
 		}
@@ -311,6 +316,7 @@ namespace uba
 					continue;
 				}
 				directory = &findIt->second;
+				PopulateDirectory(hasher, *directory); // This is needed to make sure files lookup is populated for query above
 			}
 		}
 
