[UBA]
* Fixed crash in cache server happening when failing to create view of file mapping

diff --git a/Engine/Source/Programs/UnrealBuildAccelerator/Common/Private/UbaStorageServer.cpp b/Engine/Source/Programs/UnrealBuildAccelerator/Common/Private/UbaStorageServer.cpp
--- a/Engine/Source/Programs/UnrealBuildAccelerator/Common/Private/UbaStorageServer.cpp
+++ b/Engine/Source/Programs/UnrealBuildAccelerator/Common/Private/UbaStorageServer.cpp
@@ -553,13 +553,24 @@ namespace uba
 				StringBuffer<512> casFile;
 				FileHandle readFileHandle = InvalidFileHandle;
 				auto rfg = MakeGuard([&]() { CloseFile(nullptr, readFileHandle); });
-				u64 fileSize;
+				u64 fileSize = 0;
 				u8* memoryBegin = nullptr;
 				u8* memoryPos = nullptr;
 				bool ownsMapping = false;
 
 				MappedView mappedView;
-				auto mvg = MakeGuard([&](){ m_casDataBuffer.UnmapView(mappedView, TC("FetchBegin")); });
+				auto mvg = MakeGuard([&]()
+					{
+						if (ownsMapping)
+						{
+							if (mappedView.memory)
+								UnmapViewOfFile(mappedView.memory, fileSize, TC("FetchBegin"));
+							if (mappedView.handle.IsValid())
+								CloseFileMapping(mappedView.handle);
+						}
+						else
+							m_casDataBuffer.UnmapView(mappedView, TC("FetchBegin"));
+					});
 
 				todoRemoveMe = 5;
 
@@ -590,16 +601,16 @@ namespace uba
 	#endif
 					if (fileSize > BufferSlotSize)
 					{
+						ownsMapping = true;
 						mappedView.handle = CreateFileMappingW(readFileHandle, PAGE_READONLY, fileSize, TC(""));
 						if (!mappedView.handle.IsValid())
 							return m_logger.Error(TC("Failed to create file mapping of %s (%s)"), casFile.data, LastErrorToText().data);
 						u64 offset = memoryPos - memoryBegin;
 						mappedView.memory = MapViewOfFile(mappedView.handle, FILE_MAP_READ, 0, fileSize);
 						if (!mappedView.memory)
-							return m_logger.Error(TC("Failed to map memory of %s (%s)"), casFile.data, LastErrorToText().data);
+							return m_logger.Error(TC("Failed to map memory of %s with size %llu (%s)"), casFile.data, fileSize, LastErrorToText().data);
 						memoryBegin = mappedView.memory;
 						memoryPos = memoryBegin + offset;
-						ownsMapping = true;
 						useFileMapping = true;
 					}
 				}
