[UBA]
* Added error messages if failing to close files
* Added active file mapping count stats to cache server to see if we are leaking something

diff --git a/Engine/Source/Programs/UnrealBuildAccelerator/Common/Private/UbaCacheServer.cpp b/Engine/Source/Programs/UnrealBuildAccelerator/Common/Private/UbaCacheServer.cpp
--- a/Engine/Source/Programs/UnrealBuildAccelerator/Common/Private/UbaCacheServer.cpp
+++ b/Engine/Source/Programs/UnrealBuildAccelerator/Common/Private/UbaCacheServer.cpp
@@ -27,6 +27,10 @@
 
 namespace uba
 {
+	#if !PLATFORM_WINDOWS
+	extern Atomic<u64> g_mappedViewCounter;
+	#endif
+
 	static constexpr u32 CacheFileVersion = 6;
 	static constexpr u32 CacheFileCompatibilityVersion = 3;
 
@@ -588,8 +592,12 @@ namespace uba
 					m_logger.Info(TC("Stats since last maintenance (%s ago)"), TimeToText(startTime - m_lastMaintenance, true).str);
 				m_logger.Info(TC("  Peak connection count: %4u"), peakConnectionCount);
 				auto& sentTimer = m_server.GetTotalSentTimer();
-				m_logger.Info(TC("  SendSocket         %8u %9s"), sentTimer.count.load(), TimeToText(sentTimer.time).str);
-				m_logger.Info(TC("     Bytes            %17s"), BytesToText(m_server.GetTotalSentBytes()).str);
+				m_logger.Info(TC("  Socket sent %u (%s)"), sentTimer.count.load(), BytesToText(m_server.GetTotalSentBytes()).str);
+
+				#if !PLATFORM_WINDOWS
+				m_logger.Info(TC("  Active file mappings: %llu"), g_mappedViewCounter.load());
+				#endif
+
 				//KernelStats& kernelStats = KernelStats::GetGlobal();
 				//kernelStats.Print(m_logger, false);
 				//kernelStats = {};
diff --git a/Engine/Source/Programs/UnrealBuildAccelerator/Common/Private/UbaFileMapping.cpp b/Engine/Source/Programs/UnrealBuildAccelerator/Common/Private/UbaFileMapping.cpp
--- a/Engine/Source/Programs/UnrealBuildAccelerator/Common/Private/UbaFileMapping.cpp
+++ b/Engine/Source/Programs/UnrealBuildAccelerator/Common/Private/UbaFileMapping.cpp
@@ -35,10 +35,11 @@ namespace uba
 	}
 #else
 	int asFileDescriptor(FileHandle fh);
-#endif
 
 	ReaderWriterLock g_mappingUidCounterLock;
 	Atomic<u64> g_mappingUidCounter;
+	Atomic<u64> g_mappedViewCounter;
+#endif
 
 	FileMappingHandle CreateMemoryMappingW(Logger& logger, u32 flProtect, u64 maxSize, const tchar* name)
 	{
@@ -226,7 +227,10 @@ namespace uba
 		int shmFd = fileMappingObject.shmFd;
 		void* rptr = mmap(NULL, bytesToMap, prot, MAP_SHARED, shmFd, s64(offset));
 		if (rptr != MAP_FAILED)
+		{
+			++g_mappedViewCounter;
 			return (u8*)rptr;
+		}
 		//UBA_ASSERTF(false, "Failed to map file with fd %i, desiredAccess %u offset %llu, bytesToMap %llu (%s)", fd, desiredAccess, offset, bytesToMap, strerror(errno));
 		SetLastError(errno);
 		return nullptr;
@@ -250,7 +254,10 @@ namespace uba
 #else
 		UBA_ASSERTF(bytesToUnmap, TC("bytesToUnmap is zero unmapping %p (%s)"), lpBaseAddress, hint);
 		if (munmap((void*)lpBaseAddress, bytesToUnmap) == 0)
+		{
+			--g_mappedViewCounter;
 			return true;
+		}
 		UBA_ASSERT(false);
 		return false;
 #endif
diff --git a/Engine/Source/Programs/UnrealBuildAccelerator/Common/Private/UbaStorageServer.cpp b/Engine/Source/Programs/UnrealBuildAccelerator/Common/Private/UbaStorageServer.cpp
--- a/Engine/Source/Programs/UnrealBuildAccelerator/Common/Private/UbaStorageServer.cpp
+++ b/Engine/Source/Programs/UnrealBuildAccelerator/Common/Private/UbaStorageServer.cpp
@@ -125,11 +125,15 @@ namespace uba
 			if (ownsMapping)
 			{
 				++server.m_activeUnmap;
-				server.GetServer().AddWork([&server, mb = memoryBegin, mp = mappedView, rfh = readFileHandle]()
+				server.GetServer().AddWork([&server, mb = memoryBegin, mp = mappedView, rfh = readFileHandle, reason]()
 					{
-						UnmapViewOfFile(mb, mp.size, TC(""));
-						CloseFileMapping(mp.handle);
-						CloseFile(nullptr, rfh);
+						if (!UnmapViewOfFile(mb, mp.size, reason))
+							server.m_logger.Error(TC("%s - Failed to unmap memory at 0x%llx with size %llu (%s)"), reason, mb, mp.size, LastErrorToText().data);
+						if (!CloseFileMapping(mp.handle))
+							server.m_logger.Error(TC("%s - Failed to close file mapping %llu (%s)"), reason, mp.handle.ToU64(), LastErrorToText().data);
+						if (!CloseFile(nullptr, rfh))
+							server.m_logger.Error(TC("%s - Failed to close file (%s)"), reason, LastErrorToText().data);
+
 						--server.m_activeUnmap;
 					}, 1, TC("ActiveFetchRelease"));
 			}
@@ -214,7 +218,7 @@ namespace uba
 
 	struct FileMappingScope
 	{
-		u8* MapView(Logger& logger, const tchar* func, FileMappingHandle mappingHandle, u64 mappingOffset, u64 fileSize, const tchar* fileName_)
+		u8* MapView(const tchar* func, FileMappingHandle mappingHandle, u64 mappingOffset, u64 fileSize, const tchar* fileName_)
 		{
 			fileName = fileName_;
 			constexpr u64 pageSize = 64*1024;
@@ -231,12 +235,15 @@ namespace uba
 			return nullptr;
 		}
 
+		FileMappingScope(Logger& l) : logger(l) {}
 		~FileMappingScope()
 		{
 			if (mappedMem)
-				UnmapViewOfFile(mappedMem, mapSize, fileName);
+				if (!UnmapViewOfFile(mappedMem, mapSize, fileName))
+					logger.Error(TC("Failed to unmap memory %s at 0x%llx with size %llu (%s)"), fileName, mappedMem, mapSize, LastErrorToText().data);
 		}
 
+		Logger& logger;
 		const tchar* fileName;
 		u8* mappedMem = nullptr;
 		u64 mapSize = 0;
@@ -244,8 +251,8 @@ namespace uba
 
 	bool StorageServer::StoreCasFile(CasKey& out, StringKey fileNameKey, const tchar* fileName, FileMappingHandle mappingHandle, u64 mappingOffset, u64 fileSize, const tchar* hint, bool deferCreation, bool keepMappingInMemory, bool storeCompressed)
 	{
-		FileMappingScope scope;
-		u8* fileMem = scope.MapView(m_logger, TC("StoreCasFile"), mappingHandle, mappingOffset, fileSize, fileName);
+		FileMappingScope scope(m_logger);
+		u8* fileMem = scope.MapView(TC("StoreCasFile"), mappingHandle, mappingOffset, fileSize, fileName);
 		if (!fileMem)
 			return false;
 
@@ -289,8 +296,8 @@ namespace uba
 		ExternalFileMapping& mapping = findIt->second;
 		lock.Leave();
 
-		FileMappingScope scope;
-		u8* fileMem = scope.MapView(m_logger, TC("WriteCompressed"), mapping.mappingHandle, mapping.mappingOffset, mapping.fileSize, from);
+		FileMappingScope scope(m_logger);
+		u8* fileMem = scope.MapView(TC("WriteCompressed"), mapping.mappingHandle, mapping.mappingOffset, mapping.fileSize, from);
 		if (!fileMem)
 			return false;
 		return StorageImpl::WriteCompressed(out, from, InvalidFileHandle, fileMem, mapping.fileSize, toFile, nullptr, 0);
@@ -552,7 +559,11 @@ namespace uba
 				
 				StringBuffer<512> casFile;
 				FileHandle readFileHandle = InvalidFileHandle;
-				auto rfg = MakeGuard([&]() { CloseFile(nullptr, readFileHandle); });
+				auto rfg = MakeGuard([&]()
+					{
+						if (!CloseFile(nullptr, readFileHandle))
+							m_logger.Error(TC("Failed to close file %s (%s)"), casFile.data, LastErrorToText().data);;
+					});
 				u64 fileSize = 0;
 				u8* memoryBegin = nullptr;
 				u8* memoryPos = nullptr;
@@ -564,9 +575,11 @@ namespace uba
 						if (ownsMapping)
 						{
 							if (mappedView.memory)
-								UnmapViewOfFile(mappedView.memory, fileSize, TC("FetchBegin"));
+								if (!UnmapViewOfFile(mappedView.memory, fileSize, TC("FetchBegin")))
+									m_logger.Error(TC("Failed to unmap memory %s at 0x%llx with size %llu (%s)"), hint.data, mappedView.memory, fileSize, LastErrorToText().data);
 							if (mappedView.handle.IsValid())
-								CloseFileMapping(mappedView.handle);
+								if (!CloseFileMapping(mappedView.handle))
+									m_logger.Error(TC("Failed to close file mapping for %s with handle %llu (%s)"), hint.data, mappedView.handle.ToU64(), LastErrorToText().data);
 						}
 						else
 							m_casDataBuffer.UnmapView(mappedView, TC("FetchBegin"));
@@ -664,7 +677,8 @@ namespace uba
 					memcpy(writeBuffer, memoryPos, toWrite);
 					memoryPos += toWrite;
 
-					CloseFile(casFile.data, readFileHandle);
+					if (!CloseFile(casFile.data, readFileHandle))
+						m_logger.Error(TC("Failed to close file %s (%s)"), casFile.data, LastErrorToText().data);;
 					readFileHandle = InvalidFileHandle;
 				}
 
