[UBA]
* Added more support for cross-architecture compiles. Logic is now exposed so directory mappings for cross architecture is set from the outside (from UBT)

diff --git a/Engine/Source/Programs/UnrealBuildAccelerator/Common/Private/UbaScheduler.cpp b/Engine/Source/Programs/UnrealBuildAccelerator/Common/Private/UbaScheduler.cpp
--- a/Engine/Source/Programs/UnrealBuildAccelerator/Common/Private/UbaScheduler.cpp
+++ b/Engine/Source/Programs/UnrealBuildAccelerator/Common/Private/UbaScheduler.cpp
@@ -116,7 +116,7 @@ namespace uba
 	void Scheduler::Start()
 	{
 		m_session.SetRemoteProcessReturnedEvent([this](Process& process) { RemoteProcessReturned(process); });
-		m_session.SetRemoteProcessSlotAvailableEvent([this]() { RemoteSlotAvailable(); });
+		m_session.SetRemoteProcessSlotAvailableEvent([this](bool isCrossArchitecture) { RemoteSlotAvailable(isCrossArchitecture); });
 
 		m_loop = true;
 		m_thread.Start([this]() { ThreadLoop(); return 0; });
@@ -292,8 +292,9 @@ namespace uba
 		m_updateThreadLoop.Set();
 	}
 
-	void Scheduler::RemoteSlotAvailable()
+	void Scheduler::RemoteSlotAvailable(bool isCrossArchitecture)
 	{
+		UBA_ASSERTF(false, TC("Cross architecture code path not implemented"));
 		RunQueuedProcess(false);
 	}
 
diff --git a/Engine/Source/Programs/UnrealBuildAccelerator/Common/Private/UbaSessionClient.cpp b/Engine/Source/Programs/UnrealBuildAccelerator/Common/Private/UbaSessionClient.cpp
--- a/Engine/Source/Programs/UnrealBuildAccelerator/Common/Private/UbaSessionClient.cpp
+++ b/Engine/Source/Programs/UnrealBuildAccelerator/Common/Private/UbaSessionClient.cpp
@@ -208,7 +208,7 @@ namespace uba
 			UBA_ASSERT(fileNameKey == StringKeyZero);
 			StackBinaryWriter<1024> writer;
 			NetworkMessage msg(m_client, ServiceId, SessionMessageType_EnsureBinaryFile, writer);
-			writer.WriteU32(processId);
+			writer.WriteBool(IsArm);
 			writer.WriteString(fileName);
 			writer.WriteStringKey(fileNameKey);
 			writer.WriteString(applicationDir);
@@ -1213,7 +1213,6 @@ namespace uba
 		m_detailedTrace = reader.ReadBool();
 		m_shouldSendLogToServer = reader.ReadBool();
 		m_shouldSendTraceToServer = reader.ReadBool();
-		m_isServerArm = reader.ReadBool();
 
 		if (m_shouldSendLogToServer)
 			m_logToFile = true;
@@ -1699,30 +1698,6 @@ namespace uba
 					startInfo.outputStatsThresholdMs = m_outputStatsThresholdMs != 0 ? m_outputStatsThresholdMs : startInfo.outputStatsThresholdMs;
 					startInfo.rules = GetRules(startInfo);
 
-					if constexpr (IsArm)
-					{
-						if (!m_isServerArm)
-						{
-							// WIP WIP EXTREMELY HACKY... WILL BE REFACTORED
-							const tchar* pos;
-							if (Contains(startInfo.application, TC("Hostx64"), false, &pos))
-							{
-								StringBuffer<> newPath;
-								newPath.Append(startInfo.application, pos - startInfo.application).Append(TC("HostArm64")).Append(pos + 7);
-								startInfo.applicationStr = newPath.data;
-								startInfo.application = startInfo.applicationStr.data();
-							}
-							else if (Contains(startInfo.application, TC("rc.exe")) && Contains(startInfo.application, TC("\\x64\\"), false, &pos))
-							{
-								StringBuffer<> newPath;
-								newPath.Append(startInfo.application, pos - startInfo.application).Append(TC("\\arm64\\")).Append(pos + 4);
-								startInfo.applicationStr = newPath.data;
-								startInfo.application = startInfo.applicationStr.data();
-							}
-							// WIP WIP EXTREMELY HACKY... WILL BE REFACTORED
-						}
-					}
-
 					StringBuffer<> logFile;
 					if (m_logToFile)
 					{
diff --git a/Engine/Source/Programs/UnrealBuildAccelerator/Common/Private/UbaSessionServer.cpp b/Engine/Source/Programs/UnrealBuildAccelerator/Common/Private/UbaSessionServer.cpp
--- a/Engine/Source/Programs/UnrealBuildAccelerator/Common/Private/UbaSessionServer.cpp
+++ b/Engine/Source/Programs/UnrealBuildAccelerator/Common/Private/UbaSessionServer.cpp
@@ -83,6 +83,7 @@ namespace uba
 		Vector<u8> m_trackedInputs;
 		Vector<u8> m_trackedOutputs;
 		bool m_cancelled = false;
+		bool m_allowCrossArchitecture = false;
 		u32 m_clientId = ~0u;
 		u32 m_sessionId = 0;
 		TString m_executingHost;
@@ -266,7 +267,7 @@ namespace uba
 		return {};
 	}
 
-	ProcessHandle SessionServer::RunProcessRemote(const ProcessStartInfo& startInfo, float weight, const void* knownInputs, u32 knownInputsCount)
+	ProcessHandle SessionServer::RunProcessRemote(const ProcessStartInfo& startInfo, float weight, const void* knownInputs, u32 knownInputsCount, bool allowCrossArchitecture)
 	{
 		UBA_ASSERT(!startInfo.startSuspended);
 
@@ -278,6 +279,7 @@ namespace uba
 		auto rules = GetRules(remoteProcess->m_startInfo);
 
 		remoteProcess->m_startInfo.rules = rules;
+		remoteProcess->m_allowCrossArchitecture = allowCrossArchitecture;
 
 		if (knownInputsCount)
 		{
@@ -420,7 +422,7 @@ namespace uba
 		return true;
 	}
 
-	void SessionServer::SetRemoteProcessSlotAvailableEvent(const Function<void()>& remoteProcessSlotAvailableEvent)
+	void SessionServer::SetRemoteProcessSlotAvailableEvent(const Function<void(bool isCrossArchitecture)>& remoteProcessSlotAvailableEvent)
 	{
 		SCOPED_WRITE_LOCK(m_remoteProcessSlotAvailableEventLock, lock);
 		m_remoteProcessSlotAvailableEvent = remoteProcessSlotAvailableEvent;
@@ -508,6 +510,11 @@ namespace uba
 		m_provider = provider;
 	}
 
+	void SessionServer::RegisterCrossArchitectureMapping(const tchar* from, const tchar* to)
+	{
+		m_crossArchitectureMappings.emplace_back(from, to);
+	}
+
 	void SessionServer::OnDisconnected(const Guid& clientUid, u32 clientId)
 	{
 		u32 returnCount = 0;
@@ -707,6 +714,7 @@ namespace uba
 		session.id = connectionInfo.GetId();
 		session.processSlotCount = processSlotCount;
 		session.dedicated = dedicated;
+		session.isArm = isClientArm;
 		m_availableRemoteSlotCount += processSlotCount;
 		++m_connectionCount;
 
@@ -721,7 +729,6 @@ namespace uba
 		writer.WriteBool(m_detailedTrace);
 		writer.WriteBool(m_remoteLogEnabled);
 		writer.WriteBool(m_remoteTraceEnabled);
-		writer.WriteBool(IsArm);
 		WriteRemoteEnvironmentVariables(writer);
 
 		m_trace.SessionAdded(sessionId, connectionInfo.GetId(), name.data, info.data); // Must be inside lock for TraceSessionUpdate() to not include
@@ -732,7 +739,7 @@ namespace uba
 
 	bool SessionServer::HandleEnsureBinaryFile(const ConnectionInfo& connectionInfo, BinaryReader& reader, BinaryWriter& writer)
 	{
-		u32 processId = reader.ReadU32(); (void)processId;
+		u32 clientIsArm = reader.ReadBool();
 		StringBuffer<> fileName;
 		reader.ReadString(fileName);
 		StringKey fileNameKey = reader.ReadStringKey();
@@ -763,6 +770,21 @@ namespace uba
 		CasKey casKey = CasKeyZero;
 		StringBuffer<> absoluteFile;
 
+		auto FixCrossArchitecture = [this, clientIsArm](StringBuffer<>& absoluteFile)
+			{
+				if (clientIsArm == IsArm)
+					return;
+				for (auto& mapping : m_crossArchitectureMappings)
+					if (absoluteFile.StartsWith(mapping.from.c_str()))
+					{
+						StringBuffer<> temp;
+						temp.Append(absoluteFile.data + mapping.from.size());
+						absoluteFile.Clear().Append(mapping.to).Append(temp);
+						break;
+					}
+			};
+
+
 		if (!loaderPaths.empty())
 		{
 			for (auto& loaderPath : loaderPaths)
@@ -771,23 +793,24 @@ namespace uba
 				fullPath.Append(applicationDir).EnsureEndsWithSlash().Append(loaderPath).EnsureEndsWithSlash().Append(fileName);
 				if (GetFileAttributesW(fullPath.data) == INVALID_FILE_ATTRIBUTES)
 					continue;
-				StringBuffer<> out;
 				FixPath(fullPath.data, nullptr, 0, absoluteFile);
+				FixCrossArchitecture(absoluteFile);
 				fileNameKey = ToStringKeyLower(absoluteFile);
 				if (!StoreCasFile(casKey, fileNameKey, absoluteFile.data))
 					return false;
 				break;
 			}
 		}
-		else
+		else if (SearchPathForFile(m_logger, absoluteFile, fileName.data, applicationDir.data))
 		{
-			if (SearchPathForFile(m_logger, absoluteFile, fileName.data, applicationDir.data))
-				if (!absoluteFile.StartsWith(m_systemPath.data) || !IsKnownSystemFile(absoluteFile.data))
-				{
-					fileNameKey = ToStringKeyLower(absoluteFile);
-					if (!StoreCasFile(casKey, fileNameKey, absoluteFile.data))
-						return false;
-				}
+
+			if (!absoluteFile.StartsWith(m_systemPath.data) || !IsKnownSystemFile(absoluteFile.data))
+			{
+				FixCrossArchitecture(absoluteFile);
+				fileNameKey = ToStringKeyLower(absoluteFile);
+				if (!StoreCasFile(casKey, fileNameKey, absoluteFile.data))
+					return false;
+			}
 		}
 
 		u64 startPos = writer.GetPosition();
@@ -1121,6 +1144,8 @@ namespace uba
 		ClientSession& session = *m_clientSessions[sessionIndex];
 		sessionsLock.Leave();
 
+		bool isCrossArchitecture = IsArm != session.isArm;
+
 		u32 weight32 = reader.ReadU32();
 		float availableWeight = *(float*)&weight32;
 
@@ -1134,10 +1159,34 @@ namespace uba
 			RemoteProcess* process = DequeueProcess(sessionId, connectionInfo.GetId());
 			if (process == nullptr)
 				break;
+			auto& startInfo = process->m_startInfo;
+
+			StringBuffer<> applicationOverride;
+			if (isCrossArchitecture)
+			{
+				if (!process->m_allowCrossArchitecture)
+				{
+					// There is a risk cross architecture client dequeues another client's process that is not cross architecture
+					// and in that case we just have to return it and try again later.. (break out)
+					m_queuedRemoteProcesses.push_front({process});
+					if (m_remoteProcessReturnedEvent)
+						m_remoteProcessReturnedEvent(*process);
+					break;
+				}
+
+				for (auto& mapping : m_crossArchitectureMappings)
+					if (StartsWith(startInfo.application, mapping.from.c_str()))
+					{
+						applicationOverride.Append(mapping.to).Append(startInfo.application + mapping.from.size());
+						break;
+					}
+				if (!applicationOverride.count)
+					return m_logger.Error(TC("Cross architecture client tried to find mapping for %s"), startInfo.application);
+			}
 
 			ProcessAdded(*process, sessionId);
 			writer.WriteU32(process->m_processId);
-			process->m_startInfo.Write(writer);
+			startInfo.Write(writer, applicationOverride);
 
 			for (auto kiIt = process->m_knownInputs, kiEnd = kiIt + process->m_knownInputsCount; kiIt!=kiEnd; ++kiIt)
 				if (session.sentKeys.insert(kiIt->key).second)
@@ -1148,7 +1197,7 @@ namespace uba
 			if (writer.GetCapacityLeft() < 5000) // Arbitrary number to cover all parameters above
 				break;
 
-			weightLeft -= process->m_startInfo.weight;
+			weightLeft -= startInfo.weight;
 		}
 		fillLock.Leave();
 
@@ -1808,7 +1857,8 @@ namespace uba
 
 			if (hasCalledCallback)
 				return nullptr;
-			m_remoteProcessSlotAvailableEvent();
+
+			m_remoteProcessSlotAvailableEvent(IsArm != session.isArm);
 			hasCalledCallback = true;
 		}
 		return nullptr;
diff --git a/Engine/Source/Programs/UnrealBuildAccelerator/Common/Public/UbaNetwork.h b/Engine/Source/Programs/UnrealBuildAccelerator/Common/Public/UbaNetwork.h
--- a/Engine/Source/Programs/UnrealBuildAccelerator/Common/Public/UbaNetwork.h
+++ b/Engine/Source/Programs/UnrealBuildAccelerator/Common/Public/UbaNetwork.h
@@ -16,7 +16,7 @@ namespace uba
 	
 	static constexpr u32 SystemNetworkVersion = 1339;
 	static constexpr u32 StorageNetworkVersion = 4;
-	static constexpr u32 SessionNetworkVersion = 36;
+	static constexpr u32 SessionNetworkVersion = 37;
 	static constexpr u32 CacheNetworkVersion = 5;
 
 	// Messages used over network between client/server (system, storage and session)
diff --git a/Engine/Source/Programs/UnrealBuildAccelerator/Common/Public/UbaProcessStartInfoHolder.h b/Engine/Source/Programs/UnrealBuildAccelerator/Common/Public/UbaProcessStartInfoHolder.h
--- a/Engine/Source/Programs/UnrealBuildAccelerator/Common/Public/UbaProcessStartInfoHolder.h
+++ b/Engine/Source/Programs/UnrealBuildAccelerator/Common/Public/UbaProcessStartInfoHolder.h
@@ -51,10 +51,13 @@ namespace uba
 			breadcrumbs = breadcrumbsStr.c_str();
 		}
 
-		void Write(BinaryWriter& writer)
+		void Write(BinaryWriter& writer, const StringView& applicationOverride)
 		{
 			writer.WriteString(descriptionStr);
-			writer.WriteString(applicationStr);
+			if (applicationOverride.count)
+				writer.WriteString(applicationOverride);
+			else
+				writer.WriteString(applicationStr);
 			writer.WriteString(argumentsStr);
 			writer.WriteString(workingDirStr);
 			writer.WriteString(logFileStr);
diff --git a/Engine/Source/Programs/UnrealBuildAccelerator/Common/Public/UbaScheduler.h b/Engine/Source/Programs/UnrealBuildAccelerator/Common/Public/UbaScheduler.h
--- a/Engine/Source/Programs/UnrealBuildAccelerator/Common/Public/UbaScheduler.h
+++ b/Engine/Source/Programs/UnrealBuildAccelerator/Common/Public/UbaScheduler.h
@@ -90,7 +90,7 @@ namespace uba
 		void ThreadLoop();
 		void RemoteProcessReturned(Process& process);
 		void HandleCacheMissed(ExitProcessInfo* ei);
-		void RemoteSlotAvailable();
+		void RemoteSlotAvailable(bool isCrossArchitecture);
 		void ProcessExited(ExitProcessInfo* info, const ProcessHandle& handle);
 		u32 PopProcess(bool isLocal, ProcessStatus& outPrevStatus);
 		bool RunQueuedProcess(bool isLocal);
diff --git a/Engine/Source/Programs/UnrealBuildAccelerator/Common/Public/UbaSessionClient.h b/Engine/Source/Programs/UnrealBuildAccelerator/Common/Public/UbaSessionClient.h
--- a/Engine/Source/Programs/UnrealBuildAccelerator/Common/Public/UbaSessionClient.h
+++ b/Engine/Source/Programs/UnrealBuildAccelerator/Common/Public/UbaSessionClient.h
@@ -122,7 +122,6 @@ namespace uba
 		bool m_shouldSendLogToServer = false;
 		bool m_shouldSendTraceToServer = false;
 		bool m_remoteExecutionEnabled = true;
-		bool m_isServerArm = false;
 		
 		Atomic<const tchar*> m_terminationReason;
 		Atomic<u64> m_terminationTime;
diff --git a/Engine/Source/Programs/UnrealBuildAccelerator/Common/Public/UbaSessionServer.h b/Engine/Source/Programs/UnrealBuildAccelerator/Common/Public/UbaSessionServer.h
--- a/Engine/Source/Programs/UnrealBuildAccelerator/Common/Public/UbaSessionServer.h
+++ b/Engine/Source/Programs/UnrealBuildAccelerator/Common/Public/UbaSessionServer.h
@@ -23,7 +23,7 @@ namespace uba
 		// weight is the expected core usage of the process. If processes are multithreaded it makes sense to increase weight. As an example, in UnrealBuildTool we see cl.exe as 1.5 and clang.exe as 1.0
 		// knownInputs is a memory block with null terminated tchar strings followed by an empty null terminated string to end. knownInputsCount is the number of strings in the memory block
 		// strings should be absolute or relative to working dir.
-		ProcessHandle RunProcessRemote(const ProcessStartInfo& startInfo, float weight = 1.0f, const void* knownInputs = nullptr, u32 knownInputsCount = 0);
+		ProcessHandle RunProcessRemote(const ProcessStartInfo& startInfo, float weight = 1.0f, const void* knownInputs = nullptr, u32 knownInputsCount = 0, bool allowCrossArchitecture = false);
 
 		// Will kick off a local process with the same startInfo as the one provided to start the process with id matching raceAgainstRemoteProcessId
 		// This can be useful if there are free local cores and we know local machine is faster or network connection to remote is slow
@@ -39,7 +39,7 @@ namespace uba
 		bool GetCasKeyFromTrackedInputs(CasKey& out, const tchar* fileName, const tchar* workingDir, const u8* data, u32 dataLen);
 
 		// Callback that will be called when a client is asking for processes to run. All remotes frequently ask for processes to run if they have free process slots
-		void SetRemoteProcessSlotAvailableEvent(const Function<void()>& remoteProcessSlotAvailableEvent);
+		void SetRemoteProcessSlotAvailableEvent(const Function<void(bool isCrossArchitecture)>& remoteProcessSlotAvailableEvent);
 
 		// Callback that is called when process is returned. This could be because a client unexpectedly disconnected or is running out of memory
 		void SetRemoteProcessReturnedEvent(const Function<void(Process&)>& remoteProcessReturnedEvent);
@@ -62,14 +62,19 @@ namespace uba
 		// Add external status information to trace stream. Will show in visualizer
 		void UpdateStatus(u32 statusRow, u32 statusColumn, const tchar* statusText, LogEntryType statusType, const tchar* statusLink);
 
+		// Add additional breadcrumbs to process. Can be done while process is running or after it is done
 		void AddProcessBreadcrumbs(u32 processId, const tchar* breadcrumbs, bool deleteOld = false);
 
 		// Get the network server used by this session
 		NetworkServer& GetServer();
 
+		// Additional network traffic provider. Can be used to add more network stats to trace
 		using NetworkTrafficProvider = Function<void(u64& outSent, u64& outReceive)>;
 		void RegisterNetworkTrafficProvider(const NetworkTrafficProvider& provider);
 
+		// Register mappings for cross architecture helpers (helpers with different architecture will use the mapping)
+		void RegisterCrossArchitectureMapping(const tchar* from, const tchar* to);
+
 	protected:
 		struct ClientSession;
 
@@ -117,7 +122,7 @@ namespace uba
 		static constexpr u8 ServiceId = SessionServiceId;
 
 		ReaderWriterLock m_remoteProcessSlotAvailableEventLock;
-		Function<void()> m_remoteProcessSlotAvailableEvent;
+		Function<void(bool isCrossArchitecture)> m_remoteProcessSlotAvailableEvent;
 
 		ReaderWriterLock m_remoteProcessReturnedEventLock;
 		Function<void(Process&)> m_remoteProcessReturnedEvent;
@@ -152,6 +157,7 @@ namespace uba
 			bool dedicated = false;
 			bool abort = false;
 			bool hasNotification = false;
+			bool isArm = false;
 		};
 		Vector<ClientSession*> m_clientSessions;
 
@@ -179,6 +185,10 @@ namespace uba
 
 		NetworkTrafficProvider m_provider;
 
+		ReaderWriterLock m_crossArchitectureMappingsLock;
+		struct CrossArchitectureMapping { TString from; TString to; };
+		Vector<CrossArchitectureMapping> m_crossArchitectureMappings;
+
 		Event m_memoryThreadEvent;
 		Thread m_memoryThread;
 		Atomic<u64> m_memAvail;
diff --git a/Engine/Source/Programs/UnrealBuildAccelerator/Host/Private/UbaExports.cpp b/Engine/Source/Programs/UnrealBuildAccelerator/Host/Private/UbaExports.cpp
--- a/Engine/Source/Programs/UnrealBuildAccelerator/Host/Private/UbaExports.cpp
+++ b/Engine/Source/Programs/UnrealBuildAccelerator/Host/Private/UbaExports.cpp
@@ -359,7 +359,7 @@ extern "C"
 	}
 	void SessionServer_SetRemoteProcessAvailable(uba::SessionServer* server, SessionServer_RemoteProcessAvailableCallback* available, void* userData)
 	{
-		server->SetRemoteProcessSlotAvailableEvent([available, userData]() { available(userData); });
+		server->SetRemoteProcessSlotAvailableEvent([available, userData](bool isCrossArchitecture) { available(userData, isCrossArchitecture); });
 	}
 	void SessionServer_SetRemoteProcessReturned(uba::SessionServer* server, SessionServer_RemoteProcessReturnedCallback* returned, void* userData)
 	{
@@ -385,9 +385,9 @@ extern "C"
 	{
 		return new uba::ProcessHandle(server->RunProcess(info, async, enableDetour));
 	}
-	uba::ProcessHandle* SessionServer_RunProcessRemote(uba::SessionServer* server, uba::ProcessStartInfo& info, float weight, const void* knownInputs, uba::u32 knownInputsCount)
+	uba::ProcessHandle* SessionServer_RunProcessRemote(uba::SessionServer* server, uba::ProcessStartInfo& info, float weight, const void* knownInputs, uba::u32 knownInputsCount, bool allowCrossArchitecture)
 	{
-		return new uba::ProcessHandle(server->RunProcessRemote(info, weight, knownInputs, knownInputsCount));
+		return new uba::ProcessHandle(server->RunProcessRemote(info, weight, knownInputs, knownInputsCount, allowCrossArchitecture));
 	}
 	uba::ProcessHandle* SessionServer_RunProcessRacing(uba::SessionServer* server, uba::u32 raceAgainstRemoteProcessId)
 	{
@@ -455,6 +455,11 @@ extern "C"
 			});
 	}
 
+	void SessionServer_RegisterCrossArchitectureMapping(uba::SessionServer* server, const uba::tchar* from, const uba::tchar* to)
+	{
+		server->RegisterCrossArchitectureMapping(from, to);
+	}
+
 	void SessionServer_Destroy(uba::SessionServer* server)
 	{
 		if (server)
diff --git a/Engine/Source/Programs/UnrealBuildAccelerator/Host/Public/UbaExports.h b/Engine/Source/Programs/UnrealBuildAccelerator/Host/Public/UbaExports.h
--- a/Engine/Source/Programs/UnrealBuildAccelerator/Host/Public/UbaExports.h
+++ b/Engine/Source/Programs/UnrealBuildAccelerator/Host/Public/UbaExports.h
@@ -87,7 +87,7 @@ extern "C"
 	UBA_API const uba::ProcessStartInfo* Process_GetStartInfo(uba::Process& process);
 
 	// SessionServer
-	using SessionServer_RemoteProcessAvailableCallback = void(void* userData);
+	using SessionServer_RemoteProcessAvailableCallback = void(void* userData, bool isCrossArchitecture);
 	using SessionServer_RemoteProcessReturnedCallback = void(uba::Process& process, void* userData);
 	using SessionServer_CustomServiceFunction = uba::u32(uba::ProcessHandle* handle, const void* recv, uba::u32 recvSize, void* send, uba::u32 sendCapacity, void* userData);
 
@@ -103,7 +103,7 @@ extern "C"
 	UBA_API void SessionServer_RegisterDeleteFile(uba::SessionServer* server, const uba::tchar* filePath);
 	UBA_API bool SessionServer_RegisterVirtualFile(uba::SessionServer* server, const uba::tchar* filename, const uba::tchar* sourceFile, uba::u64 sourceOffset, uba::u64 sourceSize);
 	UBA_API uba::ProcessHandle* SessionServer_RunProcess(uba::SessionServer* server, uba::ProcessStartInfo& info, bool async, bool enableDetour);
-	UBA_API uba::ProcessHandle* SessionServer_RunProcessRemote(uba::SessionServer* server, uba::ProcessStartInfo& info, float weight, const void* knownInputs = nullptr, uba::u32 knownInputsCount = 0);
+	UBA_API uba::ProcessHandle* SessionServer_RunProcessRemote(uba::SessionServer* server, uba::ProcessStartInfo& info, float weight, const void* knownInputs = nullptr, uba::u32 knownInputsCount = 0, bool allowCrossArchitecture = false);
 	UBA_API uba::ProcessHandle* SessionServer_RunProcessRacing(uba::SessionServer* server, uba::u32 raceAgainstRemoteProcessId);
 
 	UBA_API void SessionServer_SetMaxRemoteProcessCount(uba::SessionServer* server, uba::u32 count);
@@ -117,6 +117,7 @@ extern "C"
 	UBA_API void SessionServer_UpdateStatus(uba::SessionServer* server, uba::u32 statusRow, uba::u32 statusColumn, const uba::tchar* statusText, uba::LogEntryType statusType, const uba::tchar* statusLink);
 	UBA_API void SessionServer_AddProcessBreadcrumbs(uba::SessionServer* server, uba::u32 processId, const uba::tchar* breadcrumbs, bool deleteOld = false);
 	UBA_API void SessionServer_RegisterCustomService(uba::SessionServer* server, SessionServer_CustomServiceFunction* function, void* userData = nullptr);
+	UBA_API void SessionServer_RegisterCrossArchitectureMapping(uba::SessionServer* server, const uba::tchar* from, const uba::tchar* to);
 	UBA_API void SessionServer_Destroy(uba::SessionServer* server);
 
 	// Scheduler
