[UBA]
* Added agent option -nodetoursdownload which can be used to use local UbaDetours binary instead of downloading from host

diff --git a/Engine/Source/Programs/UnrealBuildAccelerator/Agent/Private/UbaAgent.cpp b/Engine/Source/Programs/UnrealBuildAccelerator/Agent/Private/UbaAgent.cpp
--- a/Engine/Source/Programs/UnrealBuildAccelerator/Agent/Private/UbaAgent.cpp
+++ b/Engine/Source/Programs/UnrealBuildAccelerator/Agent/Private/UbaAgent.cpp
@@ -95,6 +95,7 @@ namespace uba
 		logger.Info(TC("  -named=<name>           Use named events and file mappings by providing the base name in this option"));
 		logger.Info(TC("  -nopoll                 Does not keep polling for work; attempts to connect once then exits"));
 		logger.Info(TC("  -nostore                Does not use storage to store files (with a few exceptions such as binaries)"));
+		logger.Info(TC("  -nodetoursdownload      Does not download UbaDetours library from server and instead use local."));
 		logger.Info(TC("  -resetstore             Delete all cas"));
 		logger.Info(TC("  -quiet                  Does not output any logging in console"));
 		logger.Info(TC("  -maxidle=<seconds>      Max time agent will idle before disconnecting. Ignored if -nopoll is not set"));
@@ -432,6 +433,7 @@ namespace uba
 		bool verbose = false;
 		bool printSummary = false;
 		bool killRandom = false;
+		bool downloadDetoursLib = true;
 		StringBuffer<512> sentryUrl;
 		StringBuffer<128> zone;
 		u32 maxIdleSeconds = ~0u;
@@ -609,6 +611,10 @@ namespace uba
 			{
 				useStorage = false;
 			}
+			else if (name.Equals(TC("-nodetoursdownload")))
+			{
+				downloadDetoursLib = false;
+			}
 			else if (name.Equals(TC("-resetstore")))
 			{
 				resetStore = true;
@@ -1207,6 +1213,7 @@ namespace uba
 			info.useBinariesAsVersion = useBinariesAsVersion;
 			info.killRandom = killRandom;
 			info.useStorage = useStorage;
+			info.downloadDetoursLib = downloadDetoursLib;
 			info.memWaitLoadPercent = u8(memWaitLoadPercent);
 			info.memKillLoadPercent = u8(memKillLoadPercent);
 
diff --git a/Engine/Source/Programs/UnrealBuildAccelerator/Common/Private/UbaSessionClient.cpp b/Engine/Source/Programs/UnrealBuildAccelerator/Common/Private/UbaSessionClient.cpp
--- a/Engine/Source/Programs/UnrealBuildAccelerator/Common/Private/UbaSessionClient.cpp
+++ b/Engine/Source/Programs/UnrealBuildAccelerator/Common/Private/UbaSessionClient.cpp
@@ -36,6 +36,7 @@ namespace uba
 		m_maxProcessCount = info.maxProcessCount;
 		m_dedicated = info.dedicated;
 		m_useStorage = info.useStorage;
+		m_downloadDetoursLib = info.downloadDetoursLib;
 		m_defaultPriorityClass = info.defaultPriorityClass;
 		m_outputStatsThresholdMs = info.outputStatsThresholdMs;
 		m_maxIdleSeconds = info.maxIdleSeconds;
@@ -1181,8 +1182,11 @@ namespace uba
 			return;
 		}
 
+		CasKey detoursBinaryKey = reader.ReadCasKey();
+
+		StringBuffer<> detoursFile;
+		if (m_downloadDetoursLib)
 		{
-			CasKey detoursBinaryKey = reader.ReadCasKey();
 			{
 				TimerScope s(m_stats.storageRetrieve);
 				Storage::RetrieveResult result;
@@ -1190,18 +1194,22 @@ namespace uba
 					return;
 			}
 			KeyToString dir(StringKeyZero);
-			StringBuffer<> detoursFile;
 			if (!WriteBinFile(detoursFile, UBA_DETOURS_LIBRARY, detoursBinaryKey, dir, DefaultAttributes()))
 				return;
-
-			#if PLATFORM_WINDOWS
-			char dll[1024];
-			sprintf_s(dll, sizeof(dll), "%ls", detoursFile.data);
-			m_detoursLibrary = dll;
-			#else
-			m_detoursLibrary = detoursFile.data;
-			#endif
 		}
+		else
+		{
+			GetDirectoryOfCurrentModule(m_logger, detoursFile);
+			detoursFile.EnsureEndsWithSlash().Append(UBA_DETOURS_LIBRARY);
+		}
+
+		#if PLATFORM_WINDOWS
+		char dll[1024];
+		sprintf_s(dll, sizeof(dll), "%ls", detoursFile.data);
+		m_detoursLibrary = dll;
+		#else
+		m_detoursLibrary = detoursFile.data;
+		#endif
 
 		bool resetCas = reader.ReadBool();
 		if (resetCas)
diff --git a/Engine/Source/Programs/UnrealBuildAccelerator/Common/Public/UbaSessionClient.h b/Engine/Source/Programs/UnrealBuildAccelerator/Common/Public/UbaSessionClient.h
--- a/Engine/Source/Programs/UnrealBuildAccelerator/Common/Public/UbaSessionClient.h
+++ b/Engine/Source/Programs/UnrealBuildAccelerator/Common/Public/UbaSessionClient.h
@@ -26,6 +26,7 @@ namespace uba
 		bool useBinariesAsVersion = false;
 		bool killRandom = false;
 		bool useStorage = true;
+		bool downloadDetoursLib = true;
 		Function<void(const ProcessHandle&)> processFinished;
 	};
 
@@ -119,6 +120,7 @@ namespace uba
 		bool m_connected = false;
 		bool m_dedicated = false;
 		bool m_useStorage = true;
+		bool m_downloadDetoursLib = true;
 		bool m_shouldSendLogToServer = false;
 		bool m_shouldSendTraceToServer = false;
 		bool m_remoteExecutionEnabled = true;
