[UBA]
* Fixed so uba host can have both x64 and arm64 helpers connected at the same time and correct UbaDetours library is sent to the different clients

diff --git a/Engine/Source/Programs/UnrealBuildAccelerator/Common/Private/UbaSessionServer.cpp b/Engine/Source/Programs/UnrealBuildAccelerator/Common/Private/UbaSessionServer.cpp
--- a/Engine/Source/Programs/UnrealBuildAccelerator/Common/Private/UbaSessionServer.cpp
+++ b/Engine/Source/Programs/UnrealBuildAccelerator/Common/Private/UbaSessionServer.cpp
@@ -632,10 +632,13 @@ namespace uba
 		clientKeys[0] = reader.ReadCasKey();
 		clientKeys[1] = reader.ReadCasKey();
 
+		CasKey& detoursBinaryKey = m_detoursBinaryKey[isClientArm];
+		CasKey& agentBinaryKey = m_agentBinaryKey[isClientArm];
+
 		bool binAsVersion = clientKeys[0] != CasKeyZero;
 		{
 			SCOPED_WRITE_LOCK(m_binKeysLock, lock);
-			if (m_detoursBinaryKey == CasKeyZero || (binAsVersion && m_agentBinaryKey == CasKeyZero))
+			if (detoursBinaryKey == CasKeyZero || (binAsVersion && agentBinaryKey == CasKeyZero))
 			{
 				StringBuffer<> dir;
 				if (!GetDirectoryOfCurrentModule(m_logger, dir))
@@ -660,22 +663,22 @@ namespace uba
 				u64 dirCount = dir.count;
 				bool deferCreation = true;
 				bool fileIsCompressed = false;
-				if (binAsVersion && m_agentBinaryKey == CasKeyZero)
+				if (binAsVersion && agentBinaryKey == CasKeyZero)
 				{
 					UBA_ASSERT(IsWindows);
 					dir.Append(PathSeparator).Append(UBA_AGENT_EXECUTABLE);
-					m_storage.StoreCasFile(m_agentBinaryKey, dir.data, CasKeyZero, deferCreation, fileIsCompressed);
+					m_storage.StoreCasFile(agentBinaryKey, dir.data, CasKeyZero, deferCreation, fileIsCompressed);
 				}
 				dir.Resize(dirCount).Append(PathSeparator).Append(UBA_DETOURS_LIBRARY);
-				if (!m_storage.StoreCasFile(m_detoursBinaryKey, dir.data, CasKeyZero, deferCreation, fileIsCompressed))
+				if (!m_storage.StoreCasFile(detoursBinaryKey, dir.data, CasKeyZero, deferCreation, fileIsCompressed))
 					return m_logger.Error(TC("Failed to create cas for %s"), dir.data);
-				UBA_ASSERT(m_detoursBinaryKey != CasKeyZero);
+				UBA_ASSERT(detoursBinaryKey != CasKeyZero);
 			}
 		}
 
 		bool isValidVersion = clientVersion == SessionNetworkVersion;
 		if (binAsVersion)
-			isValidVersion = clientKeys[0] == m_agentBinaryKey && clientKeys[1] == m_detoursBinaryKey;
+			isValidVersion = clientKeys[0] == agentBinaryKey && clientKeys[1] == detoursBinaryKey;
 
 		writer.WriteBool(isValidVersion);
 
@@ -693,8 +696,8 @@ namespace uba
 				response.Appendf(TC("UbaAgent binaries mismatch. Disconnecting..."));
 			}
 			writer.WriteString(response);
-			writer.WriteCasKey(m_agentBinaryKey);
-			writer.WriteCasKey(m_detoursBinaryKey);
+			writer.WriteCasKey(agentBinaryKey);
+			writer.WriteCasKey(detoursBinaryKey);
 			return true;
 		}
 
@@ -721,7 +724,7 @@ namespace uba
 		if (!InitializeNameToHashTable())
 			return false;
 
-		writer.WriteCasKey(m_detoursBinaryKey);
+		writer.WriteCasKey(detoursBinaryKey);
 		writer.WriteBool(m_resetCas);
 		writer.WriteU32(sessionId);
 		writer.WriteU32(m_uiLanguage);
@@ -739,7 +742,7 @@ namespace uba
 
 	bool SessionServer::HandleEnsureBinaryFile(const ConnectionInfo& connectionInfo, BinaryReader& reader, BinaryWriter& writer)
 	{
-		u32 clientIsArm = reader.ReadBool();
+		bool clientIsArm = reader.ReadBool();
 		StringBuffer<> fileName;
 		reader.ReadString(fileName);
 		StringKey fileNameKey = reader.ReadStringKey();
diff --git a/Engine/Source/Programs/UnrealBuildAccelerator/Common/Public/UbaSessionServer.h b/Engine/Source/Programs/UnrealBuildAccelerator/Common/Public/UbaSessionServer.h
--- a/Engine/Source/Programs/UnrealBuildAccelerator/Common/Public/UbaSessionServer.h
+++ b/Engine/Source/Programs/UnrealBuildAccelerator/Common/Public/UbaSessionServer.h
@@ -136,8 +136,8 @@ namespace uba
 		u32 m_connectionCount = 0;
 
 		ReaderWriterLock m_binKeysLock;
-		CasKey m_detoursBinaryKey;
-		CasKey m_agentBinaryKey;
+		CasKey m_detoursBinaryKey[2];
+		CasKey m_agentBinaryKey[2];
 
 		struct ClientSession
 		{
