[UBA]
* Added handling for when stdin points to stderr in child processes

diff --git a/Engine/Source/Programs/UnrealBuildAccelerator/Detours/Private/Windows/UbaDetoursFunctionsKernelBase.inl b/Engine/Source/Programs/UnrealBuildAccelerator/Detours/Private/Windows/UbaDetoursFunctionsKernelBase.inl
--- a/Engine/Source/Programs/UnrealBuildAccelerator/Detours/Private/Windows/UbaDetoursFunctionsKernelBase.inl
+++ b/Engine/Source/Programs/UnrealBuildAccelerator/Detours/Private/Windows/UbaDetoursFunctionsKernelBase.inl
@@ -2768,8 +2768,13 @@ BOOL Detoured_CreateProcessW(LPCWSTR lpApplicationName, LPWSTR lpCommandLine, LP
 		DetouredHandle& dh = asDetouredHandle(lpStartupInfo->hStdInput);
 		if (dh.type == HandleType_StdIn)
 			lpStartupInfo->hStdInput = g_isDetachedProcess ? 0 : True_GetStdHandle(STD_INPUT_HANDLE);
+		else if (dh.type == HandleType_StdErr)
+			lpStartupInfo->hStdInput = g_isDetachedProcess ? 0 : True_GetStdHandle(STD_ERROR_HANDLE);
 		else
+		{
 			UBA_ASSERTF(false, L"hStdInput is detoured (%s)", lpApplicationName);
+			lpStartupInfo->hStdInput = dh.trueHandle;
+		}
 	}
 
 	lpStartupInfo->dwFlags |= STARTF_USESHOWWINDOW;
