[UBA]
* Suppress error message for ERROR_FILENAME_EXCED_RANGE because we don't want people reading the log to think this is an UBA issue

diff --git a/Engine/Source/Programs/UnrealBuildAccelerator/Common/Private/UbaProcess.cpp b/Engine/Source/Programs/UnrealBuildAccelerator/Common/Private/UbaProcess.cpp
--- a/Engine/Source/Programs/UnrealBuildAccelerator/Common/Private/UbaProcess.cpp
+++ b/Engine/Source/Programs/UnrealBuildAccelerator/Common/Private/UbaProcess.cpp
@@ -1061,6 +1061,11 @@ namespace uba
 
 		if (!result)
 		{
+			#if PLATFORM_WINDOWS
+			if (lastError == ERROR_FILENAME_EXCED_RANGE) // This is a command line issue so we dont' want uba to be blamed just because it outputs below logging entry
+				return true;
+			#endif
+
 			m_session.m_logger.Logf(LogEntryType_Info, TC("Detoured process failed to start child process - %s. %s (Working dir: %s)"), LastErrorToText(lastError).data, process.m_realApplication.c_str(), process.m_realWorkingDir);
 			return true;
 		}
