[UBA]
* Fixed mac compile errors

diff --git a/Engine/Source/Programs/UnrealBuildAccelerator/Common/Private/UbaSessionServer.cpp b/Engine/Source/Programs/UnrealBuildAccelerator/Common/Private/UbaSessionServer.cpp
--- a/Engine/Source/Programs/UnrealBuildAccelerator/Common/Private/UbaSessionServer.cpp
+++ b/Engine/Source/Programs/UnrealBuildAccelerator/Common/Private/UbaSessionServer.cpp
@@ -512,7 +512,7 @@ namespace uba
 
 	void SessionServer::RegisterCrossArchitectureMapping(const tchar* from, const tchar* to)
 	{
-		m_crossArchitectureMappings.emplace_back(from, to);
+		m_crossArchitectureMappings.emplace_back(CrossArchitectureMapping{from, to});
 	}
 
 	void SessionServer::OnDisconnected(const Guid& clientUid, u32 clientId)
