[UBA]
* Changed reading of named pipe on child process from ReadFile to instead use PeekNamedPipe+GetCodeProcess+ReadFile(if data exists).. this removed long weird stalls when building in non-detoured mode

diff --git a/Engine/Source/Programs/UnrealBuildAccelerator/Common/Private/UbaProcess.cpp b/Engine/Source/Programs/UnrealBuildAccelerator/Common/Private/UbaProcess.cpp
--- a/Engine/Source/Programs/UnrealBuildAccelerator/Common/Private/UbaProcess.cpp
+++ b/Engine/Source/Programs/UnrealBuildAccelerator/Common/Private/UbaProcess.cpp
@@ -1730,12 +1730,27 @@ namespace uba
 			TString currentString;
 			while (true)
 			{
-				char buf[4096];
-				DWORD readCount = 0;
-				if (!::ReadFile(readPipe, buf, sizeof(buf) - 1, &readCount, NULL))
+				DWORD exitCode = STILL_ACTIVE;
+				GetExitCodeProcess((HANDLE)m_nativeProcessHandle, &exitCode);
+				
+				while (true)
+				{
+					DWORD totalBytesAvailable = 0;
+					if (!PeekNamedPipe(readPipe, NULL, 0, NULL, &totalBytesAvailable, NULL))
+						break;
+					if (totalBytesAvailable == 0)
+						break;
+					char buf[4096];
+					DWORD readCount = 0;
+					if (!::ReadFile(readPipe, buf, sizeof(buf) - 1, &readCount, NULL))
+						break;
+					buf[readCount] = 0;
+					pipeReader.ReadData(buf, readCount);
+				}
+
+				if (exitCode != STILL_ACTIVE)
 					break;
-				buf[readCount] = 0;
-				pipeReader.ReadData(buf, readCount);
+				WaitForSingleObject((HANDLE)m_nativeProcessHandle, 1000);
 			}
 		}
 
