[UBA]
* Fixed bug where vctip.exe was detoured when allowfilemappings was false

diff --git a/Engine/Source/Programs/UnrealBuildAccelerator/Detours/Private/Windows/UbaDetoursFunctionsKernelBase.inl b/Engine/Source/Programs/UnrealBuildAccelerator/Detours/Private/Windows/UbaDetoursFunctionsKernelBase.inl
--- a/Engine/Source/Programs/UnrealBuildAccelerator/Detours/Private/Windows/UbaDetoursFunctionsKernelBase.inl
+++ b/Engine/Source/Programs/UnrealBuildAccelerator/Detours/Private/Windows/UbaDetoursFunctionsKernelBase.inl
@@ -2672,14 +2672,22 @@ BOOL Detoured_CreateProcessW(LPCWSTR lpApplicationName, LPWSTR lpCommandLine, LP
 	DETOURED_CALL(CreateProcessW);
 	DEBUG_LOG_DETOURED(L"CreateProcessW", L"%ls %ls CreationFlags: 0x%x StartupFlags: 0x%u Stdin: %llu WorkDir: %s", lpApplicationName, originalCmd, dwCreationFlags, lpStartupInfo->dwFlags, u64(lpStartupInfo->hStdInput), (lpCurrentDirectory ? lpCurrentDirectory : L""));
 
-	if ((!lpApplicationName || !*lpApplicationName) && !*originalCmd)
+	const tchar* testName = lpApplicationName;
+
+	if (!lpApplicationName || !*lpApplicationName)
 	{
-		SetLastError(ERROR_FILE_NOT_FOUND);
-		return FALSE;
+		if (!*originalCmd)
+		{
+			SetLastError(ERROR_FILE_NOT_FOUND);
+			return FALSE;
+		}
+
+		testName = originalCmd;
 	}
 
+
 	// Debug binaries started when process crash... we don't want to detour these.
-	if (Contains(originalCmd, L"winedbg") || Contains(originalCmd, L"werfault.exe") || Contains(originalCmd, L"vsjitdebugger.exe") || Contains(originalCmd, L"crashpad_handler.exe"))
+	if (Contains(testName, L"winedbg") || Contains(testName, L"werfault.exe") || Contains(testName, L"vsjitdebugger.exe") || Contains(testName, L"crashpad_handler.exe"))
 	{
 		if (g_runningRemote)
 		{
@@ -2697,7 +2705,7 @@ BOOL Detoured_CreateProcessW(LPCWSTR lpApplicationName, LPWSTR lpCommandLine, LP
 
 	bool isChild = true;
 	// We don't care about tracking mspdbsrv or vctip.. they are services just spawned by this process
-	if (Contains(originalCmd, L"mspdbsrv.exe") || Contains(originalCmd, L"vctip.exe") || Contains(originalCmd, L"git.exe"))
+	if (Contains(testName, L"mspdbsrv.exe") || Contains(testName, L"vctip.exe") || Contains(testName, L"git.exe"))
 	{
 		if (!g_runningRemote)
 		{
