[UBA]
* Improved error message when NtQueryDirectoryFile fails

diff --git a/Engine/Source/Programs/UnrealBuildAccelerator/Common/Public/UbaDirectoryIterator.h b/Engine/Source/Programs/UnrealBuildAccelerator/Common/Public/UbaDirectoryIterator.h
--- a/Engine/Source/Programs/UnrealBuildAccelerator/Common/Public/UbaDirectoryIterator.h
+++ b/Engine/Source/Programs/UnrealBuildAccelerator/Common/Public/UbaDirectoryIterator.h
@@ -139,7 +139,7 @@ namespace uba
 			if (res != STATUS_SUCCESS)
 			{
 				if (res != STATUS_NO_MORE_FILES)
-					logger.Error(L"NtQueryDirectoryFile returned error 0x%x", res);
+					logger.Error(L"NtQueryDirectoryFile for %s returned error 0x%x", str.data, res);
 				break;
 			}
 
