[UBA]
* Fixed handle leak in code used for compressed obj files
* Added so logger and hint is sent into file mapping functions
* Added debug code for tracking file mapping handles
* Fixed so mapping handles are cleaned up properly when session instance is destroyed
* Exposed session "AllowLocalDetour" to config file. Defaults to true. This can be set to false to disable all detouring of spawned processes

diff --git a/Engine/Source/Programs/UnrealBuildAccelerator/Common/Private/UbaFileAccessor.cpp b/Engine/Source/Programs/UnrealBuildAccelerator/Common/Private/UbaFileAccessor.cpp
--- a/Engine/Source/Programs/UnrealBuildAccelerator/Common/Private/UbaFileAccessor.cpp
+++ b/Engine/Source/Programs/UnrealBuildAccelerator/Common/Private/UbaFileAccessor.cpp
@@ -142,11 +142,11 @@ namespace uba
 			realFileName = tempFile.Append(m_tempPath).Append("Temp_").AppendValue(m_tempFileIndex).data;
 		#endif
 
-		m_mappingHandle = uba::CreateFileMappingW(m_fileHandle, PAGE_READWRITE, size, realFileName);
+		m_mappingHandle = uba::CreateFileMappingW(m_logger, m_fileHandle, PAGE_READWRITE, size, realFileName);
 		if (!m_mappingHandle.IsValid())
 			return m_logger.Error(TC("Failed to create memory map %s (%s)"), realFileName, LastErrorToText().data);
 
-		m_data = MapViewOfFile(m_mappingHandle, FILE_MAP_WRITE, 0, size);
+		m_data = MapViewOfFile(m_logger, m_mappingHandle, FILE_MAP_WRITE, 0, size);
 		if (!m_data)
 			return m_logger.Error(TC("Failed to map view of file %s with size %llu, for write (%s)"), realFileName, size, LastErrorToText().data);
 
@@ -341,9 +341,9 @@ namespace uba
 		m_size = info.size;
 #if PLATFORM_WINDOWS
 		if (m_size)
-			m_mappingHandle = uba::CreateFileMappingW(m_fileHandle, PAGE_READONLY, m_size, m_fileName);
+			m_mappingHandle = uba::CreateFileMappingW(m_logger, m_fileHandle, PAGE_READONLY, m_size, m_fileName);
 		else
-			m_mappingHandle = uba::CreateFileMappingW(InvalidFileHandle, PAGE_READONLY, 1, m_fileName);
+			m_mappingHandle = uba::CreateFileMappingW(m_logger, InvalidFileHandle, PAGE_READONLY, 1, m_fileName);
 
 		if (!m_mappingHandle.IsValid())
 			return m_logger.Error(TC("Failed to create mapping handle for %s with size %llu (%s)"), m_fileName, m_size, LastErrorToText().data);
@@ -352,7 +352,7 @@ namespace uba
 		if (offset == m_size)
 			return true;
 #endif
-		m_data = MapViewOfFile(m_mappingHandle, FILE_MAP_READ, offset, m_size);
+		m_data = MapViewOfFile(m_logger, m_mappingHandle, FILE_MAP_READ, offset, m_size);
 		if (!m_data)
 			return m_logger.Error(TC("%s - MapViewOfFile failed (%s)"), m_fileName, LastErrorToText().data);
 
@@ -368,14 +368,14 @@ namespace uba
 	{
 		if (m_data)
 		{
-			if (!UnmapViewOfFile(m_data, m_size, m_fileName))
+			if (!UnmapViewOfFile(m_logger, m_data, m_size, m_fileName))
 				return m_logger.Error(TC("Failed to unmap memory for %s (%s)"), m_fileName, LastErrorToText().data);
 			m_data = nullptr;
 		}
 
 		if (m_mappingHandle.IsValid())
 		{
-			if (!CloseFileMapping(m_mappingHandle))
+			if (!CloseFileMapping(m_logger, m_mappingHandle, m_fileName))
 				return m_logger.Error(TC("Failed to close file mapping for %s (%s)"), m_fileName, LastErrorToText().data);
 			m_mappingHandle = {};
 		}
diff --git a/Engine/Source/Programs/UnrealBuildAccelerator/Common/Private/UbaFileMapping.cpp b/Engine/Source/Programs/UnrealBuildAccelerator/Common/Private/UbaFileMapping.cpp
--- a/Engine/Source/Programs/UnrealBuildAccelerator/Common/Private/UbaFileMapping.cpp
+++ b/Engine/Source/Programs/UnrealBuildAccelerator/Common/Private/UbaFileMapping.cpp
@@ -13,25 +13,46 @@
 #include <sys/file.h>
 #endif
 
+#define UBA_DEBUG_FILE_MAPPING 0
+
 namespace uba
 {
+	#if UBA_DEBUG_FILE_MAPPING
+	ReaderWriterLock g_fileMappingsLock;
+	struct DebugFileMapping { TString hint; Atomic<u64> viewCount; };
+	UnorderedMap<HANDLE, DebugFileMapping> g_fileMappings;
+	UnorderedMap<const void*, HANDLE> g_viewMappings;
+	#endif
+
+
 #if PLATFORM_WINDOWS
 	Bottleneck& g_createFileHandleBottleneck = *new Bottleneck(8); // Allocated and leaked just to prevent shutdown asserts in debug
 
-	//Atomic<u64> g_fileMappingCount;
-
 	HANDLE asHANDLE(FileHandle fh);
 
-	HANDLE InternalCreateFileMappingW(HANDLE hFile, DWORD flProtect, DWORD dwMaximumSizeHigh, DWORD dwMaximumSizeLow, LPCWSTR lpName)
+	HANDLE InternalCreateFileMappingW(Logger& logger, HANDLE hFile, DWORD flProtect, DWORD dwMaximumSizeHigh, DWORD dwMaximumSizeLow, LPCWSTR lpName, const tchar* hint)
 	{
-		//++g_fileMappingCount;
+		HANDLE res;
+		if (flProtect == PAGE_READWRITE)
+		{
+			// Experiment to try to prevent lock happening on AWS servers when lots of helpers are sending back obj files.
+			BottleneckScope bs(g_createFileHandleBottleneck);
+			res = ::CreateFileMappingW(hFile, NULL, flProtect, dwMaximumSizeHigh, dwMaximumSizeLow, lpName);
+		}
+		else
+			res = ::CreateFileMappingW(hFile, NULL, flProtect, dwMaximumSizeHigh, dwMaximumSizeLow, lpName);
 
-		if (flProtect != PAGE_READWRITE)
-			return ::CreateFileMappingW(hFile, NULL, flProtect, dwMaximumSizeHigh, dwMaximumSizeLow, lpName);
+		#if UBA_DEBUG_FILE_MAPPING
+		if (res)
+		{
+			SCOPED_WRITE_LOCK(g_fileMappingsLock, lock);
+			auto insres = g_fileMappings.try_emplace(res, hint);(void)insres;
+			UBA_ASSERT(insres.second);
+			//logger.Info(TC("FILEMAPPING CREAT: %s %llu (%llu)"), hint, u64(res), g_fileMappings.size());
+		}
+		#endif
 
-		// Experiment to try to prevent lock happening on AWS servers when lots of helpers are sending back obj files.
-		BottleneckScope bs(g_createFileHandleBottleneck);
-		return ::CreateFileMappingW(hFile, NULL, flProtect, dwMaximumSizeHigh, dwMaximumSizeLow, lpName);
+		return res;
 	}
 #else
 	int asFileDescriptor(FileHandle fh);
@@ -41,11 +62,11 @@ namespace uba
 	Atomic<u64> g_mappedViewCounter;
 #endif
 
-	FileMappingHandle CreateMemoryMappingW(Logger& logger, u32 flProtect, u64 maxSize, const tchar* name)
+	FileMappingHandle CreateMemoryMappingW(Logger& logger, u32 flProtect, u64 maxSize, const tchar* name, const tchar* hint)
 	{
 		ExtendedTimerScope ts(KernelStats::GetCurrent().createFileMapping);
 #if PLATFORM_WINDOWS
-		return { InternalCreateFileMappingW(INVALID_HANDLE_VALUE, flProtect, (DWORD)ToHigh(maxSize), ToLow(maxSize), name) };
+		return { InternalCreateFileMappingW(logger, INVALID_HANDLE_VALUE, flProtect, (DWORD)ToHigh(maxSize), ToLow(maxSize), name, hint) };
 #else
 		UBA_ASSERT(!name);
 		UBA_ASSERT((flProtect & (~u32(PAGE_READWRITE | SEC_RESERVE))) == 0);
@@ -174,11 +195,11 @@ namespace uba
 #endif
 	}
 
-	FileMappingHandle CreateFileMappingW(FileHandle file, u32 protect, u64 maxSize, const tchar* hint)
+	FileMappingHandle CreateFileMappingW(Logger& logger, FileHandle file, u32 protect, u64 maxSize, const tchar* hint)
 	{
 		ExtendedTimerScope ts(KernelStats::GetCurrent().createFileMapping);
 #if PLATFORM_WINDOWS
-		return { InternalCreateFileMappingW(asHANDLE(file), protect, (DWORD)ToHigh(maxSize), ToLow(maxSize), NULL) };
+		return { InternalCreateFileMappingW(logger, asHANDLE(file), protect, (DWORD)ToHigh(maxSize), ToLow(maxSize), NULL, hint) };
 #else
 		FileMappingHandle h;
 		int fd = asFileDescriptor(file);
@@ -212,11 +233,13 @@ namespace uba
 #endif
 	}
 
-	u8* MapViewOfFile(FileMappingHandle fileMappingObject, u32 desiredAccess, u64 offset, u64 bytesToMap)
+	u8* MapViewOfFile(Logger& logger, FileMappingHandle fileMappingObject, u32 desiredAccess, u64 offset, u64 bytesToMap)
 	{
 		ExtendedTimerScope ts(KernelStats::GetCurrent().mapViewOfFile);
 #if PLATFORM_WINDOWS
-		return (u8*)::MapViewOfFile(fileMappingObject.handle, desiredAccess, (DWORD)ToHigh(offset), ToLow(offset), bytesToMap);
+		u8* res = (u8*)::MapViewOfFile(fileMappingObject.handle, desiredAccess, (DWORD)ToHigh(offset), ToLow(offset), bytesToMap);
+		if (!res)
+			return nullptr;
 #else
 		int prot = 0;
 		if (desiredAccess & FILE_MAP_READ)
@@ -226,15 +249,28 @@ namespace uba
 		UBA_ASSERT(fileMappingObject.IsValid());
 		int shmFd = fileMappingObject.shmFd;
 		void* rptr = mmap(NULL, bytesToMap, prot, MAP_SHARED, shmFd, s64(offset));
-		if (rptr != MAP_FAILED)
+		if (rptr == MAP_FAILED)
 		{
-			++g_mappedViewCounter;
-			return (u8*)rptr;
+			SetLastError(errno);
+			return nullptr;
 		}
+		u8* res = (u8*)rptr;
+		++g_mappedViewCounter;
 		//UBA_ASSERTF(false, "Failed to map file with fd %i, desiredAccess %u offset %llu, bytesToMap %llu (%s)", fd, desiredAccess, offset, bytesToMap, strerror(errno));
-		SetLastError(errno);
-		return nullptr;
 #endif
+
+		#if UBA_DEBUG_FILE_MAPPING
+		{
+			SCOPED_WRITE_LOCK(g_fileMappingsLock, lock);
+			auto findIt = g_fileMappings.find(fileMappingObject.handle);
+			UBA_ASSERTF(findIt != g_fileMappings.end(), TC("Mapping nonexisting handle: %llu"), u64(fileMappingObject.handle));
+			auto insres = g_viewMappings.try_emplace(res, fileMappingObject.handle);
+			UBA_ASSERT(insres.second);
+			++findIt->second.viewCount;
+		}
+		#endif
+
+		return res;
 	}
 
 	bool MapViewCommit(void* address, u64 size)
@@ -246,11 +282,25 @@ namespace uba
 #endif
 	}
 
-	bool UnmapViewOfFile(const void* lpBaseAddress, u64 bytesToUnmap, const tchar* hint)
+	bool UnmapViewOfFile(Logger& logger, const void* lpBaseAddress, u64 bytesToUnmap, const tchar* hint)
 	{
 		ExtendedTimerScope ts(KernelStats::GetCurrent().unmapViewOfFile);
+
+		#if UBA_DEBUG_FILE_MAPPING
+		{
+			SCOPED_WRITE_LOCK(g_fileMappingsLock, lock);
+			auto findIt = g_viewMappings.find(lpBaseAddress);
+			UBA_ASSERT(findIt != g_viewMappings.end());
+			auto findIt2 = g_fileMappings.find(findIt->second);
+			g_viewMappings.erase(findIt);
+			UBA_ASSERTF(findIt2 != g_fileMappings.end(), TC("Unmap nonexisting handle: %llu"), u64(findIt->second));
+			--findIt2->second.viewCount;
+		}
+		#endif
+
 #if PLATFORM_WINDOWS
-		(void)bytesToUnmap; return ::UnmapViewOfFile(lpBaseAddress);
+		(void)bytesToUnmap;
+		return ::UnmapViewOfFile(lpBaseAddress);
 #else
 		UBA_ASSERTF(bytesToUnmap, TC("bytesToUnmap is zero unmapping %p (%s)"), lpBaseAddress, hint);
 		if (munmap((void*)lpBaseAddress, bytesToUnmap) == 0)
@@ -263,17 +313,21 @@ namespace uba
 #endif
 	}
 
-	bool CloseFileMapping(FileMappingHandle h)
+	bool CloseFileMapping(Logger& logger, FileMappingHandle h, const tchar* hint)
 	{
-#if PLATFORM_WINDOWS
-		//if (h.handle)
-		//{
-		//	--g_fileMappingCount;
-		//	StringBuffer<256> sb;
-		//	sb.Appendf(TC("FILEMAPPING COUNT: %llu\r\n"), g_fileMappingCount.load());
-		//	OutputDebugStringW(sb.data);
-		//}
+		#if UBA_DEBUG_FILE_MAPPING
+		if (h.IsValid())
+		{
+			SCOPED_WRITE_LOCK(g_fileMappingsLock, lock);
+			auto findIt = g_fileMappings.find(h.handle);
+			UBA_ASSERTF(findIt != g_fileMappings.end(), TC("Handle: %llu"), u64(h.handle));
+			UBA_ASSERTF(!findIt->second.viewCount.load(), TC("Closing file mapping with %llu open views (%s)"), findIt->second.viewCount.load(), findIt->second.hint.c_str());
+			//logger.Info(TC("FILEMAPPING CLOSE: %s %llu (%llu)"), findIt->second.hint.c_str(), u64(h.handle), g_fileMappings.size() - 1);
+			g_fileMappings.erase(findIt);
+		}
+		#endif
 
+#if PLATFORM_WINDOWS
 		return CloseHandle(h.handle);
 #else
 		if (h.shmFd == -1)
@@ -300,14 +354,27 @@ namespace uba
 #endif
 	}
 
-	bool DuplicateFileMapping(ProcHandle hSourceProcessHandle, FileMappingHandle hSourceHandle, ProcHandle hTargetProcessHandle, FileMappingHandle* lpTargetHandle, u32 dwDesiredAccess, bool bInheritHandle, u32 dwOptions)
+	bool DuplicateFileMapping(Logger& logger, ProcHandle hSourceProcessHandle, FileMappingHandle hSourceHandle, ProcHandle hTargetProcessHandle, FileMappingHandle* lpTargetHandle, u32 dwDesiredAccess, bool bInheritHandle, u32 dwOptions, const tchar* hint)
 	{
 #if PLATFORM_WINDOWS
-		return DuplicateHandle((HANDLE)hSourceProcessHandle, hSourceHandle.handle, (HANDLE)hTargetProcessHandle, &lpTargetHandle->handle, dwDesiredAccess, bInheritHandle, dwOptions);
+		if (!DuplicateHandle((HANDLE)hSourceProcessHandle, hSourceHandle.handle, (HANDLE)hTargetProcessHandle, &lpTargetHandle->handle, dwDesiredAccess, bInheritHandle, dwOptions))
+			return false;
 #else
 		UBA_ASSERT(false);
-		return false;
+		if (true)
+			return false;
 #endif
+
+		#if UBA_DEBUG_FILE_MAPPING
+		{
+			SCOPED_WRITE_LOCK(g_fileMappingsLock, lock);
+			auto insres = g_fileMappings.try_emplace(lpTargetHandle->handle, hint);(void)insres;
+			UBA_ASSERT(insres.second);
+			//logger.Info(TC("FILEMAPPING DUPLI: %s %llu (%llu)"), hint, lpTargetHandle->handle, g_fileMappings.size());
+		}
+		#endif
+
+		return true;
 	}
 
 	void MapMemoryCopy(void* dest, const void* source, u64 size)
@@ -345,14 +412,13 @@ namespace uba
 
 	bool FileMappingBuffer::AddPersistent(const tchar* name, FileHandle fileHandle, u64 size, u64 capacity)
 	{
-		FileMappingHandle sparseMemoryHandle = uba::CreateFileMappingW(fileHandle, PAGE_READWRITE, capacity, name);
+		FileMappingHandle sparseMemoryHandle = uba::CreateFileMappingW(m_logger, fileHandle, PAGE_READWRITE, capacity, name);
 		if (!sparseMemoryHandle.IsValid())
 		{
 			m_logger.Error(TC("Failed to create file mapping (%s)"), LastErrorToText().data);
 			return false;
 		}
 
-
 		MappingStorage& storage = m_storage[MappedView_Persistent];
 		File& f = storage.files[storage.fileCount++];
 		f.name = name;
@@ -401,7 +467,7 @@ namespace uba
 		}
 
 		u64 mapSize = alignedOffsetEnd - alignedOffsetStart;
-		u8* data = MapViewOfFile(f.handle, FILE_MAP_WRITE, alignedOffsetStart, mapSize);
+		u8* data = MapViewOfFile(m_logger, f.handle, FILE_MAP_WRITE, alignedOffsetStart, mapSize);
 		if (!data)
 		{
 			m_logger.Error(TC("%s - AllocAndMapView failed to map view of file for %s with size %llu and offset %llu (%s)"), f.name, hint, size, alignment, LastErrorToText().data);
@@ -417,7 +483,7 @@ namespace uba
 			u64 commitSize = committedAfter - committedBefore;
 			if (!MapViewCommit(data + commitStart, commitSize))
 			{
-				UnmapViewOfFile(data, mapSize, hint);
+				UnmapViewOfFile(m_logger, data, mapSize, hint);
 				m_logger.Error(TC("%s - Failed to allocate memory for %s (%s)"), f.name, hint, LastErrorToText().data);
 				return res;
 			}
@@ -466,7 +532,7 @@ namespace uba
 				{
 					UBA_ASSERT(&storage == &m_storage[MappedView_Transient]);
 					u64 capacity = (IsWindows ? 32ull : 8ull) * 1024 * 1024 * 1024; // Linux can't have larger than 8gb
-					file-> handle = uba::CreateMemoryMappingW(m_logger, PAGE_READWRITE | SEC_RESERVE, capacity);
+					file-> handle = uba::CreateMemoryMappingW(m_logger, PAGE_READWRITE | SEC_RESERVE, capacity, nullptr, TC("FileMappingBuffer"));
 					if (!file->handle.IsValid())
 					{
 						m_logger.Error(TC("%s - Failed to create memory map (%s)"), file->name, LastErrorToText().data);
@@ -516,7 +582,7 @@ namespace uba
 		for (u32 i=0; i!=storage.fileCount; ++i)
 		{
 			//UBA_ASSERT(!storage.files[i].activeMapCount);
-			CloseFileMapping(storage.files[i].handle);
+			CloseFileMapping(m_logger, storage.files[i].handle, TC("FileMappingBuffer"));
 			CloseFile(nullptr, storage.files[i].file);
 		}
 		storage.fileCount = 0;
@@ -536,7 +602,7 @@ namespace uba
 
 		MappedView res;
 
-		u8* data = MapViewOfFile(handle, FILE_MAP_WRITE, alignedOffsetStart, alignedOffsetEnd - alignedOffsetStart);
+		u8* data = MapViewOfFile(m_logger, handle, FILE_MAP_WRITE, alignedOffsetStart, alignedOffsetEnd - alignedOffsetStart);
 		if (!data)
 		{
 			m_logger.Error(TC("%s - MapView failed to map view of file for %s with size %llu and offset %llu (%s)"), file.name, hint, size, offset, LastErrorToText().data);
@@ -567,7 +633,7 @@ namespace uba
 
 				u8* memory = view.memory - (view.offset - alignedOffsetStart);
 				u64 mapSize = alignedOffsetEnd - alignedOffsetStart;
-				if (!UnmapViewOfFile(memory, mapSize, hint))
+				if (!UnmapViewOfFile(m_logger, memory, mapSize, hint))
 				{
 					m_logger.Error(TC("%s - Failed to unmap view on address %llx (offset %llu) - %s (%s)"), file.name, u64(memory), view.offset, hint, LastErrorToText().data);
 				}
@@ -622,12 +688,12 @@ namespace uba
 	FileMappingAllocator::~FileMappingAllocator()
 	{
 		if (m_mappingHandle.IsValid())
-			CloseFileMapping(m_mappingHandle);
+			CloseFileMapping(m_logger, m_mappingHandle, TC("FileMappingAllocator"));
 	}
 
 	bool FileMappingAllocator::Init(u64 blockSize, u64 capacity)
 	{
-		m_mappingHandle = uba::CreateMemoryMappingW(m_logger, PAGE_READWRITE|SEC_RESERVE, capacity);
+		m_mappingHandle = uba::CreateMemoryMappingW(m_logger, PAGE_READWRITE|SEC_RESERVE, capacity, nullptr, TC("FileMappingAllocator"));
 		if (!m_mappingHandle.IsValid())
 			return m_logger.Error(TC("%s - Failed to create memory map with capacity %llu (%s)"), m_name, capacity, LastErrorToText().data);
 
@@ -657,7 +723,7 @@ namespace uba
 		lock.Leave();
 
 		u64 offset = index*m_blockSize;
-		u8* data = MapViewOfFile(m_mappingHandle, FILE_MAP_READ|FILE_MAP_WRITE, offset, m_blockSize);
+		u8* data = MapViewOfFile(m_logger, m_mappingHandle, FILE_MAP_READ|FILE_MAP_WRITE, offset, m_blockSize);
 		if (!data)
 		{
 			if (m_capacity < m_mappingCount*m_blockSize)
@@ -681,7 +747,7 @@ namespace uba
 	void FileMappingAllocator::Free(Allocation allocation)
 	{
 		UBA_ASSERT(allocation.handle == m_mappingHandle);
-		if (!UnmapViewOfFile(allocation.memory, m_blockSize, m_name))
+		if (!UnmapViewOfFile(m_logger, allocation.memory, m_blockSize, m_name))
 			m_logger.Error(TC("%s - Failed to unmap view of file (%s)"), m_name, LastErrorToText().data);
 		u64 index = allocation.offset / m_blockSize;
 		SCOPED_WRITE_LOCK(m_mappingLock, lock);
diff --git a/Engine/Source/Programs/UnrealBuildAccelerator/Common/Private/UbaProcess.cpp b/Engine/Source/Programs/UnrealBuildAccelerator/Common/Private/UbaProcess.cpp
--- a/Engine/Source/Programs/UnrealBuildAccelerator/Common/Private/UbaProcess.cpp
+++ b/Engine/Source/Programs/UnrealBuildAccelerator/Common/Private/UbaProcess.cpp
@@ -139,7 +139,7 @@ namespace uba
 		{
 			for (auto& pair : m_writtenFiles)
 				if (pair.second.mappingHandle.IsValid())
-					CloseFileMapping(pair.second.mappingHandle);
+					CloseFileMapping(m_session.m_logger, pair.second.mappingHandle, pair.first.c_str());
 			delete &m_writtenFiles;
 			delete &m_writtenFilesLock;
 
@@ -1340,7 +1340,7 @@ namespace uba
 		FileMappingHandle source;
 		source.FromU64(mappingHandle);
 		FileMappingHandle newHandle;
-		if (!DuplicateFileMapping(nativeProcessHandle, source, GetCurrentProcessHandle(), &newHandle, 0, false, DUPLICATE_SAME_ACCESS))
+		if (!DuplicateFileMapping(m_session.m_logger, nativeProcessHandle, source, GetCurrentProcessHandle(), &newHandle, 0, false, DUPLICATE_SAME_ACCESS, fileName.data))
 		{
 			m_session.m_logger.Error(TC("Failed to duplicate handle for temp file (%s)"), fileName.data);
 			return true;
@@ -1357,7 +1357,7 @@ namespace uba
 		tempFile.mappingHandle = newHandle;
 		tempFile.mappingWritten = mappingHandleSize;
 		tempLock.Leave();
-		CloseFileMapping(oldMapping);
+		CloseFileMapping(m_session.m_logger, oldMapping, fileName.data);
 		return true;
 	}
 
@@ -2143,7 +2143,7 @@ namespace uba
 	{
 		for (auto& pair : m_tempFiles)
 			if (pair.second.mappingHandle.IsValid())
-				CloseFileMapping(pair.second.mappingHandle);
+				CloseFileMapping(m_session.m_logger, pair.second.mappingHandle, pair.second.name.c_str());
 		m_tempFiles.clear();
 	}
 
diff --git a/Engine/Source/Programs/UnrealBuildAccelerator/Common/Private/UbaScheduler.cpp b/Engine/Source/Programs/UnrealBuildAccelerator/Common/Private/UbaScheduler.cpp
--- a/Engine/Source/Programs/UnrealBuildAccelerator/Common/Private/UbaScheduler.cpp
+++ b/Engine/Source/Programs/UnrealBuildAccelerator/Common/Private/UbaScheduler.cpp
@@ -294,7 +294,7 @@ namespace uba
 
 	void Scheduler::RemoteSlotAvailable(bool isCrossArchitecture)
 	{
-		UBA_ASSERTF(false, TC("Cross architecture code path not implemented"));
+		UBA_ASSERTF(!isCrossArchitecture, TC("Cross architecture code path not implemented"));
 		RunQueuedProcess(false);
 	}
 
diff --git a/Engine/Source/Programs/UnrealBuildAccelerator/Common/Private/UbaSession.cpp b/Engine/Source/Programs/UnrealBuildAccelerator/Common/Private/UbaSession.cpp
--- a/Engine/Source/Programs/UnrealBuildAccelerator/Common/Private/UbaSession.cpp
+++ b/Engine/Source/Programs/UnrealBuildAccelerator/Common/Private/UbaSession.cpp
@@ -221,6 +221,7 @@ namespace uba
 		table.GetValueAsBool(detailedTrace, TC("DetailedTrace"));
 		table.GetValueAsBool(traceChildProcesses, TC("TraceChildProcesses"));
 		table.GetValueAsBool(storeObjFilesCompressed, TC("StoreObjFilesCompressed"));
+		table.GetValueAsBool(allowLocalDetour, TC("AllowLocalDetour"));
 		table.GetValueAsBool(extractObjFilesSymbols, TC("ExtractObjFilesSymbols"));
 	}
 
@@ -1054,7 +1055,7 @@ namespace uba
 		UBA_ASSERTF(m_shouldWriteToDisk || m_allowMemoryMaps, TC("Can't disable both should write to disk and allow memory maps"));
 
 		m_storeObjFilesCompressed = info.storeObjFilesCompressed;
-		m_allowLocalDetour = true;
+		m_allowLocalDetour = info.allowLocalDetour;
 
 		m_detailedTrace = info.detailedTrace;
 		m_traceChildProcesses = info.traceChildProcesses;
@@ -1122,14 +1123,14 @@ namespace uba
 		m_systemPath.Append(TC("/nonexistingpath"));
 		#endif
 
-		m_fileMappingTableHandle = uba::CreateMemoryMappingW(m_logger, PAGE_READWRITE, FileMappingTableMemSize);
+		m_fileMappingTableHandle = uba::CreateMemoryMappingW(m_logger, PAGE_READWRITE, FileMappingTableMemSize, nullptr, TC("FileMappings"));
 		UBA_ASSERT(m_fileMappingTableHandle.IsValid());
-		m_fileMappingTableMem = MapViewOfFile(m_fileMappingTableHandle, FILE_MAP_WRITE, 0, FileMappingTableMemSize);
+		m_fileMappingTableMem = MapViewOfFile(m_logger, m_fileMappingTableHandle, FILE_MAP_WRITE, 0, FileMappingTableMemSize);
 		UBA_ASSERT(m_fileMappingTableMem);
 
-		m_directoryTableHandle = uba::CreateMemoryMappingW(m_logger, PAGE_READWRITE, DirTableMemSize);
+		m_directoryTableHandle = uba::CreateMemoryMappingW(m_logger, PAGE_READWRITE, DirTableMemSize, nullptr, TC("DirMappings"));
 		UBA_ASSERT(m_directoryTableHandle.IsValid());
-		m_directoryTableMem = MapViewOfFile(m_directoryTableHandle, FILE_MAP_WRITE, 0, DirTableMemSize);
+		m_directoryTableMem = MapViewOfFile(m_logger, m_directoryTableHandle, FILE_MAP_WRITE, 0, DirTableMemSize);
 		UBA_ASSERT(m_directoryTableMem);
 
 		m_directoryTable.m_memory = m_directoryTableMem;
@@ -1206,16 +1207,17 @@ namespace uba
 		FlushDeadProcesses();
 
 		for (auto& kv : m_virtualSourceFiles)
-			CloseFileMapping(kv.second);
+			CloseFileMapping(m_logger, kv.second, TC("VirtualFile"));
 
-		//for (auto& i : m_fileMappingTableLookup)
-		//	CloseHandle(i.second.mapping);
+		for (auto& i : m_fileMappingTableLookup)
+			if (i.second.owned)
+				CloseFileMapping(m_logger, i.second.mapping, TC("FileMappingKeptFromOutput"));
 
-		UnmapViewOfFile(m_fileMappingTableMem, FileMappingTableMemSize, TC("FileMappingTable"));
-		CloseFileMapping(m_fileMappingTableHandle);
+		UnmapViewOfFile(m_logger, m_fileMappingTableMem, FileMappingTableMemSize, TC("FileMappingTable"));
+		CloseFileMapping(m_logger, m_fileMappingTableHandle, TC("FileMappingTable"));
 
-		UnmapViewOfFile(m_directoryTableMem, DirTableMemSize, TC("DirectoryTable"));
-		CloseFileMapping(m_directoryTableHandle);
+		UnmapViewOfFile(m_logger, m_directoryTableMem, DirTableMemSize, TC("DirectoryTable"));
+		CloseFileMapping(m_logger, m_directoryTableHandle, TC("DirectoryTable"));
 
 		//#if !UBA_DEBUG
 		//u32 count;
@@ -1976,7 +1978,7 @@ namespace uba
 		writtenLock.Leave();
 
 		if (h.IsValid())
-			CloseFileMapping(h);
+			CloseFileMapping(process.m_session.GetLogger(), h, name.c_str());
 	}
 
 	bool Session::CloseFile(CloseFileResponse& out, const CloseFileMessage& msg)
@@ -2047,7 +2049,7 @@ namespace uba
 				{
 					if (writtenFile.mappingHandle.IsValid())
 					{
-						CloseFileMapping(writtenFile.mappingHandle);
+						CloseFileMapping(m_logger, writtenFile.mappingHandle, msg.fileName.data);
 						writtenFile.mappingHandle = {};
 					}
 				}
@@ -2061,7 +2063,7 @@ namespace uba
 				{
 					FileMappingHandle source;
 					source.FromU64(msg.mappingHandle);
-					if (!DuplicateFileMapping(msg.process.m_nativeProcessHandle, source, GetCurrentProcessHandle(), &mappingHandle, FILE_MAP_READ, false, 0))
+					if (!DuplicateFileMapping(m_logger, msg.process.m_nativeProcessHandle, source, GetCurrentProcessHandle(), &mappingHandle, FILE_MAP_READ, false, 0, msgName))
 						return m_logger.Error(TC("Failed to duplicate file mapping handle for %s"), name);
 				}
 				writtenFile.key = key;
@@ -2243,13 +2245,13 @@ namespace uba
 				#endif
 
 				u64 fileSize = file.mappingWritten;
-				u8* mem = MapViewOfFile(file.mappingHandle, FILE_MAP_READ, 0, fileSize);
+				u8* mem = MapViewOfFile(m_logger, file.mappingHandle, FILE_MAP_READ, 0, fileSize);
 				if (!mem)
 					return m_logger.Error(TC("Failed to map view of filehandle for read %s (%s)"), file.name.c_str(), LastErrorToText().data);
 
 				//PrefetchVirtualMemory(mem, fileSize);
 
-				auto memClose = MakeGuard([&](){ UnmapViewOfFile(mem, fileSize, file.name.c_str()); });
+				auto memClose = MakeGuard([&](){ UnmapViewOfFile(m_logger, mem, fileSize, file.name.c_str()); });
 
 				if (m_storeObjFilesCompressed && process.m_startInfo.rules->StoreFileCompressed(file.name))
 				{
@@ -2318,11 +2320,10 @@ namespace uba
 
 			if (shouldEvictFromMemory)
 			{
-				m_workManager->AddWork([mh = file.mappingHandle]()
+				m_workManager->AddWork([mh = file.mappingHandle, n = file.name, this]()
 					{
-						CloseFileMapping(mh);
+						CloseFileMapping(m_logger, mh, n.c_str());
 					}, 1, TC("CFM"));
-				//CloseFileMapping(file.mappingHandle);
 			}
 			else
 			{
@@ -2330,6 +2331,7 @@ namespace uba
 				Storage::GetMappingString(name, file.mappingHandle, 0);
 				SCOPED_WRITE_LOCK(m_fileMappingTableLookupLock, lookupLock);
 				auto insres = m_fileMappingTableLookup.try_emplace(file.key);
+				UBA_ASSERT(insres.second);
 				FileMappingEntry& entry = insres.first->second;
 				lookupLock.Leave();
 				SCOPED_WRITE_LOCK(entry.lock, entryCs);
@@ -2339,6 +2341,7 @@ namespace uba
 				entry.size = file.mappingWritten;
 				entry.isDir = false;
 				entry.success = true;
+				entry.owned = true;
 
 				SCOPED_WRITE_LOCK(m_fileMappingTableMemLock, lock);
 				BinaryWriter writer(m_fileMappingTableMem, m_fileMappingTableSize);
@@ -2519,7 +2522,7 @@ namespace uba
 				if (fileHandle == InvalidFileHandle)
 					return m_logger.Error(TC("[InternalRegisterVirtualFile] CreateFileW for %s failed (%s)"), sourceFile, LastErrorToText().data);
 				auto fg = MakeGuard([&]() { uba::CloseFile(sourceFile, fileHandle); });
-				insres.first->second = CreateFileMappingW(fileHandle, PAGE_READONLY, 0, sourceFile);
+				insres.first->second = CreateFileMappingW(m_logger, fileHandle, PAGE_READONLY, 0, sourceFile);
 			}
 			mappingHandle = insres.first->second;
 			if (!mappingHandle.IsValid())
@@ -2537,7 +2540,7 @@ namespace uba
 		if (entry.handled)
 			return m_logger.Error(TC("Virtual file %s has already been registered"), filePath.data);
 
-		entry.mapping = mappingHandle; // TODO: This mapping handle will leak until process close.. not good
+		entry.mapping = mappingHandle;
 		entry.mappingOffset = sourceOffset;
 		entry.size = sourceSize;
 		entry.handled = true;
@@ -2915,14 +2918,14 @@ namespace uba
 		source.FromU64(msg.mappingHandle);
 		FileMappingHandle objectFileMappingHandle;
 
-		if (!DuplicateFileMapping(msg.process.m_nativeProcessHandle, source, GetCurrentProcessHandle(), &objectFileMappingHandle, FILE_MAP_ALL_ACCESS, false, 0))
+		if (!DuplicateFileMapping(m_logger, msg.process.m_nativeProcessHandle, source, GetCurrentProcessHandle(), &objectFileMappingHandle, FILE_MAP_ALL_ACCESS, false, 0, fileName))
 			return m_logger.Error(TC("Failed to duplicate file mapping handle for %s"), fileName);
-		auto ofmh = MakeGuard([&]() { CloseFileMapping(objectFileMappingHandle); });
+		auto ofmh = MakeGuard([&]() { CloseFileMapping(m_logger, objectFileMappingHandle, fileName); });
 
-		u8* mem = MapViewOfFile(objectFileMappingHandle, FILE_MAP_ALL_ACCESS, 0, fileSize);
+		u8* mem = MapViewOfFile(m_logger, objectFileMappingHandle, FILE_MAP_ALL_ACCESS, 0, fileSize);
 		if (!mem)
 			return m_logger.Error(TC("Failed to map view of filehandle for read %s (%s)"), fileName, LastErrorToText().data);
-		auto memClose = MakeGuard([&](){ UnmapViewOfFile(mem, fileSize, fileName); });
+		auto memClose = MakeGuard([&](){ UnmapViewOfFile(m_logger, mem, fileSize, fileName); });
 
 		ObjectFile* objectFile = ObjectFile::Parse(m_logger, mem, fileSize, fileName);
 		if (!objectFile)
@@ -2941,16 +2944,16 @@ namespace uba
 		if (!objectFile->WriteImportsAndExports(m_logger, memoryBlock))
 			return false;
 
-		FileMappingHandle symHandle = CreateMemoryMappingW(m_logger, PAGE_READWRITE, memoryBlock.writtenSize);
+		FileMappingHandle symHandle = CreateMemoryMappingW(m_logger, PAGE_READWRITE, memoryBlock.writtenSize, nullptr, TC("SymHandle"));
 		if (!symHandle.IsValid())
 			return false;
-
-		u8* mem2 = MapViewOfFile(symHandle, FILE_MAP_ALL_ACCESS, 0, memoryBlock.writtenSize);
+		auto mg = MakeGuard([&]() { CloseFileMapping(m_logger, symHandle, TC("SymHandle")); });
+		u8* mem2 = MapViewOfFile(m_logger, symHandle, FILE_MAP_ALL_ACCESS, 0, memoryBlock.writtenSize);
 		if (!mem2)
 			return false;
 
 		MapMemoryCopy(mem2, memoryBlock.memory, memoryBlock.writtenSize);
-		UnmapViewOfFile(mem2, memoryBlock.writtenSize, TC(""));
+		UnmapViewOfFile(m_logger, mem2, memoryBlock.writtenSize, TC("SymHandle"));
 
 		StringKey symFileKey = CaseInsensitiveFs ? ToStringKeyLower(exportsFile) : ToStringKey(exportsFile);
 		u64 lastWriteTime = GetSystemTimeAsFileTime();
diff --git a/Engine/Source/Programs/UnrealBuildAccelerator/Common/Private/UbaSessionClient.cpp b/Engine/Source/Programs/UnrealBuildAccelerator/Common/Private/UbaSessionClient.cpp
--- a/Engine/Source/Programs/UnrealBuildAccelerator/Common/Private/UbaSessionClient.cpp
+++ b/Engine/Source/Programs/UnrealBuildAccelerator/Common/Private/UbaSessionClient.cpp
@@ -606,7 +606,7 @@ namespace uba
 		}
 		UBA_ASSERTF(casKey != CasKeyZero, TC("Failed to store cas file for %s (destination %s)"), source.name.c_str(), destination);
 	
-		CloseFileMapping(source.mappingHandle);
+		CloseFileMapping(m_logger, source.mappingHandle, source.name.c_str());
 		source.mappingHandle = {};
 
 		StackBinaryReader<128> reader;
diff --git a/Engine/Source/Programs/UnrealBuildAccelerator/Common/Private/UbaStorage.cpp b/Engine/Source/Programs/UnrealBuildAccelerator/Common/Private/UbaStorage.cpp
--- a/Engine/Source/Programs/UnrealBuildAccelerator/Common/Private/UbaStorage.cpp
+++ b/Engine/Source/Programs/UnrealBuildAccelerator/Common/Private/UbaStorage.cpp
@@ -312,12 +312,12 @@ namespace uba
 		{
 			if (!readMem)
 			{
-				FileMappingHandle fileMapping = uba::CreateFileMappingW(readHandle, PAGE_READONLY, fileSize, from);
+				FileMappingHandle fileMapping = uba::CreateFileMappingW(m_logger, readHandle, PAGE_READONLY, fileSize, from);
 				if (!fileMapping.IsValid())
 					return m_logger.Error(TC("Failed to create file mapping for %s (%s)"), from, LastErrorToText().data);
 
-				auto fmg = MakeGuard([&]() { CloseFileMapping(fileMapping); });
-				u8* uncompressedData = MapViewOfFile(fileMapping, FILE_MAP_READ, 0, fileSize);
+				auto fmg = MakeGuard([&]() { CloseFileMapping(m_logger, fileMapping, from); });
+				u8* uncompressedData = MapViewOfFile(m_logger, fileMapping, FILE_MAP_READ, 0, fileSize);
 				if (!uncompressedData)
 					return m_logger.Error(TC("Failed to map view of file mapping for %s (%s)"), from, LastErrorToText().data);
 
@@ -326,7 +326,7 @@ namespace uba
 						#if UBA_EXPERIMENTAL
 						m_workManager->AddWork([=, f = TString(from)]() { UnmapViewOfFile(uncompressedData, fileSize, f.c_str()); }, 1, TC("UnmapFile"));
 						#else
-						UnmapViewOfFile(uncompressedData, fileSize, from);
+						UnmapViewOfFile(m_logger, uncompressedData, fileSize, from);
 						#endif
 					});
 
@@ -2727,14 +2727,14 @@ namespace uba
 
 		if (fileSize > BufferSlotSize) // Note that when filesize is larger than BufferSlotSize the hash becomes a hash of hashes
 		{
-			FileMappingHandle fileMapping = uba::CreateFileMappingW(fileHandle, PAGE_READONLY, fileSize, fileName);
+			FileMappingHandle fileMapping = uba::CreateFileMappingW(m_logger, fileHandle, PAGE_READONLY, fileSize, fileName);
 			if (!fileMapping.IsValid())
 			{
 				m_logger.Error(TC("Failed to create file mapping for %s (%s)"), fileName, LastErrorToText().data);
 				return CasKeyZero;
 			}
-			auto fmg = MakeGuard([&]() { CloseFileMapping(fileMapping); });
-			u8* fileData = MapViewOfFile(fileMapping, FILE_MAP_READ, 0, fileSize);
+			auto fmg = MakeGuard([&]() { CloseFileMapping(m_logger, fileMapping, fileName); });
+			u8* fileData = MapViewOfFile(m_logger, fileMapping, FILE_MAP_READ, 0, fileSize);
 			if (!fileData)
 			{
 				m_logger.Error(TC("Failed to map view of file mapping for %s (%s)"), fileName, LastErrorToText().data);
@@ -2746,7 +2746,7 @@ namespace uba
 						m_workManager->AddWork([=, fn = TString(fileName)]() { UnmapViewOfFile(fileData, fileSize, fn.c_str()); }, 1, TC("UnmapFile"));
 					else
 					#endif
-						UnmapViewOfFile(fileData, fileSize, fileName);
+						UnmapViewOfFile(m_logger, fileData, fileSize, fileName);
 				});
 
 			return uba::CalculateCasKey(fileData, fileSize, storeCompressed, m_workManager, fileName);
@@ -2775,14 +2775,14 @@ namespace uba
 			u64 compressedSize;
 			if (!uba::GetFileSizeEx(compressedSize, fileHandle))
 				return m_logger.Error(TC("GetFileSize failed for %s (%s)"), fileName, LastErrorToText().data);
-			FileMappingHandle fileMapping = uba::CreateFileMappingW(fileHandle, PAGE_READONLY, compressedSize, fileName);
+			FileMappingHandle fileMapping = uba::CreateFileMappingW(m_logger, fileHandle, PAGE_READONLY, compressedSize, fileName);
 			if (!fileMapping.IsValid())
 				return m_logger.Error(TC("Failed to create file mapping for %s (%s)"), fileName, LastErrorToText().data);
-			auto fmg = MakeGuard([&]() { CloseFileMapping(fileMapping); });
-			u8* fileData = MapViewOfFile(fileMapping, FILE_MAP_READ, 0, compressedSize);
+			auto fmg = MakeGuard([&]() { CloseFileMapping(m_logger, fileMapping, fileName); });
+			u8* fileData = MapViewOfFile(m_logger, fileMapping, FILE_MAP_READ, 0, compressedSize);
 			if (!fileData)
 				return m_logger.Error(TC("Failed to map view of file mapping for %s (%s)"), fileName, LastErrorToText().data);
-			auto udg = MakeGuard([&]() { UnmapViewOfFile(fileData, compressedSize, fileName); });
+			auto udg = MakeGuard([&]() { UnmapViewOfFile(m_logger, fileData, compressedSize, fileName); });
 			
 			if (!DecompressMemoryToMemory(fileData + 8, dest, decompressedSize, fileName, writeHint))
 				return false;
diff --git a/Engine/Source/Programs/UnrealBuildAccelerator/Common/Private/UbaStorageClient.cpp b/Engine/Source/Programs/UnrealBuildAccelerator/Common/Private/UbaStorageClient.cpp
--- a/Engine/Source/Programs/UnrealBuildAccelerator/Common/Private/UbaStorageClient.cpp
+++ b/Engine/Source/Programs/UnrealBuildAccelerator/Common/Private/UbaStorageClient.cpp
@@ -99,7 +99,7 @@ namespace uba
 	{
 		delete m_proxyClient;
 		for (auto& pair : m_localStorageFiles)
-			CloseFileMapping(pair.second.casEntry.mappingHandle);
+			CloseFileMapping(m_logger, pair.second.casEntry.mappingHandle, pair.second.fileName.c_str());
 	}
 
 	bool StorageClient::IsUsingProxy()
@@ -812,13 +812,13 @@ namespace uba
 		}
 		else
 		{
-			fileMem = MapViewOfFile(mappingHandle, FILE_MAP_READ, mappingOffset, fileSize);
+			fileMem = MapViewOfFile(m_logger, mappingHandle, FILE_MAP_READ, mappingOffset, fileSize);
 			if (!fileMem)
 				return m_logger.Error(TC("%s - MapViewOfFile failed (%s)"), fileName, LastErrorToText().data);
 			isPersistentMapping = true;
 		}
 
-		auto unmapGuard = MakeGuard([&](){ if (isPersistentMapping) UnmapViewOfFile(fileMem, fileSize, fileName); });
+		auto unmapGuard = MakeGuard([&](){ if (isPersistentMapping) UnmapViewOfFile(m_logger, fileMem, fileSize, fileName); });
 
 		CasKey casKey;
 
@@ -840,7 +840,7 @@ namespace uba
 					if (isPersistentMapping)
 					{
 						FileMappingHandle mappingHandle2;
-						if (DuplicateFileMapping(GetCurrentProcessHandle(), mappingHandle, GetCurrentProcessHandle(), &mappingHandle2, FILE_MAP_READ, false, 0))
+						if (DuplicateFileMapping(m_logger, GetCurrentProcessHandle(), mappingHandle, GetCurrentProcessHandle(), &mappingHandle2, FILE_MAP_READ, false, 0, fileName))
 						{
 							localFile.casEntry.mappingHandle = mappingHandle2;
 							localFile.casEntry.size = fileSize;
diff --git a/Engine/Source/Programs/UnrealBuildAccelerator/Common/Private/UbaStorageServer.cpp b/Engine/Source/Programs/UnrealBuildAccelerator/Common/Private/UbaStorageServer.cpp
--- a/Engine/Source/Programs/UnrealBuildAccelerator/Common/Private/UbaStorageServer.cpp
+++ b/Engine/Source/Programs/UnrealBuildAccelerator/Common/Private/UbaStorageServer.cpp
@@ -127,9 +127,9 @@ namespace uba
 				++server.m_activeUnmap;
 				server.GetServer().AddWork([&server, mb = memoryBegin, mp = mappedView, rfh = readFileHandle, reason]()
 					{
-						if (!UnmapViewOfFile(mb, mp.size, reason))
+						if (!UnmapViewOfFile(server.m_logger, mb, mp.size, reason))
 							server.m_logger.Error(TC("%s - Failed to unmap memory at 0x%llx with size %llu (%s)"), reason, mb, mp.size, LastErrorToText().data);
-						if (!CloseFileMapping(mp.handle))
+						if (!CloseFileMapping(server.m_logger, mp.handle, reason))
 							server.m_logger.Error(TC("%s - Failed to close file mapping %llu (%s)"), reason, mp.handle.ToU64(), LastErrorToText().data);
 						if (!CloseFile(nullptr, rfh))
 							server.m_logger.Error(TC("%s - Failed to close file (%s)"), reason, LastErrorToText().data);
@@ -226,9 +226,9 @@ namespace uba
 			u64 alignedOffsetStart = AlignUp(mappingOffset - (pageSize - 1), pageSize);
 			u64 alignedOffsetEnd = AlignUp(endOffset, pageSize);
 			mapSize = alignedOffsetEnd - alignedOffsetStart;
-			mappedMem = MapViewOfFile(mappingHandle, FILE_MAP_READ, alignedOffsetStart, mapSize);
+			mappedMem = MapViewOfFile(logger, mappingHandle, FILE_MAP_READ, alignedOffsetStart, mapSize);
 			if (!mappedMem)
-				mappedMem = MapViewOfFile(mappingHandle, FILE_MAP_READ, alignedOffsetStart, 0);
+				mappedMem = MapViewOfFile(logger, mappingHandle, FILE_MAP_READ, alignedOffsetStart, 0);
 			if (mappedMem)
 				return mappedMem + (mappingOffset - alignedOffsetStart);
 			logger.Error(TC("[%s] Failed to map view of file %s at offset %llu and size %llu (%s)"), func, fileName, alignedOffsetStart, mapSize, LastErrorToText().data);
@@ -239,7 +239,7 @@ namespace uba
 		~FileMappingScope()
 		{
 			if (mappedMem)
-				if (!UnmapViewOfFile(mappedMem, mapSize, fileName))
+				if (!UnmapViewOfFile(logger, mappedMem, mapSize, fileName))
 					logger.Error(TC("Failed to unmap memory %s at 0x%llx with size %llu (%s)"), fileName, mappedMem, mapSize, LastErrorToText().data);
 		}
 
@@ -562,7 +562,7 @@ namespace uba
 				auto rfg = MakeGuard([&]()
 					{
 						if (!CloseFile(nullptr, readFileHandle))
-							m_logger.Error(TC("Failed to close file %s (%s)"), casFile.data, LastErrorToText().data);;
+							m_logger.Error(TC("Failed to close file %s (%s)"), casFile.data, LastErrorToText().data);
 					});
 				u64 fileSize = 0;
 				u8* memoryBegin = nullptr;
@@ -575,10 +575,10 @@ namespace uba
 						if (ownsMapping)
 						{
 							if (mappedView.memory)
-								if (!UnmapViewOfFile(mappedView.memory, fileSize, TC("FetchBegin")))
+								if (!UnmapViewOfFile(m_logger, mappedView.memory, fileSize, TC("FetchBegin")))
 									m_logger.Error(TC("Failed to unmap memory %s at 0x%llx with size %llu (%s)"), hint.data, mappedView.memory, fileSize, LastErrorToText().data);
 							if (mappedView.handle.IsValid())
-								if (!CloseFileMapping(mappedView.handle))
+								if (!CloseFileMapping(m_logger, mappedView.handle, hint.data))
 									m_logger.Error(TC("Failed to close file mapping for %s with handle %llu (%s)"), hint.data, mappedView.handle.ToU64(), LastErrorToText().data);
 						}
 						else
@@ -615,11 +615,11 @@ namespace uba
 					if (fileSize > BufferSlotSize)
 					{
 						ownsMapping = true;
-						mappedView.handle = CreateFileMappingW(readFileHandle, PAGE_READONLY, fileSize, TC(""));
+						mappedView.handle = CreateFileMappingW(m_logger, readFileHandle, PAGE_READONLY, fileSize, casFile.data);
 						if (!mappedView.handle.IsValid())
 							return m_logger.Error(TC("Failed to create file mapping of %s (%s)"), casFile.data, LastErrorToText().data);
 						u64 offset = memoryPos - memoryBegin;
-						mappedView.memory = MapViewOfFile(mappedView.handle, FILE_MAP_READ, 0, fileSize);
+						mappedView.memory = MapViewOfFile(m_logger, mappedView.handle, FILE_MAP_READ, 0, fileSize);
 						if (!mappedView.memory)
 							return m_logger.Error(TC("Failed to map memory of %s with size %llu (%s)"), casFile.data, fileSize, LastErrorToText().data);
 						memoryBegin = mappedView.memory;
diff --git a/Engine/Source/Programs/UnrealBuildAccelerator/Common/Private/UbaTrace.cpp b/Engine/Source/Programs/UnrealBuildAccelerator/Common/Private/UbaTrace.cpp
--- a/Engine/Source/Programs/UnrealBuildAccelerator/Common/Private/UbaTrace.cpp
+++ b/Engine/Source/Programs/UnrealBuildAccelerator/Common/Private/UbaTrace.cpp
@@ -60,15 +60,15 @@ namespace uba
 	bool Trace::StartWrite(const tchar* namedTrace, u64 traceMemCapacity)
 	{
 		m_memoryCapacity = traceMemCapacity;
-		m_memoryHandle = uba::CreateMemoryMappingW(m_logger, PAGE_READWRITE|SEC_RESERVE, m_memoryCapacity, namedTrace);
+		m_memoryHandle = uba::CreateMemoryMappingW(m_logger, PAGE_READWRITE|SEC_RESERVE, m_memoryCapacity, namedTrace, TC("Trace"));
 		if (!m_memoryHandle.IsValid())
 			return false;
 		if (GetLastError() != ERROR_ALREADY_EXISTS)
-			m_memoryBegin = MapViewOfFile(m_memoryHandle, FILE_MAP_WRITE, 0, m_memoryCapacity);
+			m_memoryBegin = MapViewOfFile(m_logger, m_memoryHandle, FILE_MAP_WRITE, 0, m_memoryCapacity);
 
 		if (!m_memoryBegin)
 		{
-			CloseFileMapping(m_memoryHandle);
+			CloseFileMapping(m_logger, m_memoryHandle, TC("Trace"));
 			m_memoryHandle = {};
 			return false;
 		}
@@ -145,12 +145,12 @@ namespace uba
 	{
 		if (m_memoryBegin)
 		{
-			UnmapViewOfFile(m_memoryBegin, m_memoryCapacity, TC("Trace"));
+			UnmapViewOfFile(m_logger, m_memoryBegin, m_memoryCapacity, TC("Trace"));
 			m_memoryBegin = nullptr;
 		}
 		if (m_memoryHandle.IsValid())
 		{
-			CloseFileMapping(m_memoryHandle);
+			CloseFileMapping(m_logger, m_memoryHandle, TC("Trace"));
 			m_memoryHandle = {};
 		}
 	}
@@ -446,7 +446,7 @@ namespace uba
 		#if PLATFORM_WINDOWS
 		StringBuffer<245> channelMutex;
 		channelMutex.Append(TC("Uba")).Append(channelName).Append(TC("Channel"));
-		m_memHandle = uba::CreateMemoryMappingW(m_logger, PAGE_READWRITE, 256, channelName);
+		m_memHandle = uba::CreateMemoryMappingW(m_logger, PAGE_READWRITE, 256, channelName, TC("TraceChannel"));
 		if (!m_memHandle.IsValid())
 		{
 			MessageBox(NULL, TC("Failed to create file mapping for trace channel"), TC("UbaVisualizer"), 0);
@@ -454,9 +454,9 @@ namespace uba
 		}
 		bool isCreator = GetLastError() != ERROR_ALREADY_EXISTS;
 
-		auto mhg = MakeGuard([&]() { CloseFileMapping(m_memHandle); m_memHandle = {}; });
+		auto mhg = MakeGuard([&]() { CloseFileMapping(m_logger, m_memHandle, TC("TraceChannel")); m_memHandle = {}; });
 
-		m_mem = MapViewOfFile(m_memHandle, FILE_MAP_WRITE, 0, 256);
+		m_mem = MapViewOfFile(m_logger, m_memHandle, FILE_MAP_WRITE, 0, 256);
 		if (!m_mem)
 		{
 			MessageBox(NULL, TC("Failed to map file mapping for uba trace channel"), TC("UbaVisualizer"), 0);
@@ -466,7 +466,7 @@ namespace uba
 		if (isCreator)
 			*(tchar*)m_mem = 0;
 
-		auto mg = MakeGuard([&]() { UnmapViewOfFile(m_mem, 256, channelMutex.data); m_mem = nullptr; });
+		auto mg = MakeGuard([&]() { UnmapViewOfFile(m_logger, m_mem, 256, channelMutex.data); m_mem = nullptr; });
 
 		channelMutex.Append(channelName).Append(TC("Mutex"));
 		m_mutex = CreateMutexW(false, channelMutex.data);
@@ -483,9 +483,9 @@ namespace uba
 	{
 		#if PLATFORM_WINDOWS
 		if (m_mem)
-			::UnmapViewOfFile(m_mem);
+			UnmapViewOfFile(m_logger, m_mem, 256, TC("TraceChannel"));
 		if (m_memHandle.IsValid())
-			CloseFileMapping(m_memHandle);
+			CloseFileMapping(m_logger, m_memHandle, TC("TraceChannel"));
 		if (m_mutex)
 			CloseHandle((HANDLE)m_mutex);
 		#endif
diff --git a/Engine/Source/Programs/UnrealBuildAccelerator/Common/Public/UbaFileMapping.h b/Engine/Source/Programs/UnrealBuildAccelerator/Common/Public/UbaFileMapping.h
--- a/Engine/Source/Programs/UnrealBuildAccelerator/Common/Public/UbaFileMapping.h
+++ b/Engine/Source/Programs/UnrealBuildAccelerator/Common/Public/UbaFileMapping.h
@@ -47,13 +47,13 @@ namespace uba
 	};
 #endif
 
-	FileMappingHandle CreateMemoryMappingW(Logger& logger, u32 flProtect, u64 maxSize, const tchar* name = nullptr);
-	FileMappingHandle CreateFileMappingW(FileHandle file, u32 flProtect, u64 maxSize = 0, const tchar* hint = TC(""));
-	u8* MapViewOfFile(FileMappingHandle fileMappingObject, u32 dwDesiredAccess, u64 offset, u64 dwNumberOfBytesToMap);
+	FileMappingHandle CreateMemoryMappingW(Logger& logger, u32 flProtect, u64 maxSize, const tchar* name, const tchar* hint);
+	FileMappingHandle CreateFileMappingW(Logger& logger, FileHandle file, u32 flProtect, u64 maxSize, const tchar* hint);
+	u8* MapViewOfFile(Logger& logger, FileMappingHandle fileMappingObject, u32 dwDesiredAccess, u64 offset, u64 dwNumberOfBytesToMap);
 	bool MapViewCommit(void* address, u64 size);
-	bool UnmapViewOfFile(const void* lpBaseAddress, u64 bytesToUnmap, const tchar* hint);
-	bool CloseFileMapping(FileMappingHandle h);
-	bool DuplicateFileMapping(ProcHandle hSourceProcessHandle, FileMappingHandle hSourceHandle, ProcHandle hTargetProcessHandle, FileMappingHandle* lpTargetHandle, u32 dwDesiredAccess, bool bInheritHandle, u32 dwOptions);
+	bool UnmapViewOfFile(Logger& logger, const void* lpBaseAddress, u64 bytesToUnmap, const tchar* hint);
+	bool CloseFileMapping(Logger& logger, FileMappingHandle h, const tchar* hint);
+	bool DuplicateFileMapping(Logger& logger, ProcHandle hSourceProcessHandle, FileMappingHandle hSourceHandle, ProcHandle hTargetProcessHandle, FileMappingHandle* lpTargetHandle, u32 dwDesiredAccess, bool bInheritHandle, u32 dwOptions, const tchar* hint);
 	void MapMemoryCopy(void* dest, const void* source, u64 size);
 
 	struct MappedView
diff --git a/Engine/Source/Programs/UnrealBuildAccelerator/Common/Public/UbaSession.h b/Engine/Source/Programs/UnrealBuildAccelerator/Common/Public/UbaSession.h
--- a/Engine/Source/Programs/UnrealBuildAccelerator/Common/Public/UbaSession.h
+++ b/Engine/Source/Programs/UnrealBuildAccelerator/Common/Public/UbaSession.h
@@ -212,7 +212,7 @@ namespace uba
 		u8* m_fileMappingTableMem;
 		u32 m_fileMappingTableSize = 0;
 		ReaderWriterLock m_fileMappingTableLookupLock;
-		struct FileMappingEntry { ReaderWriterLock lock; FileMappingHandle mapping; u64 mappingOffset = 0; u64 size = 0; bool isDir = false; bool handled = false; bool success = false;};
+		struct FileMappingEntry { ReaderWriterLock lock; FileMappingHandle mapping; u64 mappingOffset = 0; u64 size = 0; bool isDir = false; bool handled = false; bool success = false; bool owned = false; };
 		UnorderedMap<StringKey, FileMappingEntry> m_fileMappingTableLookup;
 
 		ReaderWriterLock m_nameToHashTableHandleLock;
diff --git a/Engine/Source/Programs/UnrealBuildAccelerator/Common/Public/UbaSessionCreateInfo.h b/Engine/Source/Programs/UnrealBuildAccelerator/Common/Public/UbaSessionCreateInfo.h
--- a/Engine/Source/Programs/UnrealBuildAccelerator/Common/Public/UbaSessionCreateInfo.h
+++ b/Engine/Source/Programs/UnrealBuildAccelerator/Common/Public/UbaSessionCreateInfo.h
@@ -35,6 +35,7 @@ namespace uba
 		bool detailedTrace = false;				// Enable detailed trace to include jobs, individual file I/O etc in trace dump
 		bool traceChildProcesses = false;		// Trace and visualize child processes so they can be seen in visualizer
 		bool storeObjFilesCompressed = false;	// Compiler will write .obj files compressed to disk and linker will decompress when reading
+		bool allowLocalDetour = true;			// Allow local processes to be detoured. If this is false it is up to outside logic to register all created/deleted files
 		bool extractObjFilesSymbols = false;	// Will extract import/export symbols and clear export table from obj file.
 		u64 deleteSessionsOlderThanSeconds = 12 * 60 * 60; // Delete session folders older than 12 hours by default . Set to 0 to not delete or 1 to delete all
 		u64 keepOutputFileMemoryMapsThreshold = 256 * 1024; // If allowMemoryMaps is true, output files will be kept in memory if smaller than this size
