[UBA]
* Added WriteRecievedCasFilesToDisk as option to config files
* Added AllowHintAsFallback as option to config files

diff --git a/Engine/Source/Programs/UnrealBuildAccelerator/Common/Private/UbaStorageServer.cpp b/Engine/Source/Programs/UnrealBuildAccelerator/Common/Private/UbaStorageServer.cpp
--- a/Engine/Source/Programs/UnrealBuildAccelerator/Common/Private/UbaStorageServer.cpp
+++ b/Engine/Source/Programs/UnrealBuildAccelerator/Common/Private/UbaStorageServer.cpp
@@ -1,6 +1,7 @@
 // Copyright Epic Games, Inc. All Rights Reserved.
 
 #include "UbaStorageServer.h"
+#include "UbaConfig.h"
 #include "UbaFileAccessor.h"
 #include "UbaNetworkServer.h"
 #include "UbaTrace.h"
@@ -10,6 +11,12 @@ namespace uba
 	void StorageServerCreateInfo::Apply(Config& config)
 	{
 		StorageCreateInfo::Apply(config);
+		const ConfigTable* tablePtr = config.GetTable(TC("Storage"));
+		if (!tablePtr)
+			return;
+		const ConfigTable& table = *tablePtr;
+		table.GetValueAsBool(writeRecievedCasFilesToDisk, TC("WriteRecievedCasFilesToDisk"));
+		table.GetValueAsBool(allowHintAsFallback, TC("AllowHintAsFallback"));
 	}
 
 	StorageServer::StorageServer(const StorageServerCreateInfo& info)
@@ -17,7 +24,7 @@ namespace uba
 	,	m_server(info.server)
 	{
 		m_zone = info.zone;
-		m_allowFallback = info.allowFallback;
+		m_allowHintAsFallback = info.allowHintAsFallback;
 		m_writeRecievedCasFilesToDisk = info.writeRecievedCasFilesToDisk;
 
 		if (!CreateGuid(m_uid))
@@ -492,7 +499,7 @@ namespace uba
 				if (!has)
 				{
 					todoRemoveMe = 3;
-					if (!EnsureCasFile(casKey, nullptr) && m_allowFallback)
+					if (!EnsureCasFile(casKey, nullptr) && m_allowHintAsFallback)
 					{
 						// Last resort.. use hint to load file into cas (hint should be renamed since it is now a critical parameter)
 						// We better check the caskey first to make sure it is matching on the server
diff --git a/Engine/Source/Programs/UnrealBuildAccelerator/Common/Public/UbaStorageServer.h b/Engine/Source/Programs/UnrealBuildAccelerator/Common/Public/UbaStorageServer.h
--- a/Engine/Source/Programs/UnrealBuildAccelerator/Common/Public/UbaStorageServer.h
+++ b/Engine/Source/Programs/UnrealBuildAccelerator/Common/Public/UbaStorageServer.h
@@ -14,7 +14,7 @@ namespace uba
 		void Apply(Config& config);
 
 		NetworkServer& server;
-		bool allowFallback = true; // Will fallback to file system to recalculate cas if cas content does not exist anymore
+		bool allowHintAsFallback = true; // Will fallback to file system to recalculate cas if cas content does not exist anymore
 		bool writeRecievedCasFilesToDisk = false;
 		const tchar* zone = TC("");
 	};
@@ -150,7 +150,7 @@ namespace uba
 
 		Vector<TString> m_disallowedPaths;
 
-		bool m_allowFallback;
+		bool m_allowHintAsFallback;
 		bool m_writeRecievedCasFilesToDisk;
 	};
 }
