[UBA]
* Added radlink.exe to application rules (needed when using compressed obj files)

diff --git a/Engine/Source/Programs/UnrealBuildAccelerator/Core/Private/UbaApplicationRules.cpp b/Engine/Source/Programs/UnrealBuildAccelerator/Core/Private/UbaApplicationRules.cpp
--- a/Engine/Source/Programs/UnrealBuildAccelerator/Core/Private/UbaApplicationRules.cpp
+++ b/Engine/Source/Programs/UnrealBuildAccelerator/Core/Private/UbaApplicationRules.cpp
@@ -587,7 +587,7 @@ namespace uba
 			{ ApplicationHash(TC("clang-cl.exe")),				new ApplicationRulesClangPlusPlusExe() },
 			{ ApplicationHash(TC("verse-clang-cl.exe")),		new ApplicationRulesClangPlusPlusExe() },
 			{ ApplicationHash(TC("ispc.exe")),					new ApplicationRulesISPCExe() },
-			{ 30341511,											new ApplicationRulesVcLink() },
+			{ ApplicationHash(TC("radlink.exe")),				new ApplicationRulesVcLink() },
 			{ 3340509542,										new ApplicationRulesClangPlusPlusExePlatform1() },
 			{ 4113554641,										new ApplicationRulesLdExePlatform1() },
 			{ 1752955744,										new ApplicationRulesLdExePlatform1() },
