[UBA]
* Added graceful handling of not finding cross architecture binary

diff --git a/Engine/Source/Programs/UnrealBuildAccelerator/Common/Private/UbaSessionServer.cpp b/Engine/Source/Programs/UnrealBuildAccelerator/Common/Private/UbaSessionServer.cpp
--- a/Engine/Source/Programs/UnrealBuildAccelerator/Common/Private/UbaSessionServer.cpp
+++ b/Engine/Source/Programs/UnrealBuildAccelerator/Common/Private/UbaSessionServer.cpp
@@ -1167,13 +1167,18 @@ namespace uba
 			StringBuffer<> applicationOverride;
 			if (isCrossArchitecture)
 			{
-				if (!process->m_allowCrossArchitecture)
+				auto returnProcess = [&]()
 				{
-					// There is a risk cross architecture client dequeues another client's process that is not cross architecture
-					// and in that case we just have to return it and try again later.. (break out)
 					m_queuedRemoteProcesses.push_front({process});
 					if (m_remoteProcessReturnedEvent)
 						m_remoteProcessReturnedEvent(*process);
+				};
+			
+				if (!process->m_allowCrossArchitecture)
+				{
+					// There is a risk cross architecture client dequeues another client's process that is not cross architecture
+					// and in that case we just have to return it and try again later.. (break out)
+					returnProcess();
 					break;
 				}
 
@@ -1183,8 +1188,20 @@ namespace uba
 						applicationOverride.Append(mapping.to).Append(startInfo.application + mapping.from.size());
 						break;
 					}
+
 				if (!applicationOverride.count)
-					return m_logger.Error(TC("Cross architecture client tried to find mapping for %s"), startInfo.application);
+				{
+					m_logger.Info(TC("Couldn't find cross architecture mapping for %s"), startInfo.application);
+					returnProcess();
+					break;
+				}
+
+				if (!FileExists(m_logger, applicationOverride.data))
+				{
+					m_logger.Info(TC("Couldn't find cross architecture executable %s"), applicationOverride.data);
+					returnProcess();
+					break;
+				}
 			}
 
 			ProcessAdded(*process, sessionId);
