[UBA]
* Fixed uba cache service compile error

diff --git a/Engine/Source/Programs/UnrealBuildAccelerator/CacheService/Private/UbaCacheService.cpp b/Engine/Source/Programs/UnrealBuildAccelerator/CacheService/Private/UbaCacheService.cpp
--- a/Engine/Source/Programs/UnrealBuildAccelerator/CacheService/Private/UbaCacheService.cpp
+++ b/Engine/Source/Programs/UnrealBuildAccelerator/CacheService/Private/UbaCacheService.cpp
@@ -246,7 +246,7 @@ namespace uba
 		StorageServerCreateInfo storageInfo(networkServer, g_rootDir.data, logWriter);
 		storageInfo.casCapacityBytes = storageCapacity;
 		storageInfo.storeCompressed = storeCompressed;
-		storageInfo.allowFallback = false;
+		storageInfo.allowHintAsFallback = false;
 		storageInfo.manuallyHandleOverflow = true;
 		storageInfo.writeRecievedCasFilesToDisk = true;
 		StorageServer storageServer(storageInfo);
