[UBA]
* Fixed visualizer compile errors

diff --git a/Engine/Source/Programs/UnrealBuildAccelerator/Visualizer/Private/UbaTraceReader.cpp b/Engine/Source/Programs/UnrealBuildAccelerator/Visualizer/Private/UbaTraceReader.cpp
--- a/Engine/Source/Programs/UnrealBuildAccelerator/Visualizer/Private/UbaTraceReader.cpp
+++ b/Engine/Source/Programs/UnrealBuildAccelerator/Visualizer/Private/UbaTraceReader.cpp
@@ -79,15 +79,15 @@ namespace uba
 		if (!uba::GetFileSizeEx(fileSize, readHandle))
 			return false;
 
-		m_memoryHandle = uba::CreateFileMappingW(readHandle, PAGE_READONLY);
+		m_memoryHandle = uba::CreateFileMappingW(m_logger, readHandle, PAGE_READONLY, 0, fileName);
 		if (!m_memoryHandle.IsValid())
 			return false;
-		auto closeFileMapping = MakeGuard([&]() { CloseFileMapping(m_memoryHandle); m_memoryHandle = {}; });
+		auto closeFileMapping = MakeGuard([&]() { CloseFileMapping(m_logger, m_memoryHandle, fileName); m_memoryHandle = {}; });
 
-		m_memoryBegin = MapViewOfFile(m_memoryHandle, FILE_MAP_READ, 0, 0);
+		m_memoryBegin = MapViewOfFile(m_logger, m_memoryHandle, FILE_MAP_READ, 0, 0);
 		if (!m_memoryBegin)
 			return false;
-		auto closeView = MakeGuard([&]() { UnmapViewOfFile(m_memoryBegin, 0, TC("TraceReader")); m_memoryBegin = nullptr; m_memoryPos = nullptr; m_memoryEnd = nullptr; });
+		auto closeView = MakeGuard([&]() { UnmapViewOfFile(m_logger, m_memoryBegin, 0, TC("TraceReader")); m_memoryBegin = nullptr; m_memoryPos = nullptr; m_memoryEnd = nullptr; });
 
 		m_memoryPos = m_memoryBegin;
 		m_memoryEnd = m_memoryBegin + fileSize;
@@ -161,10 +161,10 @@ namespace uba
 		Reset(out);
 
 		u32 traceMemSize = 128 * 1024 * 1024;
-		m_memoryHandle = uba::CreateMemoryMappingW(m_logger, PAGE_READWRITE, traceMemSize);
+		m_memoryHandle = CreateMemoryMappingW(m_logger, PAGE_READWRITE, traceMemSize, 0, TC("NetworkClient"));
 		if (!m_memoryHandle.IsValid())
 			return false;
-		m_memoryBegin = MapViewOfFile(m_memoryHandle, FILE_MAP_ALL_ACCESS, 0, traceMemSize);
+		m_memoryBegin = MapViewOfFile(m_logger, m_memoryHandle, FILE_MAP_ALL_ACCESS, 0, traceMemSize);
 		m_memoryPos = m_memoryBegin;
 		m_memoryEnd = m_memoryBegin;
 
@@ -232,7 +232,7 @@ namespace uba
 					m_logger.Error(L"OpenFileMappingW - Failed to open file mapping %ls (%ls)", namedTrace, LastErrorToText().data);
 				return false;
 			}
-			m_memoryBegin = MapViewOfFile(m_memoryHandle, FILE_MAP_READ, 0, 0);
+			m_memoryBegin = MapViewOfFile(m_logger, m_memoryHandle, FILE_MAP_READ, 0, 0);
 			if (!m_memoryBegin)
 				return false;
 		}
@@ -264,9 +264,9 @@ namespace uba
 		// Move memory to local mapping in case we want to replay
 		u64 traceMemSize = m_memoryEnd - m_memoryBegin;
 		u8* memoryBegin = nullptr;
-		FileMappingHandle memoryHandle = uba::CreateMemoryMappingW(m_logger, PAGE_READWRITE, traceMemSize);
+		FileMappingHandle memoryHandle = CreateMemoryMappingW(m_logger, PAGE_READWRITE, traceMemSize, nullptr, TC("UpdateReadNamed"));
 		if (memoryHandle.IsValid())
-			if ((memoryBegin = MapViewOfFile(memoryHandle, FILE_MAP_ALL_ACCESS, 0, traceMemSize)) != nullptr)
+			if ((memoryBegin = MapViewOfFile(m_logger, memoryHandle, FILE_MAP_ALL_ACCESS, 0, traceMemSize)) != nullptr)
 				memcpy(memoryBegin, m_memoryBegin, traceMemSize);
 
 		u64 pos = m_memoryPos - m_memoryBegin;
@@ -1150,10 +1150,10 @@ namespace uba
 			CloseHandle(m_hostProcess);
 		m_hostProcess = 0;
 		if (m_memoryBegin)
-			UnmapViewOfFile(m_memoryBegin, 0, TC("TraceReader"));
+			UnmapViewOfFile(m_logger, m_memoryBegin, 0, TC("TraceReader"));
 		m_memoryBegin = nullptr;
 		if (m_memoryHandle.IsValid())
-			CloseFileMapping(m_memoryHandle);
+			CloseFileMapping(m_logger, m_memoryHandle, TC("TraceReader"));
 		m_memoryHandle = {};
 		m_namedTrace.clear();
 	}
