[UBA]
* Added arm to ubaagent help text
* Improved debug logging

diff --git a/Engine/Source/Programs/UnrealBuildAccelerator/Agent/Private/UbaAgent.cpp b/Engine/Source/Programs/UnrealBuildAccelerator/Agent/Private/UbaAgent.cpp
--- a/Engine/Source/Programs/UnrealBuildAccelerator/Agent/Private/UbaAgent.cpp
+++ b/Engine/Source/Programs/UnrealBuildAccelerator/Agent/Private/UbaAgent.cpp
@@ -66,7 +66,7 @@ namespace uba
 		}
 		logger.Info(TC(""));
 		logger.Info(TC("-------------------------------------------"));
-		logger.Info(TC("   UbaAgent v%s"), Version);
+		logger.Info(TC("   UbaAgent v%s%s"), Version, (IsArm ? TC(" (ARM64)") : TC("")));
 		logger.Info(TC("-------------------------------------------"));
 		logger.Info(TC(""));
 		logger.Info(TC("  When started UbaAgent will keep trying to connect to provided host address."));
@@ -715,7 +715,7 @@ namespace uba
 				poll = false;
 				quiet = true;
 			}
-			else if (name.Equals(TC("-?")))
+			else if (name.Equals(TC("-?")) || name.Equals(TC("-help")))
 			{
 				return PrintHelp(TC(""));
 			}
diff --git a/Engine/Source/Programs/UnrealBuildAccelerator/Core/Public/UbaWinBinDependencyParser.h b/Engine/Source/Programs/UnrealBuildAccelerator/Core/Public/UbaWinBinDependencyParser.h
--- a/Engine/Source/Programs/UnrealBuildAccelerator/Core/Public/UbaWinBinDependencyParser.h
+++ b/Engine/Source/Programs/UnrealBuildAccelerator/Core/Public/UbaWinBinDependencyParser.h
@@ -40,6 +40,7 @@ namespace uba
 		L"kernelbase.dll",
 		L"mscoree.dll",
 		L"msvcp_win.dll",
+		L"msvcr120.dll",
 		L"msvcrt.dll",
 		L"mswsock.dll",
 		L"ncrypt.dll",
diff --git a/Engine/Source/Programs/UnrealBuildAccelerator/Detours/Private/Windows/UbaDetoursFunctionsKernelBase.inl b/Engine/Source/Programs/UnrealBuildAccelerator/Detours/Private/Windows/UbaDetoursFunctionsKernelBase.inl
--- a/Engine/Source/Programs/UnrealBuildAccelerator/Detours/Private/Windows/UbaDetoursFunctionsKernelBase.inl
+++ b/Engine/Source/Programs/UnrealBuildAccelerator/Detours/Private/Windows/UbaDetoursFunctionsKernelBase.inl
@@ -1,7 +1,5 @@
 // Copyright Epic Games, Inc. All Rights Reserved.
 
-const wchar_t* ToString(BOOL b) { return b ? L"Success" : L"Error"; }
-
 DWORD Local_GetLongPathNameW(LPCWSTR lpszShortPath, LPWSTR lpszLongPath, DWORD cchBuffer)
 {
 	SCOPED_WRITE_LOCK(g_longPathNameCacheLock, lock);
@@ -3085,16 +3083,20 @@ BOOL Detoured_GetThreadPreferredUILanguages(DWORD dwFlags, PULONG pulNumLanguage
 
 	if (dwFlags & MUI_LANGUAGE_ID)
 	{
-		DEBUG_LOG_DETOURED(L"GetThreadPreferredUILanguages", L"");
+		UBA_ASSERT(pulNumLanguages);
+		UBA_ASSERT(pcchLanguagesBuffer);
 		//UBA_ASSERT(!(dwFlags & ~MUI_LANGUAGE_ID);
 		*pulNumLanguages = 1;
 		*pcchLanguagesBuffer = 6;
 
 		if (!pwszLanguagesBuffer)
+		{
+			DEBUG_LOG_DETOURED(L"GetThreadPreferredUILanguages", L"(nobuf) -> TRUE");
 			return TRUE;
-
+		}
 		swprintf_s(pwszLanguagesBuffer, 6, L"%04x", g_uiLanguage);
 		pwszLanguagesBuffer[5] = 0;
+		DEBUG_LOG_DETOURED(L"GetThreadPreferredUILanguages", L"(%s) -> TRUE", pwszLanguagesBuffer);
 		return TRUE;
 	}
 	else // MUI_LANGUAGE_NAME
@@ -3723,7 +3725,7 @@ LSTATUS Detoured_RegOpenKeyExW(HKEY hKey, LPCWSTR lpSubKey, DWORD ulOptions, REG
 	//UBA_NOT_IMPLEMENTED();
 	SuppressCreateFileDetourScope cfs;
 	auto res = True_RegOpenKeyExW(hKey, lpSubKey, ulOptions, samDesired, phkResult);
-	DEBUG_LOG_TRUE(L"RegOpenKeyExW", L"(%ls) -> %ls", lpSubKey, ToString(res));
+	DEBUG_LOG_TRUE(L"RegOpenKeyExW", L"(%ls) -> %ls", lpSubKey, ToString(res == ERROR_SUCCESS));
 	return res;
 }
 
@@ -3732,7 +3734,7 @@ LSTATUS Detoured_RegOpenKeyExA(HKEY hKey, LPCSTR lpSubKey, DWORD ulOptions, REGS
 	DETOURED_CALL(RegOpenKeyExA);
 	auto res = True_RegOpenKeyExA(hKey, lpSubKey, ulOptions, samDesired, phkResult);
 	//UBA_NOT_IMPLEMENTED();
-	DEBUG_LOG_TRUE(L"RegOpenKeyExA", L"%llu (%hs) -> %ls", uintptr_t(*phkResult), lpSubKey, ToString(res));
+	DEBUG_LOG_TRUE(L"RegOpenKeyExA", L"%llu (%hs) -> %ls", uintptr_t(*phkResult), lpSubKey, ToString(res == ERROR_SUCCESS));
 	return res;
 }
 
diff --git a/Engine/Source/Programs/UnrealBuildAccelerator/Detours/Private/Windows/UbaDetoursFunctionsNtDll.inl b/Engine/Source/Programs/UnrealBuildAccelerator/Detours/Private/Windows/UbaDetoursFunctionsNtDll.inl
--- a/Engine/Source/Programs/UnrealBuildAccelerator/Detours/Private/Windows/UbaDetoursFunctionsNtDll.inl
+++ b/Engine/Source/Programs/UnrealBuildAccelerator/Detours/Private/Windows/UbaDetoursFunctionsNtDll.inl
@@ -34,7 +34,17 @@ u8 GetFileAccessFlags(DWORD desiredAccess, u32 createDisposition)
 	return access;
 }
 
-const wchar_t* ToString(NTSTATUS s) { return NT_SUCCESS(s) ? L"Success" : L"Error"; }
+StringBuffer<32> ToString(NTSTATUS s)
+{
+	StringBuffer<32> res;
+	if (NT_SUCCESS(s))
+		return res.Append(L"Success");
+	if (s == STATUS_OBJECT_NAME_NOT_FOUND)
+		return res.Append(L"STATUS_OBJECT_NAME_NOT_FOUND");
+	if (s == STATUS_OBJECT_PATH_NOT_FOUND)
+		return res.Append(L"STATUS_OBJECT_PATH_NOT_FOUND");
+	return res.Appendf(L"Error (0x%x)", u32(s));
+}
 
 struct FILE_FS_DEVICE_INFORMATION
 {
@@ -101,7 +111,7 @@ NTSTATUS Detoured_NtQueryVolumeInformationFile(HANDLE FileHandle, PIO_STATUS_BLO
 		UBA_ASSERTF(false, L"NtQueryVolumeInformationFile called in ListDirectoryHandle using class %u which is not implemented (%ls)", FsInformationClass, HandleToName(FileHandle));
 	}
 	auto res = True_NtQueryVolumeInformationFile(TrueHandle, IoStatusBlock, FsInformation, Length, FsInformationClass);
-	DEBUG_LOG_TRUE(L"NtQueryVolumeInformationFile", L"%llu (%ls) -> %ls", uintptr_t(FileHandle), HandleToName(FileHandle), ToString(res));
+	DEBUG_LOG_TRUE(L"NtQueryVolumeInformationFile", L"%llu (%ls) -> %ls", uintptr_t(FileHandle), HandleToName(FileHandle), ToString(res).data);
 	return res;
 }
 
@@ -215,7 +225,7 @@ NTSTATUS Detoured_NtQueryInformationFile(HANDLE FileHandle, PIO_STATUS_BLOCK IoS
 
 	TimerScope ts(g_kernelStats.getFileInfo);
 	auto res = True_NtQueryInformationFile(TrueHandle, IoStatusBlock, FileInformation, Length, FileInformationClass);
-	DEBUG_LOG_TRUE(L"NtQueryInformationFile", L"(%u) %llu (%ls) -> %ls", FileInformationClass, uintptr_t(FileHandle), HandleToName(FileHandle), ToString(res));
+	DEBUG_LOG_TRUE(L"NtQueryInformationFile", L"(%u) %llu (%ls) -> %ls", FileInformationClass, uintptr_t(FileHandle), HandleToName(FileHandle), ToString(res).data);
 	return res;
 }
 
@@ -472,7 +482,7 @@ NTSTATUS Detoured_NtSetInformationFile(HANDLE FileHandle, PIO_STATUS_BLOCK IoSta
 
 	TimerScope ts(g_kernelStats.setFileInfo);
 	auto res = True_NtSetInformationFile(trueHandle, IoStatusBlock, FileInformation, Length, FileInformationClass);
-	DEBUG_LOG_TRUE(L"NtSetInformationFile", L"(%u) %llu (%ls) -> %ls", FileInformationClass, uintptr_t(FileHandle), HandleToName(FileHandle), ToString(res));
+	DEBUG_LOG_TRUE(L"NtSetInformationFile", L"(%u) %llu (%ls) -> %ls", FileInformationClass, uintptr_t(FileHandle), HandleToName(FileHandle), ToString(res).data);
 	return res;
 }
 
@@ -485,7 +495,7 @@ NTSTATUS NTAPI Detoured_NtSetInformationObject(HANDLE ObjectHandle, OBJECT_INFOR
 		UBA_ASSERT(ObjectHandle != INVALID_HANDLE_VALUE);
 	}
 	auto res = True_NtSetInformationObject(ObjectHandle, ObjectInformationClass, ObjectInformation, Length);
-	DEBUG_LOG_TRUE(L"NtSetInformationObject", L"(%u) %llu (%ls) -> %ls", ObjectInformationClass, uintptr_t(ObjectHandle), HandleToName(ObjectHandle), ToString(res));
+	DEBUG_LOG_TRUE(L"NtSetInformationObject", L"(%u) %llu (%ls) -> %ls", ObjectInformationClass, uintptr_t(ObjectHandle), HandleToName(ObjectHandle), ToString(res).data);
 	return res;
 }
 
@@ -674,7 +684,7 @@ NTSTATUS NTAPI Shared_NtCreateFile(bool IsCreateFunc, PHANDLE hFileHandle, ACCES
 	if (suppressCreateFileDetour)
 	{
 		NTSTATUS res = Local_NtCreateFile(IsCreateFunc, hFileHandle, DesiredAccess, ObjectAttributes, IoStatusBlock, AllocationSize, FileAttributes, ShareAccess, CreateDisposition, CreateOptions, EaBuffer, EaLength);
-		DEBUG_LOG_TRUE(funcName, L"(SUPPRESSDETOUR) %llu (%ls) -> %ls", uintptr_t(*hFileHandle), ObjectAttributes->ObjectName->Buffer, ToString(res));
+		DEBUG_LOG_TRUE(funcName, L"(SUPPRESSDETOUR) %llu (%ls) -> %ls", uintptr_t(*hFileHandle), ObjectAttributes->ObjectName->Buffer, ToString(res).data);
 		if (createFileName && Equals(createFileName, L"NUL"))
 			g_nullFile = *hFileHandle;
 		return res;
@@ -714,7 +724,7 @@ NTSTATUS NTAPI Shared_NtCreateFile(bool IsCreateFunc, PHANDLE hFileHandle, ACCES
 		NTSTATUS res = Local_NtCreateFile(IsCreateFunc, hFileHandle, DesiredAccess, ObjectAttributes, IoStatusBlock, AllocationSize, FileAttributes, ShareAccess, CreateDisposition, CreateOptions, EaBuffer, EaLength);
 		if (NT_ERROR(res))
 			*hFileHandle = INVALID_HANDLE_VALUE;
-		DEBUG_LOG_TRUE(funcName, L"(NODETOUR)%ls %llu (%.*ls) -> %ls (%u)", isWriteStr, uintptr_t(*hFileHandle), ObjectAttributes->ObjectName->Length / 2, ObjectAttributes->ObjectName->Buffer, ToString(res), res);
+		DEBUG_LOG_TRUE(funcName, L"(NODETOUR)%ls %llu (%.*ls) -> %ls", isWriteStr, uintptr_t(*hFileHandle), ObjectAttributes->ObjectName->Length / 2, ObjectAttributes->ObjectName->Buffer, ToString(res).data);
 		if (NT_ERROR(res))
 			return res;
 
@@ -816,7 +826,7 @@ NTSTATUS NTAPI Shared_NtCreateFile(bool IsCreateFunc, PHANDLE hFileHandle, ACCES
 			UBA_ASSERT(!g_runningRemote);
 			NTSTATUS res = True_NtCreateFile(hFileHandle, DesiredAccess, ObjectAttributes, IoStatusBlock, AllocationSize, FileAttributes, ShareAccess, CreateDisposition, CreateOptions, EaBuffer, EaLength);
 
-			DEBUG_LOG_DETOURED(funcName, L"(CREATE_DIR) %llu, (%ls) -> %ls", uintptr_t(*hFileHandle), fileName.data, ToString(res));
+			DEBUG_LOG_DETOURED(funcName, L"(CREATE_DIR) %llu, (%ls) -> %ls", uintptr_t(*hFileHandle), fileName.data, ToString(res).data);
 			return res;
 		}
 
@@ -1195,12 +1205,12 @@ NTSTATUS NTAPI Shared_NtCreateFile(bool IsCreateFunc, PHANDLE hFileHandle, ACCES
 			info.lastDesiredAccess = lastDesiredAccess;
 			Rpc_UpdateCloseHandle(L"", closeId, false, L"", 0, 0, false);
 		}
-		DEBUG_LOG_TRUE(funcName, L"%ls (%ls) (%ls) -> %ls (0x%x)", isWriteStr, lpFileName, (fileName.data != lpFileName ? fileName.data : L""), ToString(res), res);
+		DEBUG_LOG_TRUE(funcName, L"%ls (%ls) (%ls) -> %ls", isWriteStr, lpFileName, (fileName.data != lpFileName ? fileName.data : L""), ToString(res).data);
 		return res;
 	}
 	if (!canDetour)
 	{
-		DEBUG_LOG_TRUE(funcName, L"%ls %llu NODETOUR (%ls) -> %ls", isWriteStr, uintptr_t(*hFileHandle), tempFileName, ToString(res));
+		DEBUG_LOG_TRUE(funcName, L"%ls %llu NODETOUR (%ls) -> %ls", isWriteStr, uintptr_t(*hFileHandle), tempFileName, ToString(res).data);
 		return res;
 	}
 
@@ -1216,7 +1226,7 @@ NTSTATUS NTAPI Shared_NtCreateFile(bool IsCreateFunc, PHANDLE hFileHandle, ACCES
 	dh->fileObject->fileInfo = &info;
 	dh->fileObject->deleteOnClose = isDeleteOnClose;
 	*hFileHandle = makeDetouredHandle(dh);
-	DEBUG_LOG_TRUE(funcName, L"%ls %llu (%ls)%s -> %ls", isWriteStr, uintptr_t(*hFileHandle), tempFileName, isDeleteOnClose ? TC(" DeleteOnClose") : TC(""), ToString(res));
+	DEBUG_LOG_TRUE(funcName, L"%ls %llu (%ls)%s -> %ls", isWriteStr, uintptr_t(*hFileHandle), tempFileName, isDeleteOnClose ? TC(" DeleteOnClose") : TC(""), ToString(res).data);
 	return res;
 }
 
@@ -1303,7 +1313,7 @@ NTSTATUS NTAPI Detoured_NtClose(HANDLE handle)
 		TimerScope ts(g_kernelStats.closeHandle);
 		auto res = True_NtClose(handle);
 #if !defined(_M_ARM64) // For some reason this log line crashes on arm64 with access violation on internal tls variable
-		DEBUG_LOG_TRUE(L"NtClose", L"%llu (%ls) -> %ls", uintptr_t(handle), HandleToName(handle), ToString(res));
+		DEBUG_LOG_TRUE(L"NtClose", L"%llu (%ls) -> %ls", uintptr_t(handle), HandleToName(handle), ToString(res).data);
 #endif
 		return res;
 	}
@@ -1323,7 +1333,7 @@ NTSTATUS NTAPI Detoured_NtClose(HANDLE handle)
 	{
 		if (dh.type >= HandleType_StdErr) // TODO: This might leak if handle is duplicated.. but ignore for now
 			return res;
-		DEBUG_LOG_TRUE(L"NtClose", L"%llu (%ls) -> %ls", uintptr_t(handle), HandleToName(handle), ToString(res));
+		DEBUG_LOG_TRUE(L"NtClose", L"%llu (%ls) -> %ls", uintptr_t(handle), HandleToName(handle), ToString(res).data);
 		delete& dh;
 		return res;
 	}
@@ -1332,7 +1342,7 @@ NTSTATUS NTAPI Detoured_NtClose(HANDLE handle)
 	UBA_ASSERT(foRefCount != ~u64(0));
 	if (foRefCount)
 	{
-		DEBUG_LOG_TRUE(L"NtClose", L"%llu (%ls) -> %ls", uintptr_t(handle), HandleToName(handle), ToString(res));
+		DEBUG_LOG_TRUE(L"NtClose", L"%llu (%ls) -> %ls", uintptr_t(handle), HandleToName(handle), ToString(res).data);
 		delete& dh;
 		return res;
 	}
@@ -1450,7 +1460,7 @@ NTSTATUS NTAPI Detoured_NtClose(HANDLE handle)
 		//UBA_ASSERTF(fo->newName.empty(), L"Got close of file that was renamed but had no closeId. Old: %ls New: %ls", fi.originalName, fo->newName.c_str());
 	}
 
-	DEBUG_LOG_DETOURED(L"NtClose", L"%llu (%ls) -> %ls", uintptr_t(handle), HandleToName(handle), ToString(res));
+	DEBUG_LOG_DETOURED(L"NtClose", L"%llu (%ls) -> %ls", uintptr_t(handle), HandleToName(handle), ToString(res).data);
 
 	if (fo->ownsFileInfo)
 	{
@@ -1481,7 +1491,7 @@ NTSTATUS Detoured_NtQueryObject(HANDLE Handle, OBJECT_INFORMATION_CLASS ObjectIn
 		UBA_ASSERTF(Handle != INVALID_HANDLE_VALUE, L"NtQueryObject");
 	}
 	auto res = True_NtQueryObject(Handle, ObjectInformationClass, ObjectInformation, ObjectInformationLength, ReturnLength);
-	DEBUG_LOG_TRUE(L"NtQueryObject", L"(%i) %llu -> %ls", ObjectInformationClass, uintptr_t(Handle), ToString(res));
+	DEBUG_LOG_TRUE(L"NtQueryObject", L"(%i) %llu -> %ls", ObjectInformationClass, uintptr_t(Handle), ToString(res).data);
 	return res;
 }
 
@@ -1494,7 +1504,7 @@ NTSTATUS Detoured_NtQueryInformationProcess(HANDLE ProcessHandle, PROCESSINFOCLA
 	}
 
 	NTSTATUS res = True_NtQueryInformationProcess(ProcessHandle, ProcessInformationClass, ProcessInformation, ProcessInformationLength, ReturnLength);
-	DEBUG_LOG_TRUE(L"NtQueryInformationProcess", L"(class %u) %llu -> %ls", ProcessInformationClass, uintptr_t(ProcessHandle), ToString(res));
+	DEBUG_LOG_TRUE(L"NtQueryInformationProcess", L"(class %u) %llu -> %ls", ProcessInformationClass, uintptr_t(ProcessHandle), ToString(res).data);
 	return res;
 }
 
@@ -1510,7 +1520,7 @@ NTSTATUS Detoured_NtQueryFullAttributesFile(POBJECT_ATTRIBUTES ObjectAttributes,
 {
 	DETOURED_CALL(NtQueryFullAttributesFile);
 	NTSTATUS res = True_NtQueryFullAttributesFile(ObjectAttributes, Attributes);
-	DEBUG_LOG_TRUE(L"NtQueryFullAttributesFile", L"(%ls) -> %ls", ObjectAttributes->ObjectName->Buffer, ToString(res));
+	DEBUG_LOG_TRUE(L"NtQueryFullAttributesFile", L"(%ls) -> %ls", ObjectAttributes->ObjectName->Buffer, ToString(res).data);
 	return res;
 }
 
diff --git a/Engine/Source/Programs/UnrealBuildAccelerator/Detours/Private/Windows/UbaDetoursFunctionsWin.cpp b/Engine/Source/Programs/UnrealBuildAccelerator/Detours/Private/Windows/UbaDetoursFunctionsWin.cpp
--- a/Engine/Source/Programs/UnrealBuildAccelerator/Detours/Private/Windows/UbaDetoursFunctionsWin.cpp
+++ b/Engine/Source/Programs/UnrealBuildAccelerator/Detours/Private/Windows/UbaDetoursFunctionsWin.cpp
@@ -486,6 +486,8 @@ void OnModuleLoaded(HMODULE moduleHandle, const wchar_t* name);
 thread_local const wchar_t* t_renameFileNewName;
 thread_local const wchar_t* t_createFileFileName;
 
+const wchar_t* ToString(BOOL b) { return b ? L"Success" : L"Error"; }
+
 #include "UbaDetoursFunctionsMiMalloc.inl"
 #include "UbaDetoursFunctionsNtDll.inl"
 #include "UbaDetoursFunctionsKernelBase.inl"
