[UBA]
* Added wine version to agent extra info

diff --git a/Engine/Source/Programs/UnrealBuildAccelerator/Agent/Private/UbaAgent.cpp b/Engine/Source/Programs/UnrealBuildAccelerator/Agent/Private/UbaAgent.cpp
--- a/Engine/Source/Programs/UnrealBuildAccelerator/Agent/Private/UbaAgent.cpp
+++ b/Engine/Source/Programs/UnrealBuildAccelerator/Agent/Private/UbaAgent.cpp
@@ -803,7 +803,7 @@ namespace uba
 			extraInfo.Append(TC(", ")).Append(zone);
 
 		if (IsRunningWine())
-			extraInfo.Appendf(TC(", Linux%s/WINE"), (IsArm ? TC("Arm64") : TC("")));
+			extraInfo.Appendf(TC(", Linux%s/Wine-"), (IsArm ? TC("Arm64") : TC(""))).Append(GetWineVersion());
 		else if (IsArm)
 			extraInfo.Append(", Arm64");
 
diff --git a/Engine/Source/Programs/UnrealBuildAccelerator/Core/Private/UbaPlatform.cpp b/Engine/Source/Programs/UnrealBuildAccelerator/Core/Private/UbaPlatform.cpp
--- a/Engine/Source/Programs/UnrealBuildAccelerator/Core/Private/UbaPlatform.cpp
+++ b/Engine/Source/Programs/UnrealBuildAccelerator/Core/Private/UbaPlatform.cpp
@@ -55,21 +55,30 @@ namespace uba
 		#endif
 	}
 
-	bool IsRunningWine()
+	const char* GetWineVersion()
 	{
 		#if PLATFORM_WINDOWS
-		static bool isRunningInWine = []()
+		static const char* wineVersion = []() -> const char*
 			{
-				if (HMODULE kernelBaseModule = GetModuleHandleW(L"kernel32.dll"))
-					return GetProcAddress(kernelBaseModule, "wine_get_unix_file_name") != nullptr;
-				return false;
+				if (HMODULE ntDllModule = GetModuleHandleW(L"ntdll.dll"))
+				{
+					using wine_get_version_func = const char*();
+					if (auto wine_get_version = (wine_get_version_func*)GetProcAddress(ntDllModule, "wine_get_version"))
+						return wine_get_version();
+				}
+				return nullptr;
 			}();
-		return isRunningInWine;
+		return wineVersion;
 		#else
-		return false;
+		return nullptr;
 		#endif
 	}
 
+	bool IsRunningWine()
+	{
+		return GetWineVersion() != nullptr;
+	}
+
 	void Sleep(u32 milliseconds)
 	{
 		#if PLATFORM_WINDOWS
diff --git a/Engine/Source/Programs/UnrealBuildAccelerator/Core/Public/UbaPlatform.h b/Engine/Source/Programs/UnrealBuildAccelerator/Core/Public/UbaPlatform.h
--- a/Engine/Source/Programs/UnrealBuildAccelerator/Core/Public/UbaPlatform.h
+++ b/Engine/Source/Programs/UnrealBuildAccelerator/Core/Public/UbaPlatform.h
@@ -81,6 +81,7 @@ namespace uba
 	using CustomAssertHandler = void(const tchar* text);
 	void SetCustomAssertHandler(CustomAssertHandler* handler);
 	bool CreateGuid(Guid& out);
+	const char* GetWineVersion();
 	bool IsRunningWine();
 	bool IsEscapePressed();
 	void Sleep(u32 milliseconds);
