[UBA]
* Added process peak memory usage to trace

diff --git a/Engine/Source/Programs/UnrealBuildAccelerator/Common/Private/UbaProcess.cpp b/Engine/Source/Programs/UnrealBuildAccelerator/Common/Private/UbaProcess.cpp
--- a/Engine/Source/Programs/UnrealBuildAccelerator/Common/Private/UbaProcess.cpp
+++ b/Engine/Source/Programs/UnrealBuildAccelerator/Common/Private/UbaProcess.cpp
@@ -10,7 +10,7 @@
 #if PLATFORM_WINDOWS
 #include "UbaDetoursPayload.h"
 #include <winternl.h>
-//#include <Psapi.h>
+#include <Psapi.h>
 #include <detours/detours.h>
 #else
 #include <wchar.h>
@@ -1995,6 +1995,10 @@ namespace uba
 
 		auto closeHandleGuard = MakeGuard([&]() { CloseHandle((HANDLE)handle); });
 
+		PROCESS_MEMORY_COUNTERS mem;
+		if (GetProcessMemoryInfo((HANDLE)handle, &mem, sizeof(mem)))
+			m_processStats.peakMemory = mem.PeakWorkingSetSize;
+
 		bool hadTimeout = false;
 		if (cancel)
 			TerminateProcess((HANDLE)handle, ProcessCancelExitCode);
diff --git a/Engine/Source/Programs/UnrealBuildAccelerator/Common/Private/UbaStats.cpp b/Engine/Source/Programs/UnrealBuildAccelerator/Common/Private/UbaStats.cpp
--- a/Engine/Source/Programs/UnrealBuildAccelerator/Common/Private/UbaStats.cpp
+++ b/Engine/Source/Programs/UnrealBuildAccelerator/Common/Private/UbaStats.cpp
@@ -21,29 +21,36 @@ namespace uba
 
 	void ProcessStats::Print(Logger& logger, u64 frequency)
 	{
-		logger.Info(TC("  Total              %8u %9s"), GetTotalCount(), TimeToText(GetTotalTime(), false, frequency).str);
-		logger.Info(TC("  WaitOnResponse     %8u %9s"), waitOnResponse.count.load(), TimeToText(waitOnResponse.time, false, frequency).str);
-		logger.Info(TC("  Host                %17s"), TimeToText(hostTotalTime, false, frequency).str);
-		logger.Info(TC(""));
-
-		struct Stat { const char* name; u64 nameLen; const Timer& timer; };
-		Stat stats[] =
+		if (hostTotalTime)
 		{
-			#define UBA_PROCESS_STAT(T, ver) { #T, sizeof(#T), T },
-			UBA_PROCESS_STATS
-			#undef UBA_PROCESS_STAT
-		};
+			logger.Info(TC("  Total              %8u %9s"), GetTotalCount(), TimeToText(GetTotalTime(), false, frequency).str);
+			logger.Info(TC("  WaitOnResponse     %8u %9s"), waitOnResponse.count.load(), TimeToText(waitOnResponse.time, false, frequency).str);
+			logger.Info(TC("  Host                %17s"), TimeToText(hostTotalTime, false, frequency).str);
+			logger.Info(TC(""));
 
-		for (Stat& s : stats)
-			if (s.timer.count)
-				logger.Info(TC("  %c") PERCENT_HS TC("%s %8u %9s"), ToUpper(s.name[0]), s.name + 1, EmptyCharArray + s.nameLen, s.timer.count.load(), TimeToText(s.timer.time, false, frequency).str);
+			struct Stat { const char* name; u64 nameLen; const Timer& timer; };
+			Stat stats[] =
+			{
+				#define UBA_PROCESS_STAT(T, ver) { #T, sizeof(#T), T },
+				UBA_PROCESS_STATS
+				#undef UBA_PROCESS_STAT
+			};
 
-		logger.Info(TC(""));
+			for (Stat& s : stats)
+				if (s.timer.count)
+					logger.Info(TC("  %c") PERCENT_HS TC("%s %8u %9s"), ToUpper(s.name[0]), s.name + 1, EmptyCharArray + s.nameLen, s.timer.count.load(), TimeToText(s.timer.time, false, frequency).str);
+
+			logger.Info(TC(""));
 
-		logger.Info(TC("  HighestMem                  %9s"), BytesToText(usedMemory).str);
-		logger.Info(TC("  Startup Time                %9s"), TimeToText(startupTime, false, frequency).str);
-		logger.Info(TC("  Exit Time                   %9s"), TimeToText(exitTime, false, frequency).str);
-		logger.Info(TC("  CPU Time                    %9s"), TimeToText(cpuTime, false, frequency).str);
+			logger.Info(TC("  Startup Time                %9s"), TimeToText(startupTime, false, frequency).str);
+			logger.Info(TC("  Exit Time                   %9s"), TimeToText(exitTime, false, frequency).str);
+			if (detoursMemory)
+				logger.Info(TC("  DetoursMem                  %9s"), BytesToText(detoursMemory).str);
+		}
+		if (peakMemory)
+			logger.Info(TC("  PeakMem                     %9s"), BytesToText(peakMemory).str);
+		if (cpuTime)
+			logger.Info(TC("  CPU Time                    %9s"), TimeToText(cpuTime, false, frequency).str);
 		logger.Info(TC("  Wall Time                   %9s"), TimeToText(wallTime, false, frequency).str);
 	}
 
@@ -83,12 +90,25 @@ namespace uba
 			#undef UBA_PROCESS_STAT
 		}
 
-		startupTime = reader.ReadU64();
-		exitTime = reader.ReadU64();
-		wallTime = reader.ReadU64();
-		cpuTime = reader.ReadU64();
-		usedMemory = reader.ReadU32();
-		hostTotalTime = reader.ReadU64();
+		if (version >= 37)
+		{
+			startupTime = reader.Read7BitEncoded();
+			exitTime = reader.Read7BitEncoded();
+			wallTime = reader.Read7BitEncoded();
+			cpuTime = reader.Read7BitEncoded();
+			detoursMemory = reader.Read7BitEncoded();
+			peakMemory = reader.Read7BitEncoded();
+			hostTotalTime = reader.Read7BitEncoded();
+		}
+		else
+		{
+			startupTime = reader.ReadU64();
+			exitTime = reader.ReadU64();
+			wallTime = reader.ReadU64();
+			cpuTime = reader.ReadU64();
+			detoursMemory = reader.ReadU32();
+			hostTotalTime = reader.ReadU64();
+		}
 	}
 
 	void ProcessStats::Add(const ProcessStats& other)
@@ -103,7 +123,8 @@ namespace uba
 		exitTime += other.exitTime;
 		wallTime += other.wallTime;
 		cpuTime += other.cpuTime;
-		usedMemory = Max(usedMemory, other.usedMemory);
+		detoursMemory = Max(detoursMemory, other.detoursMemory);
+		peakMemory = Max(peakMemory, other.peakMemory);
 		hostTotalTime += other.hostTotalTime;
 	}
 
diff --git a/Engine/Source/Programs/UnrealBuildAccelerator/Common/Public/UbaTrace.h b/Engine/Source/Programs/UnrealBuildAccelerator/Common/Public/UbaTrace.h
--- a/Engine/Source/Programs/UnrealBuildAccelerator/Common/Public/UbaTrace.h
+++ b/Engine/Source/Programs/UnrealBuildAccelerator/Common/Public/UbaTrace.h
@@ -70,7 +70,7 @@ namespace uba
 	using Color = u32;
 	inline Color toColor(u8 r, u8 g, u8 b) { return (r << 16) + (g << 8) + b; }
 
-	static constexpr u32 TraceVersion = 36;
+	static constexpr u32 TraceVersion = 37;
 	static constexpr u32 TraceReadCompatibilityVersion = 6;
 
 	class Trace : public WorkTracker
diff --git a/Engine/Source/Programs/UnrealBuildAccelerator/Core/Public/UbaProcessStats.h b/Engine/Source/Programs/UnrealBuildAccelerator/Core/Public/UbaProcessStats.h
--- a/Engine/Source/Programs/UnrealBuildAccelerator/Core/Public/UbaProcessStats.h
+++ b/Engine/Source/Programs/UnrealBuildAccelerator/Core/Public/UbaProcessStats.h
@@ -64,7 +64,8 @@ namespace uba
 		AtomicU64 wallTime;
 		AtomicU64 cpuTime;
 
-		AtomicU64 usedMemory;
+		AtomicU64 detoursMemory;
+		AtomicU64 peakMemory;
 
 		AtomicU64 hostTotalTime;
 
@@ -90,12 +91,13 @@ namespace uba
 			UBA_PROCESS_STATS
 			#undef UBA_PROCESS_STAT
 
-			writer.WriteU64(startupTime);
-			writer.WriteU64(exitTime);
-			writer.WriteU64(wallTime);
-			writer.WriteU64(cpuTime);
-			writer.WriteU32(u32(usedMemory));
-			writer.WriteU64(hostTotalTime);
+			writer.Write7BitEncoded(startupTime);
+			writer.Write7BitEncoded(exitTime);
+			writer.Write7BitEncoded(wallTime);
+			writer.Write7BitEncoded(cpuTime);
+			writer.Write7BitEncoded(detoursMemory);
+			writer.Write7BitEncoded(peakMemory);
+			writer.Write7BitEncoded(hostTotalTime);
 		}
 	};
 
diff --git a/Engine/Source/Programs/UnrealBuildAccelerator/Detours/Private/Windows/UbaDetoursFunctionsWin.cpp b/Engine/Source/Programs/UnrealBuildAccelerator/Detours/Private/Windows/UbaDetoursFunctionsWin.cpp
--- a/Engine/Source/Programs/UnrealBuildAccelerator/Detours/Private/Windows/UbaDetoursFunctionsWin.cpp
+++ b/Engine/Source/Programs/UnrealBuildAccelerator/Detours/Private/Windows/UbaDetoursFunctionsWin.cpp
@@ -511,7 +511,7 @@ void SendExitMessage(DWORD exitCode, u64 startTime)
 	if (g_trackInputsMem)
 		SendInput();
 
-	g_stats.usedMemory = u32(g_memoryBlock.writtenSize);
+	g_stats.detoursMemory = g_memoryBlock.writtenSize; // + g_directoryTable.m_memorySize + g_mappedFileTable.m_memPosition;
 
 	SCOPED_WRITE_LOCK(g_communicationLock, pcs);
 	BinaryWriter writer;
@@ -538,7 +538,15 @@ void DetourAttachFunction(void** trueFunc, void* detouredFunc, const char* funcN
 	auto error = DetourAttach(trueFunc, detouredFunc);
 	if (error == NO_ERROR)
 		return;
-	Rpc_WriteLogf(L"Failed to detour %hs", funcName);
+	const char* errorString = "Unknown error";
+	switch (error)
+	{
+	case ERROR_INVALID_BLOCK: errorString = "The function referenced is too small to be detoured."; break;
+	case ERROR_INVALID_HANDLE: errorString = "The ppPointer parameter is NULL or points to a NULL pointer."; break;
+	case ERROR_INVALID_OPERATION: errorString = "No pending transaction exists."; break;
+	case ERROR_NOT_ENOUGH_MEMORY: errorString = "Not enough memory exists to complete the operation."; break;
+	}
+	Rpc_WriteLogf(L"Failed to detour %hs (%hs)", funcName, errorString);
 	ExitProcess(error);
 }
 
diff --git a/Engine/Source/Programs/UnrealBuildAccelerator/Visualizer/Private/UbaVisualizer.cpp b/Engine/Source/Programs/UnrealBuildAccelerator/Visualizer/Private/UbaVisualizer.cpp
--- a/Engine/Source/Programs/UnrealBuildAccelerator/Visualizer/Private/UbaVisualizer.cpp
+++ b/Engine/Source/Programs/UnrealBuildAccelerator/Visualizer/Private/UbaVisualizer.cpp
@@ -1851,6 +1851,12 @@ namespace uba
 					logger.Info(L"  ----------- Detours stats -----------");
 					processStats.Print(logger, m_traceView.frequency);
 				}
+				else if (processStats.peakMemory)
+				{
+					logger.Info(L"");
+					logger.Info(L"  ----------- Process stats -----------");
+					processStats.Print(logger, m_traceView.frequency);
+				}
 
 				if (!sessionStats.IsEmpty())
 				{
