[UBA]
* Fixed so detouring works with wine arm64

diff --git a/Engine/Source/Programs/UnrealBuildAccelerator/Agent/Private/UbaAgent.cpp b/Engine/Source/Programs/UnrealBuildAccelerator/Agent/Private/UbaAgent.cpp
--- a/Engine/Source/Programs/UnrealBuildAccelerator/Agent/Private/UbaAgent.cpp
+++ b/Engine/Source/Programs/UnrealBuildAccelerator/Agent/Private/UbaAgent.cpp
@@ -803,9 +803,7 @@ namespace uba
 			extraInfo.Append(TC(", ")).Append(zone);
 
 		if (IsRunningWine())
-			extraInfo.Appendf(TC(", Linux%s/Wine-"), (IsArm ? TC("Arm64") : TC(""))).Append(GetWineVersion());
-		else if (IsArm)
-			extraInfo.Append(", Arm64");
+			extraInfo.Append(TC(", Linux/Wine-")).Append(GetWineVersion());
 
 		if (useQuic)
 			extraInfo.Append(TC(", MsQuic"));
@@ -817,7 +815,7 @@ namespace uba
 		#if UBA_DEBUG
 		dbgStr = TC(" (DEBUG)");
 		#endif
-		logger.Info(TC("UbaAgent v%s%s (Cpu: %u, MaxCon: %u, Dir: \"%s\", StoreCapacity: %uGb%s)"), Version, dbgStr, maxProcessCount, maxConnectionCount, g_rootDir.data, storageCapacityGb, extraInfo.data);
+		logger.Info(TC("UbaAgent v%s%s (Cpu%s: %u, MaxCon: %u, Dir: \"%s\", StoreCapacity: %uGb%s)"), Version, dbgStr, (IsArm ? TC("[Arm]") : TC("")), maxProcessCount, maxConnectionCount, g_rootDir.data, storageCapacityGb, extraInfo.data);
 		
 		Config config;
 		if (!configFile.IsEmpty())
diff --git a/Engine/Source/Programs/UnrealBuildAccelerator/Common/Private/UbaSession.cpp b/Engine/Source/Programs/UnrealBuildAccelerator/Common/Private/UbaSession.cpp
--- a/Engine/Source/Programs/UnrealBuildAccelerator/Common/Private/UbaSession.cpp
+++ b/Engine/Source/Programs/UnrealBuildAccelerator/Common/Private/UbaSession.cpp
@@ -2618,7 +2618,7 @@ namespace uba
 		BytesToText temp(capacity);
 		if (capacity)
 			capacityStr = temp.str;
-		out.Appendf(TC("CPU:%s%s Mem:%ugb Cas:%s"), cpuCountStr.data, hzStr.data, u32(totalMemoryInKilobytes/(1024*1024)), capacityStr);
+		out.Appendf(TC("CPU%s:%s%s Mem:%ugb Cas:%s"), (IsArm ? TC("[Arm]") : TC("")), cpuCountStr.data, hzStr.data, u32(totalMemoryInKilobytes/(1024*1024)), capacityStr);
 
 		StringBuffer<128> zone;
 		if (m_storage.GetZone(zone))
diff --git a/Engine/Source/Programs/UnrealBuildAccelerator/Detours/Private/Windows/UbaDetoursFunctionsKernelBase.inl b/Engine/Source/Programs/UnrealBuildAccelerator/Detours/Private/Windows/UbaDetoursFunctionsKernelBase.inl
--- a/Engine/Source/Programs/UnrealBuildAccelerator/Detours/Private/Windows/UbaDetoursFunctionsKernelBase.inl
+++ b/Engine/Source/Programs/UnrealBuildAccelerator/Detours/Private/Windows/UbaDetoursFunctionsKernelBase.inl
@@ -2654,14 +2654,6 @@ BOOL Detoured_SetConsoleMode(HANDLE hConsoleHandle, DWORD mode)
 	return True_SetConsoleMode(hConsoleHandle, mode);
 }
 
-BOOL Detoured_GetConsoleTitleW(LPTSTR lpConsoleTitle, DWORD nSize)
-{
-	DETOURED_CALL(GetConsoleTitleW);
-	DEBUG_LOG_DETOURED(L"GetConsoleTitleW", L"");
-	*lpConsoleTitle = 0;
-	return true;
-}
-
 BOOL Detoured_CreateProcessW(LPCWSTR lpApplicationName, LPWSTR lpCommandLine, LPSECURITY_ATTRIBUTES lpProcessAttributes, LPSECURITY_ATTRIBUTES lpThreadAttributes, BOOL bInheritHandles,
 	DWORD dwCreationFlags, LPVOID lpEnvironment, LPCWSTR lpCurrentDirectory, LPSTARTUPINFOW lpStartupInfo, LPPROCESS_INFORMATION lpProcessInformation)
 {
@@ -3148,14 +3140,6 @@ BOOL Detoured_GetFileAttributesExA(LPCSTR lpFileName, GET_FILEEX_INFO_LEVELS fIn
 	return True_GetFileAttributesExA(lpFileName, fInfoLevelId, lpFileInformation);
 }
 
-HMODULE Detoured_LoadLibraryW(LPCWSTR lpLibFileName)
-{
-	DETOURED_CALL(LoadLibraryW);
-	DEBUG_LOG_TRUE(L"LoadLibraryW", L"(%ls)", lpLibFileName);
-	return True_LoadLibraryW(lpLibFileName);
-}
-
-
 DWORD Shared_GetModuleFileNameA(HMODULE hModule, const wchar_t* moduleName, u32 moduleNameLen, LPSTR lpFilename, DWORD nSize)
 {
 	if (nSize <= moduleNameLen)
@@ -3339,24 +3323,6 @@ BOOL Detoured_ReadFileScatter(HANDLE hFile, FILE_SEGMENT_ELEMENT* aSegmentArray,
 	return True_ReadFileScatter(hFile, aSegmentArray, nNumberOfBytesToRead, lpReserved, lpOverlapped);
 }
 
-
-void Detoured_SetLastError(DWORD dwErrCode)
-{
-	DETOURED_CALL(SetLastError);
-	if (dwErrCode != ERROR_SUCCESS)
-		while (false) {}
-	True_SetLastError(dwErrCode);
-}
-
-DWORD Detoured_GetLastError()
-{
-	DETOURED_CALL(GetLastError);
-	auto res = True_GetLastError();
-	if (res != ERROR_SUCCESS)
-		while (false) {}
-	return res;
-}
-
 BOOL Detoured_SetFileValidData(HANDLE hFile, LONGLONG ValidDataLength)
 {
 	DETOURED_CALL(SetFileValidData);
diff --git a/Engine/Source/Programs/UnrealBuildAccelerator/Detours/Private/Windows/UbaDetoursFunctionsMiMalloc.inl b/Engine/Source/Programs/UnrealBuildAccelerator/Detours/Private/Windows/UbaDetoursFunctionsMiMalloc.inl
--- a/Engine/Source/Programs/UnrealBuildAccelerator/Detours/Private/Windows/UbaDetoursFunctionsMiMalloc.inl
+++ b/Engine/Source/Programs/UnrealBuildAccelerator/Detours/Private/Windows/UbaDetoursFunctionsMiMalloc.inl
@@ -142,6 +142,8 @@ void* Detoured__calloc_base(size_t _Count, size_t _Size)
 
 void* Detoured__realloc_base(void* memblock, size_t size)
 {
+	UBA_ASSERTF(false, L"_realloc_base called but is not detoured for arm64 hosts");
+	UBA_ASSERT(false); // Since this is not detourable on wine+arm64 we want to validate it is not being called
 	return mi_realloc(memblock, size);
 }
 
@@ -175,6 +177,7 @@ void* Detoured__recalloc_base(void* memblock, size_t num, size_t size)
 
 size_t Detoured__aligned_msize(void* p, size_t alignment, size_t offset)
 {
+	UBA_ASSERTF(false, L"_aligned_msize called but is only detoured in debug");
 	DETOURED_CALL(_aligned_msize);
 	return mi_usable_size(p);
 }
diff --git a/Engine/Source/Programs/UnrealBuildAccelerator/Detours/Private/Windows/UbaDetoursFunctionsWin.h b/Engine/Source/Programs/UnrealBuildAccelerator/Detours/Private/Windows/UbaDetoursFunctionsWin.h
--- a/Engine/Source/Programs/UnrealBuildAccelerator/Detours/Private/Windows/UbaDetoursFunctionsWin.h
+++ b/Engine/Source/Programs/UnrealBuildAccelerator/Detours/Private/Windows/UbaDetoursFunctionsWin.h
@@ -31,6 +31,12 @@
 	DETOURED_FUNCTIONS_UCRTBASE \
 	DETOURED_FUNCTIONS_RPCRT4 \
 
+#if defined(_M_ARM64)
+#define DETOURED_FUNCTION_X64(x)
+#else
+#define DETOURED_FUNCTION_X64(x) DETOURED_FUNCTION(x)
+#endif
+
 #define DETOURED_FUNCTIONS_KERNELBASE \
 	DETOURED_FUNCTION(GetCommandLineW) \
 	DETOURED_FUNCTION(GetCurrentDirectoryW) \
@@ -100,7 +106,7 @@
 	DETOURED_FUNCTION(GetVolumeInformationW) \
 	DETOURED_FUNCTION(GetUserDefaultUILanguage) \
 	DETOURED_FUNCTION(GetThreadPreferredUILanguages) \
-	DETOURED_FUNCTION(GetConsoleTitleW) \
+	DETOURED_FUNCTION_X64(GetConsoleTitleW) \
 	DETOURED_FUNCTION(WaitForSingleObject) \
 	DETOURED_FUNCTION(WaitForSingleObjectEx) \
 	DETOURED_FUNCTION(WaitForMultipleObjects) \
@@ -150,14 +156,6 @@
 #define DETOURED_WSPLITPATH
 #endif
 
-// TODO: For some reason _isatty can't be detoured on windows arm64
-#if !defined(_M_ARM64)
-#define DETOURED_ISATTY DETOURED_FUNCTION(_isatty)
-#else
-#define DETOURED_ISATTY
-#endif
-
-
 #define DETOURED_FUNCTIONS_UCRTBASE \
 	DETOURED_FUNCTION(_wgetcwd) \
 	DETOURED_FUNCTION(_wfullpath) \
@@ -168,7 +166,7 @@
 	DETOURED_FUNCTION(_get_osfhandle) \
 	DETOURED_FUNCTION(_write) \
 	DETOURED_FUNCTION(fputs) \
-	DETOURED_ISATTY \
+	DETOURED_FUNCTION_X64(_isatty) \
 	DETOURED_WSPLITPATH \
 	DETOURED_FUNCTIONS_UCRTBASE_DEBUG \
 
@@ -206,7 +204,7 @@
 	DETOURED_FUNCTION(_wdupenv_s) \
 	DETOURED_FUNCTION(_malloc_base) \
 	DETOURED_FUNCTION(_calloc_base) \
-	DETOURED_FUNCTION(_realloc_base) \
+	DETOURED_FUNCTION_X64(_realloc_base) \
 	DETOURED_FUNCTION(_free_base) \
 	DETOURED_FUNCTION(_expand_base) \
 	DETOURED_FUNCTION(_msize_base) \
@@ -223,8 +221,8 @@
 	DETOURED_FUNCTION(GetCommandLineA) \
 	DETOURED_FUNCTION(FreeLibrary) \
 	DETOURED_FUNCTION(RegOpenKeyExW) \
-	DETOURED_FUNCTION(SetLastError) \
-	DETOURED_FUNCTION(GetLastError) \
+	DETOURED_FUNCTION_X64(SetLastError) \
+	DETOURED_FUNCTION_X64(GetLastError) \
 	DETOURED_FUNCTION(RegOpenKeyExA) \
 	DETOURED_FUNCTION(RegCloseKey) \
 	DETOURED_FUNCTION(IsValidCodePage) \
@@ -251,7 +249,7 @@
 	DETOURED_FUNCTION(GetFullPathNameA) \
 	DETOURED_FUNCTION(GetFileAttributesA) \
 	DETOURED_FUNCTION(GetFileAttributesExA) \
-	DETOURED_FUNCTION(LoadLibraryW) \
+	DETOURED_FUNCTION_X64(LoadLibraryW) \
 	DETOURED_FUNCTION(GetModuleFileNameA) \
 	DETOURED_FUNCTION(GetModuleFileNameExA) \
 	DETOURED_FUNCTION(GetModuleBaseNameA) \
