[UBA]
* Fixed compile error

diff --git a/Engine/Source/Programs/UnrealBuildAccelerator/Detours/Private/Windows/UbaDetoursFunctionsKernelBase.inl b/Engine/Source/Programs/UnrealBuildAccelerator/Detours/Private/Windows/UbaDetoursFunctionsKernelBase.inl
--- a/Engine/Source/Programs/UnrealBuildAccelerator/Detours/Private/Windows/UbaDetoursFunctionsKernelBase.inl
+++ b/Engine/Source/Programs/UnrealBuildAccelerator/Detours/Private/Windows/UbaDetoursFunctionsKernelBase.inl
@@ -2654,6 +2654,14 @@ BOOL Detoured_SetConsoleMode(HANDLE hConsoleHandle, DWORD mode)
 	return True_SetConsoleMode(hConsoleHandle, mode);
 }
 
+BOOL Detoured_GetConsoleTitleW(LPTSTR lpConsoleTitle, DWORD nSize)
+{
+	DETOURED_CALL(GetConsoleTitleW);
+	DEBUG_LOG_DETOURED(L"GetConsoleTitleW", L"");
+	*lpConsoleTitle = 0;
+	return true;
+}
+
 BOOL Detoured_CreateProcessW(LPCWSTR lpApplicationName, LPWSTR lpCommandLine, LPSECURITY_ATTRIBUTES lpProcessAttributes, LPSECURITY_ATTRIBUTES lpThreadAttributes, BOOL bInheritHandles,
 	DWORD dwCreationFlags, LPVOID lpEnvironment, LPCWSTR lpCurrentDirectory, LPSTARTUPINFOW lpStartupInfo, LPPROCESS_INFORMATION lpProcessInformation)
 {
@@ -3140,6 +3148,15 @@ BOOL Detoured_GetFileAttributesExA(LPCSTR lpFileName, GET_FILEEX_INFO_LEVELS fIn
 	return True_GetFileAttributesExA(lpFileName, fInfoLevelId, lpFileInformation);
 }
 
+#if !defined(_M_ARM64)
+HMODULE Detoured_LoadLibraryW(LPCWSTR lpLibFileName)
+{
+	DETOURED_CALL(LoadLibraryW);
+	DEBUG_LOG_TRUE(L"LoadLibraryW", L"(%ls)", lpLibFileName);
+	return True_LoadLibraryW(lpLibFileName);
+}
+#endif
+
 DWORD Shared_GetModuleFileNameA(HMODULE hModule, const wchar_t* moduleName, u32 moduleNameLen, LPSTR lpFilename, DWORD nSize)
 {
 	if (nSize <= moduleNameLen)
@@ -3323,6 +3340,27 @@ BOOL Detoured_ReadFileScatter(HANDLE hFile, FILE_SEGMENT_ELEMENT* aSegmentArray,
 	return True_ReadFileScatter(hFile, aSegmentArray, nNumberOfBytesToRead, lpReserved, lpOverlapped);
 }
 
+#if !defined(_M_ARM64)
+void Detoured_SetLastError(DWORD dwErrCode)
+{
+	DETOURED_CALL(SetLastError);
+	if (dwErrCode != ERROR_SUCCESS)
+		while (false) {}
+	True_SetLastError(dwErrCode);
+}
+#endif
+
+#if !defined(_M_ARM64)
+DWORD Detoured_GetLastError()
+{
+	DETOURED_CALL(GetLastError);
+	auto res = True_GetLastError();
+	if (res != ERROR_SUCCESS)
+		while (false) {}
+	return res;
+}
+#endif
+
 BOOL Detoured_SetFileValidData(HANDLE hFile, LONGLONG ValidDataLength)
 {
 	DETOURED_CALL(SetFileValidData);
