[UBA]
* Fixed so DirectoryListHandle is not a pointer but a handle value. Fixed bug where memory address could accidentally overlap other types of handles
* Removed some of the detouring of memory allocation functions on wine because they just call the other detoured ones (and caused failed to detour on arm64)

diff --git a/Engine/Source/Programs/UnrealBuildAccelerator/Detours/Private/Windows/UbaDetoursFunctionsMiMalloc.inl b/Engine/Source/Programs/UnrealBuildAccelerator/Detours/Private/Windows/UbaDetoursFunctionsMiMalloc.inl
--- a/Engine/Source/Programs/UnrealBuildAccelerator/Detours/Private/Windows/UbaDetoursFunctionsMiMalloc.inl
+++ b/Engine/Source/Programs/UnrealBuildAccelerator/Detours/Private/Windows/UbaDetoursFunctionsMiMalloc.inl
@@ -142,8 +142,6 @@ void* Detoured__calloc_base(size_t _Count, size_t _Size)
 
 void* Detoured__realloc_base(void* memblock, size_t size)
 {
-	UBA_ASSERTF(false, L"_realloc_base called but is not detoured for arm64 hosts");
-	UBA_ASSERT(false); // Since this is not detourable on wine+arm64 we want to validate it is not being called
 	return mi_realloc(memblock, size);
 }
 
diff --git a/Engine/Source/Programs/UnrealBuildAccelerator/Detours/Private/Windows/UbaDetoursFunctionsNtDll.inl b/Engine/Source/Programs/UnrealBuildAccelerator/Detours/Private/Windows/UbaDetoursFunctionsNtDll.inl
--- a/Engine/Source/Programs/UnrealBuildAccelerator/Detours/Private/Windows/UbaDetoursFunctionsNtDll.inl
+++ b/Engine/Source/Programs/UnrealBuildAccelerator/Detours/Private/Windows/UbaDetoursFunctionsNtDll.inl
@@ -1619,8 +1619,9 @@ PVOID Detoured_RtlReAllocateHeap(PVOID HeapHandle, ULONG Flags, PVOID BaseAddres
 #if UBA_USE_MIMALLOC
 	if (g_checkRtlHeap && IsInMiMalloc(BaseAddress))
 	{
-		Rpc_WriteLogf(L"ERROR: RtlReAllocateHeap - This is not implemented");
-		return 0;
+		return mi_realloc(BaseAddress, Size);
+		//Rpc_WriteLogf(L"ERROR: RtlReAllocateHeap - This is not implemented");
+		//return 0;
 	}
 #endif
 	//if (Flags & HEAP_ZERO_MEMORY)
diff --git a/Engine/Source/Programs/UnrealBuildAccelerator/Detours/Private/Windows/UbaDetoursFunctionsWin.cpp b/Engine/Source/Programs/UnrealBuildAccelerator/Detours/Private/Windows/UbaDetoursFunctionsWin.cpp
--- a/Engine/Source/Programs/UnrealBuildAccelerator/Detours/Private/Windows/UbaDetoursFunctionsWin.cpp
+++ b/Engine/Source/Programs/UnrealBuildAccelerator/Detours/Private/Windows/UbaDetoursFunctionsWin.cpp
@@ -29,10 +29,12 @@ NTSTATUS NtCopyFileChunk(HANDLE Source, HANDLE Dest, HANDLE Event, PIO_STATUS_BL
 #define DETOURED_FUNCTION(Func) decltype(Func)* True_##Func = ::Func; 
 DETOURED_FUNCTIONS
 DETOURED_FUNCTIONS_MEMORY
+DETOURED_FUNCTIONS_MEMORY_NON_WINE
 #undef DETOURED_FUNCTION
 
 #define DETOURED_FUNCTION(Func) void* Local_##Func = ::Func;
 DETOURED_FUNCTIONS_MEMORY
+DETOURED_FUNCTIONS_MEMORY_NON_WINE
 #undef DETOURED_FUNCTION
 
 #include "UbaDirectoryTable.h"
@@ -100,10 +102,11 @@ void FlushDebugLog()
 #define DETOURED_FUNCTION(name) Timer timer##name;
 DETOURED_FUNCTIONS
 DETOURED_FUNCTIONS_MEMORY
+DETOURED_FUNCTIONS_MEMORY_NON_WINE
 #undef DETOURED_FUNCTION
 #define DETOURED_CALL(name) TimerScope _(timer##name)
 #else
-#define DETOURED_CALL(name)
+#define DETOURED_CALL(name) //DEBUG_LOG(TC(#name));
 #endif
 
 bool g_isDetachedProcess;
@@ -212,14 +215,14 @@ struct ListDirectoryHandle
 
 constexpr u64 ListDirHandlesRange = 4*1024*1024;
 MemoryBlock& g_listDirHandleMemoryBlock = *new MemoryBlock(ListDirHandlesRange);
-u64 g_listDirHandlesStart = u64(g_listDirHandleMemoryBlock.memory);
-u64 g_listDirHandlesEnd = g_listDirHandlesStart + g_listDirHandleMemoryBlock.reserveSize;
+constexpr u64 ListDirHandlesStart = DetouredHandleEnd;
+constexpr u64 ListDirHandlesEnd = ListDirHandlesStart + ListDirHandlesRange/sizeof(ListDirectoryHandle);
 BlockAllocator<ListDirectoryHandle> g_listDirectoryHandleAllocator(g_listDirHandleMemoryBlock);
 void* ListDirectoryHandle::operator new(size_t size) { return g_listDirectoryHandleAllocator.Allocate(); }
 void ListDirectoryHandle::operator delete(void* p) { g_listDirectoryHandleAllocator.Free(p); }
-inline bool isListDirectoryHandle(HANDLE h) { return (u64)h >= g_listDirHandlesStart && (u64)h < g_listDirHandlesEnd; }
-inline HANDLE makeListDirectoryHandle(ListDirectoryHandle* p) { return HANDLE(p);}
-inline ListDirectoryHandle& asListDirectoryHandle(HANDLE h) { return *(ListDirectoryHandle*)h; }
+inline bool isListDirectoryHandle(HANDLE h) { return (u64)h >= ListDirHandlesStart && (u64)h < ListDirHandlesEnd; }
+inline HANDLE makeListDirectoryHandle(ListDirectoryHandle* p) { return HANDLE(ListDirHandlesStart + (p - (ListDirectoryHandle*)g_listDirHandleMemoryBlock.memory)); }
+inline ListDirectoryHandle& asListDirectoryHandle(HANDLE h) { return *(((ListDirectoryHandle*)g_listDirHandleMemoryBlock.memory) + (u64(h) - ListDirHandlesStart)); }
 
 ReaderWriterLock g_loadedModulesLock;
 UnorderedMap<HMODULE, TString> g_loadedModules;
@@ -370,11 +373,10 @@ const wchar_t* HandleToName(HANDLE handle)
 	if (!isDetouredHandle(handle))
 		return L"UNKNOWN";
 	DetouredHandle& dh = asDetouredHandle(handle);
-	if (dh.fileObject)
-	{
-		if (const wchar_t* name = dh.fileObject->fileInfo->name)
-			return name;
-	}
+	if (FileObject* fo = dh.fileObject)
+		if (FileInfo* fi = fo->fileInfo)
+			if (const wchar_t* name = fi->name)
+				return name;
 	return L"DETOURED";
 }
 
@@ -604,6 +606,10 @@ int DetourAttachFunctions(bool runningRemote)
 		if (g_useMiMalloc)
 		{
 			DETOURED_FUNCTIONS_MEMORY
+			if (!g_isRunningWine)
+			{
+				DETOURED_FUNCTIONS_MEMORY_NON_WINE
+			}
 		}
 	}
 
@@ -638,6 +644,10 @@ int DetourAttachFunctions(bool runningRemote)
 	if (g_useMiMalloc)
 	{
 		DETOURED_FUNCTIONS_MEMORY
+		if (!g_isRunningWine)
+		{
+			DETOURED_FUNCTIONS_MEMORY_NON_WINE
+		}
 	}
 	#undef DETOURED_FUNCTION
 
@@ -701,6 +711,10 @@ int DetourDetachFunctions()
 	if (g_useMiMalloc)
 	{
 		DETOURED_FUNCTIONS_MEMORY
+		if (!g_isRunningWine)
+		{
+			DETOURED_FUNCTIONS_MEMORY_NON_WINE
+		}
 	}
 	DETOURED_FUNCTIONS
 	#undef DETOURED_FUNCTION
@@ -993,6 +1007,7 @@ void Deinit(u64 startTime)
 	#define DETOURED_FUNCTION(name) if (timer##name.count != 0) { char sb[1024]; sprintf_s(sb, sizeof(sb), "%s: %u %llu\n", #name, timer##name.count.load(), TimeToMs(timer##name.time.load())); WriteDebug(sb); }
 	DETOURED_FUNCTIONS
 	DETOURED_FUNCTIONS_MEMORY
+	DETOURED_FUNCTIONS_MEMORY_NON_WINE
 	#undef DETOURED_FUNCTION
 	#endif
 
diff --git a/Engine/Source/Programs/UnrealBuildAccelerator/Detours/Private/Windows/UbaDetoursFunctionsWin.h b/Engine/Source/Programs/UnrealBuildAccelerator/Detours/Private/Windows/UbaDetoursFunctionsWin.h
--- a/Engine/Source/Programs/UnrealBuildAccelerator/Detours/Private/Windows/UbaDetoursFunctionsWin.h
+++ b/Engine/Source/Programs/UnrealBuildAccelerator/Detours/Private/Windows/UbaDetoursFunctionsWin.h
@@ -202,17 +202,21 @@
 	DETOURED_FUNCTION(_aligned_offset_recalloc) \
 	DETOURED_FUNCTION(_dupenv_s) \
 	DETOURED_FUNCTION(_wdupenv_s) \
+	DETOURED_FUNCTION(_free_base) \
+	DETOURED_FUNCTIONS_MEMORY_DEBUG \
+
+// All these are calling above functions on wine
+#define DETOURED_FUNCTIONS_MEMORY_NON_WINE \
 	DETOURED_FUNCTION(_malloc_base) \
 	DETOURED_FUNCTION(_calloc_base) \
-	DETOURED_FUNCTION_X64(_realloc_base) \
-	DETOURED_FUNCTION(_free_base) \
+	DETOURED_FUNCTION(_realloc_base) \
 	DETOURED_FUNCTION(_expand_base) \
 	DETOURED_FUNCTION(_msize_base) \
 	DETOURED_FUNCTION(_recalloc_base) \
-	DETOURED_FUNCTIONS_MEMORY_DEBUG \
 
 #else
 #define DETOURED_FUNCTIONS_MEMORY
+#define DETOURED_FUNCTIONS_MEMORY_NON_WINE
 #endif
 
 #if defined(DETOURED_INCLUDE_DEBUG)
