[UBA]
* Bumped reserved memory on detoured process for memory allocations since we saw a farm build running out of reserved space

diff --git a/Engine/Source/Programs/UnrealBuildAccelerator/Detours/Private/UbaDetoursShared.cpp b/Engine/Source/Programs/UnrealBuildAccelerator/Detours/Private/UbaDetoursShared.cpp
--- a/Engine/Source/Programs/UnrealBuildAccelerator/Detours/Private/UbaDetoursShared.cpp
+++ b/Engine/Source/Programs/UnrealBuildAccelerator/Detours/Private/UbaDetoursShared.cpp
@@ -44,7 +44,7 @@ namespace uba
 		g_systemRootMem.Create();
 		g_systemTempMem.Create();
 
-		u64 reserveSizeMb = IsWindows ? 192 : 1024; // The sync primitives on linux/macos is much bigger
+		u64 reserveSizeMb = IsWindows ? 256 : 1024; // The sync primitives on linux/macos is much bigger
 		g_memoryBlockMem.Create(reserveSizeMb * 1024 * 1024);
 		g_directoryTableMem.Create(&g_memoryBlock);
 		g_mappedFileTableMem.Create(g_memoryBlock);
