[UBA]
* Fixed stack overflow bug when registering files/virtualfiles/directories from the outside
* Added some validation code in debug to be able to replace some code with simpler code once it is validated
* Fixed so "known files" work with virtual files etc
* Removed virtual files from saved file table

diff --git a/Engine/Source/Programs/UnrealBuildAccelerator/Common/Private/UbaSession.cpp b/Engine/Source/Programs/UnrealBuildAccelerator/Common/Private/UbaSession.cpp
--- a/Engine/Source/Programs/UnrealBuildAccelerator/Common/Private/UbaSession.cpp
+++ b/Engine/Source/Programs/UnrealBuildAccelerator/Common/Private/UbaSession.cpp
@@ -1364,7 +1364,7 @@ namespace uba
 			GetCurrentDirectoryW(workingDir);
 			workingDir.EnsureEndsWithSlash();
 		}
-		FixPath2(filePath, workingDir.data, workingDir.count, fixedFilePath.data, fixedFilePath.capacity, &fixedFilePath.count);
+		FixPath(filePath, workingDir.data, workingDir.count, fixedFilePath);
 
 		StringKey dirKey;
 		StringBuffer<> dirNameForHash;
@@ -1377,11 +1377,21 @@ namespace uba
 		StringKeyHasher hasher;
 		hasher.Update(dirNameForHash.data, dirNameForHash.count);
 
-		StringBuffer<128> baseFileNameForHash;
+		StringBuffer<> baseFileNameForHash;
 		baseFileNameForHash.Append(baseFileName);
 		if (CaseInsensitiveFs)
 			baseFileNameForHash.MakeLower();
-		return ToStringKey(hasher, baseFileNameForHash.data, baseFileNameForHash.count);
+		StringKey result = ToStringKey(hasher, baseFileNameForHash.data, baseFileNameForHash.count);
+
+#if UBA_DEBUG
+		StringBuffer<> testPath(fixedFilePath);
+		if (CaseInsensitiveFs)
+			testPath.MakeLower();
+		StringKey testKey = ToStringKeyLower(testPath);
+		UBA_ASSERTF(testKey == result, TC("Key mismatch for %s"), fixedFilePath.data);
+#endif
+
+		return result;
 	}
 
 	bool Session::RegisterNewFile(const tchar* filePath)
@@ -2520,13 +2530,13 @@ namespace uba
 			{
 				FileHandle fileHandle = uba::CreateFileW(sourceFile, GENERIC_READ, FILE_SHARE_READ, OPEN_EXISTING, DefaultAttributes());
 				if (fileHandle == InvalidFileHandle)
-					return m_logger.Error(TC("[InternalRegisterVirtualFile] CreateFileW for %s failed (%s)"), sourceFile, LastErrorToText().data);
+					return m_logger.Error(TC("[RegisterVirtualFileInternal] CreateFileW for %s failed (%s)"), sourceFile, LastErrorToText().data);
 				auto fg = MakeGuard([&]() { uba::CloseFile(sourceFile, fileHandle); });
 				insres.first->second = CreateFileMappingW(m_logger, fileHandle, PAGE_READONLY, 0, sourceFile);
 			}
 			mappingHandle = insres.first->second;
 			if (!mappingHandle.IsValid())
-				return m_logger.Error(TC("[InternalRegisterVirtualFile] CreateFileMapping for %s failed (%s)"), sourceFile, LastErrorToText().data);
+				return m_logger.Error(TC("[RegisterVirtualFileInternal] CreateFileMapping for %s failed (%s)"), sourceFile, LastErrorToText().data);
 		}
 
 
diff --git a/Engine/Source/Programs/UnrealBuildAccelerator/Common/Private/UbaSessionServer.cpp b/Engine/Source/Programs/UnrealBuildAccelerator/Common/Private/UbaSessionServer.cpp
--- a/Engine/Source/Programs/UnrealBuildAccelerator/Common/Private/UbaSessionServer.cpp
+++ b/Engine/Source/Programs/UnrealBuildAccelerator/Common/Private/UbaSessionServer.cpp
@@ -291,12 +291,11 @@ namespace uba
 			{
 				StringBuffer<> fileName;
 				FixPath(kiIt, workingDir.c_str(), u32(workingDir.size()), fileName);
+				StringKey fileNameKey = CaseInsensitiveFs ? ToStringKeyLower(fileName) : ToStringKey(fileName);
 
 				// Make sure cas entry exists and caskey is calculated (cas content creation is deferred in case client already has it)
 				CasKey casKey;
-				bool deferCreation = true;
-				bool fileIsCompressed = false;
-				if (!m_storage.StoreCasFile(casKey, fileName.data, CasKeyZero, deferCreation, fileIsCompressed) || casKey == CasKeyZero)
+				if (!StoreCasFile(casKey, fileNameKey, fileName.data) || casKey == CasKeyZero)
 					continue;
 
 				auto& ki = keys[keysIndex++];
@@ -306,9 +305,6 @@ namespace uba
 
 
 				// Update name to hash table
-				if (CaseInsensitiveFs)
-					fileName.MakeLower();
-				StringKey fileNameKey = ToStringKey(fileName);
 				SCOPED_WRITE_LOCK(m_nameToHashLookupLock, lock);
 				CasKey& lookupCasKey = m_nameToHashLookup[fileNameKey];
 				if (lookupCasKey != casKey)
diff --git a/Engine/Source/Programs/UnrealBuildAccelerator/Common/Private/UbaStorage.cpp b/Engine/Source/Programs/UnrealBuildAccelerator/Common/Private/UbaStorage.cpp
--- a/Engine/Source/Programs/UnrealBuildAccelerator/Common/Private/UbaStorage.cpp
+++ b/Engine/Source/Programs/UnrealBuildAccelerator/Common/Private/UbaStorage.cpp
@@ -154,6 +154,7 @@ namespace uba
 				ce.lock.LeaveWrite();
 				break;
 			}
+			UBA_ASSERT(ce.key != CasKeyZero);
 			toDelete.push_back({*it, ce.size});
 
 			m_casEvictedBytes += ce.size;
@@ -1562,7 +1563,7 @@ namespace uba
 			{
 				FileEntry& fileEntry = pair.second;
 				SCOPED_READ_LOCK(fileEntry.lock, entryLock);
-				if (fileEntry.casKey == CasKeyZero)
+				if (fileEntry.casKey == CasKeyZero || fileEntry.lastWritten == 0)
 					continue;
 				writer.WriteStringKey(pair.first);
 				writer.WriteU64(fileEntry.size);
diff --git a/Engine/Source/Programs/UnrealBuildAccelerator/Common/Public/UbaStorage.h b/Engine/Source/Programs/UnrealBuildAccelerator/Common/Public/UbaStorage.h
--- a/Engine/Source/Programs/UnrealBuildAccelerator/Common/Public/UbaStorage.h
+++ b/Engine/Source/Programs/UnrealBuildAccelerator/Common/Public/UbaStorage.h
@@ -195,7 +195,7 @@ namespace uba
 		struct FileEntry
 		{
 			ReaderWriterLock lock;
-			CasKey casKey = CasKeyZero;
+			CasKey casKey;
 			u64 size = 0;
 			u64 lastWritten = 0;
 			bool verified = false;
