[UBA]
* Fixed merged module bug where mapping handle was closed where it shouldn't
* Changed so casKey always need to be initialized in ctor for CasEntry

diff --git a/Engine/Source/Programs/UnrealBuildAccelerator/Common/Private/UbaObjectFile.cpp b/Engine/Source/Programs/UnrealBuildAccelerator/Common/Private/UbaObjectFile.cpp
--- a/Engine/Source/Programs/UnrealBuildAccelerator/Common/Private/UbaObjectFile.cpp
+++ b/Engine/Source/Programs/UnrealBuildAccelerator/Common/Private/UbaObjectFile.cpp
@@ -256,11 +256,14 @@ namespace uba
 		FileAccessor symFile(logger, filename);
 		if (!symFile.OpenMemoryRead())
 			return false;
+		if (symFile.GetSize() == 0)
+			return logger.Error(TC("%s - Import/export file corrupt (size 0)"), filename);
+
 		auto readPos = (const char*)symFile.GetData();
 
 		u8 version = *(u8*)readPos++;
 		if (SymbolFileVersion != version)
-			return logger.Error(TC("%s - Import/export file version mismatch"), filename);
+			return logger.Error(TC("%s - Import/export file version mismatch (application version %u, file version %u)"), filename, SymbolFileVersion, version);
 
 		type = *(const ObjectFileType*)readPos++;
 
diff --git a/Engine/Source/Programs/UnrealBuildAccelerator/Common/Private/UbaSession.cpp b/Engine/Source/Programs/UnrealBuildAccelerator/Common/Private/UbaSession.cpp
--- a/Engine/Source/Programs/UnrealBuildAccelerator/Common/Private/UbaSession.cpp
+++ b/Engine/Source/Programs/UnrealBuildAccelerator/Common/Private/UbaSession.cpp
@@ -2971,6 +2971,8 @@ namespace uba
 		if (!RegisterCreateFileForWrite(symFileKey, exportsFile, false, memoryBlock.writtenSize, lastWriteTime))
 			return false;
 
+		mg.Cancel();
+
 		auto insres = msg.process.m_writtenFiles.try_emplace(exportsFile.data);
 		WrittenFile& writtenFile = insres.first->second;
 
diff --git a/Engine/Source/Programs/UnrealBuildAccelerator/Common/Private/UbaStorage.cpp b/Engine/Source/Programs/UnrealBuildAccelerator/Common/Private/UbaStorage.cpp
--- a/Engine/Source/Programs/UnrealBuildAccelerator/Common/Private/UbaStorage.cpp
+++ b/Engine/Source/Programs/UnrealBuildAccelerator/Common/Private/UbaStorage.cpp
@@ -878,7 +878,7 @@ namespace uba
 	{
 		UBA_ASSERTF(IsCompressed(casKey) == m_storeCompressed, TC("CasKey compress mode must match storage compress mode (%s)"), fileName);
 		SCOPED_WRITE_LOCK(m_casLookupLock, lookupLock);
-		auto insres = m_casLookup.try_emplace(casKey);
+		auto insres = m_casLookup.try_emplace(casKey, casKey);
 		CasEntry& casEntry = insres.first->second;
 		lookupLock.Leave();
 		CasEntryAccessed(casEntry);
@@ -889,7 +889,6 @@ namespace uba
 			if (casEntry.exists)
 				return true;
 
-		casEntry.key = casKey;
 		casEntry.disallowed = IsDisallowedPath(fileName);
 
 #if !UBA_USE_SPARSEFILE
@@ -1058,7 +1057,7 @@ namespace uba
 
 
 				SCOPED_WRITE_LOCK(m_casLookupLock, lookupLock);
-				auto insres = m_casLookup.try_emplace(casKey);
+				auto insres = m_casLookup.try_emplace(casKey, casKey);
 				CasEntry& casEntry = insres.first->second;
 				casEntry.verified = true;
 				casEntry.exists = true;
@@ -1066,7 +1065,6 @@ namespace uba
 				m_casTotalBytes += size;
 				if (insres.second)
 				{
-					casEntry.key = casKey;
 					casEntry.size = size;
 					AttachEntry(casEntry);
 				}
@@ -1416,7 +1414,7 @@ namespace uba
 			CasKey casKey = reader.ReadCasKey();
 			if (casKey == CasKeyZero)
 				break;
-			auto insres = m_casLookup.try_emplace(casKey);
+			auto insres = m_casLookup.try_emplace(casKey, casKey);
 			if (!insres.second)
 			{
 				m_logger.Warning(TC("CasTable file %s is corrupt, it contains same cas key multiple times (%s)"), fileName.data, CasKeyString(casKey).str);
@@ -1426,7 +1424,6 @@ namespace uba
 				return false;
 			}
 			CasEntry& casEntry = insres.first->second;
-			casEntry.key = casKey;
 			casEntry.size = reader.ReadU64();
 			casEntry.exists = true;
 			m_casTotalBytes += casEntry.size;
@@ -2177,7 +2174,7 @@ namespace uba
 	bool StorageImpl::EnsureCasFile(const CasKey& casKey, const tchar* fileName)
 	{
 		SCOPED_WRITE_LOCK(m_casLookupLock, lookupLock);
-		auto insres = m_casLookup.try_emplace(casKey);
+		auto insres = m_casLookup.try_emplace(casKey, casKey);
 		CasEntry& casEntry = insres.first->second;
 		lookupLock.Leave();
 		CasEntryAccessed(casEntry);
@@ -2191,7 +2188,6 @@ namespace uba
 			if (!fileName)
 				return false;
 		}
-		casEntry.key = casKey;
 
 		StringBuffer<> casFile;
 #if !UBA_USE_SPARSEFILE
diff --git a/Engine/Source/Programs/UnrealBuildAccelerator/Common/Private/UbaStorageServer.cpp b/Engine/Source/Programs/UnrealBuildAccelerator/Common/Private/UbaStorageServer.cpp
--- a/Engine/Source/Programs/UnrealBuildAccelerator/Common/Private/UbaStorageServer.cpp
+++ b/Engine/Source/Programs/UnrealBuildAccelerator/Common/Private/UbaStorageServer.cpp
@@ -799,7 +799,7 @@ namespace uba
 				CasKey casKey = reader.ReadCasKey();
 				//UBA_ASSERT(IsCompressed(casKey));
 				SCOPED_WRITE_LOCK(m_casLookupLock, lookupLock);
-				auto casInsres = m_casLookup.try_emplace(casKey);
+				auto casInsres = m_casLookup.try_emplace(casKey, casKey);
 				CasEntry& casEntry = casInsres.first->second;
 				lookupLock.Leave();
 
@@ -857,14 +857,13 @@ namespace uba
 				reader.ReadString(hint);
 
 				SCOPED_WRITE_LOCK(m_casLookupLock, lookupLock);
-				auto casInsres = m_casLookup.try_emplace(casKey);
+				auto casInsres = m_casLookup.try_emplace(casKey, casKey);
 				CasEntry& casEntry = casInsres.first->second;
 				lookupLock.Leave();
 
 				SCOPED_WRITE_LOCK(casEntry.lock, entryLock);
 				if (!casEntry.verified)
 				{
-					casEntry.key = casKey;
 					casEntry.verified = true;
 				}
 				else
diff --git a/Engine/Source/Programs/UnrealBuildAccelerator/Common/Public/UbaStorage.h b/Engine/Source/Programs/UnrealBuildAccelerator/Common/Public/UbaStorage.h
--- a/Engine/Source/Programs/UnrealBuildAccelerator/Common/Public/UbaStorage.h
+++ b/Engine/Source/Programs/UnrealBuildAccelerator/Common/Public/UbaStorage.h
@@ -205,6 +205,7 @@ namespace uba
 		ReaderWriterLock m_casLookupLock;
 		struct CasEntry
 		{
+			CasEntry(const CasKey& k) : key(k) {}
 			ReaderWriterLock lock;
 			CasKey key;
 			CasEntry* prevAccessed = nullptr;
diff --git a/Engine/Source/Programs/UnrealBuildAccelerator/Common/Public/UbaStorageClient.h b/Engine/Source/Programs/UnrealBuildAccelerator/Common/Public/UbaStorageClient.h
--- a/Engine/Source/Programs/UnrealBuildAccelerator/Common/Public/UbaStorageClient.h
+++ b/Engine/Source/Programs/UnrealBuildAccelerator/Common/Public/UbaStorageClient.h
@@ -73,7 +73,7 @@ namespace uba
 		TString m_zone;
 
 		ReaderWriterLock m_localStorageFilesLock;
-		struct LocalFile { CasEntry casEntry; TString fileName; };
+		struct LocalFile { CasEntry casEntry; TString fileName; LocalFile() : casEntry(CasKeyZero) {} };
 		UnorderedMap<CasKey, LocalFile> m_localStorageFiles;
 
 		ReaderWriterLock m_sendOneAtTheTimeLock;
