[UBA]
* Added so UbaObjTool swallows /OPT command line

diff --git a/Engine/Source/Programs/UnrealBuildAccelerator/ObjTool/Private/UbaObjTool.cpp b/Engine/Source/Programs/UnrealBuildAccelerator/ObjTool/Private/UbaObjTool.cpp
--- a/Engine/Source/Programs/UnrealBuildAccelerator/ObjTool/Private/UbaObjTool.cpp
+++ b/Engine/Source/Programs/UnrealBuildAccelerator/ObjTool/Private/UbaObjTool.cpp
@@ -154,6 +154,9 @@ namespace uba
 					else if (name.Equals(TC("/LTCG")))
 					{
 					}
+					else if (name.StartsWith(TC("/OPT:")))
+					{
+					}
 					else
 					{
 						objFilesForImpLib.push_back(name.data);
