[UBA]
* Removed invalid assert

diff --git a/Engine/Source/Programs/UnrealBuildAccelerator/Common/Private/UbaStorageClient.cpp b/Engine/Source/Programs/UnrealBuildAccelerator/Common/Private/UbaStorageClient.cpp
--- a/Engine/Source/Programs/UnrealBuildAccelerator/Common/Private/UbaStorageClient.cpp
+++ b/Engine/Source/Programs/UnrealBuildAccelerator/Common/Private/UbaStorageClient.cpp
@@ -600,7 +600,6 @@ namespace uba
 
 								if (extraBuffer)
 								{
-									UBA_ASSERT(responseSize > capacity);
 									memcpy(readPosition, extraBuffer, left);
 									memmove(extraBuffer, extraBuffer + left, u32(responseSize - left));
 									if (isFirstInBlock)
