[UBA]
* Disabled compression of obj files that belong to pch. This is a bit of a hack to work around live coding support.

diff --git a/Engine/Source/Programs/UnrealBuildAccelerator/Core/Private/UbaApplicationRules.cpp b/Engine/Source/Programs/UnrealBuildAccelerator/Core/Private/UbaApplicationRules.cpp
--- a/Engine/Source/Programs/UnrealBuildAccelerator/Core/Private/UbaApplicationRules.cpp
+++ b/Engine/Source/Programs/UnrealBuildAccelerator/Core/Private/UbaApplicationRules.cpp
@@ -89,7 +89,9 @@ namespace uba
 
 		virtual bool StoreFileCompressed(const StringView& fileName) const
 		{
-			return fileName.EndsWith(TC(".obj"));
+			// We don't compress pch obj files because they are used by live-coding.
+			// TODO: Add support in live coding module to load compressed obj files instead
+			return fileName.EndsWith(TC(".obj")) && !fileName.Contains(TC("pch."));
 		}
 
 		virtual bool ShouldExtractSymbols(const StringView& fileName) const
