[UBA]
* Changed how external file mappings work for storage. It is now handled by a registered provider that can return file mappings for external files if they exist
* Added so session is external file provider and can provide all kinds of file mappings.. this will reduce number of iops and hopefully fix the problem with virtual files

diff --git a/Engine/Source/Programs/UnrealBuildAccelerator/Common/Private/UbaSession.cpp b/Engine/Source/Programs/UnrealBuildAccelerator/Common/Private/UbaSession.cpp
--- a/Engine/Source/Programs/UnrealBuildAccelerator/Common/Private/UbaSession.cpp
+++ b/Engine/Source/Programs/UnrealBuildAccelerator/Common/Private/UbaSession.cpp
@@ -524,6 +524,7 @@ namespace uba
 		auto _ = MakeGuard([&](){ uba::CloseFile(fileName, fileHandle); });
 
 		u64 size = 0;
+		u64 lastWriteTime = 0;
 
 		if (!isDir)
 		{
@@ -539,6 +540,7 @@ namespace uba
 					return false;
 
 				size = info.size;
+				lastWriteTime = info.lastWriteTime;
 			}
 		}
 
@@ -587,6 +589,7 @@ namespace uba
 
 		entry.isDir = isDir;
 		entry.size = size;
+		entry.lastWriteTime = lastWriteTime;
 		
 		out.size = size;
 		return true;
@@ -1196,6 +1199,24 @@ namespace uba
 		}
 		#endif
 
+		m_storage.RegisterExternalFileMappingsProvider([this](Storage::ExternalFileMapping& out, StringKey fileNameKey, const tchar* fileName)
+			{
+				SCOPED_READ_LOCK(m_fileMappingTableLookupLock, lookupLock);
+				auto findIt = m_fileMappingTableLookup.find(fileNameKey);
+				if (findIt == m_fileMappingTableLookup.end())
+					return false;
+				FileMappingEntry& entry = findIt->second;
+				lookupLock.Leave();
+				SCOPED_READ_LOCK(entry.lock, entryLock);
+				if (!entry.handled || !entry.success || !entry.mapping.IsValid())
+					return false;
+				out.handle = entry.mapping;
+				out.offset = entry.mappingOffset;
+				out.size = entry.size;
+				out.lastWriteTime = entry.lastWriteTime;
+				return true;
+			});
+
 		return true;
 	}
 
@@ -2349,6 +2370,7 @@ namespace uba
 				entry.mapping = file.mappingHandle;
 				entry.mappingOffset = 0;
 				entry.size = file.mappingWritten;
+				entry.lastWriteTime = file.lastWriteTime;
 				entry.isDir = false;
 				entry.success = true;
 				entry.owned = true;
@@ -2554,6 +2576,7 @@ namespace uba
 		entry.mappingOffset = sourceOffset;
 		entry.size = sourceSize;
 		entry.handled = true;
+		entry.lastWriteTime = 0; // TODO: Take lastwritetime of source file?
 
 		StringBuffer<> mappingName;
 		Storage::GetMappingString(mappingName, mappingHandle, sourceOffset);
diff --git a/Engine/Source/Programs/UnrealBuildAccelerator/Common/Private/UbaSessionClient.cpp b/Engine/Source/Programs/UnrealBuildAccelerator/Common/Private/UbaSessionClient.cpp
--- a/Engine/Source/Programs/UnrealBuildAccelerator/Common/Private/UbaSessionClient.cpp
+++ b/Engine/Source/Programs/UnrealBuildAccelerator/Common/Private/UbaSessionClient.cpp
@@ -600,8 +600,7 @@ namespace uba
 		CasKey casKey;
 		{
 			TimerScope ts(m_stats.storageSend);
-			bool deferCreation = false;
-			if (!m_storage.StoreCasFile(casKey, source.key, source.name.c_str(), source.mappingHandle, 0, source.mappingWritten, destination, deferCreation, keepMappingInMemory, compressed))
+			if (!m_storage.StoreCasFileClient(casKey, source.key, source.name.c_str(), source.mappingHandle, 0, source.mappingWritten, destination, keepMappingInMemory, compressed))
 				return false;
 		}
 		UBA_ASSERTF(casKey != CasKeyZero, TC("Failed to store cas file for %s (destination %s)"), source.name.c_str(), destination);
diff --git a/Engine/Source/Programs/UnrealBuildAccelerator/Common/Private/UbaSessionServer.cpp b/Engine/Source/Programs/UnrealBuildAccelerator/Common/Private/UbaSessionServer.cpp
--- a/Engine/Source/Programs/UnrealBuildAccelerator/Common/Private/UbaSessionServer.cpp
+++ b/Engine/Source/Programs/UnrealBuildAccelerator/Common/Private/UbaSessionServer.cpp
@@ -1650,28 +1650,10 @@ namespace uba
 			}
 		}
 
-		if (casKeyOverride != CasKeyInvalid) // if casKeyOverride is invalid it means it is a virtual file.. look in RegisterVirtualFileInternal
-		{
-			bool fileIsCompressed = false;
-			if (!m_storage.StoreCasFile(out, fileName, casKeyOverride, deferCreation, fileIsCompressed)) // We can defer the creation of the cas file since client might already have it
-				return false;
-
-			if (out != CasKeyZero)
-				return true;
-
-			if (m_shouldWriteToDisk)
-				return true;
-		}
-
-		SCOPED_READ_LOCK(m_fileMappingTableLookupLock, lookupLock);
-		auto findIt = m_fileMappingTableLookup.find(fileNameKey);
-		if (findIt == m_fileMappingTableLookup.end())
-			return true;
-
-		FileMappingEntry& entry = findIt->second;
-		lookupLock.Leave();
-		SCOPED_READ_LOCK(entry.lock, entryCs);
-		return m_storage.StoreCasFile(out, fileNameKey, fileName, entry.mapping, entry.mappingOffset, entry.size, fileName, deferCreation, true);
+		bool fileIsCompressed = false;
+		if (!m_storage.StoreCasFile(out, fileName, casKeyOverride, deferCreation, fileIsCompressed)) // We can defer the creation of the cas file since client might already have it
+			return false;
+		return true;//out != CasKeyZero;
 	}
 
 	bool SessionServer::WriteDirectoryTable(ClientSession& session, BinaryReader& reader, BinaryWriter& writer)
@@ -1940,19 +1922,6 @@ namespace uba
 		return str.Appendf(TC("%s"), findIt->second.GetStartInfo().GetDescription()).data;
 	}
 
-	bool SessionServer::RegisterVirtualFileInternal(const StringKey& fileNameKey, const StringView& filePath, const tchar* sourceFile, u64 sourceOffset, u64 sourceSize)
-	{
-		if (!Session::RegisterVirtualFileInternal(fileNameKey, filePath, sourceFile, sourceOffset, sourceSize))
-			return false;
-
-		// Abusing custom cas keys to be able to tell if a file exists only as virtual file (so we can handle StoreCasFile properly)
-		SCOPED_WRITE_LOCK(m_customCasKeysLock, lock);
-		auto insres = m_customCasKeys.try_emplace(fileNameKey);
-		CustomCasKey& customKey = insres.first->second;
-		customKey.casKey = CasKeyInvalid;
-		return true;
-	}
-
 	bool SessionServer::PrepareProcess(ProcessStartInfoHolder& startInfo, bool isChild, StringBufferBase& outRealApplication, const tchar*& outRealWorkingDir)
 	{
 		if (!Session::PrepareProcess(startInfo, isChild, outRealApplication, outRealWorkingDir))
diff --git a/Engine/Source/Programs/UnrealBuildAccelerator/Common/Private/UbaStorage.cpp b/Engine/Source/Programs/UnrealBuildAccelerator/Common/Private/UbaStorage.cpp
--- a/Engine/Source/Programs/UnrealBuildAccelerator/Common/Private/UbaStorage.cpp
+++ b/Engine/Source/Programs/UnrealBuildAccelerator/Common/Private/UbaStorage.cpp
@@ -83,6 +83,39 @@ namespace uba
 		table.GetValueAsString(rootDir, TC("RootDir"));
 	}
 
+	struct FileMappingScope
+	{
+		u8* MapView(const tchar* func, FileMappingHandle mappingHandle, u64 mappingOffset, u64 fileSize, const tchar* fileName_)
+		{
+			fileName = fileName_;
+			constexpr u64 pageSize = 64*1024;
+			u64 endOffset = mappingOffset + fileSize;
+			u64 alignedOffsetStart = AlignUp(mappingOffset - (pageSize - 1), pageSize);
+			u64 alignedOffsetEnd = AlignUp(endOffset, pageSize);
+			mapSize = alignedOffsetEnd - alignedOffsetStart;
+			mappedMem = MapViewOfFile(logger, mappingHandle, FILE_MAP_READ, alignedOffsetStart, mapSize);
+			if (!mappedMem)
+				mappedMem = MapViewOfFile(logger, mappingHandle, FILE_MAP_READ, alignedOffsetStart, 0);
+			if (mappedMem)
+				return mappedMem + (mappingOffset - alignedOffsetStart);
+			logger.Error(TC("[%s] Failed to map view of file %s at offset %llu and size %llu (%s)"), func, fileName, alignedOffsetStart, mapSize, LastErrorToText().data);
+			return nullptr;
+		}
+
+		FileMappingScope(Logger& l) : logger(l) {}
+		~FileMappingScope()
+		{
+			if (mappedMem)
+				if (!UnmapViewOfFile(logger, mappedMem, mapSize, fileName))
+					logger.Error(TC("Failed to unmap memory %s at 0x%llx with size %llu (%s)"), fileName, mappedMem, mapSize, LastErrorToText().data);
+		}
+
+		Logger& logger;
+		const tchar* fileName;
+		u8* mappedMem = nullptr;
+		u64 mapSize = 0;
+	};
+
 	void StorageImpl::CasEntryAccessed(CasEntry& casEntry)
 	{
 		bool hasMapping;
@@ -211,6 +244,8 @@ namespace uba
 			//UBA_ASSERT(false);
 #endif
 
+			//m_logger.Warning(TC("Evicted %s due to overflow"), casFile.data);
+
 			casEntry.verified = true; // Verified to be deleted
 
 			rec.casEntry.lock.LeaveWrite();
@@ -253,11 +288,29 @@ namespace uba
 		casEntry.nextAccessed = nullptr;
 	}
 
-	bool StorageImpl::WriteCompressed(WriteResult& out, const tchar* from, const tchar* to)
+	void StorageImpl::RegisterExternalFileMappingsProvider(ExternalFileMappingsProvider&& provider)
+	{
+		m_externalFileMappingsProvider = std::move(provider);
+	}
+
+	bool StorageImpl::WriteCompressed(WriteResult& out, const StringKey& fromKey, const tchar* from, const tchar* to)
 	{
 		StorageStats& stats = Stats();
 		TimerScope ts(stats.createCas);
 
+		if (m_externalFileMappingsProvider)
+		{
+			ExternalFileMapping externalMapping;
+			if (m_externalFileMappingsProvider(externalMapping, fromKey, from))
+			{
+				FileMappingScope scope(m_logger);
+				u8* fileMem = scope.MapView(TC("WriteCompressed"), externalMapping.handle, externalMapping.offset, externalMapping.size, from);
+				if (!fileMem)
+					return false;
+				return StorageImpl::WriteCompressed(out, from, InvalidFileHandle, fileMem, externalMapping.size, to, nullptr, 0);
+			}
+		}
+
 		FileHandle readHandle;
 		if (!OpenFileSequentialRead(m_logger, from, readHandle))
 			return m_logger.Error(TC("[WriteCompressed] Failed to open file %s for read (%s)"), from, LastErrorToText().data);
@@ -535,11 +588,11 @@ namespace uba
 		return true;
 	}
 
-	bool StorageImpl::WriteCasFileNoCheck(WriteResult& out, const tchar* fileName, bool fileIsCompressed, const tchar* casFile, bool storeCompressed)
+	bool StorageImpl::WriteCasFileNoCheck(WriteResult& out, const StringKey& fileNameKey, const tchar* fileName, bool fileIsCompressed, const tchar* casFile, bool storeCompressed)
 	{
 		if (storeCompressed && !fileIsCompressed)
 		{
-			if (!WriteCompressed(out, fileName, casFile))
+			if (!WriteCompressed(out, fileNameKey, fileName, casFile))
 				return false;
 		}
 		else if (fileIsCompressed && !storeCompressed)
@@ -606,7 +659,12 @@ namespace uba
 		//UBA_ASSERT(false);
 #endif
 
-		return WriteCasFileNoCheck(out, fileName, fileIsCompressed, casFile.data, IsCompressed(casKey));
+		StringBuffer<> temp;
+		temp.Append(fileName);
+		if (CaseInsensitiveFs)
+			temp.MakeLower();
+		StringKey fileNameKey = ToStringKey(temp);
+		return WriteCasFileNoCheck(out, fileNameKey, fileName, fileIsCompressed, casFile.data, IsCompressed(casKey));
 	}
 
 	void StorageImpl::CasEntryAccessed(const CasKey& casKey)
@@ -921,7 +979,7 @@ namespace uba
 
 		WriteResult res;
 #if !UBA_USE_SPARSEFILE
-		if (!WriteCasFileNoCheck(res, fileName, fileIsCompressed, casFile.data, IsCompressed(casKey)))
+		if (!WriteCasFileNoCheck(res, fileNameKey, fileName, fileIsCompressed, casFile.data, IsCompressed(casKey)))
 			return false;
 #else
 
@@ -1942,25 +2000,39 @@ namespace uba
 		}
 		fileEntry.verified = true;
 
-		// Use OpenFile+GetFileInformationByHandle+close instead of Getting file attributes because it is actually faster on cloud setups (weirdly enough)
-		FileHandle fileHandle;
-		if (!OpenFileSequentialRead(m_logger, fileName, fileHandle))
-		{
-			fileEntry.casKey = CasKeyZero;
-			out = CasKeyZero;
-			return true;
-		}
+
+		ExternalFileMapping externalMapping;
+
+		FileHandle fileHandle = InvalidFileHandle;
 		auto fileGuard = MakeGuard([&](){ CloseFile(fileName, fileHandle); });
 
-		FileInformation info;
-		if (!GetFileInformationByHandle(info, m_logger, fileName, fileHandle))
+		u64 fileSize = 0;
+		u64 lastWriteTime = 0;
+
+		if (m_externalFileMappingsProvider && m_externalFileMappingsProvider(externalMapping, fileNameKey, fileName))
 		{
-			fileEntry.casKey = CasKeyZero;
-			return m_logger.Error(TC("GetFileInformationByHandle failed on %s"), fileName);
+			fileSize = externalMapping.size;
+			lastWriteTime = externalMapping.lastWriteTime;
 		}
+		else
+		{
+			// Use OpenFile+GetFileInformationByHandle+close instead of Getting file attributes because it is actually faster on cloud setups (weirdly enough)
+			if (!OpenFileSequentialRead(m_logger, fileName, fileHandle))
+			{
+				fileEntry.casKey = CasKeyZero;
+				out = CasKeyZero;
+				return true;
+			}
 
-		u64 fileSize = info.size;
-		u64 lastWritten = info.lastWriteTime;
+			FileInformation info;
+			if (!GetFileInformationByHandle(info, m_logger, fileName, fileHandle))
+			{
+				fileEntry.casKey = CasKeyZero;
+				return m_logger.Error(TC("GetFileInformationByHandle failed on %s"), fileName);
+			}
+			fileSize = info.size;
+			lastWriteTime = info.lastWriteTime;
+		}
 
 		if (fileEntry.casKey != CasKeyZero)
 		{
@@ -1972,7 +2044,7 @@ namespace uba
 				out = fileEntry.casKey;
 				return true;
 			}
-			if (fileSize == fileEntry.size && lastWritten == fileEntry.lastWritten && fileEntry.casKey != CasKeyInvalid)
+			if (fileSize == fileEntry.size && (lastWriteTime && lastWriteTime == fileEntry.lastWritten) && fileEntry.casKey != CasKeyInvalid)
 			{
 				if (!AddCasFile(fileNameKey, fileName, fileEntry.casKey, deferCreation, fileIsCompressed))
 					return false;
@@ -1982,7 +2054,7 @@ namespace uba
 		}
 
 		fileEntry.size = fileSize;
-		fileEntry.lastWritten = lastWritten;
+		fileEntry.lastWritten = lastWriteTime;
 		if (casKeyOverride == CasKeyZero)
 		{
 			if (fileIsCompressed)
@@ -1992,6 +2064,14 @@ namespace uba
 					return m_logger.Error(TC("Failed to read header of compressed file %s (%s)"), fileName, LastErrorToText().data);
 				fileEntry.casKey = AsCompressed(header.casKey, m_storeCompressed);
 			}
+			else if (externalMapping.handle.IsValid())
+			{
+				FileMappingScope scope(m_logger);
+				u8* fileMem = scope.MapView(TC("StoreCasFile"), externalMapping.handle, externalMapping.offset, externalMapping.size, fileName);
+				if (!fileMem)
+					return false;
+				fileEntry.casKey = CalculateCasKey(fileMem, externalMapping.size, m_storeCompressed);
+			}
 			else
 			{
 				fileEntry.casKey = CalculateCasKey(fileName, fileHandle, fileSize, m_storeCompressed);
@@ -2120,7 +2200,7 @@ namespace uba
 		return true;
 	}
 
-	bool StorageImpl::StoreCasFile(CasKey& out, StringKey fileNameKey, const tchar* fileName, FileMappingHandle mappingHandle, u64 mappingOffset, u64 fileSize, const tchar* hint, bool deferCreation, bool keepMappingInMemory, bool storeCompressed)
+	bool StorageImpl::StoreCasFileClient(CasKey& out, StringKey fileNameKey, const tchar* fileName, FileMappingHandle mappingHandle, u64 mappingOffset, u64 fileSize, const tchar* hint, bool keepMappingInMemory, bool storeCompressed)
 	{
 		UBA_ASSERT(false);
 		return false;
diff --git a/Engine/Source/Programs/UnrealBuildAccelerator/Common/Private/UbaStorageClient.cpp b/Engine/Source/Programs/UnrealBuildAccelerator/Common/Private/UbaStorageClient.cpp
--- a/Engine/Source/Programs/UnrealBuildAccelerator/Common/Private/UbaStorageClient.cpp
+++ b/Engine/Source/Programs/UnrealBuildAccelerator/Common/Private/UbaStorageClient.cpp
@@ -792,7 +792,7 @@ namespace uba
 		return StorageImpl::HasCasFile(casKey, out);
 	}
 
-	bool StorageClient::StoreCasFile(CasKey& out, StringKey fileNameKey, const tchar* fileName, FileMappingHandle mappingHandle, u64 mappingOffset, u64 fileSize, const tchar* hint, bool deferCreation, bool keepMappingInMemory, bool storeCompressed)
+	bool StorageClient::StoreCasFileClient(CasKey& out, StringKey fileNameKey, const tchar* fileName, FileMappingHandle mappingHandle, u64 mappingOffset, u64 fileSize, const tchar* hint, bool keepMappingInMemory, bool storeCompressed)
 	{
 		NetworkClient& client = m_client; // Don't use proxy
 
diff --git a/Engine/Source/Programs/UnrealBuildAccelerator/Common/Private/UbaStorageServer.cpp b/Engine/Source/Programs/UnrealBuildAccelerator/Common/Private/UbaStorageServer.cpp
--- a/Engine/Source/Programs/UnrealBuildAccelerator/Common/Private/UbaStorageServer.cpp
+++ b/Engine/Source/Programs/UnrealBuildAccelerator/Common/Private/UbaStorageServer.cpp
@@ -223,93 +223,6 @@ namespace uba
 		}
 	}
 
-	struct FileMappingScope
-	{
-		u8* MapView(const tchar* func, FileMappingHandle mappingHandle, u64 mappingOffset, u64 fileSize, const tchar* fileName_)
-		{
-			fileName = fileName_;
-			constexpr u64 pageSize = 64*1024;
-			u64 endOffset = mappingOffset + fileSize;
-			u64 alignedOffsetStart = AlignUp(mappingOffset - (pageSize - 1), pageSize);
-			u64 alignedOffsetEnd = AlignUp(endOffset, pageSize);
-			mapSize = alignedOffsetEnd - alignedOffsetStart;
-			mappedMem = MapViewOfFile(logger, mappingHandle, FILE_MAP_READ, alignedOffsetStart, mapSize);
-			if (!mappedMem)
-				mappedMem = MapViewOfFile(logger, mappingHandle, FILE_MAP_READ, alignedOffsetStart, 0);
-			if (mappedMem)
-				return mappedMem + (mappingOffset - alignedOffsetStart);
-			logger.Error(TC("[%s] Failed to map view of file %s at offset %llu and size %llu (%s)"), func, fileName, alignedOffsetStart, mapSize, LastErrorToText().data);
-			return nullptr;
-		}
-
-		FileMappingScope(Logger& l) : logger(l) {}
-		~FileMappingScope()
-		{
-			if (mappedMem)
-				if (!UnmapViewOfFile(logger, mappedMem, mapSize, fileName))
-					logger.Error(TC("Failed to unmap memory %s at 0x%llx with size %llu (%s)"), fileName, mappedMem, mapSize, LastErrorToText().data);
-		}
-
-		Logger& logger;
-		const tchar* fileName;
-		u8* mappedMem = nullptr;
-		u64 mapSize = 0;
-	};
-
-	bool StorageServer::StoreCasFile(CasKey& out, StringKey fileNameKey, const tchar* fileName, FileMappingHandle mappingHandle, u64 mappingOffset, u64 fileSize, const tchar* hint, bool deferCreation, bool keepMappingInMemory, bool storeCompressed)
-	{
-		FileMappingScope scope(m_logger);
-		u8* fileMem = scope.MapView(TC("StoreCasFile"), mappingHandle, mappingOffset, fileSize, fileName);
-		if (!fileMem)
-			return false;
-
-		out = CalculateCasKey(fileMem, fileSize, storeCompressed);
-		if (out == CasKeyZero)
-			return false;
-
-		FileEntry& fileEntry = GetOrCreateFileEntry(fileNameKey);
-		SCOPED_WRITE_LOCK(fileEntry.lock, entryLock);
-		fileEntry.verified = true;
-		fileEntry.casKey = out;
-		fileEntry.size = fileSize;
-
-		SCOPED_WRITE_LOCK(m_externalFileMappingsLock, externalFileLock);
-		m_externalFileMappings.try_emplace(fileNameKey, ExternalFileMapping{mappingHandle, mappingOffset, fileSize});
-		externalFileLock.Leave();
-
-		bool fileIsCompressed = false;
-		if (!AddCasFile(fileNameKey, fileName, fileEntry.casKey, deferCreation, fileIsCompressed))
-			return false;
-
-		return true;
-	}
-
-	bool StorageServer::WriteCompressed(WriteResult& out, const tchar* from, const tchar* toFile)
-	{
-		StringBuffer<> fromForKey;
-		fromForKey.Append(from);
-		if (CaseInsensitiveFs)
-			fromForKey.MakeLower();
-		StringKey fileNameKey = ToStringKey(fromForKey);
-
-		SCOPED_WRITE_LOCK(m_externalFileMappingsLock, lock);
-		auto findIt = m_externalFileMappings.find(fileNameKey);
-		if (findIt == m_externalFileMappings.end())
-		{
-			lock.Leave();
-			return StorageImpl::WriteCompressed(out, from, toFile);
-		}
-
-		ExternalFileMapping& mapping = findIt->second;
-		lock.Leave();
-
-		FileMappingScope scope(m_logger);
-		u8* fileMem = scope.MapView(TC("WriteCompressed"), mapping.mappingHandle, mapping.mappingOffset, mapping.fileSize, from);
-		if (!fileMem)
-			return false;
-		return StorageImpl::WriteCompressed(out, from, InvalidFileHandle, fileMem, mapping.fileSize, toFile, nullptr, 0);
-	}
-
 	bool StorageServer::IsDisallowedPath(const tchar* fileName)
 	{
 		for (auto& path : m_disallowedPaths)
diff --git a/Engine/Source/Programs/UnrealBuildAccelerator/Common/Public/UbaSession.h b/Engine/Source/Programs/UnrealBuildAccelerator/Common/Public/UbaSession.h
--- a/Engine/Source/Programs/UnrealBuildAccelerator/Common/Public/UbaSession.h
+++ b/Engine/Source/Programs/UnrealBuildAccelerator/Common/Public/UbaSession.h
@@ -212,7 +212,7 @@ namespace uba
 		u8* m_fileMappingTableMem;
 		u32 m_fileMappingTableSize = 0;
 		ReaderWriterLock m_fileMappingTableLookupLock;
-		struct FileMappingEntry { ReaderWriterLock lock; FileMappingHandle mapping; u64 mappingOffset = 0; u64 size = 0; bool isDir = false; bool handled = false; bool success = false; bool owned = false; };
+		struct FileMappingEntry { ReaderWriterLock lock; FileMappingHandle mapping; u64 mappingOffset = 0; u64 size = 0; u64 lastWriteTime = 0; bool isDir = false; bool handled = false; bool success = false; bool owned = false; };
 		UnorderedMap<StringKey, FileMappingEntry> m_fileMappingTableLookup;
 
 		ReaderWriterLock m_nameToHashTableHandleLock;
diff --git a/Engine/Source/Programs/UnrealBuildAccelerator/Common/Public/UbaSessionServer.h b/Engine/Source/Programs/UnrealBuildAccelerator/Common/Public/UbaSessionServer.h
--- a/Engine/Source/Programs/UnrealBuildAccelerator/Common/Public/UbaSessionServer.h
+++ b/Engine/Source/Programs/UnrealBuildAccelerator/Common/Public/UbaSessionServer.h
@@ -98,8 +98,6 @@ namespace uba
 
 		TString GetProcessDescription(u32 processId);
 
-		virtual bool RegisterVirtualFileInternal(const StringKey& fileNameKey, const StringView& filePath, const tchar* sourceFile, u64 sourceOffset, u64 sourceSize) override;
-
 		virtual bool PrepareProcess(ProcessStartInfoHolder& startInfo, bool isChild, StringBufferBase& outRealApplication, const tchar*& outRealWorkingDir) override final;
 		virtual bool CreateFile(CreateFileResponse& out, const CreateFileMessage& msg) override final;
 		virtual void FileEntryAdded(StringKey fileNameKey, u64 lastWritten, u64 size) override final;
diff --git a/Engine/Source/Programs/UnrealBuildAccelerator/Common/Public/UbaStorage.h b/Engine/Source/Programs/UnrealBuildAccelerator/Common/Public/UbaStorage.h
--- a/Engine/Source/Programs/UnrealBuildAccelerator/Common/Public/UbaStorage.h
+++ b/Engine/Source/Programs/UnrealBuildAccelerator/Common/Public/UbaStorage.h
@@ -45,7 +45,7 @@ namespace uba
 		virtual bool InvalidateCachedFileInfo(StringKey fileNameKey) = 0;
 
 		virtual bool StoreCasFile(CasKey& out, const tchar* fileName, const CasKey& casKeyOverride, bool deferCreation, bool fileIsCompressed) = 0;
-		virtual bool StoreCasFile(CasKey& out, StringKey fileNameKey, const tchar* fileName, FileMappingHandle mappingHandle, u64 mappingOffset, u64 fileSize, const tchar* hint, bool deferCreation = false, bool keepMappingInMemory = false, bool storeCompressed = true) = 0;
+		virtual bool StoreCasFileClient(CasKey& out, StringKey fileNameKey, const tchar* fileName, FileMappingHandle mappingHandle, u64 mappingOffset, u64 fileSize, const tchar* hint, bool keepMappingInMemory = false, bool storeCompressed = true) = 0;
 		virtual bool DropCasFile(const CasKey& casKey, bool forceDelete, const tchar* hint) = 0;
 		virtual bool ReportBadCasFile(const CasKey& casKey) = 0;
 		virtual bool CalculateCasKey(CasKey& out, const tchar* fileName) = 0;
@@ -69,6 +69,10 @@ namespace uba
 
 		struct WriteResult { FileMappingHandle mappingHandle; u64 size = InvalidValue; u64 offset = InvalidValue; };
 		virtual bool WriteCompressed(WriteResult& out, const tchar* from, FileHandle readHandle, u8* readMem, u64 fileSize, const tchar* toFile, const void* header, u64 headerSize, u64 lastWriteTime = 0) = 0;
+
+		struct ExternalFileMapping { FileMappingHandle handle; u64 offset = 0; u64 size = 0; u64 lastWriteTime = 0; };
+		using ExternalFileMappingsProvider = Function<bool(ExternalFileMapping& out, StringKey fileNameKey, const tchar* fileName)>;
+		virtual void RegisterExternalFileMappingsProvider(ExternalFileMappingsProvider&& provider) {}
 	};
 
 	struct StorageCreateInfo
@@ -127,7 +131,7 @@ namespace uba
 		virtual bool VerifyAndGetCachedFileInfo(CachedFileInfo& out, StringKey fileNameKey, u64 verifiedLastWriteTime, u64 verifiedSize) override;
 		virtual bool InvalidateCachedFileInfo(StringKey fileNameKey) override;
 		virtual bool StoreCasFile(CasKey& out, const tchar* fileName, const CasKey& casKeyOverride, bool deferCreation, bool fileIsCompressed) override;
-		virtual bool StoreCasFile(CasKey& out, StringKey fileNameKey, const tchar* fileName, FileMappingHandle mappingHandle, u64 mappingOffset, u64 fileSize, const tchar* hint, bool deferCreation = false, bool keepMappingInMemory = false, bool storeCompressed = true) override;
+		virtual bool StoreCasFileClient(CasKey& out, StringKey fileNameKey, const tchar* fileName, FileMappingHandle mappingHandle, u64 mappingOffset, u64 fileSize, const tchar* hint, bool keepMappingInMemory = false, bool storeCompressed = true) override;
 		virtual bool DropCasFile(const CasKey& casKey, bool forceDelete, const tchar* hint) override;
 		virtual bool ReportBadCasFile(const CasKey& casKey) override;
 		virtual bool CalculateCasKey(CasKey& out, const tchar* fileName) override;
@@ -155,10 +159,11 @@ namespace uba
 		bool IsFileVerified(const StringKey& fileNameKey);
 		void ReportFileInfoWeak(const StringKey& fileNameKey, u64 verifiedLastWriteTime, u64 verifiedSize);
 
-		virtual bool WriteCompressed(WriteResult& out, const tchar* from, const tchar* toFile);
+		virtual void RegisterExternalFileMappingsProvider(ExternalFileMappingsProvider&& provider) override;
+		bool WriteCompressed(WriteResult& out, const StringKey& fromKey, const tchar* from, const tchar* toFile);
 		virtual bool WriteCompressed(WriteResult& out, const tchar* from, FileHandle readHandle, u8* readMem, u64 fileSize, const tchar* toFile, const void* header, u64 headerSize, u64 lastWriteTime = 0) final;
 		bool WriteMemToCompressedFile(FileAccessor& destination, u32 workCount, const u8* uncompressedData, u64 fileSize, u64 maxUncompressedBlock, u64& totalWritten);
-		bool WriteCasFileNoCheck(WriteResult& out, const tchar* fileName, bool fileIsCompressed, const tchar* casFile, bool storeCompressed);
+		bool WriteCasFileNoCheck(WriteResult& out, const StringKey& fileNameKey, const tchar* fileName, bool fileIsCompressed, const tchar* casFile, bool storeCompressed);
 		bool WriteCasFile(WriteResult& out, const tchar* fileName, bool fileIsCompressed, const CasKey& casKey);
 		bool VerifyExisting(bool& outReturnValue, ScopedWriteLock& entryLock, const CasKey& casKey, CasEntry& casEntry, StringBufferBase& casFile);
 		bool AddCasFile(StringKey fileNameKey, const tchar* fileName, const CasKey& casKey, bool deferCreation, bool fileIsCompressed);
@@ -237,6 +242,8 @@ namespace uba
 		bool m_storeCompressed = false;
 		bool m_manuallyHandleOverflow = false;
 
+		ExternalFileMappingsProvider m_externalFileMappingsProvider;
+
 		u32 m_maxParallelCopyOrLink;
 		ReaderWriterLock m_activeCopyOrLinkLock;
 		Event m_activeCopyOrLinkEvent;
diff --git a/Engine/Source/Programs/UnrealBuildAccelerator/Common/Public/UbaStorageClient.h b/Engine/Source/Programs/UnrealBuildAccelerator/Common/Public/UbaStorageClient.h
--- a/Engine/Source/Programs/UnrealBuildAccelerator/Common/Public/UbaStorageClient.h
+++ b/Engine/Source/Programs/UnrealBuildAccelerator/Common/Public/UbaStorageClient.h
@@ -55,7 +55,7 @@ namespace uba
 		virtual bool GetZone(StringBufferBase& out) override;
 		virtual bool RetrieveCasFile(RetrieveResult& out, const CasKey& casKey, const tchar* hint, FileMappingBuffer* mappingBuffer = nullptr, u64 memoryMapAlignment = 1, bool allowProxy = true) override;
 		virtual bool StoreCasFile(CasKey& out, const tchar* fileName, const CasKey& casKeyOverride, bool deferCreation, bool fileIsCompressed) override;
-		virtual bool StoreCasFile(CasKey& out, StringKey fileNameKey, const tchar* fileName, FileMappingHandle mappingFile, u64 mappingOffset, u64 fileSize, const tchar* hint, bool deferCreation = false, bool keepMappingInMemory = false, bool storeCompressed = true) override;
+		virtual bool StoreCasFileClient(CasKey& out, StringKey fileNameKey, const tchar* fileName, FileMappingHandle mappingFile, u64 mappingOffset, u64 fileSize, const tchar* hint, bool keepMappingInMemory = false, bool storeCompressed = true) override;
 		virtual bool HasCasFile(const CasKey& casKey, CasEntry** out = nullptr) override;
 		virtual void Ping() override;
 		virtual void PrintSummary(Logger& logger) override;
diff --git a/Engine/Source/Programs/UnrealBuildAccelerator/Common/Public/UbaStorageServer.h b/Engine/Source/Programs/UnrealBuildAccelerator/Common/Public/UbaStorageServer.h
--- a/Engine/Source/Programs/UnrealBuildAccelerator/Common/Public/UbaStorageServer.h
+++ b/Engine/Source/Programs/UnrealBuildAccelerator/Common/Public/UbaStorageServer.h
@@ -38,8 +38,6 @@ namespace uba
 
 		virtual bool GetZone(StringBufferBase& out) override;
 		virtual bool RetrieveCasFile(RetrieveResult& out, const CasKey& casKey, const tchar* hint, FileMappingBuffer* mappingBuffer = nullptr, u64 memoryMapAlignment = 1, bool allowProxy = true) override;
-		virtual bool StoreCasFile(CasKey& out, StringKey fileNameKey, const tchar* fileName, FileMappingHandle mappingHandle, u64 mappingOffset, u64 fileSize, const tchar* hint, bool deferCreation = false, bool keepMappingInMemory = false, bool storeCompressed = true) override;
-		virtual bool WriteCompressed(WriteResult& out, const tchar* from, const tchar* toFile) override;
 		virtual bool IsDisallowedPath(const tchar* fileName) override;
 		virtual void SetTrace(Trace* trace, bool detailed) override;
 		virtual bool HasProxy(u32 clientId) override;
@@ -111,17 +109,6 @@ namespace uba
 		Vector<u16> m_availableIds;
 		u16 m_availableIdsHigh = 1;
 
-
-		struct ExternalFileMapping
-		{
-			FileMappingHandle mappingHandle;
-			u64 mappingOffset;
-			u64 fileSize;
-		};
-		ReaderWriterLock m_externalFileMappingsLock;
-		UnorderedMap<StringKey, ExternalFileMapping> m_externalFileMappings;
-
-
 		struct ProxyEntry
 		{
 			u32 clientId = ~0u;
