[UBA]
* Fixed test compile error

diff --git a/Engine/Source/Programs/UnrealBuildAccelerator/Test/Private/UbaTestStorage.h b/Engine/Source/Programs/UnrealBuildAccelerator/Test/Private/UbaTestStorage.h
--- a/Engine/Source/Programs/UnrealBuildAccelerator/Test/Private/UbaTestStorage.h
+++ b/Engine/Source/Programs/UnrealBuildAccelerator/Test/Private/UbaTestStorage.h
@@ -147,7 +147,7 @@ namespace uba
 
 		CasKey key;
 		bool storeCompressed = false;
-		if (!storageClient.StoreCasFile(key, ToStringKeyLower(fileName), fileName.data, FileMappingHandle{}, 0, 0, TC("UbaTestFile"), false, false, storeCompressed))
+		if (!storageClient.StoreCasFileClient(key, ToStringKeyLower(fileName), fileName.data, FileMappingHandle{}, 0, 0, TC("UbaTestFile"), false, storeCompressed))
 			return logger.Error(TC("Failed to store file %s"), fileName.data);
 
 		fileName.Clear().Append(testRootDir).Append(TC("Uba")).Append(PathSeparator).Append(TC("UbaTestFile"));
