[UBA]
* Fixed bug where commandlines longer than 64k crashed uba

diff --git a/Engine/Source/Programs/UnrealBuildAccelerator/Common/Private/UbaLogger.cpp b/Engine/Source/Programs/UnrealBuildAccelerator/Common/Private/UbaLogger.cpp
--- a/Engine/Source/Programs/UnrealBuildAccelerator/Common/Private/UbaLogger.cpp
+++ b/Engine/Source/Programs/UnrealBuildAccelerator/Common/Private/UbaLogger.cpp
@@ -142,19 +142,25 @@ namespace uba
 		constexpr size_t _TRUNCATE = (size_t)-1;
 		#endif
 
-		tchar buffer[1024];
+		tchar buffer[2048];
 		int len = _vsnwprintf_s(buffer, sizeof_array(buffer), _TRUNCATE, format, args);
 		if (len != -1)
 		{
 			Log(type, buffer, u32(len));
+			return;
 		}
-		else
+
+		len = _vsnwprintf_s(buffer, sizeof_array(buffer), 0, format, args);
+		if (len == -1)
 		{
-			Vector<tchar> buf;
-			buf.resize(64 * 1024);
-			_vsnwprintf_s(buf.data(), buf.size(), buf.size(), format, args);
-			Log(type, buf.data(), u32(buf.size()));
+			Error(TC("LogArg failed. bad format"));
+			return;
 		}
+
+		Vector<tchar> buf;
+		buf.resize(len + 1);
+		_vsnwprintf_s(buf.data(), buf.size(), buf.size(), format, args);
+		Log(type, buf.data(), u32(buf.size()));
 	}
 
 	LoggerWithWriter::LoggerWithWriter(LogWriter& writer, const tchar* prefix)
diff --git a/Engine/Source/Programs/UnrealBuildAccelerator/Common/Private/UbaProcess.cpp b/Engine/Source/Programs/UnrealBuildAccelerator/Common/Private/UbaProcess.cpp
--- a/Engine/Source/Programs/UnrealBuildAccelerator/Common/Private/UbaProcess.cpp
+++ b/Engine/Source/Programs/UnrealBuildAccelerator/Common/Private/UbaProcess.cpp
@@ -946,8 +946,8 @@ namespace uba
 		#endif
 
 		TString applicationStr = reader.ReadString();
-		StringBuffer<64*1024*2> fullCommandLine;
-		reader.ReadString(fullCommandLine);
+		TString fullCommandLine = ReadLongString(reader);
+
 		StringBuffer<> currentDir;
 		reader.ReadString(currentDir);
 		if (currentDir.IsEmpty())
@@ -960,13 +960,13 @@ namespace uba
 		// Remove application from command line
 		if (fullCommandLine[0] == '"')
 		{
-			const tchar* second = fullCommandLine.First('"', 1);
+			const tchar* second = TStrchr(fullCommandLine.data() + 1, '"');
 			if (!second)
-				second = fullCommandLine.data + fullCommandLine.count;
+				second = fullCommandLine.data() + fullCommandLine.size();
 			//UBA_ASSERTF(second, TC("Missing second '\"' in command line: %s"), fullCommandLine.data); // "Unsupported cmd line format"
 			commandLine = second + 1;
 			if (applicationStr.empty())
-				applicationStr.assign(fullCommandLine.data + 1, u64(second - fullCommandLine.data - 1));
+				applicationStr.assign(fullCommandLine.data() + 1, u64(second - fullCommandLine.data() - 1));
 		}
 		else
 		{
@@ -974,17 +974,17 @@ namespace uba
 			if (applicationStr.empty() && IsWindows)
 			{
 				const tchar* extPos = nullptr;
-				fullCommandLine.Contains(TC(".exe"), true, &extPos);
+				Contains(fullCommandLine.data(),TC(".exe"), true, &extPos);
 				UBA_ASSERT(extPos);
-				searchOffset = u32(extPos - fullCommandLine.data);
+				searchOffset = u32(extPos - fullCommandLine.data());
 			}
-			const tchar* secondParamStart = fullCommandLine.First(' ', searchOffset);
+			const tchar* secondParamStart = TStrchr(fullCommandLine.data() + searchOffset, ' ');
 			if (!secondParamStart)
 				commandLine = TC("");
 			else
 				commandLine = secondParamStart + 1;
 			if (applicationStr.empty())
-				applicationStr.assign(fullCommandLine.data, u64(secondParamStart - fullCommandLine.data));
+				applicationStr.assign(fullCommandLine.data(), u64(secondParamStart - fullCommandLine.data()));
 		}
 
 		while (*commandLine == ' ')
diff --git a/Engine/Source/Programs/UnrealBuildAccelerator/Common/Public/UbaNetwork.h b/Engine/Source/Programs/UnrealBuildAccelerator/Common/Public/UbaNetwork.h
--- a/Engine/Source/Programs/UnrealBuildAccelerator/Common/Public/UbaNetwork.h
+++ b/Engine/Source/Programs/UnrealBuildAccelerator/Common/Public/UbaNetwork.h
@@ -16,7 +16,7 @@ namespace uba
 	
 	static constexpr u32 SystemNetworkVersion = 1339;
 	static constexpr u32 StorageNetworkVersion = 4;
-	static constexpr u32 SessionNetworkVersion = 37;
+	static constexpr u32 SessionNetworkVersion = 38;
 	static constexpr u32 CacheNetworkVersion = 5;
 
 	// Messages used over network between client/server (system, storage and session)
diff --git a/Engine/Source/Programs/UnrealBuildAccelerator/Core/Public/UbaProcessUtils.h b/Engine/Source/Programs/UnrealBuildAccelerator/Core/Public/UbaProcessUtils.h
--- a/Engine/Source/Programs/UnrealBuildAccelerator/Core/Public/UbaProcessUtils.h
+++ b/Engine/Source/Programs/UnrealBuildAccelerator/Core/Public/UbaProcessUtils.h
@@ -2,8 +2,10 @@
 
 #pragma once
 
+#include "UbaBinaryReaderWriter.h"
 #include "UbaPlatform.h"
 #include "UbaStringBuffer.h"
+#include <oodle2.h>
 
 namespace uba
 {
@@ -103,4 +105,41 @@ namespace uba
 	{
 		return ParseArguments(arguments, TStrlen(arguments), argumentFunc);
 	}
+
+	inline void WriteLongString(BinaryWriter& writer, const StringView& str)
+	{
+		if (str.count < 10000)
+		{
+			writer.WriteBool(true);
+			writer.WriteString(str);
+			return;
+		}
+			
+		writer.WriteBool(false);
+		u32 uncompressedSize = str.count * sizeof(tchar);
+		u64 needed = (u64)OodleLZ_GetCompressedBufferSizeNeeded(OodleLZ_Compressor_Kraken, uncompressedSize);
+		Vector<u8> mem;
+		mem.resize(needed);
+		u64 compressedSize = OodleLZ_Compress(OodleLZ_Compressor_Kraken, str.data, (OO_SINTa)uncompressedSize, mem.data(), OodleLZ_CompressionLevel_Normal);
+		UBA_ASSERT(compressedSize != OODLELZ_FAILED);
+		writer.Write7BitEncoded(str.count);
+		writer.Write7BitEncoded(compressedSize);
+		writer.WriteBytes(mem.data(), compressedSize);
+	}
+
+	inline TString ReadLongString(BinaryReader& reader)
+	{
+		if (reader.ReadBool())
+			return reader.ReadString();
+
+		u64 stringLength = reader.Read7BitEncoded();
+		u64 uncompressedSize = stringLength * sizeof(tchar);
+		u64 compressedSize = reader.Read7BitEncoded();
+		TString str;
+		str.resize(stringLength);
+		OO_SINTa decompLen = OodleLZ_Decompress(reader.GetPositionData(), (OO_SINTa)compressedSize, str.data(), (OO_SINTa)uncompressedSize);
+		UBA_ASSERT(decompLen == (OO_SINTa)uncompressedSize);(void)decompLen;
+		reader.Skip(compressedSize);
+		return str;
+	}
 }
diff --git a/Engine/Source/Programs/UnrealBuildAccelerator/Core/Public/UbaStringBuffer.h b/Engine/Source/Programs/UnrealBuildAccelerator/Core/Public/UbaStringBuffer.h
--- a/Engine/Source/Programs/UnrealBuildAccelerator/Core/Public/UbaStringBuffer.h
+++ b/Engine/Source/Programs/UnrealBuildAccelerator/Core/Public/UbaStringBuffer.h
@@ -92,6 +92,7 @@ namespace uba
 	struct StringView
 	{
 		StringView() : data(TC("")), count(0) {}
+		StringView(const StringView&) = default;
 		StringView(const tchar* d, u32 c) : data(d), count(c) {}
 		StringView(const StringBufferBase& sb) : data(sb.data), count(sb.count) {}
 		StringView(const TString& str) : data(str.data()), count(u32(str.size())) {}
diff --git a/Engine/Source/Programs/UnrealBuildAccelerator/Core/UbaCore.Build.cs b/Engine/Source/Programs/UnrealBuildAccelerator/Core/UbaCore.Build.cs
--- a/Engine/Source/Programs/UnrealBuildAccelerator/Core/UbaCore.Build.cs
+++ b/Engine/Source/Programs/UnrealBuildAccelerator/Core/UbaCore.Build.cs
@@ -29,6 +29,7 @@ public class UbaCore : ModuleRules
 		PublicIncludePathModuleNames.AddRange(new string[]
 		{
 			"UbaMimalloc",
+			"OodleDataCompression",
 		});
 
 		if (Target.Platform == UnrealTargetPlatform.Win64)
diff --git a/Engine/Source/Programs/UnrealBuildAccelerator/Detours/Private/PosixOS/UbaDetoursFunctionsPosix.cpp b/Engine/Source/Programs/UnrealBuildAccelerator/Detours/Private/PosixOS/UbaDetoursFunctionsPosix.cpp
--- a/Engine/Source/Programs/UnrealBuildAccelerator/Detours/Private/PosixOS/UbaDetoursFunctionsPosix.cpp
+++ b/Engine/Source/Programs/UnrealBuildAccelerator/Detours/Private/PosixOS/UbaDetoursFunctionsPosix.cpp
@@ -7,6 +7,7 @@
 #include "UbaDirectoryTable.h"
 #include "UbaPlatform.h"
 #include "UbaProcessStats.h"
+#include "UbaProcessUtils.h"
 #include "UbaProtocol.h"
 #include "UbaStringBuffer.h"
 #include "UbaTimer.h"
@@ -1923,7 +1924,7 @@ int shared_posix_spawn(pid_t* pid, const char* path, const posix_spawn_file_acti
 		BinaryWriter writer;
 		writer.WriteByte(MessageType_CreateProcess);
 		writer.WriteString(path); // Application
-		writer.WriteString(cmdLine); // Full command line
+		WriteLongString(writer, cmdLine); // Full command line
 		writer.WriteString(pwd); // Current dir
 		writer.WriteBool(false); // Start suspended
 		writer.WriteBool(true); // Is child
diff --git a/Engine/Source/Programs/UnrealBuildAccelerator/Detours/Private/Windows/UbaDetoursFunctionsKernelBase.inl b/Engine/Source/Programs/UnrealBuildAccelerator/Detours/Private/Windows/UbaDetoursFunctionsKernelBase.inl
--- a/Engine/Source/Programs/UnrealBuildAccelerator/Detours/Private/Windows/UbaDetoursFunctionsKernelBase.inl
+++ b/Engine/Source/Programs/UnrealBuildAccelerator/Detours/Private/Windows/UbaDetoursFunctionsKernelBase.inl
@@ -2665,16 +2665,16 @@ BOOL Detoured_GetConsoleTitleW(LPTSTR lpConsoleTitle, DWORD nSize)
 BOOL Detoured_CreateProcessW(LPCWSTR lpApplicationName, LPWSTR lpCommandLine, LPSECURITY_ATTRIBUTES lpProcessAttributes, LPSECURITY_ATTRIBUTES lpThreadAttributes, BOOL bInheritHandles,
 	DWORD dwCreationFlags, LPVOID lpEnvironment, LPCWSTR lpCurrentDirectory, LPSTARTUPINFOW lpStartupInfo, LPPROCESS_INFORMATION lpProcessInformation)
 {
-	const tchar* originalCmd = lpCommandLine ? lpCommandLine : TC("");
+	StringView originalCmd = ToView(lpCommandLine ? lpCommandLine : TC(""));
 
 	DETOURED_CALL(CreateProcessW);
-	DEBUG_LOG_DETOURED(L"CreateProcessW", L"%ls %ls CreationFlags: 0x%x StartupFlags: 0x%u Stdin: %llu WorkDir: %s", lpApplicationName, originalCmd, dwCreationFlags, lpStartupInfo->dwFlags, u64(lpStartupInfo->hStdInput), (lpCurrentDirectory ? lpCurrentDirectory : L""));
+	DEBUG_LOG_DETOURED(L"CreateProcessW", L"%ls %ls CreationFlags: 0x%x StartupFlags: 0x%u Stdin: %llu WorkDir: %s", lpApplicationName, originalCmd.data, dwCreationFlags, lpStartupInfo->dwFlags, u64(lpStartupInfo->hStdInput), (lpCurrentDirectory ? lpCurrentDirectory : L""));
 
-	const tchar* testName = lpApplicationName;
+	StringView testName = ToView(lpApplicationName ? lpApplicationName : TC(""));
 
-	if (!lpApplicationName || !*lpApplicationName)
+	if (!testName.count)
 	{
-		if (!*originalCmd)
+		if (!originalCmd.count)
 		{
 			SetLastError(ERROR_FILE_NOT_FOUND);
 			return FALSE;
@@ -2685,7 +2685,7 @@ BOOL Detoured_CreateProcessW(LPCWSTR lpApplicationName, LPWSTR lpCommandLine, LP
 
 
 	// Debug binaries started when process crash... we don't want to detour these.
-	if (Contains(testName, L"winedbg") || Contains(testName, L"werfault.exe") || Contains(testName, L"vsjitdebugger.exe") || Contains(testName, L"crashpad_handler.exe"))
+	if (testName.Contains(L"winedbg") || testName.Contains(L"werfault.exe") || testName.Contains(L"vsjitdebugger.exe") || testName.Contains(L"crashpad_handler.exe"))
 	{
 		if (g_runningRemote)
 		{
@@ -2703,7 +2703,7 @@ BOOL Detoured_CreateProcessW(LPCWSTR lpApplicationName, LPWSTR lpCommandLine, LP
 
 	bool isChild = true;
 	// We don't care about tracking mspdbsrv or vctip.. they are services just spawned by this process
-	if (Contains(testName, L"mspdbsrv.exe") || Contains(testName, L"vctip.exe") || Contains(testName, L"git.exe"))
+	if (testName.Contains(L"mspdbsrv.exe") || testName.Contains(L"vctip.exe") || testName.Contains(L"git.exe"))
 	{
 		if (!g_runningRemote)
 		{
@@ -2728,15 +2728,15 @@ BOOL Detoured_CreateProcessW(LPCWSTR lpApplicationName, LPWSTR lpCommandLine, LP
 		SCOPED_WRITE_LOCK(g_communicationLock, pcs);
 		BinaryWriter writer;
 		writer.WriteByte(MessageType_CreateProcess);
-		writer.WriteString(application.data);
-		writer.WriteString(originalCmd);
+		writer.WriteString(application);
+		WriteLongString(writer, originalCmd);
 		writer.WriteString(lpCurrentDirectory ? lpCurrentDirectory : g_virtualWorkingDir.data);
 		writer.WriteBool(startSuspended);
 		writer.WriteBool(isChild);
 		writer.Flush();
 		BinaryReader reader;
 		processId = reader.ReadU32();
-		UBA_ASSERTF(processId > 0, L"Failed to create process %s", originalCmd);
+		UBA_ASSERTF(processId > 0, L"Failed to create process %s", originalCmd.data);
 		if (!processId)
 			return FALSE;
 
@@ -2748,7 +2748,7 @@ BOOL Detoured_CreateProcessW(LPCWSTR lpApplicationName, LPWSTR lpCommandLine, LP
 
 		currentDir = reader.ReadString();
 		commandLine = reader.ReadString();
-		DEBUG_LOG_PIPE(L"CreateProcess", L"%ls %ls", application.data, originalCmd);
+		DEBUG_LOG_PIPE(L"CreateProcess", L"%ls %ls", application.data, originalCmd.data);
 	}
 
 	LPCSTR dlls[] = { dll };
diff --git a/Engine/Source/Programs/UnrealBuildAccelerator/Detours/UbaDetours.Build.cs b/Engine/Source/Programs/UnrealBuildAccelerator/Detours/UbaDetours.Build.cs
--- a/Engine/Source/Programs/UnrealBuildAccelerator/Detours/UbaDetours.Build.cs
+++ b/Engine/Source/Programs/UnrealBuildAccelerator/Detours/UbaDetours.Build.cs
@@ -13,14 +13,14 @@ public class UbaDetours : ModuleRules
 		bUseUnity = false;
 
 		PrivateDependencyModuleNames.AddRange(new string[] {
-			"UbaCore"
+			"UbaCore",
+			"OodleDataCompression",
 		});
 
 		if (Target.Platform.IsInGroup(UnrealPlatformGroup.Windows))
 		{
 			PrivateDependencyModuleNames.AddRange(new string[] {
 				"Detours",
-				"OodleDataCompression",
 			});
 
 			PublicSystemLibraries.AddRange(new string[] {
diff --git a/Engine/Source/Programs/UnrealBuildAccelerator/Visualizer/Private/UbaVisualizer.cpp b/Engine/Source/Programs/UnrealBuildAccelerator/Visualizer/Private/UbaVisualizer.cpp
--- a/Engine/Source/Programs/UnrealBuildAccelerator/Visualizer/Private/UbaVisualizer.cpp
+++ b/Engine/Source/Programs/UnrealBuildAccelerator/Visualizer/Private/UbaVisualizer.cpp
@@ -959,9 +959,15 @@ namespace uba
 			rect.right = rect.left + int(progress*float(width));//durationMs/100;
 			FillRect(hdc, &rect, m_traceView.progressErrorCount ? m_processBrushes[0].error : m_processBrushes[0].success);
 
-			const tchar* remoteDisabled = TC(""); // TODO: Only show this in "advanced mode" m_traceView.remoteExecutionDisabled ? TC("   (Remote spawn disabled)") : TC("");
 			StringBuffer<> str;
-			str.Appendf(L"%u%%    %u / %u %s", u32(progress*100.0f), m_traceView.progressProcessesDone, m_traceView.progressProcessesTotal, remoteDisabled);
+			str.Appendf(TC("%u%%    %u / %u"), u32(progress*100.0f), m_traceView.progressProcessesDone, m_traceView.progressProcessesTotal);
+
+			const tchar* remoteDisabled = TC(""); // TODO: Only show this in "advanced mode" m_traceView.remoteExecutionDisabled ? TC("   (Remote spawn disabled)") : TC("");
+			str.Append(remoteDisabled);
+
+			if (u32 active = m_traceView.totalProcessActiveCount)
+				str.Appendf(TC("   (%u active)"), active);
+
 			drawIndentedText(str, LogEntryType_Info, 6, true);
 		}
 
