[UBA]
* Fixed so obj file preloader logic preloads in same order as obj files are registered

diff --git a/Engine/Source/Programs/UnrealBuildAccelerator/Detours/Private/Windows/UbaDetoursObjFilesPreloader.h b/Engine/Source/Programs/UnrealBuildAccelerator/Detours/Private/Windows/UbaDetoursObjFilesPreloader.h
--- a/Engine/Source/Programs/UnrealBuildAccelerator/Detours/Private/Windows/UbaDetoursObjFilesPreloader.h
+++ b/Engine/Source/Programs/UnrealBuildAccelerator/Detours/Private/Windows/UbaDetoursObjFilesPreloader.h
@@ -119,12 +119,19 @@ namespace uba
 			StringKey fileNameKey = ToStringKey(fileFull.MakeLower());
 			auto insres = m_preloadedObjFiles.try_emplace(fileNameKey);
 			UBA_ASSERT(insres.second);
-			insres.first->second.handle = fileHandle;
-			insres.first->second.event.Create(true);
+			Preload& preload = insres.first->second;
+			preload.handle = fileHandle;
+			preload.event.Create(true);
+
+			if (!m_it)
+				m_it = &preload;
+			if (m_last)
+				m_last->next = &preload;
+			m_last = &preload;
 
 			SCOPED_READ_LOCK(g_mappedFileTable.m_lookupLock, _);
 			auto findIt = g_mappedFileTable.m_lookup.find(fileNameKey);
-			insres.first->second.fileInfo = &findIt->second;
+			preload.fileInfo = &findIt->second;
 		}
 
 		void Start(const wchar_t* cmdLine)
@@ -151,9 +158,9 @@ namespace uba
 		void ThreadPreload()
 		{
 			u64 totalMemSize = 0;
-			for (auto& kv : m_preloadedObjFiles)
+			for (Preload* preloadPtr=m_it; preloadPtr; preloadPtr = preloadPtr->next)
 			{
-				Preload& preload = kv.second;
+				Preload& preload = *preloadPtr;
 				HANDLE objFileHandle = preload.handle;
 				if (!objFileHandle)
 					continue;
@@ -204,7 +211,6 @@ namespace uba
 				}
 			}
 			m_totalMem = (u8*)VirtualAlloc(NULL, alignedMem, allocationType, PAGE_READWRITE);
-			m_it = m_preloadedObjFiles.begin();
 
 			for (u32 i=0;i!=sizeof_array(m_helperThreadHandles); ++i)
 				m_helperThreadHandles[i] = CreateThread(NULL, 0, [](LPVOID p) -> DWORD { ((ObjFilesPreloader*)p)->ThreadHelper(); return 0; }, this, 0, NULL);
@@ -220,12 +226,12 @@ namespace uba
 			while (true)
 			{
 				SCOPED_WRITE_LOCK(m_threadLock, lock);
-				if (m_it == m_preloadedObjFiles.end())
+				if (!m_it)
 					break;
-				Preload& preload = m_it->second;
+				Preload& preload = *m_it;
 				if (preload.objReadOffset == preload.objCompressedSize)
 				{
-					++m_it;
+					m_it = m_it->next;
 					continue;
 				}
 
@@ -256,7 +262,7 @@ namespace uba
 					OodleLZ_FuzzSafe_Yes, OodleLZ_CheckCRC_No, OodleLZ_Verbosity_None, NULL, 0, NULL, NULL, decoderMem, decoredMemSize);
 				if (decompLen != decompressedBlockSize)
 				{
-					FatalError(1356, TC("Failed to decompress .obj file %s (%s)"), info.name, info.originalName);
+					FatalError(1356, TC("Failed to decompress .obj file %s (Expected to get %llu but got %llu)"), info.originalName, u64(decompressedBlockSize), u64(decompLen));
 				}
 
 				bool isDone = preload.objLeft.fetch_sub(decompressedBlockSize) == decompressedBlockSize;
@@ -289,24 +295,29 @@ namespace uba
 			auto findIt = m_preloadedObjFiles.find(key);
 			if (findIt == m_preloadedObjFiles.end())
 				return;
+			Preload& preload = findIt->second;
 			TimerScope ts(g_stats.waitDecompress);
-			findIt->second.event.IsSet();
+			preload.event.IsSet();
 		}
 
 		struct Preload
 		{
+			Preload* next = nullptr;
 			Event event;
 			HANDLE handle;
-			u8* objMem;
-			u64 objCompressedSize;
-			u64 objDecompressedSize;
-			FileInfo* fileInfo;
-			u64 objReadOffset;
-			u64 objWriteOffset;
+			u8* objMem = nullptr;
+			u64 objCompressedSize = 0;
+			u64 objDecompressedSize = 0;
+			FileInfo* fileInfo = nullptr;
+			u64 objReadOffset = 0;
+			u64 objWriteOffset = 0;
 			Atomic<u64> objLeft;
 		};
 
-		using PreloadedObjFiles = Map<StringKey, Preload>;
+		Preload* m_it = nullptr;
+		Preload* m_last = nullptr;
+
+		using PreloadedObjFiles = UnorderedMap<StringKey, Preload>;
 		PreloadedObjFiles m_preloadedObjFiles;
 		HANDLE m_threadHandle = 0;
 		HANDLE m_helperThreadHandles[5];
@@ -314,6 +325,5 @@ namespace uba
 		u64 m_totalMemOffset = 0;
 
 		ReaderWriterLock m_threadLock;
-		PreloadedObjFiles::iterator m_it;
 	};
 }
