[UBA]
* Disabled compressed obj files for non-windows platforms since it is not implemented (yet)

diff --git a/Engine/Source/Programs/UnrealBuildAccelerator/Common/Private/UbaSession.cpp b/Engine/Source/Programs/UnrealBuildAccelerator/Common/Private/UbaSession.cpp
--- a/Engine/Source/Programs/UnrealBuildAccelerator/Common/Private/UbaSession.cpp
+++ b/Engine/Source/Programs/UnrealBuildAccelerator/Common/Private/UbaSession.cpp
@@ -1057,7 +1057,7 @@ namespace uba
 		m_shouldWriteToDisk = info.shouldWriteToDisk;
 		UBA_ASSERTF(m_shouldWriteToDisk || m_allowMemoryMaps, TC("Can't disable both should write to disk and allow memory maps"));
 
-		m_storeObjFilesCompressed = info.storeObjFilesCompressed;
+		m_storeObjFilesCompressed = info.storeObjFilesCompressed && IsWindows; // Non-windows not implemented (yet)
 		m_allowLocalDetour = info.allowLocalDetour;
 
 		m_detailedTrace = info.detailedTrace;
