[UBA]
* Fixed bug in logger when sprintf result was longer than buffer

diff --git a/Engine/Source/Programs/UnrealBuildAccelerator/Common/Private/UbaLogger.cpp b/Engine/Source/Programs/UnrealBuildAccelerator/Common/Private/UbaLogger.cpp
--- a/Engine/Source/Programs/UnrealBuildAccelerator/Common/Private/UbaLogger.cpp
+++ b/Engine/Source/Programs/UnrealBuildAccelerator/Common/Private/UbaLogger.cpp
@@ -13,6 +13,7 @@
 #else
 #define Fputs fputs
 #define _vsnwprintf_s(buffer,capacity,count,format,args) vsnprintf(buffer, capacity, format, args) // TODO: This will overflow
+#define _vscwprintf(format, args) vsnprintf(0, 0, format, args)
 #endif
 
 
@@ -150,7 +151,7 @@ namespace uba
 			return;
 		}
 
-		len = _vsnwprintf_s(buffer, sizeof_array(buffer), 0, format, args);
+		len = _vscwprintf(format, args);
 		if (len == -1)
 		{
 			Error(TC("LogArg failed. bad format"));
