[UBA]
* Added thread names for easier debugging
* Fixed so ubacli can run multiple commands in a row

diff --git a/Engine/Source/Programs/UnrealBuildAccelerator/Cli/Private/UbaCli.cpp b/Engine/Source/Programs/UnrealBuildAccelerator/Cli/Private/UbaCli.cpp
--- a/Engine/Source/Programs/UnrealBuildAccelerator/Cli/Private/UbaCli.cpp
+++ b/Engine/Source/Programs/UnrealBuildAccelerator/Cli/Private/UbaCli.cpp
@@ -184,11 +184,13 @@ namespace uba
 		bool printSummary = false;
 		bool populateCache = false;
 		bool writeCacheSummary = false;
+		bool logToFile = false;
 		TString checkFileTable;
 		TString cacheFilterString;
 		TString cacheCommand;
 		TString testCompress;
 		TString testDecompress;
+		TString addCas;
 
 		u32 loopCount = 1;
 
@@ -323,6 +325,10 @@ namespace uba
 						return PrintHelp(TC("Invalid value for -port"));
 				}
 			}
+			else if (name.Equals(TC("-log")))
+			{
+				logToFile = true;
+			}
 			else if (name.Equals(TC("-loop")))
 			{
 				if (!value.Parse(loopCount))
@@ -386,6 +392,10 @@ namespace uba
 			{
 				deleteCas = true;
 			}
+			else if (name.Equals(TC("-addcas")))
+			{
+				addCas = value.data;
+			}
 			else if (name.Equals(TC("-getcas")))
 			{
 				getCas = true;
@@ -444,6 +454,7 @@ namespace uba
 		FilteredLogWriter logWriter(g_consoleLogWriter, quiet ? LogEntryType_Warning : LogEntryType_Detail);
 		LoggerWithWriter logger(logWriter, TC(""));
 
+		bool exit = false;
 		if (deleteCas)
 		{
 			StorageImpl(StorageCreateInfo(g_rootDir.data, logWriter)).DeleteAllCas();
@@ -453,6 +464,19 @@ namespace uba
 				clientRootDir.Append(g_rootDir).Append("Agent").AppendValue(i);
 				StorageImpl(StorageCreateInfo(clientRootDir.data, logWriter)).DeleteAllCas();
 			}
+			exit = true;
+		}
+
+		if (!addCas.empty())
+		{
+			StorageCreateInfo storageInfo(g_rootDir.data, logWriter);
+			storageInfo.casCapacityBytes = 0;
+			storageInfo.storeCompressed = storeCompressed;
+			StorageImpl storage(storageInfo);
+			CasKey casKey;
+			if (!storage.StoreCasFile(casKey, addCas.c_str(), CasKeyZero, false, false))
+				return false;
+			exit = true;
 		}
 
 		if (checkCas)
@@ -461,7 +485,9 @@ namespace uba
 			storageInfo.casCapacityBytes = 0;
 			storageInfo.storeCompressed = storeCompressed;
 			StorageImpl storage(storageInfo);
-			return storage.CheckCasContent(DefaultProcessorCount);
+			if (!storage.CheckCasContent(DefaultProcessorCount))
+				return false;
+			exit = true;
 		}
 
 		if (!checkFileTable.empty())
@@ -472,7 +498,9 @@ namespace uba
 			StorageImpl storage(storageInfo);
 			if (!storage.LoadCasTable())
 				return false;
-			return storage.CheckFileTable(checkFileTable.data(), DefaultProcessorCount);
+			if (!storage.CheckFileTable(checkFileTable.data(), DefaultProcessorCount))
+				return false;
+			exit = true;
 		}
 
 		if (checkCas2) // Creates a storage server and storage client and transfer _all_ cas files over network
@@ -500,7 +528,7 @@ namespace uba
 			if (!client.Init(cii))
 				return false;
 			bool success = true;
-			WorkManagerImpl workManager(DefaultProcessorCount);
+			WorkManagerImpl workManager(DefaultProcessorCount, TC("UbaWorker(CheckCas2)"));
 			storageServer.TraverseAllCasFiles([&](const CasKey& casKey, u64 size)
 				{
 					workManager.AddWork([&, casKey]()
@@ -514,7 +542,9 @@ namespace uba
 						}, 1, TC(""));
 				});
 			workManager.FlushWork();
-			return success;
+			if (!success)
+				return false;
+			exit = true;
 		}
 
 #if UBA_USE_AWS
@@ -533,14 +563,14 @@ namespace uba
 			}
 			else
 				logger.Info(TC("Seems like we are not running inside aws."));
-			return true;
+			exit = true;
 		}
 #endif
 		
 		u64 testCompressOriginalSize = 0;
 		if (!testCompress.empty())
 		{
-			WorkManagerImpl workManager(DefaultProcessorCount);
+			WorkManagerImpl workManager(DefaultProcessorCount, TC("UbaWorker(TestCompress)"));
 
 			FileAccessor fa(logger, testCompress.c_str());
 			if (!fa.OpenMemoryRead())
@@ -565,11 +595,12 @@ namespace uba
 				return false;
 			if (testDecompress.empty())
 				return true;
+			exit = true;
 		}
 
 		if (!testDecompress.empty())
 		{
-			WorkManagerImpl workManager(DefaultProcessorCount);
+			WorkManagerImpl workManager(DefaultProcessorCount, TC("UbaWorker(TestDecompress)"));
 
 			FileAccessor fa(logger, testDecompress.c_str());
 			if (!fa.OpenMemoryRead())
@@ -617,9 +648,14 @@ namespace uba
 				reader.Skip(compressedBlockSize);
 			}
 
-			return faDest.Close();
+			if (!faDest.Close())
+				return false;
+			exit = true;
 		}
 
+		if (exit)
+			return true;
+
 		if (commandType == CommandType_NotSet)
 		{
 			const tchar* errorMsg = argc == 1 ? TC("") : TC("\nERROR: First argument must be command type. Options are 'local,remote or native'");
@@ -682,13 +718,18 @@ namespace uba
 
 		// TODO: Change workdir to make it full
 
-
-		StringBuffer<> logFile;
 		#if UBA_DEBUG
-		logFile.count = GetFullPathNameW(g_rootDir.data, logFile.capacity, logFile.data, nullptr);
-		logFile.EnsureEndsWithSlash().Append(TC("DebugLog.log"));
+		logToFile = true;
 		#endif
 
+		StringBuffer<> logFile;
+		if (logToFile)
+		{
+			logFile.count = GetFullPathNameW(g_rootDir.data, logFile.capacity, logFile.data, nullptr);
+			logFile.EnsureEndsWithSlash().Append(TC("DebugLog.log"));
+			logger.Info(TC("Logging to file: %s"), logFile.data);
+		}
+
 		#if PLATFORM_WINDOWS
 		SetConsoleCtrlHandler(ConsoleHandler, TRUE);
 		#else
@@ -795,7 +836,7 @@ namespace uba
 		}
 
 
-		if (isRemote)
+		if (true)//isRemote)
 		{
 			if (!storageServer.m_casTableLoaded)
 				if (!storageServer.LoadCasTable(true))
@@ -872,6 +913,8 @@ namespace uba
 			return true;
 		};
 
+		const tchar* clientZone = TC("DummyZone");
+
 		auto RunWithClient = [&](const Function<bool()>& func, u32 clientCount)
 			{
 				Vector<Client> clients;
@@ -879,7 +922,8 @@ namespace uba
 				u32 clientIndex = 0;
 				for (auto& c : clients)
 				{
-					ClientInitInfo cii { logWriter, networkBackend, g_rootDir.data, TC("127.0.0.1"), port, TC("DummyZone"), maxProcessCount/u32(clients.size()), clientIndex++};
+					u32 maxProcessor = Min(maxProcessCount/u32(clients.size()), 32u);
+					ClientInitInfo cii { logWriter, networkBackend, g_rootDir.data, TC("127.0.0.1"), port, clientZone, maxProcessor, clientIndex++};
 					if (!c.Init(cii))
 						return false;
 				}
@@ -983,7 +1027,10 @@ namespace uba
 				};
 
 			if (commandType == CommandType_Agent)
-				return RunWithClient([&]() { return RunQueue(); }, maxProcessCount == 1 ? 1 : 4);
+			{
+				u32 clientCount = maxProcessCount == 1 ? 1 : 4;
+				return RunWithClient([&]() { return RunQueue(); }, clientCount);
+			}
 			else
 				return RunQueue();
 		};
diff --git a/Engine/Source/Programs/UnrealBuildAccelerator/Common/Private/UbaNetworkBackendTcp.cpp b/Engine/Source/Programs/UnrealBuildAccelerator/Common/Private/UbaNetworkBackendTcp.cpp
--- a/Engine/Source/Programs/UnrealBuildAccelerator/Common/Private/UbaNetworkBackendTcp.cpp
+++ b/Engine/Source/Programs/UnrealBuildAccelerator/Common/Private/UbaNetworkBackendTcp.cpp
@@ -300,7 +300,7 @@ namespace uba
 				{
 					ThreadListen(logger, e);
 					return 0;
-				});
+				}, TC("UbaTcpListen"));
 		}
 
 		bool success = true;
@@ -435,7 +435,7 @@ namespace uba
 			SCOPED_WRITE_LOCK(m_connectionsLock, lock);
 			auto it = m_connections.emplace(m_connections.end(), logger, clientSocket);
 			auto& conn = *it;
-			conn.recvThread.Start([this, connPtr = &conn] { ThreadRecv(*connPtr); return 0; });
+			conn.recvThread.Start([this, connPtr = &conn] { ThreadRecv(*connPtr); return 0; }, TC("UbaTcpRecv"));
 			lock.Leave();
 
 			if (!entry.connectedFunc(&conn, remoteSockAddr))
@@ -685,7 +685,7 @@ namespace uba
 		SCOPED_WRITE_LOCK(m_connectionsLock, lock);
 		auto it = m_connections.emplace(m_connections.end(), logger, socketFd);
 		auto& conn = *it;
-		conn.recvThread.Start([this, connPtr = &conn] { ThreadRecv(*connPtr); return 0; });
+		conn.recvThread.Start([this, connPtr = &conn] { ThreadRecv(*connPtr); return 0; }, TC("UbaTcpRecv"));
 		lock.Leave();
 
 		if (!connectedFunc(&conn, remoteSocketAddr, timedOut))
diff --git a/Engine/Source/Programs/UnrealBuildAccelerator/Common/Private/UbaNetworkClient.cpp b/Engine/Source/Programs/UnrealBuildAccelerator/Common/Private/UbaNetworkClient.cpp
--- a/Engine/Source/Programs/UnrealBuildAccelerator/Common/Private/UbaNetworkClient.cpp
+++ b/Engine/Source/Programs/UnrealBuildAccelerator/Common/Private/UbaNetworkClient.cpp
@@ -11,7 +11,7 @@
 namespace uba
 {
 	NetworkClient::NetworkClient(bool& outCtorSuccess, const NetworkClientCreateInfo& info, const tchar* name)
-	: WorkManagerImpl(info.workerCount == 0 ? GetLogicalProcessorCount() : info.workerCount)
+	: WorkManagerImpl(info.workerCount == 0 ? GetLogicalProcessorCount() : info.workerCount, TC("UbaWorker(NetworkClient)"))
 	,	m_logWriter(info.logWriter)
 	,	m_logger(info.logWriter, SetGetPrefix(name))
 	,	m_isConnected(true)
diff --git a/Engine/Source/Programs/UnrealBuildAccelerator/Common/Private/UbaNetworkServer.cpp b/Engine/Source/Programs/UnrealBuildAccelerator/Common/Private/UbaNetworkServer.cpp
--- a/Engine/Source/Programs/UnrealBuildAccelerator/Common/Private/UbaNetworkServer.cpp
+++ b/Engine/Source/Programs/UnrealBuildAccelerator/Common/Private/UbaNetworkServer.cpp
@@ -57,7 +57,7 @@ namespace uba
 		{
 			m_context = new WorkerContext(server);
 			m_loop = true;
-			m_thread.Start([&]() { ThreadWorker(server); return 0; });
+			m_thread.Start([&]() { ThreadWorker(server); return 0; }, TC("UbaWorker(NetworkServer)"));
 		}
 
 		void Stop(NetworkServer& server)
diff --git a/Engine/Source/Programs/UnrealBuildAccelerator/Common/Private/UbaProcess.cpp b/Engine/Source/Programs/UnrealBuildAccelerator/Common/Private/UbaProcess.cpp
--- a/Engine/Source/Programs/UnrealBuildAccelerator/Common/Private/UbaProcess.cpp
+++ b/Engine/Source/Programs/UnrealBuildAccelerator/Common/Private/UbaProcess.cpp
@@ -209,7 +209,7 @@ namespace uba
 					ThreadRun(environment);
 					ThreadExit();
 					return 0;
-				});
+				}, startInfo.description);
 		else
 		{
 			ThreadRun(environment);
diff --git a/Engine/Source/Programs/UnrealBuildAccelerator/Common/Private/UbaScheduler.cpp b/Engine/Source/Programs/UnrealBuildAccelerator/Common/Private/UbaScheduler.cpp
--- a/Engine/Source/Programs/UnrealBuildAccelerator/Common/Private/UbaScheduler.cpp
+++ b/Engine/Source/Programs/UnrealBuildAccelerator/Common/Private/UbaScheduler.cpp
@@ -119,7 +119,7 @@ namespace uba
 		m_session.SetRemoteProcessSlotAvailableEvent([this](bool isCrossArchitecture) { RemoteSlotAvailable(isCrossArchitecture); });
 
 		m_loop = true;
-		m_thread.Start([this]() { ThreadLoop(); return 0; });
+		m_thread.Start([this]() { ThreadLoop(); return 0; }, TC("UbaSchedulerLoop"));
 	}
 
 	void Scheduler::Stop()
diff --git a/Engine/Source/Programs/UnrealBuildAccelerator/Common/Private/UbaSession.cpp b/Engine/Source/Programs/UnrealBuildAccelerator/Common/Private/UbaSession.cpp
--- a/Engine/Source/Programs/UnrealBuildAccelerator/Common/Private/UbaSession.cpp
+++ b/Engine/Source/Programs/UnrealBuildAccelerator/Common/Private/UbaSession.cpp
@@ -1667,7 +1667,7 @@ namespace uba
 		GetSystemInfo(systemInfo);
 		m_trace.SessionAdded(0, {}, buf, systemInfo.data);
 		m_traceThreadEvent.Create(true);
-		m_traceThread.Start([this]() { ThreadTraceLoop(); return 0; });
+		m_traceThread.Start([this]() { ThreadTraceLoop(); return 0; }, TC("UbaTraceLoop"));
 	}
 
 	bool Session::StopTrace(const tchar* writeFile)
diff --git a/Engine/Source/Programs/UnrealBuildAccelerator/Common/Private/UbaSessionClient.cpp b/Engine/Source/Programs/UnrealBuildAccelerator/Common/Private/UbaSessionClient.cpp
--- a/Engine/Source/Programs/UnrealBuildAccelerator/Common/Private/UbaSessionClient.cpp
+++ b/Engine/Source/Programs/UnrealBuildAccelerator/Common/Private/UbaSessionClient.cpp
@@ -1236,7 +1236,7 @@ namespace uba
 		BuildEnvironmentVariables(reader);
 
 		m_loop = true;
-		m_loopThread.Start([this]() { ThreadCreateProcessLoop(); return 0; });
+		m_loopThread.Start([this]() { ThreadCreateProcessLoop(); return 0; }, TC("UbaCreateProcessLoop"));
 	}
 
 	void SessionClient::BuildEnvironmentVariables(BinaryReader& reader)
diff --git a/Engine/Source/Programs/UnrealBuildAccelerator/Common/Private/UbaSessionServer.cpp b/Engine/Source/Programs/UnrealBuildAccelerator/Common/Private/UbaSessionServer.cpp
--- a/Engine/Source/Programs/UnrealBuildAccelerator/Common/Private/UbaSessionServer.cpp
+++ b/Engine/Source/Programs/UnrealBuildAccelerator/Common/Private/UbaSessionServer.cpp
@@ -165,7 +165,7 @@ namespace uba
 				m_memRequiredToSpawn = Min(u64(double(m_memTotal) * double(100 - info.memWaitLoadPercent) / 100.0), 35ull * 1024 * 1024 * 1024);
 			}
 
-			m_memoryThread.Start([this]() { ThreadMemoryCheckLoop();  return 0; });
+			m_memoryThread.Start([this]() { ThreadMemoryCheckLoop();  return 0; }, TC("UbaMemoryTrackerLoop"));
 		}
 
 		#if PLATFORM_WINDOWS
diff --git a/Engine/Source/Programs/UnrealBuildAccelerator/Common/Private/UbaStorage.cpp b/Engine/Source/Programs/UnrealBuildAccelerator/Common/Private/UbaStorage.cpp
--- a/Engine/Source/Programs/UnrealBuildAccelerator/Common/Private/UbaStorage.cpp
+++ b/Engine/Source/Programs/UnrealBuildAccelerator/Common/Private/UbaStorage.cpp
@@ -1722,7 +1722,7 @@ namespace uba
 			};
 
 		m_logger.Info(TC("Traverse cas database..."));
-		WorkManagerImpl workManager(workerCount);
+		WorkManagerImpl workManager(workerCount, TC("UbaWorker(CheckCasContent)"));
 		u32 entryCount = 0;
 		u32 errorCount = 0;
 		u64 newestWrittenError = 0;
@@ -1797,7 +1797,7 @@ namespace uba
 	bool StorageImpl::CheckFileTable(const tchar* searchPath, u32 workerCount)
 	{
 		m_logger.Info(TC("Searching %s to check files against file table..."), searchPath);
-		WorkManagerImpl workManager(workerCount);
+		WorkManagerImpl workManager(workerCount, TC("UbaWorker(CheckFileTable)"));
 
 		List<TString> directories;
 		directories.push_back(searchPath);
@@ -1907,7 +1907,7 @@ namespace uba
 	{
 		u32 deleteCount = 0;
 
-		WorkManagerImpl workManager(GetLogicalProcessorCount());
+		WorkManagerImpl workManager(GetLogicalProcessorCount(), TC("UbaWorker(DeleteCas)"));
 		{
 			Atomic<u32> atomicDeleteCount;
 			TraverseDir(m_logger, m_rootDir.data, [&](const DirectoryEntry& e)
diff --git a/Engine/Source/Programs/UnrealBuildAccelerator/Common/Private/UbaStorageClient.cpp b/Engine/Source/Programs/UnrealBuildAccelerator/Common/Private/UbaStorageClient.cpp
--- a/Engine/Source/Programs/UnrealBuildAccelerator/Common/Private/UbaStorageClient.cpp
+++ b/Engine/Source/Programs/UnrealBuildAccelerator/Common/Private/UbaStorageClient.cpp
@@ -122,7 +122,7 @@ namespace uba
 
 		u64 start = GetTime();
 
-		WorkManagerImpl workManager(workerCount);
+		WorkManagerImpl workManager(workerCount, TC("UbaWorker(PopulateCasFromDirs)"));
 		ReaderWriterLock lock;
 		bool success = true;
 		UnorderedSet<u64> seenIds;
@@ -819,16 +819,14 @@ namespace uba
 
 		auto unmapGuard = MakeGuard([&](){ if (isPersistentMapping) UnmapViewOfFile(m_logger, fileMem, fileSize, fileName); });
 
-		CasKey casKey;
+		CasKey casKey = CalculateCasKey(fileMem, fileSize, storeCompressed);
+		if (casKey == CasKeyZero)
+			return false;
 
 		bool isTemporaryStorage = true; // This function is currently only used for temporary storage (output from client)
 		//UBA_ASSERT(wcsstr(fileName, TC("\\output\\"))); // Just make sure our assumption is correct :)
 		if (isTemporaryStorage)
 		{
-			casKey = CalculateCasKey(fileMem, fileSize, storeCompressed);
-			if (casKey == CasKeyZero)
-				return false;
-
 			if (keepMappingInMemory)
 			{
 				SCOPED_WRITE_LOCK(m_localStorageFilesLock, lock);
diff --git a/Engine/Source/Programs/UnrealBuildAccelerator/Common/Private/UbaThread.cpp b/Engine/Source/Programs/UnrealBuildAccelerator/Common/Private/UbaThread.cpp
--- a/Engine/Source/Programs/UnrealBuildAccelerator/Common/Private/UbaThread.cpp
+++ b/Engine/Source/Programs/UnrealBuildAccelerator/Common/Private/UbaThread.cpp
@@ -44,9 +44,9 @@ namespace uba
 	{
 	}
 
-	Thread::Thread(Function<u32()>&& func)
+	Thread::Thread(Function<u32()>&& func, const tchar* description)
 	{
-		Start(std::move(func));
+		Start(std::move(func), description);
 	}
 
 	Thread::~Thread()
@@ -54,11 +54,13 @@ namespace uba
 		Wait();
 	}
 
-	void Thread::Start(Function<u32()>&& f)
+	void Thread::Start(Function<u32()>&& f, const tchar* description)
 	{
 		m_func = std::move(f);
 #if PLATFORM_WINDOWS
 		m_handle = CreateThread(NULL, 0, [](LPVOID p) -> DWORD { return ((Thread*)p)->m_func(); }, this, 0, NULL);
+		if (description)
+			SetThreadDescription(m_handle, description);
 		AlternateThreadGroupAffinity(m_handle);
 #else
 		int err = 0;
diff --git a/Engine/Source/Programs/UnrealBuildAccelerator/Common/Private/UbaWorkManager.cpp b/Engine/Source/Programs/UnrealBuildAccelerator/Common/Private/UbaWorkManager.cpp
--- a/Engine/Source/Programs/UnrealBuildAccelerator/Common/Private/UbaWorkManager.cpp
+++ b/Engine/Source/Programs/UnrealBuildAccelerator/Common/Private/UbaWorkManager.cpp
@@ -23,11 +23,11 @@ namespace uba
 
 	struct WorkManagerImpl::Worker
 	{
-		Worker(WorkManagerImpl& manager) : m_workAvailable(false)
+		Worker(WorkManagerImpl& manager, const tchar* workerDesc) : m_workAvailable(false)
 		{
 			m_loop = true;
 			manager.PushWorker(this);
-			m_thread.Start([&]() { ThreadWorker(manager); return 0; });
+			m_thread.Start([&]() { ThreadWorker(manager); return 0; }, workerDesc);
 		}
 		~Worker()
 		{
@@ -82,12 +82,12 @@ namespace uba
 		Thread m_thread;
 	};
 
-	WorkManagerImpl::WorkManagerImpl(u32 workerCount)
+	WorkManagerImpl::WorkManagerImpl(u32 workerCount, const tchar* workerDesc)
 	{
 		m_workers.resize(workerCount);
 		m_activeWorkerCount = workerCount;
 		for (u32 i = 0; i != workerCount; ++i)
-			m_workers[i] = new Worker(*this);
+			m_workers[i] = new Worker(*this, workerDesc);
 	}
 
 	WorkManagerImpl::~WorkManagerImpl()
diff --git a/Engine/Source/Programs/UnrealBuildAccelerator/Common/Public/UbaThread.h b/Engine/Source/Programs/UnrealBuildAccelerator/Common/Public/UbaThread.h
--- a/Engine/Source/Programs/UnrealBuildAccelerator/Common/Public/UbaThread.h
+++ b/Engine/Source/Programs/UnrealBuildAccelerator/Common/Public/UbaThread.h
@@ -15,9 +15,9 @@ namespace uba
 	{
 	public:
 		Thread();
-		Thread(Function<u32()>&& func);
+		Thread(Function<u32()>&& func, const tchar* description = nullptr);
 		~Thread();
-		void Start(Function<u32()>&& func);
+		void Start(Function<u32()>&& func, const tchar* description = nullptr);
 		bool Wait(u32 milliseconds = ~0u, Event* wakeupEvent = nullptr);
 		bool GetGroupAffinity(GroupAffinity& out);
 
diff --git a/Engine/Source/Programs/UnrealBuildAccelerator/Common/Public/UbaWorkManager.h b/Engine/Source/Programs/UnrealBuildAccelerator/Common/Public/UbaWorkManager.h
--- a/Engine/Source/Programs/UnrealBuildAccelerator/Common/Public/UbaWorkManager.h
+++ b/Engine/Source/Programs/UnrealBuildAccelerator/Common/Public/UbaWorkManager.h
@@ -40,7 +40,7 @@ namespace uba
 	class WorkManagerImpl : public WorkManager
 	{
 	public:
-		WorkManagerImpl(u32 workerCount);
+		WorkManagerImpl(u32 workerCount, const tchar* workerDesc = TC("UbaWorker"));
 		virtual ~WorkManagerImpl();
 		virtual void AddWork(const Function<void()>& work, u32 count, const tchar* desc, bool highPriority = false) override;
 		virtual u32 GetWorkerCount() override;
