[UBA]
* Fixed so uba obj tool sets worker names and also reduced number of workers to process count - 1 (since thread itself helps out)

diff --git a/Engine/Source/Programs/UnrealBuildAccelerator/ObjTool/Private/UbaObjTool.cpp b/Engine/Source/Programs/UnrealBuildAccelerator/ObjTool/Private/UbaObjTool.cpp
--- a/Engine/Source/Programs/UnrealBuildAccelerator/ObjTool/Private/UbaObjTool.cpp
+++ b/Engine/Source/Programs/UnrealBuildAccelerator/ObjTool/Private/UbaObjTool.cpp
@@ -261,8 +261,8 @@ namespace uba
 			for (auto imp : NeededImports)
 				allExternalImports.insert(imp);
 
-			u32 workerCount = DefaultProcessorCount;
-			WorkManagerImpl workManager(workerCount);
+			u32 workerCount = DefaultProcessorCount-1;
+			WorkManagerImpl workManager(workerCount, TC("UbaWorker(StripObjFiles)"));
 			workManager.ParallelFor(workerCount, objFilesDependencies, [&](auto& it)
 				{
 					const TString& exiFilename = *it;
@@ -337,8 +337,8 @@ namespace uba
 			Atomic<bool> success = true;
 			Vector<ObjectFile*> objFiles;
 			objFiles.resize(objFilesForImpLib.size());
-			u32 workerCount = DefaultProcessorCount;
-			WorkManagerImpl workManager(workerCount);
+			u32 workerCount = DefaultProcessorCount - 1;
+			WorkManagerImpl workManager(workerCount, TC("UbaWorker(Load)"));
 			workManager.ParallelFor(workerCount, objFilesForImpLib, [&](auto& it)
 			{
 				auto& o = *it;
