[UBA]
* Fixed bug when using compressed obj files and remote linking.. GetFileSize from detoured process returned the size of the compressed file but the content of the file was decompressed
* Changed so UbaObjTool let's detouring decompress obj files.. otherwise size of file mismatch
* Minor fixes in various other areas

diff --git a/Engine/Source/Programs/UnrealBuildAccelerator/Cli/Private/UbaClient.h b/Engine/Source/Programs/UnrealBuildAccelerator/Cli/Private/UbaClient.h
--- a/Engine/Source/Programs/UnrealBuildAccelerator/Cli/Private/UbaClient.h
+++ b/Engine/Source/Programs/UnrealBuildAccelerator/Cli/Private/UbaClient.h
@@ -30,7 +30,9 @@ namespace uba
 			networkBackend = &info.networkBackend;
 			networkBackendMem = new NetworkBackendMemory(info.logWriter);
 			bool ctorSuccess = true;
-			networkClient = new NetworkClient(ctorSuccess, { info.logWriter });
+			NetworkClientCreateInfo ncci(info.logWriter);
+			ncci.workerCount = Min(info.maxProcessorCount, 8u);
+			networkClient = new NetworkClient(ctorSuccess, ncci);
 			if (!ctorSuccess)
 				return false;
 
@@ -57,6 +59,7 @@ namespace uba
 				sessionClientInfo.maxProcessCount = info.maxProcessorCount;
 				sessionClientInfo.rootDir = clientRootDir.data;
 				sessionClientInfo.deleteSessionsOlderThanSeconds = 1;
+				sessionClientInfo.name.Append("Agent").AppendValue(info.index);
 
 				sessionClient = new SessionClient(sessionClientInfo);
 				sessionClient->Start();
diff --git a/Engine/Source/Programs/UnrealBuildAccelerator/Common/Private/UbaStorageUtils.cpp b/Engine/Source/Programs/UnrealBuildAccelerator/Common/Private/UbaStorageUtils.cpp
--- a/Engine/Source/Programs/UnrealBuildAccelerator/Common/Private/UbaStorageUtils.cpp
+++ b/Engine/Source/Programs/UnrealBuildAccelerator/Common/Private/UbaStorageUtils.cpp
@@ -17,8 +17,6 @@ namespace uba
 		if (fileSize == 0)
 			return ToCasKey(hasher, storeCompressed);
 
-		UBA_ASSERTF(fileSize < sizeof(CompressedObjFileHeader) || !((CompressedObjFileHeader*)fileMem)->IsValid(), TC("CalculateCasKey should never run on compressed file"));
-
 		#ifndef __clang_analyzer__
 
 		if (fileSize > BufferSlotSize) // Note that when filesize is larger than BufferSlotSize the hash becomes a hash of hashes
diff --git a/Engine/Source/Programs/UnrealBuildAccelerator/Core/Private/UbaApplicationRules.cpp b/Engine/Source/Programs/UnrealBuildAccelerator/Core/Private/UbaApplicationRules.cpp
--- a/Engine/Source/Programs/UnrealBuildAccelerator/Core/Private/UbaApplicationRules.cpp
+++ b/Engine/Source/Programs/UnrealBuildAccelerator/Core/Private/UbaApplicationRules.cpp
@@ -547,7 +547,12 @@ namespace uba
 
 		virtual bool StoreFileCompressed(const StringView& fileName) const
 		{
-			return fileName.EndsWith(TC(".obj"));
+			return fileName.EndsWith(TC(".obj")); // Extra obj file for msvc
+		}
+
+		virtual bool ShouldDecompressFiles(const StringView& fileName) const override
+		{
+			return fileName.IsEmpty() || fileName.EndsWith(TC(".o")) || fileName.EndsWith(TC(".obj"));
 		}
 	};
 
diff --git a/Engine/Source/Programs/UnrealBuildAccelerator/Core/Public/UbaDirectoryTable.h b/Engine/Source/Programs/UnrealBuildAccelerator/Core/Public/UbaDirectoryTable.h
--- a/Engine/Source/Programs/UnrealBuildAccelerator/Core/Public/UbaDirectoryTable.h
+++ b/Engine/Source/Programs/UnrealBuildAccelerator/Core/Public/UbaDirectoryTable.h
@@ -293,6 +293,7 @@ namespace uba
 					auto fileIt = directory->files.find(fileNameKey);
 					if (fileIt != directory->files.end())
 					{
+						UBA_ASSERT(fileIt->second != ~0u);
 						BinaryReader reader(m_memory, fileIt->second);
 						StringBuffer<> fileName;
 						reader.ReadString(fileName);
diff --git a/Engine/Source/Programs/UnrealBuildAccelerator/Detours/Private/Windows/UbaDetoursFunctionsKernelBase.inl b/Engine/Source/Programs/UnrealBuildAccelerator/Detours/Private/Windows/UbaDetoursFunctionsKernelBase.inl
--- a/Engine/Source/Programs/UnrealBuildAccelerator/Detours/Private/Windows/UbaDetoursFunctionsKernelBase.inl
+++ b/Engine/Source/Programs/UnrealBuildAccelerator/Detours/Private/Windows/UbaDetoursFunctionsKernelBase.inl
@@ -1785,15 +1785,21 @@ BOOL Detoured_GetFileInformationByHandleEx(HANDLE hFile, FILE_INFO_BY_HANDLE_CLA
 	const wchar_t* originalName = nullptr;
 #endif
 
+	u64 fileSize = InvalidValue;
+
 	if (isDetouredHandle(hFile))
 	{
 		DetouredHandle& dh = asDetouredHandle(hFile);
 		trueHandle = dh.trueHandle;
 		entryOffset = dh.dirTableOffset;
 
+		FileObject& fo = *dh.fileObject;
+		FileInfo& fi = *fo.fileInfo;
+		fileSize = fi.size;
+
 		if (entryOffset == ~u32(0) && trueHandle == INVALID_HANDLE_VALUE)
 		{
-			MemoryFile* mf = dh.fileObject->fileInfo->memoryFile;
+			MemoryFile* mf = fi.memoryFile;
 			UBA_ASSERTF(mf, L"GetFileInformationByHandleEx called on file %s which has no entry offset or real handle", HandleToName(hFile));
 
 			DEBUG_LOG_DETOURED(L"GetFileInformationByHandleEx", L"(MEMORY) (%u) %llu (%ls)", fileInformationClass, uintptr_t(hFile), HandleToName(hFile));
@@ -1812,7 +1818,7 @@ BOOL Detoured_GetFileInformationByHandleEx(HANDLE hFile, FILE_INFO_BY_HANDLE_CLA
 				auto& data = *(FILE_STANDARD_INFO*)lpFileInformation;
 				data.EndOfFile = ToLargeInteger(mf->writtenSize);
 				data.AllocationSize = ToLargeInteger(mf->committedSize);
-				data.DeletePending = dh.fileObject->deleteOnClose;
+				data.DeletePending = fo.deleteOnClose;
 				data.NumberOfLinks = 1;
 				data.Directory = false;
 				return TRUE;
@@ -1863,8 +1869,10 @@ BOOL Detoured_GetFileInformationByHandleEx(HANDLE hFile, FILE_INFO_BY_HANDLE_CLA
 		}
 		else if (fileInformationClass == FileStandardInfo)
 		{
+			if (fileSize == InvalidValue)  // Always use FileInfo size if available since file could be decompressed and then directory info is wrong
+				fileSize = entryInfo.size;
 			auto& data = *(FILE_STANDARD_INFO*)lpFileInformation;
-			data.EndOfFile = ToLargeInteger(entryInfo.size);
+			data.EndOfFile = ToLargeInteger(fileSize);
 			data.AllocationSize = ToLargeInteger(entryInfo.size);
 			data.DeletePending = false;
 			data.NumberOfLinks = 1;
@@ -2067,7 +2075,7 @@ BOOL Detoured_GetFileInformationByHandle(HANDLE hFile, LPBY_HANDLE_FILE_INFORMAT
 				return TRUE;
 			}
 		}
-		UBA_ASSERT(dh.trueHandle != INVALID_HANDLE_VALUE);
+		UBA_ASSERTF(dh.trueHandle != INVALID_HANDLE_VALUE, TC("GetFileInformationByHandle needs true handle for %ls"), HandleToName(hFile));
 		trueHandle = dh.trueHandle;
 	}
 
@@ -2358,8 +2366,10 @@ DWORD Detoured_GetFinalPathNameByHandleW(HANDLE hFile, LPTSTR lpszFilePath, DWOR
 	{
 		auto& dh = asDetouredHandle(hFile);
 		auto fo = dh.fileObject;
-		UBA_ASSERT(fo && fo->fileInfo->originalName);
-		const wchar_t* fileName = fo->fileInfo->originalName;
+		UBA_ASSERT(fo);
+		auto& fi = *fo->fileInfo;
+		UBA_ASSERT(fi.originalName);
+		const wchar_t* fileName = fi.originalName;
 
 		if (dwFlags == 0 || dwFlags == 2)
 		{
diff --git a/Engine/Source/Programs/UnrealBuildAccelerator/Detours/Private/Windows/UbaDetoursFunctionsNtDll.inl b/Engine/Source/Programs/UnrealBuildAccelerator/Detours/Private/Windows/UbaDetoursFunctionsNtDll.inl
--- a/Engine/Source/Programs/UnrealBuildAccelerator/Detours/Private/Windows/UbaDetoursFunctionsNtDll.inl
+++ b/Engine/Source/Programs/UnrealBuildAccelerator/Detours/Private/Windows/UbaDetoursFunctionsNtDll.inl
@@ -737,7 +737,10 @@ NTSTATUS NTAPI Shared_NtCreateFile(bool IsCreateFunc, PHANDLE hFileHandle, ACCES
 
 	StringBuffer<> fileNameLower(fileName);
 	fileNameLower.MakeLower();
-	StringKey fileNameKey;
+	StringKey fileNameKey = ToStringKey(fileNameLower);
+
+	bool wasPreload = g_objFilesPreloader.Wait(fileNameKey); // Wait for on-going preload if existing
+
 
 	if (createFileName)// && (keepInMemory || canDetour))
 	{
@@ -749,7 +752,6 @@ NTSTATUS NTAPI Shared_NtCreateFile(bool IsCreateFunc, PHANDLE hFileHandle, ACCES
 			if ((!isWrite || TODO_DELETE_failIfNotExists) && !isSystemOrTempFile) // We need to skip SystemTemp.. lots of stuff going on there.
 			{
 				CHECK_PATH(fileNameLower.data);
-				fileNameKey = ToStringKey(fileNameLower);
 				dirTableOffset = Rpc_GetEntryOffset(fileNameKey, fileName.data, fileName.count, checkIfDir);
 
 				bool allowEarlyOut = true;
@@ -793,7 +795,7 @@ NTSTATUS NTAPI Shared_NtCreateFile(bool IsCreateFunc, PHANDLE hFileHandle, ACCES
 
 				bool isWriteAttributes = (DesiredAccess & FILE_WRITE_ATTRIBUTES) != 0;
 
-				if (allowEarlyOut && !useContent && !isWriteAttributes)
+				if (allowEarlyOut && !useContent && !isWriteAttributes && !wasPreload)
 				{
 					auto dh = new DetouredHandle(HandleType_File);
 					dh->fileObject = new FileObject();
@@ -899,12 +901,8 @@ NTSTATUS NTAPI Shared_NtCreateFile(bool IsCreateFunc, PHANDLE hFileHandle, ACCES
 		return STATUS_SUCCESS;
 	}
 
-	if (fileNameKey == StringKeyZero)
-	{
-		if (!keepInMemory || !isWrite) // we might get \\pipe\ here... 
-			CHECK_PATH(fileNameLower.data);
-		fileNameKey = ToStringKey(fileNameLower);
-	}
+	if (!keepInMemory || !isWrite) // we might get \\pipe\ here... 
+		CHECK_PATH(fileNameLower.data);
 
 	const wchar_t* lpFileName = fileName.data;
 	u32 closeId = 0;
@@ -933,8 +931,6 @@ NTSTATUS NTAPI Shared_NtCreateFile(bool IsCreateFunc, PHANDLE hFileHandle, ACCES
 	}
 	else
 	{
-		g_objFilesPreloader.Wait(fileNameKey); // Wait for on-going preload if existing
-
 		if (!info.originalName)
 			info.originalName = g_memoryBlock.Strdup(fileName.data);
 		if (isWrite)
@@ -1637,4 +1633,12 @@ BOOLEAN Detoured_RtlValidateHeap(HANDLE HeapPtr, ULONG Flags, PVOID Block)
 	return True_RtlValidateHeap(HeapPtr, Flags, Block);
 }
 
+NTSTATUS Detoured_RtlDosPathNameToNtPathName_U_WithStatus(PCWSTR dos_path, PUNICODE_STRING ntpath, PWSTR* file_part, VOID* reserved)
+{
+	auto res = True_RtlDosPathNameToNtPathName_U_WithStatus(dos_path, ntpath, file_part, reserved);
+	//DEBUG_LOG_TRUE(L"Detoured_RtlDosPathNameToNtPathName_U_WithStatus", L"-> %s", ToString(res).data);
+	return res;
+}
+
+
 #endif
diff --git a/Engine/Source/Programs/UnrealBuildAccelerator/Detours/Private/Windows/UbaDetoursFunctionsWin.h b/Engine/Source/Programs/UnrealBuildAccelerator/Detours/Private/Windows/UbaDetoursFunctionsWin.h
--- a/Engine/Source/Programs/UnrealBuildAccelerator/Detours/Private/Windows/UbaDetoursFunctionsWin.h
+++ b/Engine/Source/Programs/UnrealBuildAccelerator/Detours/Private/Windows/UbaDetoursFunctionsWin.h
@@ -322,6 +322,7 @@
 	DETOURED_FUNCTION(RtlAllocateHeap) \
 	DETOURED_FUNCTION(RtlReAllocateHeap) \
 	DETOURED_FUNCTION(RtlValidateHeap) \
+	DETOURED_FUNCTION(RtlDosPathNameToNtPathName_U_WithStatus) \
 	DETOURED_FUNCTION(NtCreateIoCompletion) \
 	DETOURED_FUNCTION(NtQueryFullAttributesFile) \
 	DETOURED_FUNCTION(NtFlushBuffersFileEx) \
@@ -402,6 +403,8 @@ extern "C" {
 	void WINAPI RtlExitUserThread(ULONG);
 	BOOLEAN NTAPI RtlFreeHeap(PVOID HeapHandle, ULONG Flags, PVOID HeapBase);
 	BOOLEAN NTAPI RtlValidateHeap(HANDLE HeapPtr, ULONG Flags, PVOID Block);
+	NTSTATUS NTAPI RtlDosPathNameToNtPathName_U_WithStatus(PCWSTR dos_path, PUNICODE_STRING ntpath, PWSTR* file_part, VOID* reserved);
+
 	PVOID NTAPI RtlAllocateHeap( PVOID HeapHandle, ULONG Flags, SIZE_T Size);
 	PVOID NTAPI RtlReAllocateHeap(PVOID HeapHandle, ULONG Flags, PVOID BaseAddress, SIZE_T Size);
 	SIZE_T NTAPI RtlSizeHeap(HANDLE HeapPtr, ULONG Flags, PVOID Ptr);
diff --git a/Engine/Source/Programs/UnrealBuildAccelerator/Detours/Private/Windows/UbaDetoursObjFilesPreloader.h b/Engine/Source/Programs/UnrealBuildAccelerator/Detours/Private/Windows/UbaDetoursObjFilesPreloader.h
--- a/Engine/Source/Programs/UnrealBuildAccelerator/Detours/Private/Windows/UbaDetoursObjFilesPreloader.h
+++ b/Engine/Source/Programs/UnrealBuildAccelerator/Detours/Private/Windows/UbaDetoursObjFilesPreloader.h
@@ -183,6 +183,10 @@ namespace uba
 					continue;
 				if (!((CompressedObjFileHeader*)preload.objMem)->IsValid())
 				{
+					// Ok, this is backwards.. but it could be that the file on server's disk is compressed..
+					// but it is decompressed in client's cas.. in that case dir table has the size of the compressed file.. 
+					// and since detoured code will provide decompressed file it also needs to provide decompressed size
+					preload.fileInfo->size = objFileSize.QuadPart;
 					preload.event.Set();
 					continue;
 				}
@@ -290,14 +294,15 @@ namespace uba
 			}
 		}
 
-		void Wait(const StringKey& key)
+		bool Wait(const StringKey& key)
 		{
 			auto findIt = m_preloadedObjFiles.find(key);
 			if (findIt == m_preloadedObjFiles.end())
-				return;
+				return false;
 			Preload& preload = findIt->second;
 			TimerScope ts(g_stats.waitDecompress);
 			preload.event.IsSet();
+			return true;
 		}
 
 		struct Preload
