[UBA]
* Fixed ubaagent's -binasversion.. it hasn't been tested in a long time and needed some fixes

diff --git a/Engine/Source/Programs/UnrealBuildAccelerator/Agent/Private/UbaAgent.cpp b/Engine/Source/Programs/UnrealBuildAccelerator/Agent/Private/UbaAgent.cpp
--- a/Engine/Source/Programs/UnrealBuildAccelerator/Agent/Private/UbaAgent.cpp
+++ b/Engine/Source/Programs/UnrealBuildAccelerator/Agent/Private/UbaAgent.cpp
@@ -194,7 +194,7 @@ namespace uba
 #endif
 
 #if UBA_AUTO_UPDATE
-	const tchar* g_ubaAgentBinaries[] = { UBA_AGENT_EXECUTABLE, UBA_DETOURS_LIBRARY };
+	const tchar* g_ubaAgentBinaries[] = { UBA_AGENT_EXECUTABLE }; //, UBA_DETOURS_LIBRARY };
 
 	bool DownloadBinaries(StorageClient& storageClient, CasKey* keys)
 	{
@@ -1275,20 +1275,6 @@ namespace uba
 				loggingThread.Wait();
 			});
 
-			// We got version mismatch and have the cas keys for the needed Agent/Detours binaries
-			if (keys[0] != CasKeyZero)
-			{
-#if UBA_AUTO_UPDATE
-				logger.Info(TC("Downloading new binaries..."));
-				if (!DownloadBinaries(*storageClient, keys))
-					return false;
-				relaunch = true;
-				break;
-#else
-				return false;
-#endif
-			}
-
 			if (quiet)
 				logger.Info(TC("Client session %s started"), sessionClient->GetId());
 			else
@@ -1347,7 +1333,23 @@ namespace uba
 				}
 
 				if (sessionClient->Wait(5*1000, &wakeupSessionWait))
+				{
+					// We got version mismatch and have the cas keys for the needed Agent/Detours binaries
+					if (keys[0] != CasKeyZero)
+					{
+		#if UBA_AUTO_UPDATE
+						logger.Info(TC("Downloading new binaries..."));
+						if (!DownloadBinaries(*storageClient, keys))
+							return false;
+						isTerminating = true;
+						relaunch = true;
+						break;
+		#else
+						return false;
+		#endif
+					}
 					break;
+				}
 
 				// If we are the proxy server and have external connections we lower max process count. Note that it will always have one connection which is itself
 				if (auto proxyServer = proxy.server.load())
diff --git a/Engine/Source/Programs/UnrealBuildAccelerator/Common/Private/UbaNetworkClient.cpp b/Engine/Source/Programs/UnrealBuildAccelerator/Common/Private/UbaNetworkClient.cpp
--- a/Engine/Source/Programs/UnrealBuildAccelerator/Common/Private/UbaNetworkClient.cpp
+++ b/Engine/Source/Programs/UnrealBuildAccelerator/Common/Private/UbaNetworkClient.cpp
@@ -371,9 +371,9 @@ namespace uba
 		m_onDisconnectedFunctions.push_back(function);
 	}
 
-	void NetworkClient::RegisterOnVersionMismatch(const OnVersionMismatchFunction& function)
+	void NetworkClient::RegisterOnVersionMismatch(OnVersionMismatchFunction&& function)
 	{
-		m_versionMismatchFunction = function;
+		m_versionMismatchFunction = std::move(function);
 	}
 
 	void NetworkClient::InvokeVersionMismatch(const CasKey& exeKey, const CasKey& dllKey)
diff --git a/Engine/Source/Programs/UnrealBuildAccelerator/Common/Private/UbaSessionServer.cpp b/Engine/Source/Programs/UnrealBuildAccelerator/Common/Private/UbaSessionServer.cpp
--- a/Engine/Source/Programs/UnrealBuildAccelerator/Common/Private/UbaSessionServer.cpp
+++ b/Engine/Source/Programs/UnrealBuildAccelerator/Common/Private/UbaSessionServer.cpp
@@ -624,14 +624,13 @@ namespace uba
 
 		m_logger.Detail(TC("Client session %s connected (Id: %u, Uid: %s%s)"), name.data, connectionInfo.GetId(), GuidToString(connectionInfo.GetUid()).str, (isClientArm ? TC(", IsArm: true") : TC("")));
 
-		CasKey clientKeys[2];
-		clientKeys[0] = reader.ReadCasKey();
-		clientKeys[1] = reader.ReadCasKey();
+		CasKey clientAgentKey = reader.ReadCasKey();
+		CasKey clientDetoursKey = reader.ReadCasKey();
 
 		CasKey& detoursBinaryKey = m_detoursBinaryKey[isClientArm];
 		CasKey& agentBinaryKey = m_agentBinaryKey[isClientArm];
 
-		bool binAsVersion = clientKeys[0] != CasKeyZero;
+		bool binAsVersion = clientAgentKey != CasKeyZero;
 		{
 			SCOPED_WRITE_LOCK(m_binKeysLock, lock);
 			if (detoursBinaryKey == CasKeyZero || (binAsVersion && agentBinaryKey == CasKeyZero))
@@ -663,35 +662,47 @@ namespace uba
 				{
 					UBA_ASSERT(IsWindows);
 					dir.Append(PathSeparator).Append(UBA_AGENT_EXECUTABLE);
-					m_storage.StoreCasFile(agentBinaryKey, dir.data, CasKeyZero, deferCreation, fileIsCompressed);
+					if (!m_storage.StoreCasFile(agentBinaryKey, dir.data, CasKeyZero, deferCreation, fileIsCompressed) || agentBinaryKey == CasKeyZero)
+					{
+						// This is hacky but uba binary is not copied by ubt anymore.
+						StringBuffer<> dir2(dir);
+						const tchar* engineDir = IsWindows ? TC("\\Engine\\") : TC("/Engine/");
+						const tchar* engineDirPos;
+						if (!dir2.Contains(engineDir, true, &engineDirPos))
+							return m_logger.Error(TC("Failed to find Engine dir in %s"), dir2.data);
+						const tchar* platformStr = TC("Win64");
+						dir2.Resize(engineDirPos - dir2.data + 8).Append(TC("Binaries")).Append(PathSeparator).Append(platformStr).Append(PathSeparator).Append("UnrealBuildAccelerator").Append(PathSeparator);
+						if (IsWindows)
+							dir2.Append(isClientArm ? TC("arm64") : TC("x64")).Append(PathSeparator);
+						dir2.Append(UBA_AGENT_EXECUTABLE);
+						if (!m_storage.StoreCasFile(agentBinaryKey, dir2.data, CasKeyZero, deferCreation, fileIsCompressed) || agentBinaryKey == CasKeyZero)
+							return m_logger.Error(TC("Failed to create cas for %s"), dir2.data);
+					}
 				}
 				dir.Resize(dirCount).Append(PathSeparator).Append(UBA_DETOURS_LIBRARY);
-				if (!m_storage.StoreCasFile(detoursBinaryKey, dir.data, CasKeyZero, deferCreation, fileIsCompressed))
+				if (!m_storage.StoreCasFile(detoursBinaryKey, dir.data, CasKeyZero, deferCreation, fileIsCompressed) || detoursBinaryKey == CasKeyZero)
 					return m_logger.Error(TC("Failed to create cas for %s"), dir.data);
-				UBA_ASSERT(detoursBinaryKey != CasKeyZero);
 			}
 		}
 
-		bool isValidVersion = clientVersion == SessionNetworkVersion;
-		if (binAsVersion)
-			isValidVersion = clientKeys[0] == agentBinaryKey && clientKeys[1] == detoursBinaryKey;
+		StringBuffer<> disconnectResponse;
+
+		if (binAsVersion && clientAgentKey != agentBinaryKey)
+		{
+			m_logger.Warning(TC("UbaAgent binaries mismatch. Disconnecting %s"), name.data);
+			disconnectResponse.Appendf(TC("UbaAgent binaries mismatch. Disconnecting..."));
+		}
+		else if (clientVersion != SessionNetworkVersion)
+		{
+			m_logger.Warning(TC("Version mismatch. Server is on version %u while client is on %u. Disconnecting %s"), SessionNetworkVersion, clientVersion, name.data);
+			disconnectResponse.Appendf(TC("Version mismatch. Server is on version %u while client is on %u. Disconnecting..."), SessionNetworkVersion, clientVersion);
+		}
 
-		writer.WriteBool(isValidVersion);
+		writer.WriteBool(disconnectResponse.IsEmpty());
 
-		if (!isValidVersion)
+		if (!disconnectResponse.IsEmpty())
 		{
-			StringBuffer<> response;
-			if (!isValidVersion)
-			{
-				m_logger.Warning(TC("Version mismatch. Server is on version %u while client is on %u. Disconnecting %s"), SessionNetworkVersion, clientVersion, name.data);
-				response.Appendf(TC("Version mismatch. Server is on version %u while client is on %u. Disconnecting..."), SessionNetworkVersion, clientVersion);
-			}
-			else
-			{
-				m_logger.Warning(TC("UbaAgent binaries mismatch. Disconnecting %s"), name.data);
-				response.Appendf(TC("UbaAgent binaries mismatch. Disconnecting..."));
-			}
-			writer.WriteString(response);
+			writer.WriteString(disconnectResponse);
 			writer.WriteCasKey(agentBinaryKey);
 			writer.WriteCasKey(detoursBinaryKey);
 			return true;
diff --git a/Engine/Source/Programs/UnrealBuildAccelerator/Common/Public/UbaNetworkClient.h b/Engine/Source/Programs/UnrealBuildAccelerator/Common/Public/UbaNetworkClient.h
--- a/Engine/Source/Programs/UnrealBuildAccelerator/Common/Public/UbaNetworkClient.h
+++ b/Engine/Source/Programs/UnrealBuildAccelerator/Common/Public/UbaNetworkClient.h
@@ -53,7 +53,7 @@ namespace uba
 		void RegisterOnDisconnected(const OnDisconnectedFunction& function);
 
 		using OnVersionMismatchFunction = Function<void(const CasKey& exeKey, const CasKey& dllKey)>;
-		void RegisterOnVersionMismatch(const OnVersionMismatchFunction& function);
+		void RegisterOnVersionMismatch(OnVersionMismatchFunction&& function);
 		void InvokeVersionMismatch(const CasKey& exeKey, const CasKey& dllKey);
 
 		u64 GetMessageHeaderSize();
