[UBA]
* Removed g_mappedViewCounter since we have proven it is not the cause of the ubacacheservice crash (there are no leaks of mapped views)

diff --git a/Engine/Source/Programs/UnrealBuildAccelerator/Common/Private/UbaCacheServer.cpp b/Engine/Source/Programs/UnrealBuildAccelerator/Common/Private/UbaCacheServer.cpp
--- a/Engine/Source/Programs/UnrealBuildAccelerator/Common/Private/UbaCacheServer.cpp
+++ b/Engine/Source/Programs/UnrealBuildAccelerator/Common/Private/UbaCacheServer.cpp
@@ -27,10 +27,6 @@
 
 namespace uba
 {
-	#if !PLATFORM_WINDOWS
-	extern Atomic<u64> g_mappedViewCounter;
-	#endif
-
 	static constexpr u32 CacheFileVersion = 6;
 	static constexpr u32 CacheFileCompatibilityVersion = 3;
 
@@ -594,10 +590,6 @@ namespace uba
 				auto& sentTimer = m_server.GetTotalSentTimer();
 				m_logger.Info(TC("  Socket sent %u (%s)"), sentTimer.count.load(), BytesToText(m_server.GetTotalSentBytes()).str);
 
-				#if !PLATFORM_WINDOWS
-				m_logger.Info(TC("  Active file mappings: %llu"), g_mappedViewCounter.load());
-				#endif
-
 				//KernelStats& kernelStats = KernelStats::GetGlobal();
 				//kernelStats.Print(m_logger, false);
 				//kernelStats = {};
diff --git a/Engine/Source/Programs/UnrealBuildAccelerator/Common/Private/UbaFileMapping.cpp b/Engine/Source/Programs/UnrealBuildAccelerator/Common/Private/UbaFileMapping.cpp
--- a/Engine/Source/Programs/UnrealBuildAccelerator/Common/Private/UbaFileMapping.cpp
+++ b/Engine/Source/Programs/UnrealBuildAccelerator/Common/Private/UbaFileMapping.cpp
@@ -59,7 +59,6 @@ namespace uba
 
 	ReaderWriterLock g_mappingUidCounterLock;
 	Atomic<u64> g_mappingUidCounter;
-	Atomic<u64> g_mappedViewCounter;
 #endif
 
 	FileMappingHandle CreateMemoryMappingW(Logger& logger, u32 flProtect, u64 maxSize, const tchar* name, const tchar* hint)
@@ -255,7 +254,6 @@ namespace uba
 			return nullptr;
 		}
 		u8* res = (u8*)rptr;
-		++g_mappedViewCounter;
 		//UBA_ASSERTF(false, "Failed to map file with fd %i, desiredAccess %u offset %llu, bytesToMap %llu (%s)", fd, desiredAccess, offset, bytesToMap, strerror(errno));
 #endif
 
@@ -304,10 +302,7 @@ namespace uba
 #else
 		UBA_ASSERTF(bytesToUnmap, TC("bytesToUnmap is zero unmapping %p (%s)"), lpBaseAddress, hint);
 		if (munmap((void*)lpBaseAddress, bytesToUnmap) == 0)
-		{
-			--g_mappedViewCounter;
 			return true;
-		}
 		UBA_ASSERT(false);
 		return false;
 #endif
