[UBA]
* Fixed ever-growing connection list in tcp backend. This will hopefully fix the ubacacheservice crash happening after ~32000 connections

diff --git a/Engine/Source/Programs/UnrealBuildAccelerator/Common/Private/UbaNetworkBackendMemory.cpp b/Engine/Source/Programs/UnrealBuildAccelerator/Common/Private/UbaNetworkBackendMemory.cpp
--- a/Engine/Source/Programs/UnrealBuildAccelerator/Common/Private/UbaNetworkBackendMemory.cpp
+++ b/Engine/Source/Programs/UnrealBuildAccelerator/Common/Private/UbaNetworkBackendMemory.cpp
@@ -152,6 +152,11 @@ namespace uba
 		return false;
 	}
 
+	void NetworkBackendMemory::DeleteConnection(void* connection)
+	{
+		//delete m_connection;
+	}
+
 	void NetworkBackendMemory::GetTotalSendAndRecv(u64& outSend, u64& outRecv)
 	{
 		outSend = 0;
diff --git a/Engine/Source/Programs/UnrealBuildAccelerator/Common/Private/UbaNetworkBackendTcp.cpp b/Engine/Source/Programs/UnrealBuildAccelerator/Common/Private/UbaNetworkBackendTcp.cpp
--- a/Engine/Source/Programs/UnrealBuildAccelerator/Common/Private/UbaNetworkBackendTcp.cpp
+++ b/Engine/Source/Programs/UnrealBuildAccelerator/Common/Private/UbaNetworkBackendTcp.cpp
@@ -707,6 +707,19 @@ namespace uba
 		return true;
 	}
 
+	void NetworkBackendTcp::DeleteConnection(void* connection)
+	{
+		SCOPED_WRITE_LOCK(m_connectionsLock, lock);
+		for (auto it=m_connections.begin();it!=m_connections.end(); ++it)
+		{
+			Connection& c = *it;
+			if (&c != connection)
+				continue;
+			it = m_connections.erase(it);
+			break;
+		}
+	}
+
 	void NetworkBackendTcp::GetTotalSendAndRecv(u64& outSend, u64& outRecv)
 	{
 		outSend = m_totalSend;
diff --git a/Engine/Source/Programs/UnrealBuildAccelerator/Common/Private/UbaNetworkClient.cpp b/Engine/Source/Programs/UnrealBuildAccelerator/Common/Private/UbaNetworkClient.cpp
--- a/Engine/Source/Programs/UnrealBuildAccelerator/Common/Private/UbaNetworkClient.cpp
+++ b/Engine/Source/Programs/UnrealBuildAccelerator/Common/Private/UbaNetworkClient.cpp
@@ -291,6 +291,8 @@ namespace uba
 
 		{
 			SCOPED_WRITE_LOCK(m_connectionsLock, lock2);
+			for (auto& c : m_connections)
+				c.backend->DeleteConnection(c.backendConnection);
 			m_connections.clear();
 			m_connectionsIt = m_connections.end();
 		}
diff --git a/Engine/Source/Programs/UnrealBuildAccelerator/Common/Private/UbaNetworkServer.cpp b/Engine/Source/Programs/UnrealBuildAccelerator/Common/Private/UbaNetworkServer.cpp
--- a/Engine/Source/Programs/UnrealBuildAccelerator/Common/Private/UbaNetworkServer.cpp
+++ b/Engine/Source/Programs/UnrealBuildAccelerator/Common/Private/UbaNetworkServer.cpp
@@ -1144,6 +1144,7 @@ namespace uba
 				continue;
 			}
 			m_sendTimer += con.m_sendTimer;
+			con.m_backend.DeleteConnection(con.m_backendConnection);
 			it = m_connections.erase(it);
 		}
 	}
diff --git a/Engine/Source/Programs/UnrealBuildAccelerator/Common/Public/UbaNetworkBackend.h b/Engine/Source/Programs/UnrealBuildAccelerator/Common/Public/UbaNetworkBackend.h
--- a/Engine/Source/Programs/UnrealBuildAccelerator/Common/Public/UbaNetworkBackend.h
+++ b/Engine/Source/Programs/UnrealBuildAccelerator/Common/Public/UbaNetworkBackend.h
@@ -62,6 +62,9 @@ namespace uba
 		virtual bool Connect(Logger& logger, const tchar* ip, const ConnectedFunc& connectedFunc, u16 port = DefaultPort, bool* timedOut = nullptr) = 0;
 		virtual bool Connect(Logger& logger, const sockaddr& remoteSocketAddr, const ConnectedFunc& connectedFunc, bool* timedOut = nullptr, const tchar* nameHint = nullptr) = 0;
 
+		// Remove connection (user guarantee connection is not referenced anywhere)
+		virtual void DeleteConnection(void* connection) = 0;
+
 		virtual void GetTotalSendAndRecv(u64& outSend, u64& outRecv) = 0;
 	};
 
diff --git a/Engine/Source/Programs/UnrealBuildAccelerator/Common/Public/UbaNetworkBackendMemory.h b/Engine/Source/Programs/UnrealBuildAccelerator/Common/Public/UbaNetworkBackendMemory.h
--- a/Engine/Source/Programs/UnrealBuildAccelerator/Common/Public/UbaNetworkBackendMemory.h
+++ b/Engine/Source/Programs/UnrealBuildAccelerator/Common/Public/UbaNetworkBackendMemory.h
@@ -25,6 +25,7 @@ namespace uba
 		virtual void StopListen() override;
 		virtual bool Connect(Logger& logger, const tchar* ip, const ConnectedFunc& connectedFunc, u16 port = DefaultPort, bool* timedOut = nullptr) override;
 		virtual bool Connect(Logger& logger, const sockaddr& remoteSocketAddr, const ConnectedFunc& connectedFunc, bool* timedOut = nullptr, const tchar* nameHint = nullptr) override;
+		virtual void DeleteConnection(void* connection) override;
 
 		virtual void GetTotalSendAndRecv(u64& outSend, u64& outRecv) override;
 
diff --git a/Engine/Source/Programs/UnrealBuildAccelerator/Common/Public/UbaNetworkBackendTcp.h b/Engine/Source/Programs/UnrealBuildAccelerator/Common/Public/UbaNetworkBackendTcp.h
--- a/Engine/Source/Programs/UnrealBuildAccelerator/Common/Public/UbaNetworkBackendTcp.h
+++ b/Engine/Source/Programs/UnrealBuildAccelerator/Common/Public/UbaNetworkBackendTcp.h
@@ -25,6 +25,7 @@ namespace uba
 		virtual void StopListen() override;
 		virtual bool Connect(Logger& logger, const tchar* ip, const ConnectedFunc& connectedFunc, u16 port = DefaultPort, bool* timedOut = nullptr) override;
 		virtual bool Connect(Logger& logger, const sockaddr& remoteSocketAddr, const ConnectedFunc& connectedFunc, bool* timedOut = nullptr, const tchar* nameHint = nullptr) override;
+		virtual void DeleteConnection(void* connection) override;
 
 		virtual void GetTotalSendAndRecv(u64& outSend, u64& outRecv) override;
 
