[UBA]
* Added some documentation around assert where writtenfile changes owner.

diff --git a/Engine/Source/Programs/UnrealBuildAccelerator/Common/Private/UbaSession.cpp b/Engine/Source/Programs/UnrealBuildAccelerator/Common/Private/UbaSession.cpp
--- a/Engine/Source/Programs/UnrealBuildAccelerator/Common/Private/UbaSession.cpp
+++ b/Engine/Source/Programs/UnrealBuildAccelerator/Common/Private/UbaSession.cpp
@@ -2054,7 +2054,12 @@ namespace uba
 			auto insres = msg.process.m_writtenFiles.try_emplace(name);
 			WrittenFile& writtenFile = insres.first->second;
 
-			UBA_ASSERTF(writtenFile.owner == nullptr || writtenFile.owner == &msg.process || !m_allowOutputFiles, TC("File %s changed owner.. should not happen"), name);
+			if (m_allowOutputFiles && writtenFile.owner != nullptr && writtenFile.owner != &msg.process)
+			{
+				// This can happen when library has /GL (whole program optimization) but target has not.. then link.exe will restart
+				UBA_ASSERTF(false, TC("File %s changed owner.. should not happen (OldOwner: %s New owner: %s)"), name, writtenFile.owner->m_realApplication.c_str(), msg.process.m_realApplication.c_str());
+			}
+
 			writtenFile.owner = &msg.process;
 			writtenFile.attributes = msg.attributes;
 
