[UBA]
* Fixed bug related to compressed obj files and using cache. Don't use external file mappings if file is compressed since we can instead just read the cas out of the written file (the file mapping is not compressed)

diff --git a/Engine/Source/Programs/UnrealBuildAccelerator/Common/Private/UbaStorage.cpp b/Engine/Source/Programs/UnrealBuildAccelerator/Common/Private/UbaStorage.cpp
--- a/Engine/Source/Programs/UnrealBuildAccelerator/Common/Private/UbaStorage.cpp
+++ b/Engine/Source/Programs/UnrealBuildAccelerator/Common/Private/UbaStorage.cpp
@@ -2009,7 +2009,7 @@ namespace uba
 		u64 fileSize = 0;
 		u64 lastWriteTime = 0;
 
-		if (m_externalFileMappingsProvider && m_externalFileMappingsProvider(externalMapping, fileNameKey, fileName))
+		if (!fileIsCompressed && m_externalFileMappingsProvider && m_externalFileMappingsProvider(externalMapping, fileNameKey, fileName))
 		{
 			fileSize = externalMapping.size;
 			lastWriteTime = externalMapping.lastWriteTime;
@@ -2057,14 +2057,7 @@ namespace uba
 		fileEntry.lastWritten = lastWriteTime;
 		if (casKeyOverride == CasKeyZero)
 		{
-			if (fileIsCompressed)
-			{
-				CompressedObjFileHeader header(CasKeyZero);
-				if (!ReadFile(m_logger, fileName, fileHandle, &header, sizeof(header)))
-					return m_logger.Error(TC("Failed to read header of compressed file %s (%s)"), fileName, LastErrorToText().data);
-				fileEntry.casKey = AsCompressed(header.casKey, m_storeCompressed);
-			}
-			else if (externalMapping.handle.IsValid())
+			if (externalMapping.handle.IsValid())
 			{
 				FileMappingScope scope(m_logger);
 				u8* fileMem = scope.MapView(TC("StoreCasFile"), externalMapping.handle, externalMapping.offset, externalMapping.size, fileName);
@@ -2072,10 +2065,17 @@ namespace uba
 					return false;
 				fileEntry.casKey = CalculateCasKey(fileMem, externalMapping.size, m_storeCompressed);
 			}
-			else
+			else if (fileIsCompressed)
 			{
-				fileEntry.casKey = CalculateCasKey(fileName, fileHandle, fileSize, m_storeCompressed);
+				UBA_ASSERT(fileSize > sizeof(CompressedObjFileHeader));
+				CompressedObjFileHeader header(CasKeyZero);
+				if (!ReadFile(m_logger, fileName, fileHandle, &header, sizeof(header)))
+					return m_logger.Error(TC("Failed to read header of compressed file %s (%s)"), fileName, LastErrorToText().data);
+				UBA_ASSERT(header.IsValid());
+				fileEntry.casKey = AsCompressed(header.casKey, m_storeCompressed);
 			}
+			else
+				fileEntry.casKey = CalculateCasKey(fileName, fileHandle, fileSize, m_storeCompressed);
 		}
 		else
 			fileEntry.casKey = casKeyOverride;
