[UBA]
* Fixed so StringBuffer::Parse can take in offset to position to start parsing

diff --git a/Engine/Source/Programs/UnrealBuildAccelerator/Core/Private/UbaStringBuffer.cpp b/Engine/Source/Programs/UnrealBuildAccelerator/Core/Private/UbaStringBuffer.cpp
--- a/Engine/Source/Programs/UnrealBuildAccelerator/Core/Private/UbaStringBuffer.cpp
+++ b/Engine/Source/Programs/UnrealBuildAccelerator/Core/Private/UbaStringBuffer.cpp
@@ -341,22 +341,22 @@ namespace uba
 		return uba::Parse(out, data, count);
 	}
 
-	bool StringBufferBase::Parse(u32& out)
+	bool StringBufferBase::Parse(u32& out, u64 offset)
 	{
-		if (!count)
+		if (count <= offset)
 			return false;
 		#if PLATFORM_WINDOWS
-		out = wcstoul(data, nullptr, 10);
+		out = wcstoul(data + offset, nullptr, 10);
 		#else
-		out = strtoul(data, nullptr, 10);
+		out = strtoul(data + offset, nullptr, 10);
 		#endif
 		return out != 0 || Equals(TC("0"));
 	}
 
-	bool StringBufferBase::Parse(u16& out)
+	bool StringBufferBase::Parse(u16& out, u64 offset)
 	{
 		u32 temp;
-		if (!Parse(temp))
+		if (!Parse(temp, offset))
 			return false;
 		if (temp > 65535)
 			return false;
diff --git a/Engine/Source/Programs/UnrealBuildAccelerator/Core/Public/UbaStringBuffer.h b/Engine/Source/Programs/UnrealBuildAccelerator/Core/Public/UbaStringBuffer.h
--- a/Engine/Source/Programs/UnrealBuildAccelerator/Core/Public/UbaStringBuffer.h
+++ b/Engine/Source/Programs/UnrealBuildAccelerator/Core/Public/UbaStringBuffer.h
@@ -63,8 +63,8 @@ namespace uba
 		StringBufferBase& MakeLower();
 
 		bool Parse(u64& out);
-		bool Parse(u32& out);
-		bool Parse(u16& out);
+		bool Parse(u32& out, u64 offset = 0);
+		bool Parse(u16& out, u64 offset = 0);
 		bool Parse(float& out);
 		u32 Parse(char* out, u64 outCapacity); // Note, return value contains length + null termination.. zero if failed
 
