[UBA]
* Improved stats tracking of kernel call WriteFile in detoured process

diff --git a/Engine/Source/Programs/UnrealBuildAccelerator/Detours/Private/Windows/UbaDetoursFunctionsKernelBase.inl b/Engine/Source/Programs/UnrealBuildAccelerator/Detours/Private/Windows/UbaDetoursFunctionsKernelBase.inl
--- a/Engine/Source/Programs/UnrealBuildAccelerator/Detours/Private/Windows/UbaDetoursFunctionsKernelBase.inl
+++ b/Engine/Source/Programs/UnrealBuildAccelerator/Detours/Private/Windows/UbaDetoursFunctionsKernelBase.inl
@@ -677,8 +677,14 @@ BOOL Detoured_WriteFile(HANDLE hFile, LPCVOID lpBuffer, DWORD nNumberOfBytesToWr
 		//return True_WriteFile(trueHandle, lpBuffer, nNumberOfBytesToWrite, lpNumberOfBytesWritten, lpOverlapped);
 	}
 
+	DWORD temp;
+	if (!lpNumberOfBytesWritten)
+		lpNumberOfBytesWritten = &temp;
+
 	TimerScope ts(g_kernelStats.writeFile);
 	BOOL res = True_WriteFile(trueHandle, lpBuffer, nNumberOfBytesToWrite, lpNumberOfBytesWritten, lpOverlapped);
+
+	g_kernelStats.writeFile.bytes += *lpNumberOfBytesWritten;
 	DEBUG_LOG_TRUE(L"WriteFile", L"%llu (%ls) -> %ls", uintptr_t(hFile), HandleToName(hFile), ToString(res));
 	return res;
 }
