[UBA]
* Fixed hang in visualizer when pausing trace and jump to before time 0

diff --git a/Engine/Source/Programs/UnrealBuildAccelerator/Visualizer/Private/UbaVisualizer.cpp b/Engine/Source/Programs/UnrealBuildAccelerator/Visualizer/Private/UbaVisualizer.cpp
--- a/Engine/Source/Programs/UnrealBuildAccelerator/Visualizer/Private/UbaVisualizer.cpp
+++ b/Engine/Source/Programs/UnrealBuildAccelerator/Visualizer/Private/UbaVisualizer.cpp
@@ -2446,7 +2446,7 @@ namespace uba
 	{
 		u64 currentTime = m_paused ? m_pauseStart : GetTime();
 		u64 playTime = 0;
-		if (m_traceView.startTime)
+		if (m_traceView.startTime && currentTime > (m_traceView.startTime + m_pauseTime))
 			playTime = currentTime - m_traceView.startTime - m_pauseTime;
 		if (m_replay)
 			playTime *= m_replay;
