[UBA]
* Fixed calculation of total committed memory. Needed to take registry value PagefileOnOsVolume into account

diff --git a/Engine/Source/Programs/UnrealBuildAccelerator/Common/Private/UbaSession.cpp b/Engine/Source/Programs/UnrealBuildAccelerator/Common/Private/UbaSession.cpp
--- a/Engine/Source/Programs/UnrealBuildAccelerator/Common/Private/UbaSession.cpp
+++ b/Engine/Source/Programs/UnrealBuildAccelerator/Common/Private/UbaSession.cpp
@@ -2720,12 +2720,23 @@ namespace uba
 		// Page file can grow and we want to use the absolute max size to figure out when we need to wait to start new processes
 		if (m_maxPageSize == ~u64(0))
 		{
+			tchar systemDrive = 'c';
+			{
+				tchar temp[32];
+				if (GetEnvironmentVariableW(TC("SystemDrive"), temp, 32))
+					systemDrive = ToLower(temp[0]);
+			}
+
 			m_maxPageSize = 0;
 			wchar_t str[1024];
 			DWORD strBytes = sizeof(str);
 			LSTATUS res = RegGetValueW(HKEY_LOCAL_MACHINE, TC("SYSTEM\\CurrentControlSet\\Control\\Session Manager\\Memory Management"), TC("PagingFiles"), RRF_RT_REG_MULTI_SZ, NULL, str, &strBytes);
 			if (res == ERROR_SUCCESS)
 			{
+				DWORD pagefileOnOsVolume = 0;
+				DWORD pagefileOnOsVolumeSize = 4;
+				res = RegGetValueW(HKEY_LOCAL_MACHINE, TC("SYSTEM\\CurrentControlSet\\Control\\Session Manager\\Memory Management"), TC("PagefileOnOsVolume"), RRF_RT_DWORD, NULL, &pagefileOnOsVolume, &pagefileOnOsVolumeSize);
+
 				wchar_t* line = str;
 				for (; size_t lineLen = wcslen(line); line += lineLen + 1)
 				{
@@ -2752,8 +2763,11 @@ namespace uba
 							continue;
 
 						drive[0] = colon[-1];
+
+						if (pagefileOnOsVolume && ToLower(drive[0]) != systemDrive)
+							continue;
 					}
-					else
+					else if (!pagefileOnOsVolume || ToLower(drive[0]) == systemDrive)
 					{
 						const wchar_t* maxSizeStr = wcsrchr(line, ' ');
 						
@@ -2769,6 +2783,10 @@ namespace uba
 							continue;
 						}
 					}
+					else
+					{
+						m_logger.Warning(TC("Page file is set on drive %c: but registry key value PagefileOnOsVolume is set to 1. Fix registry"), drive[0]);
+					}
 
 					// Max possible system-managed page file
 					maxSizeMb = Max(u64(memStatus.ullTotalPhys) * 3, 4ull * 1024 * 1024 * 1024);
