[UBA]
* Changed so deletion of backend connection is made after server connection is destroyed

diff --git a/Engine/Source/Programs/UnrealBuildAccelerator/Common/Private/UbaNetworkServer.cpp b/Engine/Source/Programs/UnrealBuildAccelerator/Common/Private/UbaNetworkServer.cpp
--- a/Engine/Source/Programs/UnrealBuildAccelerator/Common/Private/UbaNetworkServer.cpp
+++ b/Engine/Source/Programs/UnrealBuildAccelerator/Common/Private/UbaNetworkServer.cpp
@@ -1144,8 +1144,10 @@ namespace uba
 				continue;
 			}
 			m_sendTimer += con.m_sendTimer;
-			con.m_backend.DeleteConnection(con.m_backendConnection);
+			auto& backend = con.m_backend;
+			void* backendConnection = con.m_backendConnection;
 			it = m_connections.erase(it);
+			backend.DeleteConnection(backendConnection);
 		}
 	}
 
