[UBA]
* Fixed so clients are refcounted by connections and deleted when connections reach zero.

diff --git a/Engine/Source/Programs/UnrealBuildAccelerator/Common/Private/UbaNetworkServer.cpp b/Engine/Source/Programs/UnrealBuildAccelerator/Common/Private/UbaNetworkServer.cpp
--- a/Engine/Source/Programs/UnrealBuildAccelerator/Common/Private/UbaNetworkServer.cpp
+++ b/Engine/Source/Programs/UnrealBuildAccelerator/Common/Private/UbaNetworkServer.cpp
@@ -286,16 +286,19 @@ namespace uba
 			SCOPED_WRITE_LOCK(conn.m_shutdownLock, shutdownLock);
 
 			SCOPED_WRITE_LOCK(server.m_clientsLock, clientsLock);
-			u32 clientId = u32(server.m_clients.size() + 1);
+			u32 clientId = 0;
 			for (auto& kv : server.m_clients)
 				if (kv.second.uid == clientUid)
 					clientId = kv.second.id;
+			if (!clientId)
+				clientId = ++server.m_clientCounter;
 			Client& client = server.m_clients.try_emplace(clientId, clientUid, clientId).first->second;
+			++client.refCount;
 			clientsLock.Leave();
 
 			conn.m_client = &client;
 
-			if (client.connectionCount.fetch_add(1) == 0)
+			if (client.connectionCount.fetch_add(1))
 			{
 				if (server.m_onConnectionFunction)
 					server.m_onConnectionFunction(clientUid, clientId);
@@ -1135,6 +1138,8 @@ namespace uba
 
 	void NetworkServer::RemoveDisconnectedConnections()
 	{
+		bool clientRefCountChanged = false;
+
 		for (auto it=m_connections.begin(); it!=m_connections.end();)
 		{
 			Connection& con = *it;
@@ -1146,9 +1151,27 @@ namespace uba
 			m_sendTimer += con.m_sendTimer;
 			auto& backend = con.m_backend;
 			void* backendConnection = con.m_backendConnection;
+			if (auto client = con.m_client)
+			{
+				--client->refCount;
+				clientRefCountChanged = true;
+			}
+
 			it = m_connections.erase(it);
 			backend.DeleteConnection(backendConnection);
 		}
+
+		if (!clientRefCountChanged)
+			return;
+
+		SCOPED_WRITE_LOCK(m_clientsLock, lock);
+		for (auto it=m_clients.begin(); it!= m_clients.end(); ++it)
+		{
+			if (it->second.refCount)
+				++it;
+			else
+				it = m_clients.erase(it);
+		}
 	}
 
 	bool NetworkServer::HandleSystemMessage(const ConnectionInfo& connectionInfo, u8 messageType, BinaryReader& reader, BinaryWriter& writer)
diff --git a/Engine/Source/Programs/UnrealBuildAccelerator/Common/Public/UbaNetworkServer.h b/Engine/Source/Programs/UnrealBuildAccelerator/Common/Public/UbaNetworkServer.h
--- a/Engine/Source/Programs/UnrealBuildAccelerator/Common/Public/UbaNetworkServer.h
+++ b/Engine/Source/Programs/UnrealBuildAccelerator/Common/Public/UbaNetworkServer.h
@@ -169,12 +169,14 @@ namespace uba
 			Client(const Guid& uid_, u32 id_) : uid(uid_), id(id_) {}
 			Guid uid;
 			u32 id;
+			Atomic<u32> refCount;
 			Atomic<u32> connectionCount;
 			Atomic<u64> sendBytes;
 			Atomic<u64> recvBytes;
 		};
 		ReaderWriterLock m_clientsLock;
 		UnorderedMap<u32, Client> m_clients;
+		u32 m_clientCounter = 0;
 
 		Timer m_sendTimer;
 		Timer m_encryptTimer;
