[UBA]
* Fixed typo in previous submit causing hangs

diff --git a/Engine/Source/Programs/UnrealBuildAccelerator/Common/Private/UbaNetworkServer.cpp b/Engine/Source/Programs/UnrealBuildAccelerator/Common/Private/UbaNetworkServer.cpp
--- a/Engine/Source/Programs/UnrealBuildAccelerator/Common/Private/UbaNetworkServer.cpp
+++ b/Engine/Source/Programs/UnrealBuildAccelerator/Common/Private/UbaNetworkServer.cpp
@@ -298,7 +298,7 @@ namespace uba
 
 			conn.m_client = &client;
 
-			if (client.connectionCount.fetch_add(1))
+			if (client.connectionCount.fetch_add(1) == 0)
 			{
 				if (server.m_onConnectionFunction)
 					server.m_onConnectionFunction(clientUid, clientId);
