[UBA]
* Added MemoryBlock::Strdup that takes/returns StringView

diff --git a/Engine/Source/Programs/UnrealBuildAccelerator/Core/Private/UbaMemory.cpp b/Engine/Source/Programs/UnrealBuildAccelerator/Core/Private/UbaMemory.cpp
--- a/Engine/Source/Programs/UnrealBuildAccelerator/Core/Private/UbaMemory.cpp
+++ b/Engine/Source/Programs/UnrealBuildAccelerator/Core/Private/UbaMemory.cpp
@@ -2,6 +2,7 @@
 
 #include "UbaMemory.h"
 #include "UbaPlatform.h"
+#include "UbaStringBuffer.h"
 
 #if PLATFORM_LINUX
 #include <linux/mman.h>
@@ -138,14 +139,18 @@ namespace uba
 			aligned_free(p);
 	}
 
-	tchar* MemoryBlock::Strdup(const tchar* str)
+	StringView MemoryBlock::Strdup(const StringView& str)
 	{
-		u64 len = TStrlen(str);
-		u64 memSize = (len + 1) * sizeof(tchar);
+		u64 memSize = (str.count + 1) * sizeof(tchar);
 		void* mem = Allocate(memSize, sizeof(tchar), TC("Strdup"));
-		const void* src = str;
+		const void* src = str.data;
 		memcpy(mem, src, memSize);
-		return (tchar*)mem;
+		return StringView((tchar*)mem, str.count);
+	}
+
+	tchar* MemoryBlock::Strdup(const tchar* str)
+	{
+		return (tchar*)Strdup(ToView(str)).data;
 	}
 
 	void MemoryBlock::Swap(MemoryBlock& other)
diff --git a/Engine/Source/Programs/UnrealBuildAccelerator/Core/Public/UbaMemory.h b/Engine/Source/Programs/UnrealBuildAccelerator/Core/Public/UbaMemory.h
--- a/Engine/Source/Programs/UnrealBuildAccelerator/Core/Public/UbaMemory.h
+++ b/Engine/Source/Programs/UnrealBuildAccelerator/Core/Public/UbaMemory.h
@@ -26,6 +26,8 @@
 
 namespace uba
 {
+	struct StringView;
+
 	template<typename T>
 	constexpr inline T AlignUp(T arg, uintptr_t alignment) { return T(((uintptr_t)arg+(alignment-1)) & ~(alignment-1)); }
 
@@ -101,6 +103,7 @@ namespace uba
 		void* AllocateNoLock(u64 bytes, u64 alignment, const tchar* hint);
 		void ReserveNoLock(u64 bytes, const tchar* hint);
 		void Free(void* p);
+		StringView Strdup(const StringView& str);
 		tchar* Strdup(const tchar* str);
 		void Swap(MemoryBlock& other);
 		
