[UBA]
* Added option for RootPaths to allow finding paths that does not have a root

diff --git a/Engine/Source/Programs/UnrealBuildAccelerator/Common/Private/UbaCacheClient.cpp b/Engine/Source/Programs/UnrealBuildAccelerator/Common/Private/UbaCacheClient.cpp
--- a/Engine/Source/Programs/UnrealBuildAccelerator/Common/Private/UbaCacheClient.cpp
+++ b/Engine/Source/Programs/UnrealBuildAccelerator/Common/Private/UbaCacheClient.cpp
@@ -1001,7 +1001,7 @@ namespace uba
 						}
 					};
 
-				if (!rootPaths.NormalizeString<char>(m_logger, (const char*)file.GetData(), file.GetSize(), handleString, path.data))
+				if (!rootPaths.NormalizeString<char>(m_logger, (const char*)file.GetData(), file.GetSize(), handleString, false, path.data))
 					return false;
 
 				if (rootOffsetsSize)
@@ -1158,7 +1158,7 @@ namespace uba
 
 		// Add arguments list to key
 		auto hashString = [&](const tchar* str, u64 strLen, u32 rootPos) { hasher.Update(str, strLen*sizeof(tchar)); };
-		if (!rootPaths.NormalizeString(m_logger, info.arguments, TStrlen(info.arguments), hashString, TC("CmdKey "), info.GetDescription()))
+		if (!rootPaths.NormalizeString(m_logger, info.arguments, TStrlen(info.arguments), hashString, false, TC("CmdKey "), info.GetDescription()))
 			return CasKeyZero;
 
 		// Add content of rsp file to key (This will cost a bit of perf since we need to normalize.. should this be part of key?)
diff --git a/Engine/Source/Programs/UnrealBuildAccelerator/Common/Private/UbaRootPaths.cpp b/Engine/Source/Programs/UnrealBuildAccelerator/Common/Private/UbaRootPaths.cpp
--- a/Engine/Source/Programs/UnrealBuildAccelerator/Common/Private/UbaRootPaths.cpp
+++ b/Engine/Source/Programs/UnrealBuildAccelerator/Common/Private/UbaRootPaths.cpp
@@ -87,7 +87,7 @@ namespace uba
 		return true;
 	}
 
-	const RootPaths::Root* RootPaths::FindRoot(const StringBufferBase& path) const
+	const RootPaths::Root* RootPaths::FindRoot(const StringView& path) const
 	{
 		if (path.count < m_shortestRoot)
 			return nullptr;
@@ -132,7 +132,7 @@ namespace uba
 				wasNormalized |= rootPos != ~0u;
 				hasher.Update(str, strLen);
 			};
-		if (!NormalizeString<char>(logger, (const char*)file.GetData(), file.GetSize(), hashString, filename))
+		if (!NormalizeString<char>(logger, (const char*)file.GetData(), file.GetSize(), hashString, false, filename))
 			return CasKeyZero;
 
 		return AsNormalized(ToCasKey(hasher, false), wasNormalized);
diff --git a/Engine/Source/Programs/UnrealBuildAccelerator/Common/Public/UbaRootPaths.h b/Engine/Source/Programs/UnrealBuildAccelerator/Common/Public/UbaRootPaths.h
--- a/Engine/Source/Programs/UnrealBuildAccelerator/Common/Public/UbaRootPaths.h
+++ b/Engine/Source/Programs/UnrealBuildAccelerator/Common/Public/UbaRootPaths.h
@@ -24,11 +24,11 @@ namespace uba
 			bool includeInKey;
 		};
 
-		const Root* FindRoot(const StringBufferBase& path) const;
+		const Root* FindRoot(const StringView& path) const;
 		const TString& GetRoot(u32 index) const;
 
 		template<typename CharType, typename Func>
-		bool NormalizeString(Logger& logger, const CharType* str, u64 strLen, const Func& func, const tchar* hint, const tchar* hint2 = TC("")) const;
+		bool NormalizeString(Logger& logger, const CharType* str, u64 strLen, const Func& func, bool allowPathsWithoutRoot, const tchar* hint, const tchar* hint2 = TC("")) const;
 
 		CasKey NormalizeAndHashFile(Logger& logger, const tchar* filename) const;
 
@@ -45,7 +45,7 @@ namespace uba
 
 
 	template<typename CharType, typename Func>
-	bool RootPaths::NormalizeString(Logger& logger, const CharType* str, u64 strLen, const Func& func, const tchar* hint, const tchar* hint2) const
+	bool RootPaths::NormalizeString(Logger& logger, const CharType* str, u64 strLen, const Func& func, bool allowPathsWithoutRoot, const tchar* hint, const tchar* hint2) const
 	{
 		auto strEnd = str + strLen;
 		auto searchPos = str;
@@ -82,6 +82,15 @@ namespace uba
 			auto root = FindRoot(path);
 			if (!root)
 			{
+				if (allowPathsWithoutRoot)
+				{
+					u32 len = 1;
+					destPos += len;
+					func(searchPos, len, ~0u);
+					searchPos += len;
+					continue;
+				}
+
 				if (auto lastQuote = path.Last('\"'))
 					path.Resize(lastQuote - path.data);
 				if (auto lineEnd = path.Last('\r'))
diff --git a/Engine/Source/Programs/UnrealBuildAccelerator/Test/Private/UbaTestBasics.h b/Engine/Source/Programs/UnrealBuildAccelerator/Test/Private/UbaTestBasics.h
--- a/Engine/Source/Programs/UnrealBuildAccelerator/Test/Private/UbaTestBasics.h
+++ b/Engine/Source/Programs/UnrealBuildAccelerator/Test/Private/UbaTestBasics.h
@@ -412,7 +412,7 @@ namespace uba
 					if (!temp.Equals(TC("foo")))
 						success = false;
 				}
-			}, TC(""));
+			}, false, TC(""));
 
 		if (!success)
 			return false;
