[UBA]
* Removed outputStatsThresholdMs

diff --git a/Engine/Source/Programs/UnrealBuildAccelerator/Agent/Private/UbaAgent.cpp b/Engine/Source/Programs/UnrealBuildAccelerator/Agent/Private/UbaAgent.cpp
--- a/Engine/Source/Programs/UnrealBuildAccelerator/Agent/Private/UbaAgent.cpp
+++ b/Engine/Source/Programs/UnrealBuildAccelerator/Agent/Private/UbaAgent.cpp
@@ -85,7 +85,6 @@ namespace uba
 		logger.Info(TC("  -config=<file>          Config file that contains options for various systems"));
 		logger.Info(TC("  -quic                   Use Quic instead of tcp backend."));
 		logger.Info(TC("  -name=<name>            The identifier of this agent. Defaults to \"%s\" on this machine"), DefaultAgentName);
-		logger.Info(TC("  -stats[=<threshold>]    Print stats for each process if higher than threshold"));
 		logger.Info(TC("  -verbose                Print debug information to console"));
 		logger.Info(TC("  -log                    Log all processes detouring information to file (only works with debug builds)"));
 		logger.Info(TC("  -nocustomalloc          Disable custom allocator for processes. If you see odd crashes this can be tested"));
@@ -407,7 +406,6 @@ namespace uba
 		u32 maxWorkerCount = DefaultProcessorCount;
 		float mulProcessValue = 1.0f;
 		u32 maxConnectionCount = DefaultMaxConnectionCount;
-		u32 outputStatsThresholdMs = 0;
 		u32 storageCapacityGb = DefaultCapacityGb;
 		StringBuffer<256> host;
 		StringBuffer<256> named;
@@ -523,13 +521,6 @@ namespace uba
 					return false;
 				configFile.Append(value);
 			}
-			else if (name.Equals(TC("-stats")))
-			{
-				if (value.IsEmpty())
-					outputStatsThresholdMs = 1;
-				else if (!value.Parse(outputStatsThresholdMs))
-					return PrintHelp(TC("Invalid for -stats"));
-			}
 			else if (name.Equals(TC("-host")))
 			{
 				if (const tchar* portIndex = value.First(':'))
@@ -1199,7 +1190,6 @@ namespace uba
 			info.maxProcessCount = maxProcessCount;
 			info.dedicated = poll;
 			info.maxIdleSeconds = maxIdleSeconds;
-			info.outputStatsThresholdMs = outputStatsThresholdMs;
 			info.name.Append(agentName);
 			info.extraInfo = extraInfo.data;
 			info.deleteSessionsOlderThanSeconds = 1; // Delete all old sessions
diff --git a/Engine/Source/Programs/UnrealBuildAccelerator/Common/Private/UbaProcess.cpp b/Engine/Source/Programs/UnrealBuildAccelerator/Common/Private/UbaProcess.cpp
--- a/Engine/Source/Programs/UnrealBuildAccelerator/Common/Private/UbaProcess.cpp
+++ b/Engine/Source/Programs/UnrealBuildAccelerator/Common/Private/UbaProcess.cpp
@@ -997,7 +997,6 @@ namespace uba
 		info.workingDir = currentDir.data;
 		info.logFile = InternalGetChildLogFile(temp);
 		info.priorityClass = m_startInfo.priorityClass;
-		info.outputStatsThresholdMs = m_startInfo.outputStatsThresholdMs;
 		info.logLineUserData = this;
 		info.logLineFunc = [](void* userData, const tchar* line, u32 length, LogEntryType type) { ((ProcessImpl*)userData)->LogLine(false, TString(line, length), type); };
 		info.startSuspended = startSuspended;
@@ -1213,12 +1212,6 @@ namespace uba
 			m_parentProcess->m_kernelStats.Add(m_kernelStats);
 		}
 
-		if (m_startInfo.outputStatsThresholdMs && TimeToMs(m_processStats.GetTotalTime()) > m_startInfo.outputStatsThresholdMs)
-		{
-			m_session.PrintProcessStats(m_processStats, logName.data);
-			m_processStats.Print(m_session.m_logger, GetFrequency());
-		}
-
 		return false;
 	}
 
diff --git a/Engine/Source/Programs/UnrealBuildAccelerator/Common/Private/UbaSessionClient.cpp b/Engine/Source/Programs/UnrealBuildAccelerator/Common/Private/UbaSessionClient.cpp
--- a/Engine/Source/Programs/UnrealBuildAccelerator/Common/Private/UbaSessionClient.cpp
+++ b/Engine/Source/Programs/UnrealBuildAccelerator/Common/Private/UbaSessionClient.cpp
@@ -38,7 +38,6 @@ namespace uba
 		m_useStorage = info.useStorage;
 		m_downloadDetoursLib = info.downloadDetoursLib;
 		m_defaultPriorityClass = info.defaultPriorityClass;
-		m_outputStatsThresholdMs = info.outputStatsThresholdMs;
 		m_maxIdleSeconds = info.maxIdleSeconds;
 		m_disableCustomAllocator = info.disableCustomAllocator;
 		m_useBinariesAsVersion = info.useBinariesAsVersion;
@@ -1702,7 +1701,6 @@ namespace uba
 					startInfo.uiLanguage = int(m_uiLanguage);
 					startInfo.priorityClass = m_defaultPriorityClass;
 					startInfo.useCustomAllocator = !m_disableCustomAllocator;
-					startInfo.outputStatsThresholdMs = m_outputStatsThresholdMs != 0 ? m_outputStatsThresholdMs : startInfo.outputStatsThresholdMs;
 					startInfo.rules = GetRules(startInfo);
 
 					StringBuffer<> logFile;
diff --git a/Engine/Source/Programs/UnrealBuildAccelerator/Common/Public/UbaProcessStartInfo.h b/Engine/Source/Programs/UnrealBuildAccelerator/Common/Public/UbaProcessStartInfo.h
--- a/Engine/Source/Programs/UnrealBuildAccelerator/Common/Public/UbaProcessStartInfo.h
+++ b/Engine/Source/Programs/UnrealBuildAccelerator/Common/Public/UbaProcessStartInfo.h
@@ -23,7 +23,6 @@ namespace uba
 		const tchar* description = TC("");		// Description. Used for on-screen logging and log file names if session.logToFile is set but logFile is ""
 		const tchar* logFile = TC("");			// Log file. If set, will always log. If not full path the session log dir will be prepended.
 		u32 priorityClass = 0x00000020;			// Priority of process. Defaults to NORMAL_PRIORITY_CLASS
-		u64 outputStatsThresholdMs = ~u64(0);	// Threshold in milliseconds where process summary should be printed to log
 		bool trackInputs = false;				// Track all files read. Can read result in ProcessHandle.GetTrackedInputs()
 		bool useCustomAllocator = true;			// Disable detouring of allocator inside processes. If Session.disableCustomAllocator is false this will be overridden
 		bool writeOutputFilesOnFail = false;	// If set to true, output files will be written to disk regardless if process succeeds or not
diff --git a/Engine/Source/Programs/UnrealBuildAccelerator/Common/Public/UbaProcessStartInfoHolder.h b/Engine/Source/Programs/UnrealBuildAccelerator/Common/Public/UbaProcessStartInfoHolder.h
--- a/Engine/Source/Programs/UnrealBuildAccelerator/Common/Public/UbaProcessStartInfoHolder.h
+++ b/Engine/Source/Programs/UnrealBuildAccelerator/Common/Public/UbaProcessStartInfoHolder.h
@@ -65,7 +65,6 @@ namespace uba
 			writer.WriteU32(*(u32*)&weight);
 			writer.WriteBool(trackInputs);
 			writer.WriteBool(writeOutputFilesOnFail);
-			writer.WriteU64(outputStatsThresholdMs);
 		}
 
 		void Read(BinaryReader& reader)
@@ -84,7 +83,6 @@ namespace uba
 			
 			trackInputs = reader.ReadBool();
 			writeOutputFilesOnFail = reader.ReadBool();
-			outputStatsThresholdMs = reader.ReadU64();
 
 			description = descriptionStr.c_str();
 			application = applicationStr.c_str();
diff --git a/Engine/Source/Programs/UnrealBuildAccelerator/Common/Public/UbaSessionClient.h b/Engine/Source/Programs/UnrealBuildAccelerator/Common/Public/UbaSessionClient.h
--- a/Engine/Source/Programs/UnrealBuildAccelerator/Common/Public/UbaSessionClient.h
+++ b/Engine/Source/Programs/UnrealBuildAccelerator/Common/Public/UbaSessionClient.h
@@ -17,7 +17,6 @@ namespace uba
 		StringBuffer<128> name;
 		u32 maxProcessCount = 1;
 		u32 defaultPriorityClass = 0x00004000; // BELOW_NORMAL_PRIORITY_CLASS;
-		u32 outputStatsThresholdMs = 0;
 		u32 maxIdleSeconds = ~0u;
 		u8 memWaitLoadPercent = 80; // When memory usage goes above this percent, no new processes will be spawned until back below
 		u8 memKillLoadPercent = 90; // When memory usage goes above this percent, newest processes will be killed to bring it back below
@@ -109,7 +108,6 @@ namespace uba
 		u32 m_sessionId = 0;
 		u32 m_uiLanguage = 0;
 		u32 m_defaultPriorityClass;
-		u32 m_outputStatsThresholdMs = 0;
 		u32 m_maxIdleSeconds = ~0u;
 		u32 m_killRandomIndex = ~0u;
 		u32 m_killRandomCounter = 0;
diff --git a/Engine/Source/Programs/UnrealBuildAccelerator/Host/Private/UbaExports.cpp b/Engine/Source/Programs/UnrealBuildAccelerator/Host/Private/UbaExports.cpp
--- a/Engine/Source/Programs/UnrealBuildAccelerator/Host/Private/UbaExports.cpp
+++ b/Engine/Source/Programs/UnrealBuildAccelerator/Host/Private/UbaExports.cpp
@@ -504,7 +504,6 @@ extern "C"
 		info->workingDir = TStrdup(workingDir);
 		info->description = TStrdup(description);
 		info->priorityClass = priorityClass;
-		info->outputStatsThresholdMs = outputStatsThresholdMs;
 		info->trackInputs = trackInputs;
 		info->logFile = TStrdup(logFile);
 		info->exitedFunc = exit;
