[UBA]
* Moved code fetching processhandle from processid to a separate function

diff --git a/Engine/Source/Programs/UnrealBuildAccelerator/Common/Private/UbaSessionServer.cpp b/Engine/Source/Programs/UnrealBuildAccelerator/Common/Private/UbaSessionServer.cpp
--- a/Engine/Source/Programs/UnrealBuildAccelerator/Common/Private/UbaSessionServer.cpp
+++ b/Engine/Source/Programs/UnrealBuildAccelerator/Common/Private/UbaSessionServer.cpp
@@ -1019,7 +1019,7 @@ namespace uba
 			RegisterCreateFileForWrite(StringKeyZero, destination, shouldWriteToDisk, 0, 0, invalidateStorage);
 
 
-			SCOPED_WRITE_LOCK(m_processesLock, lock);
+			SCOPED_READ_LOCK(m_processesLock, lock);
 			auto findIt = m_processes.find(processId);
 			if (findIt != m_processes.end())
 			{
@@ -1290,12 +1290,9 @@ namespace uba
 	bool SessionServer::HandleProcessInputs(const ConnectionInfo& connectionInfo, BinaryReader& reader, BinaryWriter& writer)
 	{
 		u32 processId = u32(reader.Read7BitEncoded());
-		SCOPED_WRITE_LOCK(m_processesLock, lock);
-		auto findIt = m_processes.find(processId);
-		if (findIt == m_processes.end())
+		ProcessHandle h(GetProcess(processId));
+		if (!h.IsValid())
 			return m_logger.Error(TC("Failed to find process for id %u when receiving custom message"), processId);
-		ProcessHandle h(findIt->second);
-		lock.Leave();
 		auto& process = *(RemoteProcess*)h.m_process;
 		auto& inputs = process.m_trackedInputs;
 		u64 size = inputs.size();
@@ -1484,12 +1481,9 @@ namespace uba
 	{
 		u32 processId = reader.ReadU32();
 		u32 prevExitCode = reader.ReadU32();
-		SCOPED_WRITE_LOCK(m_processesLock, lock);
-		auto findIt = m_processes.find(processId);
-		if (findIt == m_processes.end())
-			return m_logger.Error(TC("Failed to find process for id %u when receiving custom message"), processId);
-		ProcessHandle h(findIt->second);
-		lock.Leave();
+		ProcessHandle h(GetProcess(processId));
+		if (!h.IsValid())
+			return m_logger.Error(TC("Failed to find process for id %u when receiving GetNextProcess message"), processId);
 
 		auto& remoteProcess = *(RemoteProcess*)h.m_process;
 		SCOPED_WRITE_LOCK(remoteProcess.m_exitedLock, exitedLock);
@@ -1517,13 +1511,9 @@ namespace uba
 	bool SessionServer::HandleCustom(const ConnectionInfo& connectionInfo, BinaryReader& reader, BinaryWriter& writer)
 	{
 		u32 processId = reader.ReadU32();
-		SCOPED_WRITE_LOCK(m_processesLock, lock);
-		auto findIt = m_processes.find(processId);
-		if (findIt == m_processes.end())
+		ProcessHandle h(GetProcess(processId));
+		if (!h.IsValid())
 			return m_logger.Error(TC("Failed to find process for id %u when receiving custom message"), processId);
-		ProcessHandle h(findIt->second);
-		lock.Leave();
-
 		auto& remoteProcess = *(RemoteProcess*)h.m_process;
 		SCOPED_WRITE_LOCK(remoteProcess.m_exitedLock, exitedLock);
 		CustomMessage(remoteProcess, reader, writer);
@@ -1533,13 +1523,9 @@ namespace uba
 	bool SessionServer::HandleUpdateEnvironment(const ConnectionInfo& connectionInfo, BinaryReader& reader, BinaryWriter& writer)
 	{
 		u32 processId = reader.ReadU32();
-		SCOPED_WRITE_LOCK(m_processesLock, lock);
-		auto findIt = m_processes.find(processId);
-		if (findIt == m_processes.end())
-			return m_logger.Error(TC("Failed to find process for id %u when receiving custom message"), processId);
-		ProcessHandle h(findIt->second);
-		lock.Leave();
-
+		ProcessHandle h(GetProcess(processId));
+		if (!h.IsValid())
+			return m_logger.Error(TC("Failed to find process for id %u when receiving update environment message"), processId);
 		StringBuffer<> reason;
 		reader.ReadString(reason);
 		m_trace.ProcessEnvironmentUpdated(processId, reason, reader.GetPositionData(), reader.GetLeft(), ToView(h.GetStartInfo().breadcrumbs));
@@ -1923,6 +1909,15 @@ namespace uba
 		return h;
 	}
 
+	ProcessHandle SessionServer::GetProcess(u32 processId)
+	{
+		SCOPED_READ_LOCK(m_processesLock, lock);
+		auto findIt = m_processes.find(processId);
+		if (findIt == m_processes.end())
+			return {};
+		return ProcessHandle(findIt->second);
+	}
+
 	TString SessionServer::GetProcessDescription(u32 processId)
 	{
 		StringBuffer<512> str;
diff --git a/Engine/Source/Programs/UnrealBuildAccelerator/Common/Public/UbaSessionServer.h b/Engine/Source/Programs/UnrealBuildAccelerator/Common/Public/UbaSessionServer.h
--- a/Engine/Source/Programs/UnrealBuildAccelerator/Common/Public/UbaSessionServer.h
+++ b/Engine/Source/Programs/UnrealBuildAccelerator/Common/Public/UbaSessionServer.h
@@ -96,6 +96,7 @@ namespace uba
 		void OnCancelled(RemoteProcess* process);
 		ProcessHandle ProcessRemoved(u32 processId);
 
+		ProcessHandle GetProcess(u32 processId);
 		TString GetProcessDescription(u32 processId);
 
 		virtual bool PrepareProcess(ProcessStartInfoHolder& startInfo, bool isChild, StringBufferBase& outRealApplication, const tchar*& outRealWorkingDir) override final;
