[UBA]
* Added missing function that belonged to previous submit

diff --git a/Engine/Source/Programs/UnrealBuildAccelerator/Common/Private/UbaSession.cpp b/Engine/Source/Programs/UnrealBuildAccelerator/Common/Private/UbaSession.cpp
--- a/Engine/Source/Programs/UnrealBuildAccelerator/Common/Private/UbaSession.cpp
+++ b/Engine/Source/Programs/UnrealBuildAccelerator/Common/Private/UbaSession.cpp
@@ -121,6 +121,11 @@ namespace uba
 		return *this;
 	}
 
+	bool ProcessHandle::IsValid() const
+	{
+		return m_process != nullptr;
+	}
+
 	const ProcessStartInfo& ProcessHandle::GetStartInfo() const
 	{
 		UBA_ASSERT(m_process);
