[UBA]
* Added missing function that belonged to previous submit

diff --git a/Engine/Source/Programs/UnrealBuildAccelerator/Common/Public/UbaProcessHandle.h b/Engine/Source/Programs/UnrealBuildAccelerator/Common/Public/UbaProcessHandle.h
--- a/Engine/Source/Programs/UnrealBuildAccelerator/Common/Public/UbaProcessHandle.h
+++ b/Engine/Source/Programs/UnrealBuildAccelerator/Common/Public/UbaProcessHandle.h
@@ -31,7 +31,7 @@ namespace uba
 	class ProcessHandle
 	{
 	public:
-
+		bool IsValid() const;								// Returns true if handle points to a process
 		const ProcessStartInfo& GetStartInfo() const;		// Get the info used to start process
 		u32 GetId() const;									// Unique id of process. Is an increasing index internally
 		u32 GetExitCode() const;							// Exit code of process. Will return ~0u if not finished or default error
