[UBA]
* Changed code to use the MemoryBlock::Strdup taking StringView

diff --git a/Engine/Source/Programs/UnrealBuildAccelerator/Detours/Private/PosixOS/UbaDetoursFunctionsPosix.cpp b/Engine/Source/Programs/UnrealBuildAccelerator/Detours/Private/PosixOS/UbaDetoursFunctionsPosix.cpp
--- a/Engine/Source/Programs/UnrealBuildAccelerator/Detours/Private/PosixOS/UbaDetoursFunctionsPosix.cpp
+++ b/Engine/Source/Programs/UnrealBuildAccelerator/Detours/Private/PosixOS/UbaDetoursFunctionsPosix.cpp
@@ -379,7 +379,7 @@ int Shared_open(const char* funcName, const char* file, int flags, int mode, con
 	u32 lastDesiredAccess = info.lastDesiredAccess;
 	if (insres.second)
 	{
-		info.originalName = g_memoryBlock.Strdup(fileName.data);
+		info.originalName = g_memoryBlock.Strdup(fileName).data;
 		info.name = info.originalName;
 		if (!keepInMemory)
 		{
@@ -396,7 +396,7 @@ int Shared_open(const char* funcName, const char* file, int flags, int mode, con
 	else
 	{
 		if (!info.originalName)
-			info.originalName = g_memoryBlock.Strdup(fileName.data);
+			info.originalName = g_memoryBlock.Strdup(fileName).data;
 		if (isWrite)
 		{
 			//UBA_ASSERT(!info.isFileMap);
@@ -1547,7 +1547,7 @@ UBA_EXPORT int UBA_WRAPPER(rename)(const char* oldpath, const char* newpath)
 		auto findIt = g_mappedFileTable.m_lookup.find(oldKey);
 		if (findIt != g_mappedFileTable.m_lookup.end())
 			info = findIt->second;
-		info.originalName = g_memoryBlock.Strdup(fixedNewPath.data);
+		info.originalName = g_memoryBlock.Strdup(fixedNewPath).data;
 		info.name = info.originalName;
 	}
 
@@ -2144,7 +2144,7 @@ UBA_EXPORT int UBA_WRAPPER(dladdr)(const void *addr, Dl_info *info)
 	{
 		StringBuffer<> newPath;
 		newPath.Append(g_virtualApplication.data).Append(info->dli_fname + g_exeDir.count - 1); // TODO: Make g_handledLibraries a map and store the real path
-		info->dli_fname = g_memoryBlock.Strdup(newPath.data);
+		info->dli_fname = g_memoryBlock.Strdup(newPath).data;
 	}
 	DEBUG_LOG_TRUE("dladdr", "%s -> %i", info->dli_fname, res);
 	return res;
diff --git a/Engine/Source/Programs/UnrealBuildAccelerator/Detours/Private/UbaDetoursFileMappingTable.cpp b/Engine/Source/Programs/UnrealBuildAccelerator/Detours/Private/UbaDetoursFileMappingTable.cpp
--- a/Engine/Source/Programs/UnrealBuildAccelerator/Detours/Private/UbaDetoursFileMappingTable.cpp
+++ b/Engine/Source/Programs/UnrealBuildAccelerator/Detours/Private/UbaDetoursFileMappingTable.cpp
@@ -41,7 +41,7 @@ namespace uba
 			if (!insres.second)
 				continue;
 			info.fileNameKey = g;
-			info.name = m_memoryBlock.Strdup(mappedFileName.data);
+			info.name = m_memoryBlock.Strdup(mappedFileName).data;
 			//info.refCount = 0;
 			info.size = size;
 			//info.dwDesiredAccess = GENERIC_READ;
@@ -216,7 +216,7 @@ namespace uba
 			StringBuffer<> name;
 			reader.ReadString(name);
 			u64 fileSize = reader.ReadU64();
-			info.name = g_mappedFileTable.m_memoryBlock.Strdup(name.data);
+			info.name = g_mappedFileTable.m_memoryBlock.Strdup(name).data;
 			info.size = fileSize;
 		}
 	}
diff --git a/Engine/Source/Programs/UnrealBuildAccelerator/Detours/Private/Windows/UbaDetoursPrepopulatePchFiles.h b/Engine/Source/Programs/UnrealBuildAccelerator/Detours/Private/Windows/UbaDetoursPrepopulatePchFiles.h
--- a/Engine/Source/Programs/UnrealBuildAccelerator/Detours/Private/Windows/UbaDetoursPrepopulatePchFiles.h
+++ b/Engine/Source/Programs/UnrealBuildAccelerator/Detours/Private/Windows/UbaDetoursPrepopulatePchFiles.h
@@ -206,7 +206,7 @@ namespace uba
 			auto fileNameKey = ToStringKeyLower(t);
 			auto insres = g_mappedFileTable.m_lookup.try_emplace(fileNameKey);
 			FileInfo& info = insres.first->second;
-			info.originalName = g_memoryBlock.Strdup(include.data());
+			info.originalName = g_memoryBlock.Strdup(include).data;
 			info.name = info.originalName;
 			info.size = 0;
 			info.fileNameKey = fileNameKey;
