[UBA]
* Added support for vfs. With this change it is possible to register a vfs which is then used to convert paths inside the detoured process
* Added support for converting clang deps file back from vfs to local paths

diff --git a/Engine/Source/Programs/UnrealBuildAccelerator/Common/Private/UbaProcess.cpp b/Engine/Source/Programs/UnrealBuildAccelerator/Common/Private/UbaProcess.cpp
--- a/Engine/Source/Programs/UnrealBuildAccelerator/Common/Private/UbaProcess.cpp
+++ b/Engine/Source/Programs/UnrealBuildAccelerator/Common/Private/UbaProcess.cpp
@@ -727,6 +727,20 @@ namespace uba
 		writer.WriteU64(response.mappedFileTableHandle);
 		writer.WriteU32(response.mappedFileTableSize);
 		writer.WriteU32(response.mappedFileTableCount);
+
+		if (!m_startInfo.vfsHandle)
+		{
+			writer.WriteU16(0);
+			return true;
+		}
+
+		SCOPED_READ_LOCK(m_session.m_vfsLookupLock, lock);
+		auto findIt = m_session.m_vfsLookup.find(m_startInfo.vfsHandle);
+		if (findIt == m_session.m_vfsLookup.end())
+			return m_session.m_logger.Error(TC("Process use vfs handle %llu which does not exist"), m_startInfo.vfsHandle);
+		Session::VfsEntry& vfsEntry = findIt->second;
+		writer.WriteU16(u16(vfsEntry.memory.size()));
+		writer.WriteBytes(vfsEntry.memory.data(), vfsEntry.memory.size());
 		return true;
 	}
 	
@@ -1000,6 +1014,7 @@ namespace uba
 		info.logLineUserData = this;
 		info.logLineFunc = [](void* userData, const tchar* line, u32 length, LogEntryType type) { ((ProcessImpl*)userData)->LogLine(false, TString(line, length), type); };
 		info.startSuspended = startSuspended;
+		info.vfsHandle = m_startInfo.vfsHandle;
 
 		ProcessImpl* parent = isChild ? this : nullptr;
 		ProcessHandle h = m_session.InternalRunProcess(info, true, parent, true);
diff --git a/Engine/Source/Programs/UnrealBuildAccelerator/Common/Private/UbaSession.cpp b/Engine/Source/Programs/UnrealBuildAccelerator/Common/Private/UbaSession.cpp
--- a/Engine/Source/Programs/UnrealBuildAccelerator/Common/Private/UbaSession.cpp
+++ b/Engine/Source/Programs/UnrealBuildAccelerator/Common/Private/UbaSession.cpp
@@ -942,7 +942,24 @@ namespace uba
 		CheckDirectoryTableOverflow(dirTable);
 		return dirTable.m_memorySize;
 	}
-	
+
+	VfsHandle Session::RegisterVfs(const void* vfsData, uba::u64 vfsDataSize)
+	{
+		CasKey key = ToCasKey(CasKeyHasher().Update(vfsData, vfsDataSize), false);
+		VfsHandle vfsHandle = key.a;
+
+		SCOPED_WRITE_LOCK(m_vfsLookupLock, lock);
+		auto insres = m_vfsLookup.try_emplace(vfsHandle);
+		VfsEntry& entry = insres.first->second;
+		if (!insres.second)
+		{
+			UBA_ASSERT(memcmp(entry.memory.data(), vfsData, vfsDataSize) == 0);
+			return vfsHandle;
+		}
+		PopulateVfsEntry(entry, vfsData, vfsDataSize);
+		return vfsHandle;
+	}
+
 	void Session::CheckDirectoryTableOverflow(DirectoryTable& dirTable)
 	{
 		if (dirTable.m_memorySize < DirTableMemSize - 1024*1024)
@@ -2306,39 +2323,57 @@ namespace uba
 				}
 				else
 				{
+					FileAccessor destinationFile(m_logger, file.name.c_str());
 
-					// Seems like best combo (for windows at least) is to use writes with overlap and max 16 at the same time.
-					// On one machine we get twice as fast without overlap if no bottleneck. On another machine (ntfs compression on) we get twice as slow without overlap
-					// Both machines behaves well with overlap AND bottleneck. Both machine are 128 logical core thread rippers.
-					bool useFileMapForWrite = fileSize > 0; // ::CreateFileMappingW does not work for zero-length files
-					bool useOverlap = false;//fileSize > 8 * 1024 * 1024;
-
+					VfsHandle vfsHandle = process.GetStartInfo().vfsHandle;
+					if (vfsHandle && StringView(file.name).EndsWith(TC(".d")))
+					{
+						// Need to turn paths back into local paths
 
-					u32 attributes = DefaultAttributes();
-					if (useOverlap)
-						attributes |= FILE_FLAG_OVERLAPPED;
+						if (!destinationFile.CreateWrite(false, DefaultAttributes(), 0, m_tempPath.data))
+							return false;
 
-					FileAccessor destinationFile(m_logger, file.name.c_str());
+						MemoryBlock block(5*1024*1024);
+						if (!FormatClangDepsFile(vfsHandle, block, mem, fileSize, file.name.c_str()))
+							return false;
 
-					if (useFileMapForWrite)
-					{
-						if (!destinationFile.CreateMemoryWrite(false, attributes, fileSize, m_tempPath.data))
+						if (!destinationFile.Write(block.memory, block.writtenSize))
 							return false;
-						MapMemoryCopy(destinationFile.GetData(), mem, fileSize);
 					}
 					else
 					{
-						if (!destinationFile.CreateWrite(false, attributes, fileSize, m_tempPath.data))
-							return false;
 
-						#if PLATFORM_WINDOWS
-						//shouldBottleneck = fileSize > 64 * 1024 * 1024;
-						//if (shouldBottleneck)
-						//	bottleneck.Enter();
-						#endif
+						// Seems like best combo (for windows at least) is to use writes with overlap and max 16 at the same time.
+						// On one machine we get twice as fast without overlap if no bottleneck. On another machine (ntfs compression on) we get twice as slow without overlap
+						// Both machines behaves well with overlap AND bottleneck. Both machine are 128 logical core thread rippers.
+						bool useFileMapForWrite = fileSize > 0; // ::CreateFileMappingW does not work for zero-length files
+						bool useOverlap = false;//fileSize > 8 * 1024 * 1024;
 
-						if (!destinationFile.Write(mem, fileSize))
-							return false;
+
+						u32 attributes = DefaultAttributes();
+						if (useOverlap)
+							attributes |= FILE_FLAG_OVERLAPPED;
+
+						if (useFileMapForWrite)
+						{
+							if (!destinationFile.CreateMemoryWrite(false, attributes, fileSize, m_tempPath.data))
+								return false;
+							MapMemoryCopy(destinationFile.GetData(), mem, fileSize);
+						}
+						else
+						{
+							if (!destinationFile.CreateWrite(false, attributes, fileSize, m_tempPath.data))
+								return false;
+
+							#if PLATFORM_WINDOWS
+							//shouldBottleneck = fileSize > 64 * 1024 * 1024;
+							//if (shouldBottleneck)
+							//	bottleneck.Enter();
+							#endif
+
+							if (!destinationFile.Write(mem, fileSize))
+								return false;
+						}
 					}
 
 					if (u64 time = file.lastWriteTime)
@@ -2970,6 +3005,24 @@ namespace uba
 		return m_cpuLoad;
 	}
 
+	void Session::PopulateVfsEntry(VfsEntry& entry, const void* vfsData, uba::u64 vfsDataSize)
+	{
+		entry.memory.resize(vfsDataSize);
+		memcpy(entry.memory.data(), vfsData, vfsDataSize);
+
+		BinaryReader reader((const u8*)vfsData, 0, vfsDataSize);
+		while (reader.GetLeft())
+		{
+			StringBuffer<> temp;
+			reader.ReadString(temp);
+			entry.roots.RegisterRoot(m_logger, temp.data);
+			reader.ReadString(temp.Clear());
+			char ansi[512];
+			temp.Parse(ansi, 512);
+			entry.localsAnsi.emplace_back(ansi);
+		}
+	}
+
 	bool Session::ExtractSymbolsFromObjectFile(const CloseFileMessage& msg, const tchar* fileName, u64 fileSize)
 	{
 		if (!msg.mappingHandle)
@@ -3039,6 +3092,28 @@ namespace uba
 		return true;
 	}
 
+	bool Session::FormatClangDepsFile(VfsHandle vfsHandle, MemoryBlock& destData, const void* sourceData, u64 sourceSize, const tchar* hint)
+	{
+		SCOPED_READ_LOCK(m_vfsLookupLock, lock);
+		auto findIt = m_vfsLookup.find(vfsHandle);
+		UBA_ASSERT(findIt != m_vfsLookup.end());
+		VfsEntry& vfsEntry = findIt->second;
+		lock.Leave();
+
+		auto handleString = [&](const char* str, u64 strLen, u32 rootPos)
+			{
+				if (rootPos == ~0u)
+				{
+					memcpy(destData.Allocate(strLen, 1, TC("")), str, strLen);
+					return;
+				}
+				auto& path = vfsEntry.localsAnsi[(*str - RootPaths::RootStartByte)/3];
+				memcpy(destData.Allocate(path.size(), 1, TC("")), path.c_str(), path.size());
+			};
+
+		return vfsEntry.roots.NormalizeString<char>(m_logger, (const char*)sourceData, sourceSize, handleString, true, hint);
+	}
+
 	void Session::ThreadTraceLoop()
 	{
 		while (true)
diff --git a/Engine/Source/Programs/UnrealBuildAccelerator/Common/Private/UbaSessionClient.cpp b/Engine/Source/Programs/UnrealBuildAccelerator/Common/Private/UbaSessionClient.cpp
--- a/Engine/Source/Programs/UnrealBuildAccelerator/Common/Private/UbaSessionClient.cpp
+++ b/Engine/Source/Programs/UnrealBuildAccelerator/Common/Private/UbaSessionClient.cpp
@@ -310,6 +310,22 @@ namespace uba
 				}, 1, TC("EnsureApp"));
 		}
 
+		if (VfsHandle vfsHandle = startInfo.vfsHandle)
+		{
+			SCOPED_WRITE_LOCK(m_vfsLookupLock, vfsLock);
+			auto insres2 = m_vfsLookup.try_emplace(vfsHandle);
+			if (insres2.second)
+			{
+				StackBinaryWriter<128> writer;
+				NetworkMessage msg(m_client, ServiceId, SessionMessageType_GetVfs, writer);
+				writer.WriteU64(vfsHandle);
+				StackBinaryReader<8*1024> reader;
+				if (!msg.Send(reader, m_stats.getApplicationMsg))
+					return false;
+				PopulateVfsEntry(insres2.first->second, reader.GetPositionData(), reader.GetLeft());
+			}
+		}
+
 		while (handledCount < modules.size())
 			m_client.DoWork();
 
diff --git a/Engine/Source/Programs/UnrealBuildAccelerator/Common/Private/UbaSessionServer.cpp b/Engine/Source/Programs/UnrealBuildAccelerator/Common/Private/UbaSessionServer.cpp
--- a/Engine/Source/Programs/UnrealBuildAccelerator/Common/Private/UbaSessionServer.cpp
+++ b/Engine/Source/Programs/UnrealBuildAccelerator/Common/Private/UbaSessionServer.cpp
@@ -987,17 +987,26 @@ namespace uba
 			if (shouldWriteToDisk)
 			{
 				bool writeCompressed = false;
-				if (m_storeObjFilesCompressed)
+
+				ProcessHandle h = GetProcess(processId);
+				if (!h.IsValid())
+					return m_logger.Error(TC("Failed to find process for id %u when receiving SendFileToServer message"), processId);
+
+				VfsHandle vfsHandle = h.GetStartInfo().vfsHandle;
+
+				Storage::FormattingFunc formattingFunc;
+				if (vfsHandle && destination.EndsWith(TC(".d")))
 				{
-					SCOPED_READ_LOCK(m_processesLock, lock);
-					auto findIt = m_processes.find(processId);
-					if (findIt != m_processes.end())
-					{
-						auto& process = *(RemoteProcess*)findIt->second.m_process;
-						writeCompressed = process.m_startInfo.rules->StoreFileCompressed(destination);
-					}
+					formattingFunc = [this, vfsHandle](MemoryBlock& destData, const void* sourceData, u64 sourceSize, const tchar* hint)
+						{
+							return FormatClangDepsFile(vfsHandle, destData, sourceData, sourceSize, hint);
+						};
 				}
-				success = m_storage.CopyOrLink(casKey, destination.data, attributes, writeCompressed);
+				else if (m_storeObjFilesCompressed)
+				{
+					writeCompressed = h.GetStartInfo().rules->StoreFileCompressed(destination);
+				}
+				success = m_storage.CopyOrLink(casKey, destination.data, attributes, writeCompressed, formattingFunc);
 				if (!success)
 					m_logger.Error(TC("Failed to copy cas from %s to %s (%s)"), CasKeyString(casKey).str, destination.data, GetProcessDescription(processId).c_str());
 			}
@@ -1415,6 +1424,19 @@ namespace uba
 		return true;
 	}
 
+	bool SessionServer::HandleGetVfs(const ConnectionInfo& connectionInfo, BinaryReader& reader, BinaryWriter& writer)
+	{
+		VfsHandle vfsHandle = reader.ReadU64();
+		SCOPED_READ_LOCK(m_vfsLookupLock, lock);
+		auto findIt = m_vfsLookup.find(vfsHandle);
+		if (findIt == m_vfsLookup.end())
+			return m_logger.Error(TC("Remote process use vfs handle %llu which does not exist"), vfsHandle);
+		VfsEntry& vfsEntry = findIt->second;
+		lock.Leave();
+		writer.WriteBytes(vfsEntry.memory.data(), vfsEntry.memory.size());
+		return true;
+	}
+
 	bool SessionServer::HandleVirtualAllocFailed(const ConnectionInfo& connectionInfo, BinaryReader& reader, BinaryWriter& writer)
 	{
 		m_logger.Error(TC("VIRTUAL ALLOC FAILING ON REMOTE MACHINE %s !"), GuidToString(connectionInfo.GetUid()).str);
diff --git a/Engine/Source/Programs/UnrealBuildAccelerator/Common/Private/UbaStorage.cpp b/Engine/Source/Programs/UnrealBuildAccelerator/Common/Private/UbaStorage.cpp
--- a/Engine/Source/Programs/UnrealBuildAccelerator/Common/Private/UbaStorage.cpp
+++ b/Engine/Source/Programs/UnrealBuildAccelerator/Common/Private/UbaStorage.cpp
@@ -2426,7 +2426,7 @@ namespace uba
 		return out != CasKeyZero;
 	}
 
-	bool StorageImpl::CopyOrLink(const CasKey& casKey, const tchar* destination, u32 fileAttributes, bool writeCompressed)
+	bool StorageImpl::CopyOrLink(const CasKey& casKey, const tchar* destination, u32 fileAttributes, bool writeCompressed, const FormattingFunc& formattingFunc)
 	{
 		UBA_ASSERT(casKey != CasKeyZero);
 		UBA_ASSERT(fileAttributes);
@@ -2544,7 +2544,27 @@ namespace uba
 
 				u64 sizeOnDisk = decompressedSize;
 
-				if (writeCompressed)
+				if (formattingFunc)
+				{
+					UBA_ASSERT(!writeCompressed);
+					u8* slot = m_bufferSlots.Pop();
+					auto _ = MakeGuard([&](){ m_bufferSlots.Push(slot); });
+
+					if (!DecompressMemoryToMemory(readData, slot, decompressedSize, casFile.data, destination))
+						return false;
+
+					if (!destinationFile.CreateWrite(false, DefaultAttributes(), 0, m_tempPath.data))
+						return false;
+
+					MemoryBlock block(5*1024*1024);
+
+					if (!formattingFunc(block, slot, decompressedSize, destination))
+						return false;
+
+					if (!destinationFile.Write(block.memory, block.writtenSize))
+						return false;
+				}
+				else if (writeCompressed)
 				{
 					u64 compressedFileSize = mappedView.size;
 					if (!mappedView.memory)
diff --git a/Engine/Source/Programs/UnrealBuildAccelerator/Common/Public/UbaNetwork.h b/Engine/Source/Programs/UnrealBuildAccelerator/Common/Public/UbaNetwork.h
--- a/Engine/Source/Programs/UnrealBuildAccelerator/Common/Public/UbaNetwork.h
+++ b/Engine/Source/Programs/UnrealBuildAccelerator/Common/Public/UbaNetwork.h
@@ -16,7 +16,7 @@ namespace uba
 	
 	static constexpr u32 SystemNetworkVersion = 1339;
 	static constexpr u32 StorageNetworkVersion = 4;
-	static constexpr u32 SessionNetworkVersion = 38;
+	static constexpr u32 SessionNetworkVersion = 39;
 	static constexpr u32 CacheNetworkVersion = 5;
 
 	// Messages used over network between client/server (system, storage and session)
@@ -77,6 +77,7 @@ namespace uba
 		UBA_SESSION_MESSAGE(ProcessInputs) \
 		UBA_SESSION_MESSAGE(ProcessFinished) \
 		UBA_SESSION_MESSAGE(ProcessReturned) \
+		UBA_SESSION_MESSAGE(GetVfs) \
 		UBA_SESSION_MESSAGE(VirtualAllocFailed) \
 		UBA_SESSION_MESSAGE(GetTraceInformation) \
 		UBA_SESSION_MESSAGE(Ping) \
diff --git a/Engine/Source/Programs/UnrealBuildAccelerator/Common/Public/UbaProcessStartInfo.h b/Engine/Source/Programs/UnrealBuildAccelerator/Common/Public/UbaProcessStartInfo.h
--- a/Engine/Source/Programs/UnrealBuildAccelerator/Common/Public/UbaProcessStartInfo.h
+++ b/Engine/Source/Programs/UnrealBuildAccelerator/Common/Public/UbaProcessStartInfo.h
@@ -8,6 +8,7 @@ namespace uba
 {
 	class ApplicationRules;
 	class ProcessHandle;
+	using VfsHandle = u64;
 
 	struct ProcessStartInfo
 	{
@@ -29,6 +30,8 @@ namespace uba
 		const tchar* breadcrumbs = TC("");		// If not empty, write additional information to the UBA trace file
 		bool startSuspended = false;			// Start process suspended.. a bit internal atm and only supported on windows
 
+		VfsHandle vfsHandle = 0;				// Handle if using vfs. Handle is generated through Session::RegisterVfs
+
 		using LogLineCallback = void(void* userData, const tchar* line, u32 length, LogEntryType type);
 		LogLineCallback* logLineFunc = nullptr;	// Callback for when log entries happens
 		void* logLineUserData = nullptr;		// User data provided to logLine callback
diff --git a/Engine/Source/Programs/UnrealBuildAccelerator/Common/Public/UbaProcessStartInfoHolder.h b/Engine/Source/Programs/UnrealBuildAccelerator/Common/Public/UbaProcessStartInfoHolder.h
--- a/Engine/Source/Programs/UnrealBuildAccelerator/Common/Public/UbaProcessStartInfoHolder.h
+++ b/Engine/Source/Programs/UnrealBuildAccelerator/Common/Public/UbaProcessStartInfoHolder.h
@@ -65,6 +65,7 @@ namespace uba
 			writer.WriteU32(*(u32*)&weight);
 			writer.WriteBool(trackInputs);
 			writer.WriteBool(writeOutputFilesOnFail);
+			writer.WriteU64(vfsHandle);
 		}
 
 		void Read(BinaryReader& reader)
@@ -83,6 +84,7 @@ namespace uba
 			
 			trackInputs = reader.ReadBool();
 			writeOutputFilesOnFail = reader.ReadBool();
+			vfsHandle = reader.ReadU64();
 
 			description = descriptionStr.c_str();
 			application = applicationStr.c_str();
diff --git a/Engine/Source/Programs/UnrealBuildAccelerator/Common/Public/UbaSession.h b/Engine/Source/Programs/UnrealBuildAccelerator/Common/Public/UbaSession.h
--- a/Engine/Source/Programs/UnrealBuildAccelerator/Common/Public/UbaSession.h
+++ b/Engine/Source/Programs/UnrealBuildAccelerator/Common/Public/UbaSession.h
@@ -7,6 +7,7 @@
 #include "UbaFileMapping.h"
 #include "UbaProcessHandle.h"
 #include "UbaProcessStartInfo.h"
+#include "UbaRootPaths.h"
 #include "UbaSessionCreateInfo.h"
 #include "UbaStats.h"
 #include "UbaTrace.h"
@@ -60,6 +61,7 @@ namespace uba
 		bool RegisterNewFile(const tchar* filePath); // Tell uba a new file on disk has been added by some other system while session is running
 		bool RegisterVirtualFile(const tchar* filePath, const tchar* sourceFile, u64 sourceOffset, u64 sourceSize);
 		void RegisterDeleteFile(const tchar* filePath); // Tell uba a file on disk has been deleted by some other system while session is running
+		VfsHandle RegisterVfs(const void* vfsData, uba::u64 vfsDataSize);
 
 		using CustomServiceFunction = Function<u32(Process& handle, const void* recv, u32 recvSize, void* send, u32 sendCapacity)>;
 		void RegisterCustomService(CustomServiceFunction&& function); // Register a custom service (that can be communicated with from the remote agents)
@@ -161,7 +163,11 @@ namespace uba
 		void WriteSummary(BinaryWriter& writer, const Function<void(Logger& logger)>& summaryFunc);
 		float UpdateCpuLoad();
 
+		struct VfsEntry { Vector<u8> memory; RootPaths roots; Vector<std::string> localsAnsi; };
+		void PopulateVfsEntry(VfsEntry& entry, const void* vfsData, uba::u64 vfsDataSize);
+
 		bool ExtractSymbolsFromObjectFile(const CloseFileMessage& msg, const tchar* fileName, u64 fileSize);
+		bool FormatClangDepsFile(VfsHandle vfsHandle, MemoryBlock& destData, const void* sourceData, u64 sourceSize, const tchar* hint);
 		
 		void ThreadTraceLoop();
 		virtual void TraceSessionUpdate();
@@ -242,6 +248,9 @@ namespace uba
 		ReaderWriterLock m_virtualSourceFilesLock;
 		UnorderedMap<StringKey, FileMappingHandle> m_virtualSourceFiles;
 
+		ReaderWriterLock m_vfsLookupLock;
+		UnorderedMap<VfsHandle, VfsEntry> m_vfsLookup;
+
 		u32 m_wantsOnCloseIdCounter = 1;
 
 		SessionStats m_stats;
diff --git a/Engine/Source/Programs/UnrealBuildAccelerator/Common/Public/UbaStorage.h b/Engine/Source/Programs/UnrealBuildAccelerator/Common/Public/UbaStorage.h
--- a/Engine/Source/Programs/UnrealBuildAccelerator/Common/Public/UbaStorage.h
+++ b/Engine/Source/Programs/UnrealBuildAccelerator/Common/Public/UbaStorage.h
@@ -49,7 +49,9 @@ namespace uba
 		virtual bool DropCasFile(const CasKey& casKey, bool forceDelete, const tchar* hint) = 0;
 		virtual bool ReportBadCasFile(const CasKey& casKey) = 0;
 		virtual bool CalculateCasKey(CasKey& out, const tchar* fileName) = 0;
-		virtual bool CopyOrLink(const CasKey& casKey, const tchar* destination, u32 fileAttributes, bool writeCompressed = false) = 0;
+
+		using FormattingFunc = Function<bool(MemoryBlock& destData, const void* sourceData, u64 sourceSize, const tchar* hint)>;
+		virtual bool CopyOrLink(const CasKey& casKey, const tchar* destination, u32 fileAttributes, bool writeCompressed = false, const FormattingFunc& formattingFunc = {}) = 0;
 		virtual bool FakeCopy(const CasKey& casKey, const tchar* destination, u64 size = 0, u64 lastWritten = 0, bool deleteExisting = true) = 0;
 #if !UBA_USE_SPARSEFILE
 		virtual bool GetCasFileName(StringBufferBase& out, const CasKey& casKey) = 0;
@@ -135,7 +137,7 @@ namespace uba
 		virtual bool DropCasFile(const CasKey& casKey, bool forceDelete, const tchar* hint) override;
 		virtual bool ReportBadCasFile(const CasKey& casKey) override;
 		virtual bool CalculateCasKey(CasKey& out, const tchar* fileName) override;
-		virtual bool CopyOrLink(const CasKey& casKey, const tchar* destination, u32 fileAttributes, bool writeCompressed = false) override;
+		virtual bool CopyOrLink(const CasKey& casKey, const tchar* destination, u32 fileAttributes, bool writeCompressed = false, const FormattingFunc& formattingFunc = {}) override;
 		virtual bool FakeCopy(const CasKey& casKey, const tchar* destination, u64 size = 0, u64 lastWritten = 0, bool deleteExisting = true) override;
 #if !UBA_USE_SPARSEFILE
 		virtual bool GetCasFileName(StringBufferBase& out, const CasKey& casKey) override;
diff --git a/Engine/Source/Programs/UnrealBuildAccelerator/Detours/Private/Windows/UbaDetoursFunctionsKernelBase.inl b/Engine/Source/Programs/UnrealBuildAccelerator/Detours/Private/Windows/UbaDetoursFunctionsKernelBase.inl
--- a/Engine/Source/Programs/UnrealBuildAccelerator/Detours/Private/Windows/UbaDetoursFunctionsKernelBase.inl
+++ b/Engine/Source/Programs/UnrealBuildAccelerator/Detours/Private/Windows/UbaDetoursFunctionsKernelBase.inl
@@ -21,7 +21,7 @@ DWORD Local_GetLongPathNameW(LPCWSTR lpszShortPath, LPWSTR lpszLongPath, DWORD c
 		return len + 1;
 	}
 
-	wchar_t* newLongPath = nullptr;
+	const wchar_t* newLongPath = nullptr;
 	DWORD res = 0;
 
 	if (g_runningRemote)
@@ -40,7 +40,7 @@ DWORD Local_GetLongPathNameW(LPCWSTR lpszShortPath, LPWSTR lpszLongPath, DWORD c
 			reader.ReadString(longName);
 		}
 
-		newLongPath = g_memoryBlock.Strdup(longName.data);
+		newLongPath = g_memoryBlock.Strdup(longName).data;
 
 		if (longName.count == 0)
 		{
@@ -86,7 +86,7 @@ LPWSTR Detoured_GetCommandLineW()
 		return str;
 	}
 	DEBUG_LOG_DETOURED(L"GetCommandLineW", L"");
-	return g_virtualCommandLineW;
+	return (LPWSTR)g_virtualCommandLineW;
 }
 
 DWORD Detoured_GetCurrentDirectoryW(DWORD nBufferLength, LPWSTR lpBuffer)
diff --git a/Engine/Source/Programs/UnrealBuildAccelerator/Detours/Private/Windows/UbaDetoursFunctionsNtDll.inl b/Engine/Source/Programs/UnrealBuildAccelerator/Detours/Private/Windows/UbaDetoursFunctionsNtDll.inl
--- a/Engine/Source/Programs/UnrealBuildAccelerator/Detours/Private/Windows/UbaDetoursFunctionsNtDll.inl
+++ b/Engine/Source/Programs/UnrealBuildAccelerator/Detours/Private/Windows/UbaDetoursFunctionsNtDll.inl
@@ -690,6 +690,21 @@ NTSTATUS NTAPI Shared_NtCreateFile(bool IsCreateFunc, PHANDLE hFileHandle, ACCES
 		return res;
 	}
 
+	if (g_vfsEntryCount)
+	{
+		if (fileName.count >= g_vfsMatchingLength && Equals(fileName.data, g_vfsEntries[0].vfs.data, g_vfsMatchingLength, true))
+		{
+			for (u32 i=0, e=g_vfsEntryCount; i!=e; ++i)
+			{
+				VfsEntry& entry = g_vfsEntries[i];
+				if (!fileName.StartsWith(entry.vfs.data))
+					continue;
+				StringBuffer<> temp2(fileName.data + entry.vfs.count);
+				fileName.Clear().Append(entry.local).Append(temp2);
+			}
+		}
+	}
+
 	u32 dirTableOffset = ~u32(0);
 
 	//UBA_ASSERT(CreateDisposition != FILE_SUPERSEDE);
@@ -891,7 +906,7 @@ NTSTATUS NTAPI Shared_NtCreateFile(bool IsCreateFunc, PHANDLE hFileHandle, ACCES
 #endif
 		* hFileHandle = makeListDirectoryHandle(listHandle);
 
-		listHandle->originalName = g_memoryBlock.Strdup(fileName.data);
+		listHandle->originalName = g_memoryBlock.Strdup(fileName).data;
 
 		IoStatusBlock->Information = 1;
 		IoStatusBlock->Pointer = nullptr;
@@ -914,7 +929,7 @@ NTSTATUS NTAPI Shared_NtCreateFile(bool IsCreateFunc, PHANDLE hFileHandle, ACCES
 	if (insres.second)
 	{
 		u64 size = InvalidValue;
-		info.originalName = g_memoryBlock.Strdup(fileName.data);
+		info.originalName = g_memoryBlock.Strdup(fileName).data;
 		info.name = info.originalName;
 		if (!keepInMemory && !isSystemOrTempFile)
 		{
@@ -932,7 +947,7 @@ NTSTATUS NTAPI Shared_NtCreateFile(bool IsCreateFunc, PHANDLE hFileHandle, ACCES
 	else
 	{
 		if (!info.originalName)
-			info.originalName = g_memoryBlock.Strdup(fileName.data);
+			info.originalName = g_memoryBlock.Strdup(fileName).data;
 		if (isWrite)
 		{
 			bool lastWasWrite = IsContentWrite(info.lastDesiredAccess, 0);
@@ -1412,7 +1427,7 @@ NTSTATUS NTAPI Detoured_NtClose(HANDLE handle)
 				auto insres = g_mappedFileTable.m_lookup.try_emplace(fileNameKey);
 				FileInfo& newInfo = insres.first->second;
 				newInfo = fi;
-				newInfo.originalName = g_memoryBlock.Strdup(fo->newName.c_str());
+				newInfo.originalName = g_memoryBlock.Strdup(fo->newName).data;
 				newInfo.name = newInfo.originalName;
 				newInfo.fileNameKey = fileNameKey;
 				UBA_ASSERT(!fo->deleteOnClose);
diff --git a/Engine/Source/Programs/UnrealBuildAccelerator/Detours/Private/Windows/UbaDetoursFunctionsWin.cpp b/Engine/Source/Programs/UnrealBuildAccelerator/Detours/Private/Windows/UbaDetoursFunctionsWin.cpp
--- a/Engine/Source/Programs/UnrealBuildAccelerator/Detours/Private/Windows/UbaDetoursFunctionsWin.cpp
+++ b/Engine/Source/Programs/UnrealBuildAccelerator/Detours/Private/Windows/UbaDetoursFunctionsWin.cpp
@@ -115,7 +115,7 @@ int g_uiLanguage;
 
 StringBuffer<256> g_exeDir;
 
-wchar_t* g_virtualCommandLineW;
+const wchar_t* g_virtualCommandLineW;
 char* g_virtualCommandLineA;
 
 constexpr u32 TrackInputsMemCapacity = 512 * 1024;
@@ -231,6 +231,11 @@ u64 g_memoryFileIndexCounter = ~u64(0) - 1000000; // I really hope this will not
 ObjFilesPreloader g_objFilesPreloader;
 
 
+struct VfsEntry { StringView vfs; StringView local; };
+VfsEntry g_vfsEntries[32];
+u32 g_vfsEntryCount;
+u32 g_vfsMatchingLength;
+
 struct SuppressCreateFileDetourScope
 {
 	SuppressCreateFileDetourScope() { ++t_disallowCreateFileDetour; }
@@ -891,6 +896,36 @@ void Init(const DetoursPayload& payload, u64 startTime)
 		mappedFileTableHandle = (HANDLE)reader.ReadU64();
 		mappedFileTableSize = reader.ReadU32();
 		mappedFileTableCount = reader.ReadU32();
+
+
+		if (u16 vfsSize = reader.ReadU16())
+		{
+			BinaryReader vfsReader(reader.GetPositionData(), 0, vfsSize);
+			while (vfsReader.GetLeft())
+			{
+				u32 index = g_vfsEntryCount++;
+				VfsEntry& vfsEntry = g_vfsEntries[index];
+				StringBuffer<> str;
+				vfsReader.ReadString(str);
+				vfsEntry.vfs = g_memoryBlock.Strdup(str);
+				if (index == 0)
+					g_vfsMatchingLength = vfsEntry.vfs.count;
+				else
+				{
+					u32 shortest = Min(g_vfsMatchingLength, vfsEntry.vfs.count);
+					for (u32 i=0; i!=shortest; ++i)
+					{
+						if (g_vfsEntries[0].vfs.data[i] == vfsEntry.vfs.data[i])
+							continue;
+						shortest = i;
+						break;
+					}
+					g_vfsMatchingLength = shortest;
+				}
+				vfsReader.ReadString(str.Clear());
+				vfsEntry.local = g_memoryBlock.Strdup(str);
+			}
+		}
 		DEBUG_LOG_PIPE(L"Init", L"");
 	}
 
@@ -917,7 +952,7 @@ void Init(const DetoursPayload& payload, u64 startTime)
 		buf.Append(g_virtualApplicationDir);
 		TString realCmdLine(buf.data);
 		realCmdLine += (cmdLine + g_exeDir.count);
-		g_virtualCommandLineW = g_memoryBlock.Strdup(realCmdLine.c_str());
+		g_virtualCommandLineW = g_memoryBlock.Strdup(realCmdLine).data;
 	}
 	//else
 	//	g_virtualCommandLineW = g_memoryBlock.Strdup(cmdLine);
diff --git a/Engine/Source/Programs/UnrealBuildAccelerator/Host/Private/UbaExports.cpp b/Engine/Source/Programs/UnrealBuildAccelerator/Host/Private/UbaExports.cpp
--- a/Engine/Source/Programs/UnrealBuildAccelerator/Host/Private/UbaExports.cpp
+++ b/Engine/Source/Programs/UnrealBuildAccelerator/Host/Private/UbaExports.cpp
@@ -393,6 +393,10 @@ extern "C"
 	{
 		return new uba::ProcessHandle(server->RunProcessRacing(raceAgainstRemoteProcessId));
 	}
+	uba::u64 SessionServer_RegisterVfs(uba::SessionServer* server, const void* vfsData, uba::u64 vfsDataSize)
+	{
+		return server->RegisterVfs(vfsData, vfsDataSize);
+	}
 	void SessionServer_SetMaxRemoteProcessCount(uba::SessionServer* server, uba::u32 count)
 	{
 		return server->SetMaxRemoteProcessCount(count);
@@ -510,6 +514,21 @@ extern "C"
 		return info;
 	}
 
+	uba::ProcessStartInfo* ProcessStartInfo_Create2(const uba::tchar* application, const uba::tchar* arguments, const uba::tchar* workingDir, const uba::tchar* description, uba::u32 priorityClass, uba::u64 vfsHandle, bool trackInputs, const uba::tchar* logFile, ProcessHandle_ExitCallback* exit)
+	{
+		auto info = new uba::ProcessStartInfo();
+		info->application = TStrdup(application);
+		info->arguments = TStrdup(arguments);
+		info->workingDir = TStrdup(workingDir);
+		info->description = TStrdup(description);
+		info->priorityClass = priorityClass;
+		info->trackInputs = trackInputs;
+		info->logFile = TStrdup(logFile);
+		info->exitedFunc = exit;
+		info->vfsHandle = vfsHandle;
+		return info;
+	}
+
 	void ProcessStartInfo_Destroy(uba::ProcessStartInfo* info)
 	{
 		free((void*)info->application);
diff --git a/Engine/Source/Programs/UnrealBuildAccelerator/Host/Public/UbaExports.h b/Engine/Source/Programs/UnrealBuildAccelerator/Host/Public/UbaExports.h
--- a/Engine/Source/Programs/UnrealBuildAccelerator/Host/Public/UbaExports.h
+++ b/Engine/Source/Programs/UnrealBuildAccelerator/Host/Public/UbaExports.h
@@ -72,6 +72,7 @@ extern "C"
 	// ProcessStartInfo
 	using ProcessHandle_ExitCallback = void(void* userData, const uba::ProcessHandle&);
 	UBA_API uba::ProcessStartInfo* ProcessStartInfo_Create(const uba::tchar* application, const uba::tchar* arguments, const uba::tchar* workingDir, const uba::tchar* description, uba::u32 priorityClass, uba::u64 outputStatsThresholdMs, bool trackInputs, const uba::tchar* logFile, ProcessHandle_ExitCallback* exit);
+	UBA_API uba::ProcessStartInfo* ProcessStartInfo_Create2(const uba::tchar* application, const uba::tchar* arguments, const uba::tchar* workingDir, const uba::tchar* description, uba::u32 priorityClass, uba::u64 vfsHandle, bool trackInputs, const uba::tchar* logFile, ProcessHandle_ExitCallback* exit);
 	UBA_API void ProcessStartInfo_Destroy(uba::ProcessStartInfo* info);
 
 	// ProcessHandle
@@ -105,6 +106,7 @@ extern "C"
 	UBA_API uba::ProcessHandle* SessionServer_RunProcess(uba::SessionServer* server, uba::ProcessStartInfo& info, bool async, bool enableDetour);
 	UBA_API uba::ProcessHandle* SessionServer_RunProcessRemote(uba::SessionServer* server, uba::ProcessStartInfo& info, float weight, const void* knownInputs = nullptr, uba::u32 knownInputsCount = 0, bool allowCrossArchitecture = false);
 	UBA_API uba::ProcessHandle* SessionServer_RunProcessRacing(uba::SessionServer* server, uba::u32 raceAgainstRemoteProcessId);
+	UBA_API uba::u64 SessionServer_RegisterVfs(uba::SessionServer* server, const void* vfsData, uba::u64 vfsDataSize);
 
 	UBA_API void SessionServer_SetMaxRemoteProcessCount(uba::SessionServer* server, uba::u32 count);
 	UBA_API void SessionServer_DisableRemoteExecution(uba::SessionServer* server);
