[UBA]
* Renamed vfs logic to roots instead since it can be used for cache system without using vfs

diff --git a/Engine/Source/Programs/UnrealBuildAccelerator/Common/Private/UbaCacheClient.cpp b/Engine/Source/Programs/UnrealBuildAccelerator/Common/Private/UbaCacheClient.cpp
--- a/Engine/Source/Programs/UnrealBuildAccelerator/Common/Private/UbaCacheClient.cpp
+++ b/Engine/Source/Programs/UnrealBuildAccelerator/Common/Private/UbaCacheClient.cpp
@@ -124,6 +124,13 @@ namespace uba
 
 	CacheClient::~CacheClient() = default;
 
+	bool CacheClient::WriteToCache(u32 bucketId, const ProcessStartInfo& info, const u8* inputs, u64 inputsSize, const u8* outputs, u64 outputsSize, const u8* logLines, u64 logLinesSize, u32 processId)
+	{
+		RootPaths rootPaths;
+		m_session.PopulateCacheRoots(rootPaths, info.rootsHandle);
+		return WriteToCache(rootPaths, bucketId, info, inputs, inputsSize, outputs, outputsSize, logLines, logLinesSize, processId);
+	}
+
 	bool CacheClient::WriteToCache(const RootPaths& rootPaths, u32 bucketId, const ProcessStartInfo& info, const u8* inputs, u64 inputsSize, const u8* outputs, u64 outputsSize, const u8* logLines, u64 logLinesSize, u32 processId)
 	{
 		if (!m_connected)
diff --git a/Engine/Source/Programs/UnrealBuildAccelerator/Common/Private/UbaProcess.cpp b/Engine/Source/Programs/UnrealBuildAccelerator/Common/Private/UbaProcess.cpp
--- a/Engine/Source/Programs/UnrealBuildAccelerator/Common/Private/UbaProcess.cpp
+++ b/Engine/Source/Programs/UnrealBuildAccelerator/Common/Private/UbaProcess.cpp
@@ -728,19 +728,19 @@ namespace uba
 		writer.WriteU32(response.mappedFileTableSize);
 		writer.WriteU32(response.mappedFileTableCount);
 
-		if (!m_startInfo.vfsHandle)
+		if (!m_startInfo.rootsHandle)
 		{
 			writer.WriteU16(0);
 			return true;
 		}
 
-		SCOPED_READ_LOCK(m_session.m_vfsLookupLock, lock);
-		auto findIt = m_session.m_vfsLookup.find(m_startInfo.vfsHandle);
-		if (findIt == m_session.m_vfsLookup.end())
-			return m_session.m_logger.Error(TC("Process use vfs handle %llu which does not exist"), m_startInfo.vfsHandle);
-		Session::VfsEntry& vfsEntry = findIt->second;
-		writer.WriteU16(u16(vfsEntry.memory.size()));
-		writer.WriteBytes(vfsEntry.memory.data(), vfsEntry.memory.size());
+		SCOPED_READ_LOCK(m_session.m_rootsLookupLock, lock);
+		auto findIt = m_session.m_rootsLookup.find(m_startInfo.rootsHandle);
+		if (findIt == m_session.m_rootsLookup.end())
+			return m_session.m_logger.Error(TC("Process use roots handle %llu which does not exist"), m_startInfo.rootsHandle);
+		Session::RootsEntry& rootsEntry = findIt->second;
+		writer.WriteU16(u16(rootsEntry.memory.size()));
+		writer.WriteBytes(rootsEntry.memory.data(), rootsEntry.memory.size());
 		return true;
 	}
 	
@@ -1014,7 +1014,7 @@ namespace uba
 		info.logLineUserData = this;
 		info.logLineFunc = [](void* userData, const tchar* line, u32 length, LogEntryType type) { ((ProcessImpl*)userData)->LogLine(false, TString(line, length), type); };
 		info.startSuspended = startSuspended;
-		info.vfsHandle = m_startInfo.vfsHandle;
+		info.rootsHandle = m_startInfo.rootsHandle;
 
 		ProcessImpl* parent = isChild ? this : nullptr;
 		ProcessHandle h = m_session.InternalRunProcess(info, true, parent, true);
diff --git a/Engine/Source/Programs/UnrealBuildAccelerator/Common/Private/UbaSession.cpp b/Engine/Source/Programs/UnrealBuildAccelerator/Common/Private/UbaSession.cpp
--- a/Engine/Source/Programs/UnrealBuildAccelerator/Common/Private/UbaSession.cpp
+++ b/Engine/Source/Programs/UnrealBuildAccelerator/Common/Private/UbaSession.cpp
@@ -943,21 +943,21 @@ namespace uba
 		return dirTable.m_memorySize;
 	}
 
-	VfsHandle Session::RegisterVfs(const void* vfsData, uba::u64 vfsDataSize)
+	RootsHandle Session::RegisterRoots(const void* rootsData, uba::u64 rootsDataSize)
 	{
-		CasKey key = ToCasKey(CasKeyHasher().Update(vfsData, vfsDataSize), false);
-		VfsHandle vfsHandle = key.a;
+		CasKey key = ToCasKey(CasKeyHasher().Update(rootsData, rootsDataSize), false);
+		RootsHandle rootsHandle = key.a;
 
-		SCOPED_WRITE_LOCK(m_vfsLookupLock, lock);
-		auto insres = m_vfsLookup.try_emplace(vfsHandle);
-		VfsEntry& entry = insres.first->second;
+		SCOPED_WRITE_LOCK(m_rootsLookupLock, lock);
+		auto insres = m_rootsLookup.try_emplace(rootsHandle);
+		RootsEntry& entry = insres.first->second;
 		if (!insres.second)
 		{
-			UBA_ASSERT(memcmp(entry.memory.data(), vfsData, vfsDataSize) == 0);
-			return vfsHandle;
+			UBA_ASSERT(memcmp(entry.memory.data(), rootsData, rootsDataSize) == 0);
+			return rootsHandle;
 		}
-		PopulateVfsEntry(entry, vfsData, vfsDataSize);
-		return vfsHandle;
+		PopulateRootsEntry(entry, rootsData, rootsDataSize);
+		return rootsHandle;
 	}
 
 	void Session::CheckDirectoryTableOverflow(DirectoryTable& dirTable)
@@ -1563,6 +1563,28 @@ namespace uba
 		return ++m_processIdCounter;
 	}
 
+	bool Session::PopulateCacheRoots(RootPaths& out, RootsHandle rootsHandle)
+	{
+		if (!rootsHandle)
+			return true;
+		SCOPED_READ_LOCK(m_rootsLookupLock, lock)
+		auto findIt = m_rootsLookup.find(rootsHandle);
+		if (findIt == m_rootsLookup.end())
+			return m_logger.Error(TC("Roots handle %llu does not exist"), rootsHandle);
+
+		RootsEntry& entry = findIt->second;
+		BinaryReader reader((const u8*)entry.memory.data(), 0, entry.memory.size());
+		while (reader.GetLeft())
+		{
+			u8 index = reader.ReadByte();(void)index; // RootPaths index;
+			reader.SkipString();
+			StringBuffer<> rootPath;
+			reader.ReadString(rootPath.Clear());
+			out.RegisterRoot(m_logger, rootPath.data, true, index);
+		}
+		return true;
+	}
+
 	void Session::ProcessAdded(Process& process, u32 sessionId)
 	{
 		u32 processId = process.GetId();
@@ -2325,8 +2347,8 @@ namespace uba
 				{
 					FileAccessor destinationFile(m_logger, file.name.c_str());
 
-					VfsHandle vfsHandle = process.GetStartInfo().vfsHandle;
-					if (vfsHandle && StringView(file.name).EndsWith(TC(".d")))
+					RootsHandle rootsHandle = process.GetStartInfo().rootsHandle;
+					if (rootsHandle && StringView(file.name).EndsWith(TC(".d")))
 					{
 						// Need to turn paths back into local paths
 
@@ -2334,7 +2356,7 @@ namespace uba
 							return false;
 
 						MemoryBlock block(5*1024*1024);
-						if (!FormatClangDepsFile(vfsHandle, block, mem, fileSize, file.name.c_str()))
+						if (!FormatClangDepsFile(rootsHandle, block, mem, fileSize, file.name.c_str()))
 							return false;
 
 						if (!destinationFile.Write(block.memory, block.writtenSize))
@@ -3005,20 +3027,29 @@ namespace uba
 		return m_cpuLoad;
 	}
 
-	void Session::PopulateVfsEntry(VfsEntry& entry, const void* vfsData, uba::u64 vfsDataSize)
+	void Session::PopulateRootsEntry(RootsEntry& entry, const void* rootsData, uba::u64 rootsDataSize)
 	{
-		entry.memory.resize(vfsDataSize);
-		memcpy(entry.memory.data(), vfsData, vfsDataSize);
+		entry.memory.resize(rootsDataSize);
+		memcpy(entry.memory.data(), rootsData, rootsDataSize);
 
-		BinaryReader reader((const u8*)vfsData, 0, vfsDataSize);
+		BinaryReader reader((const u8*)rootsData, 0, rootsDataSize);
 		while (reader.GetLeft())
 		{
+			u8 index = reader.ReadByte();(void)index; // RootPaths index;
 			StringBuffer<> temp;
 			reader.ReadString(temp);
 			entry.roots.RegisterRoot(m_logger, temp.data);
 			reader.ReadString(temp.Clear());
 			char ansi[512];
-			temp.Parse(ansi, 512);
+			u32 ansiPos = 0;
+			for (u32 i=0;i!=temp.count;++i)
+			{
+				char c = (char)temp[i];
+				if (c == ' ')
+					ansi[ansiPos++] = '\\';
+				ansi[ansiPos++] = c;
+			}
+			ansi[ansiPos] = 0;
 			entry.localsAnsi.emplace_back(ansi);
 		}
 	}
@@ -3092,12 +3123,12 @@ namespace uba
 		return true;
 	}
 
-	bool Session::FormatClangDepsFile(VfsHandle vfsHandle, MemoryBlock& destData, const void* sourceData, u64 sourceSize, const tchar* hint)
+	bool Session::FormatClangDepsFile(RootsHandle rootsHandle, MemoryBlock& destData, const void* sourceData, u64 sourceSize, const tchar* hint)
 	{
-		SCOPED_READ_LOCK(m_vfsLookupLock, lock);
-		auto findIt = m_vfsLookup.find(vfsHandle);
-		UBA_ASSERT(findIt != m_vfsLookup.end());
-		VfsEntry& vfsEntry = findIt->second;
+		SCOPED_READ_LOCK(m_rootsLookupLock, lock);
+		auto findIt = m_rootsLookup.find(rootsHandle);
+		UBA_ASSERT(findIt != m_rootsLookup.end());
+		RootsEntry& rootsEntry = findIt->second;
 		lock.Leave();
 
 		auto handleString = [&](const char* str, u64 strLen, u32 rootPos)
@@ -3107,11 +3138,11 @@ namespace uba
 					memcpy(destData.Allocate(strLen, 1, TC("")), str, strLen);
 					return;
 				}
-				auto& path = vfsEntry.localsAnsi[(*str - RootPaths::RootStartByte)/3];
+				auto& path = rootsEntry.localsAnsi[(*str - RootPaths::RootStartByte)/3];
 				memcpy(destData.Allocate(path.size(), 1, TC("")), path.c_str(), path.size());
 			};
 
-		return vfsEntry.roots.NormalizeString<char>(m_logger, (const char*)sourceData, sourceSize, handleString, true, hint);
+		return rootsEntry.roots.NormalizeString<char>(m_logger, (const char*)sourceData, sourceSize, handleString, true, hint);
 	}
 
 	void Session::ThreadTraceLoop()
diff --git a/Engine/Source/Programs/UnrealBuildAccelerator/Common/Private/UbaSessionClient.cpp b/Engine/Source/Programs/UnrealBuildAccelerator/Common/Private/UbaSessionClient.cpp
--- a/Engine/Source/Programs/UnrealBuildAccelerator/Common/Private/UbaSessionClient.cpp
+++ b/Engine/Source/Programs/UnrealBuildAccelerator/Common/Private/UbaSessionClient.cpp
@@ -310,19 +310,19 @@ namespace uba
 				}, 1, TC("EnsureApp"));
 		}
 
-		if (VfsHandle vfsHandle = startInfo.vfsHandle)
+		if (RootsHandle rootsHandle = startInfo.rootsHandle)
 		{
-			SCOPED_WRITE_LOCK(m_vfsLookupLock, vfsLock);
-			auto insres2 = m_vfsLookup.try_emplace(vfsHandle);
+			SCOPED_WRITE_LOCK(m_rootsLookupLock, rootsLock);
+			auto insres2 = m_rootsLookup.try_emplace(rootsHandle);
 			if (insres2.second)
 			{
 				StackBinaryWriter<128> writer;
-				NetworkMessage msg(m_client, ServiceId, SessionMessageType_GetVfs, writer);
-				writer.WriteU64(vfsHandle);
+				NetworkMessage msg(m_client, ServiceId, SessionMessageType_GetRoots, writer);
+				writer.WriteU64(rootsHandle);
 				StackBinaryReader<8*1024> reader;
 				if (!msg.Send(reader, m_stats.getApplicationMsg))
 					return false;
-				PopulateVfsEntry(insres2.first->second, reader.GetPositionData(), reader.GetLeft());
+				PopulateRootsEntry(insres2.first->second, reader.GetPositionData(), reader.GetLeft());
 			}
 		}
 
diff --git a/Engine/Source/Programs/UnrealBuildAccelerator/Common/Private/UbaSessionServer.cpp b/Engine/Source/Programs/UnrealBuildAccelerator/Common/Private/UbaSessionServer.cpp
--- a/Engine/Source/Programs/UnrealBuildAccelerator/Common/Private/UbaSessionServer.cpp
+++ b/Engine/Source/Programs/UnrealBuildAccelerator/Common/Private/UbaSessionServer.cpp
@@ -992,14 +992,14 @@ namespace uba
 				if (!h.IsValid())
 					return m_logger.Error(TC("Failed to find process for id %u when receiving SendFileToServer message"), processId);
 
-				VfsHandle vfsHandle = h.GetStartInfo().vfsHandle;
+				RootsHandle rootsHandle = h.GetStartInfo().rootsHandle;
 
 				Storage::FormattingFunc formattingFunc;
-				if (vfsHandle && destination.EndsWith(TC(".d")))
+				if (rootsHandle && destination.EndsWith(TC(".d")))
 				{
-					formattingFunc = [this, vfsHandle](MemoryBlock& destData, const void* sourceData, u64 sourceSize, const tchar* hint)
+					formattingFunc = [this, rootsHandle](MemoryBlock& destData, const void* sourceData, u64 sourceSize, const tchar* hint)
 						{
-							return FormatClangDepsFile(vfsHandle, destData, sourceData, sourceSize, hint);
+							return FormatClangDepsFile(rootsHandle, destData, sourceData, sourceSize, hint);
 						};
 				}
 				else if (m_storeObjFilesCompressed)
@@ -1424,16 +1424,16 @@ namespace uba
 		return true;
 	}
 
-	bool SessionServer::HandleGetVfs(const ConnectionInfo& connectionInfo, BinaryReader& reader, BinaryWriter& writer)
+	bool SessionServer::HandleGetRoots(const ConnectionInfo& connectionInfo, BinaryReader& reader, BinaryWriter& writer)
 	{
-		VfsHandle vfsHandle = reader.ReadU64();
-		SCOPED_READ_LOCK(m_vfsLookupLock, lock);
-		auto findIt = m_vfsLookup.find(vfsHandle);
-		if (findIt == m_vfsLookup.end())
-			return m_logger.Error(TC("Remote process use vfs handle %llu which does not exist"), vfsHandle);
-		VfsEntry& vfsEntry = findIt->second;
+		RootsHandle rootsHandle = reader.ReadU64();
+		SCOPED_READ_LOCK(m_rootsLookupLock, lock);
+		auto findIt = m_rootsLookup.find(rootsHandle);
+		if (findIt == m_rootsLookup.end())
+			return m_logger.Error(TC("Remote process use roots handle %llu which does not exist"), rootsHandle);
+		RootsEntry& rootsEntry = findIt->second;
 		lock.Leave();
-		writer.WriteBytes(vfsEntry.memory.data(), vfsEntry.memory.size());
+		writer.WriteBytes(rootsEntry.memory.data(), rootsEntry.memory.size());
 		return true;
 	}
 
diff --git a/Engine/Source/Programs/UnrealBuildAccelerator/Common/Public/UbaCacheClient.h b/Engine/Source/Programs/UnrealBuildAccelerator/Common/Public/UbaCacheClient.h
--- a/Engine/Source/Programs/UnrealBuildAccelerator/Common/Public/UbaCacheClient.h
+++ b/Engine/Source/Programs/UnrealBuildAccelerator/Common/Public/UbaCacheClient.h
@@ -50,6 +50,7 @@ namespace uba
 		CacheClient(const CacheClientCreateInfo& info);
 		~CacheClient();
 
+		bool WriteToCache(u32 bucketId, const ProcessStartInfo& info, const u8* inputs, u64 inputsSize, const u8* outputs, u64 outputsSize, const u8* logLines, u64 logLinesSize, u32 processId = 0);
 		bool WriteToCache(const RootPaths& rootPaths, u32 bucketId, const ProcessStartInfo& info, const u8* inputs, u64 inputsSize, const u8* outputs, u64 outputsSize, const u8* logLines, u64 logLinesSize, u32 processId = 0);
 		bool FetchFromCache(CacheResult& outResult, const RootPaths& rootPaths, u32 bucketId, const ProcessStartInfo& info);
 		bool RequestServerShutdown(const tchar* reason);
diff --git a/Engine/Source/Programs/UnrealBuildAccelerator/Common/Public/UbaNetwork.h b/Engine/Source/Programs/UnrealBuildAccelerator/Common/Public/UbaNetwork.h
--- a/Engine/Source/Programs/UnrealBuildAccelerator/Common/Public/UbaNetwork.h
+++ b/Engine/Source/Programs/UnrealBuildAccelerator/Common/Public/UbaNetwork.h
@@ -77,7 +77,7 @@ namespace uba
 		UBA_SESSION_MESSAGE(ProcessInputs) \
 		UBA_SESSION_MESSAGE(ProcessFinished) \
 		UBA_SESSION_MESSAGE(ProcessReturned) \
-		UBA_SESSION_MESSAGE(GetVfs) \
+		UBA_SESSION_MESSAGE(GetRoots) \
 		UBA_SESSION_MESSAGE(VirtualAllocFailed) \
 		UBA_SESSION_MESSAGE(GetTraceInformation) \
 		UBA_SESSION_MESSAGE(Ping) \
diff --git a/Engine/Source/Programs/UnrealBuildAccelerator/Common/Public/UbaProcessStartInfo.h b/Engine/Source/Programs/UnrealBuildAccelerator/Common/Public/UbaProcessStartInfo.h
--- a/Engine/Source/Programs/UnrealBuildAccelerator/Common/Public/UbaProcessStartInfo.h
+++ b/Engine/Source/Programs/UnrealBuildAccelerator/Common/Public/UbaProcessStartInfo.h
@@ -8,7 +8,7 @@ namespace uba
 {
 	class ApplicationRules;
 	class ProcessHandle;
-	using VfsHandle = u64;
+	using RootsHandle = u64;
 
 	struct ProcessStartInfo
 	{
@@ -30,7 +30,7 @@ namespace uba
 		const tchar* breadcrumbs = TC("");		// If not empty, write additional information to the UBA trace file
 		bool startSuspended = false;			// Start process suspended.. a bit internal atm and only supported on windows
 
-		VfsHandle vfsHandle = 0;				// Handle if using vfs. Handle is generated through Session::RegisterVfs
+		RootsHandle rootsHandle = 0;			// Handle if using vfs. Handle is generated through Session::RegisterRoots
 
 		using LogLineCallback = void(void* userData, const tchar* line, u32 length, LogEntryType type);
 		LogLineCallback* logLineFunc = nullptr;	// Callback for when log entries happens
diff --git a/Engine/Source/Programs/UnrealBuildAccelerator/Common/Public/UbaProcessStartInfoHolder.h b/Engine/Source/Programs/UnrealBuildAccelerator/Common/Public/UbaProcessStartInfoHolder.h
--- a/Engine/Source/Programs/UnrealBuildAccelerator/Common/Public/UbaProcessStartInfoHolder.h
+++ b/Engine/Source/Programs/UnrealBuildAccelerator/Common/Public/UbaProcessStartInfoHolder.h
@@ -65,7 +65,7 @@ namespace uba
 			writer.WriteU32(*(u32*)&weight);
 			writer.WriteBool(trackInputs);
 			writer.WriteBool(writeOutputFilesOnFail);
-			writer.WriteU64(vfsHandle);
+			writer.WriteU64(rootsHandle);
 		}
 
 		void Read(BinaryReader& reader)
@@ -84,7 +84,7 @@ namespace uba
 			
 			trackInputs = reader.ReadBool();
 			writeOutputFilesOnFail = reader.ReadBool();
-			vfsHandle = reader.ReadU64();
+			rootsHandle = reader.ReadU64();
 
 			description = descriptionStr.c_str();
 			application = applicationStr.c_str();
diff --git a/Engine/Source/Programs/UnrealBuildAccelerator/Common/Public/UbaSession.h b/Engine/Source/Programs/UnrealBuildAccelerator/Common/Public/UbaSession.h
--- a/Engine/Source/Programs/UnrealBuildAccelerator/Common/Public/UbaSession.h
+++ b/Engine/Source/Programs/UnrealBuildAccelerator/Common/Public/UbaSession.h
@@ -61,7 +61,7 @@ namespace uba
 		bool RegisterNewFile(const tchar* filePath); // Tell uba a new file on disk has been added by some other system while session is running
 		bool RegisterVirtualFile(const tchar* filePath, const tchar* sourceFile, u64 sourceOffset, u64 sourceSize);
 		void RegisterDeleteFile(const tchar* filePath); // Tell uba a file on disk has been deleted by some other system while session is running
-		VfsHandle RegisterVfs(const void* vfsData, uba::u64 vfsDataSize);
+		RootsHandle RegisterRoots(const void* rootsData, uba::u64 rootsDataSize);
 
 		using CustomServiceFunction = Function<u32(Process& handle, const void* recv, u32 recvSize, void* send, u32 sendCapacity)>;
 		void RegisterCustomService(CustomServiceFunction&& function); // Register a custom service (that can be communicated with from the remote agents)
@@ -84,6 +84,8 @@ namespace uba
 
 		u32 CreateProcessId();
 
+		bool PopulateCacheRoots(RootPaths& out, RootsHandle rootsHandle);
+
 	protected:
 		Session(const SessionCreateInfo& info, const tchar* logPrefix, bool runningRemote, WorkManager* workManager = nullptr);
 		bool Create(const SessionCreateInfo& info);
@@ -163,11 +165,11 @@ namespace uba
 		void WriteSummary(BinaryWriter& writer, const Function<void(Logger& logger)>& summaryFunc);
 		float UpdateCpuLoad();
 
-		struct VfsEntry { Vector<u8> memory; RootPaths roots; Vector<std::string> localsAnsi; };
-		void PopulateVfsEntry(VfsEntry& entry, const void* vfsData, uba::u64 vfsDataSize);
+		struct RootsEntry { Vector<u8> memory; RootPaths roots; Vector<std::string> localsAnsi; };
+		void PopulateRootsEntry(RootsEntry& entry, const void* rootsData, uba::u64 rootsDataSize);
 
 		bool ExtractSymbolsFromObjectFile(const CloseFileMessage& msg, const tchar* fileName, u64 fileSize);
-		bool FormatClangDepsFile(VfsHandle vfsHandle, MemoryBlock& destData, const void* sourceData, u64 sourceSize, const tchar* hint);
+		bool FormatClangDepsFile(RootsHandle rootsHandle, MemoryBlock& destData, const void* sourceData, u64 sourceSize, const tchar* hint);
 		
 		void ThreadTraceLoop();
 		virtual void TraceSessionUpdate();
@@ -248,8 +250,8 @@ namespace uba
 		ReaderWriterLock m_virtualSourceFilesLock;
 		UnorderedMap<StringKey, FileMappingHandle> m_virtualSourceFiles;
 
-		ReaderWriterLock m_vfsLookupLock;
-		UnorderedMap<VfsHandle, VfsEntry> m_vfsLookup;
+		ReaderWriterLock m_rootsLookupLock;
+		UnorderedMap<RootsHandle, RootsEntry> m_rootsLookup;
 
 		u32 m_wantsOnCloseIdCounter = 1;
 
diff --git a/Engine/Source/Programs/UnrealBuildAccelerator/Detours/Private/Windows/UbaDetoursFunctionsWin.cpp b/Engine/Source/Programs/UnrealBuildAccelerator/Detours/Private/Windows/UbaDetoursFunctionsWin.cpp
--- a/Engine/Source/Programs/UnrealBuildAccelerator/Detours/Private/Windows/UbaDetoursFunctionsWin.cpp
+++ b/Engine/Source/Programs/UnrealBuildAccelerator/Detours/Private/Windows/UbaDetoursFunctionsWin.cpp
@@ -903,10 +903,17 @@ void Init(const DetoursPayload& payload, u64 startTime)
 			BinaryReader vfsReader(reader.GetPositionData(), 0, vfsSize);
 			while (vfsReader.GetLeft())
 			{
-				u32 index = g_vfsEntryCount++;
-				VfsEntry& vfsEntry = g_vfsEntries[index];
+				vfsReader.ReadByte(); // Index, unused
 				StringBuffer<> str;
 				vfsReader.ReadString(str);
+				if (!str.count)
+				{
+					vfsReader.SkipString();
+					continue;
+				}
+
+				u32 index = g_vfsEntryCount++;
+				VfsEntry& vfsEntry = g_vfsEntries[index];
 				vfsEntry.vfs = g_memoryBlock.Strdup(str);
 				if (index == 0)
 					g_vfsMatchingLength = vfsEntry.vfs.count;
diff --git a/Engine/Source/Programs/UnrealBuildAccelerator/Host/Private/UbaExports.cpp b/Engine/Source/Programs/UnrealBuildAccelerator/Host/Private/UbaExports.cpp
--- a/Engine/Source/Programs/UnrealBuildAccelerator/Host/Private/UbaExports.cpp
+++ b/Engine/Source/Programs/UnrealBuildAccelerator/Host/Private/UbaExports.cpp
@@ -393,9 +393,9 @@ extern "C"
 	{
 		return new uba::ProcessHandle(server->RunProcessRacing(raceAgainstRemoteProcessId));
 	}
-	uba::u64 SessionServer_RegisterVfs(uba::SessionServer* server, const void* vfsData, uba::u64 vfsDataSize)
+	uba::u64 SessionServer_RegisterRoots(uba::SessionServer* server, const void* rootsData, uba::u64 rootsDataSize)
 	{
-		return server->RegisterVfs(vfsData, vfsDataSize);
+		return server->RegisterRoots(rootsData, rootsDataSize);
 	}
 	void SessionServer_SetMaxRemoteProcessCount(uba::SessionServer* server, uba::u32 count)
 	{
@@ -514,7 +514,7 @@ extern "C"
 		return info;
 	}
 
-	uba::ProcessStartInfo* ProcessStartInfo_Create2(const uba::tchar* application, const uba::tchar* arguments, const uba::tchar* workingDir, const uba::tchar* description, uba::u32 priorityClass, uba::u64 vfsHandle, bool trackInputs, const uba::tchar* logFile, ProcessHandle_ExitCallback* exit)
+	uba::ProcessStartInfo* ProcessStartInfo_Create2(const uba::tchar* application, const uba::tchar* arguments, const uba::tchar* workingDir, const uba::tchar* description, uba::u32 priorityClass, uba::u64 rootsHandle, bool trackInputs, const uba::tchar* logFile, ProcessHandle_ExitCallback* exit)
 	{
 		auto info = new uba::ProcessStartInfo();
 		info->application = TStrdup(application);
@@ -525,7 +525,7 @@ extern "C"
 		info->trackInputs = trackInputs;
 		info->logFile = TStrdup(logFile);
 		info->exitedFunc = exit;
-		info->vfsHandle = vfsHandle;
+		info->rootsHandle = rootsHandle;
 		return info;
 	}
 
@@ -655,9 +655,9 @@ extern "C"
 		return cacheClient->WriteToCache(*rootPaths, bucket, process->GetStartInfo(), inputs, inputsSize, outputs, outputsSize, logLinesWriter.GetData(), logLinesWriter.GetPosition(), process->GetId());
 	}
 
-	bool CacheClient_WriteToCache2(uba::CacheClient* cacheClient, uba::RootPaths* rootPaths, uba::u32 bucket, const uba::ProcessHandle* process, const uba::u8* inputs, uba::u32 inputsSize, const uba::u8* outputs, uba::u32 outputsSize, const uba::u8* logLines, uba::u32 logLinesSize)
+	bool CacheClient_WriteToCache2(uba::CacheClient* cacheClient, uba::u32 bucket, const uba::ProcessHandle* process, const uba::u8* inputs, uba::u32 inputsSize, const uba::u8* outputs, uba::u32 outputsSize, const uba::u8* logLines, uba::u32 logLinesSize)
 	{
-		return cacheClient->WriteToCache(*rootPaths, bucket, process->GetStartInfo(), inputs, inputsSize, outputs, outputsSize, logLines, logLinesSize, process->GetId());
+		return cacheClient->WriteToCache(bucket, process->GetStartInfo(), inputs, inputsSize, outputs, outputsSize, logLines, logLinesSize, process->GetId());
 	}
 
 	uba::u32 CacheClient_FetchFromCache(uba::CacheClient* cacheClient, uba::RootPaths* rootPaths, uba::u32 bucket, const uba::ProcessStartInfo& info)
diff --git a/Engine/Source/Programs/UnrealBuildAccelerator/Host/Public/UbaExports.h b/Engine/Source/Programs/UnrealBuildAccelerator/Host/Public/UbaExports.h
--- a/Engine/Source/Programs/UnrealBuildAccelerator/Host/Public/UbaExports.h
+++ b/Engine/Source/Programs/UnrealBuildAccelerator/Host/Public/UbaExports.h
@@ -72,7 +72,7 @@ extern "C"
 	// ProcessStartInfo
 	using ProcessHandle_ExitCallback = void(void* userData, const uba::ProcessHandle&);
 	UBA_API uba::ProcessStartInfo* ProcessStartInfo_Create(const uba::tchar* application, const uba::tchar* arguments, const uba::tchar* workingDir, const uba::tchar* description, uba::u32 priorityClass, uba::u64 outputStatsThresholdMs, bool trackInputs, const uba::tchar* logFile, ProcessHandle_ExitCallback* exit);
-	UBA_API uba::ProcessStartInfo* ProcessStartInfo_Create2(const uba::tchar* application, const uba::tchar* arguments, const uba::tchar* workingDir, const uba::tchar* description, uba::u32 priorityClass, uba::u64 vfsHandle, bool trackInputs, const uba::tchar* logFile, ProcessHandle_ExitCallback* exit);
+	UBA_API uba::ProcessStartInfo* ProcessStartInfo_Create2(const uba::tchar* application, const uba::tchar* arguments, const uba::tchar* workingDir, const uba::tchar* description, uba::u32 priorityClass, uba::u64 rootsHandle, bool trackInputs, const uba::tchar* logFile, ProcessHandle_ExitCallback* exit);
 	UBA_API void ProcessStartInfo_Destroy(uba::ProcessStartInfo* info);
 
 	// ProcessHandle
@@ -106,7 +106,7 @@ extern "C"
 	UBA_API uba::ProcessHandle* SessionServer_RunProcess(uba::SessionServer* server, uba::ProcessStartInfo& info, bool async, bool enableDetour);
 	UBA_API uba::ProcessHandle* SessionServer_RunProcessRemote(uba::SessionServer* server, uba::ProcessStartInfo& info, float weight, const void* knownInputs = nullptr, uba::u32 knownInputsCount = 0, bool allowCrossArchitecture = false);
 	UBA_API uba::ProcessHandle* SessionServer_RunProcessRacing(uba::SessionServer* server, uba::u32 raceAgainstRemoteProcessId);
-	UBA_API uba::u64 SessionServer_RegisterVfs(uba::SessionServer* server, const void* vfsData, uba::u64 vfsDataSize);
+	UBA_API uba::u64 SessionServer_RegisterRoots(uba::SessionServer* server, const void* rootsData, uba::u64 rootsDataSize);
 
 	UBA_API void SessionServer_SetMaxRemoteProcessCount(uba::SessionServer* server, uba::u32 count);
 	UBA_API void SessionServer_DisableRemoteExecution(uba::SessionServer* server);
@@ -135,7 +135,7 @@ extern "C"
 	UBA_API uba::CacheClient* CacheClient_Create(uba::SessionServer* session, bool reportMissReason = false, const uba::tchar* crypto = nullptr);
 	UBA_API bool CacheClient_Connect(uba::CacheClient* cacheClient, const uba::tchar* host, int port);
 	UBA_API bool CacheClient_WriteToCache(uba::CacheClient* cacheClient, uba::RootPaths* rootPaths, uba::u32 bucket, const uba::ProcessHandle* process, const uba::u8* inputs, uba::u32 inputsSize, const uba::u8* outputs, uba::u32 outputsSize);
-	UBA_API bool CacheClient_WriteToCache2(uba::CacheClient* cacheClient, uba::RootPaths* rootPaths, uba::u32 bucket, const uba::ProcessHandle* process, const uba::u8* inputs, uba::u32 inputsSize, const uba::u8* outputs, uba::u32 outputsSize, const uba::u8* logLines, uba::u32 logLinesSize);
+	UBA_API bool CacheClient_WriteToCache2(uba::CacheClient* cacheClient, uba::u32 bucket, const uba::ProcessHandle* process, const uba::u8* inputs, uba::u32 inputsSize, const uba::u8* outputs, uba::u32 outputsSize, const uba::u8* logLines, uba::u32 logLinesSize);
 	UBA_API uba::u32 CacheClient_FetchFromCache(uba::CacheClient* cacheClient, uba::RootPaths* rootPaths, uba::u32 bucket, const uba::ProcessStartInfo& info);
 	UBA_API uba::CacheResult* CacheClient_FetchFromCache2(uba::CacheClient* cacheClient, uba::RootPaths* rootPaths, uba::u32 bucket, const uba::ProcessStartInfo& info);
 	UBA_API void CacheClient_RequestServerShutdown(uba::CacheClient* cacheClient, const uba::tchar* reason);
