[UBA]
* Fixed bug in NetworkServer::RemoveDisconnectedConnections where iterator was increased when it shouldnt

diff --git a/Engine/Source/Programs/UnrealBuildAccelerator/Common/Private/UbaNetworkServer.cpp b/Engine/Source/Programs/UnrealBuildAccelerator/Common/Private/UbaNetworkServer.cpp
--- a/Engine/Source/Programs/UnrealBuildAccelerator/Common/Private/UbaNetworkServer.cpp
+++ b/Engine/Source/Programs/UnrealBuildAccelerator/Common/Private/UbaNetworkServer.cpp
@@ -1165,7 +1165,7 @@ namespace uba
 			return;
 
 		SCOPED_WRITE_LOCK(m_clientsLock, lock);
-		for (auto it=m_clients.begin(); it!= m_clients.end(); ++it)
+		for (auto it=m_clients.begin(); it!=m_clients.end();)
 		{
 			if (it->second.refCount)
 				++it;
