[UBA]
* Fixed code that was asserting in debug because it was wrong for case sensitive file systems

diff --git a/Engine/Source/Programs/UnrealBuildAccelerator/Common/Private/UbaSession.cpp b/Engine/Source/Programs/UnrealBuildAccelerator/Common/Private/UbaSession.cpp
--- a/Engine/Source/Programs/UnrealBuildAccelerator/Common/Private/UbaSession.cpp
+++ b/Engine/Source/Programs/UnrealBuildAccelerator/Common/Private/UbaSession.cpp
@@ -1430,7 +1430,7 @@ namespace uba
 		StringBuffer<> testPath(fixedFilePath);
 		if (CaseInsensitiveFs)
 			testPath.MakeLower();
-		StringKey testKey = ToStringKeyLower(testPath);
+		StringKey testKey = ToStringKey(testPath);
 		UBA_ASSERTF(testKey == result, TC("Key mismatch for %s"), fixedFilePath.data);
 #endif
 
