[UBA]
* Temporarily revert UbaProcessStartInfo.h to match submitted binaries.. uba will not compile but at least farm will be happy for a while

diff --git a/Engine/Source/Programs/UnrealBuildAccelerator/Common/Public/UbaProcessStartInfo.h b/Engine/Source/Programs/UnrealBuildAccelerator/Common/Public/UbaProcessStartInfo.h
--- a/Engine/Source/Programs/UnrealBuildAccelerator/Common/Public/UbaProcessStartInfo.h
+++ b/Engine/Source/Programs/UnrealBuildAccelerator/Common/Public/UbaProcessStartInfo.h
@@ -8,7 +8,6 @@ namespace uba
 {
 	class ApplicationRules;
 	class ProcessHandle;
-	using RootsHandle = u64;
 
 	struct ProcessStartInfo
 	{
@@ -24,14 +23,13 @@ namespace uba
 		const tchar* description = TC("");		// Description. Used for on-screen logging and log file names if session.logToFile is set but logFile is ""
 		const tchar* logFile = TC("");			// Log file. If set, will always log. If not full path the session log dir will be prepended.
 		u32 priorityClass = 0x00000020;			// Priority of process. Defaults to NORMAL_PRIORITY_CLASS
+		u64 outputStatsThresholdMs = ~u64(0);	// Threshold in milliseconds where process summary should be printed to log
 		bool trackInputs = false;				// Track all files read. Can read result in ProcessHandle.GetTrackedInputs()
 		bool useCustomAllocator = true;			// Disable detouring of allocator inside processes. If Session.disableCustomAllocator is false this will be overridden
 		bool writeOutputFilesOnFail = false;	// If set to true, output files will be written to disk regardless if process succeeds or not
 		const tchar* breadcrumbs = TC("");		// If not empty, write additional information to the UBA trace file
 		bool startSuspended = false;			// Start process suspended.. a bit internal atm and only supported on windows
 
-		RootsHandle rootsHandle = 0;			// Handle if using vfs. Handle is generated through Session::RegisterRoots
-
 		using LogLineCallback = void(void* userData, const tchar* line, u32 length, LogEntryType type);
 		LogLineCallback* logLineFunc = nullptr;	// Callback for when log entries happens
 		void* logLineUserData = nullptr;		// User data provided to logLine callback
