[UBA]
* Reduced stack pressure in session client by moving request to fetch roots to a separate function

diff --git a/Engine/Source/Programs/UnrealBuildAccelerator/Common/Private/UbaSessionClient.cpp b/Engine/Source/Programs/UnrealBuildAccelerator/Common/Private/UbaSessionClient.cpp
--- a/Engine/Source/Programs/UnrealBuildAccelerator/Common/Private/UbaSessionClient.cpp
+++ b/Engine/Source/Programs/UnrealBuildAccelerator/Common/Private/UbaSessionClient.cpp
@@ -311,20 +311,8 @@ namespace uba
 		}
 
 		if (RootsHandle rootsHandle = startInfo.rootsHandle)
-		{
-			SCOPED_WRITE_LOCK(m_rootsLookupLock, rootsLock);
-			auto insres2 = m_rootsLookup.try_emplace(rootsHandle);
-			if (insres2.second)
-			{
-				StackBinaryWriter<128> writer;
-				NetworkMessage msg(m_client, ServiceId, SessionMessageType_GetRoots, writer);
-				writer.WriteU64(rootsHandle);
-				StackBinaryReader<8*1024> reader;
-				if (!msg.Send(reader, m_stats.getApplicationMsg))
-					return false;
-				PopulateRootsEntry(insres2.first->second, reader.GetPositionData(), reader.GetLeft());
-			}
-		}
+			if (!SendRootsHandle(rootsHandle))
+				return false;
 
 		while (handledCount < modules.size())
 			m_client.DoWork();
@@ -1497,6 +1485,22 @@ namespace uba
 		msg.Send();
 	}
 
+	bool SessionClient::SendRootsHandle(RootsHandle rootsHandle)
+	{
+		SCOPED_WRITE_LOCK(m_rootsLookupLock, rootsLock);
+		auto insres2 = m_rootsLookup.try_emplace(rootsHandle);
+		if (!insres2.second)
+			return true;
+		StackBinaryWriter<128> writer;
+		NetworkMessage msg(m_client, ServiceId, SessionMessageType_GetRoots, writer);
+		writer.WriteU64(rootsHandle);
+		StackBinaryReader<8*1024> reader;
+		if (!msg.Send(reader, m_stats.getApplicationMsg))
+			return false;
+		PopulateRootsEntry(insres2.first->second, reader.GetPositionData(), reader.GetLeft());
+		return true;
+	}
+
 	void SessionClient::SendSummary(const Function<void(Logger&)>& extraInfo)
 	{
 		StackBinaryWriter<SendMaxSize> writer;
diff --git a/Engine/Source/Programs/UnrealBuildAccelerator/Common/Public/UbaSessionClient.h b/Engine/Source/Programs/UnrealBuildAccelerator/Common/Public/UbaSessionClient.h
--- a/Engine/Source/Programs/UnrealBuildAccelerator/Common/Public/UbaSessionClient.h
+++ b/Engine/Source/Programs/UnrealBuildAccelerator/Common/Public/UbaSessionClient.h
@@ -93,6 +93,7 @@ namespace uba
 		bool SendProcessFinished(ProcessImpl& process, u32 exitCode);
 		void SendPing(u64 memAvail, u64 memTotal);
 		void SendNotification(const StringView& text);
+		bool SendRootsHandle(RootsHandle rootsHandle);
 		void SendLogFileToServer(ProcessImpl& pi);
 		void GetLogFileName(StringBufferBase& out, const tchar* logFile, const tchar* arguments, u32 processId);
 		u32 CountLogLines(ProcessImpl& process);
