[UBA]
* Fixed bug where not all roots handles got transferred to client

diff --git a/Engine/Source/Programs/UnrealBuildAccelerator/Common/Private/UbaProcess.cpp b/Engine/Source/Programs/UnrealBuildAccelerator/Common/Private/UbaProcess.cpp
--- a/Engine/Source/Programs/UnrealBuildAccelerator/Common/Private/UbaProcess.cpp
+++ b/Engine/Source/Programs/UnrealBuildAccelerator/Common/Private/UbaProcess.cpp
@@ -714,7 +714,7 @@ namespace uba
 
 	bool ProcessImpl::HandleInit(BinaryReader& reader, BinaryWriter& writer)
 	{
-		InitMessage msg;
+		InitMessage msg { *this };
 		InitResponse response;
 		m_messageSuccess = m_session.GetInitResponse(response, msg) && m_messageSuccess;
 		writer.WriteBool(m_echoOn);
diff --git a/Engine/Source/Programs/UnrealBuildAccelerator/Common/Private/UbaSessionClient.cpp b/Engine/Source/Programs/UnrealBuildAccelerator/Common/Private/UbaSessionClient.cpp
--- a/Engine/Source/Programs/UnrealBuildAccelerator/Common/Private/UbaSessionClient.cpp
+++ b/Engine/Source/Programs/UnrealBuildAccelerator/Common/Private/UbaSessionClient.cpp
@@ -310,10 +310,6 @@ namespace uba
 				}, 1, TC("EnsureApp"));
 		}
 
-		if (RootsHandle rootsHandle = startInfo.rootsHandle)
-			if (!SendRootsHandle(rootsHandle))
-				return false;
-
 		while (handledCount < modules.size())
 			m_client.DoWork();
 
@@ -424,6 +420,16 @@ namespace uba
 		return m_storage.CopyOrLink(casKey, out.data, fileAttributes);
 	}
 
+	bool SessionClient::GetInitResponse(InitResponse& out, const InitMessage& msg)
+	{
+		if (!Session::GetInitResponse(out, msg))
+			return false;
+		if (RootsHandle rootsHandle = msg.process.GetStartInfo().rootsHandle)
+			if (!SendRootsHandle(rootsHandle))
+				return false;
+		return true;
+	}
+
 	bool SessionClient::CreateFile(CreateFileResponse& out, const CreateFileMessage& msg)
 	{
 		const StringBufferBase& fileName = msg.fileName;
diff --git a/Engine/Source/Programs/UnrealBuildAccelerator/Common/Public/UbaSession.h b/Engine/Source/Programs/UnrealBuildAccelerator/Common/Public/UbaSession.h
--- a/Engine/Source/Programs/UnrealBuildAccelerator/Common/Public/UbaSession.h
+++ b/Engine/Source/Programs/UnrealBuildAccelerator/Common/Public/UbaSession.h
@@ -298,6 +298,7 @@ namespace uba
 
 	struct InitMessage
 	{
+		ProcessImpl& process;
 	};
 
 	struct InitResponse
diff --git a/Engine/Source/Programs/UnrealBuildAccelerator/Common/Public/UbaSessionClient.h b/Engine/Source/Programs/UnrealBuildAccelerator/Common/Public/UbaSessionClient.h
--- a/Engine/Source/Programs/UnrealBuildAccelerator/Common/Public/UbaSessionClient.h
+++ b/Engine/Source/Programs/UnrealBuildAccelerator/Common/Public/UbaSessionClient.h
@@ -50,6 +50,7 @@ namespace uba
 
 		virtual bool PrepareProcess(ProcessStartInfoHolder& startInfo, bool isChild, StringBufferBase& outRealApplication, const tchar*& outRealWorkingDir) override;
 		virtual void* GetProcessEnvironmentVariables() override;
+		virtual bool GetInitResponse(InitResponse& out, const InitMessage& msg) override;
 		virtual bool CreateFile(CreateFileResponse& out, const CreateFileMessage& msg) override;
 		virtual bool DeleteFile(DeleteFileResponse& out, const DeleteFileMessage& msg) override;
 		virtual bool CopyFile(CopyFileResponse& out, const CopyFileMessage& msg) override;
