[UBA]
* Fixed so all .json files uploaded to cache are normalized (this also now includes the vfs file)
* Changed so reportmissreason is or'ed together with settings from toml file

diff --git a/Engine/Source/Programs/UnrealBuildAccelerator/Common/Private/UbaCacheClient.cpp b/Engine/Source/Programs/UnrealBuildAccelerator/Common/Private/UbaCacheClient.cpp
--- a/Engine/Source/Programs/UnrealBuildAccelerator/Common/Private/UbaCacheClient.cpp
+++ b/Engine/Source/Programs/UnrealBuildAccelerator/Common/Private/UbaCacheClient.cpp
@@ -1197,9 +1197,9 @@ namespace uba
 	{
 		if (!m_useRoots)
 			return false;
-		if (path.EndsWith(TC(".dep.json"))) // Contains absolute paths (dep file for msvc)
+		if (path.EndsWith(TC(".json"))) // Contains absolute paths (dep files for msvc and vfsoverlay files for clang)
 			return true;
-		if (path.EndsWith(TC(".d"))) // Contains absolute paths (dep file for clang)
+		if (path.EndsWith(TC(".d"))) // Contains absolute paths (dep files for clang)
 			return true;
 		if (path.EndsWith(TC(".tlh"))) // Contains absolute path in a comment
 			return true;
diff --git a/Engine/Source/Programs/UnrealBuildAccelerator/Host/Private/UbaExports.cpp b/Engine/Source/Programs/UnrealBuildAccelerator/Host/Private/UbaExports.cpp
--- a/Engine/Source/Programs/UnrealBuildAccelerator/Host/Private/UbaExports.cpp
+++ b/Engine/Source/Programs/UnrealBuildAccelerator/Host/Private/UbaExports.cpp
@@ -623,7 +623,7 @@ extern "C"
 		CacheClientCreateInfo info{writer, storage, *networkClient, *session};
 		info.Apply(GetConfig());
 
-		info.reportMissReason = reportMissReason;
+		info.reportMissReason |= reportMissReason;
 		return new CacheClient(info);
 	}
 
