[UBA]
* Fixed bug in NtQueryDirectoryFile where IoStatusBlock->Information was not properly set

diff --git a/Engine/Source/Programs/UnrealBuildAccelerator/Detours/Private/Windows/UbaDetoursFunctionsNtDll.inl b/Engine/Source/Programs/UnrealBuildAccelerator/Detours/Private/Windows/UbaDetoursFunctionsNtDll.inl
--- a/Engine/Source/Programs/UnrealBuildAccelerator/Detours/Private/Windows/UbaDetoursFunctionsNtDll.inl
+++ b/Engine/Source/Programs/UnrealBuildAccelerator/Detours/Private/Windows/UbaDetoursFunctionsNtDll.inl
@@ -236,6 +236,8 @@ NTSTATUS NTAPI Detoured_NtQueryDirectoryFile(HANDLE FileHandle, HANDLE Event, PI
 
 	if (isListDirectoryHandle(FileHandle))
 	{
+		IoStatusBlock->Information = 0;
+
 		auto& listHandle = asListDirectoryHandle(FileHandle);
 		NTSTATUS res = STATUS_NO_MORE_FILES;
 
@@ -281,6 +283,7 @@ NTSTATUS NTAPI Detoured_NtQueryDirectoryFile(HANDLE FileHandle, HANDLE Event, PI
 			else
 			{
 				UBA_ASSERT(false);
+				IoStatusBlock->Status = STATUS_OBJECT_NAME_NOT_FOUND;
 				return STATUS_OBJECT_NAME_NOT_FOUND;
 			}
 
@@ -322,6 +325,9 @@ NTSTATUS NTAPI Detoured_NtQueryDirectoryFile(HANDLE FileHandle, HANDLE Event, PI
 				break;
 		}
 
+		IoStatusBlock->Status = res;
+		IoStatusBlock->Information = it - (u8*)FileInformation;
+
 #if 0//UBA_DEBUG_VALIDATE
 		if (false) // Sorting can mismatch
 		{
