[UBA]
* Fixed CacheClient code using '\' instead of PathSeparator in one place
* Fixed code path that didn't handle all possible scenarios when storage was copyorlink file

diff --git a/Engine/Source/Programs/UnrealBuildAccelerator/Common/Private/UbaCacheClient.cpp b/Engine/Source/Programs/UnrealBuildAccelerator/Common/Private/UbaCacheClient.cpp
--- a/Engine/Source/Programs/UnrealBuildAccelerator/Common/Private/UbaCacheClient.cpp
+++ b/Engine/Source/Programs/UnrealBuildAccelerator/Common/Private/UbaCacheClient.cpp
@@ -1272,7 +1272,7 @@ namespace uba
 				if (IsDirectory(e.attributes))
 					return;
 
-				path.Clear().Append('\\').Append(e.name, e.nameLen);
+				path.Clear().Append(PathSeparator).Append(e.name, e.nameLen);
 				if (CaseInsensitiveFs)
 					path.MakeLower();
 
diff --git a/Engine/Source/Programs/UnrealBuildAccelerator/Common/Private/UbaStorage.cpp b/Engine/Source/Programs/UnrealBuildAccelerator/Common/Private/UbaStorage.cpp
--- a/Engine/Source/Programs/UnrealBuildAccelerator/Common/Private/UbaStorage.cpp
+++ b/Engine/Source/Programs/UnrealBuildAccelerator/Common/Private/UbaStorage.cpp
@@ -685,8 +685,8 @@ namespace uba
 
 	bool StorageImpl::DecompressMemoryToMemory(u8* compressedData, u8* writeData, u64 decompressedSize, const tchar* readHint, const tchar* writeHint)
 	{
-		UBA_ASSERT(compressedData);
-		UBA_ASSERT(writeData);
+		UBA_ASSERTF(compressedData, TC("DecompressMemoryToMemory got readmem nullptr (%s)"), readHint);
+		UBA_ASSERTF(writeData, TC("DecompressMemoryToMemory got writemem nullptr (%s)"), writeHint);
 
 		StorageStats& stats = Stats();
 
@@ -2544,13 +2544,30 @@ namespace uba
 
 				u64 sizeOnDisk = decompressedSize;
 
+				auto decompressToMemory = [&](u8* destinationMemory)
+					{
+						if (casEntry->mappingHandle.IsValid())
+						{
+							if (!DecompressMemoryToMemory(readData, destinationMemory, decompressedSize, casFile.data, destination))
+								return false;
+						}
+						else
+						{
+							if (!DecompressFileToMemory(CasKeyString(actualKey).str, readHandle, destinationMemory, decompressedSize, destination))
+								return false;
+						}
+						return true;
+					};
+
+
 				if (formattingFunc)
 				{
 					UBA_ASSERT(!writeCompressed);
 					u8* slot = m_bufferSlots.Pop();
 					auto _ = MakeGuard([&](){ m_bufferSlots.Push(slot); });
 
-					if (!DecompressMemoryToMemory(readData, slot, decompressedSize, casFile.data, destination))
+					UBA_ASSERT(decompressedSize < BufferSlotSize);
+					if (!decompressToMemory(slot))
 						return false;
 
 					if (!destinationFile.CreateWrite(false, DefaultAttributes(), 0, m_tempPath.data))
@@ -2605,17 +2622,8 @@ namespace uba
 				{
 					if (!destinationFile.CreateMemoryWrite(allowRead, writeFlags | fileAttributes, decompressedSize, m_tempPath.data))
 						return false;
-
-					if (casEntry->mappingHandle.IsValid())
-					{
-						if (!DecompressMemoryToMemory(readData, destinationFile.GetData(), decompressedSize, casFile.data, destination))
-							return false;
-					}
-					else
-					{
-						if (!DecompressFileToMemory(CasKeyString(actualKey).str, readHandle, destinationFile.GetData(), decompressedSize, destination))
-							return false;
-					}
+					if (!decompressToMemory(destinationFile.GetData()))
+						return false;
 				}
 
 				u64 lastWriteTime = 0;
