[UBA]
* Added strcpy_s and strcat_s to posix and used them for assert code

diff --git a/Engine/Source/Programs/UnrealBuildAccelerator/Core/Private/UbaPlatform.cpp b/Engine/Source/Programs/UnrealBuildAccelerator/Core/Private/UbaPlatform.cpp
--- a/Engine/Source/Programs/UnrealBuildAccelerator/Core/Private/UbaPlatform.cpp
+++ b/Engine/Source/Programs/UnrealBuildAccelerator/Core/Private/UbaPlatform.cpp
@@ -212,35 +212,36 @@ namespace uba
 				progName[res] = 0;
 				fflush(stdout);
 				char str[4096];
-				strcpy(str, "addr2line");
+				strcpy_s(str, sizeof(str), "addr2line");
 				for (int i = 0; i < nptrs; ++i)
 				{
 					char v[32];
 					snprintf(v, sizeof_array(v), " %llx", u64(buffer[i]));
-					strcat(str, v);
+					strcat_s(str, sizeof(str), v);
 				}
-				strcat(str, " -e ");
+				strcat_s(str, sizeof(str), " -e ");
 
-				strcat(str, progName);
+				strcat_s(str, sizeof(str), progName);
 				fflush(stdout);
 				if (FILE* fp = popen(str, "r"))
 				{
 					out.Append('\n');
-					char path[1024] = { 0 };
 					u32 index = 0;
 					u32 countBeforeCallstack = out.count;
-					while (fgets(path, sizeof(path), fp) != NULL)
+					while (fgets(str, sizeof(str), fp) != NULL)
 					{
-						if (path[0] == '?')
+						if (str[0] == '?')
 							break;
 						if (index++ <= skipCallstack)
 							continue;
-						char* fileName = path;
-						char* lastSlash = strrchr(path, '/');
+						char* fileName = str;
+						char* lastSlash = strrchr(str, '/');
 						if (lastSlash)
 							fileName = lastSlash + 1;
 						if (strncmp(fileName, "function.h", 10) == 0 || strncmp(fileName, "invoke.h:", 9) == 0)
 							continue;
+						if (out.capacity - out.count < strlen(fileName) + 5)
+							break;
 						out.Appendf(TC("  %s"), fileName);
 					}
 					fclose(fp);
@@ -378,6 +379,19 @@ namespace uba
 
 
 #if !PLATFORM_WINDOWS
+	void strcpy_s(tchar* dest, u64 destCapacity, const tchar* source)
+	{
+		u64 toCopy = Min(destCapacity-1, u64(strlen(source)));
+		memcpy(dest, source, toCopy);
+		dest[toCopy] = 0;
+	}
+
+	void strcat_s(tchar* dest, u64 destCapacity, const tchar* source)
+	{
+		u64 len = strlen(dest);
+		strcpy_s(dest + len, destCapacity - len, source);
+	}
+
 	void GetMappingHandleName(StringBufferBase& out, u64 uid)
 	{
 		#if PLATFORM_MAC
diff --git a/Engine/Source/Programs/UnrealBuildAccelerator/Core/Public/UbaPlatform.h b/Engine/Source/Programs/UnrealBuildAccelerator/Core/Public/UbaPlatform.h
--- a/Engine/Source/Programs/UnrealBuildAccelerator/Core/Public/UbaPlatform.h
+++ b/Engine/Source/Programs/UnrealBuildAccelerator/Core/Public/UbaPlatform.h
@@ -142,10 +142,12 @@ namespace uba
 	#define TStrcmp(a, b) strcmp(a, b)
 	#define TSprintf_s(...) snprintf(__VA_ARGS__)
 	#define Tvsprintf_s(a, b, c, d) vsnprintf(a, b, c, d)
-	#define TStrcpy_s(a, b, c) strcpy(a, c)
-	#define TStrcat_s(a, b, c) strcat(a, c)
+	#define TStrcpy_s(a, b, c) strcpy_s(a, b, c)
+	#define TStrcat_s(a, b, c) strcat_s(a, b, c)
 	#define TStrdup strdup
 	#define localtime_s(a,b) localtime_r(b,a)
+	void strcpy_s(tchar* dest, u64 destCapacity, const tchar* source);
+	void strcat_s(tchar* dest, u64 destCapacity, const tchar* source);
 	void GetMappingHandleName(StringBufferBase& out, u64 uid);
 	inline u64 FromTimeSpec(const timespec& ts) { return u64(ts.tv_sec) * 10'000'000ull + u64(ts.tv_nsec/100); }
 	inline timespec ToTimeSpec(u64 time) { timespec ts; ts.tv_sec = time / 10'000'000ull; ts.tv_nsec = (time - (u64(ts.tv_sec) * 10'000'000ull)) * 100; return ts; }
