[UBA]
* Removed assert in detoured wait3. Seems like it is called super early when running cp

diff --git a/Engine/Source/Programs/UnrealBuildAccelerator/Detours/Private/PosixOS/UbaDetoursFunctionsPosix.cpp b/Engine/Source/Programs/UnrealBuildAccelerator/Detours/Private/PosixOS/UbaDetoursFunctionsPosix.cpp
--- a/Engine/Source/Programs/UnrealBuildAccelerator/Detours/Private/PosixOS/UbaDetoursFunctionsPosix.cpp
+++ b/Engine/Source/Programs/UnrealBuildAccelerator/Detours/Private/PosixOS/UbaDetoursFunctionsPosix.cpp
@@ -2074,7 +2074,6 @@ UBA_EXPORT pid_t UBA_WRAPPER(wait3)(int* status, int options, struct rusage* rus
 	DEBUG_LOG_TRUE("wait3", "");
 	UBA_ASSERTF(!t_inVfork, "wait3: is in fork");
 	pid_t res =TRUE_WRAPPER(wait3)(status, options, rusage);
-	UBA_ASSERTF(false, "wait3");
 	return res;
 }
 
