[UBA]
* Added full support for vfs. Uba can now have vfs paths registered and then take in files that only work against vfs

diff --git a/Engine/Source/Programs/UnrealBuildAccelerator/Common/Private/UbaSession.cpp b/Engine/Source/Programs/UnrealBuildAccelerator/Common/Private/UbaSession.cpp
--- a/Engine/Source/Programs/UnrealBuildAccelerator/Common/Private/UbaSession.cpp
+++ b/Engine/Source/Programs/UnrealBuildAccelerator/Common/Private/UbaSession.cpp
@@ -896,7 +896,7 @@ namespace uba
 			Sleep(1);
 			dirLock.Enter();
 		}
-		UBA_ASSERT(dir.parseOffset == 1);
+		UBA_ASSERTF(dir.parseOffset == 1, TC("Registering deleted file %s with unexpect dir.parseOffset %u "), fileName.data, dir.parseOffset);
 
 		if (fileNameKey == StringKeyZero)
 		{
@@ -2333,7 +2333,9 @@ namespace uba
 
 				auto memClose = MakeGuard([&](){ UnmapViewOfFile(m_logger, mem, fileSize, file.name.c_str()); });
 
-				if (m_storeObjFilesCompressed && process.m_startInfo.rules->StoreFileCompressed(file.name))
+				auto& rules = *process.m_startInfo.rules;
+
+				if (m_storeObjFilesCompressed && rules.StoreFileCompressed(file.name))
 				{
 					Storage::WriteResult res;
 					CompressedObjFileHeader header { CalculateCasKey(mem, fileSize, true, m_workManager, file.name.c_str()) };
@@ -2348,7 +2350,8 @@ namespace uba
 					FileAccessor destinationFile(m_logger, file.name.c_str());
 
 					RootsHandle rootsHandle = process.GetStartInfo().rootsHandle;
-					if (rootsHandle && StringView(file.name).EndsWith(TC(".d")))
+					bool escapeSpaces;
+					if (rootsHandle && rules.ShouldDevirtualizeFile(file.name, escapeSpaces))
 					{
 						// Need to turn paths back into local paths
 
@@ -2356,7 +2359,7 @@ namespace uba
 							return false;
 
 						MemoryBlock block(5*1024*1024);
-						if (!FormatClangDepsFile(rootsHandle, block, mem, fileSize, file.name.c_str()))
+						if (!DevirtualizeDepsFile(rootsHandle, block, mem, fileSize, escapeSpaces, file.name.c_str()))
 							return false;
 
 						if (!destinationFile.Write(block.memory, block.writtenSize))
@@ -3040,17 +3043,7 @@ namespace uba
 			reader.ReadString(temp);
 			entry.roots.RegisterRoot(m_logger, temp.data);
 			reader.ReadString(temp.Clear());
-			char ansi[512];
-			u32 ansiPos = 0;
-			for (u32 i=0;i!=temp.count;++i)
-			{
-				char c = (char)temp[i];
-				if (c == ' ')
-					ansi[ansiPos++] = '\\';
-				ansi[ansiPos++] = c;
-			}
-			ansi[ansiPos] = 0;
-			entry.localsAnsi.emplace_back(ansi);
+			entry.locals.emplace_back(temp.data, temp.count);
 		}
 	}
 
@@ -3123,7 +3116,7 @@ namespace uba
 		return true;
 	}
 
-	bool Session::FormatClangDepsFile(RootsHandle rootsHandle, MemoryBlock& destData, const void* sourceData, u64 sourceSize, const tchar* hint)
+	bool Session::DevirtualizeDepsFile(RootsHandle rootsHandle, MemoryBlock& destData, const void* sourceData, u64 sourceSize, bool escapeSpaces, const tchar* hint)
 	{
 		SCOPED_READ_LOCK(m_rootsLookupLock, lock);
 		auto findIt = m_rootsLookup.find(rootsHandle);
@@ -3131,6 +3124,32 @@ namespace uba
 		RootsEntry& rootsEntry = findIt->second;
 		lock.Leave();
 
+		Vector<std::string> localsAnsi;
+		localsAnsi.reserve(rootsEntry.locals.size());
+
+		for (auto& str : rootsEntry.locals)
+		{
+			char ansi[512];
+			u32 ansiPos = 0;
+			for (auto c : str)
+			{
+				UBA_ASSERT(c < 256);
+				if (escapeSpaces)
+				{
+					if (c == ' ')
+						ansi[ansiPos++] = '\\';
+				}
+				else
+				{
+					if (c == '\\')
+						ansi[ansiPos++] = '\\';
+				}
+				ansi[ansiPos++] = (char)c;
+			}
+			ansi[ansiPos] = 0;
+			localsAnsi.emplace_back(ansi, ansi + ansiPos);
+		}
+
 		auto handleString = [&](const char* str, u64 strLen, u32 rootPos)
 			{
 				if (rootPos == ~0u)
@@ -3138,7 +3157,7 @@ namespace uba
 					memcpy(destData.Allocate(strLen, 1, TC("")), str, strLen);
 					return;
 				}
-				auto& path = rootsEntry.localsAnsi[(*str - RootPaths::RootStartByte)/3];
+				auto& path = localsAnsi[(*str - RootPaths::RootStartByte)/3];
 				memcpy(destData.Allocate(path.size(), 1, TC("")), path.c_str(), path.size());
 			};
 
diff --git a/Engine/Source/Programs/UnrealBuildAccelerator/Common/Private/UbaSessionServer.cpp b/Engine/Source/Programs/UnrealBuildAccelerator/Common/Private/UbaSessionServer.cpp
--- a/Engine/Source/Programs/UnrealBuildAccelerator/Common/Private/UbaSessionServer.cpp
+++ b/Engine/Source/Programs/UnrealBuildAccelerator/Common/Private/UbaSessionServer.cpp
@@ -994,17 +994,20 @@ namespace uba
 
 				RootsHandle rootsHandle = h.GetStartInfo().rootsHandle;
 
+				auto& rules = *h.GetStartInfo().rules;
+
 				Storage::FormattingFunc formattingFunc;
-				if (rootsHandle && destination.EndsWith(TC(".d")))
+				bool escapeSpaces;
+				if (rootsHandle && rules.ShouldDevirtualizeFile(destination, escapeSpaces))
 				{
-					formattingFunc = [this, rootsHandle](MemoryBlock& destData, const void* sourceData, u64 sourceSize, const tchar* hint)
+					formattingFunc = [&](MemoryBlock& destData, const void* sourceData, u64 sourceSize, const tchar* hint)
 						{
-							return FormatClangDepsFile(rootsHandle, destData, sourceData, sourceSize, hint);
+							return DevirtualizeDepsFile(rootsHandle, destData, sourceData, sourceSize, escapeSpaces, hint);
 						};
 				}
 				else if (m_storeObjFilesCompressed)
 				{
-					writeCompressed = h.GetStartInfo().rules->StoreFileCompressed(destination);
+					writeCompressed = rules.StoreFileCompressed(destination);
 				}
 				success = m_storage.CopyOrLink(casKey, destination.data, attributes, writeCompressed, formattingFunc);
 				if (!success)
diff --git a/Engine/Source/Programs/UnrealBuildAccelerator/Common/Public/UbaSession.h b/Engine/Source/Programs/UnrealBuildAccelerator/Common/Public/UbaSession.h
--- a/Engine/Source/Programs/UnrealBuildAccelerator/Common/Public/UbaSession.h
+++ b/Engine/Source/Programs/UnrealBuildAccelerator/Common/Public/UbaSession.h
@@ -165,11 +165,11 @@ namespace uba
 		void WriteSummary(BinaryWriter& writer, const Function<void(Logger& logger)>& summaryFunc);
 		float UpdateCpuLoad();
 
-		struct RootsEntry { Vector<u8> memory; RootPaths roots; Vector<std::string> localsAnsi; };
+		struct RootsEntry { Vector<u8> memory; RootPaths roots; Vector<TString> locals; };
 		void PopulateRootsEntry(RootsEntry& entry, const void* rootsData, uba::u64 rootsDataSize);
 
 		bool ExtractSymbolsFromObjectFile(const CloseFileMessage& msg, const tchar* fileName, u64 fileSize);
-		bool FormatClangDepsFile(RootsHandle rootsHandle, MemoryBlock& destData, const void* sourceData, u64 sourceSize, const tchar* hint);
+		bool DevirtualizeDepsFile(RootsHandle rootsHandle, MemoryBlock& destData, const void* sourceData, u64 sourceSize, bool escapeSpaces, const tchar* hint);
 		
 		void ThreadTraceLoop();
 		virtual void TraceSessionUpdate();
diff --git a/Engine/Source/Programs/UnrealBuildAccelerator/Core/Private/UbaApplicationRules.cpp b/Engine/Source/Programs/UnrealBuildAccelerator/Core/Private/UbaApplicationRules.cpp
--- a/Engine/Source/Programs/UnrealBuildAccelerator/Core/Private/UbaApplicationRules.cpp
+++ b/Engine/Source/Programs/UnrealBuildAccelerator/Core/Private/UbaApplicationRules.cpp
@@ -87,17 +87,28 @@ namespace uba
 			return true;
 		}
 
-		virtual bool StoreFileCompressed(const StringView& fileName) const
+		virtual bool StoreFileCompressed(const StringView& fileName) const override
 		{
 			// We don't compress pch obj files because they are used by live-coding.
 			// TODO: Add support in live coding module to load compressed obj files instead
 			return fileName.EndsWith(TC(".obj")) && !fileName.Contains(TC("pch."));
 		}
 
-		virtual bool ShouldExtractSymbols(const StringView& fileName) const
+		virtual bool ShouldExtractSymbols(const StringView& fileName) const override
 		{
 			return fileName.EndsWith(TC(".obj"));
 		}
+
+		virtual bool IsFullVirtualization() const override
+		{
+			return true;
+		}
+
+		virtual bool ShouldDevirtualizeFile(const StringView& fileName, bool& outEscapeSpaces) const
+		{
+			outEscapeSpaces = false;
+			return fileName.EndsWith(TC(".dep.json"));
+		}
 	};
 
 	class ApplicationRulesVcLink : public ApplicationRulesVC
@@ -162,6 +173,11 @@ namespace uba
 		{
 			return fileName.IsEmpty() || fileName.EndsWith(TC(".obj"));
 		}
+
+		virtual bool IsFullVirtualization() const override
+		{
+			return true;
+		}
 	};
 
 	class ApplicationRulesLinkExe : public ApplicationRulesVcLink
@@ -271,6 +287,12 @@ namespace uba
 		{
 			return fileName.EndsWith(TC(".obj")) || fileName.EndsWith(TC(".o"));
 		}
+
+		virtual bool ShouldDevirtualizeFile(const StringView& fileName, bool& outEscapeSpaces) const
+		{
+			outEscapeSpaces = true;
+			return fileName.EndsWith(TC(".d"));
+		}
 	};
 
 	class ApplicationRulesLdLLdExe : public ApplicationRulesClang
diff --git a/Engine/Source/Programs/UnrealBuildAccelerator/Core/Private/UbaStringBuffer.cpp b/Engine/Source/Programs/UnrealBuildAccelerator/Core/Private/UbaStringBuffer.cpp
--- a/Engine/Source/Programs/UnrealBuildAccelerator/Core/Private/UbaStringBuffer.cpp
+++ b/Engine/Source/Programs/UnrealBuildAccelerator/Core/Private/UbaStringBuffer.cpp
@@ -253,6 +253,14 @@ namespace uba
 		return *this;
 	}
 
+	StringBufferBase& StringBufferBase::Prepend(const StringView& view)
+	{
+		UBA_ASSERT(count + view.count < capacity);
+		memmove(data+view.count, data, (count+1)*sizeof(tchar));
+		memcpy(data, view.data, view.count*sizeof(tchar));
+		return *this;
+	}
+
 	#if PLATFORM_WINDOWS
 	StringBufferBase& StringBufferBase::Append(const char* str)
 	{
diff --git a/Engine/Source/Programs/UnrealBuildAccelerator/Core/Public/UbaApplicationRules.h b/Engine/Source/Programs/UnrealBuildAccelerator/Core/Public/UbaApplicationRules.h
--- a/Engine/Source/Programs/UnrealBuildAccelerator/Core/Public/UbaApplicationRules.h
+++ b/Engine/Source/Programs/UnrealBuildAccelerator/Core/Public/UbaApplicationRules.h
@@ -97,11 +97,13 @@ namespace uba
 			return false;
 		}
 
+		// Process is allowed to detour crt allocator and use mimalloc
 		virtual bool AllowMiMalloc() const
 		{
 			return true;// !IsRunningWine();
 		}
 
+		// Can be used to disallow libraries we don't want loaded
 		virtual bool AllowLoadLibrary(const tchar* libraryName) const
 		{
 			if (Contains(libraryName, TC("nvinject.dll")))
@@ -109,51 +111,75 @@ namespace uba
 			return true;
 		}
 
+		// Silently suppress log lines.
+		// Use this to prevent log lines from be marshalled over from detoured process to session and 
 		virtual bool SuppressLogLine(const tchar* logLine, u32 logLineLen) const
 		{
 			return false;
 		}
 
+		// Return true if exit code is a success for this application
 		virtual bool IsExitCodeSuccess(u32 exitCode) const
 		{
 			return true;
 		}
 
+		// Attempt to work around bugs in pvs studio producing malformed lib paths
 		virtual void RepairMalformedLibPath(const tchar* path) const
 		{
 			// Do nothing
 		}
 
+		// Is used by uba scheduler to decide if application can be cached. (This is not used by UBT)
 		virtual bool IsCacheable() const
 		{
 			return false;
 		}
 
+		// Return true if we want an output file to be stored compressed in destination.
+		// (Used to store .obj files compressed on disk)
 		virtual bool StoreFileCompressed(const StringView& fileName) const
 		{
 			return false;
 		}
 
+		// Return true to check if we need to decompress files when loading them (used in conjunction with StoreFileCompressed)
 		virtual bool ShouldDecompressFiles(const StringView& fileName) const
 		{
 			return false;
 		}
 
+		// Return false if we for example know file sent from client is already compressed (no need to compress again)
 		virtual bool SendFileCompressedFromClient(const StringView& fileName) const
 		{
 			return true;
 		}
 
+		// Return true if uba should attempt to extract symbols from produced file
+		// (Used when compiling and extracing symbols from obj files)
 		virtual bool ShouldExtractSymbols(const StringView& fileName) const
 		{
 			return false;
 		}
 
+		// Preload if libraries. This is used to prevent deadlocks in combination with mimalloc also loading modules
 		virtual const tchar* const* LibrariesToPreload() const // Array should be null terminated
 		{
 			return nullptr;
 		}
 
+		// Full virtualization means that workingdir, exedir etc are virtualized.
+		virtual bool IsFullVirtualization() const
+		{
+			return false;
+		}
+
+		// Returns true if produced file should be devirtualized (paths are turned into local paths)
+		virtual bool ShouldDevirtualizeFile(const StringView& fileName, bool& outEscapeSpaces) const
+		{
+			return false;
+		}
+
 		u32 index = ~0u;
 	};
 
diff --git a/Engine/Source/Programs/UnrealBuildAccelerator/Core/Public/UbaStringBuffer.h b/Engine/Source/Programs/UnrealBuildAccelerator/Core/Public/UbaStringBuffer.h
--- a/Engine/Source/Programs/UnrealBuildAccelerator/Core/Public/UbaStringBuffer.h
+++ b/Engine/Source/Programs/UnrealBuildAccelerator/Core/Public/UbaStringBuffer.h
@@ -36,6 +36,7 @@ namespace uba
 		StringBufferBase& AppendHex(u64 v);
 		StringBufferBase& AppendValue(u64 v);
 		StringBufferBase& Append(const tchar* format, va_list& args);
+		StringBufferBase& Prepend(const StringView& view);
 		StringBufferBase& Resize(u64 newSize);
 		StringBufferBase& Clear();
 
