[UBA]
* Added full support for vfs. Uba can now have vfs paths registered and then take in files that only work against vfs

diff --git a/Engine/Source/Programs/UnrealBuildAccelerator/Detours/Private/Windows/UbaDetoursFunctionsKernelBase.inl b/Engine/Source/Programs/UnrealBuildAccelerator/Detours/Private/Windows/UbaDetoursFunctionsKernelBase.inl
--- a/Engine/Source/Programs/UnrealBuildAccelerator/Detours/Private/Windows/UbaDetoursFunctionsKernelBase.inl
+++ b/Engine/Source/Programs/UnrealBuildAccelerator/Detours/Private/Windows/UbaDetoursFunctionsKernelBase.inl
@@ -1035,6 +1035,7 @@ BOOL Detoured_GetFileAttributesExW(LPCWSTR lpFileName, GET_FILEEX_INFO_LEVELS fI
 
 	StringBuffer<MaxPath> fixedName;
 	FixPath(fixedName, lpFileName);
+	DevirtualizePath(fixedName);
 
 	if (!g_rules->CanExist(fixedName.data))
 	{
@@ -1059,6 +1060,7 @@ DWORD Detoured_GetFileAttributesW(LPCWSTR lpFileName)
 	StringBuffer<MaxPath> fixedPath;
 	if (!FixPath(fixedPath, lpFileName))
 		return INVALID_FILE_ATTRIBUTES;
+	DevirtualizePath(fixedPath);
 
 	WIN32_FILE_ATTRIBUTE_DATA data;
 	if (!Shared_GetFileAttributesExW(fixedPath.data, GetFileExInfoStandard, &data, lpFileName))
@@ -1371,6 +1373,8 @@ BOOL Detoured_DeleteFileW(LPCWSTR lpFileName)
 	StringBuffer<> fixedName;
 	FixPath(fixedName, lpFileName);
 
+	DevirtualizePath(fixedName);
+
 	StringBuffer<> fixedNameLower(fixedName);
 	fixedNameLower.MakeLower();
 
@@ -1429,9 +1433,11 @@ bool Shared_MoveFile(LPCWSTR lpExistingFileName, LPCWSTR lpNewFileName, DWORD dw
 
 	StringBuffer<> source;
 	FixPath(source, lpExistingFileName);
+	DevirtualizePath(source);
 
 	StringBuffer<> dest;
 	FixPath(dest, lpNewFileName);
+	DevirtualizePath(dest);
 
 	StringKey sourceKey = ToStringKeyLower(source);
 
@@ -1453,7 +1459,7 @@ bool Shared_MoveFile(LPCWSTR lpExistingFileName, LPCWSTR lpNewFileName, DWORD dw
 			auto insres = g_mappedFileTable.m_lookup.try_emplace(destKey);
 			lock2.Leave();
 			FileInfo& destInfo = insres.first->second;
-			UBA_ASSERT(!insres.second); // This is here just to get a chance to investigate this scenario.. might work
+			UBA_ASSERTF(!insres.second, TC("%s -> %s"), lpExistingFileName, lpNewFileName); // This is here just to get a chance to investigate this scenario.. might work
 			UBA_ASSERTF(!destInfo.trueFileMapHandle && (!destInfo.memoryFile || g_rules->IsThrowAway(dest, g_runningRemote)), TC("Moving file %s to %s that is an output file that is not a memory file is not supported"), source.data, lpNewFileName);
 			destInfo.memoryFile = sourceInfo.memoryFile;
 			sourceInfo.memoryFile = nullptr;
@@ -2381,29 +2387,30 @@ DWORD Detoured_GetFinalPathNameByHandleW(HANDLE hFile, LPTSTR lpszFilePath, DWOR
 		{
 			if (!fo->newName.empty())
 				fileName = fo->newName.c_str();
-			StringBuffer<> buffer;
-
-			if (dwFlags == 2)
-				buffer.Append(L"\\??\\");
 
+			StringBuffer<> buffer;
 			FixPath(fileName, g_virtualWorkingDir.data, g_virtualWorkingDir.count, buffer);
+			VirtualizePath(buffer);
+			u32 requiredBufferSize = buffer.count;
+			if (dwFlags == 2)
+				requiredBufferSize += 4;
 
-			if (cchFilePath <= buffer.count)
+			if (cchFilePath <= requiredBufferSize)
 			{
 				SetLastError(ERROR_NOT_ENOUGH_MEMORY);
 				DEBUG_LOG_DETOURED(L"GetFinalPathNameByHandleW", L"%llu (%ls) -> Error (not enough mem)", uintptr_t(hFile), lpszFilePath);
-				return buffer.count + 1;
+				return requiredBufferSize + 1;
 			}
 
 			// Unfortunately casing can be wrong here.. and we need to fix that. Let's use the directory table for that
 			// Note, this really only matters when building linux target from windows.. then there is path validation that errors if this is not properly fixed
-			StringBuffer<> buffer2;
+			g_directoryTable.GetFinalPath(buffer.Clear(), fileName);
+			VirtualizePath(buffer);
 			if (dwFlags == 2)
-				buffer2.Append(L"\\??\\");
-			g_directoryTable.GetFinalPath(buffer2, fileName);
-			UBA_ASSERT(buffer2.count == buffer.count);
+				buffer.Prepend(AsView(L"\\??\\"));
+			UBA_ASSERT(requiredBufferSize == buffer.count);
 
-			memcpy(lpszFilePath, buffer2.data, (buffer2.count + 1) * sizeof(tchar));
+			memcpy(lpszFilePath, buffer.data, (buffer.count + 1) * sizeof(tchar));
 			DEBUG_LOG_DETOURED(L"GetFinalPathNameByHandleW", L"%llu (%ls) -> Success", uintptr_t(hFile), lpszFilePath);
 
 			SetLastError(ERROR_SUCCESS);
@@ -2578,8 +2585,11 @@ HMODULE Detoured_LoadLibraryExW(LPCWSTR lpLibFileName, HANDLE hFile, DWORD dwFla
 	if (HMODULE res = Shared_LoadLibrary(lpLibFileName, hFile, dwFlags))
 		return res;
 
+	StringBuffer<> fileName(lpLibFileName);
+	DevirtualizePath(fileName);
+
 	DEBUG_LOG_TRUE(L"LoadLibraryExW", L"(%ls)", lpLibFileName);
-	return True_LoadLibraryExW(lpLibFileName, hFile, dwFlags);
+	return True_LoadLibraryExW(fileName.data, hFile, dwFlags);
 }
 
 HANDLE Detoured_GetStdHandle(DWORD nStdHandle)
@@ -2731,8 +2741,10 @@ BOOL Detoured_CreateProcessW(LPCWSTR lpApplicationName, LPWSTR lpCommandLine, LP
 
 	StringBuffer<> application;
 	if (lpApplicationName)
+	{
 		FixPath(application, lpApplicationName);
-
+		DevirtualizePath(application);
+	}
 	bool startSuspended = (dwCreationFlags & CREATE_SUSPENDED) != 0;
 
 	TString commandLine;
diff --git a/Engine/Source/Programs/UnrealBuildAccelerator/Detours/Private/Windows/UbaDetoursFunctionsNtDll.inl b/Engine/Source/Programs/UnrealBuildAccelerator/Detours/Private/Windows/UbaDetoursFunctionsNtDll.inl
--- a/Engine/Source/Programs/UnrealBuildAccelerator/Detours/Private/Windows/UbaDetoursFunctionsNtDll.inl
+++ b/Engine/Source/Programs/UnrealBuildAccelerator/Detours/Private/Windows/UbaDetoursFunctionsNtDll.inl
@@ -696,20 +696,7 @@ NTSTATUS NTAPI Shared_NtCreateFile(bool IsCreateFunc, PHANDLE hFileHandle, ACCES
 		return res;
 	}
 
-	if (g_vfsEntryCount)
-	{
-		if (fileName.count >= g_vfsMatchingLength && Equals(fileName.data, g_vfsEntries[0].vfs.data, g_vfsMatchingLength, true))
-		{
-			for (u32 i=0, e=g_vfsEntryCount; i!=e; ++i)
-			{
-				VfsEntry& entry = g_vfsEntries[i];
-				if (!fileName.StartsWith(entry.vfs.data))
-					continue;
-				StringBuffer<> temp2(fileName.data + entry.vfs.count);
-				fileName.Clear().Append(entry.local).Append(temp2);
-			}
-		}
-	}
+	DevirtualizePath(fileName);
 
 	u32 dirTableOffset = ~u32(0);
 
diff --git a/Engine/Source/Programs/UnrealBuildAccelerator/Detours/Private/Windows/UbaDetoursFunctionsUcrtBase.inl b/Engine/Source/Programs/UnrealBuildAccelerator/Detours/Private/Windows/UbaDetoursFunctionsUcrtBase.inl
--- a/Engine/Source/Programs/UnrealBuildAccelerator/Detours/Private/Windows/UbaDetoursFunctionsUcrtBase.inl
+++ b/Engine/Source/Programs/UnrealBuildAccelerator/Detours/Private/Windows/UbaDetoursFunctionsUcrtBase.inl
@@ -45,7 +45,7 @@ wchar_t* Detoured__wfullpath(wchar_t* absPath, const wchar_t* relPath, size_t ma
 		res = absPath;
 	}
 
-	DEBUG_LOG_DETOURED(L"_wfullpath", L"(%ls) -> %ls", relPath, res);
+	//DEBUG_LOG_DETOURED(L"_wfullpath", L"(%ls) -> %ls", relPath, res); // This is very spammy
 
 #if UBA_DEBUG_VALIDATE
 	if (g_validateFileAccess && !g_runningRemote)
@@ -113,15 +113,9 @@ char* Detoured__fullpath(char* absPath, const char* relPath, size_t maxLength)
 errno_t Detoured__get_wpgmptr(wchar_t** pValue)
 {
 	DETOURED_CALL(_get_wpgmptr);
-	if (g_runningRemote)
-	{
-		DEBUG_LOG_DETOURED(L"_get_wpgmptr", L"(%ls)", g_virtualApplication.data);
-		*pValue = (wchar_t*)g_virtualApplication.data;
-		return 0;
-	}
-	auto res = True__get_wpgmptr(pValue);
-	DEBUG_LOG_TRUE(L"_get_wpgmptr", L"%ls -> %u", pValue ? *pValue : L"<NoString>", res);
-	return res;
+	DEBUG_LOG_DETOURED(L"_get_wpgmptr", L"(%ls)", g_virtualApplication.data);
+	*pValue = (wchar_t*)g_virtualApplication.data;
+	return 0;
 }
 
 #if UBA_DEBUG_LOG_ENABLED
@@ -146,38 +140,44 @@ const wchar_t* WaccessResultToString(errno_t res)
 errno_t Detoured__waccess_s(const wchar_t* path, int mode)
 {
 	DETOURED_CALL(_waccess_s);
+
 	StringBuffer<> fixedPath;
 	FixPath(fixedPath, path);
+
+	const wchar_t* originalPath = path;(void)originalPath;
 	path = fixedPath.data;
 
-	StringBuffer<> temp;
-	if (g_runningRemote && StartsWith(path, g_exeDir.data))
+	if (g_runningRemote && fixedPath.StartsWith(g_exeDir.data))
 	{
-		temp.Append(g_virtualApplicationDir).Append(path + g_exeDir.count);
-		path = temp.data;
+		StringBuffer<> temp;
+		temp.Append(g_virtualApplicationDir).Append(fixedPath.data + g_exeDir.count);
+		fixedPath.Clear().Append(temp);
 	}
 
+
 	if (!CanDetour(path))
 	{
 		auto res = True__waccess_s(path, mode);
-		DEBUG_LOG_TRUE(L"_waccess_s", L"(NODETOUR) %ls %i -> %ls", path, mode, WaccessResultToString(res));
+		DEBUG_LOG_TRUE(L"_waccess_s", L"(NODETOUR) %ls %i -> %ls", originalPath, mode, WaccessResultToString(res));
 		return res;
 	}
 
+	DevirtualizePath(fixedPath);
+
 	FileAttributes attr;
 	const wchar_t* realName = Shared_GetFileAttributes(attr, path);
 
 	if (!attr.useCache)
 	{
 		auto res = True__waccess_s(realName, mode);
-		DEBUG_LOG_TRUE(L"_waccess_s", L"(NOCACHE) %ls %i -> %ls", path, mode, WaccessResultToString(res));
+		DEBUG_LOG_TRUE(L"_waccess_s", L"(NOCACHE) %ls %i -> %ls", originalPath, mode, WaccessResultToString(res));
 		return res;
 	}
 
 	SetLastError(attr.lastError);
 
 	errno_t res = attr.exists ? 0 : ENOENT;
-	DEBUG_LOG_DETOURED(L"_waccess_s", L"%ls %i -> %ls", path, mode, WaccessResultToString(res));
+	DEBUG_LOG_DETOURED(L"_waccess_s", L"%ls %i -> %ls", originalPath, mode, WaccessResultToString(res));
 	return res;
 }
 
diff --git a/Engine/Source/Programs/UnrealBuildAccelerator/Detours/Private/Windows/UbaDetoursFunctionsWin.cpp b/Engine/Source/Programs/UnrealBuildAccelerator/Detours/Private/Windows/UbaDetoursFunctionsWin.cpp
--- a/Engine/Source/Programs/UnrealBuildAccelerator/Detours/Private/Windows/UbaDetoursFunctionsWin.cpp
+++ b/Engine/Source/Programs/UnrealBuildAccelerator/Detours/Private/Windows/UbaDetoursFunctionsWin.cpp
@@ -45,7 +45,6 @@ DETOURED_FUNCTIONS_MEMORY_NON_WINE
 #include "UbaApplicationRules.h"
 #include "UbaDetoursShared.h"
 #include "UbaDetoursObjFilesPreloader.h"
-#include "UbaDetoursPrepopulatePchFiles.h"
 
 #include "Shlwapi.h"
 #include <detours/detours.h>
@@ -236,6 +235,83 @@ VfsEntry g_vfsEntries[32];
 u32 g_vfsEntryCount;
 u32 g_vfsMatchingLength;
 
+void PopulateVfs(BinaryReader& vfsReader)
+{
+	while (vfsReader.GetLeft())
+	{
+		vfsReader.ReadByte(); // Index, unused
+		StringBuffer<> str;
+		vfsReader.ReadString(str);
+		if (!str.count)
+		{
+			vfsReader.SkipString();
+			continue;
+		}
+
+		u32 index = g_vfsEntryCount++;
+		VfsEntry& vfsEntry = g_vfsEntries[index];
+		vfsEntry.vfs = g_memoryBlock.Strdup(str);
+		if (index == 0)
+			g_vfsMatchingLength = vfsEntry.vfs.count;
+		else
+		{
+			u32 shortest = Min(g_vfsMatchingLength, vfsEntry.vfs.count);
+			for (u32 i=0; i!=shortest; ++i)
+			{
+				if (g_vfsEntries[0].vfs.data[i] == vfsEntry.vfs.data[i])
+					continue;
+				shortest = i;
+				break;
+			}
+			g_vfsMatchingLength = shortest;
+		}
+		vfsReader.ReadString(str.Clear());
+		vfsEntry.local = g_memoryBlock.Strdup(str);
+	}
+}
+
+void DevirtualizePath(StringBufferBase& path)
+{
+	if (!g_vfsEntryCount)
+		return;
+
+	if (!Equals(path.data, g_vfsEntries[0].vfs.data, Min(path.count, g_vfsMatchingLength), true))
+		return;
+
+	// TODO: This is not great, the dirs above the vfs root should be empty except the dir to the roots
+	if (path.count < g_vfsMatchingLength)
+	{
+		path.Clear().Append(g_vfsEntries[0].local);
+		return;
+	}
+
+	for (u32 i=0, e=g_vfsEntryCount; i!=e; ++i)
+	{
+		VfsEntry& entry = g_vfsEntries[i];
+		if (!path.StartsWith(entry.vfs.data))
+			continue;
+		StringBuffer<> temp2(path.data + entry.vfs.count);
+		path.Clear().Append(entry.local).Append(temp2);
+		return;
+	}
+}
+
+void VirtualizePath(StringBufferBase& path)
+{
+	if (!g_vfsEntryCount)
+		return;
+	for (u32 i=0, e=g_vfsEntryCount; i!=e; ++i)
+	{
+		VfsEntry& entry = g_vfsEntries[i];
+		if (path.count < entry.local.count || !path.StartsWith(entry.local.data))
+			continue;
+		StringBuffer<> temp2(path.data + entry.local.count);
+		path.Clear().Append(entry.vfs).Append(temp2);
+		return;
+	}
+}
+
+
 struct SuppressCreateFileDetourScope
 {
 	SuppressCreateFileDetourScope() { ++t_disallowCreateFileDetour; }
@@ -901,43 +977,20 @@ void Init(const DetoursPayload& payload, u64 startTime)
 		if (u16 vfsSize = reader.ReadU16())
 		{
 			BinaryReader vfsReader(reader.GetPositionData(), 0, vfsSize);
-			while (vfsReader.GetLeft())
-			{
-				vfsReader.ReadByte(); // Index, unused
-				StringBuffer<> str;
-				vfsReader.ReadString(str);
-				if (!str.count)
-				{
-					vfsReader.SkipString();
-					continue;
-				}
-
-				u32 index = g_vfsEntryCount++;
-				VfsEntry& vfsEntry = g_vfsEntries[index];
-				vfsEntry.vfs = g_memoryBlock.Strdup(str);
-				if (index == 0)
-					g_vfsMatchingLength = vfsEntry.vfs.count;
-				else
-				{
-					u32 shortest = Min(g_vfsMatchingLength, vfsEntry.vfs.count);
-					for (u32 i=0; i!=shortest; ++i)
-					{
-						if (g_vfsEntries[0].vfs.data[i] == vfsEntry.vfs.data[i])
-							continue;
-						shortest = i;
-						break;
-					}
-					g_vfsMatchingLength = shortest;
-				}
-				vfsReader.ReadString(str.Clear());
-				vfsEntry.local = g_memoryBlock.Strdup(str);
-			}
+			PopulateVfs(vfsReader);
 		}
 		DEBUG_LOG_PIPE(L"Init", L"");
 	}
 
 	TrackInput(applicationBuffer.data);
 
+	if (g_rules->IsFullVirtualization())
+	{
+		VirtualizePath(applicationBuffer);
+		VirtualizePath(workingDirBuffer);
+		VirtualizePath(g_exeDir);
+	}
+
 	Shared_SetCurrentDirectory(workingDirBuffer.data);
 
 	{
@@ -988,6 +1041,10 @@ void Init(const DetoursPayload& payload, u64 startTime)
 		DEBUG_LOG(L"WorkingDir: %ls", g_virtualWorkingDir.data);
 		DEBUG_LOG(L"ExeDir: %ls", g_virtualApplicationDir.data);
 		DEBUG_LOG(L"ExeDir (true): %ls", g_exeDir.data);
+		for (u32 i=0; i!=g_vfsEntryCount; ++i)
+		{
+			DEBUG_LOG(L"Vfs: %s -> %s", g_vfsEntries[i].vfs.data, g_vfsEntries[i].local.data);
+		}
 	}
 	#endif
 
@@ -1030,8 +1087,6 @@ void Init(const DetoursPayload& payload, u64 startTime)
 		TimerScope ts(g_stats.preparseObjFiles);
 		g_objFilesPreloader.Start(cmdLine);
 	}
-	else if (g_rulesIndex == 1 || g_rulesIndex == 7 || g_rulesIndex == 11 || g_rulesIndex == 14)
-		PrepopulatePchIncludedFiles(cmdLine, g_rulesIndex);
 }
 
 void Deinit(u64 startTime)
diff --git a/Engine/Source/Programs/UnrealBuildAccelerator/Detours/Private/Windows/UbaDetoursObjFilesPreloader.h b/Engine/Source/Programs/UnrealBuildAccelerator/Detours/Private/Windows/UbaDetoursObjFilesPreloader.h
--- a/Engine/Source/Programs/UnrealBuildAccelerator/Detours/Private/Windows/UbaDetoursObjFilesPreloader.h
+++ b/Engine/Source/Programs/UnrealBuildAccelerator/Detours/Private/Windows/UbaDetoursObjFilesPreloader.h
@@ -13,6 +13,7 @@
 namespace uba
 {
 	constexpr bool UseLargePages = false;
+	void DevirtualizePath(StringBufferBase& path);
 
 	u64 GetLargePageSize()
 	{
@@ -112,6 +113,7 @@ namespace uba
 				return;
 			StringBuffer<> fileFull;
 			FixPath(fileFull, file.data);
+			DevirtualizePath(fileFull);
 
 			HANDLE fileHandle = CreateFileW(fileFull.data, GENERIC_READ, FILE_SHARE_READ, NULL, OPEN_EXISTING, FILE_ATTRIBUTE_NORMAL, NULL);
 			UBA_ASSERT(fileHandle != INVALID_HANDLE_VALUE);
diff --git a/Engine/Source/Programs/UnrealBuildAccelerator/Detours/Private/Windows/UbaDetoursPrepopulatePchFiles.h b/Engine/Source/Programs/UnrealBuildAccelerator/Detours/Private/Windows/UbaDetoursPrepopulatePchFiles.h
deleted file mode 100644
