[UBA]
* Added support for ignored roots when normalizing paths in files
* Added so vfs paths are ignored when calculating file hashes and files for upload

diff --git a/Engine/Source/Programs/UnrealBuildAccelerator/Common/Private/UbaRootPaths.cpp b/Engine/Source/Programs/UnrealBuildAccelerator/Common/Private/UbaRootPaths.cpp
--- a/Engine/Source/Programs/UnrealBuildAccelerator/Common/Private/UbaRootPaths.cpp
+++ b/Engine/Source/Programs/UnrealBuildAccelerator/Common/Private/UbaRootPaths.cpp
@@ -11,38 +11,7 @@
 {
 	bool RootPaths::RegisterRoot(Logger& logger, const tchar* rootPath, bool includeInKey, u8 id)
 	{
-		// Register rootPath both with single path separators and double path separators on windows because text files store them with double path separators
-		#if PLATFORM_WINDOWS
-		StringBuffer<> doubleSlash;
-		StringBuffer<> spaceEscaped;
-		bool hasSpace = false;
-		for (const tchar* it=rootPath; *it; ++it)
-		{
-			if (*it == ' ')
-			{
-				hasSpace = true;
-				spaceEscaped.Append('\\');
-			}
-			spaceEscaped.Append(*it);
-			doubleSlash.Append(*it);
-			if (*it == PathSeparator)
-				doubleSlash.Append(PathSeparator);
-		}
-
-		const tchar* rootPaths[] = { rootPath, doubleSlash.data, hasSpace ? spaceEscaped.data : TC("") };
-		#else
-		const tchar* rootPaths[] = { rootPath };
-		#endif
-
-		for (const tchar* rp : rootPaths)
-		{
-			u8 index = id;
-			if (!InternalRegisterRoot(logger, rp, includeInKey, index))
-				return false;
-			if (id)
-				++id;
-		}
-		return true;
+		return InternalRegisterRoot(logger, m_roots, rootPath, includeInKey, id);
 	}
 
 	bool RootPaths::RegisterSystemRoots(Logger& logger, u8 startId)
@@ -87,36 +56,24 @@
 		return true;
 	}
 
+	bool RootPaths::RegisterIgnoredRoot(Logger& logger, const tchar* rootPath)
+	{
+		return InternalRegisterRoot(logger, m_ignoredRoots, rootPath, false, 0);
+	}
+
 	const RootPaths::Root* RootPaths::FindRoot(const StringView& path) const
 	{
-		if (path.count < m_shortestRoot)
-			return nullptr;
-
-		StringBuffer<MaxPath> shortPath;
-		shortPath.Append(path.data, m_shortestRoot);
-		if (CaseInsensitiveFs)
-			shortPath.MakeLower();
-
-		StringKey key = ToStringKeyNoCheck(shortPath.data, m_shortestRoot);
-		for (u32 i=0, e=u32(m_roots.size()); i!=e; ++i)
-		{
-			auto& root = m_roots[i];
-			if (key != root.shortestPathKey)
-				continue;
-			if (!path.StartsWith(root.path.c_str()))
-				continue;
-			return &m_roots[i];
-		}
-		return nullptr;
+		return InternalFindRoot(m_roots, path);
 	}
 
 	static TString& EmptyString = *new TString(); // Need to leak to prevent shutdown hangs when running in managed process
 
 	const TString& RootPaths::GetRoot(u32 index) const
 	{
-		if (m_roots.size() <= index)
+		const Roots& roots = m_roots;
+		if (roots.size() <= index)
 			return EmptyString;
-		return m_roots[index].path;
+		return roots[index].path;
 	}
 
 	CasKey RootPaths::NormalizeAndHashFile(Logger& logger, const tchar* filename) const
@@ -138,20 +95,56 @@
 		return AsNormalized(ToCasKey(hasher, false), wasNormalized);
 	}
 
-	bool RootPaths::InternalRegisterRoot(Logger& logger, const tchar* rootPath, bool includeInKey, u8 index)
+	bool RootPaths::InternalRegisterRoot(Logger& logger, Roots& roots, const tchar* rootPath, bool includeInKey, u8 id)
 	{
+		// Register rootPath both with single path separators and double path separators on windows because text files store them with double path separators
+		#if PLATFORM_WINDOWS
+		StringBuffer<> doubleSlash;
+		StringBuffer<> spaceEscaped;
+		bool hasSpace = false;
+		for (const tchar* it=rootPath; *it; ++it)
+		{
+			if (*it == ' ')
+			{
+				hasSpace = true;
+				spaceEscaped.Append('\\');
+			}
+			spaceEscaped.Append(*it);
+			doubleSlash.Append(*it);
+			if (*it == PathSeparator)
+				doubleSlash.Append(PathSeparator);
+		}
+
+		const tchar* rootPaths[] = { rootPath, doubleSlash.data, hasSpace ? spaceEscaped.data : TC("") };
+		#else
+		const tchar* rootPaths[] = { rootPath };
+		#endif
+
+		for (const tchar* rp : rootPaths)
+		{
+			u8 index = id;
+			if (!InternalRegisterRoot2(logger, roots, rp, includeInKey, index))
+				return false;
+			if (id)
+				++id;
+		}
+		return true;
+	}
+
+	bool RootPaths::InternalRegisterRoot2(Logger& logger, Roots& roots, const tchar* rootPath, bool includeInKey, u8 index)
+	{
 		if (index == 0)
-			index = u8(m_roots.size());
+			index = u8(roots.size());
 		if (index == '~' - ' ') // This is not really true.. as long as value is under 256 we're good
 			return logger.Error(TC("Too many roots added (%u)"), index);
 
-		if (index >= m_roots.size())
-			m_roots.resize(index+1);
+		if (index >= roots.size())
+			roots.resize(index+1);
 
 		if (!*rootPath)
 			return true;
 
-		auto& root = m_roots[index];
+		auto& root = roots[index];
 		if (!root.path.empty())
 			return logger.Error(TC("Root at index %u already added (existing as %s, added as %s)"), index, root.path.c_str(), rootPath);
 
@@ -163,16 +156,39 @@
 
 		root.includeInKey = includeInKey;
 
-		m_longestRoot = Max(u32(root.path.size()), m_longestRoot);
+		roots.longestRoot = Max(u32(root.path.size()), roots.longestRoot);
 
-		if (!m_shortestRoot || root.path.size() < m_shortestRoot)
+		if (!roots.shortestRoot || root.path.size() < roots.shortestRoot)
 		{
-			m_shortestRoot = u32(root.path.size());
-			for (auto& r : m_roots)
-				r.shortestPathKey = ToStringKeyNoCheck(r.path.data(), m_shortestRoot);
+			roots.shortestRoot = u32(root.path.size());
+			for (auto& r : roots)
+				r.shortestPathKey = ToStringKeyNoCheck(r.path.data(), roots.shortestRoot);
 		}
 		else
-			root.shortestPathKey = ToStringKeyNoCheck(root.path.data(), m_shortestRoot);
+			root.shortestPathKey = ToStringKeyNoCheck(root.path.data(), roots.shortestRoot);
 		return true;
 	}
+
+	const RootPaths::Root* RootPaths::InternalFindRoot(const Roots& roots, const StringView& path) const
+	{
+		if (path.count < roots.shortestRoot)
+			return nullptr;
+
+		StringBuffer<MaxPath> shortPath;
+		shortPath.Append(path.data, roots.shortestRoot);
+		if (CaseInsensitiveFs)
+			shortPath.MakeLower();
+
+		StringKey key = ToStringKeyNoCheck(shortPath.data, roots.shortestRoot);
+		for (u32 i=0, e=u32(roots.size()); i!=e; ++i)
+		{
+			auto& root = roots[i];
+			if (key != root.shortestPathKey)
+				continue;
+			if (!path.StartsWith(root.path.c_str()))
+				continue;
+			return &roots[i];
+		}
+		return nullptr;
+	}
 }
diff --git a/Engine/Source/Programs/UnrealBuildAccelerator/Common/Public/UbaRootPaths.h b/Engine/Source/Programs/UnrealBuildAccelerator/Common/Public/UbaRootPaths.h
--- a/Engine/Source/Programs/UnrealBuildAccelerator/Common/Public/UbaRootPaths.h
+++ b/Engine/Source/Programs/UnrealBuildAccelerator/Common/Public/UbaRootPaths.h
@@ -15,6 +15,7 @@ namespace uba
 
 		bool RegisterRoot(Logger& logger, const tchar* rootPath, bool includeInKey = true, u8 id = 0);
 		bool RegisterSystemRoots(Logger& logger, u8 startId = 0);
+		bool RegisterIgnoredRoot(Logger& logger, const tchar* rootPath);
 
 		struct Root
 		{
@@ -35,11 +36,18 @@ namespace uba
 		static constexpr u8 RootStartByte = ' ';
 
 	private:
-		bool InternalRegisterRoot(Logger& logger, const tchar* rootPath, bool includeInKey, u8 index);
+		struct Roots : Vector<Root>
+		{
+			u32 shortestRoot = 0;
+			u32 longestRoot = 0;
+		};
+
+		bool InternalRegisterRoot(Logger& logger, Roots& roots, const tchar* rootPath, bool includeInKey, u8 id);
+		bool InternalRegisterRoot2(Logger& logger, Roots& roots, const tchar* rootPath, bool includeInKey, u8 index);
+		const Root* InternalFindRoot(const Roots& roots, const StringView& path) const;
 
-		Vector<Root> m_roots;
-		u32 m_shortestRoot = 0;
-		u32 m_longestRoot = 0;
+		Roots m_roots;
+		Roots m_ignoredRoots;
 	};
 
 
@@ -75,14 +83,14 @@ namespace uba
 				++pathEndOrMore;
 
 			u32 lenOrMore = u32(pathEndOrMore - pathStart);
-			u32 toCopy = Min(lenOrMore, m_longestRoot);
+			u32 toCopy = Min(lenOrMore, m_roots.longestRoot);
 			StringBuffer<512> path;
 			path.Append(pathStart, toCopy);
 
 			auto root = FindRoot(path);
 			if (!root)
 			{
-				if (allowPathsWithoutRoot)
+				if (allowPathsWithoutRoot || InternalFindRoot(m_ignoredRoots, path))
 				{
 					u32 len = 1;
 					destPos += len;
diff --git a/Engine/Source/Programs/UnrealBuildAccelerator/Host/Private/UbaExports.cpp b/Engine/Source/Programs/UnrealBuildAccelerator/Host/Private/UbaExports.cpp
--- a/Engine/Source/Programs/UnrealBuildAccelerator/Host/Private/UbaExports.cpp
+++ b/Engine/Source/Programs/UnrealBuildAccelerator/Host/Private/UbaExports.cpp
@@ -477,7 +477,11 @@ extern "C"
 
 	uba::RootPaths* RootPaths_Create(uba::LogWriter& writer)
 	{
-		return new uba::RootPathsWithLogger(writer);
+		auto rootPaths = new uba::RootPathsWithLogger(writer);
+		#if PLATFORM_WINDOWS
+		rootPaths->RegisterIgnoredRoot(rootPaths->logger, TC("z:\\vfs")); // TODO: Not hard coded
+		#endif
+		return rootPaths;
 	}
 
 	bool RootPaths_RegisterRoot(uba::RootPaths* rootPaths, const uba::tchar* path, bool includeInKey, uba::u8 id)
