[UBA]
* Improved log message more when finding bad paths
* Fixed hard coded path for vfs roots

diff --git a/Engine/Source/Programs/UnrealBuildAccelerator/Common/Private/UbaCacheClient.cpp b/Engine/Source/Programs/UnrealBuildAccelerator/Common/Private/UbaCacheClient.cpp
--- a/Engine/Source/Programs/UnrealBuildAccelerator/Common/Private/UbaCacheClient.cpp
+++ b/Engine/Source/Programs/UnrealBuildAccelerator/Common/Private/UbaCacheClient.cpp
@@ -1165,7 +1165,7 @@ namespace uba
 
 		// Add arguments list to key
 		auto hashString = [&](const tchar* str, u64 strLen, u32 rootPos) { hasher.Update(str, strLen*sizeof(tchar)); };
-		if (!rootPaths.NormalizeString(m_logger, info.arguments, TStrlen(info.arguments), hashString, false, info.GetDescription(), TC("Calculating command line hash")))
+		if (!rootPaths.NormalizeString(m_logger, info.arguments, TStrlen(info.arguments), hashString, false, info.GetDescription(), TC("calculating command line hash")))
 			return CasKeyZero;
 
 		// Add content of rsp file to key (This will cost a bit of perf since we need to normalize.. should this be part of key?)
diff --git a/Engine/Source/Programs/UnrealBuildAccelerator/Common/Public/UbaRootPaths.h b/Engine/Source/Programs/UnrealBuildAccelerator/Common/Public/UbaRootPaths.h
--- a/Engine/Source/Programs/UnrealBuildAccelerator/Common/Public/UbaRootPaths.h
+++ b/Engine/Source/Programs/UnrealBuildAccelerator/Common/Public/UbaRootPaths.h
@@ -103,7 +103,7 @@ namespace uba
 					path.Resize(lastQuote - path.data);
 				if (auto lineEnd = path.Last('\r'))
 					path.Resize(lineEnd - path.data);
-				logger.Info(TC("PATH WITHOUT ROOT: %s (inside %s%s at offset %u)"), path.data, hint, hint2, destPos);
+				logger.Info(TC("PATH WITHOUT ROOT: %s (inside %s at offset %u%s)"), path.data, hint, hint2, destPos);
 				return false;
 			}
 
diff --git a/Engine/Source/Programs/UnrealBuildAccelerator/Host/Private/UbaExports.cpp b/Engine/Source/Programs/UnrealBuildAccelerator/Host/Private/UbaExports.cpp
--- a/Engine/Source/Programs/UnrealBuildAccelerator/Host/Private/UbaExports.cpp
+++ b/Engine/Source/Programs/UnrealBuildAccelerator/Host/Private/UbaExports.cpp
@@ -479,7 +479,7 @@ extern "C"
 	{
 		auto rootPaths = new uba::RootPathsWithLogger(writer);
 		#if PLATFORM_WINDOWS
-		rootPaths->RegisterIgnoredRoot(rootPaths->logger, TC("z:\\vfs")); // TODO: Not hard coded
+		rootPaths->RegisterIgnoredRoot(rootPaths->logger, TC("z:\\ue")); // TODO: Not hard coded
 		#endif
 		return rootPaths;
 	}
