[UBA]
* Added so cache client throws a warning if not finding the rsp file when calculating cache key

diff --git a/Engine/Source/Programs/UnrealBuildAccelerator/Common/Private/UbaCacheClient.cpp b/Engine/Source/Programs/UnrealBuildAccelerator/Common/Private/UbaCacheClient.cpp
--- a/Engine/Source/Programs/UnrealBuildAccelerator/Common/Private/UbaCacheClient.cpp
+++ b/Engine/Source/Programs/UnrealBuildAccelerator/Common/Private/UbaCacheClient.cpp
@@ -1182,7 +1182,7 @@ namespace uba
 					rsp.Append(rspStart, rspEnd - rspStart);
 					StringBuffer<> fullPath;
 					FixPath(rsp.data, workingDir.data, workingDir.count, fullPath);
-					CasKey rspCasKey = rootPaths.NormalizeAndHashFile(m_logger, rsp.data);
+					CasKey rspCasKey = rootPaths.NormalizeAndHashFile(m_logger, rsp.data, true);
 					if (rspCasKey == CasKeyZero)
 						return CasKeyZero;
 					hasher.Update(&rspCasKey, sizeof(CasKey));
diff --git a/Engine/Source/Programs/UnrealBuildAccelerator/Common/Private/UbaRootPaths.cpp b/Engine/Source/Programs/UnrealBuildAccelerator/Common/Private/UbaRootPaths.cpp
--- a/Engine/Source/Programs/UnrealBuildAccelerator/Common/Private/UbaRootPaths.cpp
+++ b/Engine/Source/Programs/UnrealBuildAccelerator/Common/Private/UbaRootPaths.cpp
@@ -76,12 +76,15 @@ namespace uba
 		return roots[index].path;
 	}
 
-	CasKey RootPaths::NormalizeAndHashFile(Logger& logger, const tchar* filename) const
+	CasKey RootPaths::NormalizeAndHashFile(Logger& logger, const tchar* filename, bool warnOnFileNotFound) const
 	{
 		FileAccessor file(logger, filename);
 		if (!file.OpenMemoryRead(0, false))
+		{
+			if (warnOnFileNotFound)
+				logger.Warning(TC("NormalizeAndHashFile Can't find file %s"), filename);
 			return CasKeyZero;
-
+		}
 		bool wasNormalized = false;
 		CasKeyHasher hasher;
 		auto hashString = [&](const char* str, u64 strLen, u32 rootPos)
diff --git a/Engine/Source/Programs/UnrealBuildAccelerator/Common/Public/UbaRootPaths.h b/Engine/Source/Programs/UnrealBuildAccelerator/Common/Public/UbaRootPaths.h
--- a/Engine/Source/Programs/UnrealBuildAccelerator/Common/Public/UbaRootPaths.h
+++ b/Engine/Source/Programs/UnrealBuildAccelerator/Common/Public/UbaRootPaths.h
@@ -31,7 +31,7 @@ namespace uba
 		template<typename CharType, typename Func>
 		bool NormalizeString(Logger& logger, const CharType* str, u64 strLen, const Func& func, bool allowPathsWithoutRoot, const tchar* hint, const tchar* hint2 = TC("")) const;
 
-		CasKey NormalizeAndHashFile(Logger& logger, const tchar* filename) const;
+		CasKey NormalizeAndHashFile(Logger& logger, const tchar* filename, bool warnOnFileNotFound = false) const;
 
 		static constexpr u8 RootStartByte = ' ';
 
