[UBA]
* Added ReportCacheKey setting that can be used to log out cache key

diff --git a/Engine/Source/Programs/UnrealBuildAccelerator/Common/Private/UbaCacheClient.cpp b/Engine/Source/Programs/UnrealBuildAccelerator/Common/Private/UbaCacheClient.cpp
--- a/Engine/Source/Programs/UnrealBuildAccelerator/Common/Private/UbaCacheClient.cpp
+++ b/Engine/Source/Programs/UnrealBuildAccelerator/Common/Private/UbaCacheClient.cpp
@@ -29,6 +29,7 @@ namespace uba
 		table.GetValueAsBool(useDirectoryPreparsing, TC("UseDirectoryPreparsing"));
 		table.GetValueAsBool(validateCacheWritesInput, TC("ValidateCacheWritesInput"));
 		table.GetValueAsBool(validateCacheWritesOutput, TC("ValidateCacheWritesOutput"));
+		table.GetValueAsBool(reportCacheKey, TC("ReportCacheKey"));
 		table.GetValueAsBool(reportMissReason, TC("ReportMissReason"));
 		table.GetValueAsBool(useRoots, TC("UseRoots"));
 		table.GetValueAsBool(useCacheHit, TC("UseCacheHit"));
@@ -67,6 +68,7 @@ namespace uba
 	,	m_client(info.client)
 	,	m_session(info.session)
 	{
+		m_reportCacheKey = info.reportCacheKey;
 		m_reportMissReason = info.reportMissReason;
 		#if UBA_LOG_FETCH_CACHE_INFO
 		m_reportMissReason = true;
@@ -406,7 +408,14 @@ namespace uba
 
 		CasKey cmdKey = GetCmdKey(rootPaths, info);
 		if (cmdKey == CasKeyZero)
+		{
+			if (m_reportCacheKey)
+				m_logger.Info(TC("CACHEKEY %s: Failed to calculate cache key"), info.GetDescription());
 			return false;
+		}
+
+		if (m_reportCacheKey)
+			m_logger.Info(TC("CACHEKEY %s: %s (bucket %u)"), info.GetDescription(), CasKeyString(cmdKey).str, bucketId);
 
 		u8 memory[SendMaxSize];
 
@@ -451,7 +460,11 @@ namespace uba
 		#endif
 
 		if (!entryCount)
+		{
+			if (m_reportMissReason)
+				m_logger.Info(TC("Cache miss on %s because no entry with key %s was found in bucket %u"), info.GetDescription(), CasKeyString(cmdKey).str, bucketId);
 			return false;
+		}
 
 		struct MissInfo { TString path; u32 entryIndex; CasKey cache; CasKey local; };
 		Vector<MissInfo> misses;
@@ -775,7 +788,7 @@ namespace uba
 		}
 
 		for (auto& miss : misses)
-			m_logger.Info(TC("Cache miss on %s because of mismatch of %s (entry: %u, local: %s cache: %s)"), info.GetDescription(), miss.path.data(), miss.entryIndex, CasKeyString(miss.local).str, CasKeyString(miss.cache).str);
+			m_logger.Info(TC("Cache miss on %s because of mismatch of %s (entry: %u, local: %s cache: %s bucket: %u)"), info.GetDescription(), miss.path.data(), miss.entryIndex, CasKeyString(miss.local).str, CasKeyString(miss.cache).str, bucketId);
 
 		return false;
 	}
diff --git a/Engine/Source/Programs/UnrealBuildAccelerator/Common/Public/UbaCacheClient.h b/Engine/Source/Programs/UnrealBuildAccelerator/Common/Public/UbaCacheClient.h
--- a/Engine/Source/Programs/UnrealBuildAccelerator/Common/Public/UbaCacheClient.h
+++ b/Engine/Source/Programs/UnrealBuildAccelerator/Common/Public/UbaCacheClient.h
@@ -30,6 +30,7 @@ namespace uba
 
 		void Apply(Config& config);
 
+		bool reportCacheKey = false;
 		bool reportMissReason = false; // Report the reason no matching cache entry was found
 		bool useDirectoryPreparsing = true; // This is used to minimize syscalls. GetFileAttributes can be very expensive on cloud machines and we can enable this to minimize syscall count
 		bool validateCacheWritesInput = false; // Set to true to validate cas of all input files before sent to cache
@@ -84,6 +85,7 @@ namespace uba
 		StorageImpl& m_storage;
 		NetworkClient& m_client;
 		Session& m_session;
+		bool m_reportCacheKey;
 		bool m_reportMissReason;
 		bool m_useDirectoryPreParsing;
 		bool m_validateCacheWritesInput;
