[UBA]
* Extended logging for report cache key to contain each part contributing to the key

diff --git a/Engine/Source/Programs/UnrealBuildAccelerator/Common/Private/UbaCacheClient.cpp b/Engine/Source/Programs/UnrealBuildAccelerator/Common/Private/UbaCacheClient.cpp
--- a/Engine/Source/Programs/UnrealBuildAccelerator/Common/Private/UbaCacheClient.cpp
+++ b/Engine/Source/Programs/UnrealBuildAccelerator/Common/Private/UbaCacheClient.cpp
@@ -141,7 +141,7 @@ namespace uba
 		if (!inputsSize)
 			return false;
 
-		CasKey cmdKey = GetCmdKey(rootPaths, info);
+		CasKey cmdKey = GetCmdKey(rootPaths, info, false, bucketId);
 		if (cmdKey == CasKeyZero)
 		{
 			#if UBA_LOG_WRITE_CACHE_INFO
@@ -406,16 +406,9 @@ namespace uba
 		StorageStatsScope __(storageStats);
 		KernelStatsScope _(kernelStats);
 
-		CasKey cmdKey = GetCmdKey(rootPaths, info);
+		CasKey cmdKey = GetCmdKey(rootPaths, info, m_reportCacheKey, bucketId);
 		if (cmdKey == CasKeyZero)
-		{
-			if (m_reportCacheKey)
-				m_logger.Info(TC("CACHEKEY %s: Failed to calculate cache key"), info.GetDescription());
 			return false;
-		}
-
-		if (m_reportCacheKey)
-			m_logger.Info(TC("CACHEKEY %s: %s (bucket %u)"), info.GetDescription(), CasKeyString(cmdKey).str, bucketId);
 
 		u8 memory[SendMaxSize];
 
@@ -1154,11 +1147,22 @@ namespace uba
 		return true;
 	}
 
-	CasKey CacheClient::GetCmdKey(const RootPaths& rootPaths, const ProcessStartInfo& info)
+	CasKey CacheClient::GetCmdKey(const RootPaths& rootPaths, const ProcessStartInfo& info, bool report, u32 bucketId)
 	{
 		CasKeyHasher hasher;
 
+		//if (m_reportCacheKey)
+		//	m_logger.Info(TC("CACHEKEY %s: Failed to calculate cache key"), info.GetDescription());
+		//if (m_reportCacheKey)
+		//	m_logger.Info(TC("CACHEKEY %s: %s (bucket %u)"), info.GetDescription(), CasKeyString(cmdKey).str, bucketId);
 		
+		if (report)
+		{
+			m_logger.BeginScope();
+			m_logger.Info(TC("CACHEKEY %s (bucket %u)"), info.GetDescription(), bucketId);
+		}
+		auto guard = MakeGuard([&]() {if (report) m_logger.EndScope(); });
+
 		#if PLATFORM_WINDOWS
 		// cmd.exe is special.. we can't hash it because it might be different on different os versions but should do the same thing regardless of version
 		if (Contains(info.application, TC("cmd.exe")))
@@ -1173,13 +1177,22 @@ namespace uba
 			bool fileIsCompressed = false;
 			if (!m_storage.StoreCasKey(applicationCasKey, info.application, CasKeyZero, fileIsCompressed))
 				return CasKeyZero;
+			if (report)
+				m_logger.Info(TC("   %s %s"), CasKeyString(applicationCasKey).str, info.application);
 			hasher.Update(&applicationCasKey, sizeof(CasKey));
 		}
 
 		// Add arguments list to key
 		auto hashString = [&](const tchar* str, u64 strLen, u32 rootPos) { hasher.Update(str, strLen*sizeof(tchar)); };
 		if (!rootPaths.NormalizeString(m_logger, info.arguments, TStrlen(info.arguments), hashString, false, info.GetDescription(), TC("calculating command line hash")))
+		{
+			if (report)
+				m_logger.Info(TC("   Failed to normalize commandline %s"), info.arguments);
 			return CasKeyZero;
+		}
+		if (report)
+			m_logger.Info(TC("   %s %s"), CasKeyString(ToCasKey(hasher, false)).str, info.arguments);
+
 
 		// Add content of rsp file to key (This will cost a bit of perf since we need to normalize.. should this be part of key?)
 		if (auto rspStart = TStrchr(info.arguments, '@'))
@@ -1197,7 +1210,13 @@ namespace uba
 					FixPath(rsp.data, workingDir.data, workingDir.count, fullPath);
 					CasKey rspCasKey = rootPaths.NormalizeAndHashFile(m_logger, rsp.data, true);
 					if (rspCasKey == CasKeyZero)
+					{
+						if (report)
+							m_logger.Info(TC("   Failed to normalize rsp file %s"), rsp.data);
 						return CasKeyZero;
+					}
+					if (report)
+						m_logger.Info(TC("   %s %s"), CasKeyString(rspCasKey).str, rsp.data);
 					hasher.Update(&rspCasKey, sizeof(CasKey));
 				}
 			}
diff --git a/Engine/Source/Programs/UnrealBuildAccelerator/Common/Public/UbaCacheClient.h b/Engine/Source/Programs/UnrealBuildAccelerator/Common/Public/UbaCacheClient.h
--- a/Engine/Source/Programs/UnrealBuildAccelerator/Common/Public/UbaCacheClient.h
+++ b/Engine/Source/Programs/UnrealBuildAccelerator/Common/Public/UbaCacheClient.h
@@ -74,7 +74,7 @@ namespace uba
 		bool ReportUsedEntry(Vector<ProcessLogLine>& outLogLines, bool ownedLogLines, Bucket& bucket, const CasKey& cmdKey, u32 entryId);
 		bool PopulateLogLines(Vector<ProcessLogLine>& outLogLines, const u8* mem, u64 memLen);
 
-		CasKey GetCmdKey(const RootPaths& rootPaths, const ProcessStartInfo& info);
+		CasKey GetCmdKey(const RootPaths& rootPaths, const ProcessStartInfo& info, bool report, u32 bucketId);
 		bool ShouldNormalize(const StringBufferBase& path);
 
 		bool GetLocalPathAndCasKey(Bucket& bucket, const RootPaths& rootPaths, StringBufferBase& outPath, CasKey& outKey, CompactCasKeyTable& casKeyTable, CompactPathTable& pathTable, u32 offset);
