[UBA]
* Fixed bug where StringBufferBase::Prepend did not properly increase count

diff --git a/Engine/Source/Programs/UnrealBuildAccelerator/Core/Private/UbaStringBuffer.cpp b/Engine/Source/Programs/UnrealBuildAccelerator/Core/Private/UbaStringBuffer.cpp
--- a/Engine/Source/Programs/UnrealBuildAccelerator/Core/Private/UbaStringBuffer.cpp
+++ b/Engine/Source/Programs/UnrealBuildAccelerator/Core/Private/UbaStringBuffer.cpp
@@ -258,6 +258,7 @@ namespace uba
 		UBA_ASSERT(count + view.count < capacity);
 		memmove(data+view.count, data, (count+1)*sizeof(tchar));
 		memcpy(data, view.data, view.count*sizeof(tchar));
+		count += view.count;
 		return *this;
 	}
 
