[UBA]
* Fixed so UbaObjTool can handle paths with forward slashes

diff --git a/Engine/Source/Programs/UnrealBuildAccelerator/ObjTool/Private/UbaObjTool.cpp b/Engine/Source/Programs/UnrealBuildAccelerator/ObjTool/Private/UbaObjTool.cpp
--- a/Engine/Source/Programs/UnrealBuildAccelerator/ObjTool/Private/UbaObjTool.cpp
+++ b/Engine/Source/Programs/UnrealBuildAccelerator/ObjTool/Private/UbaObjTool.cpp
@@ -95,7 +95,7 @@ namespace uba
 				else
 				{
 					const tchar* colon = TStrchr(arg,':');
-					if (colon && colon[1] != '\\')
+					if (colon && colon[1] != '\\' && colon[1] != '/')
 					{
 						name.Append(arg, colon - arg);
 						const tchar* valueStart = colon+1;
