[UBA]
* Fixed vfs for more code paths
* Removed IsFullVirtualization rule function now when we have decided to always do full virtualization (no more vfsoverlay)

diff --git a/Engine/Source/Programs/UnrealBuildAccelerator/Core/Private/UbaApplicationRules.cpp b/Engine/Source/Programs/UnrealBuildAccelerator/Core/Private/UbaApplicationRules.cpp
--- a/Engine/Source/Programs/UnrealBuildAccelerator/Core/Private/UbaApplicationRules.cpp
+++ b/Engine/Source/Programs/UnrealBuildAccelerator/Core/Private/UbaApplicationRules.cpp
@@ -99,11 +99,6 @@ namespace uba
 			return fileName.EndsWith(TC(".obj"));
 		}
 
-		virtual bool IsFullVirtualization() const override
-		{
-			return true;
-		}
-
 		virtual bool ShouldDevirtualizeFile(const StringView& fileName, bool& outEscapeSpaces) const
 		{
 			outEscapeSpaces = false;
@@ -173,11 +168,6 @@ namespace uba
 		{
 			return fileName.IsEmpty() || fileName.EndsWith(TC(".obj"));
 		}
-
-		virtual bool IsFullVirtualization() const override
-		{
-			return true;
-		}
 	};
 
 	class ApplicationRulesLinkExe : public ApplicationRulesVcLink
@@ -506,6 +496,11 @@ namespace uba
 		{
 			return fileName.EndsWith(TC(".obj")) || fileName.EndsWith(TC(".o"));
 		}
+
+		virtual bool ShouldDevirtualizeFile(const StringView& fileName, bool& outEscapeSpaces) const
+		{
+			return fileName.EndsWith(TC(".ispc.txt"));
+		}
 	};
 
 	class ApplicationRulesUBTDll : public ApplicationRules
diff --git a/Engine/Source/Programs/UnrealBuildAccelerator/Core/Public/UbaApplicationRules.h b/Engine/Source/Programs/UnrealBuildAccelerator/Core/Public/UbaApplicationRules.h
--- a/Engine/Source/Programs/UnrealBuildAccelerator/Core/Public/UbaApplicationRules.h
+++ b/Engine/Source/Programs/UnrealBuildAccelerator/Core/Public/UbaApplicationRules.h
@@ -168,12 +168,6 @@ namespace uba
 			return nullptr;
 		}
 
-		// Full virtualization means that workingdir, exedir etc are virtualized.
-		virtual bool IsFullVirtualization() const
-		{
-			return false;
-		}
-
 		// Returns true if produced file should be devirtualized (paths are turned into local paths)
 		virtual bool ShouldDevirtualizeFile(const StringView& fileName, bool& outEscapeSpaces) const
 		{
diff --git a/Engine/Source/Programs/UnrealBuildAccelerator/Detours/Private/Windows/UbaDetoursFunctionsKernelBase.inl b/Engine/Source/Programs/UnrealBuildAccelerator/Detours/Private/Windows/UbaDetoursFunctionsKernelBase.inl
--- a/Engine/Source/Programs/UnrealBuildAccelerator/Detours/Private/Windows/UbaDetoursFunctionsKernelBase.inl
+++ b/Engine/Source/Programs/UnrealBuildAccelerator/Detours/Private/Windows/UbaDetoursFunctionsKernelBase.inl
@@ -438,9 +438,9 @@ BOOL Detoured_ReadConsoleW(HANDLE hConsoleInput, LPVOID lpBuffer, DWORD nNumberO
 UINT Detoured_GetDriveTypeW(LPCWSTR lpRootPathName)
 {
 	DETOURED_CALL(GetDriveTypeW);
-	if (g_runningRemote)
+	if (g_runningRemote || g_vfsEntryCount)
 	{
-		DEBUG_LOG_DETOURED(L"GetDriveType", L"%ls", lpRootPathName);
+		DEBUG_LOG_DETOURED(L"GetDriveType", L"%ls -> DRIVE_FIXED", lpRootPathName);
 		return DRIVE_FIXED;
 	}
 	DEBUG_LOG_TRUE(L"GetDriveType", L"%ls", lpRootPathName);
@@ -1103,10 +1103,15 @@ DWORD Detoured_GetLongPathNameW(LPCWSTR lpszShortPath, LPWSTR lpszLongPath, DWOR
 	{
 		StringBuffer<> fixedName;
 		FixPath(fixedName, path);
+		bool success;
 
-		DEBUG_LOG_DETOURED(L"GetLongPathNameW", L"(%ls)", path);
-		WIN32_FILE_ATTRIBUTE_DATA data;
-		bool success = Shared_GetFileAttributesExW(fixedName.data, GetFileExInfoStandard, &data, path);
+		{
+			DEBUG_LOG_DETOURED(L"GetLongPathNameW", L"(%ls)", path);
+			StringBuffer<> realName(fixedName);
+			DevirtualizePath(realName);
+			WIN32_FILE_ATTRIBUTE_DATA data;
+			success = Shared_GetFileAttributesExW(realName.data, GetFileExInfoStandard, &data, path);
+		}
 
 		DWORD res = 0;
 		if (success)
@@ -1165,7 +1170,7 @@ BOOL Detoured_GetVolumePathNameW(LPCWSTR lpszFileName, LPWSTR lpszVolumePathName
 	//	return TRUE;
 	//}
 
-	if (g_runningRemote)
+	if (g_runningRemote || g_vfsEntryCount)
 	{
 		UBA_ASSERT(cchBufferLength > 3);
 		memcpy(lpszVolumePathName, g_virtualWorkingDir.data, 6);
@@ -1218,7 +1223,7 @@ DWORD Detoured_GetModuleFileNameW(HMODULE hModule, LPWSTR lpFilename, DWORD nSiz
 			return Shared_GetModuleFileNameW(hModule, findIt->second.c_str(), u32(findIt->second.size()), lpFilename, nSize);
 	}
 
-	if (!g_runningRemote)
+	if (!g_runningRemote && !g_vfsEntryCount)
 		return True_GetModuleFileNameW(hModule, lpFilename, nSize);
 
 	wchar_t moduleName[350];
@@ -1254,10 +1259,12 @@ BOOL Detoured_CopyFileExW(LPCWSTR lpExistingFileName, LPCWSTR lpNewFileName, LPP
 
 	StringBuffer<> fromName;
 	FixPath(fromName, lpExistingFileName);
+	DevirtualizePath(fromName);
 	StringKey fromKey = ToStringKeyLower(fromName);
 
 	StringBuffer<> toName;
 	FixPath(toName, lpNewFileName);
+	DevirtualizePath(toName);
 	StringKey toKey = ToStringKeyLower(toName);
 
 	StringBuffer<> newFromName;
@@ -1326,10 +1333,12 @@ BOOL Detoured_CreateHardLinkW(LPCWSTR lpFileName, LPCWSTR lpExistingFileName, LP
 
 	StringBuffer<> fromName;
 	FixPath(fromName, lpExistingFileName);
+	DevirtualizePath(fromName);
 	StringKey fromKey = ToStringKeyLower(fromName);
 
 	StringBuffer<> toName;
 	FixPath(toName, lpFileName);
+	DevirtualizePath(toName);
 	StringKey toKey = ToStringKeyLower(toName);
 
 	StringBuffer<> newFromName;
@@ -1607,6 +1616,7 @@ __forceinline HANDLE Shared_FindFirstFileExW(LPCWSTR lpFileName, FINDEX_INFO_LEV
 
 	StringBuffer<> lowerName;
 	FixPath(lowerName, lpFileName);
+	DevirtualizePath(lowerName);
 
 	if (lowerName.StartsWith(g_systemTemp.data) || lowerName.StartsWith(g_systemRoot.data))
 		return Local_FindFirstFileExW(lpFileName, fInfoLevelId, lpFindFileData, fSearchOp, lpSearchFilter, dwAdditionalFlags, funcName);
diff --git a/Engine/Source/Programs/UnrealBuildAccelerator/Detours/Private/Windows/UbaDetoursFunctionsNtDll.inl b/Engine/Source/Programs/UnrealBuildAccelerator/Detours/Private/Windows/UbaDetoursFunctionsNtDll.inl
--- a/Engine/Source/Programs/UnrealBuildAccelerator/Detours/Private/Windows/UbaDetoursFunctionsNtDll.inl
+++ b/Engine/Source/Programs/UnrealBuildAccelerator/Detours/Private/Windows/UbaDetoursFunctionsNtDll.inl
@@ -43,6 +43,8 @@ StringBuffer<32> ToString(NTSTATUS s)
 		return res.Append(L"STATUS_OBJECT_NAME_NOT_FOUND");
 	if (s == STATUS_OBJECT_PATH_NOT_FOUND)
 		return res.Append(L"STATUS_OBJECT_PATH_NOT_FOUND");
+	if (s == STATUS_INVALID_HANDLE)
+		return res.Append(L"STATUS_INVALID_HANDLE");
 	return res.Appendf(L"Error (0x%x)", u32(s));
 }
 
@@ -394,6 +396,9 @@ NTSTATUS Detoured_NtSetInformationFile(HANDLE FileHandle, PIO_STATUS_BLOCK IoSta
 {
 	DETOURED_CALL(NtSetInformationFile);
 
+
+	u8 tempBuffer[sizeof(FILE_RENAME_INFORMATION) + 512];
+
 	HANDLE trueHandle = FileHandle;
 	if (isDetouredHandle(FileHandle))
 	{
@@ -414,6 +419,7 @@ NTSTATUS Detoured_NtSetInformationFile(HANDLE FileHandle, PIO_STATUS_BLOCK IoSta
 				newNamePtr += 4;
 			StringBuffer<> newName;
 			FixPath(newName, newNamePtr);
+			DevirtualizePath(newName);
 			FileObject& fo = *h.fileObject;
 			fo.newName = newName.data;
 
@@ -474,6 +480,16 @@ NTSTATUS Detoured_NtSetInformationFile(HANDLE FileHandle, PIO_STATUS_BLOCK IoSta
 
 			if (g_runningRemote) // This needs a proper solution as the comments above.
 				return STATUS_SUCCESS;
+
+			// In case we are using vfs we need to replace the information before calling the true NtSetInformationFile
+			auto& info2 = *(FILE_RENAME_INFORMATION*)tempBuffer;
+			memcpy(&info2, &info, sizeof(FILE_RENAME_INFORMATION));
+			memcpy(info2.FileName, TC("\\??\\"), 8);
+			memcpy(info2.FileName + 4, newName.data, newName.count * 2);
+			info2.FileNameLength = (newName.count+4) * 2;
+			
+			FileInformation = &info2;
+			Length = sizeof(FILE_RENAME_INFORMATION) + info2.FileNameLength + 2;
 		}
 
 		trueHandle = h.trueHandle;
diff --git a/Engine/Source/Programs/UnrealBuildAccelerator/Detours/Private/Windows/UbaDetoursFunctionsWin.cpp b/Engine/Source/Programs/UnrealBuildAccelerator/Detours/Private/Windows/UbaDetoursFunctionsWin.cpp
--- a/Engine/Source/Programs/UnrealBuildAccelerator/Detours/Private/Windows/UbaDetoursFunctionsWin.cpp
+++ b/Engine/Source/Programs/UnrealBuildAccelerator/Detours/Private/Windows/UbaDetoursFunctionsWin.cpp
@@ -984,12 +984,9 @@ void Init(const DetoursPayload& payload, u64 startTime)
 
 	TrackInput(applicationBuffer.data);
 
-	if (g_rules->IsFullVirtualization())
-	{
-		VirtualizePath(applicationBuffer);
-		VirtualizePath(workingDirBuffer);
-		VirtualizePath(g_exeDir);
-	}
+	VirtualizePath(applicationBuffer);
+	VirtualizePath(workingDirBuffer);
+	VirtualizePath(g_exeDir);
 
 	Shared_SetCurrentDirectory(workingDirBuffer.data);
 
