[UBA]
* Fixed so vfs devirtualization path can handle forward slashes. This is needed for ispc.txt files which have all kinds of slashes in the paths

diff --git a/Engine/Source/Programs/UnrealBuildAccelerator/Common/Private/UbaCacheClient.cpp b/Engine/Source/Programs/UnrealBuildAccelerator/Common/Private/UbaCacheClient.cpp
--- a/Engine/Source/Programs/UnrealBuildAccelerator/Common/Private/UbaCacheClient.cpp
+++ b/Engine/Source/Programs/UnrealBuildAccelerator/Common/Private/UbaCacheClient.cpp
@@ -1239,6 +1239,8 @@ namespace uba
 			return true;
 		if (path.EndsWith(TC(".bat"))) // Contains absolute paths in some cases
 			return true;
+		if (path.EndsWith(TC(".txt"))) // Contains absolute paths (ispc dependency file)
+			return true;
 		return false;
 	}
 
diff --git a/Engine/Source/Programs/UnrealBuildAccelerator/Common/Private/UbaRootPaths.cpp b/Engine/Source/Programs/UnrealBuildAccelerator/Common/Private/UbaRootPaths.cpp
--- a/Engine/Source/Programs/UnrealBuildAccelerator/Common/Private/UbaRootPaths.cpp
+++ b/Engine/Source/Programs/UnrealBuildAccelerator/Common/Private/UbaRootPaths.cpp
@@ -14,6 +14,22 @@ namespace uba
 		return InternalRegisterRoot(logger, m_roots, rootPath, includeInKey, id);
 	}
 
+	bool RootPaths::RegisterRootForwardSlashes(Logger& logger, const tchar* rootPath)
+	{
+		m_allowForwardSlashes = true;
+		tchar forwardSlash[1024];
+		u32 forwardSlashPos = 0;
+		for (const tchar* it=rootPath; *it; ++it)
+		{
+			if (*it == '\\')
+				forwardSlash[forwardSlashPos++] = '/';
+			else
+				forwardSlash[forwardSlashPos++] = *it;
+		}
+		forwardSlash[forwardSlashPos] = 0;
+		return InternalRegisterRoot2(logger, m_roots, forwardSlash, false, 0);
+	}
+
 	bool RootPaths::RegisterSystemRoots(Logger& logger, u8 startId)
 	{
 #if PLATFORM_WINDOWS
diff --git a/Engine/Source/Programs/UnrealBuildAccelerator/Common/Private/UbaSession.cpp b/Engine/Source/Programs/UnrealBuildAccelerator/Common/Private/UbaSession.cpp
--- a/Engine/Source/Programs/UnrealBuildAccelerator/Common/Private/UbaSession.cpp
+++ b/Engine/Source/Programs/UnrealBuildAccelerator/Common/Private/UbaSession.cpp
@@ -738,7 +738,7 @@ namespace uba
 		if (!shouldWriteToDisk)
 		{
 			u32 res = WriteDirectoryEntries(dirKey, dirName.data);
-			UBA_ASSERT(res); (void)res;
+			UBA_ASSERTF(res, TC("Failed to write directory entries for %s"), dirName.data); (void)res;
 		}
 
 		SCOPED_READ_LOCK(dirTable.m_lookupLock, lookupCs);
@@ -768,7 +768,7 @@ namespace uba
 			UBA_ASSERT(res); (void)res;
 			dirLock.Enter();
 		}
-		UBA_ASSERT(dir.parseOffset == 1);
+		UBA_ASSERTF(dir.parseOffset == 1, TC("Registering create file for write %s with unexpect dir.parseOffset %u "), fileName.data, dir.parseOffset);
 
 		if (fileNameKey == StringKeyZero)
 		{
@@ -3042,6 +3042,9 @@ namespace uba
 			StringBuffer<> temp;
 			reader.ReadString(temp);
 			entry.roots.RegisterRoot(m_logger, temp.data);
+			#if PLATFORM_WINDOWS
+			entry.roots.RegisterRootForwardSlashes(m_logger, temp.data); // ispc.txt files have forward slashes, sigh.
+			#endif
 			reader.ReadString(temp.Clear());
 			entry.locals.emplace_back(temp.data, temp.count);
 		}
@@ -3157,7 +3160,8 @@ namespace uba
 					memcpy(destData.Allocate(strLen, 1, TC("")), str, strLen);
 					return;
 				}
-				auto& path = localsAnsi[(*str - RootPaths::RootStartByte)/3];
+				constexpr u32 PathsPerRoot = 4; // 3 + forward slashes
+				auto& path = localsAnsi[(*str - RootPaths::RootStartByte)/PathsPerRoot];
 				memcpy(destData.Allocate(path.size(), 1, TC("")), path.c_str(), path.size());
 			};
 
diff --git a/Engine/Source/Programs/UnrealBuildAccelerator/Common/Public/UbaRootPaths.h b/Engine/Source/Programs/UnrealBuildAccelerator/Common/Public/UbaRootPaths.h
--- a/Engine/Source/Programs/UnrealBuildAccelerator/Common/Public/UbaRootPaths.h
+++ b/Engine/Source/Programs/UnrealBuildAccelerator/Common/Public/UbaRootPaths.h
@@ -14,6 +14,7 @@ namespace uba
 	public:
 
 		bool RegisterRoot(Logger& logger, const tchar* rootPath, bool includeInKey = true, u8 id = 0);
+		bool RegisterRootForwardSlashes(Logger& logger, const tchar* rootPath);
 		bool RegisterSystemRoots(Logger& logger, u8 startId = 0);
 		bool RegisterIgnoredRoot(Logger& logger, const tchar* rootPath);
 
@@ -48,6 +49,7 @@ namespace uba
 
 		Roots m_roots;
 		Roots m_ignoredRoots;
+		bool m_allowForwardSlashes = false;
 	};
 
 
@@ -64,7 +66,7 @@ namespace uba
 		{
 			auto absPathChars = searchPos;
 			CharType lastChar = 0;
-			while (absPathChars < strEnd && !(lastChar == ':' && *absPathChars == PathSeparator))
+			while (absPathChars < strEnd && !(lastChar == ':' && (*absPathChars == PathSeparator || (m_allowForwardSlashes && *absPathChars == '/'))))
 			{
 				lastChar = *absPathChars;
 				++absPathChars;
