[UBA]
* Fixed so linux linking works with vfs

diff --git a/Engine/Source/Programs/UnrealBuildAccelerator/Detours/Private/Windows/UbaDetoursFunctionsKernelBase.inl b/Engine/Source/Programs/UnrealBuildAccelerator/Detours/Private/Windows/UbaDetoursFunctionsKernelBase.inl
--- a/Engine/Source/Programs/UnrealBuildAccelerator/Detours/Private/Windows/UbaDetoursFunctionsKernelBase.inl
+++ b/Engine/Source/Programs/UnrealBuildAccelerator/Detours/Private/Windows/UbaDetoursFunctionsKernelBase.inl
@@ -2757,6 +2757,9 @@ BOOL Detoured_CreateProcessW(LPCWSTR lpApplicationName, LPWSTR lpCommandLine, LP
 	}
 	bool startSuspended = (dwCreationFlags & CREATE_SUSPENDED) != 0;
 
+	StringBuffer<> workingDir(lpCurrentDirectory ? lpCurrentDirectory : g_virtualWorkingDir.data);
+	DevirtualizePath(workingDir);
+
 	TString commandLine;
 	TString currentDir;
 	u32 processId = 0;
@@ -2768,7 +2771,7 @@ BOOL Detoured_CreateProcessW(LPCWSTR lpApplicationName, LPWSTR lpCommandLine, LP
 		writer.WriteByte(MessageType_CreateProcess);
 		writer.WriteString(application);
 		WriteLongString(writer, originalCmd);
-		writer.WriteString(lpCurrentDirectory ? lpCurrentDirectory : g_virtualWorkingDir.data);
+		writer.WriteString(workingDir);
 		writer.WriteBool(startSuspended);
 		writer.WriteBool(isChild);
 		writer.Flush();
