[UBA]
* Added exported function CacheClient_FetchFromCache3 which takes in rootshandle
* Added code to devirtualize response file path when calculate cache key

diff --git a/Engine/Source/Programs/UnrealBuildAccelerator/Common/Private/UbaCacheClient.cpp b/Engine/Source/Programs/UnrealBuildAccelerator/Common/Private/UbaCacheClient.cpp
--- a/Engine/Source/Programs/UnrealBuildAccelerator/Common/Private/UbaCacheClient.cpp
+++ b/Engine/Source/Programs/UnrealBuildAccelerator/Common/Private/UbaCacheClient.cpp
@@ -129,7 +129,7 @@ namespace uba
 	bool CacheClient::WriteToCache(u32 bucketId, const ProcessStartInfo& info, const u8* inputs, u64 inputsSize, const u8* outputs, u64 outputsSize, const u8* logLines, u64 logLinesSize, u32 processId)
 	{
 		RootPaths rootPaths;
-		m_session.PopulateCacheRoots(rootPaths, info.rootsHandle);
+		m_session.PopulateLocalToIndexRoots(rootPaths, info.rootsHandle);
 		return WriteToCache(rootPaths, bucketId, info, inputs, inputsSize, outputs, outputsSize, logLines, logLinesSize, processId);
 	}
 
@@ -390,6 +390,12 @@ namespace uba
 		return u64(bucketId) | ((u64(!CaseInsensitiveFs) + (RootPathsVersion << 1u) + (u64(!m_useRoots) << 2u)) << 32u);
 	}
 
+	bool CacheClient::FetchFromCache(CacheResult& outResult, RootsHandle rootsHandle, u32 bucketId, const ProcessStartInfo& info)
+	{
+		RootPaths rootPaths;
+		m_session.PopulateLocalToIndexRoots(rootPaths, rootsHandle);
+		return FetchFromCache(outResult, rootPaths, bucketId, info);
+	}
 
 	bool CacheClient::FetchFromCache(CacheResult& outResult, const RootPaths& rootPaths, u32 bucketId, const ProcessStartInfo& info)
 	{
@@ -1184,7 +1190,7 @@ namespace uba
 
 		// Add arguments list to key
 		auto hashString = [&](const tchar* str, u64 strLen, u32 rootPos) { hasher.Update(str, strLen*sizeof(tchar)); };
-		if (!rootPaths.NormalizeString(m_logger, info.arguments, TStrlen(info.arguments), hashString, false, info.GetDescription(), TC("calculating command line hash")))
+		if (!rootPaths.NormalizeString(m_logger, info.arguments, TStrlen(info.arguments), hashString, false, info.GetDescription(), TC(" calculating command line hash")))
 		{
 			if (report)
 				m_logger.Info(TC("   Failed to normalize commandline %s"), info.arguments);
@@ -1208,6 +1214,12 @@ namespace uba
 					rsp.Append(rspStart, rspEnd - rspStart);
 					StringBuffer<> fullPath;
 					FixPath(rsp.data, workingDir.data, workingDir.count, fullPath);
+					if (!DevirtualizePath(rsp, info.rootsHandle))
+					{
+						if (report)
+							m_logger.Warning(TC("Failed to normalize rsp file path %s"), rsp.data);
+						return CasKeyZero;
+					}
 					CasKey rspCasKey = rootPaths.NormalizeAndHashFile(m_logger, rsp.data, true);
 					if (rspCasKey == CasKeyZero)
 					{
@@ -1225,6 +1237,11 @@ namespace uba
 		return ToCasKey(hasher, false);
 	}
 
+	bool CacheClient::DevirtualizePath(StringBufferBase& inOut, RootsHandle rootsHandle)
+	{
+		return m_session.DevirtualizePath(inOut, rootsHandle);
+	}
+
 	bool CacheClient::ShouldNormalize(const StringBufferBase& path)
 	{
 		if (!m_useRoots)
diff --git a/Engine/Source/Programs/UnrealBuildAccelerator/Common/Private/UbaSession.cpp b/Engine/Source/Programs/UnrealBuildAccelerator/Common/Private/UbaSession.cpp
--- a/Engine/Source/Programs/UnrealBuildAccelerator/Common/Private/UbaSession.cpp
+++ b/Engine/Source/Programs/UnrealBuildAccelerator/Common/Private/UbaSession.cpp
@@ -48,6 +48,8 @@ typedef struct _PROCESSOR_POWER_INFORMATION {
 
 namespace uba
 {
+	constexpr u32 PathsPerRoot = 4u; // 3 + forward slashes
+
 	bool g_dummy;
 
 	ProcessStartInfo::ProcessStartInfo() = default;
@@ -1563,7 +1565,30 @@ namespace uba
 		return ++m_processIdCounter;
 	}
 
-	bool Session::PopulateCacheRoots(RootPaths& out, RootsHandle rootsHandle)
+	bool Session::DevirtualizePath(StringBufferBase& inOut, RootsHandle rootsHandle)
+	{
+		if (!rootsHandle)
+			return true;
+
+		SCOPED_READ_LOCK(m_rootsLookupLock, lock);
+		auto findIt = m_rootsLookup.find(rootsHandle);
+		if (findIt == m_rootsLookup.end())
+			return m_logger.Error(TC("Can't find entry from roots handle %llu"), rootsHandle);
+		RootsEntry& rootsEntry = findIt->second;
+		lock.Leave();
+
+		auto root = rootsEntry.roots.FindRoot(inOut);
+		if (!root)
+			return m_logger.Error(TC("Can't find root for path %s"), inOut.data);
+
+		auto& path = rootsEntry.locals[root->index/PathsPerRoot];
+		StringBuffer<> temp;
+		temp.Append(inOut.data + root->path.size());
+		inOut.Clear().Append(path).Append(temp);
+		return true;
+	}
+
+	bool Session::PopulateLocalToIndexRoots(RootPaths& out, RootsHandle rootsHandle)
 	{
 		if (!rootsHandle)
 			return true;
@@ -1571,8 +1596,9 @@ namespace uba
 		auto findIt = m_rootsLookup.find(rootsHandle);
 		if (findIt == m_rootsLookup.end())
 			return m_logger.Error(TC("Roots handle %llu does not exist"), rootsHandle);
-
 		RootsEntry& entry = findIt->second;
+		lock.Leave();
+
 		BinaryReader reader((const u8*)entry.memory.data(), 0, entry.memory.size());
 		while (reader.GetLeft())
 		{
@@ -3160,7 +3186,6 @@ namespace uba
 					memcpy(destData.Allocate(strLen, 1, TC("")), str, strLen);
 					return;
 				}
-				constexpr u32 PathsPerRoot = 4; // 3 + forward slashes
 				auto& path = localsAnsi[(*str - RootPaths::RootStartByte)/PathsPerRoot];
 				memcpy(destData.Allocate(path.size(), 1, TC("")), path.c_str(), path.size());
 			};
diff --git a/Engine/Source/Programs/UnrealBuildAccelerator/Common/Public/UbaCacheClient.h b/Engine/Source/Programs/UnrealBuildAccelerator/Common/Public/UbaCacheClient.h
--- a/Engine/Source/Programs/UnrealBuildAccelerator/Common/Public/UbaCacheClient.h
+++ b/Engine/Source/Programs/UnrealBuildAccelerator/Common/Public/UbaCacheClient.h
@@ -19,6 +19,7 @@ namespace uba
 	struct CasKey;
 	struct ProcessStartInfo;
 	struct StringView;
+	using RootsHandle = u64;
 
 	struct CacheClientCreateInfo
 	{
@@ -53,6 +54,7 @@ namespace uba
 
 		bool WriteToCache(u32 bucketId, const ProcessStartInfo& info, const u8* inputs, u64 inputsSize, const u8* outputs, u64 outputsSize, const u8* logLines, u64 logLinesSize, u32 processId = 0);
 		bool WriteToCache(const RootPaths& rootPaths, u32 bucketId, const ProcessStartInfo& info, const u8* inputs, u64 inputsSize, const u8* outputs, u64 outputsSize, const u8* logLines, u64 logLinesSize, u32 processId = 0);
+		bool FetchFromCache(CacheResult& outResult, RootsHandle rootsHandle, u32 bucketId, const ProcessStartInfo& info);
 		bool FetchFromCache(CacheResult& outResult, const RootPaths& rootPaths, u32 bucketId, const ProcessStartInfo& info);
 		bool RequestServerShutdown(const tchar* reason);
 
@@ -75,6 +77,7 @@ namespace uba
 		bool PopulateLogLines(Vector<ProcessLogLine>& outLogLines, const u8* mem, u64 memLen);
 
 		CasKey GetCmdKey(const RootPaths& rootPaths, const ProcessStartInfo& info, bool report, u32 bucketId);
+		bool DevirtualizePath(StringBufferBase& inOut, RootsHandle rootsHandle);
 		bool ShouldNormalize(const StringBufferBase& path);
 
 		bool GetLocalPathAndCasKey(Bucket& bucket, const RootPaths& rootPaths, StringBufferBase& outPath, CasKey& outKey, CompactCasKeyTable& casKeyTable, CompactPathTable& pathTable, u32 offset);
diff --git a/Engine/Source/Programs/UnrealBuildAccelerator/Common/Public/UbaRootPaths.h b/Engine/Source/Programs/UnrealBuildAccelerator/Common/Public/UbaRootPaths.h
--- a/Engine/Source/Programs/UnrealBuildAccelerator/Common/Public/UbaRootPaths.h
+++ b/Engine/Source/Programs/UnrealBuildAccelerator/Common/Public/UbaRootPaths.h
@@ -105,7 +105,7 @@ namespace uba
 					path.Resize(lastQuote - path.data);
 				if (auto lineEnd = path.Last('\r'))
 					path.Resize(lineEnd - path.data);
-				logger.Info(TC("PATH WITHOUT ROOT: %s (inside %s at offset %u%s)"), path.data, hint, hint2, destPos);
+				logger.Info(TC("PATH WITHOUT ROOT: %s (inside %s at offset %u%s)"), path.data, hint, destPos, hint2);
 				return false;
 			}
 
diff --git a/Engine/Source/Programs/UnrealBuildAccelerator/Common/Public/UbaSession.h b/Engine/Source/Programs/UnrealBuildAccelerator/Common/Public/UbaSession.h
--- a/Engine/Source/Programs/UnrealBuildAccelerator/Common/Public/UbaSession.h
+++ b/Engine/Source/Programs/UnrealBuildAccelerator/Common/Public/UbaSession.h
@@ -84,7 +84,8 @@ namespace uba
 
 		u32 CreateProcessId();
 
-		bool PopulateCacheRoots(RootPaths& out, RootsHandle rootsHandle);
+		bool DevirtualizePath(StringBufferBase& inOut, RootsHandle rootsHandle);
+		bool PopulateLocalToIndexRoots(RootPaths& out, RootsHandle rootsHandle);
 
 	protected:
 		Session(const SessionCreateInfo& info, const tchar* logPrefix, bool runningRemote, WorkManager* workManager = nullptr);
diff --git a/Engine/Source/Programs/UnrealBuildAccelerator/Detours/Private/Windows/UbaDetoursFunctionsKernelBase.inl b/Engine/Source/Programs/UnrealBuildAccelerator/Detours/Private/Windows/UbaDetoursFunctionsKernelBase.inl
--- a/Engine/Source/Programs/UnrealBuildAccelerator/Detours/Private/Windows/UbaDetoursFunctionsKernelBase.inl
+++ b/Engine/Source/Programs/UnrealBuildAccelerator/Detours/Private/Windows/UbaDetoursFunctionsKernelBase.inl
@@ -999,7 +999,7 @@ BOOL Shared_GetFileAttributesExW(LPCWSTR lpFileName, GET_FILEEX_INFO_LEVELS fInf
 
 	if (!CanDetour(fixedFileName))
 	{
-		DEBUG_LOG_TRUE(L"GetFileAttributesExW", L"(%ls)", lpFileName);
+		DEBUG_LOG_TRUE(L"GetFileAttributesExW", L"(%ls)", originalName);
 		TimerScope ts(g_kernelStats.getFileInfo);
 		return True_GetFileAttributesExW(lpFileName, fInfoLevelId, lpFileInformation);
 	}
@@ -1009,7 +1009,7 @@ BOOL Shared_GetFileAttributesExW(LPCWSTR lpFileName, GET_FILEEX_INFO_LEVELS fInf
 
 	if (!attr.useCache)
 	{
-		DEBUG_LOG_TRUE(L"GetFileAttributesExW", L"(%ls)", lpFileName);
+		DEBUG_LOG_TRUE(L"GetFileAttributesExW", L"(%ls)", originalName);
 		TimerScope ts(g_kernelStats.getFileInfo);
 		return True_GetFileAttributesExW(realName, fInfoLevelId, lpFileInformation);
 	}
@@ -1018,7 +1018,7 @@ BOOL Shared_GetFileAttributesExW(LPCWSTR lpFileName, GET_FILEEX_INFO_LEVELS fInf
 
 	memcpy(lpFileInformation, &attr.data, sizeof(WIN32_FILE_ATTRIBUTE_DATA));
 
-	DEBUG_LOG_DETOURED(L"GetFileAttributesExW", L"(%ls) -> %u", lpFileName, attr.exists);
+	DEBUG_LOG_DETOURED(L"GetFileAttributesExW", L"(%ls) -> %u", originalName, attr.exists);
 	return attr.exists;
 }
 
diff --git a/Engine/Source/Programs/UnrealBuildAccelerator/Host/Private/UbaExports.cpp b/Engine/Source/Programs/UnrealBuildAccelerator/Host/Private/UbaExports.cpp
--- a/Engine/Source/Programs/UnrealBuildAccelerator/Host/Private/UbaExports.cpp
+++ b/Engine/Source/Programs/UnrealBuildAccelerator/Host/Private/UbaExports.cpp
@@ -682,6 +682,16 @@ extern "C"
 		return nullptr;
 	}
 
+	uba::CacheResult* CacheClient_FetchFromCache3(uba::CacheClient* cacheClient, uba::u64 rootsHandle, uba::u32 bucket, const uba::ProcessStartInfo& info)
+	{
+		using namespace uba;
+		auto cacheResult = new CacheResult();
+		if (cacheClient->FetchFromCache(*cacheResult, rootsHandle, bucket, info))
+			return cacheResult;
+		delete cacheResult;
+		return nullptr;
+	}
+
 	void CacheClient_RequestServerShutdown(uba::CacheClient* cacheClient, const uba::tchar* reason)
 	{
 		cacheClient->RequestServerShutdown(reason);
diff --git a/Engine/Source/Programs/UnrealBuildAccelerator/Host/Public/UbaExports.h b/Engine/Source/Programs/UnrealBuildAccelerator/Host/Public/UbaExports.h
--- a/Engine/Source/Programs/UnrealBuildAccelerator/Host/Public/UbaExports.h
+++ b/Engine/Source/Programs/UnrealBuildAccelerator/Host/Public/UbaExports.h
@@ -138,6 +138,7 @@ extern "C"
 	UBA_API bool CacheClient_WriteToCache2(uba::CacheClient* cacheClient, uba::u32 bucket, const uba::ProcessHandle* process, const uba::u8* inputs, uba::u32 inputsSize, const uba::u8* outputs, uba::u32 outputsSize, const uba::u8* logLines, uba::u32 logLinesSize);
 	UBA_API uba::u32 CacheClient_FetchFromCache(uba::CacheClient* cacheClient, uba::RootPaths* rootPaths, uba::u32 bucket, const uba::ProcessStartInfo& info);
 	UBA_API uba::CacheResult* CacheClient_FetchFromCache2(uba::CacheClient* cacheClient, uba::RootPaths* rootPaths, uba::u32 bucket, const uba::ProcessStartInfo& info);
+	UBA_API uba::CacheResult* CacheClient_FetchFromCache3(uba::CacheClient* cacheClient, uba::u64 rootsHandle, uba::u32 bucket, const uba::ProcessStartInfo& info);
 	UBA_API void CacheClient_RequestServerShutdown(uba::CacheClient* cacheClient, const uba::tchar* reason);
 	UBA_API void CacheClient_Destroy(uba::CacheClient* cacheClient);
 	UBA_API const uba::tchar* CacheResult_GetLogLine(uba::CacheResult* result, uba::u32 index);
