[UBA]
* Changed so all combination of slashes are supported for rootpaths.. forward slash, backward slash, double backward slash and escaped spaces (backslash in front of spaces in paths)
* Fixed so uba adds system paths for cache system. There are a bunch of modules that has more then 1 extra path and it is easier to just add these for now

diff --git a/Engine/Source/Programs/UnrealBuildAccelerator/Common/Private/UbaRootPaths.cpp b/Engine/Source/Programs/UnrealBuildAccelerator/Common/Private/UbaRootPaths.cpp
--- a/Engine/Source/Programs/UnrealBuildAccelerator/Common/Private/UbaRootPaths.cpp
+++ b/Engine/Source/Programs/UnrealBuildAccelerator/Common/Private/UbaRootPaths.cpp
@@ -14,22 +14,6 @@ namespace uba
 		return InternalRegisterRoot(logger, m_roots, rootPath, includeInKey, id);
 	}
 
-	bool RootPaths::RegisterRootForwardSlashes(Logger& logger, const tchar* rootPath)
-	{
-		m_allowForwardSlashes = true;
-		tchar forwardSlash[1024];
-		u32 forwardSlashPos = 0;
-		for (const tchar* it=rootPath; *it; ++it)
-		{
-			if (*it == '\\')
-				forwardSlash[forwardSlashPos++] = '/';
-			else
-				forwardSlash[forwardSlashPos++] = *it;
-		}
-		forwardSlash[forwardSlashPos] = 0;
-		return InternalRegisterRoot2(logger, m_roots, forwardSlash, false, 0);
-	}
-
 	bool RootPaths::RegisterSystemRoots(Logger& logger, u8 startId)
 	{
 #if PLATFORM_WINDOWS
@@ -118,23 +102,47 @@ namespace uba
 	{
 		// Register rootPath both with single path separators and double path separators on windows because text files store them with double path separators
 		#if PLATFORM_WINDOWS
-		StringBuffer<> doubleSlash;
-		StringBuffer<> spaceEscaped;
+		StringBuffer<512> forwardSlash;
+		StringBuffer<512> backwardSlash;
+		StringBuffer<512> doubleBackwardSlash;
+		StringBuffer<512> spaceEscapedBackwardSlash;
 		bool hasSpace = false;
 		for (const tchar* it=rootPath; *it; ++it)
 		{
-			if (*it == ' ')
+			tchar c = *it;
+
+			if (c == '/')
+			{
+				forwardSlash.Append('/');
+				backwardSlash.Append('\\');
+				doubleBackwardSlash.Append(TC("\\\\"));
+				spaceEscapedBackwardSlash.Append('\\');
+			}
+			else if (c == '\\')
 			{
-				hasSpace = true;
-				spaceEscaped.Append('\\');
+				forwardSlash.Append('/');
+				backwardSlash.Append('\\');
+				doubleBackwardSlash.Append("\\\\");
+				spaceEscapedBackwardSlash.Append('\\');
+			}
+			else
+			{
+				if (c == ' ')
+				{
+					hasSpace = true;
+					spaceEscapedBackwardSlash.Append('\\');
+				}
+				forwardSlash.Append(c);
+				backwardSlash.Append(c);
+				doubleBackwardSlash.Append(c);
+				spaceEscapedBackwardSlash.Append(c);
 			}
-			spaceEscaped.Append(*it);
-			doubleSlash.Append(*it);
-			if (*it == PathSeparator)
-				doubleSlash.Append(PathSeparator);
 		}
 
-		const tchar* rootPaths[] = { rootPath, doubleSlash.data, hasSpace ? spaceEscaped.data : TC("") };
+		if (!hasSpace)
+			spaceEscapedBackwardSlash.Clear();
+
+		const tchar* rootPaths[] = { forwardSlash.data, backwardSlash.data, doubleBackwardSlash.data, spaceEscapedBackwardSlash.data };
 		#else
 		const tchar* rootPaths[] = { rootPath };
 		#endif
diff --git a/Engine/Source/Programs/UnrealBuildAccelerator/Common/Private/UbaSession.cpp b/Engine/Source/Programs/UnrealBuildAccelerator/Common/Private/UbaSession.cpp
--- a/Engine/Source/Programs/UnrealBuildAccelerator/Common/Private/UbaSession.cpp
+++ b/Engine/Source/Programs/UnrealBuildAccelerator/Common/Private/UbaSession.cpp
@@ -1569,6 +1569,14 @@ namespace uba
 	{
 		if (!rootsHandle)
 			return true;
+		
+#if PLATFORM_WINDOWS
+		if (inOut[1] != ':')
+			return true;
+#else
+		if (inOut[0] != '/')
+			return true;
+#endif
 
 		SCOPED_READ_LOCK(m_rootsLookupLock, lock);
 		auto findIt = m_rootsLookup.find(rootsHandle);
@@ -1606,9 +1614,18 @@ namespace uba
 			reader.SkipString();
 			StringBuffer<> rootPath;
 			reader.ReadString(rootPath.Clear());
-			out.RegisterRoot(m_logger, rootPath.data, true, index);
+			if (!out.RegisterRoot(m_logger, rootPath.data, true, index))
+				return false;
 		}
-		return true;
+
+		// TODO: Provide or calculate these
+		#if PLATFORM_WINDOWS
+		out.RegisterIgnoredRoot(m_logger, TC("z:/UEVFS"));
+		#else
+		out.RegisterIgnoredRoot(m_logger, TC("/UEVFS"));
+		#endif
+
+		return true;//out.RegisterSystemRoots(m_logger, 0);
 	}
 
 	void Session::ProcessAdded(Process& process, u32 sessionId)
@@ -3068,9 +3085,6 @@ namespace uba
 			StringBuffer<> temp;
 			reader.ReadString(temp);
 			entry.roots.RegisterRoot(m_logger, temp.data);
-			#if PLATFORM_WINDOWS
-			entry.roots.RegisterRootForwardSlashes(m_logger, temp.data); // ispc.txt files have forward slashes, sigh.
-			#endif
 			reader.ReadString(temp.Clear());
 			entry.locals.emplace_back(temp.data, temp.count);
 		}
diff --git a/Engine/Source/Programs/UnrealBuildAccelerator/Common/Public/UbaRootPaths.h b/Engine/Source/Programs/UnrealBuildAccelerator/Common/Public/UbaRootPaths.h
--- a/Engine/Source/Programs/UnrealBuildAccelerator/Common/Public/UbaRootPaths.h
+++ b/Engine/Source/Programs/UnrealBuildAccelerator/Common/Public/UbaRootPaths.h
@@ -14,7 +14,6 @@ namespace uba
 	public:
 
 		bool RegisterRoot(Logger& logger, const tchar* rootPath, bool includeInKey = true, u8 id = 0);
-		bool RegisterRootForwardSlashes(Logger& logger, const tchar* rootPath);
 		bool RegisterSystemRoots(Logger& logger, u8 startId = 0);
 		bool RegisterIgnoredRoot(Logger& logger, const tchar* rootPath);
 
@@ -49,7 +48,6 @@ namespace uba
 
 		Roots m_roots;
 		Roots m_ignoredRoots;
-		bool m_allowForwardSlashes = false;
 	};
 
 
@@ -66,7 +64,7 @@ namespace uba
 		{
 			auto absPathChars = searchPos;
 			CharType lastChar = 0;
-			while (absPathChars < strEnd && !(lastChar == ':' && (*absPathChars == PathSeparator || (m_allowForwardSlashes && *absPathChars == '/'))))
+			while (absPathChars < strEnd && !(lastChar == ':' && ((IsWindows && *absPathChars == '\\') || *absPathChars == '/')))
 			{
 				lastChar = *absPathChars;
 				++absPathChars;
