[UBA]
* Don't throw error on absolute non-virtual paths if no vfs is populated.

diff --git a/Engine/Source/Programs/UnrealBuildAccelerator/Common/Private/UbaRootPaths.cpp b/Engine/Source/Programs/UnrealBuildAccelerator/Common/Private/UbaRootPaths.cpp
--- a/Engine/Source/Programs/UnrealBuildAccelerator/Common/Private/UbaRootPaths.cpp
+++ b/Engine/Source/Programs/UnrealBuildAccelerator/Common/Private/UbaRootPaths.cpp
@@ -61,6 +61,11 @@ namespace uba
 		return InternalRegisterRoot(logger, m_ignoredRoots, rootPath, false, 0);
 	}
 
+	bool RootPaths::IsEmpty() const
+	{
+		return m_roots.empty();
+	}
+
 	const RootPaths::Root* RootPaths::FindRoot(const StringView& path) const
 	{
 		return InternalFindRoot(m_roots, path);
diff --git a/Engine/Source/Programs/UnrealBuildAccelerator/Common/Private/UbaSession.cpp b/Engine/Source/Programs/UnrealBuildAccelerator/Common/Private/UbaSession.cpp
--- a/Engine/Source/Programs/UnrealBuildAccelerator/Common/Private/UbaSession.cpp
+++ b/Engine/Source/Programs/UnrealBuildAccelerator/Common/Private/UbaSession.cpp
@@ -1585,6 +1585,9 @@ namespace uba
 		RootsEntry& rootsEntry = findIt->second;
 		lock.Leave();
 
+		if (rootsEntry.roots.IsEmpty())
+			return true;
+
 		auto root = rootsEntry.roots.FindRoot(inOut);
 		if (!root)
 			return m_logger.Error(TC("Can't find root for path %s"), inOut.data);
@@ -3084,6 +3087,8 @@ namespace uba
 			u8 index = reader.ReadByte();(void)index; // RootPaths index;
 			StringBuffer<> temp;
 			reader.ReadString(temp);
+			if (!temp.count) // If vfs is not set it means that vfs should not be used (the roots entry memory might be used for cacheclient though)
+				break;
 			entry.roots.RegisterRoot(m_logger, temp.data);
 			reader.ReadString(temp.Clear());
 			entry.locals.emplace_back(temp.data, temp.count);
diff --git a/Engine/Source/Programs/UnrealBuildAccelerator/Common/Public/UbaRootPaths.h b/Engine/Source/Programs/UnrealBuildAccelerator/Common/Public/UbaRootPaths.h
--- a/Engine/Source/Programs/UnrealBuildAccelerator/Common/Public/UbaRootPaths.h
+++ b/Engine/Source/Programs/UnrealBuildAccelerator/Common/Public/UbaRootPaths.h
@@ -17,6 +17,8 @@ namespace uba
 		bool RegisterSystemRoots(Logger& logger, u8 startId = 0);
 		bool RegisterIgnoredRoot(Logger& logger, const tchar* rootPath);
 
+		bool IsEmpty() const;
+
 		struct Root
 		{
 			TString path;
