[UBA]
* Fixed unit test error

diff --git a/Engine/Source/Programs/UnrealBuildAccelerator/Test/Private/UbaTestBasics.h b/Engine/Source/Programs/UnrealBuildAccelerator/Test/Private/UbaTestBasics.h
--- a/Engine/Source/Programs/UnrealBuildAccelerator/Test/Private/UbaTestBasics.h
+++ b/Engine/Source/Programs/UnrealBuildAccelerator/Test/Private/UbaTestBasics.h
@@ -402,7 +402,7 @@ namespace uba
 				{
 					if (strLen != 1)
 						success = false;
-					if (str[0] != RootPaths::RootStartByte + (IsWindows ? 3 : 1)) // 3 is because each root has three entries on windows
+					if (str[0] != RootPaths::RootStartByte + (IsWindows ? 5 : 1)) // 5 is because each root has four entries on windows and second entry is backslash
 						success = false;
 					rootPos = str[0];
 				}
