[UBA]
* Fixed so root ids are handled properly

diff --git a/Engine/Source/Programs/UnrealBuildAccelerator/Common/Private/UbaRootPaths.cpp b/Engine/Source/Programs/UnrealBuildAccelerator/Common/Private/UbaRootPaths.cpp
--- a/Engine/Source/Programs/UnrealBuildAccelerator/Common/Private/UbaRootPaths.cpp
+++ b/Engine/Source/Programs/UnrealBuildAccelerator/Common/Private/UbaRootPaths.cpp
@@ -107,6 +107,7 @@ namespace uba
 	{
 		// Register rootPath both with single path separators and double path separators on windows because text files store them with double path separators
 		#if PLATFORM_WINDOWS
+		u8 index = id * 4;
 		StringBuffer<512> forwardSlash;
 		StringBuffer<512> backwardSlash;
 		StringBuffer<512> doubleBackwardSlash;
@@ -149,16 +150,16 @@ namespace uba
 
 		const tchar* rootPaths[] = { forwardSlash.data, backwardSlash.data, doubleBackwardSlash.data, spaceEscapedBackwardSlash.data };
 		#else
+		u8 index = id;
 		const tchar* rootPaths[] = { rootPath };
 		#endif
 
 		for (const tchar* rp : rootPaths)
 		{
-			u8 index = id;
 			if (!InternalRegisterRoot2(logger, roots, rp, includeInKey, index))
 				return false;
 			if (id)
-				++id;
+				++index;
 		}
 		return true;
 	}
diff --git a/Engine/Source/Programs/UnrealBuildAccelerator/Common/Private/UbaSession.cpp b/Engine/Source/Programs/UnrealBuildAccelerator/Common/Private/UbaSession.cpp
--- a/Engine/Source/Programs/UnrealBuildAccelerator/Common/Private/UbaSession.cpp
+++ b/Engine/Source/Programs/UnrealBuildAccelerator/Common/Private/UbaSession.cpp
@@ -1613,11 +1613,11 @@ namespace uba
 		BinaryReader reader((const u8*)entry.memory.data(), 0, entry.memory.size());
 		while (reader.GetLeft())
 		{
-			u8 index = reader.ReadByte();(void)index; // RootPaths index;
+			u8 id = reader.ReadByte();
 			reader.SkipString();
 			StringBuffer<> rootPath;
 			reader.ReadString(rootPath.Clear());
-			if (!out.RegisterRoot(m_logger, rootPath.data, true, index))
+			if (!out.RegisterRoot(m_logger, rootPath.data, true, id))
 				return false;
 		}
 
@@ -3084,7 +3084,7 @@ namespace uba
 		BinaryReader reader((const u8*)rootsData, 0, rootsDataSize);
 		while (reader.GetLeft())
 		{
-			u8 index = reader.ReadByte();(void)index; // RootPaths index;
+			u8 id = reader.ReadByte();(void)id; // Root id.. ignore for conversion from vfs to local
 			StringBuffer<> temp;
 			reader.ReadString(temp);
 			if (!temp.count) // If vfs is not set it means that vfs should not be used (the roots entry memory might be used for cacheclient though)
