[UBA]
* Moved vfs code to detours shared
* Adde vfs logic to posix detouring code.

diff --git a/Engine/Source/Programs/UnrealBuildAccelerator/Core/Public/UbaStringBuffer.h b/Engine/Source/Programs/UnrealBuildAccelerator/Core/Public/UbaStringBuffer.h
--- a/Engine/Source/Programs/UnrealBuildAccelerator/Core/Public/UbaStringBuffer.h
+++ b/Engine/Source/Programs/UnrealBuildAccelerator/Core/Public/UbaStringBuffer.h
@@ -90,9 +90,12 @@ namespace uba
 		tchar buf[Capacity];
 	};
 
+	enum NoInitEnum { NoInit };
+
 	struct StringView
 	{
 		StringView() : data(TC("")), count(0) {}
+		StringView(NoInitEnum) {}
 		StringView(const StringView&) = default;
 		StringView(const tchar* d, u32 c) : data(d), count(c) {}
 		StringView(const StringBufferBase& sb) : data(sb.data), count(sb.count) {}
diff --git a/Engine/Source/Programs/UnrealBuildAccelerator/Detours/Private/PosixOS/UbaDetoursFunctionsPosix.cpp b/Engine/Source/Programs/UnrealBuildAccelerator/Detours/Private/PosixOS/UbaDetoursFunctionsPosix.cpp
--- a/Engine/Source/Programs/UnrealBuildAccelerator/Detours/Private/PosixOS/UbaDetoursFunctionsPosix.cpp
+++ b/Engine/Source/Programs/UnrealBuildAccelerator/Detours/Private/PosixOS/UbaDetoursFunctionsPosix.cpp
@@ -310,6 +310,7 @@ int Shared_open(const char* funcName, const char* file, int flags, int mode, con
 {
 	StringBuffer<> fileName;
 	FixPath(fileName, file);
+	DevirtualizePath(fileName);
 
 	if (!CanDetour2(fileName))
 	{
@@ -620,7 +621,7 @@ int Shared_fstat(const char* funcName, int fd, struct stat* attr, const True_fst
 		memcpy(attr, &fileAttr.data, sizeof(struct stat));
 
 	#if UBA_DEBUG_VALIDATE
-	if (!g_runningRemote)
+	if (!g_runningRemote && !IsVfsEnabled())
 	{
 		struct stat attr2;
 		int res2 = trueFstat(fd, &attr2);
@@ -664,6 +665,8 @@ int Shared_stat(const char* funcName, const char* file, struct stat* attr, const
 		fixedFile.Clear().Append(temp);
 	}
 
+	DevirtualizePath(fixedFile);
+
 	FileAttributes fileAttr;
 	const char* realName = Shared_GetFileAttributes(fileAttr, fixedFile.data);
 
@@ -684,7 +687,7 @@ int Shared_stat(const char* funcName, const char* file, struct stat* attr, const
 		memcpy(attr, &fileAttr.data, sizeof(fileAttr.data));
 
 	#if UBA_DEBUG_VALIDATE
-	if (!g_runningRemote && *file)
+	if (!g_runningRemote && !IsVfsEnabled() && *file)
 	{
 		struct stat attr2;
 		int res2 = trueStat(file, &attr2);
@@ -788,6 +791,7 @@ UBA_EXPORT int UBA_WRAPPER(mkdir)(const char* path, mode_t mode)
 		DEBUG_LOG_TRUE("mkdir", "%s -> %i", path, res);
 		return res;
 	}
+	DevirtualizePath(pathName);
 
 	u32 directoryTableSize;
 	int res;
@@ -827,6 +831,8 @@ UBA_EXPORT int UBA_WRAPPER(rmdir)(const char* path)
 		return res;
 	}
 
+	DevirtualizePath(pathName);
+
 	u32 directoryTableSize;
 	bool res;
 	u32 errorCode = 0;
@@ -1368,6 +1374,7 @@ UBA_EXPORT int UBA_WRAPPER(access)(const char* pathname, int mode)
 		DEBUG_LOG_TRUE("access", "%s %i -> %i (%s)", pathname, mode, res, StrError(res, errno));
 		return res;
 	}
+	DevirtualizePath(fixedPath);
 
 	bool checkIfDir = false;
 	StringBuffer<> temp;
@@ -1455,10 +1462,12 @@ UBA_EXPORT int UBA_WRAPPER(rename)(const char* oldpath, const char* newpath)
 
 	StringBuffer<> fixedOldPath;
 	FixPath(fixedOldPath, oldpath);
+	DevirtualizePath(fixedOldPath);
 	StringKey oldKey = ToFilenameKey(fixedOldPath);
 
 	StringBuffer<> fixedNewPath;
 	FixPath(fixedNewPath, newpath);
+	DevirtualizePath(fixedNewPath);
 	StringKey newKey = ToFilenameKey(fixedNewPath);
 
 	// TODO: This might be really slow but it seems you can rename files on linux while they are open and they won't be properly renamed until closed
@@ -1569,6 +1578,7 @@ UBA_EXPORT int UBA_WRAPPER(chmod)(const char* pathname, mode_t mode)
 
 	StringBuffer<> fixedName;
 	FixPath(fixedName, pathname);
+	DevirtualizePath(fixedName);
 
 	if (!CanDetour2(fixedName))
 	{
@@ -1622,8 +1632,10 @@ UBA_EXPORT int UBA_WRAPPER(symlink)(const char* path1, const char* path2)
 
 	StringBuffer<> fixedPath1;
 	FixPath(fixedPath1, path1);
+	DevirtualizePath(fixedPath1);
 	StringBuffer<> fixedPath2;
 	FixPath(fixedPath2, path2);
+	DevirtualizePath(fixedPath2);
 
 	if (!CanDetour2(fixedPath1) && !CanDetour2(fixedPath2))
 	{
@@ -1662,6 +1674,7 @@ int Shared_DeleteFile(const char* funcName, const char* pathname)
 {
 	StringBuffer<> fixedName;
 	FixPath(fixedName, pathname);
+	DevirtualizePath(fixedName);
 
 	if (!CanDetour2(fixedName))
 	{
@@ -1717,8 +1730,10 @@ UBA_EXPORT int UBA_WRAPPER(link)(const char* oldpath, const char* newpath)
 
 	StringBuffer<> fixedOldPath;
 	FixPath(fixedOldPath, oldpath);
+	DevirtualizePath(fixedOldPath);
 	StringBuffer<> fixedNewPath;
 	FixPath(fixedNewPath, newpath);
+	DevirtualizePath(fixedNewPath);
 
 	if (!CanDetour2(fixedNewPath))
 	{
@@ -2493,9 +2508,20 @@ namespace uba
 			mappedFileTableHandle = reader.ReadU64();
 			mappedFileTableSize = reader.ReadU32();
 			mappedFileTableCount = reader.ReadU32();
+
+			if (u16 vfsSize = reader.ReadU16())
+			{
+				BinaryReader vfsReader(reader.GetPositionData(), 0, vfsSize);
+				PopulateVfs(vfsReader);
+			}
+
 			DEBUG_LOG_PIPE(L"Init", L"");
 		}
 
+		VirtualizePath(g_virtualApplication);
+		VirtualizePath(g_virtualWorkingDir);
+		VirtualizePath(g_exeDir);
+
 		UBA_ASSERTF(g_virtualApplicationDir.capacity > 0, "g_virtualApplicationDir.capacity > 0");
 
 		const char* lastSlash = strrchr(g_virtualApplication.data, '/');
@@ -2530,6 +2556,8 @@ namespace uba
 		//},[](){},[](){});
 
 
+		LogVfsInfo();
+
 		g_isDetouring = true;
 		DEBUG_LOG("Detouring enabled");
 
diff --git a/Engine/Source/Programs/UnrealBuildAccelerator/Detours/Private/UbaDetoursShared.cpp b/Engine/Source/Programs/UnrealBuildAccelerator/Detours/Private/UbaDetoursShared.cpp
--- a/Engine/Source/Programs/UnrealBuildAccelerator/Detours/Private/UbaDetoursShared.cpp
+++ b/Engine/Source/Programs/UnrealBuildAccelerator/Detours/Private/UbaDetoursShared.cpp
@@ -165,6 +165,101 @@ namespace uba
 		return FixPath2(path, g_virtualWorkingDir.data, g_virtualWorkingDir.count, out.data, out.capacity, &out.count);
 	}
 
+	struct VfsEntry { StringView vfs; StringView local; VfsEntry() : vfs(NoInit), local(NoInit) {}; };
+	VfsEntry g_vfsEntries[32];
+	u32 g_vfsEntryCount;
+	u32 g_vfsMatchingLength;
+
+	void PopulateVfs(BinaryReader& vfsReader)
+	{
+		while (vfsReader.GetLeft())
+		{
+			vfsReader.ReadByte(); // Index, unused
+			StringBuffer<> str;
+			vfsReader.ReadString(str);
+			if (!str.count)
+			{
+				vfsReader.SkipString();
+				continue;
+			}
+
+			u32 index = g_vfsEntryCount++;
+			UBA_ASSERT(index < sizeof_array(g_vfsEntries));
+			VfsEntry& vfsEntry = g_vfsEntries[index];
+			vfsEntry.vfs = g_memoryBlock.Strdup(str);
+			if (index == 0)
+				g_vfsMatchingLength = vfsEntry.vfs.count;
+			else
+			{
+				u32 shortest = Min(g_vfsMatchingLength, vfsEntry.vfs.count);
+				for (u32 i=0; i!=shortest; ++i)
+				{
+					if (g_vfsEntries[0].vfs.data[i] == vfsEntry.vfs.data[i])
+						continue;
+					shortest = i;
+					break;
+				}
+				g_vfsMatchingLength = shortest;
+			}
+			vfsReader.ReadString(str.Clear());
+			vfsEntry.local = g_memoryBlock.Strdup(str);
+		}
+	}
+
+	bool IsVfsEnabled()
+	{
+		return g_vfsEntryCount > 0;
+	}
+
+	void DevirtualizePath(StringBufferBase& path)
+	{
+		if (!g_vfsEntryCount)
+			return;
+
+		if (!Equals(path.data, g_vfsEntries[0].vfs.data, Min(path.count, g_vfsMatchingLength), CaseInsensitiveFs))
+			return;
+
+		// TODO: This is not great, the dirs above the vfs root should be empty except the dir to the roots
+		if (path.count < g_vfsMatchingLength)
+		{
+			path.Clear().Append(g_vfsEntries[0].local);
+			return;
+		}
+
+		for (u32 i=0, e=g_vfsEntryCount; i!=e; ++i)
+		{
+			VfsEntry& entry = g_vfsEntries[i];
+			if (!path.StartsWith(entry.vfs.data))
+				continue;
+			StringBuffer<> temp2(path.data + entry.vfs.count);
+			path.Clear().Append(entry.local).Append(temp2);
+			return;
+		}
+	}
+
+	void VirtualizePath(StringBufferBase& path)
+	{
+		if (!g_vfsEntryCount)
+			return;
+		for (u32 i=0, e=g_vfsEntryCount; i!=e; ++i)
+		{
+			VfsEntry& entry = g_vfsEntries[i];
+			if (path.count < entry.local.count || !path.StartsWith(entry.local.data))
+				continue;
+			StringBuffer<> temp2(path.data + entry.local.count);
+			path.Clear().Append(entry.vfs).Append(temp2);
+			return;
+		}
+	}
+
+	void LogVfsInfo()
+	{
+		for (u32 i=0; i!=g_vfsEntryCount; ++i)
+		{
+			DEBUG_LOG(TC("Vfs: %s -> %s"), g_vfsEntries[i].vfs.data, g_vfsEntries[i].local.data);
+		}
+	}
+
 	const tchar* GetApplicationShortName()
 	{
 		const tchar* lastBackslash = TStrrchr(g_virtualApplication.data, '\\');
diff --git a/Engine/Source/Programs/UnrealBuildAccelerator/Detours/Private/UbaDetoursShared.h b/Engine/Source/Programs/UnrealBuildAccelerator/Detours/Private/UbaDetoursShared.h
--- a/Engine/Source/Programs/UnrealBuildAccelerator/Detours/Private/UbaDetoursShared.h
+++ b/Engine/Source/Programs/UnrealBuildAccelerator/Detours/Private/UbaDetoursShared.h
@@ -22,6 +22,12 @@ namespace uba
 
 	const tchar* GetApplicationShortName();
 	bool FixPath(StringBufferBase& out, const tchar* path);
+	
+	void PopulateVfs(BinaryReader& vfsReader);
+	bool IsVfsEnabled();
+	void DevirtualizePath(StringBufferBase& path);
+	void VirtualizePath(StringBufferBase& path);
+	void LogVfsInfo();
 
 	#if UBA_DEBUG_LOG_ENABLED
 		#define DEBUG_LOG_PREFIX(Prefix, Command, fmt, ...) \
diff --git a/Engine/Source/Programs/UnrealBuildAccelerator/Detours/Private/Windows/UbaDetoursFunctionsKernelBase.inl b/Engine/Source/Programs/UnrealBuildAccelerator/Detours/Private/Windows/UbaDetoursFunctionsKernelBase.inl
--- a/Engine/Source/Programs/UnrealBuildAccelerator/Detours/Private/Windows/UbaDetoursFunctionsKernelBase.inl
+++ b/Engine/Source/Programs/UnrealBuildAccelerator/Detours/Private/Windows/UbaDetoursFunctionsKernelBase.inl
@@ -438,7 +438,7 @@ BOOL Detoured_ReadConsoleW(HANDLE hConsoleInput, LPVOID lpBuffer, DWORD nNumberO
 UINT Detoured_GetDriveTypeW(LPCWSTR lpRootPathName)
 {
 	DETOURED_CALL(GetDriveTypeW);
-	if (g_runningRemote || g_vfsEntryCount)
+	if (g_runningRemote || IsVfsEnabled())
 	{
 		DEBUG_LOG_DETOURED(L"GetDriveType", L"%ls -> DRIVE_FIXED", lpRootPathName);
 		return DRIVE_FIXED;
@@ -1170,7 +1170,7 @@ BOOL Detoured_GetVolumePathNameW(LPCWSTR lpszFileName, LPWSTR lpszVolumePathName
 	//	return TRUE;
 	//}
 
-	if (g_runningRemote || g_vfsEntryCount)
+	if (g_runningRemote || IsVfsEnabled())
 	{
 		UBA_ASSERT(cchBufferLength > 3);
 		memcpy(lpszVolumePathName, g_virtualWorkingDir.data, 6);
@@ -1223,7 +1223,7 @@ DWORD Detoured_GetModuleFileNameW(HMODULE hModule, LPWSTR lpFilename, DWORD nSiz
 			return Shared_GetModuleFileNameW(hModule, findIt->second.c_str(), u32(findIt->second.size()), lpFilename, nSize);
 	}
 
-	if (!g_runningRemote && !g_vfsEntryCount)
+	if (!g_runningRemote && !IsVfsEnabled())
 		return True_GetModuleFileNameW(hModule, lpFilename, nSize);
 
 	wchar_t moduleName[350];
diff --git a/Engine/Source/Programs/UnrealBuildAccelerator/Detours/Private/Windows/UbaDetoursFunctionsWin.cpp b/Engine/Source/Programs/UnrealBuildAccelerator/Detours/Private/Windows/UbaDetoursFunctionsWin.cpp
--- a/Engine/Source/Programs/UnrealBuildAccelerator/Detours/Private/Windows/UbaDetoursFunctionsWin.cpp
+++ b/Engine/Source/Programs/UnrealBuildAccelerator/Detours/Private/Windows/UbaDetoursFunctionsWin.cpp
@@ -230,88 +230,6 @@ u64 g_memoryFileIndexCounter = ~u64(0) - 1000000; // I really hope this will not
 ObjFilesPreloader g_objFilesPreloader;
 
 
-struct VfsEntry { StringView vfs; StringView local; };
-VfsEntry g_vfsEntries[32];
-u32 g_vfsEntryCount;
-u32 g_vfsMatchingLength;
-
-void PopulateVfs(BinaryReader& vfsReader)
-{
-	while (vfsReader.GetLeft())
-	{
-		vfsReader.ReadByte(); // Index, unused
-		StringBuffer<> str;
-		vfsReader.ReadString(str);
-		if (!str.count)
-		{
-			vfsReader.SkipString();
-			continue;
-		}
-
-		u32 index = g_vfsEntryCount++;
-		VfsEntry& vfsEntry = g_vfsEntries[index];
-		vfsEntry.vfs = g_memoryBlock.Strdup(str);
-		if (index == 0)
-			g_vfsMatchingLength = vfsEntry.vfs.count;
-		else
-		{
-			u32 shortest = Min(g_vfsMatchingLength, vfsEntry.vfs.count);
-			for (u32 i=0; i!=shortest; ++i)
-			{
-				if (g_vfsEntries[0].vfs.data[i] == vfsEntry.vfs.data[i])
-					continue;
-				shortest = i;
-				break;
-			}
-			g_vfsMatchingLength = shortest;
-		}
-		vfsReader.ReadString(str.Clear());
-		vfsEntry.local = g_memoryBlock.Strdup(str);
-	}
-}
-
-void DevirtualizePath(StringBufferBase& path)
-{
-	if (!g_vfsEntryCount)
-		return;
-
-	if (!Equals(path.data, g_vfsEntries[0].vfs.data, Min(path.count, g_vfsMatchingLength), true))
-		return;
-
-	// TODO: This is not great, the dirs above the vfs root should be empty except the dir to the roots
-	if (path.count < g_vfsMatchingLength)
-	{
-		path.Clear().Append(g_vfsEntries[0].local);
-		return;
-	}
-
-	for (u32 i=0, e=g_vfsEntryCount; i!=e; ++i)
-	{
-		VfsEntry& entry = g_vfsEntries[i];
-		if (!path.StartsWith(entry.vfs.data))
-			continue;
-		StringBuffer<> temp2(path.data + entry.vfs.count);
-		path.Clear().Append(entry.local).Append(temp2);
-		return;
-	}
-}
-
-void VirtualizePath(StringBufferBase& path)
-{
-	if (!g_vfsEntryCount)
-		return;
-	for (u32 i=0, e=g_vfsEntryCount; i!=e; ++i)
-	{
-		VfsEntry& entry = g_vfsEntries[i];
-		if (path.count < entry.local.count || !path.StartsWith(entry.local.data))
-			continue;
-		StringBuffer<> temp2(path.data + entry.local.count);
-		path.Clear().Append(entry.vfs).Append(temp2);
-		return;
-	}
-}
-
-
 struct SuppressCreateFileDetourScope
 {
 	SuppressCreateFileDetourScope() { ++t_disallowCreateFileDetour; }
@@ -973,7 +891,6 @@ void Init(const DetoursPayload& payload, u64 startTime)
 		mappedFileTableSize = reader.ReadU32();
 		mappedFileTableCount = reader.ReadU32();
 
-
 		if (u16 vfsSize = reader.ReadU16())
 		{
 			BinaryReader vfsReader(reader.GetPositionData(), 0, vfsSize);
@@ -1038,10 +955,7 @@ void Init(const DetoursPayload& payload, u64 startTime)
 		DEBUG_LOG(L"WorkingDir: %ls", g_virtualWorkingDir.data);
 		DEBUG_LOG(L"ExeDir: %ls", g_virtualApplicationDir.data);
 		DEBUG_LOG(L"ExeDir (true): %ls", g_exeDir.data);
-		for (u32 i=0; i!=g_vfsEntryCount; ++i)
-		{
-			DEBUG_LOG(L"Vfs: %s -> %s", g_vfsEntries[i].vfs.data, g_vfsEntries[i].local.data);
-		}
+		LogVfsInfo();
 	}
 	#endif
 
