[UBA]
* Fixed rename fallback code which was handling sendfile in the wrong way.

diff --git a/Engine/Source/Programs/UnrealBuildAccelerator/Common/Private/UbaFileAccessor.cpp b/Engine/Source/Programs/UnrealBuildAccelerator/Common/Private/UbaFileAccessor.cpp
--- a/Engine/Source/Programs/UnrealBuildAccelerator/Common/Private/UbaFileAccessor.cpp
+++ b/Engine/Source/Programs/UnrealBuildAccelerator/Common/Private/UbaFileAccessor.cpp
@@ -426,8 +426,14 @@ namespace uba
 				int sourceFd = asFileDescriptor(m_fileHandle);
 				if (lseek(sourceFd, 0, SEEK_SET) == -1)
 					return m_logger.Error(TC("Failed to do lseek to beginning for sendfile (%s)"), strerror(errno));
-				if (sendfile(targetFd, sourceFd, NULL, m_size) != m_size)
-					return m_logger.Error(TC("Failed to do sendfile from temporary %s to file %s (%s)"), realFileName, m_fileName, strerror(errno));
+				u64 left = m_size;
+				while (left)
+				{
+					u64 written = sendfile(targetFd, sourceFd, NULL, left);
+					if (written == -1)
+						return m_logger.Error(TC("Failed to do sendfile from temporary %s (%i) to file %s (%i) (%s)"), realFileName, targetFd, m_fileName, sourceFd, strerror(errno));
+					left -= written;
+				}
 				#endif
 
 				remove(realFileName); // Remove real file now when we have copied it over
