[UBA]
* Lots of fixes around vfs and cache combo for posix
* Fixed early out of devirtualization if there are no vfs paths
* This change also fix the bug with zero-sized dependency files on posix when produced by helpers

diff --git a/Engine/Source/Programs/UnrealBuildAccelerator/Common/Private/UbaCacheClient.cpp b/Engine/Source/Programs/UnrealBuildAccelerator/Common/Private/UbaCacheClient.cpp
--- a/Engine/Source/Programs/UnrealBuildAccelerator/Common/Private/UbaCacheClient.cpp
+++ b/Engine/Source/Programs/UnrealBuildAccelerator/Common/Private/UbaCacheClient.cpp
@@ -129,7 +129,8 @@ namespace uba
 	bool CacheClient::WriteToCache(u32 bucketId, const ProcessStartInfo& info, const u8* inputs, u64 inputsSize, const u8* outputs, u64 outputsSize, const u8* logLines, u64 logLinesSize, u32 processId)
 	{
 		RootPaths rootPaths;
-		m_session.PopulateLocalToIndexRoots(rootPaths, info.rootsHandle);
+		if (!m_session.PopulateLocalToIndexRoots(rootPaths, info.rootsHandle))
+			return false;
 		return WriteToCache(rootPaths, bucketId, info, inputs, inputsSize, outputs, outputsSize, logLines, logLinesSize, processId);
 	}
 
@@ -728,12 +729,14 @@ namespace uba
 							pathTemp.Append(root);
 							char rootPath[512];
 							u32 rootPathLen = pathTemp.Parse(rootPath, sizeof_array(rootPath));
+							UBA_ASSERT(rootPathLen);
+							--rootPathLen;
 							#else
 							const char* rootPath = root.data();
 							u32 rootPathLen = root.size();
 							#endif
 
-							if (u32 toWrite = rootPathLen - 1)
+							if (u32 toWrite = rootPathLen)
 								memcpy(localBlock.Allocate(toWrite, 1, TC("")), rootPath, toWrite);
 							lastWritten = rootOffset + 1;
 						}
@@ -1214,21 +1217,21 @@ namespace uba
 					rsp.Append(rspStart, rspEnd - rspStart);
 					StringBuffer<> fullPath;
 					FixPath(rsp.data, workingDir.data, workingDir.count, fullPath);
-					if (!DevirtualizePath(rsp, info.rootsHandle))
+					if (!DevirtualizePath(fullPath, info.rootsHandle))
 					{
 						if (report)
-							m_logger.Warning(TC("Failed to normalize rsp file path %s"), rsp.data);
+							m_logger.Warning(TC("Failed to normalize rsp file path %s"), fullPath.data);
 						return CasKeyZero;
 					}
-					CasKey rspCasKey = rootPaths.NormalizeAndHashFile(m_logger, rsp.data, true);
+					CasKey rspCasKey = rootPaths.NormalizeAndHashFile(m_logger, fullPath.data, true);
 					if (rspCasKey == CasKeyZero)
 					{
 						if (report)
-							m_logger.Info(TC("   Failed to normalize rsp file %s"), rsp.data);
+							m_logger.Info(TC("   Failed to normalize rsp file %s"), fullPath.data);
 						return CasKeyZero;
 					}
 					if (report)
-						m_logger.Info(TC("   %s %s"), CasKeyString(rspCasKey).str, rsp.data);
+						m_logger.Info(TC("   %s %s"), CasKeyString(rspCasKey).str, fullPath.data);
 					hasher.Update(&rspCasKey, sizeof(CasKey));
 				}
 			}
diff --git a/Engine/Source/Programs/UnrealBuildAccelerator/Common/Private/UbaProcess.cpp b/Engine/Source/Programs/UnrealBuildAccelerator/Common/Private/UbaProcess.cpp
--- a/Engine/Source/Programs/UnrealBuildAccelerator/Common/Private/UbaProcess.cpp
+++ b/Engine/Source/Programs/UnrealBuildAccelerator/Common/Private/UbaProcess.cpp
@@ -42,6 +42,8 @@ extern char **environ;
 
 namespace uba
 {
+	static constexpr const tchar g_extractExportsStr[] = TC("/extractexports");
+
 	void Process::AddRef()
 	{
 		++m_refCount;
@@ -172,7 +174,7 @@ namespace uba
 
 		m_startInfo.Expand();
 
-		m_extractExports = Contains(m_startInfo.arguments, TC("/extractexports"), true);
+		m_extractExports = Contains(m_startInfo.arguments, g_extractExportsStr, true);
 
 		StringBuffer<> realApplication(m_startInfo.application);
 		const tchar* realWorkingDir = m_startInfo.workingDir;
@@ -734,13 +736,11 @@ namespace uba
 			return true;
 		}
 
-		SCOPED_READ_LOCK(m_session.m_rootsLookupLock, lock);
-		auto findIt = m_session.m_rootsLookup.find(m_startInfo.rootsHandle);
-		if (findIt == m_session.m_rootsLookup.end())
-			return m_session.m_logger.Error(TC("Process use roots handle %llu which does not exist"), m_startInfo.rootsHandle);
-		Session::RootsEntry& rootsEntry = findIt->second;
-		writer.WriteU16(u16(rootsEntry.memory.size()));
-		writer.WriteBytes(rootsEntry.memory.data(), rootsEntry.memory.size());
+		auto rootsEntry = m_session.GetRootsEntry(m_startInfo.rootsHandle);
+		if (!rootsEntry)
+			return false;
+		writer.WriteU16(u16(rootsEntry->memory.size()));
+		writer.WriteBytes(rootsEntry->memory.data(), rootsEntry->memory.size());
 		return true;
 	}
 	
@@ -1470,10 +1470,9 @@ namespace uba
 
 			if (m_extractExports)
 			{
-				tchar extractExports[] = TC(" /extractexports");
 				const tchar* pos = nullptr;
-				Contains(commandLine.c_str(), extractExports, true, &pos);
-				commandLine.erase(pos - commandLine.c_str(), sizeof_array(extractExports));
+				Contains(commandLine.c_str(), g_extractExportsStr, true, &pos);
+				commandLine.erase(pos - commandLine.c_str(), sizeof_array(g_extractExportsStr));
 			}
 
 			LPCSTR dlls[] = { detoursLib };
@@ -1775,7 +1774,12 @@ namespace uba
 			}
 
 			Vector<TString> arguments;
-			if (!ParseArguments(m_startInfo.arguments, [&](const tchar* arg, u32 argLen) { arguments.push_back(TString(arg, argLen)); }))
+			if (!ParseArguments(m_startInfo.arguments, [&](const tchar* arg, u32 argLen)
+				{
+					if (Equals(arg, g_extractExportsStr))
+						return;
+					arguments.push_back(TString(arg, argLen));
+				}))
 			{
 				logger.Error("Failed to parse arguments: %s", m_startInfo.arguments);
 				return UBA_EXIT_CODE(16);
diff --git a/Engine/Source/Programs/UnrealBuildAccelerator/Common/Private/UbaSession.cpp b/Engine/Source/Programs/UnrealBuildAccelerator/Common/Private/UbaSession.cpp
--- a/Engine/Source/Programs/UnrealBuildAccelerator/Common/Private/UbaSession.cpp
+++ b/Engine/Source/Programs/UnrealBuildAccelerator/Common/Private/UbaSession.cpp
@@ -48,8 +48,6 @@ typedef struct _PROCESSOR_POWER_INFORMATION {
 
 namespace uba
 {
-	constexpr u32 PathsPerRoot = 4u; // 3 + forward slashes
-
 	bool g_dummy;
 
 	ProcessStartInfo::ProcessStartInfo() = default;
@@ -948,18 +946,16 @@ namespace uba
 	RootsHandle Session::RegisterRoots(const void* rootsData, uba::u64 rootsDataSize)
 	{
 		CasKey key = ToCasKey(CasKeyHasher().Update(rootsData, rootsDataSize), false);
-		RootsHandle rootsHandle = key.a;
+		RootsHandle rootsHandle = WithVfs(key.a, false);
 
 		SCOPED_WRITE_LOCK(m_rootsLookupLock, lock);
 		auto insres = m_rootsLookup.try_emplace(rootsHandle);
 		RootsEntry& entry = insres.first->second;
-		if (!insres.second)
-		{
+		if (insres.second)
+			PopulateRootsEntry(entry, rootsData, rootsDataSize);
+		else
 			UBA_ASSERT(memcmp(entry.memory.data(), rootsData, rootsDataSize) == 0);
-			return rootsHandle;
-		}
-		PopulateRootsEntry(entry, rootsData, rootsDataSize);
-		return rootsHandle;
+		return WithVfs(rootsHandle, !entry.locals.empty());
 	}
 
 	void Session::CheckDirectoryTableOverflow(DirectoryTable& dirTable)
@@ -1578,21 +1574,18 @@ namespace uba
 			return true;
 #endif
 
-		SCOPED_READ_LOCK(m_rootsLookupLock, lock);
-		auto findIt = m_rootsLookup.find(rootsHandle);
-		if (findIt == m_rootsLookup.end())
-			return m_logger.Error(TC("Can't find entry from roots handle %llu"), rootsHandle);
-		RootsEntry& rootsEntry = findIt->second;
-		lock.Leave();
+		auto rootsEntry = GetRootsEntry(rootsHandle);
+		if (!rootsEntry)
+			return false;
 
-		if (rootsEntry.roots.IsEmpty())
+		if (rootsEntry->roots.IsEmpty())
 			return true;
 
-		auto root = rootsEntry.roots.FindRoot(inOut);
+		auto root = rootsEntry->roots.FindRoot(inOut);
 		if (!root)
 			return m_logger.Error(TC("Can't find root for path %s"), inOut.data);
 
-		auto& path = rootsEntry.locals[root->index/PathsPerRoot];
+		auto& path = rootsEntry->locals[root->index/PathsPerRoot];
 		StringBuffer<> temp;
 		temp.Append(inOut.data + root->path.size());
 		inOut.Clear().Append(path).Append(temp);
@@ -1603,14 +1596,11 @@ namespace uba
 	{
 		if (!rootsHandle)
 			return true;
-		SCOPED_READ_LOCK(m_rootsLookupLock, lock)
-		auto findIt = m_rootsLookup.find(rootsHandle);
-		if (findIt == m_rootsLookup.end())
-			return m_logger.Error(TC("Roots handle %llu does not exist"), rootsHandle);
-		RootsEntry& entry = findIt->second;
-		lock.Leave();
+		auto rootsEntry = GetRootsEntry(rootsHandle);
+		if (!rootsEntry)
+			return false;
 
-		BinaryReader reader((const u8*)entry.memory.data(), 0, entry.memory.size());
+		BinaryReader reader((const u8*)rootsEntry->memory.data(), 0, rootsEntry->memory.size());
 		while (reader.GetLeft())
 		{
 			u8 id = reader.ReadByte();
@@ -2181,6 +2171,36 @@ namespace uba
 				}
 			}
 
+			ProcessStartInfo& startInfo = msg.process.m_startInfo;
+
+			if (!m_runningRemote && HasVfs(startInfo.rootsHandle)) // For posix we write the dependency file directly to disk so we need to update it if vfs is enabled
+			{
+				bool escapeSpaces;
+				if (!msg.mappingHandle && startInfo.rules->ShouldDevirtualizeFile(file.name, escapeSpaces))
+				{
+					// TODO: On linux we don't use file mappings for outputs yet.. so we have to open the file and change it
+					FileAccessor readFile(m_logger, name);
+					if (!readFile.OpenMemoryRead())
+						return false;
+					void* mem = readFile.GetData();
+					fileSize = readFile.GetSize();
+					MemoryBlock block(5*1024*1024);
+					RootsHandle rootsHandle = startInfo.rootsHandle;
+					if (!DevirtualizeDepsFile(rootsHandle, block, mem, fileSize, escapeSpaces, name))
+						return false;
+					if (!readFile.Close())
+						return false;
+					FileAccessor writeFile(m_logger, name);
+					if (!writeFile.CreateWrite())
+						return false;
+					if (!writeFile.Write(block.memory, block.writtenSize))
+						return false;
+					if (!writeFile.Close(&lastWriteTime))
+						return false;
+					fileSize = block.writtenSize;
+				}
+			}
+
 			if (addMapping)
 			{
 				writtenFile.name = msgName;
@@ -2206,7 +2226,7 @@ namespace uba
 				lastWriteTime = writtenFile.lastWriteTime;
 			}
 
-			if (msg.process.m_extractExports && msg.process.m_startInfo.rules->ShouldExtractSymbols(file.name))
+			if (msg.process.m_extractExports && startInfo.rules->ShouldExtractSymbols(file.name))
 				if (!ExtractSymbolsFromObjectFile(msg, name, fileSize))
 					return false;
 		}
@@ -2397,7 +2417,7 @@ namespace uba
 
 					RootsHandle rootsHandle = process.GetStartInfo().rootsHandle;
 					bool escapeSpaces;
-					if (rootsHandle && rules.ShouldDevirtualizeFile(file.name, escapeSpaces))
+					if (HasVfs(rootsHandle) && rules.ShouldDevirtualizeFile(file.name, escapeSpaces))
 					{
 						// Need to turn paths back into local paths
 
@@ -3076,6 +3096,26 @@ namespace uba
 		return m_cpuLoad;
 	}
 
+	bool Session::HasVfs(RootsHandle handle) const
+	{
+		return (handle & 1ull) == 1ull;
+	}
+
+	RootsHandle Session::WithVfs(RootsHandle handle, bool vfs) const
+	{
+		return vfs ? (handle | 1ull) : (handle & ~1ull);
+	}
+
+	const Session::RootsEntry* Session::GetRootsEntry(RootsHandle rootsHandle)
+	{
+		SCOPED_READ_LOCK(m_rootsLookupLock, lock);
+		auto findIt = m_rootsLookup.find(WithVfs(rootsHandle, false));
+		if (findIt != m_rootsLookup.end())
+			return &findIt->second;
+		m_logger.Error(TC("Can't find entry from roots handle %llu"), rootsHandle);
+		return nullptr;
+	}
+
 	void Session::PopulateRootsEntry(RootsEntry& entry, const void* rootsData, uba::u64 rootsDataSize)
 	{
 		entry.memory.resize(rootsDataSize);
@@ -3166,11 +3206,12 @@ namespace uba
 
 	bool Session::DevirtualizeDepsFile(RootsHandle rootsHandle, MemoryBlock& destData, const void* sourceData, u64 sourceSize, bool escapeSpaces, const tchar* hint)
 	{
-		SCOPED_READ_LOCK(m_rootsLookupLock, lock);
-		auto findIt = m_rootsLookup.find(rootsHandle);
-		UBA_ASSERT(findIt != m_rootsLookup.end());
-		RootsEntry& rootsEntry = findIt->second;
-		lock.Leave();
+		auto rootsEntryPtr = GetRootsEntry(rootsHandle);
+		if (!rootsEntryPtr)
+			return false;
+		const RootsEntry& rootsEntry = *rootsEntryPtr;
+
+		UBA_ASSERT(!rootsEntry.locals.empty());
 
 		Vector<std::string> localsAnsi;
 		localsAnsi.reserve(rootsEntry.locals.size());
diff --git a/Engine/Source/Programs/UnrealBuildAccelerator/Common/Private/UbaSessionClient.cpp b/Engine/Source/Programs/UnrealBuildAccelerator/Common/Private/UbaSessionClient.cpp
--- a/Engine/Source/Programs/UnrealBuildAccelerator/Common/Private/UbaSessionClient.cpp
+++ b/Engine/Source/Programs/UnrealBuildAccelerator/Common/Private/UbaSessionClient.cpp
@@ -1494,7 +1494,7 @@ namespace uba
 	bool SessionClient::SendRootsHandle(RootsHandle rootsHandle)
 	{
 		SCOPED_WRITE_LOCK(m_rootsLookupLock, rootsLock);
-		auto insres2 = m_rootsLookup.try_emplace(rootsHandle);
+		auto insres2 = m_rootsLookup.try_emplace(WithVfs(rootsHandle, false));
 		if (!insres2.second)
 			return true;
 		StackBinaryWriter<128> writer;
diff --git a/Engine/Source/Programs/UnrealBuildAccelerator/Common/Private/UbaSessionServer.cpp b/Engine/Source/Programs/UnrealBuildAccelerator/Common/Private/UbaSessionServer.cpp
--- a/Engine/Source/Programs/UnrealBuildAccelerator/Common/Private/UbaSessionServer.cpp
+++ b/Engine/Source/Programs/UnrealBuildAccelerator/Common/Private/UbaSessionServer.cpp
@@ -998,7 +998,7 @@ namespace uba
 
 				Storage::FormattingFunc formattingFunc;
 				bool escapeSpaces;
-				if (rootsHandle && rules.ShouldDevirtualizeFile(destination, escapeSpaces))
+				if (HasVfs(rootsHandle) && rules.ShouldDevirtualizeFile(destination, escapeSpaces))
 				{
 					formattingFunc = [&](MemoryBlock& destData, const void* sourceData, u64 sourceSize, const tchar* hint)
 						{
@@ -1430,13 +1430,10 @@ namespace uba
 	bool SessionServer::HandleGetRoots(const ConnectionInfo& connectionInfo, BinaryReader& reader, BinaryWriter& writer)
 	{
 		RootsHandle rootsHandle = reader.ReadU64();
-		SCOPED_READ_LOCK(m_rootsLookupLock, lock);
-		auto findIt = m_rootsLookup.find(rootsHandle);
-		if (findIt == m_rootsLookup.end())
-			return m_logger.Error(TC("Remote process use roots handle %llu which does not exist"), rootsHandle);
-		RootsEntry& rootsEntry = findIt->second;
-		lock.Leave();
-		writer.WriteBytes(rootsEntry.memory.data(), rootsEntry.memory.size());
+		auto rootsEntry = GetRootsEntry(rootsHandle);
+		if (!rootsEntry)
+			return false;
+		writer.WriteBytes(rootsEntry->memory.data(), rootsEntry->memory.size());
 		return true;
 	}
 
diff --git a/Engine/Source/Programs/UnrealBuildAccelerator/Common/Public/UbaRootPaths.h b/Engine/Source/Programs/UnrealBuildAccelerator/Common/Public/UbaRootPaths.h
--- a/Engine/Source/Programs/UnrealBuildAccelerator/Common/Public/UbaRootPaths.h
+++ b/Engine/Source/Programs/UnrealBuildAccelerator/Common/Public/UbaRootPaths.h
@@ -8,6 +8,7 @@
 namespace uba
 {
 	static constexpr u8 RootPathsVersion = 1;
+	static constexpr u32 PathsPerRoot = IsWindows ? 4u : 1;
 
 	class RootPaths
 	{
@@ -62,15 +63,27 @@ namespace uba
 
 		u32 destPos = 0;
 
+#if !PLATFORM_WINDOWS
+		allowPathsWithoutRoot = true; // Posix uses forward slash and no drive letter so we can't see if a path is in the middle of a path or beginning
+#endif
+
 		while (true)
 		{
 			auto absPathChars = searchPos;
 			CharType lastChar = 0;
-			while (absPathChars < strEnd && !(lastChar == ':' && ((IsWindows && *absPathChars == '\\') || *absPathChars == '/')))
+#if PLATFORM_WINDOWS
+			while (absPathChars < strEnd && !(lastChar == ':' && (*absPathChars == '\\' || *absPathChars == '/')))
+			{
+				lastChar = *absPathChars;
+				++absPathChars;
+			}
+#else
+			while (absPathChars < strEnd && *absPathChars != '/')
 			{
 				lastChar = *absPathChars;
 				++absPathChars;
 			}
+#endif
 		
 			if (absPathChars == strEnd)
 			{
@@ -78,7 +91,7 @@ namespace uba
 				return true;
 			}
 
-			auto pathStart = absPathChars - 2;
+			auto pathStart = absPathChars - (IsWindows ? 2 : 0);
 
 			auto pathEndOrMore = pathStart;
 			while (pathEndOrMore < strEnd && *pathEndOrMore != '\n')
@@ -105,7 +118,9 @@ namespace uba
 					path.Resize(lastQuote - path.data);
 				if (auto lineEnd = path.Last('\r'))
 					path.Resize(lineEnd - path.data);
+#if PLATFORM_WINDOWS
 				logger.Info(TC("PATH WITHOUT ROOT: %s (inside %s at offset %u%s)"), path.data, hint, destPos, hint2);
+#endif
 				return false;
 			}
 
diff --git a/Engine/Source/Programs/UnrealBuildAccelerator/Common/Public/UbaSession.h b/Engine/Source/Programs/UnrealBuildAccelerator/Common/Public/UbaSession.h
--- a/Engine/Source/Programs/UnrealBuildAccelerator/Common/Public/UbaSession.h
+++ b/Engine/Source/Programs/UnrealBuildAccelerator/Common/Public/UbaSession.h
@@ -166,7 +166,10 @@ namespace uba
 		void WriteSummary(BinaryWriter& writer, const Function<void(Logger& logger)>& summaryFunc);
 		float UpdateCpuLoad();
 
+		bool HasVfs(RootsHandle handle) const;
+		RootsHandle WithVfs(RootsHandle key, bool vfs) const;
 		struct RootsEntry { Vector<u8> memory; RootPaths roots; Vector<TString> locals; };
+		const RootsEntry* GetRootsEntry(RootsHandle rootsHandle);
 		void PopulateRootsEntry(RootsEntry& entry, const void* rootsData, uba::u64 rootsDataSize);
 
 		bool ExtractSymbolsFromObjectFile(const CloseFileMessage& msg, const tchar* fileName, u64 fileSize);
diff --git a/Engine/Source/Programs/UnrealBuildAccelerator/Core/Private/UbaApplicationRules.cpp b/Engine/Source/Programs/UnrealBuildAccelerator/Core/Private/UbaApplicationRules.cpp
--- a/Engine/Source/Programs/UnrealBuildAccelerator/Core/Private/UbaApplicationRules.cpp
+++ b/Engine/Source/Programs/UnrealBuildAccelerator/Core/Private/UbaApplicationRules.cpp
@@ -631,6 +631,7 @@ namespace uba
 #else
 			{ ApplicationHash(TC("clang++")),					new ApplicationRulesClangPlusPlusExe() },
 			{ ApplicationHash(TC("ld.lld")),					new ApplicationRulesLdLLdExe() },
+			{ ApplicationHash(TC("ispc")),						new ApplicationRulesISPCExe() },
 			{ ApplicationHash(TC("shadercompileworker")),		new ApplicationRulesShaderCompileWorker() },
 #endif
 			{ 0, nullptr }
diff --git a/Engine/Source/Programs/UnrealBuildAccelerator/Test/Private/UbaTestBasics.h b/Engine/Source/Programs/UnrealBuildAccelerator/Test/Private/UbaTestBasics.h
--- a/Engine/Source/Programs/UnrealBuildAccelerator/Test/Private/UbaTestBasics.h
+++ b/Engine/Source/Programs/UnrealBuildAccelerator/Test/Private/UbaTestBasics.h
@@ -402,7 +402,7 @@ namespace uba
 				{
 					if (strLen != 1)
 						success = false;
-					if (str[0] != RootPaths::RootStartByte + (IsWindows ? 5 : 1)) // 5 is because each root has four entries on windows and second entry is backslash
+					if (str[0] != RootPaths::RootStartByte + PathsPerRoot + (IsWindows ? 1 : 0)) // Add one for windows because second entry is backslash
 						success = false;
 					rootPos = str[0];
 				}
