[UBA]
* Fixed so errors are devirtualized

diff --git a/Engine/Source/Programs/UnrealBuildAccelerator/Common/Private/UbaProcess.cpp b/Engine/Source/Programs/UnrealBuildAccelerator/Common/Private/UbaProcess.cpp
--- a/Engine/Source/Programs/UnrealBuildAccelerator/Common/Private/UbaProcess.cpp
+++ b/Engine/Source/Programs/UnrealBuildAccelerator/Common/Private/UbaProcess.cpp
@@ -1171,6 +1171,8 @@ namespace uba
 		TString line = reader.ReadString();
 		LogEntryType entryType = isError ? LogEntryType_Error : LogEntryType_Info;
 
+		m_session.DevirtualizeString(line, m_startInfo.rootsHandle, TC("LogLine"));
+
 		if (!m_session.LogLine(*this, line.c_str(), entryType))
 			return false;
 
diff --git a/Engine/Source/Programs/UnrealBuildAccelerator/Common/Private/UbaSession.cpp b/Engine/Source/Programs/UnrealBuildAccelerator/Common/Private/UbaSession.cpp
--- a/Engine/Source/Programs/UnrealBuildAccelerator/Common/Private/UbaSession.cpp
+++ b/Engine/Source/Programs/UnrealBuildAccelerator/Common/Private/UbaSession.cpp
@@ -1592,6 +1592,59 @@ namespace uba
 		return true;
 	}
 
+	bool Session::DevirtualizeString(TString& inOut, RootsHandle rootsHandle, const tchar* hint)
+	{
+		if (!HasVfs(rootsHandle))
+			return true;
+		auto rootsEntry = GetRootsEntry(rootsHandle);
+		if (!rootsEntry)
+			return false;
+
+		u64 newSize = 0;
+		bool hasRoot = false;
+		auto checkString = [&](const tchar* str, u64 strLen, u32 rootPos)
+			{
+				if (rootPos == ~0u)
+				{
+					newSize += strLen;
+					return;
+				}
+				hasRoot = true;
+				u32 localsIndex = (*str - RootPaths::RootStartByte)/PathsPerRoot;
+				auto& path = rootsEntry->locals[localsIndex];
+				newSize += path.size() - strLen;
+			};
+
+		if (!rootsEntry->roots.NormalizeString<tchar>(m_logger, inOut.data(), inOut.size(), checkString, false, hint))
+			return false;
+
+		if (!hasRoot)
+			return false;
+
+		TString newString;
+		newString.resize(newSize);
+		tchar* newStringPos = newString.data();
+
+		auto handleString = [&](const tchar* str, u64 strLen, u32 rootPos)
+			{
+				if (rootPos == ~0u)
+				{
+					memcpy(newStringPos, str, strLen*sizeof(tchar));
+					newStringPos += strLen;
+					return;
+				}
+				u32 localsIndex = (*str - RootPaths::RootStartByte)/PathsPerRoot;
+				auto& path = rootsEntry->locals[localsIndex];
+				memcpy(newStringPos, path.data(), path.size()*sizeof(tchar));
+				newStringPos += path.size();
+			};
+
+		rootsEntry->roots.NormalizeString<tchar>(m_logger, inOut.data(), inOut.size(), handleString, false, hint);
+		inOut = std::move(newString);
+
+		return true;
+	}
+
 	bool Session::PopulateLocalToIndexRoots(RootPaths& out, RootsHandle rootsHandle)
 	{
 		if (!rootsHandle)
diff --git a/Engine/Source/Programs/UnrealBuildAccelerator/Common/Public/UbaSession.h b/Engine/Source/Programs/UnrealBuildAccelerator/Common/Public/UbaSession.h
--- a/Engine/Source/Programs/UnrealBuildAccelerator/Common/Public/UbaSession.h
+++ b/Engine/Source/Programs/UnrealBuildAccelerator/Common/Public/UbaSession.h
@@ -85,6 +85,7 @@ namespace uba
 		u32 CreateProcessId();
 
 		bool DevirtualizePath(StringBufferBase& inOut, RootsHandle rootsHandle);
+		bool DevirtualizeString(TString& inOut, RootsHandle rootsHandle, const tchar* hint = TC(""));
 		bool PopulateLocalToIndexRoots(RootPaths& out, RootsHandle rootsHandle);
 
 	protected:
