[UBA]
* Fixed code traversing text to normalize roots to not go character by character when found root to skip

diff --git a/Engine/Source/Programs/UnrealBuildAccelerator/Common/Public/UbaRootPaths.h b/Engine/Source/Programs/UnrealBuildAccelerator/Common/Public/UbaRootPaths.h
--- a/Engine/Source/Programs/UnrealBuildAccelerator/Common/Public/UbaRootPaths.h
+++ b/Engine/Source/Programs/UnrealBuildAccelerator/Common/Public/UbaRootPaths.h
@@ -107,7 +107,7 @@ namespace uba
 			{
 				if (allowPathsWithoutRoot || InternalFindRoot(m_ignoredRoots, path))
 				{
-					u32 len = 1;
+					u32 len = u32(absPathChars - searchPos) + 1;
 					destPos += len;
 					func(searchPos, len, ~0u);
 					searchPos += len;
