[UBA]
* Fixed off-by-one size for code devirtualizing error messages

diff --git a/Engine/Source/Programs/UnrealBuildAccelerator/Common/Private/UbaSession.cpp b/Engine/Source/Programs/UnrealBuildAccelerator/Common/Private/UbaSession.cpp
--- a/Engine/Source/Programs/UnrealBuildAccelerator/Common/Private/UbaSession.cpp
+++ b/Engine/Source/Programs/UnrealBuildAccelerator/Common/Private/UbaSession.cpp
@@ -1622,7 +1622,7 @@ namespace uba
 			return false;
 
 		TString newString;
-		newString.resize(newSize);
+		newString.resize(newSize + 1);
 		tchar* newStringPos = newString.data();
 
 		auto handleString = [&](const tchar* str, u64 strLen, u32 rootPos)
