[UBA]
* Improved error handling when timing out

diff --git a/Engine/Source/Programs/UnrealBuildAccelerator/Common/Private/UbaNetworkClient.cpp b/Engine/Source/Programs/UnrealBuildAccelerator/Common/Private/UbaNetworkClient.cpp
--- a/Engine/Source/Programs/UnrealBuildAccelerator/Common/Private/UbaNetworkClient.cpp
+++ b/Engine/Source/Programs/UnrealBuildAccelerator/Common/Private/UbaNetworkClient.cpp
@@ -572,12 +572,16 @@ namespace uba
 			{
 				m_logger.Error(TC("Timed out after %s waiting for message response from server."), TimeToText(GetTime() - waitStart, true).str);
 				message.m_error = 4;
+				OnDisconnected(connection, 4);
 			}
 			else if (m_cryptoKey && !message.m_error && message.m_responseSize)
 			{
 				TimerScope ts(m_decryptTimer);
 				if (!Crypto::Decrypt(m_logger, m_cryptoKey, (u8*)message.m_response, message.m_responseSize))
+				{
 					message.m_error = 5;
+					OnDisconnected(connection, 5);
+				}
 			}
 		}
 		return !message.m_error;
diff --git a/Engine/Source/Programs/UnrealBuildAccelerator/Common/Private/UbaStorageUtils.cpp b/Engine/Source/Programs/UnrealBuildAccelerator/Common/Private/UbaStorageUtils.cpp
--- a/Engine/Source/Programs/UnrealBuildAccelerator/Common/Private/UbaStorageUtils.cpp
+++ b/Engine/Source/Programs/UnrealBuildAccelerator/Common/Private/UbaStorageUtils.cpp
@@ -101,7 +101,7 @@ namespace uba
 		return ToCasKey(hasher, storeCompressed);
 	}
 
-	bool SendBatchMessages(Logger& logger, NetworkClient& client, u16 fetchId, u8* slot, u64 capacity, u64 left, u32 messageMaxSize, u32& readIndex, u32& responseSize, const Function<bool()>& runInWaitFunc)
+	bool SendBatchMessages(Logger& logger, NetworkClient& client, u16 fetchId, u8* slot, u64 capacity, u64 left, u32 messageMaxSize, u32& readIndex, u32& responseSize, const Function<bool()>& runInWaitFunc, const tchar* hint)
 	{
 		responseSize = 0;
 
@@ -155,7 +155,7 @@ namespace uba
 			Entry& entry = entries[i];
 			if (!entry.done.IsSet(timeOutTimeMs))
 			{
-				logger.Error(TC("SendBatchMessages timed out after 5 minutes getting async message response (%u/%u). Received %llu bytes so far. FetchId: %u"), i, inFlightCount, responseSize, fetchId);
+				logger.Error(TC("SendBatchMessages timed out after 5 minutes getting async message response (%u/%u). Received %llu bytes so far. FetchId: %u (%s)"), i, inFlightCount, responseSize, fetchId, hint);
 				timeOutTimeMs = 10;
 			}
 			if (!entry.message.ProcessAsyncResults(entry.reader))
@@ -531,7 +531,7 @@ namespace uba
 				if (fetchId == u16(~0))
 					return logger.Error(TC("Cas content error (2). Server believes %s was only one segment but client sees more. "));//UncompressedSize: %llu LeftUncompressed: %llu Size: %llu Left to read: %llu ResponseSize: %u. (%s)"), destination, actualSize, leftUncompressed, fileSize, left, responseSize, CasKeyString(casKey).str);
 
-				if (!SendBatchMessages(logger, client, fetchId, source, BufferSlotHalfSize, leftCompressed, sizeOfFirstMessage, readIndex, responseSize, writePrev))
+				if (!SendBatchMessages(logger, client, fetchId, source, BufferSlotHalfSize, leftCompressed, sizeOfFirstMessage, readIndex, responseSize, writePrev, destination))
 					return logger.Error(TC("Failed to send batched messages to server (%s)"), CasKeyString(casKey).str);
 
 				lastSource = source;
@@ -626,6 +626,8 @@ namespace uba
 						}
 						else
 						{
+							if (left < responseSize)
+								return logger.Error(TC("Something went wrong. Left %u, Response: %u (%s)"), left, responseSize, destination);
 							left -= responseSize;
 						}
 					}
diff --git a/Engine/Source/Programs/UnrealBuildAccelerator/Common/Public/UbaStorageUtils.h b/Engine/Source/Programs/UnrealBuildAccelerator/Common/Public/UbaStorageUtils.h
--- a/Engine/Source/Programs/UnrealBuildAccelerator/Common/Public/UbaStorageUtils.h
+++ b/Engine/Source/Programs/UnrealBuildAccelerator/Common/Public/UbaStorageUtils.h
@@ -16,7 +16,7 @@ namespace uba
 	static constexpr u8 DefaultCompressionLevel = OodleLZ_CompressionLevel_SuperFast;
 
 	CasKey CalculateCasKey(u8* fileMem, u64 fileSize, bool storeCompressed, WorkManager* workManager, const tchar* hint);
-	bool SendBatchMessages(Logger& logger, NetworkClient& client, u16 fetchId, u8* slot, u64 capacity, u64 left, u32 messageMaxSize, u32& readIndex, u32& responseSize, const Function<bool()>& runInWaitFunc = {});
+	bool SendBatchMessages(Logger& logger, NetworkClient& client, u16 fetchId, u8* slot, u64 capacity, u64 left, u32 messageMaxSize, u32& readIndex, u32& responseSize, const Function<bool()>& runInWaitFunc = {}, const tchar* hint = TC(""));
 	bool SendFile(Logger& logger, NetworkClient& client, const CasKey& casKey, const u8* sourceMem, u64 sourceSize, const tchar* hint);
 
 	struct FileSender
