[UBA]
* Changed so local paths are always stored with forward slash

diff --git a/Engine/Source/Programs/UnrealBuildAccelerator/Common/Private/UbaSession.cpp b/Engine/Source/Programs/UnrealBuildAccelerator/Common/Private/UbaSession.cpp
--- a/Engine/Source/Programs/UnrealBuildAccelerator/Common/Private/UbaSession.cpp
+++ b/Engine/Source/Programs/UnrealBuildAccelerator/Common/Private/UbaSession.cpp
@@ -1610,8 +1610,7 @@ namespace uba
 					return;
 				}
 				hasRoot = true;
-				u32 localsIndex = (*str - RootPaths::RootStartByte)/PathsPerRoot;
-				auto& path = rootsEntry->locals[localsIndex];
+				auto& path = rootsEntry->locals[(*str - RootPaths::RootStartByte)/PathsPerRoot];
 				newSize += path.size() - strLen;
 			};
 
@@ -1633,10 +1632,22 @@ namespace uba
 					newStringPos += strLen;
 					return;
 				}
-				u32 localsIndex = (*str - RootPaths::RootStartByte)/PathsPerRoot;
+				u32 rootIndex = (*str - RootPaths::RootStartByte);
+				u32 localsIndex = rootIndex/PathsPerRoot;
 				auto& path = rootsEntry->locals[localsIndex];
-				memcpy(newStringPos, path.data(), path.size()*sizeof(tchar));
+				tchar* start = newStringPos;
 				newStringPos += path.size();
+				memcpy(start, path.data(), path.size()*sizeof(tchar));
+
+				#if PLATFORM_WINDOWS
+				u32 type = rootIndex % RootPaths::RootStartByte;
+				if (type != 1) // Not backslash
+				{
+					UBA_ASSERTF(type == 0, TC("Not implemented. (%s)"), hint); // We don't like double backslash or escaped space
+					*newStringPos = 0;
+					Replace(start, '\\', '/');
+				}
+				#endif
 			};
 
 		rootsEntry->roots.NormalizeString<tchar>(m_logger, inOut.data(), inOut.size(), handleString, false, hint);
