[UBA]
* Removed architectures option in Deploy-Windows.bat to prevent build of arm binaries since most people don't have that toolchain installed

diff --git a/Engine/Source/Programs/UnrealBuildAccelerator/Deploy-Windows.bat b/Engine/Source/Programs/UnrealBuildAccelerator/Deploy-Windows.bat
--- a/Engine/Source/Programs/UnrealBuildAccelerator/Deploy-Windows.bat
+++ b/Engine/Source/Programs/UnrealBuildAccelerator/Deploy-Windows.bat
@@ -14,7 +14,7 @@ set config=%1
 echo.
 echo === Building Targets ===
 pushd "../../../.."
-call Engine/Build/BatchFiles/RunUBT.bat -NoUba -NoUbaLocal -NoSNDBS -NoXGE -Architectures=x64+arm64 ^
+call Engine/Build/BatchFiles/RunUBT.bat -NoUba -NoUbaLocal -NoSNDBS -NoXGE ^
 	-Target="UbaAgent Win64 %config%" ^
 	-Target="UbaCli Win64 %config%" ^
 	-Target="UbaDetours Win64 %config%" ^
