[UBA]
* Cleaned up CreateProcess/posix_spawn code to work well with vfs.
* Fixed vfs logic for some posix detouring paths

diff --git a/Engine/Source/Programs/UnrealBuildAccelerator/Common/Private/UbaProcess.cpp b/Engine/Source/Programs/UnrealBuildAccelerator/Common/Private/UbaProcess.cpp
--- a/Engine/Source/Programs/UnrealBuildAccelerator/Common/Private/UbaProcess.cpp
+++ b/Engine/Source/Programs/UnrealBuildAccelerator/Common/Private/UbaProcess.cpp
@@ -963,7 +963,7 @@ namespace uba
 		#endif
 
 		TString applicationStr = reader.ReadString();
-		TString fullCommandLine = ReadLongString(reader);
+		TString commandLineWithoutApplication = ReadLongString(reader);
 
 		StringBuffer<> currentDir;
 		reader.ReadString(currentDir);
@@ -972,45 +972,10 @@ namespace uba
 		bool startSuspended = reader.ReadBool();
 		bool isChild = reader.ReadBool();
 
-		const tchar* commandLine = nullptr;
-
-		// Remove application from command line
-		if (fullCommandLine[0] == '"')
-		{
-			const tchar* second = TStrchr(fullCommandLine.data() + 1, '"');
-			if (!second)
-				second = fullCommandLine.data() + fullCommandLine.size();
-			//UBA_ASSERTF(second, TC("Missing second '\"' in command line: %s"), fullCommandLine.data); // "Unsupported cmd line format"
-			commandLine = second + 1;
-			if (applicationStr.empty())
-				applicationStr.assign(fullCommandLine.data() + 1, u64(second - fullCommandLine.data() - 1));
-		}
-		else
-		{
-			u32 searchOffset = 1;
-			if (applicationStr.empty() && IsWindows)
-			{
-				const tchar* extPos = nullptr;
-				Contains(fullCommandLine.data(),TC(".exe"), true, &extPos);
-				UBA_ASSERT(extPos);
-				searchOffset = u32(extPos - fullCommandLine.data());
-			}
-			const tchar* secondParamStart = TStrchr(fullCommandLine.data() + searchOffset, ' ');
-			if (!secondParamStart)
-				commandLine = TC("");
-			else
-				commandLine = secondParamStart + 1;
-			if (applicationStr.empty())
-				applicationStr.assign(fullCommandLine.data(), u64(secondParamStart - fullCommandLine.data()));
-		}
-
-		while (*commandLine == ' ')
-			++commandLine;
-
 		StringBuffer<> temp;
 		ProcessStartInfo info;
 		info.application = applicationStr.c_str();
-		info.arguments = commandLine;
+		info.arguments = commandLineWithoutApplication.c_str();
 		info.workingDir = currentDir.data;
 		info.logFile = InternalGetChildLogFile(temp);
 		info.priorityClass = m_startInfo.priorityClass;
@@ -1048,13 +1013,8 @@ namespace uba
 
 		writer.WriteString(m_realWorkingDir);
 
-		commandLine = process.m_startInfo.arguments;
-		#if PLATFORM_WINDOWS
-		TString realCommandLine = TC("\"") + process.m_realApplication + TC("\" ") + commandLine;
-		writer.WriteString(realCommandLine);
-		#else
 		writer.WriteString(process.m_realApplication);
-		writer.WriteString(commandLine);
+		#if !PLATFORM_WINDOWS
 		writer.WriteU64(process.m_comMemory.handle.uid);
 		writer.WriteU32(process.m_comMemory.offset);
 		writer.WriteString(info.logFile);
@@ -1469,7 +1429,10 @@ namespace uba
 			if (!*detoursLib)
 				detoursLib = UBA_DETOURS_LIBRARY_ANSI;
 
-			TString commandLine = TC("\"") + m_realApplication + TC("\" ") + m_startInfo.arguments;
+			StringBuffer<> application(m_realApplication);
+			m_session.VirtualizePath(application, m_startInfo.rootsHandle);
+
+			TString commandLine = TString(TC("\"")) + application.data + TC("\" ") + m_startInfo.arguments;
 
 			if (m_extractExports)
 			{
@@ -1551,7 +1514,7 @@ namespace uba
 
 				if (m_detourEnabled)
 				{
-					if (DetourCreateProcessWithDlls(NULL, (tchar*)commandLine.c_str(), NULL, NULL, inheritHandles, creationFlags, environment, workingDir, &si, &processInfo, sizeof_array(dlls), dlls, NULL))
+					if (DetourCreateProcessWithDlls(m_realApplication.c_str(), (tchar*)commandLine.c_str(), NULL, NULL, inheritHandles, creationFlags, environment, workingDir, &si, &processInfo, sizeof_array(dlls), dlls, NULL))
 					{
 						DisableProcessPowerThrottling(processInfo.hProcess);
 						break;
@@ -1789,7 +1752,11 @@ namespace uba
 			}
 			Vector<const char*> arguments2;
 			arguments2.reserve(arguments.size() + 2);
-			arguments2.push_back(m_startInfo.application);
+			{
+				StringBuffer<512> application(m_startInfo.application);
+				m_session.VirtualizePath(application, m_startInfo.rootsHandle);
+				arguments2.push_back(application.data);
+			}
 			for (auto& s : arguments)
 				arguments2.push_back(s.data());
 			arguments2.push_back(nullptr);
diff --git a/Engine/Source/Programs/UnrealBuildAccelerator/Common/Private/UbaSession.cpp b/Engine/Source/Programs/UnrealBuildAccelerator/Common/Private/UbaSession.cpp
--- a/Engine/Source/Programs/UnrealBuildAccelerator/Common/Private/UbaSession.cpp
+++ b/Engine/Source/Programs/UnrealBuildAccelerator/Common/Private/UbaSession.cpp
@@ -1561,26 +1561,43 @@ namespace uba
 		return ++m_processIdCounter;
 	}
 
-	bool Session::DevirtualizePath(StringBufferBase& inOut, RootsHandle rootsHandle)
+	bool Session::VirtualizePath(StringBufferBase& inOut, RootsHandle rootsHandle)
 	{
 		if (!rootsHandle)
 			return true;
-		
-#if PLATFORM_WINDOWS
-		if (inOut[1] != ':')
+		if (IsWindows ? inOut[1] != ':' : inOut[0] != '/')
 			return true;
-#else
-		if (inOut[0] != '/')
+		auto rootsEntry = GetRootsEntry(rootsHandle);
+		if (!rootsEntry)
+			return false;
+		if (rootsEntry->roots.IsEmpty())
 			return true;
-#endif
+		auto& locals = rootsEntry->locals;
+		auto& vfs = rootsEntry->vfs;
+		for (u32 i=0, e=u32(locals.size()); i!=e; ++i)
+		{
+			if (!inOut.StartsWith(locals[i].c_str()))
+				continue;
+			StringBuffer<> temp;
+			temp.Append(inOut.data + locals[i].size());
+			inOut.Clear().Append(vfs[i]).Append(temp);
+			return true;
+		}
+
+		return true;
+	}
 
+	bool Session::DevirtualizePath(StringBufferBase& inOut, RootsHandle rootsHandle)
+	{
+		if (!rootsHandle)
+			return true;
+		if (IsWindows ? inOut[1] != ':' : inOut[0] != '/')
+			return true;
 		auto rootsEntry = GetRootsEntry(rootsHandle);
 		if (!rootsEntry)
 			return false;
-
 		if (rootsEntry->roots.IsEmpty())
 			return true;
-
 		auto root = rootsEntry->roots.FindRoot(inOut);
 		if (!root)
 			return m_logger.Error(TC("Can't find root for path %s"), inOut.data);
@@ -3194,6 +3211,10 @@ namespace uba
 			if (!temp.count) // If vfs is not set it means that vfs should not be used (the roots entry memory might be used for cacheclient though)
 				break;
 			entry.roots.RegisterRoot(m_logger, temp.data);
+			entry.vfs.emplace_back(temp.data, temp.count);
+			#if PLATFORM_WINDOWS
+			Replace((tchar*)entry.vfs.data(), '/', '\\');
+			#endif
 			reader.ReadString(temp.Clear());
 			entry.locals.emplace_back(temp.data, temp.count);
 		}
diff --git a/Engine/Source/Programs/UnrealBuildAccelerator/Common/Public/UbaSession.h b/Engine/Source/Programs/UnrealBuildAccelerator/Common/Public/UbaSession.h
--- a/Engine/Source/Programs/UnrealBuildAccelerator/Common/Public/UbaSession.h
+++ b/Engine/Source/Programs/UnrealBuildAccelerator/Common/Public/UbaSession.h
@@ -84,6 +84,7 @@ namespace uba
 
 		u32 CreateProcessId();
 
+		bool VirtualizePath(StringBufferBase& inOut, RootsHandle rootsHandle);
 		bool DevirtualizePath(StringBufferBase& inOut, RootsHandle rootsHandle);
 		bool DevirtualizeString(TString& inOut, RootsHandle rootsHandle, const tchar* hint = TC(""));
 		bool PopulateLocalToIndexRoots(RootPaths& out, RootsHandle rootsHandle);
@@ -169,7 +170,7 @@ namespace uba
 
 		bool HasVfs(RootsHandle handle) const;
 		RootsHandle WithVfs(RootsHandle key, bool vfs) const;
-		struct RootsEntry { Vector<u8> memory; RootPaths roots; Vector<TString> locals; };
+		struct RootsEntry { Vector<u8> memory; RootPaths roots; Vector<TString> locals; Vector<TString> vfs; };
 		const RootsEntry* GetRootsEntry(RootsHandle rootsHandle);
 		void PopulateRootsEntry(RootsEntry& entry, const void* rootsData, uba::u64 rootsDataSize);
 
diff --git a/Engine/Source/Programs/UnrealBuildAccelerator/Detours/Private/PosixOS/UbaDetoursFunctionsPosix.cpp b/Engine/Source/Programs/UnrealBuildAccelerator/Detours/Private/PosixOS/UbaDetoursFunctionsPosix.cpp
--- a/Engine/Source/Programs/UnrealBuildAccelerator/Detours/Private/PosixOS/UbaDetoursFunctionsPosix.cpp
+++ b/Engine/Source/Programs/UnrealBuildAccelerator/Detours/Private/PosixOS/UbaDetoursFunctionsPosix.cpp
@@ -723,7 +723,10 @@ int Shared_stat(const char* funcName, const char* file, struct stat* attr, const
 UBA_EXPORT int UBA_WRAPPER(_NSGetExecutablePath)(char* buf, uint32_t* bufsize)
 {
 	if (!g_isDetouring)
+	{
+		DEBUG_LOG_TRUE("NSGetExecutablePath", "");
 		return TRUE_WRAPPER(_NSGetExecutablePath)(buf, bufsize);
+	}
 	if (bufsize == nullptr)
 		return -1;
 	const uint32_t requiredBufsize = g_virtualApplication.count + 1;
@@ -733,6 +736,7 @@ UBA_EXPORT int UBA_WRAPPER(_NSGetExecutablePath)(char* buf, uint32_t* bufsize)
 		return -1;
 	if (buf != nullptr)
 		memcpy(buf, g_virtualApplication.data, requiredBufsize);
+	DEBUG_LOG_DETOURED("NSGetExecutablePath", "%s", buf);
 	return 0;
 }
 
@@ -903,7 +907,7 @@ UBA_EXPORT char* UBA_WRAPPER(realpath)(const char* path, char* resolved_path)
 {
 	UBA_INIT_DETOUR(realpath, path, resolved_path);
 
-	if (!g_runningRemote)
+	if (!g_runningRemote && !IsVfsEnabled())
 	{
 		auto res = TRUE_WRAPPER(realpath)(path, resolved_path);
 		DEBUG_LOG_TRUE("realpath", "(%s) -> %s", path, res);
@@ -1374,7 +1378,6 @@ UBA_EXPORT int UBA_WRAPPER(access)(const char* pathname, int mode)
 		DEBUG_LOG_TRUE("access", "%s %i -> %i (%s)", pathname, mode, res, StrError(res, errno));
 		return res;
 	}
-	DevirtualizePath(fixedPath);
 
 	bool checkIfDir = false;
 	StringBuffer<> temp;
@@ -1390,6 +1393,8 @@ UBA_EXPORT int UBA_WRAPPER(access)(const char* pathname, int mode)
 		fixedPath.Clear().Append(temp);
 	}
 
+	DevirtualizePath(fixedPath);
+
 	if (!CanDetour(pathname))
 	{
 		auto res = TRUE_WRAPPER(access)(pathname, mode);
@@ -1899,22 +1904,24 @@ int shared_posix_spawn(pid_t* pid, const char* path, const posix_spawn_file_acti
 		tempArgv[argc3] = 0;
 		argv = (char*const*)tempArgv;
 	}
-
+	else if (IsVfsEnabled() && (!path || !*path))
+	{
+		FixPath(result, argv[0]);
+		DevirtualizePath(result);
+		path = result.data;
+	}
 	//DEBUG_LOG("LIBRARY_SEARCH_PATHS: %s", getenv("LIBRARY_SEARCH_PATHS"));
 	//DEBUG_LOG("PATH: %s", getenv("PATH"));
 
-	TString cmdLine;
-	u32 argc = 0;
-	for (u32 i = 0; argv[i]; ++i)
+	TString cmdLineWithoutApplication;
+	for (u32 i = 1; argv[i]; ++i)
 	{
 		if (i != 0)
-			cmdLine.append(" ");
-		cmdLine.append(argv[i]);
-		++argc;
+			cmdLineWithoutApplication.append(" ");
+		cmdLineWithoutApplication.append(argv[i]);
 	}
 
 	TString realApplication;
-	TString commandLine;
 	u32 processId = 0;
 	StringBuffer<512> currentDir;
 	StringBuffer<256> comIdVar;
@@ -1939,7 +1946,7 @@ int shared_posix_spawn(pid_t* pid, const char* path, const posix_spawn_file_acti
 		BinaryWriter writer;
 		writer.WriteByte(MessageType_CreateProcess);
 		writer.WriteString(path); // Application
-		WriteLongString(writer, cmdLine); // Full command line
+		WriteLongString(writer, cmdLineWithoutApplication);
 		writer.WriteString(pwd); // Current dir
 		writer.WriteBool(false); // Start suspended
 		writer.WriteBool(true); // Is child
@@ -1963,7 +1970,6 @@ int shared_posix_spawn(pid_t* pid, const char* path, const posix_spawn_file_acti
 		reader.ReadString(currentDir);
 
 		realApplication = reader.ReadString();
-		commandLine = reader.ReadString();
 
 		comIdVar.Append("UBA_COMID=").AppendValue(reader.ReadU64()).Append('+').AppendValue(reader.ReadU32());
 
@@ -2522,6 +2528,10 @@ namespace uba
 		VirtualizePath(g_virtualWorkingDir);
 		VirtualizePath(g_exeDir);
 
+		DEBUG_LOG("Application: %s", g_virtualApplication.data);
+		DEBUG_LOG("WorkingDir: %s", g_virtualWorkingDir.data);
+		DEBUG_LOG("ExeDir: %s", g_exeDir.data);
+
 		UBA_ASSERTF(g_virtualApplicationDir.capacity > 0, "g_virtualApplicationDir.capacity > 0");
 
 		const char* lastSlash = strrchr(g_virtualApplication.data, '/');
diff --git a/Engine/Source/Programs/UnrealBuildAccelerator/Detours/Private/UbaDetoursShared.cpp b/Engine/Source/Programs/UnrealBuildAccelerator/Detours/Private/UbaDetoursShared.cpp
--- a/Engine/Source/Programs/UnrealBuildAccelerator/Detours/Private/UbaDetoursShared.cpp
+++ b/Engine/Source/Programs/UnrealBuildAccelerator/Detours/Private/UbaDetoursShared.cpp
@@ -183,10 +183,15 @@ namespace uba
 				continue;
 			}
 
+			#if PLATFORM_WINDOWS
+			str.Replace('/', '\\');
+			#endif
+
 			u32 index = g_vfsEntryCount++;
 			UBA_ASSERT(index < sizeof_array(g_vfsEntries));
 			VfsEntry& vfsEntry = g_vfsEntries[index];
 			vfsEntry.vfs = g_memoryBlock.Strdup(str);
+
 			if (index == 0)
 				g_vfsMatchingLength = vfsEntry.vfs.count;
 			else
@@ -211,19 +216,19 @@ namespace uba
 		return g_vfsEntryCount > 0;
 	}
 
-	void DevirtualizePath(StringBufferBase& path)
+	bool DevirtualizePath(StringBufferBase& path)
 	{
 		if (!g_vfsEntryCount)
-			return;
+			return false;
 
 		if (!Equals(path.data, g_vfsEntries[0].vfs.data, Min(path.count, g_vfsMatchingLength), CaseInsensitiveFs))
-			return;
+			return false;
 
 		// TODO: This is not great, the dirs above the vfs root should be empty except the dir to the roots
 		if (path.count < g_vfsMatchingLength)
 		{
 			path.Clear().Append(g_vfsEntries[0].local);
-			return;
+			return true;
 		}
 
 		for (u32 i=0, e=g_vfsEntryCount; i!=e; ++i)
@@ -233,14 +238,15 @@ namespace uba
 				continue;
 			StringBuffer<> temp2(path.data + entry.vfs.count);
 			path.Clear().Append(entry.local).Append(temp2);
-			return;
+			return true;
 		}
+		return false;
 	}
 
-	void VirtualizePath(StringBufferBase& path)
+	bool VirtualizePath(StringBufferBase& path)
 	{
 		if (!g_vfsEntryCount)
-			return;
+			return false;
 		for (u32 i=0, e=g_vfsEntryCount; i!=e; ++i)
 		{
 			VfsEntry& entry = g_vfsEntries[i];
@@ -248,8 +254,9 @@ namespace uba
 				continue;
 			StringBuffer<> temp2(path.data + entry.local.count);
 			path.Clear().Append(entry.vfs).Append(temp2);
-			return;
+			return true;
 		}
+		return false;
 	}
 
 	void LogVfsInfo()
diff --git a/Engine/Source/Programs/UnrealBuildAccelerator/Detours/Private/UbaDetoursShared.h b/Engine/Source/Programs/UnrealBuildAccelerator/Detours/Private/UbaDetoursShared.h
--- a/Engine/Source/Programs/UnrealBuildAccelerator/Detours/Private/UbaDetoursShared.h
+++ b/Engine/Source/Programs/UnrealBuildAccelerator/Detours/Private/UbaDetoursShared.h
@@ -25,8 +25,8 @@ namespace uba
 	
 	void PopulateVfs(BinaryReader& vfsReader);
 	bool IsVfsEnabled();
-	void DevirtualizePath(StringBufferBase& path);
-	void VirtualizePath(StringBufferBase& path);
+	bool DevirtualizePath(StringBufferBase& path);
+	bool VirtualizePath(StringBufferBase& path);
 	void LogVfsInfo();
 
 	#if UBA_DEBUG_LOG_ENABLED
diff --git a/Engine/Source/Programs/UnrealBuildAccelerator/Detours/Private/Windows/UbaDetoursFunctionsKernelBase.inl b/Engine/Source/Programs/UnrealBuildAccelerator/Detours/Private/Windows/UbaDetoursFunctionsKernelBase.inl
--- a/Engine/Source/Programs/UnrealBuildAccelerator/Detours/Private/Windows/UbaDetoursFunctionsKernelBase.inl
+++ b/Engine/Source/Programs/UnrealBuildAccelerator/Detours/Private/Windows/UbaDetoursFunctionsKernelBase.inl
@@ -2453,8 +2453,22 @@ DWORD Detoured_SearchPathW(LPCWSTR lpPath, LPCWSTR lpFileName, LPCWSTR lpExtensi
 		return DWORD(pathLen);
 	}
 
-	DEBUG_LOG_TRUE(L"SearchPathW", L"%ls %ls", lpPath, lpFileName);
-	return True_SearchPathW(lpPath, lpFileName, lpExtension, nBufferLength, lpBuffer, lpFilePart);
+	DWORD res = True_SearchPathW(lpPath, lpFileName, lpExtension, nBufferLength, lpBuffer, lpFilePart);
+	if (res && IsVfsEnabled())
+	{
+		UBA_ASSERT(res < nBufferLength);
+		StringBuffer<> temp(lpBuffer);
+		if (VirtualizePath(temp))
+		{
+			memcpy(lpBuffer, temp.data, sizeof(tchar)*temp.count+1);
+			res = temp.count;
+			if (lpFilePart)
+				*lpFilePart = TStrrchr(lpBuffer, '\\') + 1;
+		}
+
+	}
+	DEBUG_LOG_TRUE(L"SearchPathW", L"%ls %ls -> %s", lpPath, lpFileName, lpBuffer);
+	return res;
 }
 
 using AdditionalLoads = Vector<HMODULE, GrowingAllocator<HMODULE>>;
@@ -2749,18 +2763,38 @@ BOOL Detoured_CreateProcessW(LPCWSTR lpApplicationName, LPWSTR lpCommandLine, LP
 		isChild = false;
 	}
 
-	StringBuffer<> application;
+	StringBuffer<512> application;
+	StringView commandLineWithoutApplication;
+
+	{
+		const tchar* endOfApplication;
+		if (originalCmd.data[0] == '"')
+			endOfApplication = TStrchr(originalCmd.data+1, '"') + 1;
+		else
+			endOfApplication = TStrchr(originalCmd.data, ' ') + 1;
+		commandLineWithoutApplication = StringView(endOfApplication, originalCmd.count - u32(endOfApplication - originalCmd.data));
+	}
+
 	if (lpApplicationName)
 	{
 		FixPath(application, lpApplicationName);
-		DevirtualizePath(application);
 	}
+	else
+	{
+		StringBuffer<512> temp;
+		if (originalCmd.data[0] == '"')
+			temp.Append(originalCmd.data+1, commandLineWithoutApplication.data - originalCmd.data - 2);
+		else
+			temp.Append(originalCmd.data, commandLineWithoutApplication.data - originalCmd.data - 1);
+		FixPath(application, temp.data);
+	}
+	DevirtualizePath(application);
+
 	bool startSuspended = (dwCreationFlags & CREATE_SUSPENDED) != 0;
 
 	StringBuffer<> workingDir(lpCurrentDirectory ? lpCurrentDirectory : g_virtualWorkingDir.data);
 	DevirtualizePath(workingDir);
 
-	TString commandLine;
 	TString currentDir;
 	u32 processId = 0;
 	char dll[1024];
@@ -2770,7 +2804,7 @@ BOOL Detoured_CreateProcessW(LPCWSTR lpApplicationName, LPWSTR lpCommandLine, LP
 		BinaryWriter writer;
 		writer.WriteByte(MessageType_CreateProcess);
 		writer.WriteString(application);
-		WriteLongString(writer, originalCmd);
+		WriteLongString(writer, commandLineWithoutApplication);
 		writer.WriteString(workingDir);
 		writer.WriteBool(startSuspended);
 		writer.WriteBool(isChild);
@@ -2788,7 +2822,6 @@ BOOL Detoured_CreateProcessW(LPCWSTR lpApplicationName, LPWSTR lpCommandLine, LP
 		dll[dllNameSize] = 0;
 
 		currentDir = reader.ReadString();
-		commandLine = reader.ReadString();
 		DEBUG_LOG_PIPE(L"CreateProcess", L"%ls %ls", application.data, originalCmd.data);
 	}
 
@@ -2841,7 +2874,7 @@ BOOL Detoured_CreateProcessW(LPCWSTR lpApplicationName, LPWSTR lpCommandLine, LP
 	while (true)
 	{
 		res = true;
-		if (DetourCreateProcessWithDllsW(NULL, (wchar_t*)commandLine.c_str(), NULL, NULL, bInheritHandles, dwCreationFlags, lpEnvironment, currentDir.c_str(), lpStartupInfo, lpProcessInformation, sizeof_array(dlls), dlls, True_CreateProcessW))
+		if (DetourCreateProcessWithDllsW(application.data, lpCommandLine, NULL, NULL, bInheritHandles, dwCreationFlags, lpEnvironment, currentDir.c_str(), lpStartupInfo, lpProcessInformation, sizeof_array(dlls), dlls, True_CreateProcessW))
 			break;
 		res = false;
 		lastError = GetLastError();
@@ -2853,7 +2886,7 @@ BOOL Detoured_CreateProcessW(LPCWSTR lpApplicationName, LPWSTR lpCommandLine, LP
 		if (retryCount++ >= 5)
 			break;
 		const wchar_t* errorText = lastError == ERROR_ACCESS_DENIED ? L"access denied" : L"internal error";
-		Rpc_WriteLogf(L"DetourCreateProcessWithDllEx failed with %ls, retrying %ls (Working dir: %ls)", errorText, commandLine.c_str(), currentDir.c_str());
+		Rpc_WriteLogf(L"DetourCreateProcessWithDllEx failed with %ls, retrying %ls (Working dir: %ls)", errorText, originalCmd.data, currentDir.c_str());
 		Sleep(100 + (rand() % 200)); // We have no idea
 		continue;
 	}
@@ -2873,10 +2906,10 @@ BOOL Detoured_CreateProcessW(LPCWSTR lpApplicationName, LPWSTR lpCommandLine, LP
 		writer.WriteU32(lpProcessInformation->dwProcessId);
 		writer.WriteU64(u64(lpProcessInformation->hThread));
 		writer.Flush();
-		DEBUG_LOG_PIPE(L"StartProcess", L"%ls %ls", lpApplicationName, originalCmd);
+		DEBUG_LOG_PIPE(L"StartProcess", L"%ls %ls", application.data, originalCmd.data);
 	}
 
-	UBA_ASSERTF(res, L"Failed to spawn process %ls (Error code: %u)", commandLine.c_str(), lastError);
+	UBA_ASSERTF(res, L"Failed to spawn process %ls (Error code: %u)", originalCmd.data, lastError);
 
 	HANDLE trueHandle = lpProcessInformation->hProcess;
 
diff --git a/Engine/Source/Programs/UnrealBuildAccelerator/Detours/Private/Windows/UbaDetoursObjFilesPreloader.h b/Engine/Source/Programs/UnrealBuildAccelerator/Detours/Private/Windows/UbaDetoursObjFilesPreloader.h
--- a/Engine/Source/Programs/UnrealBuildAccelerator/Detours/Private/Windows/UbaDetoursObjFilesPreloader.h
+++ b/Engine/Source/Programs/UnrealBuildAccelerator/Detours/Private/Windows/UbaDetoursObjFilesPreloader.h
@@ -13,7 +13,7 @@
 namespace uba
 {
 	constexpr bool UseLargePages = false;
-	void DevirtualizePath(StringBufferBase& path);
+	bool DevirtualizePath(StringBufferBase& path);
 
 	u64 GetLargePageSize()
 	{
