[UBA]
* Fixed posix_spawn child process logic when path argument is set

diff --git a/Engine/Source/Programs/UnrealBuildAccelerator/Detours/Private/PosixOS/UbaDetoursFunctionsPosix.cpp b/Engine/Source/Programs/UnrealBuildAccelerator/Detours/Private/PosixOS/UbaDetoursFunctionsPosix.cpp
--- a/Engine/Source/Programs/UnrealBuildAccelerator/Detours/Private/PosixOS/UbaDetoursFunctionsPosix.cpp
+++ b/Engine/Source/Programs/UnrealBuildAccelerator/Detours/Private/PosixOS/UbaDetoursFunctionsPosix.cpp
@@ -960,12 +960,14 @@ UBA_EXPORT ssize_t UBA_WRAPPER(readlink)(const char* pathname, char* buf, size_t
 			FileInfo& info = *fo->fileInfo;
 			u32 len = TStrlen(info.originalName);
 			UBA_ASSERTF(len < bufsiz, "buffer size is smaller than length of name");
+			UBA_ASSERT(!IsVfsEnabled());
 			memcpy(buf, info.originalName, len + 1);
 			DEBUG_LOG_DETOURED("readlink", "(%s) (%s) -> %u", pathname, buf, len);
 			return len;
 		}
 	}
 	auto res = TRUE_WRAPPER(readlink)(pathname, buf, bufsiz);
+	UBA_ASSERT(!IsVfsEnabled());
 	DEBUG_LOG_TRUE("readlink", "(%s) (%s) -> %llu", pathname, buf, u64(res));
 	return res;
 }
@@ -1904,9 +1906,13 @@ int shared_posix_spawn(pid_t* pid, const char* path, const posix_spawn_file_acti
 		tempArgv[argc3] = 0;
 		argv = (char*const*)tempArgv;
 	}
-	else if (IsVfsEnabled() && (!path || !*path))
+	else if (IsVfsEnabled())
 	{
-		FixPath(result, argv[0]);
+		if (!path || !*path)
+			FixPath(result, argv[0]);
+		else
+			FixPath(result, path);
+
 		DevirtualizePath(result);
 		path = result.data;
 	}
