[UBA]
* Added asserts in strategical places to catch unsupported vfs paths

diff --git a/Engine/Source/Programs/UnrealBuildAccelerator/Detours/Private/PosixOS/UbaDetoursFunctionsPosix.cpp b/Engine/Source/Programs/UnrealBuildAccelerator/Detours/Private/PosixOS/UbaDetoursFunctionsPosix.cpp
--- a/Engine/Source/Programs/UnrealBuildAccelerator/Detours/Private/PosixOS/UbaDetoursFunctionsPosix.cpp
+++ b/Engine/Source/Programs/UnrealBuildAccelerator/Detours/Private/PosixOS/UbaDetoursFunctionsPosix.cpp
@@ -960,15 +960,25 @@ UBA_EXPORT ssize_t UBA_WRAPPER(readlink)(const char* pathname, char* buf, size_t
 			FileInfo& info = *fo->fileInfo;
 			u32 len = TStrlen(info.originalName);
 			UBA_ASSERTF(len < bufsiz, "buffer size is smaller than length of name");
-			UBA_ASSERT(!IsVfsEnabled());
 			memcpy(buf, info.originalName, len + 1);
 			DEBUG_LOG_DETOURED("readlink", "(%s) (%s) -> %u", pathname, buf, len);
 			return len;
 		}
 	}
+	else
+	{
+		UBA_ASSERTF(!StartsWith(pathname, "/UEVFS"), "Need to devirtualize %s", pathname);
+	}
+
 	auto res = TRUE_WRAPPER(readlink)(pathname, buf, bufsiz);
-	UBA_ASSERT(!IsVfsEnabled());
 	DEBUG_LOG_TRUE("readlink", "(%s) (%s) -> %llu", pathname, buf, u64(res));
+
+	if (res && res < bufsiz && IsVfsEnabled())
+	{
+		StringBuffer<> temp(buf);
+		UBA_ASSERT(!DevirtualizePath(temp));
+	}
+
 	return res;
 }
 
