[UBA]
* Fixed bug in CreateProcess detour code where wrong "real application" ended up being used on remote sessions

diff --git a/Engine/Source/Programs/UnrealBuildAccelerator/Detours/Private/Windows/UbaDetoursFunctionsKernelBase.inl b/Engine/Source/Programs/UnrealBuildAccelerator/Detours/Private/Windows/UbaDetoursFunctionsKernelBase.inl
--- a/Engine/Source/Programs/UnrealBuildAccelerator/Detours/Private/Windows/UbaDetoursFunctionsKernelBase.inl
+++ b/Engine/Source/Programs/UnrealBuildAccelerator/Detours/Private/Windows/UbaDetoursFunctionsKernelBase.inl
@@ -2822,6 +2822,7 @@ BOOL Detoured_CreateProcessW(LPCWSTR lpApplicationName, LPWSTR lpCommandLine, LP
 		dll[dllNameSize] = 0;
 
 		currentDir = reader.ReadString();
+		reader.ReadString(application.Clear());
 		DEBUG_LOG_PIPE(L"CreateProcess", L"%ls %ls", application.data, originalCmd.data);
 	}
 
