[UBA]
* Moved logic fetching rootpaths to be earlier since it can now be used to devirtualize process path

diff --git a/Engine/Source/Programs/UnrealBuildAccelerator/Common/Private/UbaProcess.cpp b/Engine/Source/Programs/UnrealBuildAccelerator/Common/Private/UbaProcess.cpp
--- a/Engine/Source/Programs/UnrealBuildAccelerator/Common/Private/UbaProcess.cpp
+++ b/Engine/Source/Programs/UnrealBuildAccelerator/Common/Private/UbaProcess.cpp
@@ -465,6 +465,12 @@ namespace uba
 		if (HandleSpecialApplication())
 			return;
 
+		if (!m_session.ProcessThreadStart(*this))
+		{
+			m_logLines.push_back({ TString(TC("Session::ProcessThreadStart failed")), LogEntryType_Error });
+			return;
+		}
+
 		u8* comMemory = m_comMemory.memory;
 		u64 comMemorySize = CommunicationMemSize;
 		#if !PLATFORM_WINDOWS
diff --git a/Engine/Source/Programs/UnrealBuildAccelerator/Common/Private/UbaSession.cpp b/Engine/Source/Programs/UnrealBuildAccelerator/Common/Private/UbaSession.cpp
--- a/Engine/Source/Programs/UnrealBuildAccelerator/Common/Private/UbaSession.cpp
+++ b/Engine/Source/Programs/UnrealBuildAccelerator/Common/Private/UbaSession.cpp
@@ -1758,6 +1758,11 @@ namespace uba
 		lock.Leave();
 	}
 
+	bool Session::ProcessThreadStart(ProcessImpl& process)
+	{
+		return true;
+	}
+
 	bool Session::GetInitResponse(InitResponse& out, const InitMessage& msg)
 	{
 		out.directoryTableHandle = m_directoryTableHandle.ToU64();
diff --git a/Engine/Source/Programs/UnrealBuildAccelerator/Common/Private/UbaSessionClient.cpp b/Engine/Source/Programs/UnrealBuildAccelerator/Common/Private/UbaSessionClient.cpp
--- a/Engine/Source/Programs/UnrealBuildAccelerator/Common/Private/UbaSessionClient.cpp
+++ b/Engine/Source/Programs/UnrealBuildAccelerator/Common/Private/UbaSessionClient.cpp
@@ -420,11 +420,11 @@ namespace uba
 		return m_storage.CopyOrLink(casKey, out.data, fileAttributes);
 	}
 
-	bool SessionClient::GetInitResponse(InitResponse& out, const InitMessage& msg)
+	bool SessionClient::ProcessThreadStart(ProcessImpl& process)
 	{
-		if (!Session::GetInitResponse(out, msg))
+		if (!Session::ProcessThreadStart(process))
 			return false;
-		if (RootsHandle rootsHandle = msg.process.GetStartInfo().rootsHandle)
+		if (RootsHandle rootsHandle = process.GetStartInfo().rootsHandle)
 			if (!SendRootsHandle(rootsHandle))
 				return false;
 		return true;
diff --git a/Engine/Source/Programs/UnrealBuildAccelerator/Common/Public/UbaSession.h b/Engine/Source/Programs/UnrealBuildAccelerator/Common/Public/UbaSession.h
--- a/Engine/Source/Programs/UnrealBuildAccelerator/Common/Public/UbaSession.h
+++ b/Engine/Source/Programs/UnrealBuildAccelerator/Common/Public/UbaSession.h
@@ -133,6 +133,7 @@ namespace uba
 
 		virtual bool RegisterVirtualFileInternal(const StringKey& fileNameKey, const StringView& filePath, const tchar* sourceFile, u64 sourceOffset, u64 sourceSize);
 
+		virtual bool ProcessThreadStart(ProcessImpl& process);
 		virtual bool GetInitResponse(InitResponse& out, const InitMessage& msg);
 		virtual bool CreateFile(CreateFileResponse& out, const CreateFileMessage& msg);
 		virtual bool CloseFile(CloseFileResponse& out, const CloseFileMessage& msg);
diff --git a/Engine/Source/Programs/UnrealBuildAccelerator/Common/Public/UbaSessionClient.h b/Engine/Source/Programs/UnrealBuildAccelerator/Common/Public/UbaSessionClient.h
--- a/Engine/Source/Programs/UnrealBuildAccelerator/Common/Public/UbaSessionClient.h
+++ b/Engine/Source/Programs/UnrealBuildAccelerator/Common/Public/UbaSessionClient.h
@@ -50,7 +50,7 @@ namespace uba
 
 		virtual bool PrepareProcess(ProcessStartInfoHolder& startInfo, bool isChild, StringBufferBase& outRealApplication, const tchar*& outRealWorkingDir) override;
 		virtual void* GetProcessEnvironmentVariables() override;
-		virtual bool GetInitResponse(InitResponse& out, const InitMessage& msg) override;
+		virtual bool ProcessThreadStart(ProcessImpl& process);
 		virtual bool CreateFile(CreateFileResponse& out, const CreateFileMessage& msg) override;
 		virtual bool DeleteFile(DeleteFileResponse& out, const DeleteFileMessage& msg) override;
 		virtual bool CopyFile(CopyFileResponse& out, const CopyFileMessage& msg) override;
