[UBA]
* Reduced stack usage to fix mac remote session crash

diff --git a/Engine/Source/Programs/UnrealBuildAccelerator/Common/Private/UbaSessionClient.cpp b/Engine/Source/Programs/UnrealBuildAccelerator/Common/Private/UbaSessionClient.cpp
--- a/Engine/Source/Programs/UnrealBuildAccelerator/Common/Private/UbaSessionClient.cpp
+++ b/Engine/Source/Programs/UnrealBuildAccelerator/Common/Private/UbaSessionClient.cpp
@@ -336,9 +336,9 @@ namespace uba
 
 	bool SessionClient::ReadModules(List<ModuleInfo>& outModules, u32 processId, const tchar* application)
 	{
-		StackBinaryReader<SendMaxSize> reader;
+		StackBinaryReader<16*1024> reader;
 		{
-			StackBinaryWriter<1024> writer;
+			StackBinaryWriter<256> writer;
 			NetworkMessage msg(m_client, ServiceId, SessionMessageType_GetApplication, writer);
 			writer.WriteU32(processId);
 			writer.WriteString(application);
