[UBA]
* Fixed so linux linking works in uba again
* Fixed weird bug where realpath is returning null.. it is like it is suddenly using an older version of glibc?

diff --git a/Engine/Source/Programs/UnrealBuildAccelerator/Detours/Private/PosixOS/UbaDetoursFunctionsPosix.cpp b/Engine/Source/Programs/UnrealBuildAccelerator/Detours/Private/PosixOS/UbaDetoursFunctionsPosix.cpp
--- a/Engine/Source/Programs/UnrealBuildAccelerator/Detours/Private/PosixOS/UbaDetoursFunctionsPosix.cpp
+++ b/Engine/Source/Programs/UnrealBuildAccelerator/Detours/Private/PosixOS/UbaDetoursFunctionsPosix.cpp
@@ -139,6 +139,8 @@ using namespace uba;
 	DETOURED_FUNCTION(__fxstatat64) \
 	DETOURED_FUNCTION(open64) \
 	DETOURED_FUNCTION(fstatat64) \
+	DETOURED_FUNCTION(fpathconf) \
+	DETOURED_FUNCTION(pathconf) \
 
 #else
 #define DETOURED_FUNCTION_LINUX
@@ -873,6 +875,32 @@ UBA_EXPORT int UBA_WRAPPER(chroot)(const char* path)
 UBA_EXPORT char* UBA_WRAPPER(getcwd)(char* buf, size_t size)
 {
 	UBA_INIT_DETOUR(getcwd, buf, size);
+
+	if (size == 0)
+	{
+		DEBUG_LOG_DETOURED("getcwd", "-> null (Size 0)");
+
+		if (buf == nullptr)
+		{
+			UBA_ASSERT(size == 0);
+			size = g_virtualWorkingDir.count+1;
+			buf = (char*)malloc(size);
+			memcpy(buf, g_virtualWorkingDir.data, size);
+			return buf;
+		}
+
+		errno = EINVAL;
+		return nullptr;
+	}
+	if (size < g_virtualWorkingDir.count+1)
+	{
+		DEBUG_LOG_DETOURED("getcwd", "-> null (Buffer too small: %u)", u32(size));
+		errno = ERANGE;
+		return nullptr;
+	}
+	if (!buf)
+		buf = (char*)malloc(size);
+		
 	UBA_ASSERTF(g_virtualWorkingDir.count < size, "getcwd with size smaller than path not implemented");
 	memcpy(buf, g_virtualWorkingDir.data, g_virtualWorkingDir.count+1);
 	DEBUG_LOG_DETOURED("getcwd", "%s -> %p", buf, buf);
@@ -909,8 +937,10 @@ UBA_EXPORT char* UBA_WRAPPER(realpath)(const char* path, char* resolved_path)
 
 	if (!g_runningRemote && !IsVfsEnabled())
 	{
+		if (!resolved_path) // TODO: This is weird.. it seems like newly built uba uses old glibc?
+			resolved_path  = (char*)malloc(PATH_MAX);
 		auto res = TRUE_WRAPPER(realpath)(path, resolved_path);
-		DEBUG_LOG_TRUE("realpath", "(%s) -> %s", path, res);
+		DEBUG_LOG_TRUE("realpath", "(%s) -> %s (%s)", path, res, StrError(res == 0 ? -1 : 0, errno));
 		return res;
 	}
 
@@ -1355,8 +1385,25 @@ UBA_EXPORT int UBA_WRAPPER(fstatat64)(int dirfd, const char* pathname, struct st
 {
 	UBA_INIT_DETOUR(fstatat64, dirfd, pathname, buf, flags);
 	DEBUG_LOG_TRUE("fstatat64", "");
+	UBA_ASSERT(false);
 	return TRUE_WRAPPER(fstatat64)(dirfd, pathname, buf, flags);
 }
+
+UBA_EXPORT long UBA_WRAPPER(fpathconf)(int fd, int name)
+{
+	UBA_INIT_DETOUR(fpathconf, fd, name);
+	DEBUG_LOG_TRUE("fpathconf", "");
+	UBA_ASSERT(false);
+	return TRUE_WRAPPER(fpathconf)(fd, name);
+}
+UBA_EXPORT long UBA_WRAPPER(pathconf)(char *path, int name)
+{
+	UBA_INIT_DETOUR(pathconf, path, name);
+	DEBUG_LOG_TRUE("pathconf", "");
+	UBA_ASSERT(false);
+	return TRUE_WRAPPER(pathconf)(path, name);
+}
+
 #endif
 
 UBA_EXPORT int UBA_WRAPPER(lstat)(const char *path, struct stat *buf)
@@ -1916,7 +1963,7 @@ int shared_posix_spawn(pid_t* pid, const char* path, const posix_spawn_file_acti
 		tempArgv[argc3] = 0;
 		argv = (char*const*)tempArgv;
 	}
-	else if (IsVfsEnabled())
+	else
 	{
 		if (!path || !*path)
 			FixPath(result, argv[0]);
@@ -2039,7 +2086,7 @@ int shared_posix_spawn(pid_t* pid, const char* path, const posix_spawn_file_acti
 		writer.Flush();
 	}
 
-	DEBUG_LOG("         Child process started %s -> %i", path, res);
+	DEBUG_LOG("         Child process started %s -> %i (pid: %i)", path, res, *pid);
 
 	return res;
 }
@@ -2069,7 +2116,7 @@ UBA_EXPORT pid_t UBA_WRAPPER(waitpid)(pid_t pid, int* status, int options)
 	UBA_INIT_DETOUR(waitpid, pid, status, options);
 	// TODO: Should probably report id to session
 	auto res = TRUE_WRAPPER(waitpid)(pid, status, options);
-	DEBUG_LOG_TRUE("waitpid", "(%i) -> %i", pid, res);
+	DEBUG_LOG_TRUE("waitpid", "(%i) -> %i (%i)", pid, res, *status);
 	if (WIFEXITED(*status))
 		Rpc_UpdateTables();
 	return res;
@@ -2117,9 +2164,11 @@ UBA_EXPORT pid_t UBA_WRAPPER(wait3)(int* status, int options, struct rusage* rus
 UBA_EXPORT pid_t UBA_WRAPPER(wait4)(pid_t pid, int* status, int options, struct rusage* rusage)
 {
 	UBA_INIT_DETOUR(wait4, pid, status, options, rusage);
-	UBA_ASSERTF(!t_inVfork, "wait4: is in fork");
+	//UBA_ASSERTF(!t_inVfork, "wait4: is in fork");
 	pid_t res = TRUE_WRAPPER(wait4)(pid, status, options, rusage);
-	DEBUG_LOG_TRUE("wait4", "");
+	if (WIFEXITED(*status))
+		Rpc_UpdateTables();
+	//DEBUG_LOG_TRUE("wait4", "(pid %i) -> %i (%i)", pid, res, *status);
 	return res;
 }
 
@@ -2214,7 +2263,7 @@ int Internal_execve(const char* pathname, char* const _Nullable argv[], char* co
 	int status;
 	{
 		res = TRUE_WRAPPER(waitpid)(pid, &status, WUNTRACED | WCONTINUED);
-		DEBUG_LOG_TRUE("waitpid", "(%i) -> %i", pid, res);
+		DEBUG_LOG_TRUE("waitpid", "(execve) (%i) -> %i (%i)", pid, res, status);
 		UBA_ASSERTF(res == pid, "execve: wait result was not same as pid");
 		UBA_ASSERTF(WIFEXITED(status), "execve: Unsupported status from waitpid");
 	}
@@ -2233,10 +2282,11 @@ int Internal_execve(const char* pathname, char* const _Nullable argv[], char* co
 UBA_EXPORT int UBA_WRAPPER(execve)(const char* pathname, char* const _Nullable argv[], char* const _Nullable envp[])
 {
 	int res = Internal_execve(pathname, argv, envp);
-	Deinit();
-	CloseCom();
-	TRUE_WRAPPER(_exit)(res);
-	return 0;
+	if (t_inVfork)
+	{
+		TRUE_WRAPPER(_exit)(res); // execve should exit unless it really fails
+	}
+	return res;
 }
 
 UBA_EXPORT int UBA_WRAPPER(execvp)(const char* file, char* const argv[])
@@ -2322,8 +2372,11 @@ UBA_EXPORT void UBA_WRAPPER(exit)(int status)
 UBA_EXPORT void UBA_WRAPPER(_exit)(int status)
 {
 	DEBUG_LOG_TRUE("_exit", "(%i)", status);
-	Deinit();
-	CloseCom();
+	if (!t_inVfork)
+	{
+		Deinit();
+		CloseCom();
+	}
 	TRUE_WRAPPER(_exit)(status);
 }
 
@@ -2653,6 +2706,11 @@ namespace uba
 		//fsync(g_debugFile);
 		errno = t;
 	}
+	void FlushDebugLog()
+	{
+		if (isLogging())
+			fsync(g_debugFile);
+	}
 	#endif
 
 	ANALYSIS_NORETURN void UbaAssert(const tchar* text, const char* file, u32 line, const char* expr, u32 terminateCode)
