[UBA]
* Added special handling for running /usr/bin/stat via execve.. This is only used for linux relink actions which does not need to run remotely

diff --git a/Engine/Source/Programs/UnrealBuildAccelerator/Detours/Private/PosixOS/UbaDetoursFunctionsPosix.cpp b/Engine/Source/Programs/UnrealBuildAccelerator/Detours/Private/PosixOS/UbaDetoursFunctionsPosix.cpp
--- a/Engine/Source/Programs/UnrealBuildAccelerator/Detours/Private/PosixOS/UbaDetoursFunctionsPosix.cpp
+++ b/Engine/Source/Programs/UnrealBuildAccelerator/Detours/Private/PosixOS/UbaDetoursFunctionsPosix.cpp
@@ -2248,6 +2248,12 @@ int Internal_execve(const char* pathname, char* const _Nullable argv[], char* co
 	// We are most likely in a vfork/fork here, which means that we won't see the exit call since _exit will drop back out to the entrance of fork.
 	DEBUG_LOG_TRUE("execve", "%s", pathname);
 
+	if (StartsWith(pathname, "/usr/bin/stat"))
+	{
+		UBA_ASSERTF(!g_runningRemote, "Tried using posix_spawn with ExecuteHostRun+SpawnEcho but didn't work.");
+		return TRUE_WRAPPER(execve)(pathname, argv, envp);
+	}
+
 	bool inVfork = t_inVfork != 0;
 	pid_t pid;
 	int res = shared_posix_spawn(&pid, pathname, nullptr, nullptr, argv, envp);
