[UBA]
* Fixed bug where stack memory was used after going out of scope

diff --git a/Engine/Source/Programs/UnrealBuildAccelerator/Common/Private/UbaProcess.cpp b/Engine/Source/Programs/UnrealBuildAccelerator/Common/Private/UbaProcess.cpp
--- a/Engine/Source/Programs/UnrealBuildAccelerator/Common/Private/UbaProcess.cpp
+++ b/Engine/Source/Programs/UnrealBuildAccelerator/Common/Private/UbaProcess.cpp
@@ -1758,11 +1758,11 @@ namespace uba
 			}
 			Vector<const char*> arguments2;
 			arguments2.reserve(arguments.size() + 2);
-			{
-				StringBuffer<512> application(m_startInfo.application);
-				m_session.VirtualizePath(application, m_startInfo.rootsHandle);
-				arguments2.push_back(application.data);
-			}
+
+			StringBuffer<512> application(m_startInfo.application);
+			m_session.VirtualizePath(application, m_startInfo.rootsHandle);
+			arguments2.push_back(application.data);
+
 			for (auto& s : arguments)
 				arguments2.push_back(s.data());
 			arguments2.push_back(nullptr);
