[UBA]
* Fixed so GetFileLastWriteTime handles non existing files (posix only)
* Fixed FileAccessor close file while using temp file and fetching last write time (posix only)
* Fixed bugs related to Logger and long strings (posix only)

diff --git a/Engine/Source/Programs/UnrealBuildAccelerator/Common/Private/UbaFile.cpp b/Engine/Source/Programs/UnrealBuildAccelerator/Common/Private/UbaFile.cpp
--- a/Engine/Source/Programs/UnrealBuildAccelerator/Common/Private/UbaFile.cpp
+++ b/Engine/Source/Programs/UnrealBuildAccelerator/Common/Private/UbaFile.cpp
@@ -614,7 +614,8 @@ namespace uba
 		int res = fstat(asFileDescriptor(hFile), &attr);
 		if (res != 0)
 		{
-			UBA_ASSERTF(false, TC("GetFileLastWriteTime (fstat) error handling not implemented: %s"), strerror(errno));
+			UBA_ASSERTF(errno == ENOENT, TC("GetFileLastWriteTime (fstat) error handling not implemented: %s"), strerror(errno));
+			SetLastError(ERROR_FILE_NOT_FOUND);
 			return false;
 		}
 		outTime = FromTimeSpec(attr.st_mtimespec);
diff --git a/Engine/Source/Programs/UnrealBuildAccelerator/Common/Private/UbaFileAccessor.cpp b/Engine/Source/Programs/UnrealBuildAccelerator/Common/Private/UbaFileAccessor.cpp
--- a/Engine/Source/Programs/UnrealBuildAccelerator/Common/Private/UbaFileAccessor.cpp
+++ b/Engine/Source/Programs/UnrealBuildAccelerator/Common/Private/UbaFileAccessor.cpp
@@ -408,35 +408,46 @@ namespace uba
 			if (!SetDeleteOnClose(m_logger, realFileName, m_fileHandle, false))
 				return m_logger.Error(TC("Failed to remove delete on close for file %s (%s)"), realFileName, LastErrorToText().data);
 			#else
-			if (m_tempPath && rename(realFileName, m_fileName) == -1)
+			if (m_tempPath)
 			{
-				if (errno != EXDEV)
-					return m_logger.Error(TC("Failed to rename temporary file %s to %s (%s)"), realFileName, m_fileName, strerror(errno));
-
-				// Need to copy, can't rename over devices
-				int targetFd = open(m_fileName, O_CREAT | O_TRUNC | O_WRONLY | O_CLOEXEC, m_flagsAndAttributes);
-				auto g = MakeGuard([targetFd]() { close(targetFd); });
-				if (targetFd == -1)
-					return m_logger.Error(TC("Failed to create file %s for move from temporary file %s (%s)"), m_fileName, realFileName, strerror(errno));
+				if (rename(realFileName, m_fileName) == -1)
+				{
+					if (errno != EXDEV)
+						return m_logger.Error(TC("Failed to rename temporary file %s to %s (%s)"), realFileName, m_fileName, strerror(errno));
+
+					// Need to copy, can't rename over devices
+					int targetFd = open(m_fileName, O_CREAT | O_TRUNC | O_WRONLY | O_CLOEXEC, m_flagsAndAttributes);
+					auto closeTargetFile = MakeGuard([&]() { close(targetFd); });
+					if (targetFd == -1)
+						return m_logger.Error(TC("Failed to create file %s for move from temporary file %s (%s)"), m_fileName, realFileName, strerror(errno));
 						
-				#if PLATFORM_MAC
-				if (fcopyfile(asFileDescriptor(m_fileHandle), targetFd, 0, COPYFILE_ALL) == -1)
-					return m_logger.Error(TC("Failed to do fcopyfile from temporary %s to file %s (%s)"), realFileName, m_fileName, strerror(errno));
-				#else
-				int sourceFd = asFileDescriptor(m_fileHandle);
-				if (lseek(sourceFd, 0, SEEK_SET) == -1)
-					return m_logger.Error(TC("Failed to do lseek to beginning for sendfile (%s)"), strerror(errno));
-				u64 left = m_size;
-				while (left)
+					int sourceFd = asFileDescriptor(m_fileHandle);
+					#if PLATFORM_MAC
+					if (fcopyfile(sourceFd, targetFd, 0, COPYFILE_ALL) == -1)
+						return m_logger.Error(TC("Failed to do fcopyfile from temporary %s to file %s (%s)"), realFileName, m_fileName, strerror(errno));
+					#else
+					if (lseek(sourceFd, 0, SEEK_SET) == -1)
+						return m_logger.Error(TC("Failed to do lseek to beginning for sendfile (%s)"), strerror(errno));
+					u64 left = m_size;
+					while (left)
+					{
+						u64 written = sendfile(targetFd, sourceFd, NULL, left);
+						if (written == -1)
+							return m_logger.Error(TC("Failed to do sendfile from temporary %s (%i) to file %s (%i) (%s)"), realFileName, targetFd, m_fileName, sourceFd, strerror(errno));
+						left -= written;
+					}
+					#endif
+
+					remove(realFileName); // Remove real file now when we have copied it over
+				}
+				
+				if (lastWriteTime) // Have to reopen it seems
 				{
-					u64 written = sendfile(targetFd, sourceFd, NULL, left);
-					if (written == -1)
-						return m_logger.Error(TC("Failed to do sendfile from temporary %s (%i) to file %s (%i) (%s)"), realFileName, targetFd, m_fileName, sourceFd, strerror(errno));
-					left -= written;
+					if (!CloseFile(m_fileName, m_fileHandle))
+						return m_logger.Error(TC("Failed to close file %s (%s)"), m_fileName, LastErrorToText().data);
+					if (!OpenFileSequentialRead(m_logger, m_fileName, m_fileHandle))
+						return m_logger.Error(TC("Failed to re-open file %s (%s)"), m_fileName, LastErrorToText().data);
 				}
-				#endif
-
-				remove(realFileName); // Remove real file now when we have copied it over
 			}
 			#endif
 
@@ -444,7 +455,9 @@ namespace uba
 			{
 				*lastWriteTime = 0;
 				if (!GetFileLastWriteTime(*lastWriteTime, m_fileHandle))
-					m_logger.Warning(TC("Failed to get file time for %s (%s)"), realFileName, LastErrorToText().data);
+				{
+					m_logger.Warning(TC("Failed to get file time for %s %llu (%s)"), m_fileName, m_fileHandle, LastErrorToText().data);
+				}
 			}
 		}
 		else
diff --git a/Engine/Source/Programs/UnrealBuildAccelerator/Common/Private/UbaLogger.cpp b/Engine/Source/Programs/UnrealBuildAccelerator/Common/Private/UbaLogger.cpp
--- a/Engine/Source/Programs/UnrealBuildAccelerator/Common/Private/UbaLogger.cpp
+++ b/Engine/Source/Programs/UnrealBuildAccelerator/Common/Private/UbaLogger.cpp
@@ -141,27 +141,40 @@ namespace uba
 	{
 		#if !PLATFORM_WINDOWS
 		constexpr size_t _TRUNCATE = (size_t)-1;
+		va_list args2;
+		va_copy(args2, args);
+		auto g = MakeGuard([&]() { va_end(args2); });
 		#endif
 
 		tchar buffer[2048];
 		int len = _vsnwprintf_s(buffer, sizeof_array(buffer), _TRUNCATE, format, args);
-		if (len != -1)
+		if (len >= 0 && len < sizeof_array(buffer))
 		{
 			Log(type, buffer, u32(len));
 			return;
 		}
 
+		#if PLATFORM_WINDOWS
 		len = _vscwprintf(format, args);
-		if (len == -1)
+		if (len < 0)
 		{
 			Error(TC("LogArg failed. bad format"));
 			return;
 		}
+		va_list& newArgs = args;
+		#else
+		va_list& newArgs = args2;
+		#endif
 
 		Vector<tchar> buf;
 		buf.resize(len + 1);
-		_vsnwprintf_s(buf.data(), buf.size(), buf.size(), format, args);
-		Log(type, buf.data(), u32(buf.size()));
+		len = _vsnwprintf_s(buf.data(), buf.size(), buf.size(), format, newArgs);
+		if (len < 0)
+		{
+			Error(TC("LogArg failed. bad format"));
+			return;
+		}
+		Log(type, buf.data(), u32(len));
 	}
 
 	LoggerWithWriter::LoggerWithWriter(LogWriter& writer, const tchar* prefix)
