[UBA]
* Added so detour library pdb/debug file is transferred to helpers when running in debug

diff --git a/Engine/Source/Programs/UnrealBuildAccelerator/Common/Private/UbaSessionClient.cpp b/Engine/Source/Programs/UnrealBuildAccelerator/Common/Private/UbaSessionClient.cpp
--- a/Engine/Source/Programs/UnrealBuildAccelerator/Common/Private/UbaSessionClient.cpp
+++ b/Engine/Source/Programs/UnrealBuildAccelerator/Common/Private/UbaSessionClient.cpp
@@ -1191,6 +1191,7 @@ namespace uba
 		}
 
 		CasKey detoursBinaryKey = reader.ReadCasKey();
+		CasKey detoursSymbolsKey = reader.ReadCasKey();
 
 		StringBuffer<> detoursFile;
 		if (m_downloadDetoursLib)
@@ -1204,6 +1205,20 @@ namespace uba
 			KeyToString dir(StringKeyZero);
 			if (!WriteBinFile(detoursFile, UBA_DETOURS_LIBRARY, detoursBinaryKey, dir, DefaultAttributes()))
 				return;
+
+			if (detoursSymbolsKey != CasKeyZero)
+			{
+				StringBuffer<128> symbolsFile(UBA_DETOURS_LIBRARY);
+				#if PLATFORM_WINDOWS
+				symbolsFile.Resize(symbolsFile.count - 3).Append("pdb");
+				#else
+				symbolsFile.Resize(symbolsFile.count - 2).Append("debug");
+				#endif
+				Storage::RetrieveResult result;
+				StringBuffer<> throwaway;
+				if (m_storage.RetrieveCasFile(result, AsCompressed(detoursSymbolsKey, false), symbolsFile.data))
+					WriteBinFile(throwaway, symbolsFile.data, detoursSymbolsKey, dir, DefaultAttributes());
+			}
 		}
 		else
 		{
diff --git a/Engine/Source/Programs/UnrealBuildAccelerator/Common/Private/UbaSessionServer.cpp b/Engine/Source/Programs/UnrealBuildAccelerator/Common/Private/UbaSessionServer.cpp
--- a/Engine/Source/Programs/UnrealBuildAccelerator/Common/Private/UbaSessionServer.cpp
+++ b/Engine/Source/Programs/UnrealBuildAccelerator/Common/Private/UbaSessionServer.cpp
@@ -628,6 +628,7 @@ namespace uba
 		CasKey clientDetoursKey = reader.ReadCasKey();
 
 		CasKey& detoursBinaryKey = m_detoursBinaryKey[isClientArm];
+		CasKey& detoursSymbolsKey = m_detoursSymbolsKey[isClientArm];
 		CasKey& agentBinaryKey = m_agentBinaryKey[isClientArm];
 
 		bool binAsVersion = clientAgentKey != CasKeyZero;
@@ -682,6 +683,16 @@ namespace uba
 				dir.Resize(dirCount).Append(PathSeparator).Append(UBA_DETOURS_LIBRARY);
 				if (!m_storage.StoreCasFile(detoursBinaryKey, dir.data, CasKeyZero, deferCreation, fileIsCompressed) || detoursBinaryKey == CasKeyZero)
 					return m_logger.Error(TC("Failed to create cas for %s"), dir.data);
+
+#if UBA_DEBUG
+				#if PLATFORM_WINDOWS
+				dir.Resize(dir.count - 3).Append("pdb");
+				#else
+				dir.Resize(dir.count - 2).Append("debug");
+				#endif
+				if (!m_storage.StoreCasFile(detoursSymbolsKey, dir.data, CasKeyZero, deferCreation, fileIsCompressed) || detoursSymbolsKey == CasKeyZero)
+					return m_logger.Error(TC("Failed to create cas for %s"), dir.data);
+#endif
 			}
 		}
 
@@ -732,6 +743,7 @@ namespace uba
 			return false;
 
 		writer.WriteCasKey(detoursBinaryKey);
+		writer.WriteCasKey(detoursSymbolsKey);
 		writer.WriteBool(m_resetCas);
 		writer.WriteU32(sessionId);
 		writer.WriteU32(m_uiLanguage);
diff --git a/Engine/Source/Programs/UnrealBuildAccelerator/Common/Public/UbaNetwork.h b/Engine/Source/Programs/UnrealBuildAccelerator/Common/Public/UbaNetwork.h
--- a/Engine/Source/Programs/UnrealBuildAccelerator/Common/Public/UbaNetwork.h
+++ b/Engine/Source/Programs/UnrealBuildAccelerator/Common/Public/UbaNetwork.h
@@ -16,7 +16,7 @@ namespace uba
 	
 	static constexpr u32 SystemNetworkVersion = 1339;
 	static constexpr u32 StorageNetworkVersion = 4;
-	static constexpr u32 SessionNetworkVersion = 39;
+	static constexpr u32 SessionNetworkVersion = 40;
 	static constexpr u32 CacheNetworkVersion = 5;
 
 	// Messages used over network between client/server (system, storage and session)
