[UBA]
* Added so detour library pdb/debug file is transferred to helpers when running in debug

diff --git a/Engine/Source/Programs/UnrealBuildAccelerator/Common/Public/UbaSessionServer.h b/Engine/Source/Programs/UnrealBuildAccelerator/Common/Public/UbaSessionServer.h
--- a/Engine/Source/Programs/UnrealBuildAccelerator/Common/Public/UbaSessionServer.h
+++ b/Engine/Source/Programs/UnrealBuildAccelerator/Common/Public/UbaSessionServer.h
@@ -136,6 +136,7 @@ namespace uba
 
 		ReaderWriterLock m_binKeysLock;
 		CasKey m_detoursBinaryKey[2];
+		CasKey m_detoursSymbolsKey[2];
 		CasKey m_agentBinaryKey[2];
 
 		struct ClientSession
