[UBA]
* Improved assert output when triggered from detoured process.. still not perfect, will revisit
* Added so tables are updated after a successful wait3

diff --git a/Engine/Source/Programs/UnrealBuildAccelerator/Core/Private/UbaPlatform.cpp b/Engine/Source/Programs/UnrealBuildAccelerator/Core/Private/UbaPlatform.cpp
--- a/Engine/Source/Programs/UnrealBuildAccelerator/Core/Private/UbaPlatform.cpp
+++ b/Engine/Source/Programs/UnrealBuildAccelerator/Core/Private/UbaPlatform.cpp
@@ -6,7 +6,7 @@
 #if !PLATFORM_WINDOWS
 #include <execinfo.h>
 #include <limits.h>
-extern const char* __progname;
+#include <dlfcn.h>
 #else
 #include <conio.h>
 #include <psapi.h>
@@ -194,10 +194,10 @@ namespace uba
 		}
 #else
 		out.Append(TC("ASSERT: ")).Append(*text ? text : TC("Unknown")).Append('\n');
-		if (expr && strcmp(expr, "false") != 0)
+		if (!*text && expr && strcmp(expr, "false") != 0)
 			out.Appendf(TC(" EXPR: %s\n"), expr);
 		if (file && *file)
-			out.Appendf(TC("LOCATION: %s:%u"), file, line);
+			out.Appendf(TC(" LOCATION: %s:%u"), file, line);
 
 		constexpr u32 maxCallers = 100;
 		void* buffer[maxCallers];
@@ -205,29 +205,55 @@ namespace uba
 
 		if (true)
 		{
+			bool tryResolve = false;
 			char progName[256];
-			int res = readlink("/proc/self/exe", progName, sizeof_array(progName));
-			if (res != -1)
+
+			u64 baseOffset = 0;
+			Dl_info info;
+			if (dladdr((void*)&WriteAssertInfo, &info))
+			{
+				if (strstr(info.dli_fname, "UbaDetours"))
+				{
+					baseOffset = u64(info.dli_fbase);
+					strcpy(progName, info.dli_fname);
+					tryResolve = true;
+				}
+			}
+			if (!tryResolve)
+			{
+				int res = readlink("/proc/self/exe", progName, sizeof_array(progName));
+				if (res > 0)
+				{
+					progName[res] = 0;
+					tryResolve = true;
+				}
+			}
+
+			if (tryResolve)
 			{
-				progName[res] = 0;
 				fflush(stdout);
-				char str[4096];
+				char str[2048];
 				strcpy_s(str, sizeof(str), "addr2line");
 				for (int i = 0; i < nptrs; ++i)
 				{
 					char v[32];
-					snprintf(v, sizeof_array(v), " %llx", u64(buffer[i]));
+					u64 offset = u64(buffer[i]);
+					if (offset > baseOffset)
+						offset -= baseOffset;
+					snprintf(v, sizeof_array(v), " %llx", offset);
 					strcat_s(str, sizeof(str), v);
 				}
 				strcat_s(str, sizeof(str), " -e ");
 
 				strcat_s(str, sizeof(str), progName);
 				fflush(stdout);
+
 				if (FILE* fp = popen(str, "r"))
 				{
 					out.Append('\n');
 					u32 index = 0;
 					u32 countBeforeCallstack = out.count;
+					errno = 0;
 					while (fgets(str, sizeof(str), fp) != NULL)
 					{
 						if (str[0] == '?')
@@ -242,10 +268,14 @@ namespace uba
 							continue;
 						if (out.capacity - out.count < strlen(fileName) + 5)
 							break;
-						out.Appendf(TC("  %s"), fileName);
+						out.Appendf(TC(" %s"), fileName);
 					}
+					int errnoCopy = errno;
+					int err = ferror(fp);
 					fclose(fp);
-					if (countBeforeCallstack != out.count)
+					if (err)
+						out.Appendf(TC("   fgets failed for %s (%s)\n"), progName, strerror(errnoCopy));
+					else if (countBeforeCallstack != out.count)
 						return;
 				}
 			}
@@ -258,7 +288,12 @@ namespace uba
 			return;
 		}
 		for (int i = 0; i < nptrs; ++i)
-			out.Appendf(TC("%s\n"), strings[i]);
+		{
+			if (char* absoluteAddressStart = strrchr(strings[i], '['))
+				if (absoluteAddressStart > strings[i] && absoluteAddressStart[-1] == ' ')
+					absoluteAddressStart[-1] = 0;
+			out.Appendf(TC("  %s\n"), strings[i]);
+		}
 		free(strings);
 #endif
 	}
diff --git a/Engine/Source/Programs/UnrealBuildAccelerator/Detours/Private/PosixOS/UbaDetoursFunctionsPosix.cpp b/Engine/Source/Programs/UnrealBuildAccelerator/Detours/Private/PosixOS/UbaDetoursFunctionsPosix.cpp
--- a/Engine/Source/Programs/UnrealBuildAccelerator/Detours/Private/PosixOS/UbaDetoursFunctionsPosix.cpp
+++ b/Engine/Source/Programs/UnrealBuildAccelerator/Detours/Private/PosixOS/UbaDetoursFunctionsPosix.cpp
@@ -958,6 +958,8 @@ UBA_EXPORT char* UBA_WRAPPER(realpath)(const char* path, char* resolved_path)
 
 UBA_EXPORT ssize_t UBA_WRAPPER(readlink)(const char* pathname, char* buf, size_t bufsiz)
 {
+	UBA_INIT_DETOUR(readlink, pathname, buf, bufsiz);
+
 	// Beautiful hack. Some of our tools use je_malloc and dlsym do memory allocation so we end up in a deadlock when initializing detour (since detour use dlsym to figure out true function)
 	if (!g_isDetouring && Equals(pathname, "/etc/je_malloc.conf"))
 	{
@@ -1712,7 +1714,7 @@ UBA_EXPORT int UBA_WRAPPER(symlink)(const char* path1, const char* path2)
 	return TRUE_WRAPPER(symlink)(path1, path2);
 }
 
-UBA_EXPORT ssize_t	 UBA_WRAPPER(pread)(int __fd, void * __buf, size_t __nbyte, off_t __offset)
+UBA_EXPORT ssize_t UBA_WRAPPER(pread)(int __fd, void * __buf, size_t __nbyte, off_t __offset)
 {
 	// char filePath[PATH_MAX];
 	// if (fcntl(__fd, F_GETPATH, filePath) != -1)
@@ -2155,9 +2157,11 @@ UBA_EXPORT int UBA_WRAPPER(waitid)(idtype_t idtype, id_t id, siginfo_t* infop, i
 UBA_EXPORT pid_t UBA_WRAPPER(wait3)(int* status, int options, struct rusage* rusage)
 {
 	UBA_INIT_DETOUR(wait3, status, options, rusage);
-	DEBUG_LOG_TRUE("wait3", "");
 	UBA_ASSERTF(!t_inVfork, "wait3: is in fork");
 	pid_t res =TRUE_WRAPPER(wait3)(status, options, rusage);
+	DEBUG_LOG_TRUE("wait3", "-> %i (%i)", res, *status);
+	if (WIFEXITED(*status))
+		Rpc_UpdateTables();
 	return res;
 }
 
