[UBA]
* Added special handling of shell scripts running remotely. Fixed so binary is not downloaded and also fixed so environment is flushed between each process spawned by shell script

diff --git a/Engine/Source/Programs/UnrealBuildAccelerator/Common/Private/UbaSessionClient.cpp b/Engine/Source/Programs/UnrealBuildAccelerator/Common/Private/UbaSessionClient.cpp
--- a/Engine/Source/Programs/UnrealBuildAccelerator/Common/Private/UbaSessionClient.cpp
+++ b/Engine/Source/Programs/UnrealBuildAccelerator/Common/Private/UbaSessionClient.cpp
@@ -251,6 +251,12 @@ namespace uba
 		outRealWorkingDir = m_processWorkingDir.data;
 		if (StartsWith(startInfo.application, TC("ubacopy")))
 			return true;
+
+		#if PLATFORM_LINUX
+		if (EndsWith(startInfo.application, TStrlen(startInfo.application), "/sh"))
+			return true;
+		#endif
+
 		outRealApplication.Clear();
 
 		const tchar* application = startInfo.application;
@@ -424,6 +430,23 @@ namespace uba
 	{
 		if (!Session::ProcessThreadStart(process))
 			return false;
+
+		#if PLATFORM_LINUX
+		// TODO. Linux linker has lots of stuff going on. it is a shell script deleting, renaming, overwriting same file causing confusion with child processes
+		// Since directory table is not updated from server between child processes we end up with not up-to-date table.. therefore we flush everything in between
+		if (auto parent = process.m_parentProcess)
+		{
+			auto& si = parent->m_startInfo;
+			if (EndsWith(si.application, TStrlen(si.application), "/sh"))
+			{
+				FlushWrittenFiles(*parent);
+				UpdateEnvironment(*parent, AsView("ShellScript"), false);
+			}
+		}
+		#endif
+
+
+
 		if (RootsHandle rootsHandle = process.GetStartInfo().rootsHandle)
 			if (!SendRootsHandle(rootsHandle))
 				return false;
