[UBA]
* Fixed so DebugLogger is owned by session to fix but where log file was overwritten when having multiple sessions in the same process

diff --git a/Engine/Source/Programs/UnrealBuildAccelerator/Common/Private/UbaLogger.cpp b/Engine/Source/Programs/UnrealBuildAccelerator/Common/Private/UbaLogger.cpp
--- a/Engine/Source/Programs/UnrealBuildAccelerator/Common/Private/UbaLogger.cpp
+++ b/Engine/Source/Programs/UnrealBuildAccelerator/Common/Private/UbaLogger.cpp
@@ -414,16 +414,17 @@ namespace uba
 		FileAccessor* m_file = nullptr;
 		ReaderWriterLock m_logLock;
 
-	} g_debugLogWriter;
+	};
 
-	bool StartDebugLogger(Logger& outerLogger, const tchar* fileName)
+	Logger* StartDebugLogger(Logger& outerLogger, const tchar* fileName)
 	{
-		g_debugLogWriter.m_fileName = fileName;
-		auto fa = new FileAccessor(outerLogger, g_debugLogWriter.m_fileName.c_str());
+		auto debugWriter = new DebugLogWriter();
+		debugWriter->m_fileName = fileName;
+		auto fa = new FileAccessor(outerLogger, debugWriter->m_fileName.c_str());
 		if (!fa->CreateWrite())
 		{
 			delete fa;
-			return false;
+			return new LoggerWithWriter(g_nullLogWriter);
 		}
 
 		#if PLATFORM_WINDOWS
@@ -431,20 +432,28 @@ namespace uba
 		fa->Write(utf8BOM, sizeof(utf8BOM));
 		#endif
 
-		g_debugLogWriter.m_file = fa;
-		return true;
+		debugWriter->m_file = fa;
+		return new LoggerWithWriter(*debugWriter);
 	}
 
-	void StopDebugLogger()
+	Logger* StopDebugLogger(Logger* logger)
 	{
-		if (!g_debugLogWriter.m_file)
-			return;
-		g_debugLogWriter.m_file->Close();
-		delete g_debugLogWriter.m_file;
-		g_debugLogWriter.m_file = nullptr;
+		auto debugLogger = (LoggerWithWriter*)logger;
+		if (&debugLogger->m_writer != &g_nullLogWriter)
+		{
+			auto debugWriter = (DebugLogWriter*)&debugLogger->m_writer;
+			if (debugWriter->m_file)
+			{
+				debugWriter->m_file->Close();
+				delete debugWriter->m_file;
+				debugWriter->m_file = nullptr;
+			}
+			delete debugWriter;
+		}
+		delete debugLogger;
+		return nullptr;
 	}
 
-	LoggerWithWriter g_debugLogger(g_debugLogWriter);
 #endif
 
 	#if UBA_TRACK_CONTENTION
diff --git a/Engine/Source/Programs/UnrealBuildAccelerator/Common/Private/UbaProcess.cpp b/Engine/Source/Programs/UnrealBuildAccelerator/Common/Private/UbaProcess.cpp
--- a/Engine/Source/Programs/UnrealBuildAccelerator/Common/Private/UbaProcess.cpp
+++ b/Engine/Source/Programs/UnrealBuildAccelerator/Common/Private/UbaProcess.cpp
@@ -564,7 +564,7 @@ namespace uba
 		KernelStats::GetGlobal().Add(m_kernelStats);
 
 		#if UBA_DEBUG_TRACK_PROCESS
-		g_debugLogger.Info(TC("ProcessExitedStart (%u)\n"), m_id);
+		m_session.m_debugLogger->Info(TC("ProcessExitedStart (%u)\n"), m_id);
 		#endif
 
 		// For some reason a parent can exit before a child. Need to figure out repro for this but I've seen it happen on ClangEditor win64
@@ -616,7 +616,7 @@ namespace uba
 		m_session.ProcessExited(*this, m_processStats.wallTime);
 
 		#if UBA_DEBUG_TRACK_PROCESS
-		g_debugLogger.Info(TC("ProcessExitedDone  (%u)\n"), m_id);
+		m_session.m_debugLogger->Info(TC("ProcessExitedDone  (%u)\n"), m_id);
 		#endif
 	}
 
@@ -1027,7 +1027,7 @@ namespace uba
 		#endif
 
 		#if UBA_DEBUG_TRACK_PROCESS
-		g_debugLogger.Info(TC("CreateChildProcess (%u creating child %u at index %u) %s %s (%s)\n"), m_id, process.m_id, childProcessId-1, process.m_realApplication.c_str(), commandLine, info.logFile);
+		m_session.m_debugLogger->Info(TC("CreateChildProcess (%u creating child %u at index %u) %s %s (%s)\n"), m_id, process.m_id, childProcessId-1, process.m_realApplication.c_str(), commandLineWithoutApplication.c_str(), info.logFile);
 		#endif
 		return true;
 	}
@@ -1077,13 +1077,13 @@ namespace uba
 #endif
 		
 		#if UBA_DEBUG_TRACK_PROCESS
-		g_debugLogger.Info(TC("WaitForChildProcessReady (%u waiting for child %u at index %u)\n"), m_id, process.m_id, processId-1);
+		m_session.m_debugLogger->Info(TC("WaitForChildProcessReady (%u waiting for child %u at index %u)\n"), m_id, process.m_id, processId-1);
 		#endif
 
 		setWaitForParent.Execute();
 
 		#if UBA_DEBUG_TRACK_PROCESS
-		g_debugLogger.Info(TC("WaitForChildProcessReadyDone (%u waiting for child %u at index %u)\n"), m_id, process.m_id, processId-1);
+		m_session.m_debugLogger->Info(TC("WaitForChildProcessReadyDone (%u waiting for child %u at index %u)\n"), m_id, process.m_id, processId-1);
 		#endif
 
 		#if PLATFORM_WINDOWS
@@ -1608,13 +1608,13 @@ namespace uba
 			m_nativeThreadHandle = processInfo.hThread;
 
 			#if UBA_DEBUG_TRACK_PROCESS
-			g_debugLogger.Info(TC("CreateRealProcess  (%u) %s %s\n"), m_id, m_realApplication.c_str(), m_startInfo.arguments);
+			m_session.m_debugLogger->Info(TC("CreateRealProcess  (%u) %s %s\n"), m_id, m_realApplication.c_str(), m_startInfo.arguments);
 			#endif
 		}
 		else
 		{
 			#if UBA_DEBUG_TRACK_PROCESS
-			g_debugLogger.Info(TC("WaitingForParentReady (%u)\n"), m_id);
+			m_session.m_debugLogger->Info(TC("WaitingForParentReady (%u)\n"), m_id);
 			#endif
 
 			WaitForParent();
@@ -1622,7 +1622,7 @@ namespace uba
 				return UBA_EXIT_CODE(7);
 
 			#if UBA_DEBUG_TRACK_PROCESS
-			g_debugLogger.Info(TC("WaitingForParentReadyDone (%u)\n"), m_id);
+			m_session.m_debugLogger->Info(TC("WaitingForParentReadyDone (%u)\n"), m_id);
 			#endif
 
 			m_extractExports = m_parentProcess->m_extractExports;
@@ -1945,13 +1945,13 @@ namespace uba
 			pipeGuard0.Cancel();
 
 			#if UBA_DEBUG_TRACK_PROCESS
-			g_debugLogger.Info(TC("CreateRealProcess  (%u) %s %.100s\n"), m_id, m_realApplication.c_str(), m_startInfo.arguments);
+			m_session.m_debugLogger->Info(TC("CreateRealProcess  (%u) %s %.100s\n"), m_id, m_realApplication.c_str(), m_startInfo.arguments);
 			#endif
 		}
 		else
 		{
 			#if UBA_DEBUG_TRACK_PROCESS
-			g_debugLogger.Info(TC("WaitingForParent (%u) %.100s\n"), m_id, m_realApplication.c_str());
+			m_session.m_debugLogger->Info(TC("WaitingForParent (%u) %.100s\n"), m_id, m_realApplication.c_str());
 			#endif
 
 			//logger.Info("Waiting for parent");
diff --git a/Engine/Source/Programs/UnrealBuildAccelerator/Common/Private/UbaSession.cpp b/Engine/Source/Programs/UnrealBuildAccelerator/Common/Private/UbaSession.cpp
--- a/Engine/Source/Programs/UnrealBuildAccelerator/Common/Private/UbaSession.cpp
+++ b/Engine/Source/Programs/UnrealBuildAccelerator/Common/Private/UbaSession.cpp
@@ -268,9 +268,9 @@ namespace uba
 		}
 
 		#if UBA_DEBUG_TRACK_DIR
-		g_debugLogger.BeginScope();
-		auto dg = MakeGuard([]() { g_debugLogger.EndScope(); });
-		g_debugLogger.Info(TC("TRACKDIR %s\n"), dirPath);
+		m_debugLogger->BeginScope();
+		auto dg = MakeGuard([&]() { m_debugLogger->EndScope(); });
+		m_debugLogger->Info(TC("TRACKDIR %s\n"), dirPath);
 		#endif
 
 
@@ -292,7 +292,7 @@ namespace uba
 				memoryWriter.WriteString(e.name, e.nameLen);
 
 				#if UBA_DEBUG_TRACK_DIR
-				g_debugLogger.Info(TC("    %s (Size: %llu, Key: %s, Id: %llu)\n"), e.name, e.size, KeyToString(fileKey).data, e.id);
+				m_debugLogger->Info(TC("    %s (Size: %llu, Key: %s, Id: %llu)\n"), e.name, e.size, KeyToString(fileKey).data, e.id);
 				#endif
 
 				u64 id = e.id;
@@ -851,7 +851,7 @@ namespace uba
 		}
 
 		#if UBA_DEBUG_TRACK_DIR
-		g_debugLogger.Info(TC("TRACKADD    %s (Size: %llu, Key: %s, Id: %llu)\n"), fileName.data, fileSize, KeyToString(fileNameKey).data, fileIndex);
+		m_debugLogger->Info(TC("TRACKADD    %s (Size: %llu, Key: %s, Id: %llu)\n"), fileName.data, fileSize, KeyToString(fileNameKey).data, fileIndex);
 		#endif
 
 
@@ -928,7 +928,7 @@ namespace uba
 		}
 
 		#if UBA_DEBUG_TRACK_DIR
-		g_debugLogger.Info(TC("TRACKDEL    %s (Key: %s)\n"), fileName.data, KeyToString(fileNameKey).data);
+		m_debugLogger->Info(TC("TRACKDEL    %s (Key: %s)\n"), fileName.data, KeyToString(fileNameKey).data);
 		#endif
 
 		SCOPED_WRITE_LOCK(dirTable.m_memoryLock, memoryLock);
@@ -1000,7 +1000,7 @@ namespace uba
 				return m_logger.Error("Failed to find file %s", applicationName);
 			//UBA_ASSERTF(false, TC("DIR NOT FOUND: %s"), applicationName);
 			#else
-			return m_logger.Error("Code path not implemented for linux!");
+			return m_logger.Error("Code path not implemented for linux! (CopyImports %s)", library);
 			#endif
 
 			applicationName = temp;
@@ -1137,7 +1137,7 @@ namespace uba
 	bool Session::Create(const SessionCreateInfo& info)
 	{
 		#if UBA_DEBUG_LOGGER
-		StartDebugLogger(m_logger, StringBuffer<512>().Append(m_sessionDir).Append(TC("SessionDebug.log")).data);
+		m_debugLogger = StartDebugLogger(m_logger, StringBuffer<512>().Append(m_sessionDir).Append(TC("SessionDebug.log")).data);
 		#endif
 
 		#if PLATFORM_WINDOWS
@@ -1266,7 +1266,7 @@ namespace uba
 		//#endif
 
 		#if UBA_DEBUG_LOGGER
-		StopDebugLogger();
+		m_debugLogger = StopDebugLogger(m_debugLogger);
 		#endif
 	}
 
diff --git a/Engine/Source/Programs/UnrealBuildAccelerator/Common/Private/UbaSessionClient.cpp b/Engine/Source/Programs/UnrealBuildAccelerator/Common/Private/UbaSessionClient.cpp
--- a/Engine/Source/Programs/UnrealBuildAccelerator/Common/Private/UbaSessionClient.cpp
+++ b/Engine/Source/Programs/UnrealBuildAccelerator/Common/Private/UbaSessionClient.cpp
@@ -652,6 +652,11 @@ namespace uba
 			if (!msg.Send(reader, Stats().sendFileMsg))
 				return m_logger.Error(TC("Failed to send file %s to server"), source.name.c_str());
 		}
+
+		#if UBA_DEBUG_LOGGER
+		m_debugLogger->Info(TC("SENDFILE %s\n"), destination);
+		#endif
+
 		if (!reader.ReadBool())
 			return m_logger.Error(TC("Server failed to receive file %s (%s)"), source.name.c_str(), destination);
 		return true;
diff --git a/Engine/Source/Programs/UnrealBuildAccelerator/Common/Public/UbaLogger.h b/Engine/Source/Programs/UnrealBuildAccelerator/Common/Public/UbaLogger.h
--- a/Engine/Source/Programs/UnrealBuildAccelerator/Common/Public/UbaLogger.h
+++ b/Engine/Source/Programs/UnrealBuildAccelerator/Common/Public/UbaLogger.h
@@ -82,9 +82,8 @@ namespace uba
 
 
 	#if UBA_DEBUG_LOGGER
-	bool StartDebugLogger(Logger& outerLogger, const tchar* fileName);
-	void StopDebugLogger();
-	extern LoggerWithWriter g_debugLogger;
+	Logger* StartDebugLogger(Logger& outerLogger, const tchar* fileName);
+	Logger* StopDebugLogger(Logger* logger);
 	#endif
 
 	void PrintContentionSummary(class Logger& logger);
diff --git a/Engine/Source/Programs/UnrealBuildAccelerator/Common/Public/UbaSession.h b/Engine/Source/Programs/UnrealBuildAccelerator/Common/Public/UbaSession.h
--- a/Engine/Source/Programs/UnrealBuildAccelerator/Common/Public/UbaSession.h
+++ b/Engine/Source/Programs/UnrealBuildAccelerator/Common/Public/UbaSession.h
@@ -287,6 +287,10 @@ namespace uba
 		GetNextProcessFunction m_getNextProcessFunction;
 		CustomServiceFunction m_customServiceFunction;
 
+		#if UBA_DEBUG_LOGGER
+		Logger* m_debugLogger = nullptr;
+		#endif
+
 		friend class ProcessImpl;
 	};
 
