[UBA]
* Changed default for application rules to not store files in file mappings

diff --git a/Engine/Source/Programs/UnrealBuildAccelerator/Core/Private/UbaApplicationRules.cpp b/Engine/Source/Programs/UnrealBuildAccelerator/Core/Private/UbaApplicationRules.cpp
--- a/Engine/Source/Programs/UnrealBuildAccelerator/Core/Private/UbaApplicationRules.cpp
+++ b/Engine/Source/Programs/UnrealBuildAccelerator/Core/Private/UbaApplicationRules.cpp
@@ -294,11 +294,6 @@ namespace uba
 			return fileName.Contains(TC(".tmp")); // both .so.tmp and .tmp123456
 		}
 
-		virtual bool IsRarelyRead(const StringView& file) const override
-		{
-			return file.EndsWith(TC(".so.rsp"));
-		}
-
 		virtual u64 FileTypeMaxSize(const StringView& file, bool isSystemOrTempFile) const override
 		{
 			return 14ull * 1024 * 1024 * 1024; // This is ridiculous (needed for asan targets)
diff --git a/Engine/Source/Programs/UnrealBuildAccelerator/Core/Public/UbaApplicationRules.h b/Engine/Source/Programs/UnrealBuildAccelerator/Core/Public/UbaApplicationRules.h
--- a/Engine/Source/Programs/UnrealBuildAccelerator/Core/Public/UbaApplicationRules.h
+++ b/Engine/Source/Programs/UnrealBuildAccelerator/Core/Public/UbaApplicationRules.h
@@ -75,14 +75,14 @@ namespace uba
 		// This is an optimization to not store the file in the mapping table since it will not be read again and would just take up space
 		virtual bool IsRarelyRead(const StringView& file) const
 		{
-			return false;
+			return true;
 		}
 
 		// Return true if the file is never/rarely read after it was written.
 		// This is an optimization where the written file is not kept in file mappings after written
 		virtual bool IsRarelyReadAfterWritten(const StringView& fileName) const
 		{
-			return false;
+			return true;
 		}
 
 		virtual bool AllowStorageProxy(const StringView& file) const
diff --git a/Engine/Source/Programs/UnrealBuildAccelerator/Detours/Private/Windows/UbaDetoursFunctionsKernelBase.inl b/Engine/Source/Programs/UnrealBuildAccelerator/Detours/Private/Windows/UbaDetoursFunctionsKernelBase.inl
--- a/Engine/Source/Programs/UnrealBuildAccelerator/Detours/Private/Windows/UbaDetoursFunctionsKernelBase.inl
+++ b/Engine/Source/Programs/UnrealBuildAccelerator/Detours/Private/Windows/UbaDetoursFunctionsKernelBase.inl
@@ -1521,9 +1521,13 @@ bool Shared_MoveFile(LPCWSTR lpExistingFileName, LPCWSTR lpNewFileName, DWORD dw
 
 	DEBUG_LOG_DETOURED(L"MoveFileExW", L"(PIPE) (%ls to %ls) -> %ls (%u)", lpExistingFileName, lpNewFileName, ToString(result), errorCode);
 
-	g_directoryTable.ParseDirectoryTable(directoryTableSize);
-	g_mappedFileTable.SetDeleted(sourceKey, source.data, true);
-	g_mappedFileTable.SetDeleted(destKey, dest.data, false);
+	if (result)
+	{
+		g_directoryTable.ParseDirectoryTable(directoryTableSize);
+		g_mappedFileTable.SetDeleted(sourceKey, source.data, true);
+		g_mappedFileTable.SetDeleted(destKey, dest.data, false);
+	}
+
 	SetLastError(errorCode);
 
 	return result;
