[UBA]
* Improved assert/crash handling. Most paths will now get a properly resolved callstack (at least in debug builds). Because of issues in wine we have to resolve the callstack on the host so system modules will be garbage but that is fine as long as the top of the stack looks correct
* Fixed so when ubalog is enabled remote logs are always sent back
* Disabled using local "sh" on posix because the aws linux sh hangs when linking
* Always set vectorized exception handler on windows and now filter out all exceptions happening outside uba dll
* Added ReaderWriterLock::TryEnterWrite that is used when handling asserts/seh.. needed if assert/seh happens while in the middle of sending a message
* Added nviewH64.dll in exclude list of dlls we never want to load into our detoured process
* Fixed so rsp file path is devirtualized in obj file preloader handling decompression of obj files
* Minor improvements (assert descriptions etc)

diff --git a/Engine/Source/Programs/UnrealBuildAccelerator/Agent/Private/UbaAgent.cpp b/Engine/Source/Programs/UnrealBuildAccelerator/Agent/Private/UbaAgent.cpp
--- a/Engine/Source/Programs/UnrealBuildAccelerator/Agent/Private/UbaAgent.cpp
+++ b/Engine/Source/Programs/UnrealBuildAccelerator/Agent/Private/UbaAgent.cpp
@@ -31,10 +31,8 @@
 
 #if PLATFORM_WINDOWS
 #define UBA_AUTO_UPDATE 1
-#define UBA_USE_EXCEPTION_HANDLER 0
 #else
 #define UBA_AUTO_UPDATE 0
-#define UBA_USE_EXCEPTION_HANDLER 0
 #endif
 //#include <dbghelp.h>
 //#pragma comment (lib, "Dbghelp.lib")
@@ -145,14 +143,6 @@ namespace uba
 	}
 
 	#if PLATFORM_WINDOWS
-	int ReportSEH(LPEXCEPTION_POINTERS exceptionInfo)
-	{
-		StringBuffer<4096> assertInfo;
-		assertInfo.Appendf(TC("SEH EXCEPTION %u (0x%llx)"), exceptionInfo->ExceptionRecord->ExceptionCode, exceptionInfo->ExceptionRecord->ExceptionAddress);
-		WriteAssertInfo(assertInfo, nullptr, nullptr, 0, nullptr, 1);
-		LoggerWithWriter(g_consoleLogWriter).Log(LogEntryType_Error, assertInfo.data, assertInfo.count);
-		return EXCEPTION_CONTINUE_SEARCH;
-	}
 	BOOL ConsoleHandler(DWORD signal)
 	{
 		CtrlBreakPressed();
@@ -167,31 +157,6 @@ namespace uba
 
 	StringBuffer<> g_rootDir(DefaultRootDir);
 
-#if UBA_USE_EXCEPTION_HANDLER
-	LONG WINAPI UbaUnhandledExceptionFilter(EXCEPTION_POINTERS* ExceptionInfo)
-	{
-		StringBuffer<4096> assertInfo;
-		WriteAssertInfo(assertInfo, TC(""), nullptr, 0, nullptr, 1);
-		LoggerWithWriter(g_consoleLogWriter).Log(LogEntryType_Error, assertInfo.data, assertInfo.count);
-	//	time_t rawtime;
-	//	time(&rawtime);
-	//	tm ti;
-	//	localtime_s(&ti, &rawtime);
-	//
-	//	StringBuffer<> dumpFile;
-	//	dumpFile.Append(g_rootDir).EnsureEndsWithSlash().Appendf(TC("UbaAgentCrash_%02u%02u%02u_%02u%02u%02u.dmp"), ti.tm_year - 100, ti.tm_mon + 1, ti.tm_mday, ti.tm_hour, ti.tm_min, ti.tm_sec);
-	//
-	//	wprintf(TC("Unhandled exception - Writing minidump %s\n"), dumpFile.data);
-	//	HANDLE hFile = CreateFile(dumpFile.data, GENERIC_READ | GENERIC_WRITE, FILE_SHARE_DELETE | FILE_SHARE_READ | FILE_SHARE_WRITE, NULL, CREATE_ALWAYS, FILE_ATTRIBUTE_NORMAL, NULL);
-	//	MINIDUMP_EXCEPTION_INFORMATION mei;
-	//	mei.ThreadId = GetCurrentThreadId();
-	//	mei.ClientPointers = TRUE;
-	//	mei.ExceptionPointers = ExceptionInfo;
-	//	MiniDumpWriteDump(GetCurrentProcess(), GetCurrentProcessId(), hFile, MiniDumpNormal, &mei, NULL, NULL);
-		return EXCEPTION_EXECUTE_HANDLER;
-	}
-#endif
-
 #if UBA_AUTO_UPDATE
 	const tchar* g_ubaAgentBinaries[] = { UBA_AGENT_EXECUTABLE }; //, UBA_DETOURS_LIBRARY };
 
@@ -398,10 +363,6 @@ namespace uba
 
 	bool WrappedMain(int argc, tchar*argv[])
 	{
-		#if UBA_USE_EXCEPTION_HANDLER
-		SetUnhandledExceptionFilter(UbaUnhandledExceptionFilter);
-		#endif
-
 		u32 maxProcessCount = DefaultProcessorCount;
 		u32 maxWorkerCount = DefaultProcessorCount;
 		float mulProcessValue = 1.0f;
@@ -432,6 +393,7 @@ namespace uba
 		bool printSummary = false;
 		bool killRandom = false;
 		bool downloadDetoursLib = true;
+		bool useExceptionHandler = true;
 		StringBuffer<512> sentryUrl;
 		StringBuffer<128> zone;
 		u32 maxIdleSeconds = ~0u;
@@ -602,6 +564,10 @@ namespace uba
 			{
 				useStorage = false;
 			}
+			else if (name.Equals(TC("-nohandler")))
+			{
+				useExceptionHandler = false;
+			}
 			else if (name.Equals(TC("-nodetoursdownload")))
 			{
 				downloadDetoursLib = false;
@@ -718,6 +684,9 @@ namespace uba
 			}
 		}
 
+		if (useExceptionHandler)
+			AddExceptionHandler();
+
 		if (!named.IsEmpty()) // We only run once with named connection
 			poll = false;
 
@@ -1430,15 +1399,7 @@ namespace uba
 #if PLATFORM_WINDOWS
 int wmain(int argc, wchar_t* argv[])
 {
-	using namespace uba;
-	__try
-	{
-		return WrappedMain(argc, argv) ? 0 : -1;
-	}
-	__except(ReportSEH(GetExceptionInformation()))
-	{
-		return -1;
-	}
+	return uba::WrappedMain(argc, argv) ? 0 : -1;
 }
 #else
 int main(int argc, char* argv[])
diff --git a/Engine/Source/Programs/UnrealBuildAccelerator/Common/Private/UbaFile.cpp b/Engine/Source/Programs/UnrealBuildAccelerator/Common/Private/UbaFile.cpp
--- a/Engine/Source/Programs/UnrealBuildAccelerator/Common/Private/UbaFile.cpp
+++ b/Engine/Source/Programs/UnrealBuildAccelerator/Common/Private/UbaFile.cpp
@@ -907,4 +907,23 @@ namespace uba
 		SCOPED_WRITE_LOCK(m_createdDirsLock, lock);
 		m_createdDirs.clear();
 	}
+
+	bool GetAlternativeUbaPath(Logger& logger, StringBufferBase& out, StringView firstPath, bool isWindowsArm)
+	{
+		out.Append(firstPath);
+		const tchar* engineDir = IsWindows ? TC("\\Engine\\") : TC("/Engine/");
+		const tchar* engineDirPos;
+		if (!out.Contains(engineDir, true, &engineDirPos))
+			return false;//logger.Error(TC("Failed to find Engine dir in %s"), out.data);
+		#if PLATFORM_WINDOWS
+		const tchar* platformStr = TC("Win64");
+		#elif PLATFORM_LINUX
+		const tchar* platformStr = TC("Linux");
+		#else
+		const tchar* platformStr = TC("Mac");
+		#endif
+		out.Resize(engineDirPos - out.data + 8).Append(TC("Binaries")).Append(PathSeparator).Append(platformStr).Append(PathSeparator).Append("UnrealBuildAccelerator").Append(PathSeparator);
+		out.Append(isWindowsArm ? TC("arm64") : TC("x64")).Append(PathSeparator);
+		return true;
+	};
 }
diff --git a/Engine/Source/Programs/UnrealBuildAccelerator/Common/Private/UbaLogger.cpp b/Engine/Source/Programs/UnrealBuildAccelerator/Common/Private/UbaLogger.cpp
--- a/Engine/Source/Programs/UnrealBuildAccelerator/Common/Private/UbaLogger.cpp
+++ b/Engine/Source/Programs/UnrealBuildAccelerator/Common/Private/UbaLogger.cpp
@@ -9,6 +9,7 @@
 
 #if PLATFORM_WINDOWS
 #include <io.h>
+#include <dbghelp.h>
 #define Fputs fputws
 #else
 #define Fputs fputs
@@ -21,27 +22,48 @@ namespace uba
 {
 	CustomAssertHandler* g_assertHandler;
 
-	ANALYSIS_NORETURN void UbaAssert(const tchar* text, const char* file, u32 line, const char* expr, u32 terminateCode)
+	ANALYSIS_NORETURN void UbaAssert(const tchar* text, const char* file, u32 line, const char* expr, bool allowTerminate, u32 terminateCode, void* context)
 	{
 		static ReaderWriterLock& assertLock = *new ReaderWriterLock(); // Leak to prevent asan annoyances during shutdown when asserts happen
 		SCOPED_WRITE_LOCK(assertLock, lock);
 
-		StringBuffer<4096> b;
-		WriteAssertInfo(b, text, file, line, expr, 1);
+		static u8* writerMem = new u8[4096];
+		BinaryWriter writer(writerMem, 0, 4096);
+		WriteCallstackInfo(writer, 2, context);
+
+		static auto& sb = *new StringBuffer<16*1024>();
+		WriteAssertInfo(sb.Clear(), text, file, line, expr);
+
+		StringView searchPaths[3];
+		StringBuffer<512> currentModuleDir;
+		LoggerWithWriter logger(g_nullLogWriter);
+		GetDirectoryOfCurrentModule(logger, currentModuleDir);
+		searchPaths[0] = currentModuleDir;
+		StringBuffer<512> alternativePath;
+		if (GetAlternativeUbaPath(logger, alternativePath, currentModuleDir, IsWindows && IsArm))
+			searchPaths[1] = alternativePath;
+
+		tchar executable[512] = { 0 };
+		#if !PLATFORM_WINDOWS
+		readlink("/proc/self/exe", executable, sizeof_array(executable));
+		#endif
+
+		BinaryReader reader(writerMem, 0, writer.GetPosition());
+		ParseCallstackInfo(sb, reader, executable, searchPaths);
 
 		if (g_assertHandler)
 		{
-			g_assertHandler(b.data);
+			g_assertHandler(sb.data);
 			return;
 		}
 
-		Fputs(b.data, stdout);
+		Fputs(sb.data, stdout);
 		Fputs(TC("\n"), stdout);
 		fflush(stdout);
 
 #if PLATFORM_WINDOWS
 #if UBA_ASSERT_MESSAGEBOX
-		int ret = MessageBoxW(GetConsoleWindow(), b.data, TC("Assert"), MB_ABORTRETRYIGNORE);
+		int ret = MessageBoxW(GetConsoleWindow(), sb.data, TC("Assert"), MB_ABORTRETRYIGNORE);
 		if (ret != IDABORT)
 		{
 			if (ret == IDRETRY)
@@ -460,6 +482,189 @@ namespace uba
 	List<ContentionTracker>& GetContentionTrackerList();
 	#endif
 
+	void ParseCallstackInfo(StringBufferBase& out, BinaryReader& reader, const tchar* executable, const StringView* searchPaths)
+	{
+		bool isRunningWine = reader.ReadBool();
+
+		struct CallstackEntry { u64 moduleIndex; u64 memoryOffset; };
+		u64 callstackCount = reader.Read7BitEncoded();
+		Vector<CallstackEntry> entries(callstackCount);
+		for (auto& entry : entries)
+		{
+			entry.moduleIndex = reader.Read7BitEncoded();
+			entry.memoryOffset = reader.Read7BitEncoded();
+		}
+
+		struct ModuleEntry { u64 start; u64 size; TString name; bool handled; };
+		u64 moduleCount = reader.Read7BitEncoded();
+		Vector<ModuleEntry> modules(moduleCount);
+		for (auto& mod : modules)
+		{
+			mod.start = reader.Read7BitEncoded();
+			mod.size = reader.Read7BitEncoded();
+			mod.name = reader.ReadString();
+			mod.handled = false;
+		}
+
+		if (entries.empty())
+			return;
+
+		out.Appendf(TC("\n CALLSTACK%s:"), isRunningWine ? TC(" (Wine)") : TC(""));
+
+#if PLATFORM_WINDOWS
+		static u64 processHandleCounter = 45234523;
+		auto processHandle = (HANDLE)processHandleCounter++; // TODO: Reuse handle based on a bunch of criterias?
+
+		StringBuffer<> searchPathString;
+		for (auto it=searchPaths; it->count; ++it)
+		{
+			if (it != searchPaths)
+				searchPathString.Append(';');
+			searchPathString.Append(*it);
+		}
+
+		bool symInitialized = false;
+
+		if (SymInitializeW(processHandle, searchPathString.data, FALSE))
+		{
+			symInitialized = true;
+			for (auto& mod : modules)
+			{
+				if (mod.name.empty())
+				{
+					mod.name = TC("<Unknown>");
+					continue;
+				}
+
+				auto res = SymLoadModuleExW(processHandle, NULL, mod.name.c_str(), NULL, mod.start, (DWORD)mod.size, NULL, 0);
+				if (res != 0 || GetLastError() == ERROR_SUCCESS)
+				{
+					mod.handled = true;
+					continue;
+				}
+
+				// Try something more?
+			}
+		}
+		auto symCleanup = MakeGuard([processHandle]() { SymCleanup(processHandle); });
+
+		for (auto& entry : entries)
+		{
+			out.Append(TC("\n   "));
+			if (entry.moduleIndex == ~0u)
+			{
+				out.Appendf(TC("<Unknown>: 0x%llx"), entry.memoryOffset);
+				continue;
+			}
+			auto& mod = modules[entry.moduleIndex];
+			if (mod.handled)
+			{
+				u8 buffer[1024];
+
+				auto& info = *(SYMBOL_INFOW*)buffer;
+				memset(&info, 0, sizeof(info));
+				info.SizeOfStruct = sizeof(info);
+				info.MaxNameLen = (sizeof(buffer) - sizeof(info))/sizeof(tchar);
+				DWORD64 displacement2 = 0;
+				bool gotSymbol = false;
+				if (SymFromAddrW(processHandle, mod.start + entry.memoryOffset, &displacement2, &info))
+				{
+					out.Appendf(TC("%s"), info.Name);
+					gotSymbol = true;
+				}
+
+				auto& line = *(IMAGEHLP_LINEW64*)buffer;
+				line.SizeOfStruct = sizeof(line);
+				DWORD displacement = 0;
+				bool gotLine = false;
+				if (SymGetLineFromAddrW64(processHandle, mod.start + entry.memoryOffset, &displacement, &line))
+				{
+					auto fileName = line.FileName;
+					if (auto lastSlash = TStrrchr(fileName, '\\'))
+						fileName = lastSlash + 1;
+					if (gotSymbol)
+						out.Appendf(TC(" ("));
+					out.Appendf(TC("%s:%u"), fileName, line.LineNumber);
+					if (gotSymbol)
+						out.Appendf(TC(")"));
+					gotLine = true;
+				}
+				if (gotSymbol || gotLine)
+					continue;
+			}
+			out.Appendf(TC("%s: +0x%llx"), mod.name.c_str(), entry.memoryOffset);
+		}
+#else
+
+		for (auto& entry : entries)
+		{
+			if (entry.moduleIndex == ~0u)
+			{
+				out.Appendf(TC("\n   <Unknown>: 0x%llx"), entry.memoryOffset);
+				continue;
+			}
+			auto& mod = modules[entry.moduleIndex];
+			
+			if (!mod.handled)
+			{
+				if (mod.name.empty())
+					mod.name = executable;
+				StringBuffer<256> buf(mod.name);
+				if (auto lastSlash = buf.Last(PathSeparator))
+					buf.Resize(lastSlash - buf.data);
+				for (auto it=searchPaths; it->count; ++it)
+				{
+					struct stat attr;
+					buf.Clear().Append(*it).EnsureEndsWithSlash().Append(mod.name);
+					if (stat(buf.data, &attr) == -1)
+						continue;
+					mod.name = buf.data;
+					break;
+				}
+				mod.handled  = true;
+			}
+
+			auto pg = MakeGuard([&]() { out.Appendf(TC("\n   %s: 0x%llx"), mod.name.c_str(), entry.memoryOffset); });
+
+			StringBuffer<1024> cmd;
+			cmd.Appendf("addr2line 0x%x", entry.memoryOffset).Append(" -e ").Append(mod.name).Append(" 2>/dev/null");
+
+			fflush(stdout);
+
+			FILE* fp = popen(cmd.data, "r");
+			if (!fp)
+				break;
+			auto cg = MakeGuard([&]() { fclose(fp); });
+
+			u32 index = 0;
+			u32 countBeforeCallstack = out.count;
+			errno = 0;
+			StringBuffer<1024> str;
+			if (!fgets(str.data, str.capacity, fp))
+				continue;
+			str.count = TStrlen(str.data);
+			if (str.StartsWith("addr2line"))
+				continue;
+			if (str.EndsWith("\n"))
+				str.Resize(str.count - 1);
+			if (!str.count || str[0] == ':' || str[0] == '?')
+				continue;
+			char* fileName = str.data;
+			if (char* lastSlash = strrchr(fileName, '/'))
+				fileName = lastSlash + 1;
+			if (strncmp(fileName, "function.h", 10) == 0 || strncmp(fileName, "invoke.h:", 9) == 0)
+				continue;
+
+			pg.Cancel();
+			if (out.capacity - out.count < strlen(fileName) + 5)
+				break;
+			out.Appendf(TC("\n   %s"), fileName);
+
+			int errnoCopy = errno;
+			int err = ferror(fp);
+		}
+#endif
+	}
 
 	void PrintContentionSummary(Logger& logger)
 	{
diff --git a/Engine/Source/Programs/UnrealBuildAccelerator/Common/Private/UbaProcess.cpp b/Engine/Source/Programs/UnrealBuildAccelerator/Common/Private/UbaProcess.cpp
--- a/Engine/Source/Programs/UnrealBuildAccelerator/Common/Private/UbaProcess.cpp
+++ b/Engine/Source/Programs/UnrealBuildAccelerator/Common/Private/UbaProcess.cpp
@@ -1306,6 +1306,12 @@ namespace uba
 		return true;
 	}
 
+	bool ProcessImpl::HandleResolveCallstack(BinaryReader& reader, BinaryWriter& writer)
+	{
+		m_session.GetSymbols(*this, reader, writer);
+		return true;
+	}
+
 	bool ProcessImpl::CreateTempFile(BinaryReader& reader, ProcHandle nativeProcessHandle, const tchar* application)
 	{
 		StringKey key = reader.ReadStringKey();
diff --git a/Engine/Source/Programs/UnrealBuildAccelerator/Common/Private/UbaSession.cpp b/Engine/Source/Programs/UnrealBuildAccelerator/Common/Private/UbaSession.cpp
--- a/Engine/Source/Programs/UnrealBuildAccelerator/Common/Private/UbaSession.cpp
+++ b/Engine/Source/Programs/UnrealBuildAccelerator/Common/Private/UbaSession.cpp
@@ -1824,7 +1824,8 @@ namespace uba
 		else
 			LoggerWithWriter(m_logger.m_writer).Info(TC("---- Starting trace ----"));
 
-		m_trace.StartWrite(traceName, m_detailedTrace ? 512*1024*1024 : 128*1024*1024);
+		if (!m_trace.StartWrite(traceName, m_detailedTrace ? 512*1024*1024 : 128*1024*1024))
+			return;
 
 		tchar buf[256];
 		if (!GetComputerNameW(buf, sizeof_array(buf)))
@@ -2370,9 +2371,12 @@ namespace uba
 	bool Session::MoveFile(MoveFileResponse& out, const MoveFileMessage& msg)
 	{
 		out.result = MoveFileExW(msg.fromName.data, msg.toName.data, msg.flags);
-		out.errorCode = ERROR_SUCCESS;
 		if (!out.result)
+		{
 			out.errorCode = GetLastError();
+			return true;
+		}
+		out.errorCode = ERROR_SUCCESS;
 		RegisterCreateFileForWrite(msg.toKey, msg.toName, true);
 		out.directoryTableSize = RegisterDeleteFile(msg.fromKey, msg.fromName);
 		return true;
@@ -2584,7 +2588,7 @@ namespace uba
 				Storage::GetMappingString(name, file.mappingHandle, 0);
 				SCOPED_WRITE_LOCK(m_fileMappingTableLookupLock, lookupLock);
 				auto insres = m_fileMappingTableLookup.try_emplace(file.key);
-				UBA_ASSERT(insres.second);
+				UBA_ASSERTF(insres.second, TC("Failed to add %s to file mapping, already exists"), file.name.c_str());
 				FileMappingEntry& entry = insres.first->second;
 				lookupLock.Leave();
 				SCOPED_WRITE_LOCK(entry.lock, entryCs);
@@ -2718,6 +2722,22 @@ namespace uba
 		return true;
 	}
 
+	bool Session::GetSymbols(Process& process, BinaryReader& reader, BinaryWriter& writer)
+	{
+		StringBuffer<256> detoursLibPath;
+		StringBuffer<256> alternativeLibPath;
+
+		detoursLibPath.Append(m_detoursLibrary.c_str()).Resize(detoursLibPath.Last(PathSeparator) - detoursLibPath.data);
+		GetAlternativeUbaPath(m_logger, alternativeLibPath, detoursLibPath, IsArm);
+
+		StringView searchPaths[3] = { detoursLibPath, alternativeLibPath, {} };
+
+		StringBuffer<16*1024> sb;
+		ParseCallstackInfo(sb, reader, process.GetStartInfo().application, searchPaths);
+		writer.WriteString(sb);
+		return true;
+	}
+
 	void Session::FileEntryAdded(StringKey fileNameKey, u64 lastWritten, u64 size)
 	{
 	}
diff --git a/Engine/Source/Programs/UnrealBuildAccelerator/Common/Private/UbaSessionClient.cpp b/Engine/Source/Programs/UnrealBuildAccelerator/Common/Private/UbaSessionClient.cpp
--- a/Engine/Source/Programs/UnrealBuildAccelerator/Common/Private/UbaSessionClient.cpp
+++ b/Engine/Source/Programs/UnrealBuildAccelerator/Common/Private/UbaSessionClient.cpp
@@ -253,8 +253,8 @@ namespace uba
 			return true;
 
 		#if PLATFORM_LINUX
-		if (EndsWith(startInfo.application, TStrlen(startInfo.application), "/sh"))
-			return true;
+		//if (EndsWith(startInfo.application, TStrlen(startInfo.application), "/sh")) // For some reason the local shell application on the helpers hang (need to investigate aws shell application)
+		//	return true;
 		#endif
 
 		outRealApplication.Clear();
@@ -440,7 +440,8 @@ namespace uba
 			if (EndsWith(si.application, TStrlen(si.application), "/sh"))
 			{
 				FlushWrittenFiles(*parent);
-				UpdateEnvironment(*parent, AsView("ShellScript"), false);
+				StackBinaryReader<SendMaxSize> reader;
+				SendUpdateDirectoryTable(reader);
 			}
 		}
 		#endif
@@ -1219,7 +1220,6 @@ namespace uba
 		}
 
 		CasKey detoursBinaryKey = reader.ReadCasKey();
-		CasKey detoursSymbolsKey = reader.ReadCasKey();
 
 		StringBuffer<> detoursFile;
 		if (m_downloadDetoursLib)
@@ -1233,20 +1233,6 @@ namespace uba
 			KeyToString dir(StringKeyZero);
 			if (!WriteBinFile(detoursFile, UBA_DETOURS_LIBRARY, detoursBinaryKey, dir, DefaultAttributes()))
 				return;
-
-			if (detoursSymbolsKey != CasKeyZero)
-			{
-				StringBuffer<128> symbolsFile(UBA_DETOURS_LIBRARY);
-				#if PLATFORM_WINDOWS
-				symbolsFile.Resize(symbolsFile.count - 3).Append("pdb");
-				#else
-				symbolsFile.Resize(symbolsFile.count - 2).Append("debug");
-				#endif
-				Storage::RetrieveResult result;
-				StringBuffer<> throwaway;
-				if (m_storage.RetrieveCasFile(result, AsCompressed(detoursSymbolsKey, false), symbolsFile.data))
-					WriteBinFile(throwaway, symbolsFile.data, detoursSymbolsKey, dir, DefaultAttributes());
-			}
 		}
 		else
 		{
@@ -1273,9 +1259,6 @@ namespace uba
 		m_shouldSendLogToServer = reader.ReadBool();
 		m_shouldSendTraceToServer = reader.ReadBool();
 
-		if (m_shouldSendLogToServer)
-			m_logToFile = true;
-
 		if (m_shouldSendTraceToServer)
 		{
 			//if (m_detailedTrace)
@@ -1773,7 +1756,7 @@ namespace uba
 					startInfo.rules = GetRules(startInfo);
 
 					StringBuffer<> logFile;
-					if (m_logToFile)
+					if (m_logToFile || (*startInfo.logFile && m_shouldSendLogToServer))
 					{
 						GetLogFileName(logFile, startInfo.logFile, startInfo.arguments, startInfo.processId);
 						startInfo.logFile = logFile.data;
@@ -2065,7 +2048,8 @@ namespace uba
 			outNextProcess.workingDir = reader.ReadString();
 			outNextProcess.description = reader.ReadString();
 			outNextProcess.logFile = reader.ReadString();
-			if (m_logToFile)
+
+			if (m_logToFile || (!outNextProcess.logFile.empty() && m_shouldSendLogToServer))
 			{
 				StringBuffer<512> logFile;
 				GetLogFileName(logFile, outNextProcess.logFile.c_str(), outNextProcess.arguments.c_str(), process.GetId());
@@ -2141,6 +2125,37 @@ namespace uba
 		return true;
 	}
 
+	bool SessionClient::GetSymbols(Process& process, BinaryReader& reader, BinaryWriter& writer)
+	{
+		StackBinaryWriter<SendMaxSize> msgWriter;
+		NetworkMessage msg(m_client, ServiceId, SessionMessageType_GetSymbols, msgWriter);
+		msgWriter.WriteU32(process.GetId());
+		msgWriter.WriteBytes(reader.GetPositionData(), reader.GetLeft());
+
+		BinaryReader responseReader(writer.GetData(), 0, writer.GetCapacityLeft());
+		if (!msg.Send(responseReader, m_stats.customMsg))
+			return false;
+		writer.AllocWrite(responseReader.GetPosition());
+		return true;
+
+#if 0
+		CasKey detoursSymbolsKey = msgReader.ReadCasKey();
+		if (detoursSymbolsKey == CasKeyZero)
+			return true;
+		StringBuffer<128> symbolsFile(UBA_DETOURS_LIBRARY);
+		#if PLATFORM_WINDOWS
+		symbolsFile.Resize(symbolsFile.count - 3).Append("pdb");
+		#else
+		symbolsFile.Resize(symbolsFile.count - 2).Append("debug");
+		#endif
+		Storage::RetrieveResult result;
+		StringBuffer<> throwaway;
+		if (m_storage.RetrieveCasFile(result, AsCompressed(detoursSymbolsKey, false), symbolsFile.data))
+			WriteBinFile(throwaway, symbolsFile.data, detoursSymbolsKey, KeyToString(StringKeyZero), DefaultAttributes());
+#endif
+	}
+
+
 	bool SessionClient::FlushWrittenFiles(ProcessImpl& process)
 	{
 		SCOPED_WRITE_LOCK(process.m_writtenFilesLock, lock);
diff --git a/Engine/Source/Programs/UnrealBuildAccelerator/Common/Private/UbaSessionServer.cpp b/Engine/Source/Programs/UnrealBuildAccelerator/Common/Private/UbaSessionServer.cpp
--- a/Engine/Source/Programs/UnrealBuildAccelerator/Common/Private/UbaSessionServer.cpp
+++ b/Engine/Source/Programs/UnrealBuildAccelerator/Common/Private/UbaSessionServer.cpp
@@ -628,7 +628,6 @@ namespace uba
 		CasKey clientDetoursKey = reader.ReadCasKey();
 
 		CasKey& detoursBinaryKey = m_detoursBinaryKey[isClientArm];
-		CasKey& detoursSymbolsKey = m_detoursSymbolsKey[isClientArm];
 		CasKey& agentBinaryKey = m_agentBinaryKey[isClientArm];
 
 		bool binAsVersion = clientAgentKey != CasKeyZero;
@@ -666,15 +665,9 @@ namespace uba
 					if (!m_storage.StoreCasFile(agentBinaryKey, dir.data, CasKeyZero, deferCreation, fileIsCompressed) || agentBinaryKey == CasKeyZero)
 					{
 						// This is hacky but uba binary is not copied by ubt anymore.
-						StringBuffer<> dir2(dir);
-						const tchar* engineDir = IsWindows ? TC("\\Engine\\") : TC("/Engine/");
-						const tchar* engineDirPos;
-						if (!dir2.Contains(engineDir, true, &engineDirPos))
-							return m_logger.Error(TC("Failed to find Engine dir in %s"), dir2.data);
-						const tchar* platformStr = TC("Win64");
-						dir2.Resize(engineDirPos - dir2.data + 8).Append(TC("Binaries")).Append(PathSeparator).Append(platformStr).Append(PathSeparator).Append("UnrealBuildAccelerator").Append(PathSeparator);
-						if (IsWindows)
-							dir2.Append(isClientArm ? TC("arm64") : TC("x64")).Append(PathSeparator);
+						StringBuffer<> dir2;
+						if (!GetAlternativeUbaPath(m_logger, dir2, dir, IsWindows && isClientArm))
+							return false;
 						dir2.Append(UBA_AGENT_EXECUTABLE);
 						if (!m_storage.StoreCasFile(agentBinaryKey, dir2.data, CasKeyZero, deferCreation, fileIsCompressed) || agentBinaryKey == CasKeyZero)
 							return m_logger.Error(TC("Failed to create cas for %s"), dir2.data);
@@ -683,16 +676,6 @@ namespace uba
 				dir.Resize(dirCount).Append(PathSeparator).Append(UBA_DETOURS_LIBRARY);
 				if (!m_storage.StoreCasFile(detoursBinaryKey, dir.data, CasKeyZero, deferCreation, fileIsCompressed) || detoursBinaryKey == CasKeyZero)
 					return m_logger.Error(TC("Failed to create cas for %s"), dir.data);
-
-#if UBA_DEBUG
-				#if PLATFORM_WINDOWS
-				dir.Resize(dir.count - 3).Append("pdb");
-				#else
-				dir.Resize(dir.count - 2).Append("debug");
-				#endif
-				if (!m_storage.StoreCasFile(detoursSymbolsKey, dir.data, CasKeyZero, deferCreation, fileIsCompressed) || detoursSymbolsKey == CasKeyZero)
-					return m_logger.Error(TC("Failed to create cas for %s"), dir.data);
-#endif
 			}
 		}
 
@@ -743,7 +726,6 @@ namespace uba
 			return false;
 
 		writer.WriteCasKey(detoursBinaryKey);
-		writer.WriteCasKey(detoursSymbolsKey);
 		writer.WriteBool(m_resetCas);
 		writer.WriteU32(sessionId);
 		writer.WriteU32(m_uiLanguage);
@@ -2035,6 +2017,11 @@ namespace uba
 		CasKey& lookupCasKey = m_nameToHashLookup[fileNameKey];
 		if (lookupCasKey == cachedInfo.casKey)
 			return;
+
+		#if 0
+		//m_debugLogger->Info(TC("NAMETOHASHADD    %s %s\n"), KeyToString(fileNameKey).data, CasKeyString(lookupCasKey).str);
+		#endif
+
 		lookupCasKey = cachedInfo.casKey;
 		BinaryWriter w(m_nameToHashTableMem.memory, m_nameToHashTableMem.writtenSize, NameToHashMemSize);
 		m_nameToHashTableMem.AllocateNoLock(sizeof(StringKey) + sizeof(CasKey), 1, TC("NameToHashTable"));
@@ -2271,4 +2258,34 @@ namespace uba
 	{
 		return HostRun(reader, writer);
 	}
+
+	bool SessionServer::HandleGetSymbols(const ConnectionInfo& connectionInfo, BinaryReader& reader, BinaryWriter& writer)
+	{
+		u32 processId = reader.ReadU32();
+		ProcessHandle h(GetProcess(processId));
+		if (!h.IsValid())
+			return m_logger.Error(TC("Failed to find process for id %u when receiving update environment message"), processId);
+
+		GetSymbols(*h.m_process, reader, writer);
+		return true;
+
+		#if 0
+		bool isClientArm = reader.ReadBool();
+		CasKey detoursSymbolsKey;
+		StringBuffer<> dir;
+		if (GetDirectoryOfCurrentModule(m_logger, dir))
+		{
+			bool fileIsCompressed = false;
+			bool deferCreation = true;
+			auto ChangeToSymbolExtension = [](StringBufferBase& str) -> StringBufferBase& { IsWindows ? str.Resize(str.count - 3).Append("pdb") : str.Resize(str.count - 2).Append("debug"); return str; };
+			if (!m_storage.StoreCasFile(detoursSymbolsKey, ChangeToSymbolExtension(dir).data, CasKeyZero, deferCreation, fileIsCompressed) || detoursSymbolsKey == CasKeyZero)
+			{
+				StringBuffer<> dir2;
+				if (GetAlternativeUbaPath(m_logger, dir2, dir, IsWindows && isClientArm))
+					m_storage.StoreCasFile(detoursSymbolsKey, ChangeToSymbolExtension(dir2.Append(UBA_DETOURS_LIBRARY)).data, CasKeyZero, deferCreation, fileIsCompressed);
+			}
+		}
+		writer.WriteCasKey(detoursSymbolsKey);
+		#endif
+	}
 }
diff --git a/Engine/Source/Programs/UnrealBuildAccelerator/Common/Public/UbaFile.h b/Engine/Source/Programs/UnrealBuildAccelerator/Common/Public/UbaFile.h
--- a/Engine/Source/Programs/UnrealBuildAccelerator/Common/Public/UbaFile.h
+++ b/Engine/Source/Programs/UnrealBuildAccelerator/Common/Public/UbaFile.h
@@ -78,6 +78,7 @@ namespace uba
 	u64 GetFileTimeAsTime(u64 fileTime);
 	u64 GetSecondsAsFileTime(u64 seconds);
 	bool GetCurrentDirectoryW(StringBufferBase& out);
+	bool GetAlternativeUbaPath(Logger& logger, StringBufferBase& out, StringView firstPath, bool isWindowsArm);
 
 
 	class DirectoryCache
diff --git a/Engine/Source/Programs/UnrealBuildAccelerator/Common/Public/UbaLogger.h b/Engine/Source/Programs/UnrealBuildAccelerator/Common/Public/UbaLogger.h
--- a/Engine/Source/Programs/UnrealBuildAccelerator/Common/Public/UbaLogger.h
+++ b/Engine/Source/Programs/UnrealBuildAccelerator/Common/Public/UbaLogger.h
@@ -10,6 +10,10 @@
 
 namespace uba
 {
+	class StringBufferBase;
+	struct BinaryReader;
+	struct StringView;
+
 	struct LogEntry
 	{
 		LogEntryType type;
@@ -86,5 +90,6 @@ namespace uba
 	Logger* StopDebugLogger(Logger* logger);
 	#endif
 
+	void ParseCallstackInfo(StringBufferBase& out, BinaryReader& reader, const tchar* executable, const StringView* searchPaths);
 	void PrintContentionSummary(class Logger& logger);
 }
diff --git a/Engine/Source/Programs/UnrealBuildAccelerator/Common/Public/UbaNetwork.h b/Engine/Source/Programs/UnrealBuildAccelerator/Common/Public/UbaNetwork.h
--- a/Engine/Source/Programs/UnrealBuildAccelerator/Common/Public/UbaNetwork.h
+++ b/Engine/Source/Programs/UnrealBuildAccelerator/Common/Public/UbaNetwork.h
@@ -16,7 +16,7 @@ namespace uba
 	
 	static constexpr u32 SystemNetworkVersion = 1339;
 	static constexpr u32 StorageNetworkVersion = 4;
-	static constexpr u32 SessionNetworkVersion = 40;
+	static constexpr u32 SessionNetworkVersion = 41;
 	static constexpr u32 CacheNetworkVersion = 5;
 
 	// Messages used over network between client/server (system, storage and session)
@@ -90,6 +90,7 @@ namespace uba
 		UBA_SESSION_MESSAGE(SHGetKnownFolderPath) \
 		UBA_SESSION_MESSAGE(DebugFileNotFoundError) \
 		UBA_SESSION_MESSAGE(HostRun) \
+		UBA_SESSION_MESSAGE(GetSymbols) \
 
 	enum SessionMessageType : u8 
 	{
diff --git a/Engine/Source/Programs/UnrealBuildAccelerator/Common/Public/UbaSession.h b/Engine/Source/Programs/UnrealBuildAccelerator/Common/Public/UbaSession.h
--- a/Engine/Source/Programs/UnrealBuildAccelerator/Common/Public/UbaSession.h
+++ b/Engine/Source/Programs/UnrealBuildAccelerator/Common/Public/UbaSession.h
@@ -152,6 +152,7 @@ namespace uba
 		virtual bool CustomMessage(Process& process, BinaryReader& reader, BinaryWriter& writer);
 		virtual bool SHGetKnownFolderPath(Process& process, BinaryReader& reader, BinaryWriter& writer);
 		virtual bool HostRun(BinaryReader& reader, BinaryWriter& writer);
+		virtual bool GetSymbols(Process& process, BinaryReader& reader, BinaryWriter& writer);
 		virtual void FileEntryAdded(StringKey fileNameKey, u64 lastWritten, u64 size);
 		virtual bool FlushWrittenFiles(ProcessImpl& process);
 		virtual bool UpdateEnvironment(ProcessImpl& process, const StringView& reason, bool resetStats);
diff --git a/Engine/Source/Programs/UnrealBuildAccelerator/Common/Public/UbaSessionClient.h b/Engine/Source/Programs/UnrealBuildAccelerator/Common/Public/UbaSessionClient.h
--- a/Engine/Source/Programs/UnrealBuildAccelerator/Common/Public/UbaSessionClient.h
+++ b/Engine/Source/Programs/UnrealBuildAccelerator/Common/Public/UbaSessionClient.h
@@ -68,6 +68,7 @@ namespace uba
 		virtual bool CustomMessage(Process& process, BinaryReader& reader, BinaryWriter& writer) override;
 		virtual bool SHGetKnownFolderPath(Process& process, BinaryReader& reader, BinaryWriter& writer) override;
 		virtual bool HostRun(BinaryReader& reader, BinaryWriter& writer) override;
+		virtual bool GetSymbols(Process& process, BinaryReader& reader, BinaryWriter& writer);
 		virtual bool FlushWrittenFiles(ProcessImpl& process) override;
 		virtual bool UpdateEnvironment(ProcessImpl& process, const StringView& reason, bool resetStats) override;
 		virtual bool LogLine(ProcessImpl& process, const tchar* line, LogEntryType logType) override;
diff --git a/Engine/Source/Programs/UnrealBuildAccelerator/Common/Public/UbaSessionServer.h b/Engine/Source/Programs/UnrealBuildAccelerator/Common/Public/UbaSessionServer.h
--- a/Engine/Source/Programs/UnrealBuildAccelerator/Common/Public/UbaSessionServer.h
+++ b/Engine/Source/Programs/UnrealBuildAccelerator/Common/Public/UbaSessionServer.h
@@ -136,7 +136,6 @@ namespace uba
 
 		ReaderWriterLock m_binKeysLock;
 		CasKey m_detoursBinaryKey[2];
-		CasKey m_detoursSymbolsKey[2];
 		CasKey m_agentBinaryKey[2];
 
 		struct ClientSession
diff --git a/Engine/Source/Programs/UnrealBuildAccelerator/Core/Private/UbaApplicationRules.cpp b/Engine/Source/Programs/UnrealBuildAccelerator/Core/Private/UbaApplicationRules.cpp
--- a/Engine/Source/Programs/UnrealBuildAccelerator/Core/Private/UbaApplicationRules.cpp
+++ b/Engine/Source/Programs/UnrealBuildAccelerator/Core/Private/UbaApplicationRules.cpp
@@ -217,11 +217,6 @@ namespace uba
 	class ApplicationRulesClang : public ApplicationRules
 	{
 	public:
-		virtual bool EnableVectoredExceptionHandler() const override
-		{
-			return true;
-		}
-
 		virtual bool IsExitCodeSuccess(u32 exitCode) const override
 		{
 			return exitCode == 0;
diff --git a/Engine/Source/Programs/UnrealBuildAccelerator/Core/Private/UbaPlatform.cpp b/Engine/Source/Programs/UnrealBuildAccelerator/Core/Private/UbaPlatform.cpp
--- a/Engine/Source/Programs/UnrealBuildAccelerator/Core/Private/UbaPlatform.cpp
+++ b/Engine/Source/Programs/UnrealBuildAccelerator/Core/Private/UbaPlatform.cpp
@@ -3,15 +3,21 @@
 #include "UbaPlatform.h"
 #include "UbaProcessStats.h"
 
-#if !PLATFORM_WINDOWS
+#if PLATFORM_WINDOWS
+#include <conio.h>
+#include <psapi.h>
+#else
 #include <execinfo.h>
 #include <limits.h>
 #include <dlfcn.h>
+#if PLATFORM_LINUX
+#include <link.h>
 #else
-#include <conio.h>
-#include <psapi.h>
+#include <mach-o/dyld.h>
+#endif
 #endif
 
+
 namespace uba
 {
 	KernelStats g_kernelStatsInternal;
@@ -135,166 +141,224 @@ namespace uba
 		#endif
 	}
 
-	UBA_NOINLINE void WriteAssertInfo(StringBufferBase& out, const tchar* text, const char* file, u32 line, const char* expr, u32 skipCallstack = 0)
+	void WriteAssertInfo(StringBufferBase& out, const tchar* text, const char* file, u32 line, const char* expr)
+	{
+		out.Append(TC("ASSERT: ")).Append(*text ? text : TC("Unknown"));
+		if (!*text && expr && *expr && strcmp(expr, "false") != 0)
+			out.Appendf(TC("\n EXPR: ")).Append(expr);
+		if (file && *file)
+			out.Append(TC("\n LOCATION: ")).Append(file).Append(':').AppendValue(line);
+	}
+
+	UBA_NOINLINE bool WriteCallstackInfo(BinaryWriter& out, u32 skipCallstack, void* contextPtr)
 	{
 #if PLATFORM_WINDOWS
-		if (text)
-			out.Append(L"ASSERT: ").Append(text).Append(L"\r\n");
-		if (file)
-			out.Appendf(L"%hs:%u (%hs)", file, line, expr);
-
-		typedef USHORT(WINAPI* CaptureStackBackTraceType)(__in ULONG, __in ULONG, __out PVOID*, __out_opt PULONG);
-		static CaptureStackBackTraceType func = (CaptureStackBackTraceType)(GetProcAddress(GetModuleHandleW(L"ntdll.dll"), "RtlCaptureStackBackTrace"));
-		if (func != NULL)
-		{
-			struct ModuleRec { u64 start; HMODULE handle; };
-			Map<u64, ModuleRec> moduleEndAddresses;
+		out.WriteBool(IsRunningWine());
 
-			HMODULE loadedModules[512];
-			DWORD needed = 0;
-			if (EnumProcessModules(GetCurrentProcess(), loadedModules, sizeof(loadedModules), &needed))
+		const int MaxDepth = 16;// 62;
+		void* callers[MaxDepth];
+		u32 count = 0;//RtlCaptureStackBackTrace(0, MaxDepth, callers, NULL); // For some reason this does not work on wine
+		if (!count)
+		{
+			CONTEXT context;
+			if (contextPtr)
 			{
-				for (u32 i = 0, e = needed / sizeof(HMODULE); i != e; ++i)
-				{
-					MODULEINFO mi;
-					GetModuleInformation(GetCurrentProcess(), loadedModules[i], &mi, sizeof(mi));
-					ModuleRec& rec = moduleEndAddresses[u64(mi.lpBaseOfDll) + mi.SizeOfImage];
-					rec.handle = loadedModules[i];
-					rec.start = u64(mi.lpBaseOfDll);
-				}
+				context = *(CONTEXT*)contextPtr;
+				skipCallstack = 0;
 			}
+			else
+				RtlCaptureContext(&context);
 
-			const int kMaxCallers = 16;// 62;
-			void* callers[kMaxCallers];
-			if (u32 count = (func)(0, kMaxCallers, callers, NULL))
+			if (context.Rip && context.Rsp)
 			{
-				if (file)
-					out.Appendf(L"\n\n");
-				out.Appendf(L"  Callstack:");
-				for (u32 i = skipCallstack; i < count; i++)
+				UNWIND_HISTORY_TABLE unwindHistoryTable;
+				RtlZeroMemory(&unwindHistoryTable, sizeof(UNWIND_HISTORY_TABLE));
+				for (u32 i=0; i<MaxDepth && context.Rip; ++i)
 				{
-					tchar str[1024];
-					auto addr = u64(callers[i]);
-
-					auto findIt = moduleEndAddresses.lower_bound(addr);
-					if (findIt != moduleEndAddresses.end() && addr >= findIt->second.start)
+					callers[i] = (void*)context.Rip;
+					count = i;
+					ULONG64 ImageBase = 0;
+					if (PRUNTIME_FUNCTION RuntimeFunction = RtlLookupFunctionEntry(context.Rip, &ImageBase, &unwindHistoryTable); RuntimeFunction)
 					{
-						if (GetModuleFileNameW(findIt->second.handle, str, sizeof_array(str)))
-						{
-							const tchar* moduleName = str;
-							if (const tchar* lastSeparator = TStrrchr(str, PathSeparator))
-								moduleName = lastSeparator + 1;
-							out.Appendf(L"\n   %s: +0x%llx", moduleName, addr - findIt->second.start);
-						}
+						KNONVOLATILE_CONTEXT_POINTERS nvcontext;
+						RtlZeroMemory(&nvcontext, sizeof(KNONVOLATILE_CONTEXT_POINTERS));
+						PVOID HandlerData = nullptr;
+						ULONG64 EstablisherFrame = 0;
+						RtlVirtualUnwind(0, ImageBase, context.Rip, RuntimeFunction, &context, &HandlerData, &EstablisherFrame, &nvcontext);
 					}
 					else
-						out.Appendf(L"\n   <Unknown>: 0x%llx", callers[i]);
+					{
+						context.Rip = (*(u64*)context.Rsp);
+						context.Rsp += 8;
+					}
 				}
 			}
 		}
-#else
-		out.Append(TC("ASSERT: ")).Append(*text ? text : TC("Unknown")).Append('\n');
-		if (!*text && expr && strcmp(expr, "false") != 0)
-			out.Appendf(TC(" EXPR: %s\n"), expr);
-		if (file && *file)
-			out.Appendf(TC(" LOCATION: %s:%u"), file, line);
 
-		constexpr u32 maxCallers = 100;
-		void* buffer[maxCallers];
-		int nptrs = backtrace(buffer, maxCallers);
+		struct ModuleRec { u64 start; u64 size; HMODULE handle; u32 index; };
+		Map<u64, ModuleRec> moduleEndAddresses;
 
-		if (true)
+		HMODULE loadedModules[512];
+		DWORD needed = 0;
+		if (EnumProcessModules(GetCurrentProcess(), loadedModules, sizeof(loadedModules), &needed))
 		{
-			bool tryResolve = false;
-			char progName[256];
+			for (u32 i = 0, e = needed / sizeof(HMODULE); i != e; ++i)
+			{
+				MODULEINFO mi;
+				GetModuleInformation(GetCurrentProcess(), loadedModules[i], &mi, sizeof(mi));
+				ModuleRec& rec = moduleEndAddresses[u64(mi.lpBaseOfDll) + mi.SizeOfImage];
+				rec.handle = loadedModules[i];
+				rec.start = u64(mi.lpBaseOfDll);
+				rec.size = mi.SizeOfImage;
+				rec.index = ~0u;
+			}
+		}
 
-			u64 baseOffset = 0;
-			Dl_info info;
-			if (dladdr((void*)&WriteAssertInfo, &info))
+		Vector<ModuleRec*> usedModules;
+
+		out.Write7BitEncoded(count < skipCallstack ? 0 : count - skipCallstack);
+
+		for (u32 i = skipCallstack; i < count; i++)
+		{
+			auto addr = u64(callers[i]);
+
+			auto findIt = moduleEndAddresses.lower_bound(addr);
+			if (findIt != moduleEndAddresses.end() && addr >= findIt->second.start)
 			{
-				if (strstr(info.dli_fname, "UbaDetours"))
+				ModuleRec& rec = findIt->second;
+				if (rec.index == ~0u)
 				{
-					baseOffset = u64(info.dli_fbase);
-					strcpy(progName, info.dli_fname);
-					tryResolve = true;
+					rec.index = u32(usedModules.size());
+					usedModules.push_back(&rec);
 				}
+				out.Write7BitEncoded(rec.index);
+				out.Write7BitEncoded(addr - rec.start);
 			}
-			if (!tryResolve)
+			else
 			{
-				int res = readlink("/proc/self/exe", progName, sizeof_array(progName));
-				if (res > 0)
-				{
-					progName[res] = 0;
-					tryResolve = true;
-				}
+				out.Write7BitEncoded(~0u);
+				out.Write7BitEncoded(addr);
 			}
+		}
 
-			if (tryResolve)
+		out.Write7BitEncoded(usedModules.size());
+		for (auto* rec : usedModules)
+		{
+			out.Write7BitEncoded(rec->start);
+			out.Write7BitEncoded(rec->size);
+			tchar str[1024];
+			if (GetModuleFileNameW(rec->handle, str, sizeof_array(str)))
 			{
-				fflush(stdout);
-				char str[2048];
-				strcpy_s(str, sizeof(str), "addr2line");
-				for (int i = 0; i < nptrs; ++i)
-				{
-					char v[32];
-					u64 offset = u64(buffer[i]);
-					if (offset > baseOffset)
-						offset -= baseOffset;
-					snprintf(v, sizeof_array(v), " %llx", offset);
-					strcat_s(str, sizeof(str), v);
-				}
-				strcat_s(str, sizeof(str), " -e ");
+				const tchar* moduleName = str;
+				if (const tchar* lastSeparator = TStrrchr(str, PathSeparator))
+					moduleName = lastSeparator + 1;
+				out.WriteString(moduleName);
+			}
+			else
+			{
+				out.WriteString(TC(""));
+			}
+		}
+		return true;
+#else
 
-				strcat_s(str, sizeof(str), progName);
-				fflush(stdout);
+		struct ModuleRec { u64 start; u64 size; TString name; u32 index; };
+		Map<u64, ModuleRec> moduleEndAddresses;
 
-				if (FILE* fp = popen(str, "r"))
+#if PLATFORM_LINUX
+		dl_iterate_phdr([](struct dl_phdr_info* info, size_t size, void* data)
+			{
+				auto& moduleEndAddresses = *(Map<u64, ModuleRec>*)data;
+				if (!info->dlpi_name)
+					return 0;
+				u64 base = info->dlpi_addr;
+				u64 end = base;
+				for (int i = 0; i < info->dlpi_phnum; i++)
 				{
-					out.Append('\n');
-					u32 index = 0;
-					u32 countBeforeCallstack = out.count;
-					errno = 0;
-					while (fgets(str, sizeof(str), fp) != NULL)
-					{
-						if (str[0] == '?')
-							break;
-						if (index++ <= skipCallstack)
-							continue;
-						char* fileName = str;
-						char* lastSlash = strrchr(str, '/');
-						if (lastSlash)
-							fileName = lastSlash + 1;
-						if (strncmp(fileName, "function.h", 10) == 0 || strncmp(fileName, "invoke.h:", 9) == 0)
-							continue;
-						if (out.capacity - out.count < strlen(fileName) + 5)
-							break;
-						out.Appendf(TC(" %s"), fileName);
-					}
-					int errnoCopy = errno;
-					int err = ferror(fp);
-					fclose(fp);
-					if (err)
-						out.Appendf(TC("   fgets failed for %s (%s)\n"), progName, strerror(errnoCopy));
-					else if (countBeforeCallstack != out.count)
-						return;
+					if (info->dlpi_phdr[i].p_type != PT_LOAD)
+						continue;
+					u64 start = base + info->dlpi_phdr[i].p_vaddr;
+					end = Max(end, start + info->dlpi_phdr[i].p_memsz);
 				}
+
+				ModuleRec& rec = moduleEndAddresses[end];
+				rec.start = base;
+				rec.size = end - base;
+				rec.name = info->dlpi_name;
+				rec.index = ~0u;
+				return 0;
+			}, &moduleEndAddresses);
+#else
+		for (u32 i = 0, e = _dyld_image_count(); i < e; i++)
+		{
+			const char* name = _dyld_get_image_name(i);
+			if (!name)
+				continue;
+			auto hdr = _dyld_get_image_header(i);
+			size_t imageSize = 0;
+			if (hdr->magic == MH_MAGIC_64)
+			{
+				auto hdr64 = (const mach_header_64 *)hdr;
+				auto cmd = (const load_command*)(hdr64 + 1);
+				for (uint32_t i = 0; i < hdr64->ncmds; i++, cmd = (const load_command*)((const char *)cmd + cmd->cmdsize))
+					if (cmd->cmd == LC_SEGMENT_64)
+						imageSize += ((const segment_command_64*)cmd)->vmsize;
+			}
+			else if (hdr->magic == MH_MAGIC)
+			{
+				auto hdr32 = hdr;
+				auto cmd = (const load_command*)(hdr32 + 1);
+				for (uint32_t i = 0; i < hdr32->ncmds; i++, cmd = (const load_command*)((const char *)cmd + cmd->cmdsize))
+					if (cmd->cmd == LC_SEGMENT)
+						imageSize += ((const segment_command*)cmd)->vmsize;
 			}
+
+			ModuleRec& rec = moduleEndAddresses[end];
+			rec.start = u64(hdr) + _dyld_get_image_vmaddr_slide(i);
+			rec.size = imageSize;
+			rec.name = name;
+			rec.index = ~0u;
 		}
+#endif
 
-		char** strings = backtrace_symbols(buffer, nptrs);
-		if (strings == NULL)
+		constexpr u32 maxCallers = 100;
+		void* buffer[maxCallers];
+		int nptrs = backtrace(buffer, maxCallers);
+
+		Vector<ModuleRec*> usedModules;
+
+		out.Write7BitEncoded(nptrs - skipCallstack);
+
+		for (int i = skipCallstack; i < nptrs; ++i)
 		{
-			out.Append(TC("Failed to get symbols"));
-			return;
+			u64 addr = u64(buffer[i]);
+			auto findIt = moduleEndAddresses.lower_bound(addr);
+			if (findIt != moduleEndAddresses.end() && addr >= findIt->second.start)
+			{
+				ModuleRec& rec = findIt->second;
+				if (rec.index == ~0u)
+				{
+					rec.index = u32(usedModules.size());
+					usedModules.push_back(&rec);
+				}
+				out.Write7BitEncoded(rec.index);
+				out.Write7BitEncoded(addr - rec.start);
+			}
+			else
+			{
+				out.Write7BitEncoded(~0u);
+				out.Write7BitEncoded(addr);
+			}
 		}
-		for (int i = 0; i < nptrs; ++i)
+
+		out.Write7BitEncoded(usedModules.size());
+		for (auto* rec : usedModules)
 		{
-			if (char* absoluteAddressStart = strrchr(strings[i], '['))
-				if (absoluteAddressStart > strings[i] && absoluteAddressStart[-1] == ' ')
-					absoluteAddressStart[-1] = 0;
-			out.Appendf(TC("  %s\n"), strings[i]);
+			out.Write7BitEncoded(rec->start);
+			out.Write7BitEncoded(rec->size);
+			out.WriteString(rec->name);
 		}
-		free(strings);
+		return false;
 #endif
 	}
 
@@ -412,6 +476,79 @@ namespace uba
 		#endif
 	}
 
+#if PLATFORM_WINDOWS
+	void* g_startOfCurrentLibrary;
+	void* g_endOfCurrentLibrary;
+	LONG UbaExceptionHandler(_EXCEPTION_POINTERS* ExceptionInfo)
+	{
+		auto record = ExceptionInfo->ExceptionRecord;
+		u32 exceptionCode = record->ExceptionCode;
+
+		if (exceptionCode != EXCEPTION_STACK_OVERFLOW && exceptionCode != EXCEPTION_ACCESS_VIOLATION)
+			return EXCEPTION_CONTINUE_SEARCH;
+
+		if (!g_startOfCurrentLibrary)
+		{
+			HMODULE currentModule;
+			GetModuleHandleExW(GET_MODULE_HANDLE_EX_FLAG_FROM_ADDRESS, (LPCWSTR)&g_startOfCurrentLibrary, &currentModule);
+			MODULEINFO mi;
+			GetModuleInformation(GetCurrentProcess(), currentModule, &mi, sizeof(mi));
+			g_startOfCurrentLibrary = mi.lpBaseOfDll;
+			g_endOfCurrentLibrary = (u8*)mi.lpBaseOfDll + mi.SizeOfImage;
+		}
+
+		if (record->ExceptionAddress < g_startOfCurrentLibrary || record->ExceptionAddress >= g_endOfCurrentLibrary)
+			return EXCEPTION_CONTINUE_SEARCH;
+
+		//CONTEXT context;
+		//RtlCaptureContext(&context);
+		//CONTEXT* contextPtr = &context;
+		CONTEXT* contextPtr = ExceptionInfo->ContextRecord;
+
+		StringBuffer<> text;
+		text.Appendf(L"Unhandled Exception (Code: 0x%x)", exceptionCode);
+		UbaAssert(text.data, nullptr, 0, nullptr, true, exceptionCode, contextPtr);
+		return EXCEPTION_EXECUTE_HANDLER;
+	}
+#else
+	static const int g_allSignals[] = {
+		//#ifdef SIGHUP
+		//SIGHUP,
+		//#endif
+		//#ifdef SIGQUIT
+		//SIGQUIT,
+		//#endif
+		//#ifdef SIGTRAP
+		//SIGTRAP,
+		//#endif
+		//#ifdef SIGIO
+		//SIGIO,
+		//#endif
+		//SIGABRT,
+		//SIGFPE,
+		SIGILL,
+		//SIGINT,
+		SIGSEGV,
+		//SIGTERM 
+	};
+	static void segfault_sigaction(int signal) { uba::UbaAssert("Segmentation fault", "", 0, "", true, -1, nullptr); }
+#endif
+
+	void AddExceptionHandler()
+	{
+#if PLATFORM_WINDOWS
+		::AddVectoredExceptionHandler(TRUE, UbaExceptionHandler);
+#else
+		struct sigaction sa;
+		memset(&sa, 0, sizeof(struct sigaction));
+		sigemptyset(&sa.sa_mask);
+		sa.sa_handler = segfault_sigaction;
+		sa.sa_flags = 0;// SA_SIGINFO;
+		for (auto signal : g_allSignals)
+			sigaction(signal, &sa, NULL);
+#endif
+	}
+
 
 #if !PLATFORM_WINDOWS
 	void strcpy_s(tchar* dest, u64 destCapacity, const tchar* source)
diff --git a/Engine/Source/Programs/UnrealBuildAccelerator/Core/Private/UbaSynchronization.cpp b/Engine/Source/Programs/UnrealBuildAccelerator/Core/Private/UbaSynchronization.cpp
--- a/Engine/Source/Programs/UnrealBuildAccelerator/Core/Private/UbaSynchronization.cpp
+++ b/Engine/Source/Programs/UnrealBuildAccelerator/Core/Private/UbaSynchronization.cpp
@@ -114,6 +114,15 @@ namespace uba
 		#endif
 	}
 
+	bool ReaderWriterLock::TryEnterWrite()
+	{
+		#if PLATFORM_WINDOWS
+		return TryAcquireSRWLockExclusive((SRWLOCK*)&data) != 0;
+		#else
+		return pthread_rwlock_trywrlock((pthread_rwlock_t*)data) == 0;
+		#endif
+	}
+
 	void ReaderWriterLock::LeaveWrite()
 	{
 		#if PLATFORM_WINDOWS
diff --git a/Engine/Source/Programs/UnrealBuildAccelerator/Core/Public/UbaApplicationRules.h b/Engine/Source/Programs/UnrealBuildAccelerator/Core/Public/UbaApplicationRules.h
--- a/Engine/Source/Programs/UnrealBuildAccelerator/Core/Public/UbaApplicationRules.h
+++ b/Engine/Source/Programs/UnrealBuildAccelerator/Core/Public/UbaApplicationRules.h
@@ -90,13 +90,6 @@ namespace uba
 			return !IsRarelyRead(file);
 		}
 
-		// Enable vectored exception handler. This can't be enabled for all processes because some of them actually allow
-		// access violations etc and then catch them.
-		virtual bool EnableVectoredExceptionHandler() const
-		{
-			return false;
-		}
-
 		// Process is allowed to detour crt allocator and use mimalloc
 		virtual bool AllowMiMalloc() const
 		{
@@ -106,7 +99,7 @@ namespace uba
 		// Can be used to disallow libraries we don't want loaded
 		virtual bool AllowLoadLibrary(const tchar* libraryName) const
 		{
-			if (Contains(libraryName, TC("nvinject.dll")))
+			if (Contains(libraryName, TC("nvinject.dll")) || Contains(libraryName, TC("nviewH64.dll")))
 				return false;
 			return true;
 		}
diff --git a/Engine/Source/Programs/UnrealBuildAccelerator/Core/Public/UbaPlatform.h b/Engine/Source/Programs/UnrealBuildAccelerator/Core/Public/UbaPlatform.h
--- a/Engine/Source/Programs/UnrealBuildAccelerator/Core/Public/UbaPlatform.h
+++ b/Engine/Source/Programs/UnrealBuildAccelerator/Core/Public/UbaPlatform.h
@@ -59,13 +59,17 @@
 
 namespace uba
 {
+	class StringBufferBase;
+	struct BinaryReader;
+	struct BinaryWriter;
+
 	#ifndef sizeof_array
 	#define	sizeof_array(array) int(sizeof(array)/sizeof(array[0]))
 	#endif
 
 	#if UBA_DEBUG
-	#define UBA_ASSERT(x) do { if (x) break; uba::UbaAssert(TC(""), __FILE__, __LINE__, #x, 543221); } while(false)
-	#define UBA_ASSERTF(x, fmt, ...) do { UBA_FMT_CHECK(fmt, ##__VA_ARGS__); if (x) break; uba::StringBuffer<> _buf; _buf.Appendf(fmt, ##__VA_ARGS__); uba::UbaAssert(_buf.data, __FILE__, __LINE__, #x, 543221); } while(false)
+	#define UBA_ASSERT(x) do { if (x) break; uba::UbaAssert(TC(""), __FILE__, __LINE__, #x, true, 543221, nullptr); } while(false)
+	#define UBA_ASSERTF(x, fmt, ...) do { UBA_FMT_CHECK(fmt, ##__VA_ARGS__); if (x) break; uba::StringBuffer<> _buf; _buf.Appendf(fmt, ##__VA_ARGS__); uba::UbaAssert(_buf.data, __FILE__, __LINE__, #x, true, 543221, nullptr); } while(false)
 	#define UBA_ASSERT_MESSAGEBOX 0
 	#else
 	#define UBA_ASSERT(x) do { } while(false)
@@ -75,8 +79,9 @@ namespace uba
 
 	#define UBA_NOT_IMPLEMENTED(x) UBA_ASSERTF(false, TC("%s not implemented!"), #x);
 
-	void WriteAssertInfo(class StringBufferBase& out, const tchar* text, const char* file, u32 line, const char* expr, u32 skipCallstack);
-	ANALYSIS_NORETURN void UbaAssert(const tchar* text, const char* file, u32 line, const char* expr, u32 terminateCode);
+	bool WriteCallstackInfo(BinaryWriter& out, u32 skipCallstack, void* context);
+	void WriteAssertInfo(StringBufferBase& out, const tchar* text, const char* file, u32 line, const char* expr);
+	ANALYSIS_NORETURN void UbaAssert(const tchar* text, const char* file, u32 line, const char* expr, bool allowTerminate, u32 terminateCode, void* context);
 	ANALYSIS_NORETURN void FatalError(u32 code, const tchar* format, ...);
 	using CustomAssertHandler = void(const tchar* text);
 	void SetCustomAssertHandler(CustomAssertHandler* handler);
@@ -110,6 +115,7 @@ namespace uba
 
 	void ElevateCurrentThreadPriority();
 	void PrefetchVirtualMemory(const void* mem, u64 size);
+	void AddExceptionHandler();
 
 #if PLATFORM_WINDOWS
 	inline constexpr bool CaseInsensitiveFs = true;
diff --git a/Engine/Source/Programs/UnrealBuildAccelerator/Core/Public/UbaProtocol.h b/Engine/Source/Programs/UnrealBuildAccelerator/Core/Public/UbaProtocol.h
--- a/Engine/Source/Programs/UnrealBuildAccelerator/Core/Public/UbaProtocol.h
+++ b/Engine/Source/Programs/UnrealBuildAccelerator/Core/Public/UbaProtocol.h
@@ -38,6 +38,7 @@ namespace uba
 		UBA_PROCESS_MESSAGE(SHGetKnownFolderPath) \
 		UBA_PROCESS_MESSAGE(RpcCommunication) \
 		UBA_PROCESS_MESSAGE(HostRun) \
+		UBA_PROCESS_MESSAGE(ResolveCallstack) \
 
 	enum MessageType : u8
 	{
diff --git a/Engine/Source/Programs/UnrealBuildAccelerator/Core/Public/UbaSynchronization.h b/Engine/Source/Programs/UnrealBuildAccelerator/Core/Public/UbaSynchronization.h
--- a/Engine/Source/Programs/UnrealBuildAccelerator/Core/Public/UbaSynchronization.h
+++ b/Engine/Source/Programs/UnrealBuildAccelerator/Core/Public/UbaSynchronization.h
@@ -72,6 +72,8 @@ namespace uba
 		void EnterWrite();
 		void LeaveWrite();
 
+		bool TryEnterWrite();
+
 	private:
 
 		#if PLATFORM_WINDOWS
diff --git a/Engine/Source/Programs/UnrealBuildAccelerator/Core/Public/UbaWinBinDependencyParser.h b/Engine/Source/Programs/UnrealBuildAccelerator/Core/Public/UbaWinBinDependencyParser.h
--- a/Engine/Source/Programs/UnrealBuildAccelerator/Core/Public/UbaWinBinDependencyParser.h
+++ b/Engine/Source/Programs/UnrealBuildAccelerator/Core/Public/UbaWinBinDependencyParser.h
@@ -121,9 +121,9 @@ namespace uba
 		}
 		__except(GetExceptionCode() == EXCEPTION_ACCESS_VIOLATION ? EXCEPTION_EXECUTE_HANDLER : EXCEPTION_CONTINUE_SEARCH)
 		{
-			StringBuffer<> buf;
-			buf.Appendf(L"Access violation reading %s", fileName);
-			UbaAssert(buf.data, __FILE__, __LINE__, "", GetExceptionCode());
+			//StringBuffer<> buf;
+			//buf.Appendf(L"Access violation reading %s", fileName);
+			//UbaAssert(buf.data, __FILE__, __LINE__, "", true, GetExceptionCode(), GetExceptionInformation()->ContextRecord);
 			return false;
 		}
 		return true;
diff --git a/Engine/Source/Programs/UnrealBuildAccelerator/Detours/Private/PosixOS/UbaDetoursFunctionsPosix.cpp b/Engine/Source/Programs/UnrealBuildAccelerator/Detours/Private/PosixOS/UbaDetoursFunctionsPosix.cpp
--- a/Engine/Source/Programs/UnrealBuildAccelerator/Detours/Private/PosixOS/UbaDetoursFunctionsPosix.cpp
+++ b/Engine/Source/Programs/UnrealBuildAccelerator/Detours/Private/PosixOS/UbaDetoursFunctionsPosix.cpp
@@ -1893,7 +1893,7 @@ void UnsupportedHostRun(char* const argv[], const char* msg)
 {
 	StringBuffer<4096> command;
 	FlattenArgs(command, argv);
-	UbaAssert(msg, __FILE__, __LINE__, command.data, 1999);
+	UbaAssert(msg, __FILE__, __LINE__, command.data, true, 1999, nullptr);
 }
 
 int shared_posix_spawn(pid_t* pid, const char* path, const posix_spawn_file_actions_t* file_actions, const posix_spawnattr_t* attrp, char* const argv[], char* const envp[])
@@ -1915,7 +1915,7 @@ int shared_posix_spawn(pid_t* pid, const char* path, const posix_spawn_file_acti
 	if (strstr(path, "xcode-select"))
 	{
 		if (strcmp(argv[1], "--print-path") != 0)
-			UbaAssert("xcode-select only supported with --print-path", __FILE__, __LINE__, "", 1999);
+			UbaAssert("xcode-select only supported with --print-path", __FILE__, __LINE__, "", true, 1999, nullptr);
 		if (!ExecuteHostRun(result, argv))
 			return -1;
 		return SpawnEcho(result.data, pid,  file_actions, attrp, envp);
@@ -2723,12 +2723,18 @@ namespace uba
 	}
 	#endif
 
-	ANALYSIS_NORETURN void UbaAssert(const tchar* text, const char* file, u32 line, const char* expr, u32 terminateCode)
+	ANALYSIS_NORETURN void UbaAssert(const tchar* text, const char* file, u32 line, const char* expr, bool allowTerminate, u32 terminateCode, void* context)
 	{
 		SuppressDetourScope s;
-		StringBuffer<8*1024> b;
-		WriteAssertInfo(b, text, file, line, expr, 1);
-		Rpc_WriteLog(b.data, b.count, true, true);
+		static CriticalSection cs;
+		ScopedCriticalSection scs(cs);
+		static auto& sb =  *new StringBuffer<8*1024>;
+		WriteAssertInfo(sb, text, file, line, expr);
+		Rpc_ResolveCallstack(sb, 3, context);
+		Rpc_WriteLog(sb.data, sb.count, true, true);
+
+		if (!allowTerminate)
+			return;
 
 		BinaryWriter writer;
 		writer.WriteByte(MessageType_Exit);
diff --git a/Engine/Source/Programs/UnrealBuildAccelerator/Detours/Private/PosixOS/UbaDetoursMainPosix.cpp b/Engine/Source/Programs/UnrealBuildAccelerator/Detours/Private/PosixOS/UbaDetoursMainPosix.cpp
--- a/Engine/Source/Programs/UnrealBuildAccelerator/Detours/Private/PosixOS/UbaDetoursMainPosix.cpp
+++ b/Engine/Source/Programs/UnrealBuildAccelerator/Detours/Private/PosixOS/UbaDetoursMainPosix.cpp
@@ -23,33 +23,7 @@ namespace uba
 	void PreInit(const char* logFile);
 	void Init();
 	void Deinit();
-}
-
-static const int g_allSignals[] = {
-	#ifdef SIGHUP
-	SIGHUP,
-	#endif
-	#ifdef SIGQUIT
-	SIGQUIT,
-	#endif
-	#ifdef SIGTRAP
-	SIGTRAP,
-	#endif
-	#ifdef SIGIO
-	SIGIO,
-	#endif
-	SIGABRT,
-	SIGFPE,
-	SIGILL,
-	SIGINT,
-	SIGSEGV,
-	SIGTERM 
-};
-
-
-static void segfault_sigaction(int signal)
-{
-	uba::UbaAssert("Segmentation fault", "", 0, "", -1);
+	void AddExceptionHandler();
 }
 
 static void __attribute__((constructor(102))) PreInitCtor()
@@ -107,15 +81,7 @@ static void __attribute__((constructor(102))) PreInitCtor()
 	g_writeEvent = ((Event*)rptr) + 2;
 	g_messageMappingMem = rptr + sizeof(Event) * 3;
 
-	#if 0 // TODO: For some reason this does not work
-	struct sigaction sa;
-	memset(&sa, 0, sizeof(struct sigaction));
-	sigemptyset(&sa.sa_mask);
-	sa.sa_handler = segfault_sigaction;
-	sa.sa_flags = 0;// SA_SIGINFO;
-	for (auto signal : g_allSignals)
-		sigaction(signal, &sa, NULL);
-	#endif
+	AddExceptionHandler();
 
 	g_sessionPid = atoi(getenv("UBA_SESSION_PROCESS"));
 
diff --git a/Engine/Source/Programs/UnrealBuildAccelerator/Detours/Private/UbaDetoursShared.cpp b/Engine/Source/Programs/UnrealBuildAccelerator/Detours/Private/UbaDetoursShared.cpp
--- a/Engine/Source/Programs/UnrealBuildAccelerator/Detours/Private/UbaDetoursShared.cpp
+++ b/Engine/Source/Programs/UnrealBuildAccelerator/Detours/Private/UbaDetoursShared.cpp
@@ -323,6 +323,34 @@ namespace uba
 		Rpc_WriteLog(buffer, u32(count), false, false);
 	}
 
+	UBA_NOINLINE void Rpc_ResolveCallstack(StringBufferBase& out, u32 skipCallstackCount, void* context)
+	{
+		u32 tryCount = 0;
+		bool hasLock = false;
+		while (tryCount++ < 5)
+		{
+			hasLock = g_communicationLock.TryEnterWrite();
+			if (hasLock)
+				break;
+			Sleep(100);
+		}
+
+		BinaryWriter writer;
+		writer.WriteByte(MessageType_ResolveCallstack);
+		if (WriteCallstackInfo(writer, skipCallstackCount, context))
+		{
+			writer.Flush();
+			BinaryReader reader;
+			reader.ReadString(out);
+		}
+		else
+		{
+			out.Append(TC("\n   Failed to resolve callstack\n"));
+		}
+		// Note, we leave the lock even though we might not have it because we want to be able to report
+		g_communicationLock.LeaveWrite();
+	}
+
 	//TODO: Implement SetConsoleTextAttribute.. clang is using it to color errors
 
 	tchar g_consoleString[4096];
diff --git a/Engine/Source/Programs/UnrealBuildAccelerator/Detours/Private/UbaDetoursShared.h b/Engine/Source/Programs/UnrealBuildAccelerator/Detours/Private/UbaDetoursShared.h
--- a/Engine/Source/Programs/UnrealBuildAccelerator/Detours/Private/UbaDetoursShared.h
+++ b/Engine/Source/Programs/UnrealBuildAccelerator/Detours/Private/UbaDetoursShared.h
@@ -19,6 +19,7 @@ namespace uba
 
 	void Rpc_WriteLog(const tchar* text, u64 textCharLength, bool printInSession, bool isError);
 	void Rpc_WriteLogf(const tchar* format, ...);
+	void Rpc_ResolveCallstack(StringBufferBase& out, u32 skipCallstackCount, void* context);
 
 	const tchar* GetApplicationShortName();
 	bool FixPath(StringBufferBase& out, const tchar* path);
diff --git a/Engine/Source/Programs/UnrealBuildAccelerator/Detours/Private/Windows/UbaDetoursFunctionsKernelBase.inl b/Engine/Source/Programs/UnrealBuildAccelerator/Detours/Private/Windows/UbaDetoursFunctionsKernelBase.inl
--- a/Engine/Source/Programs/UnrealBuildAccelerator/Detours/Private/Windows/UbaDetoursFunctionsKernelBase.inl
+++ b/Engine/Source/Programs/UnrealBuildAccelerator/Detours/Private/Windows/UbaDetoursFunctionsKernelBase.inl
@@ -2743,7 +2743,7 @@ BOOL Detoured_CreateProcessW(LPCWSTR lpApplicationName, LPWSTR lpCommandLine, LP
 	{
 		if (g_runningRemote)
 		{
-			UbaAssert(L"Suppress debugger startup and try to report issue instead. This message is here to hopefully see callstack", "", 0, "", 0, false);
+			UbaAssert(L"Suppress debugger startup and try to report issue instead. This message is here to hopefully see callstack", "", 0, "", false, 0, nullptr);
 			return false;
 		}
 		else
diff --git a/Engine/Source/Programs/UnrealBuildAccelerator/Detours/Private/Windows/UbaDetoursFunctionsWin.cpp b/Engine/Source/Programs/UnrealBuildAccelerator/Detours/Private/Windows/UbaDetoursFunctionsWin.cpp
--- a/Engine/Source/Programs/UnrealBuildAccelerator/Detours/Private/Windows/UbaDetoursFunctionsWin.cpp
+++ b/Engine/Source/Programs/UnrealBuildAccelerator/Detours/Private/Windows/UbaDetoursFunctionsWin.cpp
@@ -325,13 +325,16 @@ void ToInvestigate(const wchar_t* format, ...)
 #endif
 }
 
-UBA_NOINLINE void UbaAssert(const wchar_t* text, const char* file, u32 line, const char* expr, u32 terminateCode, bool allowTerminate)
+UBA_NOINLINE void UbaAssert(const tchar* text, const char* file, u32 line, const char* expr, bool allowTerminate, u32 terminateCode, void* context)
 {
 	SuppressDetourScope _;
-	
-	StringBuffer<32*1024> b;
-	WriteAssertInfo(b, text, file, line, expr, 3);
-	Rpc_WriteLog(b.data, b.count, true, true);
+	static CriticalSection cs;
+	ScopedCriticalSection s(cs);
+	static auto& sb = *new StringBuffer<16*1024>();
+	WriteAssertInfo(sb.Clear(), text, file, line, expr);
+	Rpc_ResolveCallstack(sb, 4, context);
+	Rpc_WriteLog(sb.data, sb.count, true, true);
+
 	#if UBA_DEBUG_LOG_ENABLED
 	FlushDebugLog();
 	#endif
@@ -339,7 +342,7 @@ UBA_NOINLINE void UbaAssert(const wchar_t* text, const char* file, u32 line, con
 	#if UBA_ASSERT_MESSAGEBOX
 	StringBuffer<> title;
 	title.Appendf(L"Assert %ls - pid %u", GetApplicationShortName(), GetCurrentProcessId());
-	int ret = MessageBoxW(GetConsoleWindow(), b.data, title.data, MB_ABORTRETRYIGNORE|MB_SYSTEMMODAL);
+	int ret = MessageBoxW(GetConsoleWindow(), sb.data, title.data, MB_ABORTRETRYIGNORE|MB_SYSTEMMODAL);
 	if (ret == IDABORT)
 		ExitProcess(terminateCode);
 	else if (ret == IDRETRY && IsDebuggerPresent())
@@ -350,11 +353,6 @@ UBA_NOINLINE void UbaAssert(const wchar_t* text, const char* file, u32 line, con
 	#endif
 }
 
-void UbaAssert(const wchar_t* text, const char* file, u32 line, const char* expr, u32 terminateCode)
-{
-	UbaAssert(text, file, line, expr, terminateCode, true);
-}
-
 extern HANDLE g_hostProcess;
 
 const wchar_t* HandleToName(HANDLE handle)
@@ -810,19 +808,7 @@ void PreInit(const DetoursPayload& payload)
 
 void Init(const DetoursPayload& payload, u64 startTime)
 {
-	if (g_rules->EnableVectoredExceptionHandler())
-	{
-		AddVectoredExceptionHandler(0, [](_EXCEPTION_POINTERS* ExceptionInfo) -> LONG
-			{
-				u32 exceptionCode = ExceptionInfo->ExceptionRecord->ExceptionCode;
-				if (exceptionCode != EXCEPTION_STACK_OVERFLOW && exceptionCode != EXCEPTION_ACCESS_VIOLATION)
-					return EXCEPTION_CONTINUE_SEARCH;
-				StringBuffer<> text;
-				text.Appendf(L"Unhandled Exception (Code: 0x%x)", exceptionCode);
-				UbaAssert(text.data, nullptr, 0, nullptr, exceptionCode, false);
-				return EXCEPTION_EXECUTE_HANDLER;
-			});
-	}
+	AddExceptionHandler();
 
 	DetourAttachFunctions(g_runningRemote);
 
diff --git a/Engine/Source/Programs/UnrealBuildAccelerator/Detours/Private/Windows/UbaDetoursObjFilesPreloader.h b/Engine/Source/Programs/UnrealBuildAccelerator/Detours/Private/Windows/UbaDetoursObjFilesPreloader.h
--- a/Engine/Source/Programs/UnrealBuildAccelerator/Detours/Private/Windows/UbaDetoursObjFilesPreloader.h
+++ b/Engine/Source/Programs/UnrealBuildAccelerator/Detours/Private/Windows/UbaDetoursObjFilesPreloader.h
@@ -103,6 +103,8 @@ namespace uba
 						return;
 				}
 				rspFile.Append(begin, end - begin);
+
+				DevirtualizePath(rspFile);
 				ParseRsp(rspFile);
 				return;
 			}
@@ -116,7 +118,7 @@ namespace uba
 			DevirtualizePath(fileFull);
 
 			HANDLE fileHandle = CreateFileW(fileFull.data, GENERIC_READ, FILE_SHARE_READ, NULL, OPEN_EXISTING, FILE_ATTRIBUTE_NORMAL, NULL);
-			UBA_ASSERT(fileHandle != INVALID_HANDLE_VALUE);
+			UBA_ASSERTF(fileHandle != INVALID_HANDLE_VALUE, TC("Failed to open %s for read"), fileFull.data);
 
 			StringKey fileNameKey = ToStringKey(fileFull.MakeLower());
 			auto insres = m_preloadedObjFiles.try_emplace(fileNameKey);
diff --git a/Engine/Source/Programs/UnrealBuildAccelerator/Host/Private/UbaExports.cpp b/Engine/Source/Programs/UnrealBuildAccelerator/Host/Private/UbaExports.cpp
--- a/Engine/Source/Programs/UnrealBuildAccelerator/Host/Private/UbaExports.cpp
+++ b/Engine/Source/Programs/UnrealBuildAccelerator/Host/Private/UbaExports.cpp
@@ -76,17 +76,6 @@ namespace uba
 		LoggerWithWriter logger;
 	};
 
-	#define UBA_USE_SIGNALHANDLER 0//PLATFORM_LINUX // It might be that we can't use signal handlers in c# processes.. so don't set this to 1
-
-	#if UBA_USE_SIGNALHANDLER
-	void SignalHandler(int sig, siginfo_t* si, void* unused)
-	{
-		StringBuffer<256> desc;
-		desc.Append("Segmentation fault at +0x").AppendHex(u64(si->si_addr));
-		UbaAssert(desc.data, "", 0, "", -1);
-	}
-	#endif
-
 	Config& GetConfig(const tchar* fileName = nullptr)
 	{
 		static Config config;
@@ -115,15 +104,7 @@ extern "C"
 
 	uba::LogWriter* CreateCallbackLogWriter(uba::CallbackLogWriter::BeginScopeCallback begin, uba::CallbackLogWriter::EndScopeCallback end, uba::CallbackLogWriter::LogCallback log)
 	{
-		#if UBA_USE_SIGNALHANDLER
-		struct sigaction action;
-		memset(&action, 0, sizeof(action));
-		sigfillset(&action.sa_mask);
-		action.sa_flags = SA_SIGINFO | SA_RESTART | SA_ONSTACK;
-		action.sa_sigaction = uba::SignalHandler;
-		sigaction(SIGSEGV, &action, NULL);
-		//sigaction(SIGABRT, &action, NULL);
-		#endif
+		uba::AddExceptionHandler();
 
 		return new uba::CallbackLogWriter(begin, end, log);
 	}
@@ -342,7 +323,7 @@ extern "C"
 		info->allowKillOnMem = allowKillOnMem;
 		info->storeObjFilesCompressed = storeObjFilesCompressed;
 		//info->remoteTraceEnabled = true;
-		//info->remoteLogEnabled = true;
+		info->remoteLogEnabled = true;
 		return info;
 	}
 
diff --git a/Engine/Source/Programs/UnrealBuildAccelerator/TestApp/Private/PosixOS/UbaTestAppPosix.cpp b/Engine/Source/Programs/UnrealBuildAccelerator/TestApp/Private/PosixOS/UbaTestAppPosix.cpp
--- a/Engine/Source/Programs/UnrealBuildAccelerator/TestApp/Private/PosixOS/UbaTestAppPosix.cpp
+++ b/Engine/Source/Programs/UnrealBuildAccelerator/TestApp/Private/PosixOS/UbaTestAppPosix.cpp
@@ -11,7 +11,6 @@
 
 #if PLATFORM_LINUX
 #include <link.h>
-#include <dlfcn.h>
 #include <wait.h>
 #else
 #include <mach-o/dyld.h>
