[UBA]
* Fixed callstack resolving code for osx

diff --git a/Engine/Source/Programs/UnrealBuildAccelerator/Common/Private/UbaLogger.cpp b/Engine/Source/Programs/UnrealBuildAccelerator/Common/Private/UbaLogger.cpp
--- a/Engine/Source/Programs/UnrealBuildAccelerator/Common/Private/UbaLogger.cpp
+++ b/Engine/Source/Programs/UnrealBuildAccelerator/Common/Private/UbaLogger.cpp
@@ -484,7 +484,11 @@ namespace uba
 
 	void ParseCallstackInfo(StringBufferBase& out, BinaryReader& reader, const tchar* executable, const StringView* searchPaths)
 	{
+		#if PLATFORM_WINDOWS
 		bool isRunningWine = reader.ReadBool();
+		#else
+		bool isRunningWine = false;
+		#endif
 
 		struct CallstackEntry { u64 moduleIndex; u64 memoryOffset; };
 		u64 callstackCount = reader.Read7BitEncoded();
@@ -627,14 +631,18 @@ namespace uba
 			auto pg = MakeGuard([&]() { out.Appendf(TC("\n   %s: 0x%llx"), mod.name.c_str(), entry.memoryOffset); });
 
 			StringBuffer<1024> cmd;
+			#if PLATFORM_LINUX
 			cmd.Appendf("addr2line 0x%x", entry.memoryOffset).Append(" -e ").Append(mod.name).Append(" 2>/dev/null");
+			#else
+			cmd.Appendf("atos -o %s --offset 0x%llx", mod.name.c_str(), entry.memoryOffset);
+			#endif
 
 			fflush(stdout);
 
 			FILE* fp = popen(cmd.data, "r");
 			if (!fp)
-				break;
-			auto cg = MakeGuard([&]() { fclose(fp); });
+				continue;
+			auto cg = MakeGuard([&]() { pclose(fp); });
 
 			u32 index = 0;
 			u32 countBeforeCallstack = out.count;
diff --git a/Engine/Source/Programs/UnrealBuildAccelerator/Core/Private/UbaPlatform.cpp b/Engine/Source/Programs/UnrealBuildAccelerator/Core/Private/UbaPlatform.cpp
--- a/Engine/Source/Programs/UnrealBuildAccelerator/Core/Private/UbaPlatform.cpp
+++ b/Engine/Source/Programs/UnrealBuildAccelerator/Core/Private/UbaPlatform.cpp
@@ -289,35 +289,32 @@ namespace uba
 				return 0;
 			}, &moduleEndAddresses);
 #else
-		for (u32 i = 0, e = _dyld_image_count(); i < e; i++)
 		{
-			const char* name = _dyld_get_image_name(i);
-			if (!name)
-				continue;
-			auto hdr = _dyld_get_image_header(i);
-			size_t imageSize = 0;
-			if (hdr->magic == MH_MAGIC_64)
+			Map<u64, ModuleRec> modules;
+			for (u32 i = 0, e = _dyld_image_count(); i < e; i++)
 			{
-				auto hdr64 = (const mach_header_64 *)hdr;
-				auto cmd = (const load_command*)(hdr64 + 1);
-				for (uint32_t i = 0; i < hdr64->ncmds; i++, cmd = (const load_command*)((const char *)cmd + cmd->cmdsize))
-					if (cmd->cmd == LC_SEGMENT_64)
-						imageSize += ((const segment_command_64*)cmd)->vmsize;
+				const char* name = _dyld_get_image_name(i);
+				if (!name)
+					continue;
+				auto hdr = _dyld_get_image_header(i);
+				u64 start = u64(hdr);// + _dyld_get_image_vmaddr_slide(i);
+				ModuleRec& rec = modules[start];
+				rec.start = start;
+				rec.size = 0;
+				rec.name = name;
+				rec.index = ~0u;
 			}
-			else if (hdr->magic == MH_MAGIC)
+			ModuleRec* prev = nullptr;
+			for (auto& kv : modules)
 			{
-				auto hdr32 = hdr;
-				auto cmd = (const load_command*)(hdr32 + 1);
-				for (uint32_t i = 0; i < hdr32->ncmds; i++, cmd = (const load_command*)((const char *)cmd + cmd->cmdsize))
-					if (cmd->cmd == LC_SEGMENT)
-						imageSize += ((const segment_command*)cmd)->vmsize;
+				if (prev)
+				{
+					u64 end = kv.second.start;
+					auto& rec = moduleEndAddresses[end] = *prev;
+					rec.size = end - prev->start;
+				}
+				prev = &kv.second;
 			}
-
-			ModuleRec& rec = moduleEndAddresses[end];
-			rec.start = u64(hdr) + _dyld_get_image_vmaddr_slide(i);
-			rec.size = imageSize;
-			rec.name = name;
-			rec.index = ~0u;
 		}
 #endif
 
@@ -327,7 +324,7 @@ namespace uba
 
 		Vector<ModuleRec*> usedModules;
 
-		out.Write7BitEncoded(nptrs - skipCallstack);
+		out.Write7BitEncoded(nptrs < skipCallstack ? 0 : nptrs - skipCallstack);
 
 		for (int i = skipCallstack; i < nptrs; ++i)
 		{
