[UBA]
* Fixed so segfaults on posix is not showing the top of the callstack

diff --git a/Engine/Source/Programs/UnrealBuildAccelerator/Core/Private/UbaPlatform.cpp b/Engine/Source/Programs/UnrealBuildAccelerator/Core/Private/UbaPlatform.cpp
--- a/Engine/Source/Programs/UnrealBuildAccelerator/Core/Private/UbaPlatform.cpp
+++ b/Engine/Source/Programs/UnrealBuildAccelerator/Core/Private/UbaPlatform.cpp
@@ -324,6 +324,9 @@ namespace uba
 
 		Vector<ModuleRec*> usedModules;
 
+		if (contextPtr) // Use this to know if we come from a signal
+			skipCallstack +=2;
+
 		out.Write7BitEncoded(nptrs < skipCallstack ? 0 : nptrs - skipCallstack);
 
 		for (int i = skipCallstack; i < nptrs; ++i)
@@ -528,7 +531,7 @@ namespace uba
 		SIGSEGV,
 		//SIGTERM 
 	};
-	static void segfault_sigaction(int signal) { uba::UbaAssert("Segmentation fault", "", 0, "", true, -1, nullptr); }
+	static void segfault_sigaction(int signal) { uba::UbaAssert("Segmentation fault", "", 0, "", true, -1, &signal); }
 #endif
 
 	void AddExceptionHandler()
