[UBA]
* Fixed ubacli code to use exception handler
* Added a bunch of disabled debug code in ubacli

diff --git a/Engine/Source/Programs/UnrealBuildAccelerator/Cli/Private/UbaCli.cpp b/Engine/Source/Programs/UnrealBuildAccelerator/Cli/Private/UbaCli.cpp
--- a/Engine/Source/Programs/UnrealBuildAccelerator/Cli/Private/UbaCli.cpp
+++ b/Engine/Source/Programs/UnrealBuildAccelerator/Cli/Private/UbaCli.cpp
@@ -135,31 +135,11 @@ namespace uba
 	
 	StringBuffer<> g_rootDir(DefaultRootDir);
 
-	//LONG WINAPI UbaUnhandledExceptionFilter(EXCEPTION_POINTERS* ExceptionInfo)
-	//{
-	//	time_t rawtime;
-	//	time(&rawtime);
-	//	tm ti;
-	//	localtime_s(&ti, &rawtime);
-	//
-	//	StringBuffer<> dumpFile;
-	//	dumpFile.Append(g_rootDir).EnsureEndsWithSlash().Appendf(TC("UbaCliCrash_%02u%02u%02u_%02u%02u%02u.dmp"), ti.tm_year - 100, ti.tm_mon + 1, ti.tm_mday, ti.tm_hour, ti.tm_min, ti.tm_sec);
-	//
-	//	wprintf(TC("Unhandled exception - Writing minidump %s\n"), dumpFile.data);
-	//	HANDLE hFile = ::CreateFileW(dumpFile.data, GENERIC_READ | GENERIC_WRITE, FILE_SHARE_DELETE | FILE_SHARE_READ | FILE_SHARE_WRITE, NULL, CREATE_ALWAYS, FILE_ATTRIBUTE_NORMAL, NULL);
-	//	MINIDUMP_EXCEPTION_INFORMATION mei;
-	//	mei.ThreadId = GetCurrentThreadId();
-	//	mei.ClientPointers = TRUE;
-	//	mei.ExceptionPointers = ExceptionInfo;
-	//	MiniDumpWriteDump(GetCurrentProcess(), ::GetCurrentProcessId(), hFile, MiniDumpNormal, &mei, NULL, NULL);
-	//	return EXCEPTION_EXECUTE_HANDLER;
-	//}
-
-
 	bool WrappedMain(int argc, tchar* argv[])
 	{
 		using namespace uba;
-		//SetUnhandledExceptionFilter(UbaUnhandledExceptionFilter);
+
+		AddExceptionHandler();
 
 		u32 storageCapacityGb = DefaultCapacityGb;
 		StringBuffer<256> workDir;
@@ -170,7 +150,9 @@ namespace uba
 		TString coordinatorPool;
 		u32 coordinatorMaxCoreCount = 400;
 		u16 port = DefaultPort;
+		u16 cachePort = DefaultCachePort;
 		u32 maxProcessCount = DefaultProcessorCount;
+		u32 agentCount = 1;
 		bool launchVisualizer = false;
 		bool storeCompressed = true;
 		bool disableCustomAllocator = false;
@@ -408,7 +390,13 @@ namespace uba
 			{
 				if (value.IsEmpty())
 					return PrintHelp(TC("-cache needs a value"));
-				cacheHost.Append(value);
+				if (const tchar* colon = value.First(':'))
+				{
+					value.Parse(cachePort, colon - value.data + 1);
+					cacheHost.Append(value.data, colon - value.data);
+				}
+				else
+					cacheHost.Append(value);
 			}
 			else if (name.Equals(TC("-populatecache")))
 			{
@@ -458,7 +446,7 @@ namespace uba
 		if (deleteCas)
 		{
 			StorageImpl(StorageCreateInfo(g_rootDir.data, logWriter)).DeleteAllCas();
-			for (u32 i=0; i!=4; ++i)
+			for (u32 i=0; i!=agentCount; ++i)
 			{
 				StringBuffer<> clientRootDir;
 				clientRootDir.Append(g_rootDir).Append("Agent").AppendValue(i);
@@ -797,7 +785,7 @@ namespace uba
 		if (cacheHost.count)
 		{
 			CreateCacheClient();
-			if (!cacheClient->GetClient().Connect(networkBackend, cacheHost.data, DefaultCachePort))
+			if (!cacheClient->GetClient().Connect(networkBackend, cacheHost.data, cachePort))
 				return logger.Error(TC("Failed to connect to cache server"));
 
 			if (!storageServer.LoadCasTable(true))
@@ -835,6 +823,37 @@ namespace uba
 			arguments.resize(arguments.size() - 1);
 		}
 
+		// Vfs testing
+		RootsHandle rootsHandle = 0;
+		if (false)
+		{
+			struct VfsEntry { const tchar* vfs; const tchar* local; };
+			VfsEntry vfsEntries[] =
+			{
+#if PLATFORM_WINDOWS
+				{ TC("Z:/UEVFS/Root"), TC("E:\\dev\\fn") },
+				{ TC("Z:/UEVFS/Compiler"), TC("c:\\sdk\\AutoSDK\\HostWin64\\Win64\\VS2022\\14.38.33130") },
+				{ TC("Z:/UEVFS/WinSDK"), TC("C:\\Program Files (x86)\\Windows Kits\\10") },
+				//{ TC("Z:/UEVFS/Root"), TC("E:\\dev\\fn") },
+				//{ TC("Z:/UEVFS/Compiler"), TC("c:\\sdk\\AutoSDK\\HostWin64\\Win64\\LLVM\\18.1.8") },
+				//{ TC("Z:/UEVFS/Toolchain"), TC("c:\\sdk\\AutoSDK\\HostWin64\\Win64\\VS2022\\14.38.33130") },
+				//{ TC("Z:/UEVFS/WinSDK"), TC("c:\\sdk\\AutoSDK\\HostWin64\\Win64\\Windows Kits\\10.0.22621.0") },
+#else
+				{ TC("/UEVFS/Root"), TC("[REDACTED_USER_HOME]/fn") },
+				{ TC("/UEVFS/PlatformSDK"), TC("[REDACTED_USER_HOME]/AutoSDK/HostLinux/Linux_x64/v23_clang-18.1.0-rockylinux8/x86_64-unknown-linux-gnu") },
+#endif
+			};
+			StackBinaryWriter<1024> writer;
+			for (VfsEntry& entry : vfsEntries)
+			{
+				writer.WriteByte(0);
+				writer.WriteString(entry.vfs);
+				writer.WriteString(entry.local);
+			}
+			rootsHandle = sessionServer.RegisterRoots(writer.GetData(), writer.GetPosition());
+		}
+
+
 
 		if (true)//isRemote)
 		{
@@ -856,6 +875,7 @@ namespace uba
 			pinfo.application = app.c_str();
 			pinfo.arguments = arg.c_str();
 			pinfo.workingDir = workDir.data;
+			pinfo.rootsHandle = rootsHandle;
 
 			u32 bucketId = 1337;
 			if (cacheClient)
@@ -901,6 +921,7 @@ namespace uba
 			pinfo.workingDir = workDir.data;
 			pinfo.logFile = logFile.data;
 			pinfo.logLineUserData = &logger;
+			pinfo.rootsHandle = rootsHandle;
 			if (enableStdOut)
 				pinfo.logLineFunc = [](void* userData, const tchar* line, u32 length, LogEntryType type) { ((Logger*)userData)->Log(type, line, length); };
 			logger.Info(TC("Running %s %s"), app.c_str(), arg.c_str());
@@ -1028,7 +1049,7 @@ namespace uba
 
 			if (commandType == CommandType_Agent)
 			{
-				u32 clientCount = maxProcessCount == 1 ? 1 : 4;
+				u32 clientCount = maxProcessCount == 1 ? 1 : agentCount;
 				return RunWithClient([&]() { return RunQueue(); }, clientCount);
 			}
 			else
@@ -1057,14 +1078,19 @@ namespace uba
 		}
 		auto cg = MakeGuard([&]() { coordinator.Destroy(); });
 
-#if 0	// Annoying that link.exe/lld-link.exe needs path to windows folder.. 
+#if PLATFORM_WINDOWS	// Annoying that link.exe/lld-link.exe needs path to windows folder.. 
 		if (!useScheduler)
 		{
-			StringBuffer<2048> temp;
-			temp.count = GetEnvironmentVariableW(TC("PATH"), temp.data, temp.capacity);
-			temp.Append(TC("c:\\sdk\\AutoSDK\\HostWin64\\Win64\\Windows Kits\\10\\bin\\10.0.19041.0\\x64;"));
-			SetEnvironmentVariableW(TC("PATH"), temp.data);
-
+			StringBuffer<512> sdkbin;
+			//sdkbin.count = GetEnvironmentVariable(TC("WindowsSdkVerBinPath"), sdkbin.data, sdkbin.capacity);
+			sdkbin.Append(TC("C:\\Program Files (x86)\\Windows Kits\\10\\bin\\10.0.22621.0\\x64;"));
+			if (sdkbin.count)
+			{
+				StringBuffer<4096> temp;
+				temp.count = GetEnvironmentVariableW(TC("PATH"), temp.data, temp.capacity);
+				temp.Append(sdkbin);
+				SetEnvironmentVariableW(TC("PATH"), temp.data);
+			}
 		}
 #endif
 
@@ -1118,7 +1144,11 @@ namespace uba
 #if PLATFORM_WINDOWS
 int wmain(int argc, wchar_t* argv[])
 {
-	int res = uba::WrappedMain(argc, argv) ? 0 : -1;
+	int res;
+	//while (true)
+	{
+		res = uba::WrappedMain(argc, argv) ? 0 : -1;
+	}
 	Sleep(1); // Here to be able to put a breakpoint just before exit :-)
 	return res;
 }
