[UBA]
* Fixed win arm64 compile errors

diff --git a/Engine/Source/Programs/UnrealBuildAccelerator/Core/Private/UbaPlatform.cpp b/Engine/Source/Programs/UnrealBuildAccelerator/Core/Private/UbaPlatform.cpp
--- a/Engine/Source/Programs/UnrealBuildAccelerator/Core/Private/UbaPlatform.cpp
+++ b/Engine/Source/Programs/UnrealBuildAccelerator/Core/Private/UbaPlatform.cpp
@@ -150,6 +150,14 @@ namespace uba
 			out.Append(TC("\n LOCATION: ")).Append(file).Append(':').AppendValue(line);
 	}
 
+	#if defined(_M_ARM64)
+		#define GET_PC(context) context.Pc
+		#define GET_SP(context) context.Sp
+	#else
+		#define GET_PC(context) context.Rip
+		#define GET_SP(context) context.Rsp
+	#endif
+
 	UBA_NOINLINE bool WriteCallstackInfo(BinaryWriter& out, u32 skipCallstack, void* contextPtr)
 	{
 #if PLATFORM_WINDOWS
@@ -169,27 +177,27 @@ namespace uba
 			else
 				RtlCaptureContext(&context);
 
-			if (context.Rip && context.Rsp)
+			if (GET_PC(context) && GET_SP(context))
 			{
 				UNWIND_HISTORY_TABLE unwindHistoryTable;
 				RtlZeroMemory(&unwindHistoryTable, sizeof(UNWIND_HISTORY_TABLE));
-				for (u32 i=0; i<MaxDepth && context.Rip; ++i)
+				for (u32 i=0; i<MaxDepth && GET_PC(context); ++i)
 				{
-					callers[i] = (void*)context.Rip;
+					callers[i] = (void*)GET_PC(context);
 					count = i;
 					ULONG64 ImageBase = 0;
-					if (PRUNTIME_FUNCTION RuntimeFunction = RtlLookupFunctionEntry(context.Rip, &ImageBase, &unwindHistoryTable); RuntimeFunction)
+					if (PRUNTIME_FUNCTION RuntimeFunction = RtlLookupFunctionEntry(GET_PC(context), &ImageBase, &unwindHistoryTable); RuntimeFunction)
 					{
 						KNONVOLATILE_CONTEXT_POINTERS nvcontext;
 						RtlZeroMemory(&nvcontext, sizeof(KNONVOLATILE_CONTEXT_POINTERS));
 						PVOID HandlerData = nullptr;
 						ULONG64 EstablisherFrame = 0;
-						RtlVirtualUnwind(0, ImageBase, context.Rip, RuntimeFunction, &context, &HandlerData, &EstablisherFrame, &nvcontext);
+						RtlVirtualUnwind(0, ImageBase, GET_PC(context), RuntimeFunction, &context, &HandlerData, &EstablisherFrame, &nvcontext);
 					}
 					else
 					{
-						context.Rip = (*(u64*)context.Rsp);
-						context.Rsp += 8;
+						GET_PC(context) = (*(u64*)GET_SP(context));
+						GET_SP(context) += 8;
 					}
 				}
 			}
