[UBA]
* Added exception handler to cache server too

diff --git a/Engine/Source/Programs/UnrealBuildAccelerator/CacheService/Private/UbaCacheService.cpp b/Engine/Source/Programs/UnrealBuildAccelerator/CacheService/Private/UbaCacheService.cpp
--- a/Engine/Source/Programs/UnrealBuildAccelerator/CacheService/Private/UbaCacheService.cpp
+++ b/Engine/Source/Programs/UnrealBuildAccelerator/CacheService/Private/UbaCacheService.cpp
@@ -88,31 +88,11 @@ namespace uba
 	
 	StringBuffer<> g_rootDir(DefaultRootDir);
 
-	//LONG WINAPI UbaUnhandledExceptionFilter(EXCEPTION_POINTERS* ExceptionInfo)
-	//{
-	//	time_t rawtime;
-	//	time(&rawtime);
-	//	tm ti;
-	//	localtime_s(&ti, &rawtime);
-	//
-	//	StringBuffer<> dumpFile;
-	//	dumpFile.Append(g_rootDir).EnsureEndsWithSlash().Appendf(TC("UbaCliCrash_%02u%02u%02u_%02u%02u%02u.dmp"), ti.tm_year - 100, ti.tm_mon + 1, ti.tm_mday, ti.tm_hour, ti.tm_min, ti.tm_sec);
-	//
-	//	wprintf(TC("Unhandled exception - Writing minidump %s\n"), dumpFile.data);
-	//	HANDLE hFile = ::CreateFileW(dumpFile.data, GENERIC_READ | GENERIC_WRITE, FILE_SHARE_DELETE | FILE_SHARE_READ | FILE_SHARE_WRITE, NULL, CREATE_ALWAYS, FILE_ATTRIBUTE_NORMAL, NULL);
-	//	MINIDUMP_EXCEPTION_INFORMATION mei;
-	//	mei.ThreadId = GetCurrentThreadId();
-	//	mei.ClientPointers = TRUE;
-	//	mei.ExceptionPointers = ExceptionInfo;
-	//	MiniDumpWriteDump(GetCurrentProcess(), ::GetCurrentProcessId(), hFile, MiniDumpNormal, &mei, NULL, NULL);
-	//	return EXCEPTION_EXECUTE_HANDLER;
-	//}
-
-
 	int WrappedMain(int argc, tchar* argv[])
 	{
 		using namespace uba;
-		//SetUnhandledExceptionFilter(UbaUnhandledExceptionFilter);
+
+		AddExceptionHandler();
 
 		u32 storageCapacityGb = DefaultCapacityGb;
 		StringBuffer<256> workDir;
