[UBA]
* Change detoured NtCreateFile to exit readerwriter lock earlier. It was spanning over the real NtCreate in one case which we don't want in case there is reentrant code caused by antivirus etc

diff --git a/Engine/Source/Programs/UnrealBuildAccelerator/Detours/Private/Windows/UbaDetoursFunctionsNtDll.inl b/Engine/Source/Programs/UnrealBuildAccelerator/Detours/Private/Windows/UbaDetoursFunctionsNtDll.inl
--- a/Engine/Source/Programs/UnrealBuildAccelerator/Detours/Private/Windows/UbaDetoursFunctionsNtDll.inl
+++ b/Engine/Source/Programs/UnrealBuildAccelerator/Detours/Private/Windows/UbaDetoursFunctionsNtDll.inl
@@ -931,7 +931,7 @@ NTSTATUS NTAPI Shared_NtCreateFile(bool IsCreateFunc, PHANDLE hFileHandle, ACCES
 	const wchar_t* lpFileName = fileName.data;
 	u32 closeId = 0;
 
-	SCOPED_WRITE_LOCK(g_mappedFileTable.m_lookupLock, _);
+	SCOPED_WRITE_LOCK(g_mappedFileTable.m_lookupLock, lookupLock);
 	auto insres = g_mappedFileTable.m_lookup.try_emplace(fileNameKey);
 	FileInfo& info = insres.first->second;
 	u32 lastDesiredAccess = info.lastDesiredAccess;
@@ -1016,7 +1016,7 @@ NTSTATUS NTAPI Shared_NtCreateFile(bool IsCreateFunc, PHANDLE hFileHandle, ACCES
 
 	if (lpFileName[0] == '$')
 	{
-		_.Leave();
+		lookupLock.Leave();
 
 		UBA_ASSERT(!lpFileName[2]);
 
@@ -1063,6 +1063,9 @@ NTSTATUS NTAPI Shared_NtCreateFile(bool IsCreateFunc, PHANDLE hFileHandle, ACCES
 			UBA_ASSERTF(info.trueFileMapHandle, L"Can't duplicate handle 0x%llx (%ls) for file %ls", uintptr_t(mappingHandle), lpFileName, info.originalName);
 			return STATUS_UNSUCCESSFUL;
 		}
+
+		lookupLock.Leave();
+
 		DetouredHandle* dh = new DetouredHandle(HandleType_File);
 		UBA_ASSERT(info.size != InvalidValue);
 		dh->dirTableOffset = dirTableOffset;
@@ -1160,7 +1163,7 @@ NTSTATUS NTAPI Shared_NtCreateFile(bool IsCreateFunc, PHANDLE hFileHandle, ACCES
 			#endif
 		}
 
-		_.Leave();
+		lookupLock.Leave();
 
 		auto fileObject = new FileObject();
 		DetouredHandle* dh = new DetouredHandle(HandleType_File);
@@ -1178,6 +1181,8 @@ NTSTATUS NTAPI Shared_NtCreateFile(bool IsCreateFunc, PHANDLE hFileHandle, ACCES
 		return STATUS_SUCCESS;
 	}
 
+	lookupLock.Leave();
+
 	if (lpFileName[0] == ':')
 	{
 		auto fileObject = new FileObject();
