[UBA]
* Fixed linux compile error in UbaTest

diff --git a/Engine/Source/Programs/UnrealBuildAccelerator/Test/Private/Linux/UbaTestLinux.cpp b/Engine/Source/Programs/UnrealBuildAccelerator/Test/Private/Linux/UbaTestLinux.cpp
--- a/Engine/Source/Programs/UnrealBuildAccelerator/Test/Private/Linux/UbaTestLinux.cpp
+++ b/Engine/Source/Programs/UnrealBuildAccelerator/Test/Private/Linux/UbaTestLinux.cpp
@@ -2,23 +2,12 @@
 
 #include "UbaTestAll.h"
 
-void segfault_sigaction(int signal, siginfo_t* si, void* arg)
-{
-	uba::UbaAssert("Segmentation fault", "", 0, "", -1);
-}
-
 int main(int argc, char* argv[])
 {
-	#if !__has_feature(thread_sanitizer)
-	struct sigaction sa;
-	memset(&sa, 0, sizeof(struct sigaction));
-	sigemptyset(&sa.sa_mask);
-	sa.sa_sigaction = segfault_sigaction;
-	sa.sa_flags = SA_SIGINFO;
-	sigaction(SIGSEGV, &sa, NULL);
-	#endif
-
 	using namespace uba;
+
+	AddExceptionHandler();
+
 	if (!RunTests(argc, argv))
 		return -1;
 	return 0;
