[UBA]
* Added session RegisterNewDirectory
* Fixed session RefreshDirectory
* Changed lots of places to use StringView instead of const tchar* (and length sometimes)
* Change DirectoryTable to take MemoryBlock by reference since it is not optional
* Change DirectoryTable memory to be const since it should never be modified

diff --git a/Engine/Source/Programs/UnrealBuildAccelerator/Common/Private/UbaCacheClient.cpp b/Engine/Source/Programs/UnrealBuildAccelerator/Common/Private/UbaCacheClient.cpp
--- a/Engine/Source/Programs/UnrealBuildAccelerator/Common/Private/UbaCacheClient.cpp
+++ b/Engine/Source/Programs/UnrealBuildAccelerator/Common/Private/UbaCacheClient.cpp
@@ -1320,7 +1320,7 @@ namespace uba
 			return;
 
 		// Traverse all files in directory and report the file information... but only if it has not been reported before.. we don't want to interfere with other reports
-		TraverseDir(m_logger, path.data, 
+		TraverseDir(m_logger, path, 
 			[&](const DirectoryEntry& e)
 			{
 				if (IsDirectory(e.attributes))
diff --git a/Engine/Source/Programs/UnrealBuildAccelerator/Common/Private/UbaCacheServer.cpp b/Engine/Source/Programs/UnrealBuildAccelerator/Common/Private/UbaCacheServer.cpp
--- a/Engine/Source/Programs/UnrealBuildAccelerator/Common/Private/UbaCacheServer.cpp
+++ b/Engine/Source/Programs/UnrealBuildAccelerator/Common/Private/UbaCacheServer.cpp
@@ -270,7 +270,7 @@ namespace uba
 		{
 			StringBuffer<MaxPath> bucketsDir(m_rootDir);
 			bucketsDir.EnsureEndsWithSlash().Append(TC("buckets"));
-			TraverseDir(m_logger, bucketsDir.data, [&](const DirectoryEntry& e)
+			TraverseDir(m_logger, bucketsDir, [&](const DirectoryEntry& e)
 				{
 					StringBuffer<128> keyName;
 					keyName.Append(e.name, e.nameLen);
diff --git a/Engine/Source/Programs/UnrealBuildAccelerator/Common/Private/UbaFile.cpp b/Engine/Source/Programs/UnrealBuildAccelerator/Common/Private/UbaFile.cpp
--- a/Engine/Source/Programs/UnrealBuildAccelerator/Common/Private/UbaFile.cpp
+++ b/Engine/Source/Programs/UnrealBuildAccelerator/Common/Private/UbaFile.cpp
@@ -320,7 +320,7 @@ namespace uba
 	bool DeleteAllFiles(Logger& logger, const tchar* dir, bool deleteDir, u32* count)
 	{
 		bool success = true;
-		bool traverseRes = TraverseDir(logger, dir,
+		bool traverseRes = TraverseDir(logger, ToView(dir),
 			[&](const DirectoryEntry& e)
 			{
 				StringBuffer<> fullPath(dir);
diff --git a/Engine/Source/Programs/UnrealBuildAccelerator/Common/Private/UbaFileMapping.cpp b/Engine/Source/Programs/UnrealBuildAccelerator/Common/Private/UbaFileMapping.cpp
--- a/Engine/Source/Programs/UnrealBuildAccelerator/Common/Private/UbaFileMapping.cpp
+++ b/Engine/Source/Programs/UnrealBuildAccelerator/Common/Private/UbaFileMapping.cpp
@@ -86,7 +86,7 @@ namespace uba
 				}
 
 			// Clear out all orphaned shm_open
-			TraverseDir(logger, lockDir.data,
+			TraverseDir(logger, lockDir,
 				[&](const DirectoryEntry& e)
 				{
 					u32 uid = strtoul(e.name, nullptr, 10);
diff --git a/Engine/Source/Programs/UnrealBuildAccelerator/Common/Private/UbaLogger.cpp b/Engine/Source/Programs/UnrealBuildAccelerator/Common/Private/UbaLogger.cpp
--- a/Engine/Source/Programs/UnrealBuildAccelerator/Common/Private/UbaLogger.cpp
+++ b/Engine/Source/Programs/UnrealBuildAccelerator/Common/Private/UbaLogger.cpp
@@ -73,8 +73,12 @@ namespace uba
 
 		SetFocus(GetConsoleWindow());
 		SetActiveWindow(GetConsoleWindow());
+#else
+		if (IsDebuggerPresent())
+			DebugBreak();
 #endif
 
+
 		// *(int*)nullptr = 42; // Use this to force a crashdump instead of exiting
 
 		ExitProcess(terminateCode);
diff --git a/Engine/Source/Programs/UnrealBuildAccelerator/Common/Private/UbaProcess.cpp b/Engine/Source/Programs/UnrealBuildAccelerator/Common/Private/UbaProcess.cpp
--- a/Engine/Source/Programs/UnrealBuildAccelerator/Common/Private/UbaProcess.cpp
+++ b/Engine/Source/Programs/UnrealBuildAccelerator/Common/Private/UbaProcess.cpp
@@ -906,7 +906,7 @@ namespace uba
 		reader.ReadString(msg.directoryName);
 		msg.directoryNameKey = reader.ReadStringKey();
 		ListDirectoryResponse response;
-		m_messageSuccess = m_session.GetListDirectoryInfo(response, msg.directoryName.data, msg.directoryNameKey) && m_messageSuccess;
+		m_messageSuccess = m_session.GetListDirectoryInfo(response, msg.directoryName, msg.directoryNameKey) && m_messageSuccess;
 		writer.WriteU32(response.tableSize);
 		writer.WriteU32(response.tableOffset);
 		return true;
diff --git a/Engine/Source/Programs/UnrealBuildAccelerator/Common/Private/UbaSession.cpp b/Engine/Source/Programs/UnrealBuildAccelerator/Common/Private/UbaSession.cpp
--- a/Engine/Source/Programs/UnrealBuildAccelerator/Common/Private/UbaSession.cpp
+++ b/Engine/Source/Programs/UnrealBuildAccelerator/Common/Private/UbaSession.cpp
@@ -238,7 +238,7 @@ namespace uba
 		m_environmentVariables.push_back(0);
 	}
 
-	bool Session::WriteDirectoryEntriesInternal(DirectoryTable::Directory& dir, const StringKey& dirKey, const tchar* dirPath, bool isRefresh, u32& outTableOffset)
+	bool Session::WriteDirectoryEntriesInternal(DirectoryTable::Directory& dir, const StringKey& dirKey, const StringView& dirPath, bool isRefresh, u32& outTableOffset)
 	{
 		if (dir.tableOffset != InvalidTableOffset && !isRefresh)
 		{
@@ -257,11 +257,11 @@ namespace uba
 		u32 itemCount = 0;
 
 		StringKeyHasher hasher;
-		u32 dirPathLen = TStrlen(dirPath);
+		u32 dirPathLen = dirPath.count;
 		if (dirPathLen)
 		{
 			StringBuffer<> forHash;
-			forHash.Append(dirPath, dirPathLen);
+			forHash.Append(dirPath);
 			if (CaseInsensitiveFs)
 				forHash.MakeLower();
 			hasher.Update(forHash.data, forHash.count);
@@ -270,13 +270,13 @@ namespace uba
 		#if UBA_DEBUG_TRACK_DIR
 		m_debugLogger->BeginScope();
 		auto dg = MakeGuard([&]() { m_debugLogger->EndScope(); });
-		m_debugLogger->Info(TC("TRACKDIR %s\n"), dirPath);
+		m_debugLogger->Info(TC("TRACKDIR %s\n"), StringBuffer<>(dirPath).data);
 		#endif
 
 
 		BinaryWriter memoryWriter(memoryBlock.data(), 0, memoryBlock.size());
 
-		bool res = TraverseDir(m_logger, dirPathLen ? dirPath : TC("/"),
+		bool res = TraverseDir(m_logger, dirPathLen ? dirPath : AsView(TC("/")),
 			[&](const DirectoryEntry& e)
 			{
 				StringBuffer<256> fileNameForHash;
@@ -333,7 +333,7 @@ namespace uba
 
 		SCOPED_WRITE_LOCK(dirTable.m_memoryLock, memoryLock);
 		u32 writePos = dirTable.m_memorySize;
-		BinaryWriter tableWriter(dirTable.m_memory + dirTable.m_memorySize);
+		BinaryWriter tableWriter(m_directoryTableMem + dirTable.m_memorySize);
 
 		if (isRefresh)
 		{
@@ -373,7 +373,7 @@ namespace uba
 		return true;
 	}
 
-	void Session::WriteDirectoryEntriesRecursive(const StringKey& dirKey, tchar* dirPath, u32& outTableOffset)
+	void Session::WriteDirectoryEntriesRecursive(const StringKey& dirKey, StringView dirPath, u32& outTableOffset)
 	{
 		auto& dirTable = m_directoryTable;
 		SCOPED_WRITE_LOCK(dirTable.m_lookupLock, lookupLock);
@@ -399,14 +399,14 @@ namespace uba
 			dir.parseOffset = 1;
 		}
 
-		u64 dirlen = TStrlen(dirPath);
+		u64 dirlen = dirPath.count;
 
 		if (!dirlen) // This is for non-windows.. '/' is actually empty to get hashes correct
 			return;
 
 		// scan backwards first
-		tchar* rit = (tchar*)dirPath + dirlen - 2;
-		while (rit > dirPath)
+		const tchar* rit = (tchar*)dirPath.data + dirlen - 2;
+		while (rit > dirPath.data)
 		{
 			if (*rit != PathSeparator)
 			{
@@ -416,16 +416,17 @@ namespace uba
 			break;
 		}
 
-		if (rit <= dirPath) // There are no path separators left, this is the drive
+		if (rit <= dirPath.data) // There are no path separators left, this is the drive
 		{
-			*(dirPath) = 0;
+			dirPath.count = 0;
 			return;
 		}
 
-		*(rit) = 0;
+		dirPath.count = u32(rit - dirPath.data);
+		//*(rit) = 0;
 
 		StringBuffer<> parentDirForHash;
-		parentDirForHash.Append(dirPath, u64(rit - dirPath));
+		parentDirForHash.Append(dirPath);
 		if (CaseInsensitiveFs)
 			parentDirForHash.MakeLower();
 		StringKey parentKey = ToStringKey(parentDirForHash);
@@ -435,7 +436,7 @@ namespace uba
 		WriteDirectoryEntriesRecursive(parentKey, dirPath, parentOffset);
 	}
 
-	u32 Session::WriteDirectoryEntries(const StringKey& dirKey, tchar* dirPath, u32* outTableOffset)
+	u32 Session::WriteDirectoryEntries(const StringKey& dirKey, const StringView& dirPath, u32* outTableOffset)
 	{
 		auto& dirTable = m_directoryTable;
 		u32 temp;
@@ -737,7 +738,7 @@ namespace uba
 		// When not writing to disk we need to populate lookup before adding non-written files.. otherwise they will be lost once lookup is actually populated
 		if (!shouldWriteToDisk)
 		{
-			u32 res = WriteDirectoryEntries(dirKey, dirName.data);
+			u32 res = WriteDirectoryEntries(dirKey, dirName);
 			UBA_ASSERTF(res, TC("Failed to write directory entries for %s"), dirName.data); (void)res;
 		}
 
@@ -764,7 +765,7 @@ namespace uba
 		if (dir.parseOffset == 2)
 		{
 			dirLock.Leave();
-			u32 res = WriteDirectoryEntries(dirKey, dirName.data);
+			u32 res = WriteDirectoryEntries(dirKey, dirName);
 			UBA_ASSERT(res); (void)res;
 			dirLock.Enter();
 		}
@@ -811,7 +812,7 @@ namespace uba
 		// We check this by using lastWriteTime. If it hasn't change, directory table is already up-to-date
 		if (!insres.second && insres.first->second != ~u32(0))
 		{
-			BinaryReader reader(dirTable.m_memory + insres.first->second);
+			BinaryReader reader(m_directoryTableMem + insres.first->second);
 			u64 oldLastWriteTime = reader.ReadU64();
 
 			if (lastWriteTime == oldLastWriteTime)
@@ -856,7 +857,7 @@ namespace uba
 
 
 		SCOPED_WRITE_LOCK(dirTable.m_memoryLock, memoryLock);
-		u8* startPos = dirTable.m_memory + dirTable.m_memorySize;
+		u8* startPos = m_directoryTableMem + dirTable.m_memorySize;
 		BinaryWriter writer(startPos);
 		writer.Write7BitEncoded(written); // Storage size
 		insres.first->second = dirTable.m_memorySize + u32(writer.GetPosition() + entryPos); // Storing position to last write time
@@ -932,7 +933,7 @@ namespace uba
 		#endif
 
 		SCOPED_WRITE_LOCK(dirTable.m_memoryLock, memoryLock);
-		u8* startPos = dirTable.m_memory + dirTable.m_memorySize;
+		u8* startPos = m_directoryTableMem + dirTable.m_memorySize;
 		BinaryWriter writer(startPos);
 		writer.Write7BitEncoded(written); // Storage size
 		u32 tableOffset = u32(writer.GetPosition()) + sizeof(StringKey);
@@ -1041,7 +1042,7 @@ namespace uba
 	:	m_storage(info.storage)
 	,	m_logger(info.logWriter, logPrefix)
 	,	m_workManager(workManager)
-	,	m_directoryTable(&m_directoryTableMemory)
+	,	m_directoryTable(m_directoryTableMemory)
 	,	m_fileMappingBuffer(m_logger, workManager)
 	,	m_processCommunicationAllocator(m_logger, TC("CommunicationAllocator"))
 	,	m_trace(info.logWriter)
@@ -1093,7 +1094,7 @@ namespace uba
 
 			u64 systemTimeAsFileTime = GetSystemTimeAsFileTime();
 
-			TraverseDir(m_logger, sessionsDir.data,
+			TraverseDir(m_logger, sessionsDir,
 				[&](const DirectoryEntry& e)
 				{
 					u64 seconds = GetFileTimeAsSeconds(systemTimeAsFileTime - e.lastWritten);
@@ -1364,40 +1365,7 @@ namespace uba
 		return h;
 	}
 
-	bool Session::RefreshDirectory(const tchar* dirName)
-	{
-		UBA_ASSERT(!m_runningRemote);
-
-		StringBuffer<> dirPath;
-		FixPath2(dirName, nullptr, 0, dirPath.data, dirPath.capacity, &dirPath.count);
-		StringKey dirKey = CaseInsensitiveFs ? ToStringKeyLower(dirPath) : ToStringKey(dirPath);
-
-		auto& dirTable = m_directoryTable;
-		SCOPED_READ_LOCK(dirTable.m_lookupLock, lookupLock);
-		auto res = dirTable.m_lookup.find(dirKey);
-		if (res == dirTable.m_lookup.end())
-			return true;
-		DirectoryTable::Directory& dir = res->second;
-		lookupLock.Leave();
-		SCOPED_WRITE_LOCK(dir.lock, dirLock);
-
-		while (dir.parseOffset == 0)
-		{
-			dirLock.Leave();
-			Sleep(1);
-			dirLock.Enter();
-		}
-		UBA_ASSERT(dir.parseOffset == 1);
-
-		StringKeyHasher hasher;
-		hasher.Update(dirPath.data, TStrlen(dirPath.data));
-		m_directoryTable.PopulateDirectoryNoLock(hasher, dir);
-
-		u32 tableOffset;
-		return WriteDirectoryEntriesInternal(dir, dirKey, dirPath.data, true, tableOffset);
-	}
-
-	StringKey GetKeyAndFixedName(StringBuffer<>& fixedFilePath, const tchar* filePath)
+	StringKey GetKeyAndFixedName(StringBuffer<>& outFixedFilePath, StringKeyHasher& outHasher, const tchar* filePath)
 	{
 		StringBuffer<> workingDir;
 		if (!IsAbsolutePath(filePath))
@@ -1405,36 +1373,81 @@ namespace uba
 			GetCurrentDirectoryW(workingDir);
 			workingDir.EnsureEndsWithSlash();
 		}
-		FixPath(filePath, workingDir.data, workingDir.count, fixedFilePath);
+		FixPath(filePath, workingDir.data, workingDir.count, outFixedFilePath);
 
 		StringKey dirKey;
 		StringBuffer<> dirNameForHash;
 		const tchar* baseFileName;
-		GetDirKey(dirKey, dirNameForHash, baseFileName, fixedFilePath);
+		GetDirKey(dirKey, dirNameForHash, baseFileName, outFixedFilePath);
 
 		if (CaseInsensitiveFs)
 			dirNameForHash.MakeLower();
 
-		StringKeyHasher hasher;
-		hasher.Update(dirNameForHash.data, dirNameForHash.count);
+		outHasher.Update(dirNameForHash);
 
-		StringBuffer<> baseFileNameForHash;
+		StringBuffer<256> baseFileNameForHash;
 		baseFileNameForHash.Append(baseFileName);
 		if (CaseInsensitiveFs)
 			baseFileNameForHash.MakeLower();
-		StringKey result = ToStringKey(hasher, baseFileNameForHash.data, baseFileNameForHash.count);
+
+		outHasher.Update(baseFileNameForHash);
+		StringKey result = ToStringKey(outHasher);
 
 #if UBA_DEBUG
-		StringBuffer<> testPath(fixedFilePath);
+		StringBuffer<> testPath(outFixedFilePath);
 		if (CaseInsensitiveFs)
 			testPath.MakeLower();
 		StringKey testKey = ToStringKey(testPath);
-		UBA_ASSERTF(testKey == result, TC("Key mismatch for %s"), fixedFilePath.data);
+		UBA_ASSERTF(testKey == result, TC("Key mismatch for %s"), outFixedFilePath.data);
 #endif
 
 		return result;
 	}
 
+	StringKey GetKeyAndFixedName(StringBuffer<>& outFixedFilePath, const tchar* filePath)
+	{
+		StringKeyHasher hasher;
+		return GetKeyAndFixedName(outFixedFilePath, hasher, filePath);
+	}
+
+	bool Session::RefreshDirectory(const tchar* dirName, bool forceRegister)
+	{
+		UBA_ASSERT(!m_runningRemote);
+
+		StringBuffer<> dirPath;
+		StringKeyHasher hasher;
+		GetKeyAndFixedName(dirPath, hasher, dirName);
+
+		StringKey dirKey = ToStringKey(hasher);
+
+		auto& dirTable = m_directoryTable;
+		SCOPED_READ_LOCK(dirTable.m_lookupLock, lookupLock);
+		auto res = dirTable.m_lookup.find(dirKey);
+		if (res == dirTable.m_lookup.end())
+		{
+			lookupLock.Leave();
+			if (forceRegister)
+				WriteDirectoryEntries(dirKey, dirPath);
+			return true;
+		}
+		DirectoryTable::Directory& dir = res->second;
+		lookupLock.Leave();
+		SCOPED_WRITE_LOCK(dir.lock, dirLock);
+
+		while (dir.parseOffset == 0)
+		{
+			dirLock.Leave();
+			Sleep(1);
+			dirLock.Enter();
+		}
+		UBA_ASSERT(dir.parseOffset == 1);
+
+		//m_directoryTable.PopulateDirectoryNoLock(hasher, dir);
+
+		u32 tableOffset;
+		return WriteDirectoryEntriesInternal(dir, dirKey, dirPath, true, tableOffset);
+	}
+
 	bool Session::RegisterNewFile(const tchar* filePath)
 	{
 		UBA_ASSERT(!m_runningRemote);
@@ -1461,6 +1474,17 @@ namespace uba
 		RegisterDeleteFile(key, fixedFilePath);
 	}
 
+	bool Session::RegisterNewDirectory(const tchar* directoryPath)
+	{
+		UBA_ASSERT(!m_runningRemote);
+		StringBuffer<> fixedDirPath;
+		auto dirKey = GetKeyAndFixedName(fixedDirPath, directoryPath);
+		if (!RegisterCreateFileForWrite(dirKey, fixedDirPath, true))
+			return false;
+		WriteDirectoryEntries(dirKey, fixedDirPath);
+		return true;
+	}
+
 	void Session::RegisterCustomService(CustomServiceFunction&& function)
 	{
 		m_customServiceFunction = function;
@@ -2417,7 +2441,7 @@ namespace uba
 		{
 			// Both these functions need to be called. otherwise we can get created directories that does not end up in directory table
 			RegisterCreateFileForWrite(msg.nameKey, msg.name, true);
-			WriteDirectoryEntries(dirKey, dirName.data);
+			WriteDirectoryEntries(dirKey, dirName);
 		}
 
 		out.directoryTableSize = GetDirectoryTableSize();
@@ -2449,7 +2473,7 @@ namespace uba
 		return false;
 	}
 
-	bool Session::GetListDirectoryInfo(ListDirectoryResponse& out, tchar* dirName, const StringKey& dirKey)
+	bool Session::GetListDirectoryInfo(ListDirectoryResponse& out, const StringView& dirName, const StringKey& dirKey)
 	{
 		u32 tableOffset;
 		u32 tableSize = WriteDirectoryEntries(dirKey, dirName, &tableOffset);
diff --git a/Engine/Source/Programs/UnrealBuildAccelerator/Common/Private/UbaSessionClient.cpp b/Engine/Source/Programs/UnrealBuildAccelerator/Common/Private/UbaSessionClient.cpp
--- a/Engine/Source/Programs/UnrealBuildAccelerator/Common/Private/UbaSessionClient.cpp
+++ b/Engine/Source/Programs/UnrealBuildAccelerator/Common/Private/UbaSessionClient.cpp
@@ -928,7 +928,7 @@ namespace uba
 		return true;
 	}
 
-	bool SessionClient::GetListDirectoryInfo(ListDirectoryResponse& out, tchar* dirName, const StringKey& dirKey)
+	bool SessionClient::GetListDirectoryInfo(ListDirectoryResponse& out, const StringView& dirName, const StringKey& dirKey)
 	{
 		StackBinaryWriter<1024> writer;
 		NetworkMessage msg(m_client, ServiceId, SessionMessageType_ListDirectory, writer);
@@ -1050,7 +1050,7 @@ namespace uba
 		{
 			u32 readPos = reader.ReadU32();
 
-			u8* pos = dirTable.m_memory + readPos;
+			u8* pos = m_directoryTableMem + readPos;
 			u32 toRead = u32(reader.GetLeft());
 
 			SCOPED_WRITE_LOCK(m_directoryTableLock, lock);
diff --git a/Engine/Source/Programs/UnrealBuildAccelerator/Common/Private/UbaSessionServer.cpp b/Engine/Source/Programs/UnrealBuildAccelerator/Common/Private/UbaSessionServer.cpp
--- a/Engine/Source/Programs/UnrealBuildAccelerator/Common/Private/UbaSessionServer.cpp
+++ b/Engine/Source/Programs/UnrealBuildAccelerator/Common/Private/UbaSessionServer.cpp
@@ -1113,7 +1113,7 @@ namespace uba
 		reader.ReadString(dirName);
 		StringKey dirKey = reader.ReadStringKey();
 		ListDirectoryResponse out;
-		GetListDirectoryInfo(out, dirName.data, dirKey);
+		GetListDirectoryInfo(out, dirName, dirKey);
 		writer.WriteU32(out.tableOffset);
 		WriteDirectoryTable(session, reader, writer);
 		return true;
@@ -2172,12 +2172,12 @@ namespace uba
 
 		// Make a copy of dir table since we can't populate it
 		MemoryBlock block(64*1024*1024);
-		DirectoryTable dirTable(&block);
+		DirectoryTable dirTable(block);
 		u8* dirMem;
 		u32 dirMemSize;
 		{
 			SCOPED_READ_LOCK(m_directoryTable.m_memoryLock, lock);
-			dirMem = m_directoryTable.m_memory;
+			dirMem = m_directoryTableMem;
 			dirMemSize = m_directoryTable.m_memorySize;
 		}
 
diff --git a/Engine/Source/Programs/UnrealBuildAccelerator/Common/Private/UbaStorage.cpp b/Engine/Source/Programs/UnrealBuildAccelerator/Common/Private/UbaStorage.cpp
--- a/Engine/Source/Programs/UnrealBuildAccelerator/Common/Private/UbaStorage.cpp
+++ b/Engine/Source/Programs/UnrealBuildAccelerator/Common/Private/UbaStorage.cpp
@@ -996,7 +996,7 @@ namespace uba
 	{
 		Atomic<u32> workLeft;
 
-		auto traverseCasFileDir = [&func, &workLeft, this](const tchar* casFileDir)
+		auto traverseCasFileDir = [&func, &workLeft, this](const StringView& casFileDir)
 			{
 				TraverseDir(m_logger, casFileDir,
 					[&](const DirectoryEntry& e)
@@ -1009,7 +1009,7 @@ namespace uba
 				--workLeft;
 			};
 
-		TraverseDir(m_logger, dir,
+		TraverseDir(m_logger, ToView(dir),
 			[&](const DirectoryEntry& e)
 			{
 				if (!IsDirectory(e.attributes))
@@ -1020,11 +1020,11 @@ namespace uba
 
 				if (allowParallel && m_workManager)
 				{
-					m_workManager->AddWork([&traverseCasFileDir, p = TString(fullPath.data)]() { traverseCasFileDir(p.data()); }, 1, TC("TraverseCasFiles"));
+					m_workManager->AddWork([&traverseCasFileDir, p = TString(fullPath.data)]() { traverseCasFileDir(p); }, 1, TC("TraverseCasFiles"));
 				}
 				else
 				{
-					traverseCasFileDir(fullPath.data);
+					traverseCasFileDir(fullPath);
 				}
 			});
 
@@ -1813,7 +1813,7 @@ namespace uba
 		{
 			TString dir = directories.front();
 			directories.pop_front();
-			TraverseDir(m_logger, dir.data(), [&](const DirectoryEntry& e)
+			TraverseDir(m_logger, dir, [&](const DirectoryEntry& e)
 				{
 					StringBuffer<> path(dir);
 					path.EnsureEndsWithSlash().Append(e.name);
@@ -1910,7 +1910,7 @@ namespace uba
 		WorkManagerImpl workManager(GetLogicalProcessorCount(), TC("UbaWorker(DeleteCas)"));
 		{
 			Atomic<u32> atomicDeleteCount;
-			TraverseDir(m_logger, m_rootDir.data, [&](const DirectoryEntry& e)
+			TraverseDir(m_logger, m_rootDir, [&](const DirectoryEntry& e)
 				{
 					if (IsDirectory(e.attributes))
 						return;
diff --git a/Engine/Source/Programs/UnrealBuildAccelerator/Common/Private/UbaStorageClient.cpp b/Engine/Source/Programs/UnrealBuildAccelerator/Common/Private/UbaStorageClient.cpp
--- a/Engine/Source/Programs/UnrealBuildAccelerator/Common/Private/UbaStorageClient.cpp
+++ b/Engine/Source/Programs/UnrealBuildAccelerator/Common/Private/UbaStorageClient.cpp
@@ -924,7 +924,7 @@ namespace uba
 		StringBuffer<> fullPath;
 		fullPath.Append(dir).EnsureEndsWithSlash();
 		u32 dirLen = fullPath.count;
-		TraverseDir(m_logger, dir, [&](const DirectoryEntry& e)
+		TraverseDir(m_logger, ToView(dir), [&](const DirectoryEntry& e)
 			{
 				fullPath.Resize(dirLen).Append(e.name);
 				if (IsDirectory(e.attributes))
diff --git a/Engine/Source/Programs/UnrealBuildAccelerator/Common/Public/UbaDirectoryIterator.h b/Engine/Source/Programs/UnrealBuildAccelerator/Common/Public/UbaDirectoryIterator.h
--- a/Engine/Source/Programs/UnrealBuildAccelerator/Common/Public/UbaDirectoryIterator.h
+++ b/Engine/Source/Programs/UnrealBuildAccelerator/Common/Public/UbaDirectoryIterator.h
@@ -72,12 +72,12 @@ struct FILE_ID_BOTH_DIR_INFORMATION {
 
 namespace uba
 {
-	inline bool TraverseDir(Logger& logger, const tchar* dirPath, const IteratorFunc& iteratorFunc, bool errorOnNotFound = false, const DirectoryInfoFunc& infoFunc = {})
+	inline bool TraverseDir(Logger& logger, const StringView& dirPath, const IteratorFunc& iteratorFunc, bool errorOnNotFound = false, const DirectoryInfoFunc& infoFunc = {})
 	{
 		StringBuffer<> str;
 		str.Append(L"\\??\\").Append(dirPath);
 
-		if (dirPath[1] == ':' && dirPath[2] == 0) // if we are all the way to the root of the drive we add a slash since NtCreateFile needs that
+		if (dirPath.data[1] == ':' && dirPath.data[2] == 0) // if we are all the way to the root of the drive we add a slash since NtCreateFile needs that
 			str.Append(L"\\");
 
 		UNICODE_STRING uniName;
@@ -105,7 +105,7 @@ namespace uba
 			if (res == STATUS_NOT_A_DIRECTORY || res == STATUS_ACCESS_DENIED || res == STATUS_NO_MEDIA_IN_DEVICE || res == STATUS_FVE_LOCKED_VOLUME)
 				return false;
 			#if UBA_DEBUG
-			return logger.Error(L"NtCreateFile for TraverseDir on %s failed with code %x", dirPath, (DWORD)res);
+			return logger.Error(L"NtCreateFile for TraverseDir on %s failed with code %x", str.data, (DWORD)res);
 			#else
 			return false;
 			#endif
@@ -194,18 +194,20 @@ namespace uba
 
 namespace uba
 {
-	inline bool TraverseDir(Logger& logger, const tchar* dirPath, const IteratorFunc& iteratorFunc, bool errorOnNotFound = false, const DirectoryInfoFunc& infoFunc = {})
+	inline bool TraverseDir(Logger& logger, const StringView& dirPath, const IteratorFunc& iteratorFunc, bool errorOnNotFound = false, const DirectoryInfoFunc& infoFunc = {})
 	{
 		char dirPath2[1024];
-		strcpy(dirPath2, dirPath);
-		u32 dirPathLen = strlen(dirPath2);
+		u32 dirPathLen = dirPath.count;
+
+		memcpy(dirPath2, dirPath.data, dirPathLen);
+		dirPath2[dirPathLen] = 0;
 
 		DIR* dir = opendir(dirPath2);
 		if (!dir)
 		{
 			if (errno == ENOENT || errno == ENOTDIR)
 				return !errorOnNotFound;
-			UBA_ASSERTF(false, "TraverseDir error handling not implemented %s (%s)", dirPath, strerror(errno));
+			UBA_ASSERTF(false, "TraverseDir error handling not implemented %s (%s)", dirPath2, strerror(errno));
 			return false;
 		}
 
diff --git a/Engine/Source/Programs/UnrealBuildAccelerator/Common/Public/UbaSession.h b/Engine/Source/Programs/UnrealBuildAccelerator/Common/Public/UbaSession.h
--- a/Engine/Source/Programs/UnrealBuildAccelerator/Common/Public/UbaSession.h
+++ b/Engine/Source/Programs/UnrealBuildAccelerator/Common/Public/UbaSession.h
@@ -57,10 +57,12 @@ namespace uba
 		void CancelAllProcessesAndWait(bool terminate = true); // Cancel all processes and wait for them to go away
 
 		void PrintSummary(Logger& logger); // Print summary stats of session
-		bool RefreshDirectory(const tchar* dirName); // Tell uba a directory on disk has been changed by some other system while session is running
+		bool RefreshDirectory(const tchar* dirName, bool forceRegister = false); // Tell uba a directory on disk has been changed by some other system while session is running
 		bool RegisterNewFile(const tchar* filePath); // Tell uba a new file on disk has been added by some other system while session is running
 		bool RegisterVirtualFile(const tchar* filePath, const tchar* sourceFile, u64 sourceOffset, u64 sourceSize);
 		void RegisterDeleteFile(const tchar* filePath); // Tell uba a file on disk has been deleted by some other system while session is running
+		bool RegisterNewDirectory(const tchar* directoryPath); // Tell uba a directory on disk has been added by some other system while session is running
+
 		RootsHandle RegisterRoots(const void* rootsData, uba::u64 rootsDataSize);
 
 		using CustomServiceFunction = Function<u32(Process& handle, const void* recv, u32 recvSize, void* send, u32 sendCapacity)>;
@@ -89,6 +91,10 @@ namespace uba
 		bool DevirtualizeString(TString& inOut, RootsHandle rootsHandle, const tchar* hint = TC(""));
 		bool PopulateLocalToIndexRoots(RootPaths& out, RootsHandle rootsHandle);
 
+		// For testing
+		const u8* GetDirectoryTableMemory() { return m_directoryTableMem; }
+		u32 GetDirectoryTableSize();
+
 	protected:
 		Session(const SessionCreateInfo& info, const tchar* logPrefix, bool runningRemote, WorkManager* workManager = nullptr);
 		bool Create(const SessionCreateInfo& info);
@@ -102,7 +108,6 @@ namespace uba
 		void StartTrace(const tchar* traceName);
 		bool StopTrace(const tchar* writeFile);
 		void StopTraceThread();
-		u32 GetDirectoryTableSize();
 		u32 GetFileMappingSize();
 		u32 GetMemoryMapAlignment(const StringView& fileName) const;
 		u32 GetMemoryMapAlignment(const StringView& fileName, bool runningRemote) const;
@@ -116,7 +121,7 @@ namespace uba
 		bool IsRarelyReadAfterWritten(ProcessImpl& process, const StringView& fileName) const;
 		bool IsKnownSystemFile(const tchar* applicationName);
 		bool ShouldWriteToDisk(const StringView& fileName);
-		u32 WriteDirectoryEntries(const StringKey& dirKey, tchar* dirPath, u32* outTableOffset = nullptr);
+		u32 WriteDirectoryEntries(const StringKey& dirKey, const StringView& dirPath, u32* outTableOffset = nullptr);
 		u32 AddFileMapping(StringKey fileNameKey, const tchar* fileName, const tchar* newFileName, u64 fileSize = InvalidValue);
 		
 		struct MemoryMap { StringBuffer<128> name; u64 size = 0; };
@@ -145,7 +150,7 @@ namespace uba
 		virtual bool RemoveDirectory(RemoveDirectoryResponse& out, const RemoveDirectoryMessage& msg);
 		virtual bool GetFullFileName(GetFullFileNameResponse& out, const GetFullFileNameMessage& msg);
 		virtual bool GetLongPathName(GetLongPathNameResponse& out, const GetLongPathNameMessage& msg);
-		virtual bool GetListDirectoryInfo(ListDirectoryResponse& out, tchar* dirName, const StringKey& dirKey);
+		virtual bool GetListDirectoryInfo(ListDirectoryResponse& out, const StringView& dirName, const StringKey& dirKey);
 		virtual bool WriteFilesToDisk(ProcessImpl& process, WrittenFile** files, u32 fileCount);
 		virtual bool AllocFailed(Process& process, const tchar* allocType, u32 error);
 		virtual bool GetNextProcess(Process& process, bool& outNewProcess, NextProcessInfo& outNextProcess, u32 prevExitCode, BinaryReader& statsReader);
@@ -161,8 +166,8 @@ namespace uba
 		static constexpr CasKey CasKeyIsDirectory = { ~u64(0), ~u64(0), ~u32(0) };
 
 		void AddEnvironmentVariableNoLock(const tchar* key, const tchar* value);
-		bool WriteDirectoryEntriesInternal(DirectoryTable::Directory& dir, const StringKey& dirKey, const tchar* dirPath, bool isRefresh, u32& outTableOffset);
-		void WriteDirectoryEntriesRecursive(const StringKey& dirKey, tchar* dirPath, u32& outTableOffset);
+		bool WriteDirectoryEntriesInternal(DirectoryTable::Directory& dir, const StringKey& dirKey, const StringView& dirPath, bool isRefresh, u32& outTableOffset);
+		void WriteDirectoryEntriesRecursive(const StringKey& dirKey, StringView dirPath, u32& outTableOffset);
 		bool CopyImports(Vector<BinaryModule>& out, const tchar* library, tchar* applicationDir, tchar* applicationDirEnd, UnorderedSet<TString>& handledImports, const char* const* loaderPaths);
 		bool CreateProcessJobObject();
 		void GetSystemInfo(StringBufferBase& out);
diff --git a/Engine/Source/Programs/UnrealBuildAccelerator/Common/Public/UbaSessionClient.h b/Engine/Source/Programs/UnrealBuildAccelerator/Common/Public/UbaSessionClient.h
--- a/Engine/Source/Programs/UnrealBuildAccelerator/Common/Public/UbaSessionClient.h
+++ b/Engine/Source/Programs/UnrealBuildAccelerator/Common/Public/UbaSessionClient.h
@@ -60,7 +60,7 @@ namespace uba
 		virtual bool RemoveDirectory(RemoveDirectoryResponse& out, const RemoveDirectoryMessage& msg) override;
 		virtual bool GetFullFileName(GetFullFileNameResponse& out, const GetFullFileNameMessage& msg) override;
 		virtual bool GetLongPathName(GetLongPathNameResponse& out, const GetLongPathNameMessage& msg) override;
-		virtual bool GetListDirectoryInfo(ListDirectoryResponse& out, tchar* dirName, const StringKey& dirKey) override;
+		virtual bool GetListDirectoryInfo(ListDirectoryResponse& out, const StringView& dirName, const StringKey& dirKey) override;
 		virtual bool WriteFilesToDisk(ProcessImpl& process, WrittenFile** files, u32 fileCount) override;
 		virtual bool AllocFailed(Process& process, const tchar* allocType, u32 error) override;
 		virtual void PrintSessionStats(Logger& logger) override;
diff --git a/Engine/Source/Programs/UnrealBuildAccelerator/Core/Private/UbaHash.cpp b/Engine/Source/Programs/UnrealBuildAccelerator/Core/Private/UbaHash.cpp
--- a/Engine/Source/Programs/UnrealBuildAccelerator/Core/Private/UbaHash.cpp
+++ b/Engine/Source/Programs/UnrealBuildAccelerator/Core/Private/UbaHash.cpp
@@ -11,16 +11,21 @@ namespace uba
 		blake3_hasher_init((blake3_hasher*)&hasher);
 	}
 
+	void StringKeyHasher::Update(const StringView& str)
+	{
+		Update(str.data, str.count);
+	}
+		
 	void StringKeyHasher::Update(const tchar* str, u64 strLen)
 	{
-		CHECK_PATH(str);
+		CHECK_PATH(StringView(str, u32(strLen)));
 		if (strLen != 0)
 			blake3_hasher_update((blake3_hasher*)&hasher, str, strLen * sizeof(tchar));
 	}
 
 	StringKey ToStringKey(const tchar* str, u64 strLen)
 	{
-		CHECK_PATH(str);
+		CHECK_PATH(StringView(str, u32(strLen)));
 		blake3_hasher hasher;
 		blake3_hasher_init(&hasher);
 		blake3_hasher_update(&hasher, str, strLen * sizeof(tchar));
@@ -47,7 +52,7 @@ namespace uba
 	}
 	StringKey ToStringKey(const StringKeyHasher& hasher, const tchar* str, u64 strLen)
 	{
-		CHECK_PATH(str);
+		CHECK_PATH(StringView(str, u32(strLen)));
 		StringKeyHasher temp(hasher);
 		blake3_hasher_update((blake3_hasher*)&temp.hasher, str, strLen * sizeof(tchar));
 		uint8_t output[BLAKE3_OUT_LEN];
diff --git a/Engine/Source/Programs/UnrealBuildAccelerator/Core/Public/UbaDirectoryTable.h b/Engine/Source/Programs/UnrealBuildAccelerator/Core/Public/UbaDirectoryTable.h
--- a/Engine/Source/Programs/UnrealBuildAccelerator/Core/Public/UbaDirectoryTable.h
+++ b/Engine/Source/Programs/UnrealBuildAccelerator/Core/Public/UbaDirectoryTable.h
@@ -13,14 +13,13 @@ namespace uba
 	{
 	public:
 		using EntryLookup = GrowingUnorderedMap<StringKey, u32>;
-		struct Directory { Directory(MemoryBlock* block) : files(block) {} u32 tableOffset = InvalidTableOffset; u32 parseOffset = InvalidTableOffset; EntryLookup files; ReaderWriterLock lock; };
+		struct Directory { Directory(MemoryBlock& block) : files(block) {} u32 tableOffset = InvalidTableOffset; u32 parseOffset = InvalidTableOffset; EntryLookup files; ReaderWriterLock lock; };
 
-		void Init(u8* mem, u32 tableCount, u32 tableSize)
+		void Init(const u8* mem, u32 tableCount, u32 tableSize)
 		{
 			m_memory = mem;
 			m_lookup.reserve(tableCount + 100);
-			if (m_memoryBlock)
-				m_memoryBlock->ReserveNoLock(tableCount*(sizeof(GrowingUnorderedMap<StringKey, Directory>::value_type)+16), TC(""));
+			m_memoryBlock.ReserveNoLock(tableCount*(sizeof(GrowingUnorderedMap<StringKey, Directory>::value_type)+16), TC(""));
 			ParseDirectoryTable(tableSize);
 		}
 
@@ -140,13 +139,13 @@ namespace uba
 			Exists_Maybe,
 		};
 
-		Exists EntryExists(StringKey entryKey, const tchar* entryName, u64 entryNameLen, bool checkIfDir = false, u32* tableOffset = nullptr)
+		Exists EntryExists(StringKey entryKey, const StringView& entryName, bool checkIfDir = false, u32* tableOffset = nullptr)
 		{
 			SCOPED_READ_LOCK(m_lookupLock, lock);
-			return EntryExistsNoLock(entryKey, entryName, entryNameLen, checkIfDir, tableOffset);
+			return EntryExistsNoLock(entryKey, entryName, checkIfDir, tableOffset);
 		}
 
-		Exists EntryExistsNoLock(StringKey entryKey, const tchar* entryName, u64 entryNameLen, bool checkIfDir = false, u32* tableOffset = nullptr)
+		Exists EntryExistsNoLock(StringKey entryKey, const StringView& entryName, bool checkIfDir = false, u32* tableOffset = nullptr)
 		{
 			if (checkIfDir)
 			{
@@ -160,9 +159,9 @@ namespace uba
 			}
 
 			// scan backwards first
-			const tchar* rit = entryName + entryNameLen - 2;
+			const tchar* rit = entryName.data + entryName.count - 2;
 			bool inAncestor = false;
-			while (rit > entryName)
+			while (rit > entryName.data)
 			{
 				if (*rit != PathSeparator)
 				{
@@ -170,9 +169,9 @@ namespace uba
 					continue;
 				}
 
-				u64 sublen = u64(rit - entryName);
+				u64 sublen = u64(rit - entryName.data);
 				StringKeyHasher ancestorHasher;
-				ancestorHasher.Update(entryName, sublen);
+				ancestorHasher.Update(entryName.data, sublen);
 				StringKey ancestorKey = ToStringKey(ancestorHasher);
 				auto dirIt = m_lookup.find(ancestorKey);
 				if (dirIt != m_lookup.end())
@@ -205,6 +204,16 @@ namespace uba
 			return Exists_Maybe;
 		}
 
+		Exists EntryExists(const StringView& str, bool checkIfDir = false)
+		{
+			StringBuffer<> str2(str);
+			if (str2[str2.count-1] == PathSeparator)
+				str2.Resize(str2.count-1);
+			if (CaseInsensitiveFs)
+				str2.MakeLower();
+			return EntryExists(ToStringKey(str2), str2, checkIfDir, nullptr);
+		}
+
 		struct EntryInformation
 		{
 			u32 attributes;
@@ -358,16 +367,16 @@ namespace uba
 		}
 #endif
 
-		DirectoryTable(MemoryBlock* block) : m_memoryBlock(block), m_lookup(block) {}
+		DirectoryTable(MemoryBlock& block) : m_memoryBlock(block), m_lookup(block) {}
 		DirectoryTable(const DirectoryTable&) = delete;
 		void operator=(const DirectoryTable&) = delete;
 
-		MemoryBlock* m_memoryBlock;
+		MemoryBlock& m_memoryBlock;
 		ReaderWriterLock m_lookupLock;
 		GrowingUnorderedMap<StringKey, Directory> m_lookup;
 
 		ReaderWriterLock m_memoryLock;
-		u8* m_memory = nullptr;
+		const u8* m_memory = nullptr;
 		u32 m_memorySize = 0;
 	};
 
diff --git a/Engine/Source/Programs/UnrealBuildAccelerator/Core/Public/UbaHash.h b/Engine/Source/Programs/UnrealBuildAccelerator/Core/Public/UbaHash.h
--- a/Engine/Source/Programs/UnrealBuildAccelerator/Core/Public/UbaHash.h
+++ b/Engine/Source/Programs/UnrealBuildAccelerator/Core/Public/UbaHash.h
@@ -96,6 +96,7 @@ namespace uba
 	struct StringKeyHasher
 	{
 		StringKeyHasher();
+		void Update(const StringView& str);
 		void Update(const tchar* str, u64 strLen);
 		u64 hasher[1912 / sizeof(u64)];
 	};
@@ -176,26 +177,26 @@ namespace uba
 	};
 
 	#if UBA_DEBUG
-	inline void CheckPath(const tchar* fileName)
+	inline void CheckPath(const StringView& fileName)
 	{
-		auto pos = fileName;
+		auto pos = fileName.data;
 		tchar history[3] = { 0 };
 		while (tchar c = *pos++)
 		{
-			UBA_ASSERTF(!CaseInsensitiveFs || c < 'A' || c > 'Z', TC("Path is not valid (%s)"), fileName); // No upper case characters if case insensitive
-			UBA_ASSERTF(!(c == PathSeparator && history[0] == '.' && history[1] == PathSeparator) || (IsWindows && TStrstr(fileName, TC("\\\\.\\pipe"))), TC("Path is not valid (%s)"), fileName);  //  /./ not allowed
-			UBA_ASSERTF(!(c == PathSeparator && history[0] == '.' && history[1] == '.' && history[2] == PathSeparator), TC("Path is not valid (%s)"), fileName); // /../ not allowed
-			UBA_ASSERTF(c != NonPathSeparator, TC("Path is not valid (%s)"), fileName);
-			UBA_ASSERTF(c != PathSeparator || history[0] != PathSeparator || (IsWindows && TStrstr(fileName, TC("\\\\.\\pipe"))), TC("Path is not valid (%s)"), fileName); // Double backslash
+			UBA_ASSERTF(!CaseInsensitiveFs || c < 'A' || c > 'Z', TC("Path is not valid (%.*s)"), fileName.count, fileName.data); // No upper case characters if case insensitive
+			UBA_ASSERTF(!(c == PathSeparator && history[0] == '.' && history[1] == PathSeparator) || (IsWindows && fileName.Contains(TC("\\\\.\\pipe"))), TC("Path is not valid (%.*s)"), fileName.count, fileName.data);  //  /./ not allowed
+			UBA_ASSERTF(!(c == PathSeparator && history[0] == '.' && history[1] == '.' && history[2] == PathSeparator), TC("Path is not valid (%.*s)"), fileName.count, fileName.data); // /../ not allowed
+			UBA_ASSERTF(c != NonPathSeparator, TC("Path is not valid (%.*s)"), fileName.count, fileName.data);
+			UBA_ASSERTF(c != PathSeparator || history[0] != PathSeparator || (IsWindows && fileName.Contains(TC("\\\\.\\pipe"))), TC("Path is not valid (%.*s)"), fileName.count, fileName.data); // Double backslash
 			history[2] = history[1];
 			history[1] = history[0];
 			history[0] = c;
 		}
-		UBA_ASSERTF(history[0] != ' ' && (history[0] != PathSeparator || (fileName[1] == ':' && fileName[3] == 0)), TC("Path is not valid (%s)"), fileName);
+		UBA_ASSERTF(history[0] != ' ' && (history[0] != PathSeparator || (fileName.data[1] == ':' && fileName.data[3] == 0)), TC("Path is not valid (%.*s)"), fileName.count, fileName.data);
 		// Commented out because this asserts on dlls in subfolders on remote machines (example: 2057\clui.dll). Likely related to helper not having some language id as host
 		// UBA_ASSERTF(!hasSlash || fileName[1] == ':', TC("Invalid path: %s"), fileName);
 	}
-	#define CHECK_PATH(path) CheckPath((const tchar*)path);
+	#define CHECK_PATH(path) CheckPath(path);
 	#else
 	#define CHECK_PATH(path) while (false) {}
 	#endif
diff --git a/Engine/Source/Programs/UnrealBuildAccelerator/Core/Public/UbaMemory.h b/Engine/Source/Programs/UnrealBuildAccelerator/Core/Public/UbaMemory.h
--- a/Engine/Source/Programs/UnrealBuildAccelerator/Core/Public/UbaMemory.h
+++ b/Engine/Source/Programs/UnrealBuildAccelerator/Core/Public/UbaMemory.h
@@ -122,7 +122,7 @@ namespace uba
 	public:
 		using value_type = Type;
 
-		GrowingAllocator(MemoryBlock* block) : m_block(block) {}
+		GrowingAllocator(MemoryBlock& block) : m_block(&block) {}
 		GrowingAllocator(const GrowingAllocator& o) : m_block(o.m_block) {}
 		GrowingAllocator(GrowingAllocator&& o) noexcept : m_block(o.m_block) {}
 		template <class _Other>
diff --git a/Engine/Source/Programs/UnrealBuildAccelerator/Core/Public/UbaStringBuffer.h b/Engine/Source/Programs/UnrealBuildAccelerator/Core/Public/UbaStringBuffer.h
--- a/Engine/Source/Programs/UnrealBuildAccelerator/Core/Public/UbaStringBuffer.h
+++ b/Engine/Source/Programs/UnrealBuildAccelerator/Core/Public/UbaStringBuffer.h
@@ -86,6 +86,7 @@ namespace uba
 		explicit StringBuffer(const TString& str) : StringBufferBase(Capacity) { *buf = 0; Append(str); }
 		explicit StringBuffer(const tchar* str) : StringBufferBase(Capacity) { *buf = 0; if (str) Append(str); }
 		StringBuffer(const StringBufferBase& str) : StringBufferBase(Capacity) { *buf = 0; Append(str); }
+		StringBuffer(const StringView& str) : StringBufferBase(Capacity) { *buf = 0; Append(str); }
 	private:
 		tchar buf[Capacity];
 	};
diff --git a/Engine/Source/Programs/UnrealBuildAccelerator/Detours/Private/PosixOS/UbaDetoursFunctionsPosix.cpp b/Engine/Source/Programs/UnrealBuildAccelerator/Detours/Private/PosixOS/UbaDetoursFunctionsPosix.cpp
--- a/Engine/Source/Programs/UnrealBuildAccelerator/Detours/Private/PosixOS/UbaDetoursFunctionsPosix.cpp
+++ b/Engine/Source/Programs/UnrealBuildAccelerator/Detours/Private/PosixOS/UbaDetoursFunctionsPosix.cpp
@@ -1224,13 +1224,13 @@ UBA_EXPORT DIR* UBA_WRAPPER(opendir)(const char* name)
 	StringBuffer<> forHash(dirName);
 	if (CaseInsensitiveFs)
 		forHash.MakeLower();
-	DirHash hash(forHash.data,forHash.count);
+	DirHash hash(forHash);
 
 	SCOPED_WRITE_LOCK(g_directoryTable.m_lookupLock, lookLock);
-	auto insres = g_directoryTable.m_lookup.try_emplace(hash.key, &g_memoryBlock);
+	auto insres = g_directoryTable.m_lookup.try_emplace(hash.key, g_memoryBlock);
 	DirectoryTable::Directory& dir = insres.first->second;
 	if (insres.second)
-		if (g_directoryTable.EntryExistsNoLock(hash.key, forHash.data, forHash.count) != DirectoryTable::Exists_No)
+		if (g_directoryTable.EntryExistsNoLock(hash.key, forHash) != DirectoryTable::Exists_No)
 			Rpc_UpdateDirectory(hash.key, dirName.data, dirName.count, false);
 \
 	bool exists = false;
diff --git a/Engine/Source/Programs/UnrealBuildAccelerator/Detours/Private/UbaDetoursFileMappingTable.cpp b/Engine/Source/Programs/UnrealBuildAccelerator/Detours/Private/UbaDetoursFileMappingTable.cpp
--- a/Engine/Source/Programs/UnrealBuildAccelerator/Detours/Private/UbaDetoursFileMappingTable.cpp
+++ b/Engine/Source/Programs/UnrealBuildAccelerator/Detours/Private/UbaDetoursFileMappingTable.cpp
@@ -10,7 +10,7 @@
 
 namespace uba
 {
-	MappedFileTable::MappedFileTable(MemoryBlock& memoryBlock) : m_memoryBlock(memoryBlock), m_lookup(&memoryBlock)
+	MappedFileTable::MappedFileTable(MemoryBlock& memoryBlock) : m_memoryBlock(memoryBlock), m_lookup(memoryBlock)
 	{
 	}
 
@@ -231,8 +231,8 @@ namespace uba
 		else if (entryNameForKey.count == 1 && entryNameForKey[0] == '/')
 			checkIfDir = true;
 
-		CHECK_PATH(entryNameForKey.data);
-		DirectoryTable::Exists exists = g_directoryTable.EntryExists(entryNameKey, entryNameForKey.data, entryNameLen, checkIfDir, &dirTableOffset);
+		CHECK_PATH(entryNameForKey);
+		DirectoryTable::Exists exists = g_directoryTable.EntryExists(entryNameKey, entryNameForKey, checkIfDir, &dirTableOffset);
 		if (exists != DirectoryTable::Exists_Maybe)
 			return dirTableOffset;
 
@@ -247,8 +247,8 @@ namespace uba
 		UBA_ASSERT(wcsncmp(entryName, g_systemTemp.data, g_systemTemp.count) != 0);
 		#endif
 
-		u64 dirNameLen = lastPathSeparator - entryName;
-		DirHash hash(entryNameForKey.data, dirNameLen);
+		u32 dirNameLen = u32(lastPathSeparator - entryName);
+		DirHash hash(StringView(entryNameForKey.data, dirNameLen));
 
 		if (Rpc_UpdateDirectory(hash.key, entryName, dirNameLen) == ~u32(0))
 			return ~u32(0);
@@ -331,10 +331,10 @@ namespace uba
 		pathLen = tempBuf.count;
 	}
 
-	DirHash::DirHash(const tchar* str, u64 strLen)
+	DirHash::DirHash(const StringView& str)
 	{
 		CHECK_PATH(str);
-		open.Update(str, strLen);
+		open.Update(str);
 		key = ToStringKey(open);
 	}
 }
diff --git a/Engine/Source/Programs/UnrealBuildAccelerator/Detours/Private/UbaDetoursFileMappingTable.h b/Engine/Source/Programs/UnrealBuildAccelerator/Detours/Private/UbaDetoursFileMappingTable.h
--- a/Engine/Source/Programs/UnrealBuildAccelerator/Detours/Private/UbaDetoursFileMappingTable.h
+++ b/Engine/Source/Programs/UnrealBuildAccelerator/Detours/Private/UbaDetoursFileMappingTable.h
@@ -69,7 +69,7 @@ namespace uba
 
 	struct DirHash
 	{
-		DirHash(const tchar* str, u64 strLen);
+		DirHash(const StringView& str);
 		StringKey key;
 		StringKeyHasher open;
 	};
diff --git a/Engine/Source/Programs/UnrealBuildAccelerator/Detours/Private/UbaDetoursShared.cpp b/Engine/Source/Programs/UnrealBuildAccelerator/Detours/Private/UbaDetoursShared.cpp
--- a/Engine/Source/Programs/UnrealBuildAccelerator/Detours/Private/UbaDetoursShared.cpp
+++ b/Engine/Source/Programs/UnrealBuildAccelerator/Detours/Private/UbaDetoursShared.cpp
@@ -46,7 +46,7 @@ namespace uba
 
 		u64 reserveSizeMb = IsWindows ? 256 : 1024; // The sync primitives on linux/macos is much bigger
 		g_memoryBlockMem.Create(reserveSizeMb * 1024 * 1024);
-		g_directoryTableMem.Create(&g_memoryBlock);
+		g_directoryTableMem.Create(g_memoryBlock);
 		g_mappedFileTableMem.Create(g_memoryBlock);
 		g_consoleStringCsMem.Create();
 	}
@@ -405,7 +405,7 @@ namespace uba
 			fileNameForKey.MakeLower();
 
 		UBA_ASSERT(fileNameForKey.count);
-		CHECK_PATH(fileNameForKey.data);
+		CHECK_PATH(fileNameForKey);
 		StringKey fileNameKey = ToStringKey(fileNameForKey);
 
 		memset(&outAttr.data, 0, sizeof(outAttr.data));
diff --git a/Engine/Source/Programs/UnrealBuildAccelerator/Detours/Private/Windows/UbaDetoursFunctionsKernelBase.inl b/Engine/Source/Programs/UnrealBuildAccelerator/Detours/Private/Windows/UbaDetoursFunctionsKernelBase.inl
--- a/Engine/Source/Programs/UnrealBuildAccelerator/Detours/Private/Windows/UbaDetoursFunctionsKernelBase.inl
+++ b/Engine/Source/Programs/UnrealBuildAccelerator/Detours/Private/Windows/UbaDetoursFunctionsKernelBase.inl
@@ -1634,7 +1634,7 @@ __forceinline HANDLE Shared_FindFirstFileExW(LPCWSTR lpFileName, FINDEX_INFO_LEV
 		fileName = lastBackslash + 1;
 
 	UBA_ASSERT(lastBackslash);
-	u64 bufChars = (lastBackslash - buf + 1);
+	u32 bufChars = u32(lastBackslash - buf + 1);
 
 	if (wcscmp(fileName, L"*") == 0 || wcscmp(fileName, L"*.*") == 0)
 	{
@@ -1649,15 +1649,15 @@ __forceinline HANDLE Shared_FindFirstFileExW(LPCWSTR lpFileName, FINDEX_INFO_LEV
 			temp[--bufChars] = 0;
 	}
 
-	DirHash hash(buf, bufChars);
+	DirHash hash(StringView(buf, bufChars));
 
 	SCOPED_WRITE_LOCK(g_directoryTable.m_lookupLock, lookLock);
-	auto insres = g_directoryTable.m_lookup.try_emplace(hash.key, &g_memoryBlock);
+	auto insres = g_directoryTable.m_lookup.try_emplace(hash.key, g_memoryBlock);
 	DirectoryTable::Directory& dir = insres.first->second;
 	if (insres.second)
 	{
-		CHECK_PATH(buf);
-		if (g_directoryTable.EntryExistsNoLock(hash.key, buf, bufChars) != DirectoryTable::Exists_No)
+		CHECK_PATH(StringView(buf, bufChars));
+		if (g_directoryTable.EntryExistsNoLock(hash.key, StringView(buf, bufChars)) != DirectoryTable::Exists_No)
 			Rpc_UpdateDirectory(hash.key, buf, bufChars, false);
 	}
 	bool exists = false;
@@ -2485,7 +2485,7 @@ HMODULE Recursive_LoadLibraryExW(LPCWSTR lpLibFileName, LPCWSTR originalName, DW
 
 	// Important that this code is not doing allocations.. it could cause a recursive stack overflow
 	struct Import { wchar_t name[128]; bool isKnown;  Import(const wchar_t* s, bool ik) : isKnown(ik) { wcscpy_s(name, sizeof_array(name), s); } };
-	std::vector<Import, GrowingAllocator<Import>> importedModules(&g_memoryBlock);
+	std::vector<Import, GrowingAllocator<Import>> importedModules(g_memoryBlock);
 	{
 		SuppressCreateFileDetourScope cfs;
 		StringBuffer<256> error;
@@ -2593,8 +2593,8 @@ HMODULE Shared_LoadLibrary(LPCWSTR lpLibFileName, HANDLE hFile, DWORD dwFlags)
 	if (g_runningRemote)
 		Rpc_GetFullFileName(newPath, newPathLen, tempBuf, false);
 
-	AdditionalLoads additionalLoads(&g_memoryBlock); // Don't do allocations
-	VisitedModules visitedModules(&g_memoryBlock);
+	AdditionalLoads additionalLoads(g_memoryBlock); // Don't do allocations
+	VisitedModules visitedModules(g_memoryBlock);
 	HMODULE res = Recursive_LoadLibraryExW(newPath, fileName, dwFlags, additionalLoads, visitedModules);
 	for (HMODULE h : additionalLoads)
 		FreeLibrary(h);
diff --git a/Engine/Source/Programs/UnrealBuildAccelerator/Detours/Private/Windows/UbaDetoursFunctionsNtDll.inl b/Engine/Source/Programs/UnrealBuildAccelerator/Detours/Private/Windows/UbaDetoursFunctionsNtDll.inl
--- a/Engine/Source/Programs/UnrealBuildAccelerator/Detours/Private/Windows/UbaDetoursFunctionsNtDll.inl
+++ b/Engine/Source/Programs/UnrealBuildAccelerator/Detours/Private/Windows/UbaDetoursFunctionsNtDll.inl
@@ -775,7 +775,7 @@ NTSTATUS NTAPI Shared_NtCreateFile(bool IsCreateFunc, PHANDLE hFileHandle, ACCES
 			// .. clang otoh is using CreateFile with tooons of different paths trying to open files.. in remote worker case this becomes super expensive
 			if ((!isWrite || TODO_DELETE_failIfNotExists) && !isSystemOrTempFile) // We need to skip SystemTemp.. lots of stuff going on there.
 			{
-				CHECK_PATH(fileNameLower.data);
+				CHECK_PATH(fileNameLower);
 				dirTableOffset = Rpc_GetEntryOffset(fileNameKey, fileName.data, fileName.count, checkIfDir);
 
 				bool allowEarlyOut = true;
@@ -856,18 +856,17 @@ NTSTATUS NTAPI Shared_NtCreateFile(bool IsCreateFunc, PHANDLE hFileHandle, ACCES
 			return res;
 		}
 
-		u32 pathLen = fileNameLower.count;
 		UBA_ASSERT(fileNameLower.data[fileNameLower.count - 1] != '\\');
-		DirHash hash(fileNameLower.data, pathLen);
+		DirHash hash(fileNameLower);
 
 		SCOPED_WRITE_LOCK(g_directoryTable.m_lookupLock, lookupLock);
-		auto insres = g_directoryTable.m_lookup.try_emplace(hash.key, &g_memoryBlock);
+		auto insres = g_directoryTable.m_lookup.try_emplace(hash.key, g_memoryBlock);
 		DirectoryTable::Directory& dir = insres.first->second;
 		if (insres.second)
 		{
-			auto existsResult = g_directoryTable.EntryExistsNoLock(hash.key, fileNameLower.data, pathLen);
+			auto existsResult = g_directoryTable.EntryExistsNoLock(hash.key, fileNameLower);
 			if (existsResult != DirectoryTable::Exists_No)
-				Rpc_UpdateDirectory(hash.key, fileNameLower.data, pathLen, false);
+				Rpc_UpdateDirectory(hash.key, fileNameLower.data, fileNameLower.count, false);
 		}
 
 		bool exists = false;
@@ -926,7 +925,7 @@ NTSTATUS NTAPI Shared_NtCreateFile(bool IsCreateFunc, PHANDLE hFileHandle, ACCES
 	}
 
 	if (!keepInMemory || !isWrite) // we might get \\pipe\ here... 
-		CHECK_PATH(fileNameLower.data);
+		CHECK_PATH(fileNameLower);
 
 	const wchar_t* lpFileName = fileName.data;
 	u32 closeId = 0;
diff --git a/Engine/Source/Programs/UnrealBuildAccelerator/Detours/Private/Windows/UbaDetoursFunctionsWin.cpp b/Engine/Source/Programs/UnrealBuildAccelerator/Detours/Private/Windows/UbaDetoursFunctionsWin.cpp
--- a/Engine/Source/Programs/UnrealBuildAccelerator/Detours/Private/Windows/UbaDetoursFunctionsWin.cpp
+++ b/Engine/Source/Programs/UnrealBuildAccelerator/Detours/Private/Windows/UbaDetoursFunctionsWin.cpp
@@ -432,7 +432,7 @@ bool EnsureMapped(DetouredHandle& handle, DWORD dwFileOffsetHigh = 0, DWORD dwFi
 }
 
 ReaderWriterLock g_longPathNameCacheLock;
-GrowingUnorderedMap<const wchar_t*, const wchar_t*, HashString, EqualString> g_longPathNameCache(&g_memoryBlock);
+GrowingUnorderedMap<const wchar_t*, const wchar_t*, HashString, EqualString> g_longPathNameCache(g_memoryBlock);
 
 void Rpc_AllocFailed(const wchar_t* allocType, u32 error)
 {
diff --git a/Engine/Source/Programs/UnrealBuildAccelerator/Host/Private/UbaExports.cpp b/Engine/Source/Programs/UnrealBuildAccelerator/Host/Private/UbaExports.cpp
--- a/Engine/Source/Programs/UnrealBuildAccelerator/Host/Private/UbaExports.cpp
+++ b/Engine/Source/Programs/UnrealBuildAccelerator/Host/Private/UbaExports.cpp
@@ -358,6 +358,10 @@ extern "C"
 	{
 		server->RegisterDeleteFile(filePath);
 	}
+	bool SessionServer_RegisterNewDirectory(uba::SessionServer* server, const uba::tchar* directoryPath)
+	{
+		return server->RegisterNewDirectory(directoryPath);
+	}
 	bool SessionServer_RegisterVirtualFile(uba::SessionServer* server, const uba::tchar* filename, const uba::tchar* sourceFile, uba::u64 sourceOffset, uba::u64 sourceSize)
 	{
 		return server->RegisterVirtualFile(filename, sourceFile, sourceOffset, sourceSize);
diff --git a/Engine/Source/Programs/UnrealBuildAccelerator/Host/Public/UbaExports.h b/Engine/Source/Programs/UnrealBuildAccelerator/Host/Public/UbaExports.h
--- a/Engine/Source/Programs/UnrealBuildAccelerator/Host/Public/UbaExports.h
+++ b/Engine/Source/Programs/UnrealBuildAccelerator/Host/Public/UbaExports.h
@@ -102,6 +102,7 @@ extern "C"
 	UBA_API bool SessionServer_RefreshDirectory(uba::SessionServer* server, const uba::tchar* directory);
 	UBA_API bool SessionServer_RegisterNewFile(uba::SessionServer* server, const uba::tchar* filePath);
 	UBA_API void SessionServer_RegisterDeleteFile(uba::SessionServer* server, const uba::tchar* filePath);
+	UBA_API bool SessionServer_RegisterNewDirectory(uba::SessionServer* server, const uba::tchar* directoryPath);
 	UBA_API bool SessionServer_RegisterVirtualFile(uba::SessionServer* server, const uba::tchar* filename, const uba::tchar* sourceFile, uba::u64 sourceOffset, uba::u64 sourceSize);
 	UBA_API uba::ProcessHandle* SessionServer_RunProcess(uba::SessionServer* server, uba::ProcessStartInfo& info, bool async, bool enableDetour);
 	UBA_API uba::ProcessHandle* SessionServer_RunProcessRemote(uba::SessionServer* server, uba::ProcessStartInfo& info, float weight, const void* knownInputs = nullptr, uba::u32 knownInputsCount = 0, bool allowCrossArchitecture = false);
diff --git a/Engine/Source/Programs/UnrealBuildAccelerator/Test/Private/UbaTestAll.h b/Engine/Source/Programs/UnrealBuildAccelerator/Test/Private/UbaTestAll.h
--- a/Engine/Source/Programs/UnrealBuildAccelerator/Test/Private/UbaTestAll.h
+++ b/Engine/Source/Programs/UnrealBuildAccelerator/Test/Private/UbaTestAll.h
@@ -41,6 +41,7 @@ namespace uba
 		UBA_TEST(TestCrypto) \
 		UBA_TEST(TestBinDependencies) \
 		UBA_TEST(TestRootPaths) \
+		UBA_TEST(TestRegisterChangesToSession);
 
 
 #define UBA_POSIX_TESTS \
diff --git a/Engine/Source/Programs/UnrealBuildAccelerator/Test/Private/UbaTestBasics.h b/Engine/Source/Programs/UnrealBuildAccelerator/Test/Private/UbaTestBasics.h
--- a/Engine/Source/Programs/UnrealBuildAccelerator/Test/Private/UbaTestBasics.h
+++ b/Engine/Source/Programs/UnrealBuildAccelerator/Test/Private/UbaTestBasics.h
@@ -163,7 +163,7 @@ namespace uba
 			return logger.Error(TC("Failed to get last written time"));
 
 		u64 writeTime2 = 0;
-		TraverseDir(logger, rootDir.data, [&](const DirectoryEntry& de)
+		TraverseDir(logger, rootDir, [&](const DirectoryEntry& de)
 			{
 				if (Equals(de.name, TC("UbaTestFile")))
 					writeTime2 = de.lastWritten;
@@ -240,7 +240,7 @@ namespace uba
 			return logger.Error(TC("GetCurrentDirectoryW failed"));
 
 		bool foundFile = false;
-		if (!TraverseDir(logger, rootDir.data, [&](const DirectoryEntry& de) { foundFile |= TStrcmp(de.name, TC("UbaTestFile2")) == 0; }, true))
+		if (!TraverseDir(logger, rootDir, [&](const DirectoryEntry& de) { foundFile |= TStrcmp(de.name, TC("UbaTestFile2")) == 0; }, true))
 			return logger.Error(TC("Failed to TraverseDir '.'"));
 
 		if (!foundFile)
@@ -251,7 +251,7 @@ namespace uba
 		#endif
 
 		LoggerWithWriter nullLogger(g_nullLogWriter);
-		if (TraverseDir(nullLogger, TC("TestDir2"), [&](const DirectoryEntry&) {}, true))
+		if (TraverseDir(nullLogger, AsView(TC("TestDir2")), [&](const DirectoryEntry&) {}, true))
 			return logger.Error(TC("TraverseDir failed to report fail on non existing dir"));
 
 		return true;
diff --git a/Engine/Source/Programs/UnrealBuildAccelerator/Test/Private/UbaTestSession.h b/Engine/Source/Programs/UnrealBuildAccelerator/Test/Private/UbaTestSession.h
--- a/Engine/Source/Programs/UnrealBuildAccelerator/Test/Private/UbaTestSession.h
+++ b/Engine/Source/Programs/UnrealBuildAccelerator/Test/Private/UbaTestSession.h
@@ -464,4 +464,74 @@ namespace uba
 	{
 		return RunLocal(logger, testRootDir, RunSystemApplicationAndLookForLog, false);
 	}
+
+	bool TestRegisterChangesToSession(LoggerWithWriter& logger, const StringBufferBase& testRootDir)
+	{
+		return RunLocal(logger, testRootDir, [](LoggerWithWriter& logger, SessionServer& session, const tchar* workingDir, const RunProcessFunction& runProcess)
+			{
+				MemoryBlock temp;
+				DirectoryTable dirTable(temp);
+				dirTable.Init(session.GetDirectoryTableMemory(), 0, 0);
+
+				if (!session.RefreshDirectory(workingDir, true))
+					return false;
+				if (!session.RefreshDirectory(workingDir))
+					return false;
+
+				if (dirTable.EntryExists(ToView(workingDir)) != DirectoryTable::Exists_Maybe)
+					return false;
+				dirTable.ParseDirectoryTable(session.GetDirectoryTableSize());
+				if (dirTable.EntryExists(ToView(workingDir), true) != DirectoryTable::Exists_Yes)
+					return false;
+
+				StringBuffer<MaxPath> sourceFile;
+				sourceFile.Append(workingDir).Append(TC("Code.cpp"));
+				FileAccessor codeFile(logger, sourceFile.data);
+				if (!codeFile.CreateWrite())
+					return false;
+				if (!codeFile.Close())
+					return false;
+				if (!session.RegisterNewFile(sourceFile.data))
+					return false;
+				if (dirTable.EntryExists(sourceFile) != DirectoryTable::Exists_No)
+					return false;
+				dirTable.ParseDirectoryTable(session.GetDirectoryTableSize());
+				if (dirTable.EntryExists(sourceFile) != DirectoryTable::Exists_Yes)
+					return false;
+
+				StringBuffer<MaxPath> newDir;
+				newDir.Append(workingDir).Append(TC("NewDir/"));
+				if (!CreateDirectoryW(newDir.data))
+					return false;
+				if (!session.RegisterNewFile(newDir.data))
+					return false;
+				newDir.Resize(newDir.count - 1);
+				if (dirTable.EntryExists(newDir) != DirectoryTable::Exists_No)
+					return false;
+				dirTable.ParseDirectoryTable(session.GetDirectoryTableSize());
+				if (dirTable.EntryExists(newDir) != DirectoryTable::Exists_Yes)
+					return false;
+
+				StringBuffer<MaxPath> newDir2;
+				newDir2.Append(workingDir).Append(TC("NewDir2"));
+				if (!CreateDirectoryW(newDir2.data))
+					return false;
+				if (!session.RefreshDirectory(workingDir))
+					return false;
+				dirTable.ParseDirectoryTable(session.GetDirectoryTableSize());
+				if (dirTable.EntryExists(newDir2) != DirectoryTable::Exists_Yes)
+					return false;
+
+				return true;
+			});
+
+		/*
+		StringBuffer<MaxPath> sourceFile;
+		sourceFile.Append(workingDir).Append(TC("Code.cpp"));
+		FileAccessor codeFile(logger, sourceFile.data);
+		if (!codeFile.CreateWrite())
+			return false;
+		*/
+		return true;
+	}
 }
