[UBA]
* Fixed bug where error strings are devirtualized

diff --git a/Engine/Source/Programs/UnrealBuildAccelerator/Common/Private/UbaSession.cpp b/Engine/Source/Programs/UnrealBuildAccelerator/Common/Private/UbaSession.cpp
--- a/Engine/Source/Programs/UnrealBuildAccelerator/Common/Private/UbaSession.cpp
+++ b/Engine/Source/Programs/UnrealBuildAccelerator/Common/Private/UbaSession.cpp
@@ -1681,10 +1681,10 @@ namespace uba
 				memcpy(start, path.data(), path.size()*sizeof(tchar));
 
 				#if PLATFORM_WINDOWS
-				u32 type = rootIndex % RootPaths::RootStartByte;
+				u32 type = rootIndex % PathsPerRoot;
 				if (type != 1) // Not backslash
 				{
-					UBA_ASSERTF(type == 0, TC("Not implemented. (%s)"), hint); // We don't like double backslash or escaped space
+					UBA_ASSERTF(type == 0, TC("Not root path type %u not implemented (%s)"), type, hint); // We don't like double backslash or escaped space
 					*newStringPos = 0;
 					Replace(start, '\\', '/');
 				}
