[UBA]
* More tests around GetFileAttributes etc after doing RegisterNewFile

diff --git a/Engine/Source/Programs/UnrealBuildAccelerator/Core/Public/UbaProtocol.h b/Engine/Source/Programs/UnrealBuildAccelerator/Core/Public/UbaProtocol.h
--- a/Engine/Source/Programs/UnrealBuildAccelerator/Core/Public/UbaProtocol.h
+++ b/Engine/Source/Programs/UnrealBuildAccelerator/Core/Public/UbaProtocol.h
@@ -59,7 +59,7 @@ namespace uba
 // Currently only used for detoured process
 #if UBA_DEBUG
 #define UBA_DEBUG_LOG_ENABLED 1
-#define UBA_DEBUG_VALIDATE 1
+#define UBA_DEBUG_VALIDATE 0 // Does not work with some unit tests
 #else
 #define UBA_DEBUG_LOG_ENABLED 0
 #define UBA_DEBUG_VALIDATE 0
diff --git a/Engine/Source/Programs/UnrealBuildAccelerator/Test/Private/UbaTestAll.h b/Engine/Source/Programs/UnrealBuildAccelerator/Test/Private/UbaTestAll.h
--- a/Engine/Source/Programs/UnrealBuildAccelerator/Test/Private/UbaTestAll.h
+++ b/Engine/Source/Programs/UnrealBuildAccelerator/Test/Private/UbaTestAll.h
@@ -2,6 +2,10 @@
 
 #pragma once
 
+#define CHECK_TRUE(x) \
+	if (!(x)) \
+		return logger.Error(TC("Failed %s (%s:%u)"), TC(#x), __FILE__, __LINE__);
+
 #include "UbaTestBasics.h"
 #include "UbaTestCache.h"
 #include "UbaTestConfig.h"
@@ -41,7 +45,8 @@ namespace uba
 		UBA_TEST(TestCrypto) \
 		UBA_TEST(TestBinDependencies) \
 		UBA_TEST(TestRootPaths) \
-		UBA_TEST(TestRegisterChangesToSession);
+		UBA_TEST(TestRegisterChanges) \
+		UBA_TEST(TestRegisterChangesRemote) \
 
 
 #define UBA_POSIX_TESTS \
diff --git a/Engine/Source/Programs/UnrealBuildAccelerator/Test/Private/UbaTestSession.h b/Engine/Source/Programs/UnrealBuildAccelerator/Test/Private/UbaTestSession.h
--- a/Engine/Source/Programs/UnrealBuildAccelerator/Test/Private/UbaTestSession.h
+++ b/Engine/Source/Programs/UnrealBuildAccelerator/Test/Private/UbaTestSession.h
@@ -340,15 +340,11 @@ namespace uba
 				file.Append(workingDir).Append(TC("TouchFile.h"));
 				FileAccessor fr(logger, file.data);
 
-				if (!fr.CreateWrite())
-					return false;
-				if (!fr.Write("Foo", 4))
-					return false;
-				if (!fr.Close())
-					return false;
+				CHECK_TRUE(fr.CreateWrite());
+				CHECK_TRUE(fr.Write("Foo", 4));
+				CHECK_TRUE(fr.Close());
 				FileInformation oldInfo;
-				if (!GetFileInformation(oldInfo, logger, file.data))
-					return false;
+				CHECK_TRUE(GetFileInformation(oldInfo, logger, file.data));
 
 				Sleep(100);
 
@@ -365,8 +361,7 @@ namespace uba
 					return false;
 
 				FileInformation newInfo;
-				if (!GetFileInformation(newInfo, logger, file.data))
-					return false;
+				CHECK_TRUE(GetFileInformation(newInfo, logger, file.data));
 				if (newInfo.lastWriteTime == oldInfo.lastWriteTime)
 					return logger.Error(TC("File time not changed after touch"));
 				return true;
@@ -465,73 +460,87 @@ namespace uba
 		return RunLocal(logger, testRootDir, RunSystemApplicationAndLookForLog, false);
 	}
 
-	bool TestRegisterChangesToSession(LoggerWithWriter& logger, const StringBufferBase& testRootDir)
+	bool CheckAttributes(LoggerWithWriter& logger, SessionServer& session, const tchar* workingDir, const RunProcessFunction& runProcess)
 	{
-		return RunLocal(logger, testRootDir, [](LoggerWithWriter& logger, SessionServer& session, const tchar* workingDir, const RunProcessFunction& runProcess)
+		StringBuffer<MaxPath> testApp;
+		GetTestAppPath(logger, testApp);
+		ProcessStartInfo processInfo;
+		processInfo.application = testApp.data;
+		processInfo.workingDir = workingDir;
+		processInfo.logLineFunc = [](void* userData, const tchar* line, u32 length, LogEntryType type)
 			{
-				MemoryBlock temp;
-				DirectoryTable dirTable(temp);
-				dirTable.Init(session.GetDirectoryTableMemory(), 0, 0);
-
-				if (!session.RefreshDirectory(workingDir, true))
-					return false;
-				if (!session.RefreshDirectory(workingDir))
-					return false;
-
-				if (dirTable.EntryExists(ToView(workingDir)) != DirectoryTable::Exists_Maybe)
-					return false;
-				dirTable.ParseDirectoryTable(session.GetDirectoryTableSize());
-				if (dirTable.EntryExists(ToView(workingDir), true) != DirectoryTable::Exists_Yes)
-					return false;
-
-				StringBuffer<MaxPath> sourceFile;
-				sourceFile.Append(workingDir).Append(TC("Code.cpp"));
-				FileAccessor codeFile(logger, sourceFile.data);
-				if (!codeFile.CreateWrite())
-					return false;
-				if (!codeFile.Close())
-					return false;
-				if (!session.RegisterNewFile(sourceFile.data))
-					return false;
-				if (dirTable.EntryExists(sourceFile) != DirectoryTable::Exists_No)
-					return false;
-				dirTable.ParseDirectoryTable(session.GetDirectoryTableSize());
-				if (dirTable.EntryExists(sourceFile) != DirectoryTable::Exists_Yes)
-					return false;
+				LoggerWithWriter(g_consoleLogWriter, TC("")).Info(line);
+			};
 
-				StringBuffer<MaxPath> newDir;
-				newDir.Append(workingDir).Append(TC("NewDir/"));
-				if (!CreateDirectoryW(newDir.data))
-					return false;
-				if (!session.RegisterNewFile(newDir.data))
-					return false;
-				newDir.Resize(newDir.count - 1);
-				if (dirTable.EntryExists(newDir) != DirectoryTable::Exists_No)
-					return false;
-				dirTable.ParseDirectoryTable(session.GetDirectoryTableSize());
-				if (dirTable.EntryExists(newDir) != DirectoryTable::Exists_Yes)
-					return false;
+		auto GetAttributes = [&](const StringView& file) -> u32
+			{
+				StringBuffer<> arg(TC("-GetFileAttributes="));
+				arg.Append(file);
+				processInfo.arguments = arg.data;
+				ProcessHandle process = runProcess(processInfo);
+				if (!process.WaitForExit(100000))
+					return logger.Error(TC("UbaTestApp did not exit in 10 seconds"));
+				u32 exitCode =  process.GetExitCode();
+				return exitCode == 255 ? INVALID_FILE_ATTRIBUTES : exitCode;
+			};
 
-				StringBuffer<MaxPath> newDir2;
-				newDir2.Append(workingDir).Append(TC("NewDir2"));
-				if (!CreateDirectoryW(newDir2.data))
-					return false;
-				if (!session.RefreshDirectory(workingDir))
-					return false;
-				dirTable.ParseDirectoryTable(session.GetDirectoryTableSize());
-				if (dirTable.EntryExists(newDir2) != DirectoryTable::Exists_Yes)
-					return false;
+		MemoryBlock temp;
+		DirectoryTable dirTable(temp);
+		dirTable.Init(session.GetDirectoryTableMemory(), 0, 0);
 
-				return true;
-			});
+		CHECK_TRUE(session.RefreshDirectory(workingDir, true));
+		CHECK_TRUE(session.RefreshDirectory(workingDir));
+		CHECK_TRUE(dirTable.EntryExists(ToView(workingDir)) == DirectoryTable::Exists_Maybe);
+		dirTable.ParseDirectoryTable(session.GetDirectoryTableSize());
+		CHECK_TRUE(dirTable.EntryExists(ToView(workingDir), true) == DirectoryTable::Exists_Yes);
 
-		/*
 		StringBuffer<MaxPath> sourceFile;
 		sourceFile.Append(workingDir).Append(TC("Code.cpp"));
+
+		CHECK_TRUE(GetAttributes(sourceFile) == INVALID_FILE_ATTRIBUTES);
 		FileAccessor codeFile(logger, sourceFile.data);
-		if (!codeFile.CreateWrite())
-			return false;
-		*/
+		CHECK_TRUE(codeFile.CreateWrite());
+		CHECK_TRUE(codeFile.Close());
+		CHECK_TRUE(session.RegisterNewFile(sourceFile.data));
+		CHECK_TRUE(GetAttributes(sourceFile) != INVALID_FILE_ATTRIBUTES);
+
+		CHECK_TRUE(dirTable.EntryExists(sourceFile) == DirectoryTable::Exists_No);
+		dirTable.ParseDirectoryTable(session.GetDirectoryTableSize());
+		CHECK_TRUE(dirTable.EntryExists(sourceFile) == DirectoryTable::Exists_Yes);
+
+		StringBuffer<MaxPath> newDir;
+		newDir.Append(workingDir).Append(TC("NewDir"));
+		StringBuffer<MaxPath> newDirAndSlash(newDir);
+		newDirAndSlash.Append('/');
+
+		CHECK_TRUE(GetAttributes(newDir) == INVALID_FILE_ATTRIBUTES);
+		CHECK_TRUE(CreateDirectoryW(newDir.data));
+		CHECK_TRUE(session.RegisterNewFile(newDir.data));
+		CHECK_TRUE(dirTable.EntryExists(newDir) == DirectoryTable::Exists_No);
+		dirTable.ParseDirectoryTable(session.GetDirectoryTableSize());
+		CHECK_TRUE(dirTable.EntryExists(newDir) == DirectoryTable::Exists_Yes);
+		CHECK_TRUE(GetAttributes(newDir) != INVALID_FILE_ATTRIBUTES);
+		CHECK_TRUE(GetAttributes(newDirAndSlash) != INVALID_FILE_ATTRIBUTES);
+
+		StringBuffer<MaxPath> newDir2;
+		newDir2.Append(workingDir).Append(TC("NewDir2"));
+		CHECK_TRUE(CreateDirectoryW(newDir2.data));
+		CHECK_TRUE(GetAttributes(newDir2) == INVALID_FILE_ATTRIBUTES);
+		CHECK_TRUE(session.RefreshDirectory(workingDir))
+		CHECK_TRUE(GetAttributes(newDir2) != INVALID_FILE_ATTRIBUTES);
+		dirTable.ParseDirectoryTable(session.GetDirectoryTableSize());
+		CHECK_TRUE(dirTable.EntryExists(newDir2) == DirectoryTable::Exists_Yes);
+
 		return true;
 	}
+
+	bool TestRegisterChanges(LoggerWithWriter& logger, const StringBufferBase& testRootDir)
+	{
+		return RunLocal(logger, testRootDir, CheckAttributes);
+	}
+
+	bool TestRegisterChangesRemote(LoggerWithWriter& logger, const StringBufferBase& testRootDir)
+	{
+		return RunRemote(logger, testRootDir, CheckAttributes);
+	}
 }
diff --git a/Engine/Source/Programs/UnrealBuildAccelerator/TestApp/Private/PosixOS/UbaTestAppPosix.cpp b/Engine/Source/Programs/UnrealBuildAccelerator/TestApp/Private/PosixOS/UbaTestAppPosix.cpp
--- a/Engine/Source/Programs/UnrealBuildAccelerator/TestApp/Private/PosixOS/UbaTestAppPosix.cpp
+++ b/Engine/Source/Programs/UnrealBuildAccelerator/TestApp/Private/PosixOS/UbaTestAppPosix.cpp
@@ -230,6 +230,14 @@ int main(int argc, char* argv[])
 			return LogError("stat for 'FileW' in child process failed");
 		return 0;
 	}
+	else if (strncmp(argv[1], "-GetFileAttributes=", 6) == 0)
+	{
+		const char* str = argv[1] + 19;
+		struct stat attr;
+		if (stat(str, &attr) == -1)
+			return 255;
+		return attr.st_mode;
+	}
 	else if (strncmp(argv[1], "-file=", 6) == 0)
 	{
 		void* detoursHandle = dlopen(UBA_DETOURS_LIBRARY, RTLD_LAZY);
diff --git a/Engine/Source/Programs/UnrealBuildAccelerator/TestApp/Private/Windows/UbaTestAppWin.cpp b/Engine/Source/Programs/UnrealBuildAccelerator/TestApp/Private/Windows/UbaTestAppWin.cpp
--- a/Engine/Source/Programs/UnrealBuildAccelerator/TestApp/Private/Windows/UbaTestAppWin.cpp
+++ b/Engine/Source/Programs/UnrealBuildAccelerator/TestApp/Private/Windows/UbaTestAppWin.cpp
@@ -237,6 +237,12 @@ int wmain(int argc, wchar_t* argv[])
 
 		return 0;
 	}
+	else if (wcsncmp(argv[1], L"-GetFileAttributes=", 6) == 0)
+	{
+		const wchar_t* str = argv[1] + 19;
+		DWORD attr = GetFileAttributesW(str);
+		return attr == INVALID_FILE_ATTRIBUTES ? 255 : attr;
+	}
 	else if (wcsncmp(argv[1], L"-stdout=", 8) == 0)
 	{
 		const wchar_t* str = argv[1] + 8;
