[UBA]
* Fixed macos compile error

diff --git a/Engine/Source/Programs/UnrealBuildAccelerator/Agent/Private/UbaAgent.cpp b/Engine/Source/Programs/UnrealBuildAccelerator/Agent/Private/UbaAgent.cpp
--- a/Engine/Source/Programs/UnrealBuildAccelerator/Agent/Private/UbaAgent.cpp
+++ b/Engine/Source/Programs/UnrealBuildAccelerator/Agent/Private/UbaAgent.cpp
@@ -826,7 +826,7 @@ namespace uba
 			StringBuffer<> applicationsDir;
 			applicationsDir.Append("/Applications");
 			
-			TraverseDir(logger, applicationsDir.data,
+			TraverseDir(logger, applicationsDir,
 						[&](const DirectoryEntry& e)
 						{
 				if (IsDirectory(e.attributes) && StartsWith(e.name, "Xcode"))
