[UBA]
* Removed callstack resolver spam on osx

diff --git a/Engine/Source/Programs/UnrealBuildAccelerator/Common/Private/UbaLogger.cpp b/Engine/Source/Programs/UnrealBuildAccelerator/Common/Private/UbaLogger.cpp
--- a/Engine/Source/Programs/UnrealBuildAccelerator/Common/Private/UbaLogger.cpp
+++ b/Engine/Source/Programs/UnrealBuildAccelerator/Common/Private/UbaLogger.cpp
@@ -638,7 +638,7 @@ namespace uba
 			#if PLATFORM_LINUX
 			cmd.Appendf("addr2line 0x%x", entry.memoryOffset).Append(" -e ").Append(mod.name).Append(" 2>/dev/null");
 			#else
-			cmd.Appendf("atos -o %s --offset 0x%llx", mod.name.c_str(), entry.memoryOffset);
+			cmd.Appendf("atos -o %s --offset 0x%llx", mod.name.c_str(), entry.memoryOffset).Append(" 2>/dev/null");
 			#endif
 
 			fflush(stdout);
