[UBA]
* Added -version to UbaAgent that can show all the various system network versions

diff --git a/Engine/Source/Programs/UnrealBuildAccelerator/Agent/Private/UbaAgent.cpp b/Engine/Source/Programs/UnrealBuildAccelerator/Agent/Private/UbaAgent.cpp
--- a/Engine/Source/Programs/UnrealBuildAccelerator/Agent/Private/UbaAgent.cpp
+++ b/Engine/Source/Programs/UnrealBuildAccelerator/Agent/Private/UbaAgent.cpp
@@ -101,6 +101,7 @@ namespace uba
 		logger.Info(TC("  -eventfile=<file>       File containing external events to agent. Things like machine is about to be terminated etc"));
 		logger.Info(TC("  -sentry                 Enable sentry"));
 		logger.Info(TC("  -zone                   Set the zone this machine exists in. This info is used to figure out if proxies should be created."));
+		logger.Info(TC("  -version                Prints the version for this executable."));
 		logger.Info(TC("  -noproxy                Does not allow this agent to be a storage proxy for other agents"));
 		logger.Info(TC("  -killrandom             Kills random process and exit session"));
 		logger.Info(TC("  -memwait=<percent>      The amount of memory needed to spawn a process. Set this to 100 to disable. Defaults to 80%%"));
@@ -660,6 +661,15 @@ namespace uba
 					return PrintHelp(TC("-zone needs a value"));
 				zone.Append(value);
 			}
+			else if (name.Equals(TC("-version")))
+			{
+				const tchar* dbgStr = TC("");
+				#if UBA_DEBUG
+				dbgStr = TC(" (DEBUG)");
+				#endif
+				LoggerWithWriter(g_consoleLogWriter, TC("")).Info(TC("v%s%s (Network: %u, Storage: %u, Session: %u, Cache: %u)"), Version, dbgStr, SystemNetworkVersion, StorageNetworkVersion, SessionNetworkVersion, CacheNetworkVersion);
+				return true;
+			}
 			else if (name.Equals(TC("-noproxy")))
 			{
 				allowProxy = false;
