[UBA]
* Fixed bug in createprocess when process have no quotes or arguments

diff --git a/Engine/Source/Programs/UnrealBuildAccelerator/Detours/Private/Windows/UbaDetoursFunctionsKernelBase.inl b/Engine/Source/Programs/UnrealBuildAccelerator/Detours/Private/Windows/UbaDetoursFunctionsKernelBase.inl
--- a/Engine/Source/Programs/UnrealBuildAccelerator/Detours/Private/Windows/UbaDetoursFunctionsKernelBase.inl
+++ b/Engine/Source/Programs/UnrealBuildAccelerator/Detours/Private/Windows/UbaDetoursFunctionsKernelBase.inl
@@ -2767,18 +2767,24 @@ BOOL Detoured_CreateProcessW(LPCWSTR lpApplicationName, LPWSTR lpCommandLine, LP
 		isChild = false;
 	}
 
-	StringBuffer<512> application;
 	StringView commandLineWithoutApplication;
-
 	{
 		const tchar* endOfApplication;
 		if (originalCmd.data[0] == '"')
-			endOfApplication = TStrchr(originalCmd.data+1, '"') + 1;
+		{
+			const tchar* quoteEnd = TStrchr(originalCmd.data+1, '"');
+			UBA_ASSERT(quoteEnd);
+			endOfApplication = quoteEnd + 1;
+		}
+		else if (const tchar* firstSpace = TStrchr(originalCmd.data, ' '))
+			endOfApplication = firstSpace + 1;
 		else
-			endOfApplication = TStrchr(originalCmd.data, ' ') + 1;
+			endOfApplication = originalCmd.data + originalCmd.count;
+
 		commandLineWithoutApplication = StringView(endOfApplication, originalCmd.count - u32(endOfApplication - originalCmd.data));
 	}
 
+	StringBuffer<512> application;
 	if (lpApplicationName)
 	{
 		FixPath(application, lpApplicationName);
