[UBA]
* Added WINDOWS_LEAN_AND_MEAN before windows includes

diff --git a/Engine/Source/Programs/UnrealBuildAccelerator/Core/Public/UbaPlatform.h b/Engine/Source/Programs/UnrealBuildAccelerator/Core/Public/UbaPlatform.h
--- a/Engine/Source/Programs/UnrealBuildAccelerator/Core/Public/UbaPlatform.h
+++ b/Engine/Source/Programs/UnrealBuildAccelerator/Core/Public/UbaPlatform.h
@@ -5,6 +5,7 @@
 #include "UbaBase.h"
 
 #if PLATFORM_WINDOWS
+	#define WINDOWS_LEAN_AND_MEAN
 	#define NOMINMAX
 	#include <ntstatus.h>
 	#define WIN32_NO_STATUS
diff --git a/Engine/Source/Programs/UnrealBuildAccelerator/ObjTool/Private/UbaImportLibWriter.cpp b/Engine/Source/Programs/UnrealBuildAccelerator/ObjTool/Private/UbaImportLibWriter.cpp
--- a/Engine/Source/Programs/UnrealBuildAccelerator/ObjTool/Private/UbaImportLibWriter.cpp
+++ b/Engine/Source/Programs/UnrealBuildAccelerator/ObjTool/Private/UbaImportLibWriter.cpp
@@ -5,6 +5,7 @@
 #include "UbaObjectFile.h"
 
 #if PLATFORM_WINDOWS
+#define WINDOWS_LEAN_AND_MEAN
 #include <windows.h>
 #include <algorithm>
 #endif
diff --git a/Engine/Source/Programs/UnrealBuildAccelerator/TestApp/Private/Windows/UbaTestAppWin.cpp b/Engine/Source/Programs/UnrealBuildAccelerator/TestApp/Private/Windows/UbaTestAppWin.cpp
--- a/Engine/Source/Programs/UnrealBuildAccelerator/TestApp/Private/Windows/UbaTestAppWin.cpp
+++ b/Engine/Source/Programs/UnrealBuildAccelerator/TestApp/Private/Windows/UbaTestAppWin.cpp
@@ -1,5 +1,6 @@
 // Copyright Epic Games, Inc. All Rights Reserved.
 
+#define WINDOWS_LEAN_AND_MEAN
 #include <Windows.h>
 #include <wchar.h> 
 #include <stdarg.h>
