[UBA]
* Added some code to collect more information about when there are active remote processes but zero connections. Should not be possible but has now been seen once in the wild
* Added code to early out when deqeueing processes to disconnected session (again should not happen)

diff --git a/Engine/Source/Programs/UnrealBuildAccelerator/Common/Private/UbaSessionServer.cpp b/Engine/Source/Programs/UnrealBuildAccelerator/Common/Private/UbaSessionServer.cpp
--- a/Engine/Source/Programs/UnrealBuildAccelerator/Common/Private/UbaSessionServer.cpp
+++ b/Engine/Source/Programs/UnrealBuildAccelerator/Common/Private/UbaSessionServer.cpp
@@ -550,10 +550,10 @@ namespace uba
 		{
 			++sessionId;
 			auto& s = *sptr;
-			if (s.id != clientId)
+			if (s.clientId != clientId)
 				continue;
 
-			if (!returnCount && !s.hasNotification && !m_remoteExecutionEnabled)
+			if (!returnCount && !s.hasNotification && !s.enabled)
 				m_trace.SessionNotification(sessionId, TC("Done"));
 
 			m_trace.SessionDisconnect(sessionId);
@@ -565,6 +565,7 @@ namespace uba
 			if (s.enabled)
 				m_availableRemoteSlotCount -= s.processSlotCount - returnCount;
 			s.enabled = false;
+			s.connected = false;
 			--m_connectionCount;
 		}
 
@@ -596,6 +597,8 @@ namespace uba
 
 		if (!m_activeRemoteProcesses.empty())
 		{
+			// This path has been seen in the wild (Once over million of runs)....
+			// And the theory is that 
 			m_logger.Error(TC("No client sessions connected but there are %llu active remote processes. This should not happen, there is a bug in the code!!"), m_activeRemoteProcesses.size());
 		}
 	}
@@ -715,7 +718,7 @@ namespace uba
 		m_clientSessions.push_back(&session);
 		u32 sessionId = u32(m_clientSessions.size());
 		session.name = name.data;
-		session.id = connectionInfo.GetId();
+		session.clientId = connectionInfo.GetId();
 		session.processSlotCount = processSlotCount;
 		session.dedicated = dedicated;
 		session.isArm = isClientArm;
@@ -1172,7 +1175,7 @@ namespace uba
 		SCOPED_WRITE_LOCK(m_fillUpOneAtTheTimeLock, fillLock); // This is a lock to group files better (all clients connect at the same time);
 		while (weightLeft > 0)
 		{
-			RemoteProcess* process = DequeueProcess(sessionId, connectionInfo.GetId());
+			RemoteProcess* process = DequeueProcess(session, sessionId, connectionInfo.GetId());
 			if (process == nullptr)
 				break;
 			auto& startInfo = process->m_startInfo;
@@ -1572,7 +1575,7 @@ namespace uba
 			{
 				if (!s->enabled)
 					continue;
-				WriteString(StringBuffer<>().Appendf(TC("Session %u (%s)"), s->id, s->name.c_str()).data);
+				WriteString(StringBuffer<>().Appendf(TC("Session %u (%s)"), s->clientId, s->name.c_str()).data);
 				WriteString(StringBuffer<>().Appendf(TC("   Process slots used %u/%u"), s->usedSlotCount, s->processSlotCount).data);
 				if (s->pingTime)
 					WriteString(StringBuffer<>().Appendf(TC("   Last ping %s ago"), TimeToText(time - s->pingTime).str).data);
@@ -1596,7 +1599,7 @@ namespace uba
 			{
 				if (!s->enabled || s->abort)
 					continue;
-				bool hasProxy = m_storage.HasProxy(s->id);
+				bool hasProxy = m_storage.HasProxy(s->clientId);
 				if (abortWithProxy && hasProxy)
 					s->abort = true;
 				else if (abortUseProxy && !hasProxy)
@@ -1833,17 +1836,20 @@ namespace uba
 		m_newestWaitingProcess = nullptr;
 	}
 
-	SessionServer::RemoteProcess* SessionServer::DequeueProcess(u32 sessionId, u32 clientId)
+	SessionServer::RemoteProcess* SessionServer::DequeueProcess(ClientSession& session, u32 sessionId, u32 clientId)
 	{
 		SCOPED_READ_LOCK(m_remoteProcessSlotAvailableEventLock, lock);
 		bool hasCalledCallback = !m_remoteProcessSlotAvailableEvent;
-		u32 sessionIndex = sessionId - 1;
 
 		while (true)
 		{
 			ScopedCriticalSection queueLock(m_remoteProcessAndSessionLock);
-			UBA_ASSERT(sessionIndex < m_clientSessions.size());
-			auto& session = *m_clientSessions[sessionIndex];
+
+			if (!session.connected) // This should not be possible
+			{
+				m_logger.Warning(TC("Dequeing process to session that is not connected. This should never happen. Report to Epic (%u)"), clientId);
+				return nullptr;
+			}
 
 			while (!m_queuedRemoteProcesses.empty())
 			{
@@ -2050,7 +2056,7 @@ namespace uba
 		{
 			auto& s = *sptr;
 			NetworkServer::ClientStats stats;
-			m_server.GetClientStats(stats, s.id);
+			m_server.GetClientStats(stats, s.clientId);
 			if (stats.connectionCount && (stats.send || stats.recv))
 				m_trace.SessionUpdate(sessionIndex, stats.connectionCount, stats.send, stats.recv, s.lastPing, s.memAvail, s.memTotal, s.cpuLoad);
 			++sessionIndex;
diff --git a/Engine/Source/Programs/UnrealBuildAccelerator/Common/Public/UbaSessionServer.h b/Engine/Source/Programs/UnrealBuildAccelerator/Common/Public/UbaSessionServer.h
--- a/Engine/Source/Programs/UnrealBuildAccelerator/Common/Public/UbaSessionServer.h
+++ b/Engine/Source/Programs/UnrealBuildAccelerator/Common/Public/UbaSessionServer.h
@@ -92,7 +92,7 @@ namespace uba
 
 		class RemoteProcess;
 
-		RemoteProcess* DequeueProcess(u32 sessionId, u32 clientId);
+		RemoteProcess* DequeueProcess(ClientSession& session, u32 sessionId, u32 clientId);
 		void OnCancelled(RemoteProcess* process);
 		ProcessHandle ProcessRemoved(u32 processId);
 
@@ -144,7 +144,7 @@ namespace uba
 			UnorderedSet<CasKey> sentKeys;
 			ReaderWriterLock dirTablePosLock;
 			u32 dirTablePos = 0;
-			u32 id = ~0u;
+			u32 clientId = ~0u;
 			u32 processSlotCount = 1;
 			u32 usedSlotCount = 0;
 			u64 lastPing = 0;
@@ -152,6 +152,7 @@ namespace uba
 			u64 memTotal = 0;
 			u64 pingTime = 0;
 			float cpuLoad = 0;
+			bool connected = true;
 			bool enabled = true;
 			bool dedicated = false;
 			bool abort = false;
