[UBA]
* Fixed so detoured SetConsoleMode is forwarding the true handle and not detoured one to true SetConsoleMode.
* Fixed so CREATE_NO_WINDOW is not set for CreateProcessW if there has been calls to SetConsoleMode since it is likely something related to stdin/out has changed (batch files like doing this)
* Fixed bug when application is a .bat file and arguments are quoted..
* Fixed so we don't try to find module imports in batch files
* Improved error message when failing to resolve module imports
* Fixed so uba assert callstack contains last entry on win64

diff --git a/Engine/Source/Programs/UnrealBuildAccelerator/Common/Private/UbaProcess.cpp b/Engine/Source/Programs/UnrealBuildAccelerator/Common/Private/UbaProcess.cpp
--- a/Engine/Source/Programs/UnrealBuildAccelerator/Common/Private/UbaProcess.cpp
+++ b/Engine/Source/Programs/UnrealBuildAccelerator/Common/Private/UbaProcess.cpp
@@ -1444,7 +1444,11 @@ namespace uba
 			StringBuffer<> application(m_realApplication);
 			m_session.VirtualizePath(application, m_startInfo.rootsHandle);
 
-			TString commandLine = TString(TC("\"")) + application.data + TC("\" ") + m_startInfo.arguments;
+			TString commandLine;
+			if (application.EndsWith(TC(".bat"))) // If there are quotes around arguments used by batch file we need to quote the entire thing 
+				commandLine = TString(TC("\"\"")) + application.data + TC("\" ") + m_startInfo.arguments + TC("\"");
+			else
+				commandLine = TString(TC("\"")) + application.data + TC("\" ") + m_startInfo.arguments;
 
 			if (m_extractExports)
 			{
@@ -1559,7 +1563,7 @@ namespace uba
 					si.hStdOutput = writePipe;
 					si.dwFlags |= STARTF_USESTDHANDLES;
 
-					if (CreateProcessW(NULL, (tchar*)commandLine.c_str(), NULL, NULL, TRUE, creationFlags, environment, workingDir, &si, &processInfo))
+					if (CreateProcessW(m_realApplication.c_str(), (tchar*)commandLine.c_str(), NULL, NULL, TRUE, creationFlags, environment, workingDir, &si, &processInfo))
 					{
 						DisableProcessPowerThrottling(processInfo.hProcess);
 						break;
diff --git a/Engine/Source/Programs/UnrealBuildAccelerator/Common/Private/UbaSessionClient.cpp b/Engine/Source/Programs/UnrealBuildAccelerator/Common/Private/UbaSessionClient.cpp
--- a/Engine/Source/Programs/UnrealBuildAccelerator/Common/Private/UbaSessionClient.cpp
+++ b/Engine/Source/Programs/UnrealBuildAccelerator/Common/Private/UbaSessionClient.cpp
@@ -1428,8 +1428,8 @@ namespace uba
 		NetworkMessage msg(m_client, ServiceId, SessionMessageType_ProcessFinished, writer);
 		writer.WriteU32(process.m_id);
 		writer.WriteU32(exitCode);
-		writer.WriteU32(CountLogLines(process));
-		WriteLogLines(writer, process);
+		u32* lineCount = (u32*)writer.AllocWrite(sizeof(u32));
+		*lineCount = WriteLogLines(writer, process);
 
 		// This is normally set after callback so we need to calculate it here
 		auto& exitTime = process.m_processStats.exitTime;
@@ -1974,23 +1974,20 @@ namespace uba
 		}
 	}
 
-	u32 SessionClient::CountLogLines(ProcessImpl& process)
+	u32 SessionClient::WriteLogLines(BinaryWriter& writer, ProcessImpl& process)
 	{
-		u32 count = u32(process.m_logLines.size());
+		u32 logLineCount = 0;
 		for (auto& child : process.m_childProcesses)
-			count += CountLogLines(*(ProcessImpl*)child.m_process);
-		return count;
-	}
-
-	void SessionClient::WriteLogLines(BinaryWriter& writer, ProcessImpl& process)
-	{
-		for (auto& child : process.m_childProcesses)
-			WriteLogLines(writer, *(ProcessImpl*)child.m_process);
+			logLineCount += WriteLogLines(writer, *(ProcessImpl*)child.m_process);
 		for (auto& line : process.m_logLines)
 		{
+			if (line.text.size()*sizeof(tchar) + 1000 >= writer.GetCapacityLeft())
+				break;
 			writer.WriteString(line.text);
 			writer.WriteByte(line.type);
+			++logLineCount;
 		}
+		return logLineCount;
 	}
 
 	bool SessionClient::AllocFailed(Process& process, const tchar* allocType, u32 error)
diff --git a/Engine/Source/Programs/UnrealBuildAccelerator/Common/Public/UbaSessionClient.h b/Engine/Source/Programs/UnrealBuildAccelerator/Common/Public/UbaSessionClient.h
--- a/Engine/Source/Programs/UnrealBuildAccelerator/Common/Public/UbaSessionClient.h
+++ b/Engine/Source/Programs/UnrealBuildAccelerator/Common/Public/UbaSessionClient.h
@@ -98,8 +98,7 @@ namespace uba
 		bool SendRootsHandle(RootsHandle rootsHandle);
 		void SendLogFileToServer(ProcessImpl& pi);
 		void GetLogFileName(StringBufferBase& out, const tchar* logFile, const tchar* arguments, u32 processId);
-		u32 CountLogLines(ProcessImpl& process);
-		void WriteLogLines(BinaryWriter& writer, ProcessImpl& process);
+		u32 WriteLogLines(BinaryWriter& writer, ProcessImpl& process);
 
 		void ThreadCreateProcessLoop();
 
diff --git a/Engine/Source/Programs/UnrealBuildAccelerator/Core/Private/UbaPlatform.cpp b/Engine/Source/Programs/UnrealBuildAccelerator/Core/Private/UbaPlatform.cpp
--- a/Engine/Source/Programs/UnrealBuildAccelerator/Core/Private/UbaPlatform.cpp
+++ b/Engine/Source/Programs/UnrealBuildAccelerator/Core/Private/UbaPlatform.cpp
@@ -518,6 +518,17 @@ namespace uba
 		UbaAssert(text.data, nullptr, 0, nullptr, true, exceptionCode, contextPtr);
 		return EXCEPTION_EXECUTE_HANDLER;
 	}
+
+	DWORD UbaExceptionHandlerAllowAccessViolation(DWORD code, _EXCEPTION_POINTERS* info, const tchar* fileName)
+	{
+		if (code != EXCEPTION_ACCESS_VIOLATION)
+			return EXCEPTION_CONTINUE_SEARCH;
+		StringBuffer<> buf;
+		buf.Appendf(L"Access violation reading %s", fileName);
+		UbaAssert(buf.data, __FILE__, __LINE__, "", true, code, info->ContextRecord);
+		return EXCEPTION_EXECUTE_HANDLER;
+	}
+
 #else
 	static const int g_allSignals[] = {
 		//#ifdef SIGHUP
diff --git a/Engine/Source/Programs/UnrealBuildAccelerator/Core/Public/UbaWinBinDependencyParser.h b/Engine/Source/Programs/UnrealBuildAccelerator/Core/Public/UbaWinBinDependencyParser.h
--- a/Engine/Source/Programs/UnrealBuildAccelerator/Core/Public/UbaWinBinDependencyParser.h
+++ b/Engine/Source/Programs/UnrealBuildAccelerator/Core/Public/UbaWinBinDependencyParser.h
@@ -99,6 +99,8 @@ namespace uba
 		return rva - seh->VirtualAddress + seh->PointerToRawData;
 	} 
 
+	DWORD UbaExceptionHandlerAllowAccessViolation(DWORD code, _EXCEPTION_POINTERS* info, const tchar* fileName);
+
 	inline bool FindImportsInMem(const tchar* fileName, void* mem, const Function<void(const tchar* import, bool isKnown, const char* const* importLoaderPaths)>& func)
 	{
 		auto hdrs = (PIMAGE_NT_HEADERS)(PCHAR(mem) + PIMAGE_DOS_HEADER(mem)->e_lfanew);   
@@ -119,11 +121,8 @@ namespace uba
 				++importDesc;
 			}
 		}
-		__except(GetExceptionCode() == EXCEPTION_ACCESS_VIOLATION ? EXCEPTION_EXECUTE_HANDLER : EXCEPTION_CONTINUE_SEARCH)
+		__except(UbaExceptionHandlerAllowAccessViolation(GetExceptionCode(), GetExceptionInformation(), fileName))
 		{
-			//StringBuffer<> buf;
-			//buf.Appendf(L"Access violation reading %s", fileName);
-			//UbaAssert(buf.data, __FILE__, __LINE__, "", true, GetExceptionCode(), GetExceptionInformation()->ContextRecord);
 			return false;
 		}
 		return true;
@@ -131,6 +130,8 @@ namespace uba
 
 	inline bool FindImports(const tchar* fileName, const Function<void(const tchar* import, bool isKnown, const char* const* importLoaderPaths)>& func, StringBufferBase& outError)
 	{
+		if (EndsWith(fileName, TStrlen(fileName), TC(".bat")))
+			return true;
 		HANDLE fileHandle = True_CreateFileW(fileName, GENERIC_READ, FILE_SHARE_READ, 0, OPEN_EXISTING, FILE_ATTRIBUTE_NORMAL | FILE_FLAG_SEQUENTIAL_SCAN, 0);
 		if (fileHandle == INVALID_HANDLE_VALUE)
 			return true;
diff --git a/Engine/Source/Programs/UnrealBuildAccelerator/Detours/Private/Windows/UbaDetoursFunctionsKernelBase.inl b/Engine/Source/Programs/UnrealBuildAccelerator/Detours/Private/Windows/UbaDetoursFunctionsKernelBase.inl
--- a/Engine/Source/Programs/UnrealBuildAccelerator/Detours/Private/Windows/UbaDetoursFunctionsKernelBase.inl
+++ b/Engine/Source/Programs/UnrealBuildAccelerator/Detours/Private/Windows/UbaDetoursFunctionsKernelBase.inl
@@ -2690,12 +2690,21 @@ BOOL Detoured_GetConsoleMode(HANDLE hConsoleHandle, LPDWORD lpMode)
 	return res;
 }
 
+bool g_setConsoleModeCalled;
+
 BOOL Detoured_SetConsoleMode(HANDLE hConsoleHandle, DWORD mode)
 {
 	DETOURED_CALL(SetConsoleMode);
-	DEBUG_LOG_DETOURED(L"SetConsoleMode", L"%llu (%u)", u64(hConsoleHandle), mode);
+	DEBUG_LOG_DETOURED(L"SetConsoleMode", L"%llu (0x%x)", u64(hConsoleHandle), mode);
+
+	HANDLE trueHandle = hConsoleHandle;
 
-	if (hConsoleHandle == g_stdHandle[1])
+	if (hConsoleHandle == g_stdHandle[0])
+	{
+		if (g_isDetachedProcess)
+			trueHandle = True_GetStdHandle(STD_ERROR_HANDLE);
+	}
+	else if (hConsoleHandle == g_stdHandle[1])
 	{
 		g_echoOn = (mode & ~503) != 0; // TODO: This might be wrong. Trying to figure out how echo off in batch files work in terms of win32 calls
 
@@ -2704,8 +2713,23 @@ BOOL Detoured_SetConsoleMode(HANDLE hConsoleHandle, DWORD mode)
 		writer.WriteByte(MessageType_EchoOn);
 		writer.WriteBool(g_echoOn);
 		writer.Flush();
+		if (g_isDetachedProcess)
+			trueHandle = True_GetStdHandle(STD_OUTPUT_HANDLE);
+	}
+	else if (hConsoleHandle == g_stdHandle[0])
+	{
+		if (g_isDetachedProcess)
+			trueHandle = True_GetStdHandle(STD_INPUT_HANDLE);
 	}
-	return True_SetConsoleMode(hConsoleHandle, mode);
+	else if (isDetouredHandle(hConsoleHandle))
+	{
+		g_setConsoleModeCalled = true;
+		DetouredHandle& dh = asDetouredHandle(hConsoleHandle);
+		trueHandle = dh.trueHandle;
+		UBA_ASSERTF(trueHandle != INVALID_HANDLE_VALUE, TC("SetConsoleMode is using %s as handle to set mode 0x%x"), HandleToName(hConsoleHandle), mode);
+	}
+
+	return True_SetConsoleMode(trueHandle, mode);
 }
 
 BOOL Detoured_GetConsoleTitleW(LPTSTR lpConsoleTitle, DWORD nSize)
@@ -2873,7 +2897,7 @@ BOOL Detoured_CreateProcessW(LPCWSTR lpApplicationName, LPWSTR lpCommandLine, LP
 
 	if (g_rules->AllowDetach())
 		dwCreationFlags |= DETACHED_PROCESS;
-	else
+	else if (!g_setConsoleModeCalled) // If console mode has been called there are probably some sort of stdin/out rerouting stuff.. so let's not set CREATE_NO_WINDOW
 		dwCreationFlags |= CREATE_NO_WINDOW;
 
 	dwCreationFlags |= CREATE_SUSPENDED;
diff --git a/Engine/Source/Programs/UnrealBuildAccelerator/Detours/Private/Windows/UbaDetoursFunctionsWin.cpp b/Engine/Source/Programs/UnrealBuildAccelerator/Detours/Private/Windows/UbaDetoursFunctionsWin.cpp
--- a/Engine/Source/Programs/UnrealBuildAccelerator/Detours/Private/Windows/UbaDetoursFunctionsWin.cpp
+++ b/Engine/Source/Programs/UnrealBuildAccelerator/Detours/Private/Windows/UbaDetoursFunctionsWin.cpp
@@ -332,7 +332,7 @@ UBA_NOINLINE void UbaAssert(const tchar* text, const char* file, u32 line, const
 	ScopedCriticalSection s(cs);
 	static auto& sb = *new StringBuffer<16*1024>();
 	WriteAssertInfo(sb.Clear(), text, file, line, expr);
-	Rpc_ResolveCallstack(sb, 4, context);
+	Rpc_ResolveCallstack(sb, 3, context);
 	Rpc_WriteLog(sb.data, sb.count, true, true);
 
 	#if UBA_DEBUG_LOG_ENABLED
