[UBA]
* Added exception handler to win and linux test app

diff --git a/Engine/Source/Programs/UnrealBuildAccelerator/Test/Private/Linux/UbaTestLinux.cpp b/Engine/Source/Programs/UnrealBuildAccelerator/Test/Private/Linux/UbaTestLinux.cpp
deleted file mode 100644
diff --git a/Engine/Source/Programs/UnrealBuildAccelerator/Test/Private/Mac/UbaTestMac.cpp b/Engine/Source/Programs/UnrealBuildAccelerator/Test/Private/Mac/UbaTestMac.cpp
deleted file mode 100644
diff --git a/Engine/Source/Programs/UnrealBuildAccelerator/Test/Private/UbaTest.cpp b/Engine/Source/Programs/UnrealBuildAccelerator/Test/Private/UbaTest.cpp
new file mode 100644
--- /dev/null
+++ b/Engine/Source/Programs/UnrealBuildAccelerator/Test/Private/UbaTest.cpp
@@ -0,0 +1,24 @@
+// Copyright Epic Games, Inc. All Rights Reserved.
+
+#include "UbaTestAll.h"
+
+namespace uba
+{
+	bool WrappedMain(int argc, tchar* argv[])
+	{
+		AddExceptionHandler();
+		return RunTests(argc, argv);
+	}
+}
+
+#if PLATFORM_WINDOWS
+int wmain(int argc, wchar_t* argv[])
+{
+	return uba::WrappedMain(argc, argv) ? 0 : -1;
+}
+#else
+int main(int argc, char* argv[])
+{
+	return uba::WrappedMain(argc, argv) ? 0 : -1;
+}
+#endif
diff --git a/Engine/Source/Programs/UnrealBuildAccelerator/Test/Private/Windows/UbaTestWin.cpp b/Engine/Source/Programs/UnrealBuildAccelerator/Test/Private/Windows/UbaTestWin.cpp
deleted file mode 100644
