[UBA]
* Added asserts in various places to try to figure out what is going wrong on the farm tests

diff --git a/Engine/Source/Programs/UnrealBuildAccelerator/Common/Private/UbaNetworkServer.cpp b/Engine/Source/Programs/UnrealBuildAccelerator/Common/Private/UbaNetworkServer.cpp
--- a/Engine/Source/Programs/UnrealBuildAccelerator/Common/Private/UbaNetworkServer.cpp
+++ b/Engine/Source/Programs/UnrealBuildAccelerator/Common/Private/UbaNetworkServer.cpp
@@ -1124,6 +1124,8 @@ namespace uba
 			lock.Enter();
 		}
 
+		UBA_ASSERT(m_firstActiveWorker == nullptr);
+
 		auto worker = m_firstAvailableWorker;
 		while (worker)
 		{
diff --git a/Engine/Source/Programs/UnrealBuildAccelerator/Common/Private/UbaProcess.cpp b/Engine/Source/Programs/UnrealBuildAccelerator/Common/Private/UbaProcess.cpp
--- a/Engine/Source/Programs/UnrealBuildAccelerator/Common/Private/UbaProcess.cpp
+++ b/Engine/Source/Programs/UnrealBuildAccelerator/Common/Private/UbaProcess.cpp
@@ -971,6 +971,8 @@ namespace uba
 		TString applicationStr = reader.ReadString();
 		TString commandLineWithoutApplication = ReadLongString(reader);
 
+		UBA_ASSERTF(!applicationStr.empty() && (applicationStr[0] != '\"' || applicationStr[1] != '\"'), TC("Application: %s"), applicationStr.c_str());
+
 		StringBuffer<> currentDir;
 		reader.ReadString(currentDir);
 		if (currentDir.IsEmpty())
diff --git a/Engine/Source/Programs/UnrealBuildAccelerator/Common/Private/UbaSession.cpp b/Engine/Source/Programs/UnrealBuildAccelerator/Common/Private/UbaSession.cpp
--- a/Engine/Source/Programs/UnrealBuildAccelerator/Common/Private/UbaSession.cpp
+++ b/Engine/Source/Programs/UnrealBuildAccelerator/Common/Private/UbaSession.cpp
@@ -1505,14 +1505,15 @@ namespace uba
 	{
 		u32 exeNameStart = 0;
 		u32 exeNameEnd = TStrlen(si.application);
-		const tchar* lastSeparator = TStrrchr(si.application, PathSeparator);
-		if (lastSeparator)
+		UBA_ASSERT(exeNameEnd > 1);
+		if (const tchar* lastSeparator = TStrrchr(si.application, PathSeparator))
 			exeNameStart = u32(lastSeparator - si.application + 1);
 		else if (si.application[exeNameStart] == '"')
 			++exeNameStart;
 		if (si.application[exeNameEnd - 1] == '"')
 			--exeNameEnd;
 		StringBuffer<128> exeName;
+		UBA_ASSERTF(exeNameStart < exeNameEnd, TC("Bad application string: %s"), si.application);
 		exeName.Append(si.application + exeNameStart, exeNameEnd - exeNameStart);
 		
 		auto rules = GetApplicationRules();
diff --git a/Engine/Source/Programs/UnrealBuildAccelerator/Detours/Private/Windows/UbaDetoursFunctionsKernelBase.inl b/Engine/Source/Programs/UnrealBuildAccelerator/Detours/Private/Windows/UbaDetoursFunctionsKernelBase.inl
--- a/Engine/Source/Programs/UnrealBuildAccelerator/Detours/Private/Windows/UbaDetoursFunctionsKernelBase.inl
+++ b/Engine/Source/Programs/UnrealBuildAccelerator/Detours/Private/Windows/UbaDetoursFunctionsKernelBase.inl
@@ -2809,7 +2809,7 @@ BOOL Detoured_CreateProcessW(LPCWSTR lpApplicationName, LPWSTR lpCommandLine, LP
 	}
 
 	StringBuffer<512> application;
-	if (lpApplicationName)
+	if (lpApplicationName && *lpApplicationName)
 	{
 		FixPath(application, lpApplicationName);
 	}
diff --git a/Engine/Source/Programs/UnrealBuildAccelerator/Test/Private/UbaTestStorage.h b/Engine/Source/Programs/UnrealBuildAccelerator/Test/Private/UbaTestStorage.h
--- a/Engine/Source/Programs/UnrealBuildAccelerator/Test/Private/UbaTestStorage.h
+++ b/Engine/Source/Programs/UnrealBuildAccelerator/Test/Private/UbaTestStorage.h
@@ -131,6 +131,7 @@ namespace uba
 		Sleep(100);
 		if (!client.Connect(clientTcp, TC("127.0.0.1"), 1234))
 			return logger.Error(TC("Failed to connect"));
+		auto cg = MakeGuard([&]() { client.Disconnect(); });
 
 		StringBuffer<> fileName;
 		{
@@ -197,6 +198,7 @@ namespace uba
 		Sleep(100);
 		if (!client.Connect(clientTcp, TC("127.0.0.1"), 1234))
 			return logger.Error(TC("Failed to connect"));
+		auto cg = MakeGuard([&]() { client.Disconnect(); });
 
 		StringBuffer<> fileName;
 #if 1
