[UBA]
* Fixed so log entries are only devirtualized once
* Fixed some paths that did not handle devirtualization properly
* Fixed so logging from Shared_GetFileAttributesExW is showing the caller name

diff --git a/Engine/Source/Programs/UnrealBuildAccelerator/Common/Private/UbaProcess.cpp b/Engine/Source/Programs/UnrealBuildAccelerator/Common/Private/UbaProcess.cpp
--- a/Engine/Source/Programs/UnrealBuildAccelerator/Common/Private/UbaProcess.cpp
+++ b/Engine/Source/Programs/UnrealBuildAccelerator/Common/Private/UbaProcess.cpp
@@ -61,7 +61,7 @@ namespace uba
 		~PipeReader()
 		{
 			if (!currentString.empty())
-				process.LogLine(false, TString(currentString), logType);
+				process.InternalLogLine(false, TString(currentString), logType);
 		}
 
 		void ReadData(char* buf, u32 readCount)
@@ -79,7 +79,7 @@ namespace uba
 				if (endOfLine > buf && endOfLine[-1] == '\r')
 					--endOfLine;
 				currentString.append(TString(startPos, endOfLine));
-				process.LogLine(false, TString(currentString), logType);
+				process.InternalLogLine(false, TString(currentString), logType);
 				currentString.clear();
 				startPos = newStart;
 			}
@@ -712,9 +712,6 @@ namespace uba
 	{
 		if (IsCancelled())// || m_startInfo.stdoutHandle != INVALID_HANDLE_VALUE)
 			return;
-
-		m_session.DevirtualizeString(line, m_startInfo.rootsHandle, TC("LogLine"));
-
 		if (printInSession)
 			m_session.m_logger.Log(LogEntryType_Warning, line.c_str(), u32(line.size()));
 		if (m_startInfo.logLineFunc)
@@ -971,7 +968,7 @@ namespace uba
 		TString applicationStr = reader.ReadString();
 		TString commandLineWithoutApplication = ReadLongString(reader);
 
-		UBA_ASSERTF(!applicationStr.empty() && (applicationStr[0] != '\"' || applicationStr[1] != '\"'), TC("Application: %s"), applicationStr.c_str());
+		UBA_ASSERTF(!applicationStr.empty() && (applicationStr[0] != '\"' || applicationStr[1] != '\"'), TC("Invalid application name: %s"), applicationStr.c_str());
 
 		StringBuffer<> currentDir;
 		reader.ReadString(currentDir);
@@ -1145,7 +1142,7 @@ namespace uba
 		if (!m_session.LogLine(*this, line.c_str(), entryType))
 			return false;
 
-		LogLine(printInSession, std::move(line), entryType);
+		InternalLogLine(printInSession, std::move(line), entryType);
 		return true;
 	}
 
@@ -2097,6 +2094,12 @@ namespace uba
 #endif
 	}
 
+	void ProcessImpl::InternalLogLine(bool printInSession, TString&& line, LogEntryType logType)
+	{
+		m_session.DevirtualizeString(line, m_startInfo.rootsHandle, true, TC("LogLine"));
+		LogLine(printInSession, std::move(line), logType);
+	}
+
 #if !PLATFORM_WINDOWS
 	bool ProcessImpl::PollStdPipes(PipeReader& outReader, PipeReader& errReader, int timeoutMs)
 	{
diff --git a/Engine/Source/Programs/UnrealBuildAccelerator/Common/Private/UbaSession.cpp b/Engine/Source/Programs/UnrealBuildAccelerator/Common/Private/UbaSession.cpp
--- a/Engine/Source/Programs/UnrealBuildAccelerator/Common/Private/UbaSession.cpp
+++ b/Engine/Source/Programs/UnrealBuildAccelerator/Common/Private/UbaSession.cpp
@@ -1634,7 +1634,7 @@ namespace uba
 		return true;
 	}
 
-	bool Session::DevirtualizeString(TString& inOut, RootsHandle rootsHandle, const tchar* hint)
+	bool Session::DevirtualizeString(TString& inOut, RootsHandle rootsHandle, bool allowPathsWithoutRoot, const tchar* hint)
 	{
 		if (!HasVfs(rootsHandle))
 			return true;
@@ -1656,7 +1656,7 @@ namespace uba
 				newSize += path.size() - strLen;
 			};
 
-		if (!rootsEntry->roots.NormalizeString<tchar>(m_logger, inOut.data(), inOut.size(), checkString, false, hint))
+		if (!rootsEntry->roots.NormalizeString<tchar>(m_logger, inOut.data(), inOut.size(), checkString, allowPathsWithoutRoot, hint))
 			return false;
 
 		if (!hasRoot)
@@ -1692,7 +1692,7 @@ namespace uba
 				#endif
 			};
 
-		rootsEntry->roots.NormalizeString<tchar>(m_logger, inOut.data(), inOut.size(), handleString, false, hint);
+		rootsEntry->roots.NormalizeString<tchar>(m_logger, inOut.data(), inOut.size(), handleString, allowPathsWithoutRoot, hint);
 		inOut = std::move(newString);
 
 		return true;
diff --git a/Engine/Source/Programs/UnrealBuildAccelerator/Common/Public/UbaProcess.h b/Engine/Source/Programs/UnrealBuildAccelerator/Common/Public/UbaProcess.h
--- a/Engine/Source/Programs/UnrealBuildAccelerator/Common/Public/UbaProcess.h
+++ b/Engine/Source/Programs/UnrealBuildAccelerator/Common/Public/UbaProcess.h
@@ -93,6 +93,7 @@ namespace uba
 
 		u32 InternalCreateProcess(void* environment, FileMappingHandle communicationHandle, u64 communicationOffset);
 		u32	InternalExitProcess(bool cancel);
+		void InternalLogLine(bool printInSession, TString&& line, LogEntryType logType);
 		void ClearTempFiles();
 		void WaitForParent();
 
diff --git a/Engine/Source/Programs/UnrealBuildAccelerator/Common/Public/UbaSession.h b/Engine/Source/Programs/UnrealBuildAccelerator/Common/Public/UbaSession.h
--- a/Engine/Source/Programs/UnrealBuildAccelerator/Common/Public/UbaSession.h
+++ b/Engine/Source/Programs/UnrealBuildAccelerator/Common/Public/UbaSession.h
@@ -88,7 +88,7 @@ namespace uba
 
 		bool VirtualizePath(StringBufferBase& inOut, RootsHandle rootsHandle);
 		bool DevirtualizePath(StringBufferBase& inOut, RootsHandle rootsHandle);
-		bool DevirtualizeString(TString& inOut, RootsHandle rootsHandle, const tchar* hint = TC(""));
+		bool DevirtualizeString(TString& inOut, RootsHandle rootsHandle, bool allowPathsWithoutRoot, const tchar* hint = TC(""));
 		bool PopulateLocalToIndexRoots(RootPaths& out, RootsHandle rootsHandle);
 
 		// For testing
diff --git a/Engine/Source/Programs/UnrealBuildAccelerator/Detours/Private/Windows/UbaDetoursFunctionsKernelBase.inl b/Engine/Source/Programs/UnrealBuildAccelerator/Detours/Private/Windows/UbaDetoursFunctionsKernelBase.inl
--- a/Engine/Source/Programs/UnrealBuildAccelerator/Detours/Private/Windows/UbaDetoursFunctionsKernelBase.inl
+++ b/Engine/Source/Programs/UnrealBuildAccelerator/Detours/Private/Windows/UbaDetoursFunctionsKernelBase.inl
@@ -991,7 +991,7 @@ DWORD Detoured_GetFileType(HANDLE hFile)
 	return True_GetFileType(hFile); // Calling NtQueryVolumeInformationFile
 }
 
-BOOL Shared_GetFileAttributesExW(LPCWSTR lpFileName, GET_FILEEX_INFO_LEVELS fInfoLevelId, LPVOID lpFileInformation, LPCWSTR originalName)
+BOOL Shared_GetFileAttributesExW(const tchar* caller, LPCWSTR lpFileName, GET_FILEEX_INFO_LEVELS fInfoLevelId, LPVOID lpFileInformation, LPCWSTR originalName)
 {
 	const wchar_t* fixedFileName = lpFileName;
 	if (StartsWith(fixedFileName, L"\\\\?\\"))
@@ -999,7 +999,7 @@ BOOL Shared_GetFileAttributesExW(LPCWSTR lpFileName, GET_FILEEX_INFO_LEVELS fInf
 
 	if (!CanDetour(fixedFileName))
 	{
-		DEBUG_LOG_TRUE(L"GetFileAttributesExW", L"(%ls)", originalName);
+		DEBUG_LOG_TRUE(L"GetFileAttributesExW", L"(%s) (%s)", caller, originalName);
 		TimerScope ts(g_kernelStats.getFileInfo);
 		return True_GetFileAttributesExW(lpFileName, fInfoLevelId, lpFileInformation);
 	}
@@ -1009,7 +1009,7 @@ BOOL Shared_GetFileAttributesExW(LPCWSTR lpFileName, GET_FILEEX_INFO_LEVELS fInf
 
 	if (!attr.useCache)
 	{
-		DEBUG_LOG_TRUE(L"GetFileAttributesExW", L"(%ls)", originalName);
+		DEBUG_LOG_TRUE(L"GetFileAttributesExW", L"(%s) (%s)", caller, originalName);
 		TimerScope ts(g_kernelStats.getFileInfo);
 		return True_GetFileAttributesExW(realName, fInfoLevelId, lpFileInformation);
 	}
@@ -1018,7 +1018,7 @@ BOOL Shared_GetFileAttributesExW(LPCWSTR lpFileName, GET_FILEEX_INFO_LEVELS fInf
 
 	memcpy(lpFileInformation, &attr.data, sizeof(WIN32_FILE_ATTRIBUTE_DATA));
 
-	DEBUG_LOG_DETOURED(L"GetFileAttributesExW", L"(%ls) -> %u", originalName, attr.exists);
+	DEBUG_LOG_DETOURED(caller, L"(%ls) -> %u", originalName, attr.exists);
 	return attr.exists;
 }
 
@@ -1043,7 +1043,7 @@ BOOL Detoured_GetFileAttributesExW(LPCWSTR lpFileName, GET_FILEEX_INFO_LEVELS fI
 		return FALSE;
 	}
 
-	return Shared_GetFileAttributesExW(fixedName.data, fInfoLevelId, lpFileInformation, lpFileName);
+	return Shared_GetFileAttributesExW(TC("GetFileAttributesExW"), fixedName.data, fInfoLevelId, lpFileInformation, lpFileName);
 }
 
 DWORD Detoured_GetFileAttributesW(LPCWSTR lpFileName)
@@ -1063,7 +1063,7 @@ DWORD Detoured_GetFileAttributesW(LPCWSTR lpFileName)
 	DevirtualizePath(fixedPath);
 
 	WIN32_FILE_ATTRIBUTE_DATA data;
-	if (!Shared_GetFileAttributesExW(fixedPath.data, GetFileExInfoStandard, &data, lpFileName))
+	if (!Shared_GetFileAttributesExW(TC("GetFileAttributesW"), fixedPath.data, GetFileExInfoStandard, &data, lpFileName))
 		return INVALID_FILE_ATTRIBUTES;
 
 	return data.dwFileAttributes;
@@ -1110,7 +1110,7 @@ DWORD Detoured_GetLongPathNameW(LPCWSTR lpszShortPath, LPWSTR lpszLongPath, DWOR
 			StringBuffer<> realName(fixedName);
 			DevirtualizePath(realName);
 			WIN32_FILE_ATTRIBUTE_DATA data;
-			success = Shared_GetFileAttributesExW(realName.data, GetFileExInfoStandard, &data, path);
+			success = Shared_GetFileAttributesExW(TC("GetLongPathNameW"), realName.data, GetFileExInfoStandard, &data, path);
 		}
 
 		DWORD res = 0;
@@ -1186,25 +1186,25 @@ BOOL Detoured_GetVolumePathNameW(LPCWSTR lpszFileName, LPWSTR lpszVolumePathName
 	return res;
 }
 
-DWORD Shared_GetModuleFileNameW(HMODULE hModule, const wchar_t* moduleName, u32 moduleNameLen, LPWSTR lpFilename, DWORD nSize)
+DWORD Shared_GetModuleFileNameW(HMODULE hModule, const StringView& moduleName, LPWSTR lpFilename, DWORD nSize)
 {
-	if (nSize <= moduleNameLen)
+	if (nSize <= moduleName.count)
 	{
 		if (nSize)
 		{
-			memcpy(lpFilename, moduleName, nSize * 2);
+			memcpy(lpFilename, moduleName.data, nSize * 2);
 			lpFilename[nSize - 1] = 0;
 		}
 		SetLastError(ERROR_INSUFFICIENT_BUFFER);
-		DEBUG_LOG_DETOURED(L"GetModuleFileNameW", L"%llu  %u INSUFFICIENT BUFFER (%ls) -> %u", uintptr_t(hModule), nSize, moduleName, moduleNameLen + 1);
+		DEBUG_LOG_DETOURED(L"GetModuleFileNameW", L"%llu  %u INSUFFICIENT BUFFER (%ls) -> %u", uintptr_t(hModule), nSize, moduleName.data, moduleName.count + 1);
 		return nSize;
 	}
 
-	memcpy(lpFilename, moduleName, moduleNameLen * 2);
-	lpFilename[moduleNameLen] = 0;
-	DEBUG_LOG_DETOURED(L"GetModuleFileNameW", L"%llu  %u (%ls) -> %u", uintptr_t(hModule), nSize, lpFilename, moduleNameLen);
+	memcpy(lpFilename, moduleName.data, moduleName.count * 2);
+	lpFilename[moduleName.count] = 0;
+	DEBUG_LOG_DETOURED(L"GetModuleFileNameW", L"%llu  %u (%ls) -> %u", uintptr_t(hModule), nSize, lpFilename, moduleName.count);
 	SetLastError(ERROR_SUCCESS);
-	return moduleNameLen;
+	return moduleName.count;
 }
 
 DWORD Detoured_GetModuleFileNameW(HMODULE hModule, LPWSTR lpFilename, DWORD nSize)
@@ -1213,34 +1213,38 @@ DWORD Detoured_GetModuleFileNameW(HMODULE hModule, LPWSTR lpFilename, DWORD nSiz
 
 	// If null we use the virtual application name
 	if (hModule == NULL)
-		return Shared_GetModuleFileNameW(hModule, g_virtualApplication.data, u32(g_virtualApplication.count), lpFilename, nSize);
+		return Shared_GetModuleFileNameW(hModule, g_virtualApplication, lpFilename, nSize);
 
 	{
 		// Check if there are any stored paths from dynamically loaded dlls
 		SCOPED_READ_LOCK(g_loadedModulesLock, lock);
 		auto findIt = g_loadedModules.find(hModule);
 		if (findIt != g_loadedModules.end())
-			return Shared_GetModuleFileNameW(hModule, findIt->second.c_str(), u32(findIt->second.size()), lpFilename, nSize);
+			return Shared_GetModuleFileNameW(hModule, findIt->second, lpFilename, nSize);
 	}
 
 	if (!g_runningRemote && !IsVfsEnabled())
 		return True_GetModuleFileNameW(hModule, lpFilename, nSize);
 
-	wchar_t moduleName[350];
-	DWORD res = True_GetModuleFileNameW(hModule, moduleName, sizeof_array(moduleName));
+	StringBuffer<512> moduleName;
+	DWORD res = True_GetModuleFileNameW(hModule, moduleName.data, moduleName.capacity);
 	if (res == 0)
 		return res;
 	UBA_ASSERT(GetLastError() != ERROR_INSUFFICIENT_BUFFER);
+	moduleName.count = res;
 
 	// This could be dlls that are loaded early one so might not exist in g_loadedModules
 	// TODO: These could be wrong.. since the files could have been copied from different directories into the remote exedir
-	if (!StartsWith(moduleName, g_exeDir.data))
-		return Shared_GetModuleFileNameW(hModule, moduleName, res, lpFilename, nSize);
+	if (!moduleName.StartsWith(g_exeDir.data))
+	{
+		VirtualizePath(moduleName);
+		return Shared_GetModuleFileNameW(hModule, moduleName, lpFilename, nSize);
+	}
 
 	StringBuffer<350> fileName;
 	fileName.Append(g_virtualApplicationDir);
-	fileName.Append(moduleName + g_exeDir.count);
-	return Shared_GetModuleFileNameW(hModule, fileName.data, u32(fileName.count), lpFilename, nSize);
+	fileName.Append(moduleName.data + g_exeDir.count);
+	return Shared_GetModuleFileNameW(hModule, fileName, lpFilename, nSize);
 }
 
 DWORD Detoured_GetModuleFileNameExW(HANDLE hProcess, HMODULE hModule, LPWSTR lpFilename, DWORD nSize)
@@ -2549,6 +2553,7 @@ HMODULE Recursive_LoadLibraryExW(LPCWSTR lpLibFileName, LPCWSTR originalName, DW
 	{
 		if (originalName[1] == ':')
 		{
+			// TODO: Virtualize!
 			SCOPED_WRITE_LOCK(g_loadedModulesLock, lock);
 			g_loadedModules[res] = originalName;
 		}
@@ -2606,17 +2611,22 @@ HMODULE Detoured_LoadLibraryExW(LPCWSTR lpLibFileName, HANDLE hFile, DWORD dwFla
 	DETOURED_CALL(LoadLibraryExW);
 	DEBUG_LOG_DETOURED(L"LoadLibraryExW", L"(%ls)", lpLibFileName);
 
-
 	if (!g_rules->AllowLoadLibrary(lpLibFileName))
 		return 0;
 
 	if (HMODULE res = Shared_LoadLibrary(lpLibFileName, hFile, dwFlags))
 		return res;
 
-	StringBuffer<> fileName(lpLibFileName);
-	DevirtualizePath(fileName);
+	StringBuffer<> fileName;
+	if (lpLibFileName[1] == ':')
+	{
+		FixPath(fileName, lpLibFileName);
+		DevirtualizePath(fileName);
+	}
+	else
+		fileName.Append(lpLibFileName);
 
-	DEBUG_LOG_TRUE(L"LoadLibraryExW", L"(%ls)", lpLibFileName);
+	DEBUG_LOG_TRUE(L"LoadLibraryExW", L"%ls (%s)", lpLibFileName, fileName.data);
 	return True_LoadLibraryExW(fileName.data, hFile, dwFlags);
 }
 
@@ -2812,6 +2822,7 @@ BOOL Detoured_CreateProcessW(LPCWSTR lpApplicationName, LPWSTR lpCommandLine, LP
 	if (lpApplicationName && *lpApplicationName)
 	{
 		FixPath(application, lpApplicationName);
+		UBA_ASSERTF(application.count, TC("Invalid application name from application field: %s"), lpApplicationName);
 	}
 	else
 	{
@@ -2821,6 +2832,7 @@ BOOL Detoured_CreateProcessW(LPCWSTR lpApplicationName, LPWSTR lpCommandLine, LP
 		else
 			temp.Append(originalCmd.data, commandLineWithoutApplication.data - originalCmd.data - 1);
 		FixPath(application, temp.data);
+		UBA_ASSERTF(application.count, TC("Invalid application from command line (name: %s)"), originalCmd.data);
 	}
 	DevirtualizePath(application);
 
@@ -3317,7 +3329,7 @@ DWORD Detoured_GetModuleFileNameExA(HANDLE hProcess, HMODULE hModule, LPSTR lpFi
 			return Shared_GetModuleFileNameA(hModule, findIt->second.c_str(), u32(findIt->second.size()), lpFilename, nSize);
 	}
 
-	if (!g_runningRemote)
+	if (!g_runningRemote && !IsVfsEnabled())
 		return True_GetModuleFileNameA(hModule, lpFilename, nSize);
 
 	wchar_t moduleName[350];
