[UBA]
* Added support for not preloading obj files but decompress them when opened instead and also once they are closed after they have been mapped or read
* Reduced code duplication in detoured NtCreateFile
* Changed so FileInfo stores fileMapMemSize instead of fileMapMemEnd to simplify code
* Added various asserts and log entries to make it easier to debug code
* Added lock for log entries to make sure all entries in a scope are after each other (confusing log when processes are multithreaded

diff --git a/Engine/Source/Programs/UnrealBuildAccelerator/Core/Private/UbaApplicationRules.cpp b/Engine/Source/Programs/UnrealBuildAccelerator/Core/Private/UbaApplicationRules.cpp
--- a/Engine/Source/Programs/UnrealBuildAccelerator/Core/Private/UbaApplicationRules.cpp
+++ b/Engine/Source/Programs/UnrealBuildAccelerator/Core/Private/UbaApplicationRules.cpp
@@ -265,7 +265,7 @@ namespace uba
 
 		virtual bool StoreFileCompressed(const StringView& fileName) const
 		{
-			return fileName.EndsWith(TC(".obj")) || fileName.EndsWith(TC(".o"));
+			return (fileName.EndsWith(TC(".obj")) || fileName.EndsWith(TC(".o"))) && !fileName.Contains(TC("pch."));;
 		}
 
 		virtual bool ShouldExtractSymbols(const StringView& fileName) const
diff --git a/Engine/Source/Programs/UnrealBuildAccelerator/Detours/Private/UbaDetoursFileMappingTable.h b/Engine/Source/Programs/UnrealBuildAccelerator/Detours/Private/UbaDetoursFileMappingTable.h
--- a/Engine/Source/Programs/UnrealBuildAccelerator/Detours/Private/UbaDetoursFileMappingTable.h
+++ b/Engine/Source/Programs/UnrealBuildAccelerator/Detours/Private/UbaDetoursFileMappingTable.h
@@ -12,24 +12,49 @@ namespace uba
 
 	struct FileInfo
 	{
+		// Size of file. Can be different than what Directory table says (For example for decompressed obj files)
 		u64 size = InvalidValue;
-		bool tracked = false;
+
+		// Key of fixed up file name
 		StringKey fileNameKey;
+
+		// "Real" name of file. Can also be ids/handles to memory buffers etc
 		const tchar* name = nullptr;
-		const tchar* originalName = nullptr;
 
-		u32 lastDesiredAccess = 0;
+		// Original name of file as the detoured process sees it
+		const tchar* originalName = nullptr;
 
+		// If file is a memory file, this is set. Memory files are writable and can be shared between process hierarchies
+		// If this is set, isFileMap is false.
 		MemoryFile* memoryFile = nullptr;
 
+		// Holds previous file open access flags. Is used to figure out if file open is the first write
+		u32 lastDesiredAccess = 0;
+
 		// This is the Filemap handle and mapped view that represents the File HANDLE.
 		// It can be used for other mappings too but doesnt' have to.
 		bool isFileMap = false;
-		bool deleted = false; // Only used by remote builds since directory table might not contain local temporary files
+
+		// If this is true, fileMapMem will be deleted when
+		bool freeFileMapOnClose = false;
+
+		// Only used by remote builds since directory table might not contain local temporary files
+		bool deleted = false;
+
+		// Is set to true once file has been tracked to prevent tracking same file multiple times
+		bool tracked = false;
+
+		// File mapping handle coming from session process (Set in NtCreateFile)
 		void* trueFileMapHandle = 0;
+
+		// Offset into trueFileMapHandle
 		u64 trueFileMapOffset = 0;
+
+		// If file mapping is resolved, then this points to the memory. Is also used for compressed obj files
 		u8* fileMapMem = nullptr;
-		u8* fileMapMemEnd = nullptr;
+
+		// Size of fileMapMem
+		u64 fileMapMemSize = 0;
 
 		#if PLATFORM_WINDOWS
 		u32 fileMapDesiredAccess = PAGE_READONLY;
diff --git a/Engine/Source/Programs/UnrealBuildAccelerator/Detours/Private/UbaDetoursShared.cpp b/Engine/Source/Programs/UnrealBuildAccelerator/Detours/Private/UbaDetoursShared.cpp
--- a/Engine/Source/Programs/UnrealBuildAccelerator/Detours/Private/UbaDetoursShared.cpp
+++ b/Engine/Source/Programs/UnrealBuildAccelerator/Detours/Private/UbaDetoursShared.cpp
@@ -60,6 +60,7 @@ namespace uba
 	thread_local char t_b[LogBufSize];
 	thread_local u32 t_b_size;
 	thread_local u32 t_logScopeCount;
+	ReaderWriterLock g_logScopeLock;
 
 	void GetPrefixExtra(StringBufferBase& out)
 	{
@@ -137,13 +138,16 @@ namespace uba
 	}
 	LogScope::LogScope()
 	{
-		++t_logScopeCount;
+		if (++t_logScopeCount  == 1)
+			g_logScopeLock.EnterWrite();
 	}
 	LogScope::~LogScope()
 	{
-		--t_logScopeCount;
-		if (!t_logScopeCount && t_b_size)
+		if (--t_logScopeCount)
+			return;
+		if (t_b_size)
 			Flush();
+		g_logScopeLock.LeaveWrite();
 	}
 	void LogScope::Flush()
 	{
diff --git a/Engine/Source/Programs/UnrealBuildAccelerator/Detours/Private/Windows/UbaDetoursFunctionsKernelBase.inl b/Engine/Source/Programs/UnrealBuildAccelerator/Detours/Private/Windows/UbaDetoursFunctionsKernelBase.inl
--- a/Engine/Source/Programs/UnrealBuildAccelerator/Detours/Private/Windows/UbaDetoursFunctionsKernelBase.inl
+++ b/Engine/Source/Programs/UnrealBuildAccelerator/Detours/Private/Windows/UbaDetoursFunctionsKernelBase.inl
@@ -158,10 +158,12 @@ BOOL Detoured_DuplicateHandle(HANDLE hSourceProcessHandle, HANDLE hSourceHandle,
 	if (hSourceHandle == PseudoHandle || !isDetouredHandle(hSourceHandle))
 	{
 		auto res = True_DuplicateHandle(hSourceProcessHandle, hSourceHandle, hTargetProcessHandle, lpTargetHandle, dwDesiredAccess, bInheritHandle, dwOptions);
-		DEBUG_LOG_TRUE(L"DuplicateHandle", L"%llu %llu -> %ls", uintptr_t(hSourceHandle), lpTargetHandle ? uintptr_t(*lpTargetHandle) : 0ull, ToString(res));
+		DEBUG_LOG_TRUE(L"DuplicateHandle", L"%llu to %llu -> %ls", uintptr_t(hSourceHandle), lpTargetHandle ? uintptr_t(*lpTargetHandle) : 0ull, ToString(res));
 		return res;
 	}
 
+	UBA_ASSERT(hSourceProcessHandle == hTargetProcessHandle);
+	UBA_ASSERT(!(dwOptions & DUPLICATE_CLOSE_SOURCE));
 	auto& dh = asDetouredHandle(hSourceHandle);
 
 	HANDLE trueHandle = dh.trueHandle;
@@ -183,7 +185,7 @@ BOOL Detoured_DuplicateHandle(HANDLE hSourceProcessHandle, HANDLE hSourceHandle,
 		InterlockedIncrement(&fo->refCount);
 	}
 	*lpTargetHandle = makeDetouredHandle(newDh);
-	DEBUG_LOG_TRUE(L"DuplicateHandle", L"%llu %llu -> %ls", uintptr_t(hSourceHandle), uintptr_t(*lpTargetHandle), ToString(res));
+	DEBUG_LOG_DETOURED(L"DuplicateHandle", L"%llu to %llu -> %ls", uintptr_t(hSourceHandle), uintptr_t(*lpTargetHandle), ToString(res));
 	return res;
 }
 
@@ -370,7 +372,11 @@ BOOL Detoured_ReadFile(HANDLE hFile, LPVOID lpBuffer, DWORD nNumberOfBytesToRead
 			return TRUE;
 		}
 
-		FileInfo& fi = *dh.fileObject->fileInfo;
+		auto& fo = *dh.fileObject;
+		fo.wasUsed = true;
+
+		FileInfo& fi = *fo.fileInfo;
+
 		if (fi.isFileMap || fi.memoryFile)
 		{
 			// TODO: Handle lpOverlapped - If a read happen and there is 0 left it should return 0 with SetLastError(ERROR_HANDLE_EOF)
@@ -382,7 +388,7 @@ BOOL Detoured_ReadFile(HANDLE hFile, LPVOID lpBuffer, DWORD nNumberOfBytesToRead
 			UBA_ASSERT(fi.fileMapMem || !fi.memoryFile->isThrowAway);
 
 			u8* mem = fi.fileMapMem ? fi.fileMapMem : fi.memoryFile->baseAddress;
-			u64 size = fi.fileMapMem ? fi.fileMapMemEnd - mem : fi.memoryFile->writtenSize;
+			u64 size = fi.fileMapMem ? fi.fileMapMemSize : fi.memoryFile->writtenSize;
 			UBA_ASSERTF(dh.pos <= size, L"Filepointer is higher than size of file (pointer: %llu, size: %llu) (%ls)", dh.pos, size, HandleToName(hFile));
 			u64 leftToRead = size - dh.pos;
 			if (nNumberOfBytesToRead > leftToRead)
@@ -819,7 +825,7 @@ DWORD Detoured_SetFilePointer(HANDLE hFile, LONG lDistanceToMove, PLONG lpDistan
 				else
 				{
 					//UBA_ASSERT(fi.fileMapMemEnd);
-					size = fi.size;// fileMapMemEnd - fi.fileMapMem;
+					size = fi.size;// fi.fileMapMemSize;
 				}
 				dh.pos = Max(0ll, (LONGLONG)size + liDistanceToMove.QuadPart);
 			}
@@ -866,7 +872,7 @@ DWORD Detoured_SetFilePointerEx(HANDLE hFile, LARGE_INTEGER liDistanceToMove, PL
 				if (fi.memoryFile)
 					size = fi.memoryFile->writtenSize;
 				else if (fi.fileMapMem)
-					size = fi.fileMapMemEnd - fi.fileMapMem;
+					size = fi.fileMapMemSize;
 				else
 					size = fi.size;
 				dh.pos = Max(0ll, (LONGLONG)size + liDistanceToMove.QuadPart);
@@ -2176,7 +2182,11 @@ HANDLE Detoured_CreateFileMappingW(HANDLE hFile, LPSECURITY_ATTRIBUTES lpFileMap
 	if (isDetouredHandle(hFile))
 	{
 		auto& dh = asDetouredHandle(hFile);
-		FileInfo& fi = *dh.fileObject->fileInfo;
+
+		FileObject& fo = *dh.fileObject;
+		fo.wasUsed = true;
+
+		FileInfo& fi = *fo.fileInfo;
 		if (fi.memoryFile || fi.isFileMap)
 		{
 			auto mdh = new DetouredHandle(HandleType_FileMapping);
@@ -2200,7 +2210,7 @@ HANDLE Detoured_CreateFileMappingW(HANDLE hFile, LPSECURITY_ATTRIBUTES lpFileMap
 			}
 			InterlockedIncrement(&dh.fileObject->refCount);
 			HANDLE res = makeDetouredHandle(mdh);
-			DEBUG_LOG_DETOURED(L"CreateFileMappingW", L"(%ls) %llu %llu (%ls) -> Success", (fi.memoryFile ? L"MEMORYFILE" : L"FILEMAP"), uintptr_t(res), uintptr_t(hFile), HandleToName(hFile));
+			DEBUG_LOG_DETOURED(L"CreateFileMappingW", L"(%ls) File %llu %ls -> %llu", (fi.memoryFile ? L"MEMORYFILE" : L"FILEMAP"), uintptr_t(hFile), HandleToName(hFile), uintptr_t(res));
 			SetLastError(ERROR_SUCCESS);
 			return res;
 		}
@@ -2212,7 +2222,7 @@ HANDLE Detoured_CreateFileMappingW(HANDLE hFile, LPSECURITY_ATTRIBUTES lpFileMap
 	HANDLE mappingHandle = True_CreateFileMappingW(trueHandle, lpFileMappingAttributes, flProtect, dwMaximumSizeHigh, dwMaximumSizeLow, lpName);
 	if (!mappingHandle)
 	{
-		DEBUG_LOG_TRUE(L"CreateFileMappingW", L"%llu  (File %llu) (%ls) -> Error", uintptr_t(mappingHandle), uintptr_t(hFile), HandleToName(hFile));
+		DEBUG_LOG_TRUE(L"CreateFileMappingW", L"File %llu (%s) -> Error", uintptr_t(hFile), HandleToName(hFile));
 		return NULL;
 	}
 
@@ -2224,7 +2234,7 @@ HANDLE Detoured_CreateFileMappingW(HANDLE hFile, LPSECURITY_ATTRIBUTES lpFileMap
 		detouredHandle->trueHandle = mappingHandle;
 		mappingHandle = makeDetouredHandle(detouredHandle);
 	}
-	DEBUG_LOG_TRUE(L"CreateFileMappingW", L"%llu  (File %llu, Size: %llu) (%ls) -> %llu", uintptr_t(mappingHandle), uintptr_t(hFile), ToLargeInteger(dwMaximumSizeHigh, dwMaximumSizeLow).QuadPart, HandleToName(hFile), u64(mappingHandle));
+	DEBUG_LOG_TRUE(L"CreateFileMappingW", L"File %llu, Size: %llu (%s) -> %llu", uintptr_t(hFile), ToLargeInteger(dwMaximumSizeHigh, dwMaximumSizeLow).QuadPart, HandleToName(hFile), u64(mappingHandle));
 	return mappingHandle;
 }
 
@@ -2270,7 +2280,6 @@ LPVOID Detoured_MapViewOfFileEx(HANDLE hFileMappingObject, DWORD dwDesiredAccess
 		{
 			u64 offset = ToLargeInteger(dwFileOffsetHigh, dwFileOffsetLow).QuadPart;
 
-			DEBUG_LOG_DETOURED(L"MapViewOfFileEx", L"(%ls)", HandleToName(hFileMappingObject));
 			FileInfo& fi = *dh.fileObject->fileInfo;
 			if (fi.fileMapMem && lpBaseAddress && lpBaseAddress != fi.fileMapMem) // This scenario happens with pch files in msvc cl.exe
 			{
@@ -2292,7 +2301,7 @@ LPVOID Detoured_MapViewOfFileEx(HANDLE hFileMappingObject, DWORD dwDesiredAccess
 				//			UBA_ASSERT(false);
 				//		}
 				//		fi.fileMapMem = (u8*)res;
-				//		fi.fileMapMemEnd = fi.fileMapMem + dwNumberOfBytesToMap;
+				//		fi.fileMapMemSize = dwNumberOfBytesToMap;
 				//		g_mappedFileTable.m_memLookup[res + offset] = 1;
 				//	}
 				//	res += offset;
@@ -2317,7 +2326,7 @@ LPVOID Detoured_MapViewOfFileEx(HANDLE hFileMappingObject, DWORD dwDesiredAccess
 					} while (counter <= 10);
 				}
 
-				DEBUG_LOG_TRUE(L"INTERNAL MapViewOfFileEx", L"New FileObject for different base address (%ls)", HandleToName(hFileMappingObject));
+				DEBUG_LOG_TRUE(L"MapViewOfFileEx", L"(INTERNAL) New FileObject for different base address %llu (%s) -> 0x%llx", uintptr_t(hFileMappingObject), HandleToName(hFileMappingObject), uintptr_t(res));
 				return res;
 			}
 			else if (!fi.fileMapMem)
@@ -2335,6 +2344,8 @@ LPVOID Detoured_MapViewOfFileEx(HANDLE hFileMappingObject, DWORD dwDesiredAccess
 			if (fi.memoryFile && (dwDesiredAccess & FILE_MAP_WRITE)) // We assume changes will happen
 				fi.memoryFile->isReported = false;
 
+			DEBUG_LOG_DETOURED(L"MapViewOfFileEx", L"%llu (%s) -> 0x%llx", uintptr_t(hFileMappingObject), HandleToName(hFileMappingObject), u64(mem));
+
 			SCOPED_WRITE_LOCK(g_mappedFileTable.m_memLookupLock, lock);
 			++g_mappedFileTable.m_memLookup[mem];
 			return mem;
diff --git a/Engine/Source/Programs/UnrealBuildAccelerator/Detours/Private/Windows/UbaDetoursFunctionsNtDll.inl b/Engine/Source/Programs/UnrealBuildAccelerator/Detours/Private/Windows/UbaDetoursFunctionsNtDll.inl
--- a/Engine/Source/Programs/UnrealBuildAccelerator/Detours/Private/Windows/UbaDetoursFunctionsNtDll.inl
+++ b/Engine/Source/Programs/UnrealBuildAccelerator/Detours/Private/Windows/UbaDetoursFunctionsNtDll.inl
@@ -819,7 +819,7 @@ NTSTATUS NTAPI Shared_NtCreateFile(bool IsCreateFunc, PHANDLE hFileHandle, ACCES
 
 				bool isWriteAttributes = (DesiredAccess & FILE_WRITE_ATTRIBUTES) != 0;
 
-				if (allowEarlyOut && !useContent && !isWriteAttributes && !wasPreload)
+				if (allowEarlyOut && !useContent && !isWriteAttributes && !wasPreload && !g_objFilesPreloader.CouldBeCompressedFile(fileName))
 				{
 					auto dh = new DetouredHandle(HandleType_File);
 					dh->fileObject = new FileObject();
@@ -1013,6 +1013,19 @@ NTSTATUS NTAPI Shared_NtCreateFile(bool IsCreateFunc, PHANDLE hFileHandle, ACCES
 			}
 		};
 
+	auto CreateFileHandle = [&](HANDLE th = INVALID_HANDLE_VALUE)
+		{
+			auto fo = new FileObject();
+			fo->desiredAccess = DesiredAccess;
+			fo->closeId = closeId;
+			fo->fileInfo = &info;
+			fo->deleteOnClose = isDeleteOnClose;
+			auto dh = new DetouredHandle(HandleType_File, th);
+			dh->dirTableOffset = dirTableOffset;
+			dh->fileObject = fo;
+			return makeDetouredHandle(dh);
+		};
+
 	if (lpFileName[0] == '$')
 	{
 		lookupLock.Leave();
@@ -1026,14 +1039,8 @@ NTSTATUS NTAPI Shared_NtCreateFile(bool IsCreateFunc, PHANDLE hFileHandle, ACCES
 		MemoryFile& mf = g_emptyMemoryFile;
 		info.memoryFile = &mf;
 
-		DetouredHandle* dh = new DetouredHandle(HandleType_File);
-		dh->dirTableOffset = dirTableOffset;
-		dh->fileObject = new FileObject();
-		dh->fileObject->desiredAccess = DesiredAccess;
-		dh->fileObject->closeId = closeId;
-		dh->fileObject->fileInfo = &info;
 		UBA_ASSERT(!isDeleteOnClose);
-		*hFileHandle = makeDetouredHandle(dh);
+		*hFileHandle = CreateFileHandle();
 
 		TrackFileInput();
 
@@ -1065,15 +1072,9 @@ NTSTATUS NTAPI Shared_NtCreateFile(bool IsCreateFunc, PHANDLE hFileHandle, ACCES
 
 		lookupLock.Leave();
 
-		DetouredHandle* dh = new DetouredHandle(HandleType_File);
 		UBA_ASSERT(info.size != InvalidValue);
-		dh->dirTableOffset = dirTableOffset;
-		dh->fileObject = new FileObject();
-		dh->fileObject->desiredAccess = DesiredAccess;
-		dh->fileObject->closeId = closeId;
-		dh->fileObject->fileInfo = &info;
 		UBA_ASSERT(!isDeleteOnClose);
-		*hFileHandle = makeDetouredHandle(dh);
+		*hFileHandle = CreateFileHandle();
 
 		TrackFileInput();
 
@@ -1164,15 +1165,7 @@ NTSTATUS NTAPI Shared_NtCreateFile(bool IsCreateFunc, PHANDLE hFileHandle, ACCES
 
 		lookupLock.Leave();
 
-		auto fileObject = new FileObject();
-		DetouredHandle* dh = new DetouredHandle(HandleType_File);
-		dh->fileObject = fileObject;
-		dh->dirTableOffset = dirTableOffset;
-		dh->fileObject->desiredAccess = DesiredAccess;
-		dh->fileObject->closeId = closeId;
-		dh->fileObject->fileInfo = &info;
-		dh->fileObject->deleteOnClose = isDeleteOnClose;
-		*hFileHandle = makeDetouredHandle(dh);
+		*hFileHandle = CreateFileHandle();
 
 		TrackFileInput();
 
@@ -1184,16 +1177,7 @@ NTSTATUS NTAPI Shared_NtCreateFile(bool IsCreateFunc, PHANDLE hFileHandle, ACCES
 
 	if (lpFileName[0] == ':')
 	{
-		auto fileObject = new FileObject();
-		DetouredHandle* dh = new DetouredHandle(HandleType_File);
-		dh->fileObject = fileObject;
-		dh->dirTableOffset = dirTableOffset;
-		dh->fileObject->desiredAccess = DesiredAccess;
-		dh->fileObject->closeId = closeId;
-		dh->fileObject->fileInfo = &info;
-		dh->fileObject->deleteOnClose = isDeleteOnClose;
-		*hFileHandle = makeDetouredHandle(dh);
-
+		*hFileHandle = CreateFileHandle();
 		DEBUG_LOG_DETOURED(funcName, L"(PRELOAD)%ls %llu (%ls) (%ls) -> Success", isWriteStr, uintptr_t(*hFileHandle), lpFileName, (fileName.data != lpFileName ? fileName.data : L""));
 		return STATUS_SUCCESS;
 	}
@@ -1204,6 +1188,15 @@ NTSTATUS NTAPI Shared_NtCreateFile(bool IsCreateFunc, PHANDLE hFileHandle, ACCES
 	else
 		tempFileName = info.name;
 
+
+	if (g_objFilesPreloader.IsDecompressed(fileNameKey, fileName, info))
+	{
+		*hFileHandle = CreateFileHandle();
+		TrackFileInput();
+		DEBUG_LOG_DETOURED(funcName, L"(DECOMPRESSED)%ls %llu (%ls) (%ls) -> Success", isWriteStr, uintptr_t(*hFileHandle), lpFileName, (fileName.data != lpFileName ? fileName.data : L""));
+		return STATUS_SUCCESS;
+	}
+
 	StringBuffer<> temp;
 	temp.Append(L"\\??\\");
 	temp.Append(tempFileName);
@@ -1241,15 +1234,7 @@ NTSTATUS NTAPI Shared_NtCreateFile(bool IsCreateFunc, PHANDLE hFileHandle, ACCES
 	TrackFileInput();
 
 	UBA_ASSERT(info.originalName);
-	auto dh = new DetouredHandle(HandleType_File);
-	dh->trueHandle = *hFileHandle;
-	dh->dirTableOffset = dirTableOffset;
-	dh->fileObject = new FileObject();
-	dh->fileObject->desiredAccess = DesiredAccess;
-	dh->fileObject->closeId = closeId;
-	dh->fileObject->fileInfo = &info;
-	dh->fileObject->deleteOnClose = isDeleteOnClose;
-	*hFileHandle = makeDetouredHandle(dh);
+	*hFileHandle = CreateFileHandle(*hFileHandle);
 	DEBUG_LOG_TRUE(funcName, L"%ls %llu (%ls)%s -> %ls", isWriteStr, uintptr_t(*hFileHandle), tempFileName, isDeleteOnClose ? TC(" DeleteOnClose") : TC(""), ToString(res).data);
 	return res;
 }
@@ -1484,6 +1469,15 @@ NTSTATUS NTAPI Detoured_NtClose(HANDLE handle)
 		//UBA_ASSERTF(fo->newName.empty(), L"Got close of file that was renamed but had no closeId. Old: %ls New: %ls", fi.originalName, fo->newName.c_str());
 	}
 
+	// TODO: Can't free here..... Might still be open with CreateFileMapping or MapViewOfFile
+	if (fo->wasUsed && fi.freeFileMapOnClose)
+	{
+		DEBUG_LOG_DETOURED(L"VirtualFree", L"%llu (%ls) -> %ls", uintptr_t(handle), HandleToName(handle), ToString(res).data);
+		//VirtualFree(fi.fileMapMem, 0, MEM_RELEASE);
+		_aligned_free(fi.fileMapMem);
+		fi.fileMapMem = nullptr;
+	}
+
 	DEBUG_LOG_DETOURED(L"NtClose", L"%llu (%ls) -> %ls", uintptr_t(handle), HandleToName(handle), ToString(res).data);
 
 	if (fo->ownsFileInfo)
diff --git a/Engine/Source/Programs/UnrealBuildAccelerator/Detours/Private/Windows/UbaDetoursFunctionsWin.cpp b/Engine/Source/Programs/UnrealBuildAccelerator/Detours/Private/Windows/UbaDetoursFunctionsWin.cpp
--- a/Engine/Source/Programs/UnrealBuildAccelerator/Detours/Private/Windows/UbaDetoursFunctionsWin.cpp
+++ b/Engine/Source/Programs/UnrealBuildAccelerator/Detours/Private/Windows/UbaDetoursFunctionsWin.cpp
@@ -415,6 +415,7 @@ bool EnsureMapped(DetouredHandle& handle, DWORD dwFileOffsetHigh = 0, DWORD dwFi
 	}
 	else
 	{
+		UBA_ASSERT(!info.freeFileMapOnClose);
 		TimerScope ts(g_kernelStats.mapViewOfFile);
 		info.fileMapMem = (u8*)True_MapViewOfFileEx(info.trueFileMapHandle, info.fileMapViewDesiredAccess, 0, 0, numberOfBytesToMap, baseAddress);
 	}
@@ -425,7 +426,7 @@ bool EnsureMapped(DetouredHandle& handle, DWORD dwFileOffsetHigh = 0, DWORD dwFi
 		return false;
 	}
 	info.fileMapMem += (offset - alignedOffsetStart);
-	info.fileMapMemEnd = info.fileMapMem + info.size;
+	info.fileMapMemSize = info.size;
 
 	DEBUG_LOG_TRUE(L"INTERNAL MapViewOfFileEx", L"(%ls) (size: %llu) (%ls) -> 0x%llx", info.name, numberOfBytesToMap, info.originalName, uintptr_t(info.fileMapMem));
 	return true;
@@ -980,10 +981,7 @@ void Init(const DetoursPayload& payload, u64 startTime)
 	g_stats.attach.count = 1;
 
 	if (payload.storeObjFilesCompressed && g_rules->ShouldDecompressFiles(StringView()))
-	{
-		TimerScope ts(g_stats.preparseObjFiles);
 		g_objFilesPreloader.Start(cmdLine);
-	}
 }
 
 void Deinit(u64 startTime)
diff --git a/Engine/Source/Programs/UnrealBuildAccelerator/Detours/Private/Windows/UbaDetoursObjFilesPreloader.h b/Engine/Source/Programs/UnrealBuildAccelerator/Detours/Private/Windows/UbaDetoursObjFilesPreloader.h
--- a/Engine/Source/Programs/UnrealBuildAccelerator/Detours/Private/Windows/UbaDetoursObjFilesPreloader.h
+++ b/Engine/Source/Programs/UnrealBuildAccelerator/Detours/Private/Windows/UbaDetoursObjFilesPreloader.h
@@ -119,6 +119,8 @@ namespace uba
 
 			HANDLE fileHandle = CreateFileW(fileFull.data, GENERIC_READ, FILE_SHARE_READ, NULL, OPEN_EXISTING, FILE_ATTRIBUTE_NORMAL, NULL);
 			UBA_ASSERTF(fileHandle != INVALID_HANDLE_VALUE, TC("Failed to open %s for read"), fileFull.data);
+			if (fileHandle == INVALID_HANDLE_VALUE)
+				return;
 
 			StringKey fileNameKey = ToStringKey(fileFull.MakeLower());
 			auto insres = m_preloadedObjFiles.try_emplace(fileNameKey);
@@ -142,7 +144,16 @@ namespace uba
 		{
 			Oodle_SetUsageWarnings(Oodle_UsageWarnings_Disabled);
 
+			m_isEnabled = true;
+
+			if (true)
+				return;
+
+			TimerScope ts(g_stats.preparseObjFiles);
+
+			m_isParsing = true;
 			ParseArguments(cmdLine, [&](const tchar* arg, u32 argLen) { HandleLine(StringView(arg, argLen)); });
+			m_isParsing = false;
 
 			if (!m_preloadedObjFiles.empty())
 				m_threadHandle = CreateThread(NULL, 0, [](LPVOID p) -> DWORD { ((ObjFilesPreloader*)p)->ThreadPreload(); return 0; }, this, 0, NULL);
@@ -159,53 +170,125 @@ namespace uba
 				CloseHandle(m_helperThreadHandles[i]);
 		}
 
-		void ThreadPreload()
+		bool CouldBeCompressedFile(const StringView& fileName)
 		{
-			u64 totalMemSize = 0;
-			for (Preload* preloadPtr=m_it; preloadPtr; preloadPtr = preloadPtr->next)
-			{
-				Preload& preload = *preloadPtr;
-				HANDLE objFileHandle = preload.handle;
-				if (!objFileHandle)
-					continue;
+			if (!m_isEnabled)
+				return false;
+			return g_rules->ShouldDecompressFiles(fileName);
+		}
 
-				LARGE_INTEGER objFileSize;
-				BOOL res = GetFileSizeEx(objFileHandle, &objFileSize);
-				UBA_ASSERT(res);(void)res;
-				if (objFileSize.QuadPart < sizeof(CompressedObjFileHeader))
-					continue;
+		bool IsDecompressed(const StringKey& fileNameKey, const StringView& fileName, FileInfo& info)
+		{
+			if (!m_isEnabled || m_isParsing)
+				return false;
 
-				HANDLE objFileMappingHandle = CreateFileMappingW(objFileHandle, NULL, PAGE_READONLY, objFileSize.HighPart, objFileSize.LowPart, NULL);
-				UBA_ASSERT(objFileMappingHandle);
-				if (!objFileMappingHandle)
-					continue;
-				CloseHandle(objFileHandle);
-				preload.handle = NULL;
-				preload.objMem = (u8*)MapViewOfFile(objFileMappingHandle, FILE_MAP_READ, 0, 0, objFileSize.QuadPart);
-				UBA_ASSERT(preload.objMem);
-				if (!preload.objMem)
-					continue;
-				if (!((CompressedObjFileHeader*)preload.objMem)->IsValid())
+			if (!g_rules->ShouldDecompressFiles(fileName))
+				return false;
+
+			SCOPED_WRITE_LOCK(m_preloadedObjFilesLock, lock);
+			auto insres = m_preloadedObjFiles.try_emplace(fileNameKey);
+			Preload& preload = insres.first->second;
+			if (!insres.second)
+			{
+				lock.Leave();
+				TimerScope ts(g_stats.waitDecompress);
+				preload.event.IsSet();
+				UBA_ASSERT(preload.fileInfo->fileMapMem || !preload.fileInfo->freeFileMapOnClose);
+				return preload.objMem != nullptr;
+			}
+
+			DEBUG_LOG_TRUE(L"DecompressingFile", L"(%s)", fileName.data);
+			{
+				SuppressDetourScope _;
+				preload.handle = CreateFileW(fileName.data, GENERIC_READ, FILE_SHARE_READ, NULL, OPEN_EXISTING, FILE_ATTRIBUTE_NORMAL, NULL);
+				UBA_ASSERTF(preload.handle != INVALID_HANDLE_VALUE, TC("Failed to open %s for read"), fileName.data);
+				if (preload.handle == INVALID_HANDLE_VALUE)
 				{
-					// Ok, this is backwards.. but it could be that the file on server's disk is compressed..
-					// but it is decompressed in client's cas.. in that case dir table has the size of the compressed file.. 
-					// and since detoured code will provide decompressed file it also needs to provide decompressed size
-					preload.fileInfo->size = objFileSize.QuadPart;
 					preload.event.Set();
-					continue;
+					return false;
 				}
-				preload.objCompressedSize = objFileSize.QuadPart;
-				preload.objDecompressedSize = *(u64*)(preload.objMem + sizeof(CompressedObjFileHeader));
-				preload.objReadOffset = 8 + sizeof(CompressedObjFileHeader);
-				preload.objWriteOffset = 0;
-				preload.objLeft = preload.objDecompressedSize;
-				CloseHandle(objFileMappingHandle);
+			}
 
-				if (ToView(preload.fileInfo->originalName).EndsWith(TC(".h.obj"))) // TODO: Copied this from UbaSession align code
-					totalMemSize = AlignUp(totalMemSize, 4 * 1024);
+			TimerScope ts(g_stats.waitDecompress);
 
-				totalMemSize += preload.objDecompressedSize;
+			preload.fileInfo = &info;
+			preload.event.Create(true);
+			u64 memoryNeeded = PreparePreload(preload, 0);
+			if (!memoryNeeded)
+			{
+				preload.event.Set();
+				return false;
 			}
+			u64 memoryOffset = 0;
+
+			//u64 alignedMem = AlignUp(memoryNeeded, 64*1024);
+			//DWORD allocationType = MEM_RESERVE | MEM_COMMIT;
+			//u8* memory = (u8*)VirtualAlloc(NULL, alignedMem, allocationType, PAGE_READWRITE);
+			u8* memory = (u8*)_aligned_malloc(memoryNeeded, GetAlignment(fileName));
+			Preload* it = &preload;
+			DecompressWork(it, memory, memoryOffset, true);
+			return true;
+		}
+
+		struct Preload;
+
+		u64 GetAlignment(const StringView& fileName)
+		{
+			if (fileName.EndsWith(TC(".h.obj"))) // TODO: Copied this from UbaSession align code
+				return 4 * 1024;
+			return 1;
+		}
+
+		u64 PreparePreload(Preload& preload, u64 totalMemSize)
+		{
+			HANDLE objFileHandle = preload.handle;
+			LARGE_INTEGER objFileSize;
+			BOOL res = GetFileSizeEx(objFileHandle, &objFileSize);
+			UBA_ASSERT(res);(void)res;
+			if (objFileSize.QuadPart < sizeof(CompressedObjFileHeader))
+			{
+				CloseHandle(objFileHandle);
+				preload.event.Set();
+				return totalMemSize;
+			}
+
+			HANDLE objFileMappingHandle = CreateFileMappingW(objFileHandle, NULL, PAGE_READONLY, objFileSize.HighPart, objFileSize.LowPart, NULL);
+			CloseHandle(objFileHandle);
+			UBA_ASSERT(objFileMappingHandle);
+			if (!objFileMappingHandle)
+				return totalMemSize;
+			preload.handle = NULL;
+			u8* objMem = (u8*)MapViewOfFile(objFileMappingHandle, FILE_MAP_READ, 0, 0, objFileSize.QuadPart);
+			UBA_ASSERT(objMem);
+			CloseHandle(objFileMappingHandle);
+			if (!objMem)
+				return totalMemSize;
+			if (!((CompressedObjFileHeader*)objMem)->IsValid())
+			{
+				UnmapViewOfFile(objMem);
+				// Ok, this is backwards.. but it could be that the file on server's disk is compressed..
+				// but it is decompressed in client's cas.. in that case dir table has the size of the compressed file.. 
+				// and since detoured code will provide decompressed file it also needs to provide decompressed size
+				preload.fileInfo->size = objFileSize.QuadPart;
+				preload.event.Set();
+				return totalMemSize;
+			}
+			preload.objMem = objMem;
+			preload.objCompressedSize = objFileSize.QuadPart;
+			preload.objDecompressedSize = *(u64*)(preload.objMem + sizeof(CompressedObjFileHeader));
+			preload.objReadOffset = 8 + sizeof(CompressedObjFileHeader);
+			preload.objWriteOffset = 0;
+			preload.objLeft = preload.objDecompressedSize;
+
+			totalMemSize = AlignUp(totalMemSize, GetAlignment(ToView(preload.fileInfo->originalName)));
+			return totalMemSize + preload.objDecompressedSize;
+		}
+
+		void ThreadPreload()
+		{
+			u64 totalMemSize = 0;
+			for (Preload* preloadPtr=m_it; preloadPtr; preloadPtr = preloadPtr->next)
+				totalMemSize = PreparePreload(*preloadPtr, totalMemSize);
 
 			u64 alignedMem = AlignUp(totalMemSize, 64*1024);
 			DWORD allocationType = MEM_RESERVE | MEM_COMMIT;
@@ -221,11 +304,16 @@ namespace uba
 			m_totalMem = (u8*)VirtualAlloc(NULL, alignedMem, allocationType, PAGE_READWRITE);
 
 			for (u32 i=0;i!=sizeof_array(m_helperThreadHandles); ++i)
-				m_helperThreadHandles[i] = CreateThread(NULL, 0, [](LPVOID p) -> DWORD { ((ObjFilesPreloader*)p)->ThreadHelper(); return 0; }, this, 0, NULL);
-			ThreadHelper();
+				m_helperThreadHandles[i] = CreateThread(NULL, 0, [](LPVOID p) -> DWORD
+					{
+						auto& pl = *(ObjFilesPreloader*)p;
+						pl.DecompressWork(pl.m_it, pl.m_totalMem, pl.m_totalMemOffset, false);
+						return 0;
+					}, this, 0, NULL);
+			DecompressWork(m_it, m_totalMem, m_totalMemOffset, false);
 		}
 
-		void ThreadHelper()
+		void DecompressWork(Preload*& it, u8* totalMem, u64& totalMemOffset, bool allowFreeOnClose)
 		{
 			OO_SINTa decoredMemSize = OodleLZDecoder_MemorySizeNeeded(OodleLZ_Compressor_Kraken);
 			void* decoderMem = malloc(decoredMemSize);
@@ -234,12 +322,12 @@ namespace uba
 			while (true)
 			{
 				SCOPED_WRITE_LOCK(m_threadLock, lock);
-				if (!m_it)
+				if (!it)
 					break;
-				Preload& preload = *m_it;
+				Preload& preload = *it;
 				if (preload.objReadOffset == preload.objCompressedSize)
 				{
-					m_it = m_it->next;
+					it = it->next;
 					continue;
 				}
 
@@ -247,11 +335,9 @@ namespace uba
 
 				if (preload.objWriteOffset == 0)
 				{
-					if (ToView(info.originalName).EndsWith(TC(".h.obj"))) // TODO: Copied this from UbaSession align code
-						m_totalMemOffset = AlignUp(m_totalMemOffset, 4 * 1024);
-
-					info.fileMapMem = m_totalMem + m_totalMemOffset;
-					m_totalMemOffset += preload.objDecompressedSize;
+					totalMemOffset = AlignUp(totalMemOffset, GetAlignment(ToView(info.originalName)));
+					info.fileMapMem = totalMem + totalMemOffset;
+					totalMemOffset += preload.objDecompressedSize;
 				}
 
 				BinaryReader reader(preload.objMem, preload.objReadOffset, preload.objCompressedSize);
@@ -281,7 +367,8 @@ namespace uba
 				info.memoryFile = nullptr;
 				info.name = L":";
 				info.isFileMap = true;
-				info.fileMapMemEnd = info.fileMapMem + info.size;
+				info.freeFileMapOnClose = allowFreeOnClose;
+				info.fileMapMemSize = info.size;
 				info.trueFileMapHandle = nullptr;
 				info.trueFileMapOffset = 0;
 
@@ -300,9 +387,12 @@ namespace uba
 
 		bool Wait(const StringKey& key)
 		{
+			SCOPED_READ_LOCK(m_preloadedObjFilesLock, lock);
 			auto findIt = m_preloadedObjFiles.find(key);
 			if (findIt == m_preloadedObjFiles.end())
 				return false;
+			lock.Leave();
+
 			Preload& preload = findIt->second;
 			TimerScope ts(g_stats.waitDecompress);
 			preload.event.IsSet();
@@ -333,6 +423,9 @@ namespace uba
 		u8* m_totalMem = nullptr;
 		u64 m_totalMemOffset = 0;
 
+		ReaderWriterLock m_preloadedObjFilesLock;
 		ReaderWriterLock m_threadLock;
+		bool m_isParsing = false;
+		bool m_isEnabled = false;
 	};
 }
diff --git a/Engine/Source/Programs/UnrealBuildAccelerator/Detours/Private/Windows/UbaDetoursUtilsWin.h b/Engine/Source/Programs/UnrealBuildAccelerator/Detours/Private/Windows/UbaDetoursUtilsWin.h
--- a/Engine/Source/Programs/UnrealBuildAccelerator/Detours/Private/Windows/UbaDetoursUtilsWin.h
+++ b/Engine/Source/Programs/UnrealBuildAccelerator/Detours/Private/Windows/UbaDetoursUtilsWin.h
@@ -17,6 +17,11 @@ namespace uba
 		u32 desiredAccess = 0;
 		bool deleteOnClose = false;
 		bool ownsFileInfo = false;
+
+		// Is set to true once content of file has been accessed (ReadFile or MapViewOfFile)
+		bool wasUsed = false;
+
+
 		TString newName;
 	};
 	extern BlockAllocator<FileObject> g_fileObjectAllocator;
