[UBA]
* UbaDetoursObjFilesPreloader.h -> UbaDetoursFileDecompressor.h

diff --git a/Engine/Source/Programs/UnrealBuildAccelerator/Detours/Private/Windows/UbaDetoursFileDecompressor.h b/Engine/Source/Programs/UnrealBuildAccelerator/Detours/Private/Windows/UbaDetoursFileDecompressor.h
new file mode 100644
--- /dev/null
+++ b/Engine/Source/Programs/UnrealBuildAccelerator/Detours/Private/Windows/UbaDetoursFileDecompressor.h
@@ -0,0 +1,431 @@
+// Copyright Epic Games, Inc. All Rights Reserved.
+
+#pragma once
+
+#include "UbaCompressedObjFileHeader.h"
+#include "UbaDetoursUtilsWin.h"
+#include "UbaEvent.h"
+#include "UbaProcessUtils.h"
+#include <oodle2.h>
+#define _NTDEF_
+#include <ntsecapi.h>
+
+namespace uba
+{
+	constexpr bool UseLargePages = false;
+	bool DevirtualizePath(StringBufferBase& path);
+
+	u64 GetLargePageSize()
+	{
+		HINSTANCE hDll = LoadLibrary(TEXT("kernel32.dll"));
+		if (!hDll)
+			return 0;
+		auto lg = MakeGuard([hDll]() { FreeLibrary(hDll); });
+
+		using GetLargePageMinimumType = int(void);
+		auto getLargePageMinimumFunc = (GetLargePageMinimumType*)GetProcAddress(hDll, "GetLargePageMinimum");
+		if (getLargePageMinimumFunc == NULL)
+			return 0;
+
+		DWORD size = (*getLargePageMinimumFunc)();
+		if (!size)
+			return 0;
+
+		HANDLE hToken;
+		if (!OpenProcessToken(GetCurrentProcess(), TOKEN_ADJUST_PRIVILEGES | TOKEN_QUERY, &hToken))
+			return 0;
+		auto tg = MakeGuard([hToken]() { CloseHandle(hToken); });
+
+		TOKEN_PRIVILEGES tp;
+		tp.PrivilegeCount = 1;
+		tp.Privileges[0].Attributes = SE_PRIVILEGE_ENABLED;
+		if (!LookupPrivilegeValue(NULL, L"SeLockMemoryPrivilege", &tp.Privileges[0].Luid))
+			return 0;
+
+		if (!AdjustTokenPrivileges(hToken, FALSE, &tp, 0, (PTOKEN_PRIVILEGES)NULL, 0))
+			return 0;
+
+		// It is possible for AdjustTokenPrivileges to return TRUE and still not succeed. So always check for the last error value.
+		if (GetLastError() != ERROR_SUCCESS)
+			return 0;
+
+		return size;
+	}
+
+	class ObjFilesPreloader
+	{
+	public:
+		void ParseRsp(const StringView& rspFile)
+		{
+			HANDLE rspFileHandle = CreateFileW(rspFile.data, GENERIC_READ, FILE_SHARE_READ, NULL, OPEN_EXISTING, FILE_ATTRIBUTE_NORMAL, NULL);
+			if (rspFileHandle == INVALID_HANDLE_VALUE)
+				return;
+			LARGE_INTEGER rspFileSize;
+			if (!GetFileSizeEx(rspFileHandle, &rspFileSize))
+			{
+				CloseHandle(rspFileHandle);
+				return;
+			}
+			HANDLE rspFileMappingHandle = CreateFileMappingW(rspFileHandle, NULL, PAGE_READONLY, rspFileSize.HighPart, rspFileSize.LowPart, NULL);
+			CloseHandle(rspFileHandle);
+			if (!rspFileMappingHandle)
+				return;
+			const void* rspMem = MapViewOfFile(rspFileMappingHandle, FILE_MAP_READ, 0, 0, rspFileSize.QuadPart);
+			CloseHandle(rspFileMappingHandle);
+			if (!rspMem)
+				return;
+
+			ParseArguments((const char*)rspMem, rspFileSize.QuadPart, [&](const char* arg, u32 argLen)
+				{
+					StringBuffer<> sb;
+					sb.Append(arg, argLen);
+					HandleLine(sb);
+				});
+
+			UnmapViewOfFile(rspMem);
+		}
+
+		void HandleLine(const StringView& line)
+		{
+			if (line.data[0] == '/' || line.data[0] == '-')
+				return;
+
+			if (line.data[0] == '@')
+			{
+				const wchar_t* begin = line.data + 1;
+				const wchar_t* end = line.data + line.count;
+				StringBuffer<> rspFile;
+				if (*begin == '\"')
+				{
+					++begin;
+					--end;
+					if (!end)
+						return;
+				}
+				rspFile.Append(begin, end - begin);
+
+				DevirtualizePath(rspFile);
+				ParseRsp(rspFile);
+				return;
+			}
+
+			StringBuffer<> file;
+			file.Append(line);
+			if (!g_rules->ShouldDecompressFiles(file))
+				return;
+			StringBuffer<> fileFull;
+			FixPath(fileFull, file.data);
+			DevirtualizePath(fileFull);
+
+			HANDLE fileHandle = CreateFileW(fileFull.data, GENERIC_READ, FILE_SHARE_READ, NULL, OPEN_EXISTING, FILE_ATTRIBUTE_NORMAL, NULL);
+			UBA_ASSERTF(fileHandle != INVALID_HANDLE_VALUE, TC("Failed to open %s for read"), fileFull.data);
+			if (fileHandle == INVALID_HANDLE_VALUE)
+				return;
+
+			StringKey fileNameKey = ToStringKey(fileFull.MakeLower());
+			auto insres = m_preloadedObjFiles.try_emplace(fileNameKey);
+			UBA_ASSERT(insres.second);
+			Preload& preload = insres.first->second;
+			preload.handle = fileHandle;
+			preload.event.Create(true);
+
+			if (!m_it)
+				m_it = &preload;
+			if (m_last)
+				m_last->next = &preload;
+			m_last = &preload;
+
+			SCOPED_READ_LOCK(g_mappedFileTable.m_lookupLock, _);
+			auto findIt = g_mappedFileTable.m_lookup.find(fileNameKey);
+			preload.fileInfo = &findIt->second;
+		}
+
+		void Start(const wchar_t* cmdLine)
+		{
+			Oodle_SetUsageWarnings(Oodle_UsageWarnings_Disabled);
+
+			m_isEnabled = true;
+
+			if (true)
+				return;
+
+			TimerScope ts(g_stats.preparseObjFiles);
+
+			m_isParsing = true;
+			ParseArguments(cmdLine, [&](const tchar* arg, u32 argLen) { HandleLine(StringView(arg, argLen)); });
+			m_isParsing = false;
+
+			if (!m_preloadedObjFiles.empty())
+				m_threadHandle = CreateThread(NULL, 0, [](LPVOID p) -> DWORD { ((ObjFilesPreloader*)p)->ThreadPreload(); return 0; }, this, 0, NULL);
+		}
+
+		void Stop()
+		{
+			if (!m_threadHandle)
+				return;
+			WaitForSingleObject(m_threadHandle, INFINITE);
+			CloseHandle(m_threadHandle);
+			WaitForMultipleObjects(sizeof_array(m_helperThreadHandles), m_helperThreadHandles, true, INFINITE);
+			for (u32 i=0;i!=sizeof_array(m_helperThreadHandles); ++i)
+				CloseHandle(m_helperThreadHandles[i]);
+		}
+
+		bool CouldBeCompressedFile(const StringView& fileName)
+		{
+			if (!m_isEnabled)
+				return false;
+			return g_rules->ShouldDecompressFiles(fileName);
+		}
+
+		bool IsDecompressed(const StringKey& fileNameKey, const StringView& fileName, FileInfo& info)
+		{
+			if (!m_isEnabled || m_isParsing)
+				return false;
+
+			if (!g_rules->ShouldDecompressFiles(fileName))
+				return false;
+
+			SCOPED_WRITE_LOCK(m_preloadedObjFilesLock, lock);
+			auto insres = m_preloadedObjFiles.try_emplace(fileNameKey);
+			Preload& preload = insres.first->second;
+			if (!insres.second)
+			{
+				lock.Leave();
+				TimerScope ts(g_stats.waitDecompress);
+				preload.event.IsSet();
+				UBA_ASSERT(preload.fileInfo->fileMapMem || !preload.fileInfo->freeFileMapOnClose);
+				return preload.objMem != nullptr;
+			}
+
+			DEBUG_LOG_TRUE(L"DecompressingFile", L"(%s)", fileName.data);
+			{
+				SuppressDetourScope _;
+				preload.handle = CreateFileW(fileName.data, GENERIC_READ, FILE_SHARE_READ, NULL, OPEN_EXISTING, FILE_ATTRIBUTE_NORMAL, NULL);
+				UBA_ASSERTF(preload.handle != INVALID_HANDLE_VALUE, TC("Failed to open %s for read"), fileName.data);
+				if (preload.handle == INVALID_HANDLE_VALUE)
+				{
+					preload.event.Set();
+					return false;
+				}
+			}
+
+			TimerScope ts(g_stats.waitDecompress);
+
+			preload.fileInfo = &info;
+			preload.event.Create(true);
+			u64 memoryNeeded = PreparePreload(preload, 0);
+			if (!memoryNeeded)
+			{
+				preload.event.Set();
+				return false;
+			}
+			u64 memoryOffset = 0;
+
+			//u64 alignedMem = AlignUp(memoryNeeded, 64*1024);
+			//DWORD allocationType = MEM_RESERVE | MEM_COMMIT;
+			//u8* memory = (u8*)VirtualAlloc(NULL, alignedMem, allocationType, PAGE_READWRITE);
+			u8* memory = (u8*)_aligned_malloc(memoryNeeded, GetAlignment(fileName));
+			Preload* it = &preload;
+			DecompressWork(it, memory, memoryOffset, true);
+			return true;
+		}
+
+		struct Preload;
+
+		u64 GetAlignment(const StringView& fileName)
+		{
+			if (fileName.EndsWith(TC(".h.obj"))) // TODO: Copied this from UbaSession align code
+				return 4 * 1024;
+			return 1;
+		}
+
+		u64 PreparePreload(Preload& preload, u64 totalMemSize)
+		{
+			HANDLE objFileHandle = preload.handle;
+			LARGE_INTEGER objFileSize;
+			BOOL res = GetFileSizeEx(objFileHandle, &objFileSize);
+			UBA_ASSERT(res);(void)res;
+			if (objFileSize.QuadPart < sizeof(CompressedObjFileHeader))
+			{
+				CloseHandle(objFileHandle);
+				preload.event.Set();
+				return totalMemSize;
+			}
+
+			HANDLE objFileMappingHandle = CreateFileMappingW(objFileHandle, NULL, PAGE_READONLY, objFileSize.HighPart, objFileSize.LowPart, NULL);
+			CloseHandle(objFileHandle);
+			UBA_ASSERT(objFileMappingHandle);
+			if (!objFileMappingHandle)
+				return totalMemSize;
+			preload.handle = NULL;
+			u8* objMem = (u8*)MapViewOfFile(objFileMappingHandle, FILE_MAP_READ, 0, 0, objFileSize.QuadPart);
+			UBA_ASSERT(objMem);
+			CloseHandle(objFileMappingHandle);
+			if (!objMem)
+				return totalMemSize;
+			if (!((CompressedObjFileHeader*)objMem)->IsValid())
+			{
+				UnmapViewOfFile(objMem);
+				// Ok, this is backwards.. but it could be that the file on server's disk is compressed..
+				// but it is decompressed in client's cas.. in that case dir table has the size of the compressed file.. 
+				// and since detoured code will provide decompressed file it also needs to provide decompressed size
+				preload.fileInfo->size = objFileSize.QuadPart;
+				preload.event.Set();
+				return totalMemSize;
+			}
+			preload.objMem = objMem;
+			preload.objCompressedSize = objFileSize.QuadPart;
+			preload.objDecompressedSize = *(u64*)(preload.objMem + sizeof(CompressedObjFileHeader));
+			preload.objReadOffset = 8 + sizeof(CompressedObjFileHeader);
+			preload.objWriteOffset = 0;
+			preload.objLeft = preload.objDecompressedSize;
+
+			totalMemSize = AlignUp(totalMemSize, GetAlignment(ToView(preload.fileInfo->originalName)));
+			return totalMemSize + preload.objDecompressedSize;
+		}
+
+		void ThreadPreload()
+		{
+			u64 totalMemSize = 0;
+			for (Preload* preloadPtr=m_it; preloadPtr; preloadPtr = preloadPtr->next)
+				totalMemSize = PreparePreload(*preloadPtr, totalMemSize);
+
+			u64 alignedMem = AlignUp(totalMemSize, 64*1024);
+			DWORD allocationType = MEM_RESERVE | MEM_COMMIT;
+
+			if (UseLargePages)
+			{
+				if (u64 largePageSize = GetLargePageSize())
+				{
+					alignedMem = AlignUp(totalMemSize, largePageSize);
+					allocationType |= MEM_LARGE_PAGES;
+				}
+			}
+			m_totalMem = (u8*)VirtualAlloc(NULL, alignedMem, allocationType, PAGE_READWRITE);
+
+			for (u32 i=0;i!=sizeof_array(m_helperThreadHandles); ++i)
+				m_helperThreadHandles[i] = CreateThread(NULL, 0, [](LPVOID p) -> DWORD
+					{
+						auto& pl = *(ObjFilesPreloader*)p;
+						pl.DecompressWork(pl.m_it, pl.m_totalMem, pl.m_totalMemOffset, false);
+						return 0;
+					}, this, 0, NULL);
+			DecompressWork(m_it, m_totalMem, m_totalMemOffset, false);
+		}
+
+		void DecompressWork(Preload*& it, u8* totalMem, u64& totalMemOffset, bool allowFreeOnClose)
+		{
+			OO_SINTa decoredMemSize = OodleLZDecoder_MemorySizeNeeded(OodleLZ_Compressor_Kraken);
+			void* decoderMem = malloc(decoredMemSize);
+			auto mg = MakeGuard([decoderMem]() { free(decoderMem); });
+
+			while (true)
+			{
+				SCOPED_WRITE_LOCK(m_threadLock, lock);
+				if (!it)
+					break;
+				Preload& preload = *it;
+				if (preload.objReadOffset == preload.objCompressedSize)
+				{
+					it = it->next;
+					continue;
+				}
+
+				FileInfo& info = *preload.fileInfo;
+
+				if (preload.objWriteOffset == 0)
+				{
+					totalMemOffset = AlignUp(totalMemOffset, GetAlignment(ToView(info.originalName)));
+					info.fileMapMem = totalMem + totalMemOffset;
+					totalMemOffset += preload.objDecompressedSize;
+				}
+
+				BinaryReader reader(preload.objMem, preload.objReadOffset, preload.objCompressedSize);
+
+				u32 compressedBlockSize = reader.ReadU32();
+				u32 decompressedBlockSize = reader.ReadU32();
+
+				u8* destMem = info.fileMapMem + preload.objWriteOffset;
+
+				preload.objReadOffset += compressedBlockSize + 8;
+				preload.objWriteOffset += decompressedBlockSize;
+
+				lock.Leave();
+
+				OO_SINTa decompLen = OodleLZ_Decompress(reader.GetPositionData(), (OO_SINTa)compressedBlockSize, destMem, (OO_SINTa)decompressedBlockSize,
+					OodleLZ_FuzzSafe_Yes, OodleLZ_CheckCRC_No, OodleLZ_Verbosity_None, NULL, 0, NULL, NULL, decoderMem, decoredMemSize);
+				if (decompLen != decompressedBlockSize)
+				{
+					FatalError(1356, TC("Failed to decompress .obj file %s (Expected to get %llu but got %llu)"), info.originalName, u64(decompressedBlockSize), u64(decompLen));
+				}
+
+				bool isDone = preload.objLeft.fetch_sub(decompressedBlockSize) == decompressedBlockSize;
+				if (!isDone)
+					continue;
+
+				info.size = preload.objDecompressedSize;
+				info.memoryFile = nullptr;
+				info.name = L":";
+				info.isFileMap = true;
+				info.freeFileMapOnClose = allowFreeOnClose;
+				info.fileMapMemSize = info.size;
+				info.trueFileMapHandle = nullptr;
+				info.trueFileMapOffset = 0;
+
+				{
+					SCOPED_WRITE_LOCK(g_mappedFileTable.m_memLookupLock, lock3);
+					auto insres = g_mappedFileTable.m_memLookup.try_emplace(destMem);
+					UBA_ASSERT(insres.second);
+					insres.first->second = 1;
+				}
+
+				preload.event.Set();
+
+				UnmapViewOfFile(preload.objMem);
+			}
+		}
+
+		bool Wait(const StringKey& key)
+		{
+			SCOPED_READ_LOCK(m_preloadedObjFilesLock, lock);
+			auto findIt = m_preloadedObjFiles.find(key);
+			if (findIt == m_preloadedObjFiles.end())
+				return false;
+			lock.Leave();
+
+			Preload& preload = findIt->second;
+			TimerScope ts(g_stats.waitDecompress);
+			preload.event.IsSet();
+			return true;
+		}
+
+		struct Preload
+		{
+			Preload* next = nullptr;
+			Event event;
+			HANDLE handle;
+			u8* objMem = nullptr;
+			u64 objCompressedSize = 0;
+			u64 objDecompressedSize = 0;
+			FileInfo* fileInfo = nullptr;
+			u64 objReadOffset = 0;
+			u64 objWriteOffset = 0;
+			Atomic<u64> objLeft;
+		};
+
+		Preload* m_it = nullptr;
+		Preload* m_last = nullptr;
+
+		using PreloadedObjFiles = UnorderedMap<StringKey, Preload>;
+		PreloadedObjFiles m_preloadedObjFiles;
+		HANDLE m_threadHandle = 0;
+		HANDLE m_helperThreadHandles[5];
+		u8* m_totalMem = nullptr;
+		u64 m_totalMemOffset = 0;
+
+		ReaderWriterLock m_preloadedObjFilesLock;
+		ReaderWriterLock m_threadLock;
+		bool m_isParsing = false;
+		bool m_isEnabled = false;
+	};
+}
diff --git a/Engine/Source/Programs/UnrealBuildAccelerator/Detours/Private/Windows/UbaDetoursFunctionsWin.cpp b/Engine/Source/Programs/UnrealBuildAccelerator/Detours/Private/Windows/UbaDetoursFunctionsWin.cpp
--- a/Engine/Source/Programs/UnrealBuildAccelerator/Detours/Private/Windows/UbaDetoursFunctionsWin.cpp
+++ b/Engine/Source/Programs/UnrealBuildAccelerator/Detours/Private/Windows/UbaDetoursFunctionsWin.cpp
@@ -44,7 +44,7 @@ DETOURED_FUNCTIONS_MEMORY_NON_WINE
 #include "UbaDetoursPayload.h"
 #include "UbaApplicationRules.h"
 #include "UbaDetoursShared.h"
-#include "UbaDetoursObjFilesPreloader.h"
+#include "UbaDetoursFileDecompressor.h"
 
 #include "Shlwapi.h"
 #include <detours/detours.h>
diff --git a/Engine/Source/Programs/UnrealBuildAccelerator/Detours/Private/Windows/UbaDetoursObjFilesPreloader.h b/Engine/Source/Programs/UnrealBuildAccelerator/Detours/Private/Windows/UbaDetoursObjFilesPreloader.h
deleted file mode 100644
