[UBA]
* Renamed ObjFilePreloader to FileDecompressor

diff --git a/Engine/Source/Programs/UnrealBuildAccelerator/Detours/Private/Windows/UbaDetoursFileDecompressor.h b/Engine/Source/Programs/UnrealBuildAccelerator/Detours/Private/Windows/UbaDetoursFileDecompressor.h
--- a/Engine/Source/Programs/UnrealBuildAccelerator/Detours/Private/Windows/UbaDetoursFileDecompressor.h
+++ b/Engine/Source/Programs/UnrealBuildAccelerator/Detours/Private/Windows/UbaDetoursFileDecompressor.h
@@ -52,7 +52,7 @@ namespace uba
 		return size;
 	}
 
-	class ObjFilesPreloader
+	class FileDecompressor
 	{
 	public:
 		void ParseRsp(const StringView& rspFile)
@@ -156,7 +156,7 @@ namespace uba
 			m_isParsing = false;
 
 			if (!m_preloadedObjFiles.empty())
-				m_threadHandle = CreateThread(NULL, 0, [](LPVOID p) -> DWORD { ((ObjFilesPreloader*)p)->ThreadPreload(); return 0; }, this, 0, NULL);
+				m_threadHandle = CreateThread(NULL, 0, [](LPVOID p) -> DWORD { ((FileDecompressor*)p)->ThreadPreload(); return 0; }, this, 0, NULL);
 		}
 
 		void Stop()
@@ -306,7 +306,7 @@ namespace uba
 			for (u32 i=0;i!=sizeof_array(m_helperThreadHandles); ++i)
 				m_helperThreadHandles[i] = CreateThread(NULL, 0, [](LPVOID p) -> DWORD
 					{
-						auto& pl = *(ObjFilesPreloader*)p;
+						auto& pl = *(FileDecompressor*)p;
 						pl.DecompressWork(pl.m_it, pl.m_totalMem, pl.m_totalMemOffset, false);
 						return 0;
 					}, this, 0, NULL);
diff --git a/Engine/Source/Programs/UnrealBuildAccelerator/Detours/Private/Windows/UbaDetoursFunctionsNtDll.inl b/Engine/Source/Programs/UnrealBuildAccelerator/Detours/Private/Windows/UbaDetoursFunctionsNtDll.inl
--- a/Engine/Source/Programs/UnrealBuildAccelerator/Detours/Private/Windows/UbaDetoursFunctionsNtDll.inl
+++ b/Engine/Source/Programs/UnrealBuildAccelerator/Detours/Private/Windows/UbaDetoursFunctionsNtDll.inl
@@ -763,7 +763,7 @@ NTSTATUS NTAPI Shared_NtCreateFile(bool IsCreateFunc, PHANDLE hFileHandle, ACCES
 	fileNameLower.MakeLower();
 	StringKey fileNameKey = ToStringKey(fileNameLower);
 
-	bool wasPreload = g_objFilesPreloader.Wait(fileNameKey); // Wait for on-going preload if existing
+	bool wasPreload = g_fileDecompressor.Wait(fileNameKey); // Wait for on-going preload if existing
 
 
 	if (createFileName)// && (keepInMemory || canDetour))
@@ -819,7 +819,7 @@ NTSTATUS NTAPI Shared_NtCreateFile(bool IsCreateFunc, PHANDLE hFileHandle, ACCES
 
 				bool isWriteAttributes = (DesiredAccess & FILE_WRITE_ATTRIBUTES) != 0;
 
-				if (allowEarlyOut && !useContent && !isWriteAttributes && !wasPreload && !g_objFilesPreloader.CouldBeCompressedFile(fileName))
+				if (allowEarlyOut && !useContent && !isWriteAttributes && !wasPreload && !g_fileDecompressor.CouldBeCompressedFile(fileName))
 				{
 					auto dh = new DetouredHandle(HandleType_File);
 					dh->fileObject = new FileObject();
@@ -1189,7 +1189,7 @@ NTSTATUS NTAPI Shared_NtCreateFile(bool IsCreateFunc, PHANDLE hFileHandle, ACCES
 		tempFileName = info.name;
 
 
-	if (g_objFilesPreloader.IsDecompressed(fileNameKey, fileName, info))
+	if (g_fileDecompressor.IsDecompressed(fileNameKey, fileName, info))
 	{
 		*hFileHandle = CreateFileHandle();
 		TrackFileInput();
diff --git a/Engine/Source/Programs/UnrealBuildAccelerator/Detours/Private/Windows/UbaDetoursFunctionsWin.cpp b/Engine/Source/Programs/UnrealBuildAccelerator/Detours/Private/Windows/UbaDetoursFunctionsWin.cpp
--- a/Engine/Source/Programs/UnrealBuildAccelerator/Detours/Private/Windows/UbaDetoursFunctionsWin.cpp
+++ b/Engine/Source/Programs/UnrealBuildAccelerator/Detours/Private/Windows/UbaDetoursFunctionsWin.cpp
@@ -227,7 +227,7 @@ ReaderWriterLock g_loadedModulesLock;
 UnorderedMap<HMODULE, TString> g_loadedModules;
 u64 g_memoryFileIndexCounter = ~u64(0) - 1000000; // I really hope this will not collide with anything
 
-ObjFilesPreloader g_objFilesPreloader;
+FileDecompressor g_fileDecompressor;
 
 
 struct SuppressCreateFileDetourScope
@@ -981,12 +981,12 @@ void Init(const DetoursPayload& payload, u64 startTime)
 	g_stats.attach.count = 1;
 
 	if (payload.storeObjFilesCompressed && g_rules->ShouldDecompressFiles(StringView()))
-		g_objFilesPreloader.Start(cmdLine);
+		g_fileDecompressor.Start(cmdLine);
 }
 
 void Deinit(u64 startTime)
 {
-	g_objFilesPreloader.Stop();
+	g_fileDecompressor.Stop();
 
 	if (g_isRunningWine) // mt.exe etc fails if detaching is not done during shutdown
 	{
