[UBA]
* Fixed remote linking when compressed obj files are enabled

diff --git a/Engine/Source/Programs/UnrealBuildAccelerator/Detours/Private/UbaDetoursFileMappingTable.cpp b/Engine/Source/Programs/UnrealBuildAccelerator/Detours/Private/UbaDetoursFileMappingTable.cpp
--- a/Engine/Source/Programs/UnrealBuildAccelerator/Detours/Private/UbaDetoursFileMappingTable.cpp
+++ b/Engine/Source/Programs/UnrealBuildAccelerator/Detours/Private/UbaDetoursFileMappingTable.cpp
@@ -289,6 +289,10 @@ namespace uba
 
 		u32 mappedFileTableSize;
 
+		#if UBA_DEBUG
+		StringBuffer<> virtualName;
+		#endif
+
 		{
 			TimerScope ts(g_stats.getFullFileName);
 			SCOPED_WRITE_LOCK(g_communicationLock, pcs);
@@ -308,16 +312,11 @@ namespace uba
 			if (useVirtualName)
 			{
 				reader.ReadString(tempBuf.Clear());
-				#if UBA_DEBUG
-				DEBUG_LOG_DETOURED(TC("Rpc_GetFullFileName"), TC("%s -> %s"), path, tempBuf.data);
-				#endif
 			}
 			else
 			{
 				#if UBA_DEBUG
-				StringBuffer<> virtualName;
 				reader.ReadString(virtualName);
-				DEBUG_LOG_DETOURED(TC("Rpc_GetFullFileName"), TC("%s -> %s (%s)"), path, tempBuf.data, virtualName.data);
 				#else
 				reader.SkipString();
 				#endif
@@ -326,6 +325,13 @@ namespace uba
 			DEBUG_LOG_PIPE(TC("GetFileName"), TC("(%ls)"), tempBuf);
 		}
 
+		#if UBA_DEBUG
+		if (useVirtualName)
+		{ DEBUG_LOG_DETOURED(TC("Rpc_GetFullFileName"), TC("%s -> %s"), path, tempBuf.data); }
+		else
+		{ DEBUG_LOG_DETOURED(TC("Rpc_GetFullFileName"), TC("%s -> %s (%s)"), path, tempBuf.data, virtualName.data); }
+		#endif
+
 		g_mappedFileTable.Parse(mappedFileTableSize);
 		path = tempBuf.data;
 		pathLen = tempBuf.count;
diff --git a/Engine/Source/Programs/UnrealBuildAccelerator/Detours/Private/UbaDetoursShared.cpp b/Engine/Source/Programs/UnrealBuildAccelerator/Detours/Private/UbaDetoursShared.cpp
--- a/Engine/Source/Programs/UnrealBuildAccelerator/Detours/Private/UbaDetoursShared.cpp
+++ b/Engine/Source/Programs/UnrealBuildAccelerator/Detours/Private/UbaDetoursShared.cpp
@@ -138,8 +138,9 @@ namespace uba
 	}
 	LogScope::LogScope()
 	{
-		if (++t_logScopeCount  == 1)
-			g_logScopeLock.EnterWrite();
+		if (++t_logScopeCount > 1)
+			return;
+		//g_logScopeLock.EnterWrite(); // Deadlocks in a few places
 	}
 	LogScope::~LogScope()
 	{
@@ -147,7 +148,7 @@ namespace uba
 			return;
 		if (t_b_size)
 			Flush();
-		g_logScopeLock.LeaveWrite();
+		//g_logScopeLock.LeaveWrite();
 	}
 	void LogScope::Flush()
 	{
diff --git a/Engine/Source/Programs/UnrealBuildAccelerator/Detours/Private/Windows/UbaDetoursFileDecompressor.h b/Engine/Source/Programs/UnrealBuildAccelerator/Detours/Private/Windows/UbaDetoursFileDecompressor.h
--- a/Engine/Source/Programs/UnrealBuildAccelerator/Detours/Private/Windows/UbaDetoursFileDecompressor.h
+++ b/Engine/Source/Programs/UnrealBuildAccelerator/Detours/Private/Windows/UbaDetoursFileDecompressor.h
@@ -177,7 +177,7 @@ namespace uba
 			return g_rules->ShouldDecompressFiles(fileName);
 		}
 
-		bool IsDecompressed(const StringKey& fileNameKey, const StringView& fileName, FileInfo& info)
+		bool IsDecompressed(const StringKey& fileNameKey, const StringView& fileName, const StringView& realFileName, FileInfo& info)
 		{
 			if (!m_isEnabled || m_isParsing)
 				return false;
@@ -200,8 +200,8 @@ namespace uba
 			DEBUG_LOG_TRUE(L"DecompressingFile", L"(%s)", fileName.data);
 			{
 				SuppressDetourScope _;
-				preload.handle = CreateFileW(fileName.data, GENERIC_READ, FILE_SHARE_READ, NULL, OPEN_EXISTING, FILE_ATTRIBUTE_NORMAL, NULL);
-				UBA_ASSERTF(preload.handle != INVALID_HANDLE_VALUE, TC("Failed to open %s for read"), fileName.data);
+				preload.handle = CreateFileW(realFileName.data, GENERIC_READ, FILE_SHARE_READ, NULL, OPEN_EXISTING, FILE_ATTRIBUTE_NORMAL, NULL);
+				UBA_ASSERTF(preload.handle != INVALID_HANDLE_VALUE, TC("Failed to open %s for read"), realFileName.data);
 				if (preload.handle == INVALID_HANDLE_VALUE)
 				{
 					preload.event.Set();
diff --git a/Engine/Source/Programs/UnrealBuildAccelerator/Detours/Private/Windows/UbaDetoursFunctionsNtDll.inl b/Engine/Source/Programs/UnrealBuildAccelerator/Detours/Private/Windows/UbaDetoursFunctionsNtDll.inl
--- a/Engine/Source/Programs/UnrealBuildAccelerator/Detours/Private/Windows/UbaDetoursFunctionsNtDll.inl
+++ b/Engine/Source/Programs/UnrealBuildAccelerator/Detours/Private/Windows/UbaDetoursFunctionsNtDll.inl
@@ -1182,14 +1182,14 @@ NTSTATUS NTAPI Shared_NtCreateFile(bool IsCreateFunc, PHANDLE hFileHandle, ACCES
 		return STATUS_SUCCESS;
 	}
 
-	const wchar_t* tempFileName = lpFileName;
-	if (tempFileName[0] == '#')
-		tempFileName = fileName.data;
+	StringView tempFileName;
+	if (lpFileName[0] == '#')
+		tempFileName = fileName;
 	else
-		tempFileName = info.name;
+		tempFileName = ToView(info.name);
 
 
-	if (g_fileDecompressor.IsDecompressed(fileNameKey, fileName, info))
+	if (g_fileDecompressor.IsDecompressed(fileNameKey, fileName, tempFileName, info))
 	{
 		*hFileHandle = CreateFileHandle();
 		TrackFileInput();
@@ -1227,7 +1227,7 @@ NTSTATUS NTAPI Shared_NtCreateFile(bool IsCreateFunc, PHANDLE hFileHandle, ACCES
 	}
 	if (!canDetour)
 	{
-		DEBUG_LOG_TRUE(funcName, L"%ls %llu NODETOUR (%ls) -> %ls", isWriteStr, uintptr_t(*hFileHandle), tempFileName, ToString(res).data);
+		DEBUG_LOG_TRUE(funcName, L"%ls %llu NODETOUR (%ls) -> %ls", isWriteStr, uintptr_t(*hFileHandle), tempFileName.data, ToString(res).data);
 		return res;
 	}
 
@@ -1235,7 +1235,7 @@ NTSTATUS NTAPI Shared_NtCreateFile(bool IsCreateFunc, PHANDLE hFileHandle, ACCES
 
 	UBA_ASSERT(info.originalName);
 	*hFileHandle = CreateFileHandle(*hFileHandle);
-	DEBUG_LOG_TRUE(funcName, L"%ls %llu (%ls)%s -> %ls", isWriteStr, uintptr_t(*hFileHandle), tempFileName, isDeleteOnClose ? TC(" DeleteOnClose") : TC(""), ToString(res).data);
+	DEBUG_LOG_TRUE(funcName, L"%ls %llu (%ls)%s -> %ls", isWriteStr, uintptr_t(*hFileHandle), tempFileName.data, isDeleteOnClose ? TC(" DeleteOnClose") : TC(""), ToString(res).data);
 	return res;
 }
 
