[UBA]
* Fixed bug where freed decompressed memory was not removed from mapped file table memory lookup

diff --git a/Engine/Source/Programs/UnrealBuildAccelerator/Detours/Private/Windows/UbaDetoursFileDecompressor.h b/Engine/Source/Programs/UnrealBuildAccelerator/Detours/Private/Windows/UbaDetoursFileDecompressor.h
--- a/Engine/Source/Programs/UnrealBuildAccelerator/Detours/Private/Windows/UbaDetoursFileDecompressor.h
+++ b/Engine/Source/Programs/UnrealBuildAccelerator/Detours/Private/Windows/UbaDetoursFileDecompressor.h
@@ -230,6 +230,18 @@ namespace uba
 			return true;
 		}
 
+		void ReturnMemory(FileInfo& fi)
+		{
+			DEBUG_LOG_DETOURED(L"ReturnMemory", L"%s", fi.originalName);
+			SCOPED_WRITE_LOCK(g_mappedFileTable.m_memLookupLock, lock);
+			auto res = g_mappedFileTable.m_memLookup.erase(fi.fileMapMem);
+			UBA_ASSERT(res);
+			lock.Leave();
+			//VirtualFree(fi.fileMapMem, 0, MEM_RELEASE);
+			_aligned_free(fi.fileMapMem);
+			fi.fileMapMem = nullptr;
+		}
+
 		struct Preload;
 
 		u64 GetAlignment(const StringView& fileName)
@@ -375,7 +387,7 @@ namespace uba
 				{
 					SCOPED_WRITE_LOCK(g_mappedFileTable.m_memLookupLock, lock3);
 					auto insres = g_mappedFileTable.m_memLookup.try_emplace(destMem);
-					UBA_ASSERT(insres.second);
+					UBA_ASSERTF(insres.second, TC("Failed to add memory 0x%llx for %s"), u64(destMem), info.originalName);
 					insres.first->second = 1;
 				}
 
diff --git a/Engine/Source/Programs/UnrealBuildAccelerator/Detours/Private/Windows/UbaDetoursFunctionsNtDll.inl b/Engine/Source/Programs/UnrealBuildAccelerator/Detours/Private/Windows/UbaDetoursFunctionsNtDll.inl
--- a/Engine/Source/Programs/UnrealBuildAccelerator/Detours/Private/Windows/UbaDetoursFunctionsNtDll.inl
+++ b/Engine/Source/Programs/UnrealBuildAccelerator/Detours/Private/Windows/UbaDetoursFunctionsNtDll.inl
@@ -1471,12 +1471,7 @@ NTSTATUS NTAPI Detoured_NtClose(HANDLE handle)
 
 	// TODO: Can't free here..... Might still be open with CreateFileMapping or MapViewOfFile
 	if (fo->wasUsed && fi.freeFileMapOnClose)
-	{
-		DEBUG_LOG_DETOURED(L"VirtualFree", L"%llu (%ls) -> %ls", uintptr_t(handle), HandleToName(handle), ToString(res).data);
-		//VirtualFree(fi.fileMapMem, 0, MEM_RELEASE);
-		_aligned_free(fi.fileMapMem);
-		fi.fileMapMem = nullptr;
-	}
+		g_fileDecompressor.ReturnMemory(fi);
 
 	DEBUG_LOG_DETOURED(L"NtClose", L"%llu (%ls) -> %ls", uintptr_t(handle), HandleToName(handle), ToString(res).data);
 
