[UBA]
* Fixed shipping compile warning

diff --git a/Engine/Source/Programs/UnrealBuildAccelerator/Detours/Private/Windows/UbaDetoursFileDecompressor.h b/Engine/Source/Programs/UnrealBuildAccelerator/Detours/Private/Windows/UbaDetoursFileDecompressor.h
--- a/Engine/Source/Programs/UnrealBuildAccelerator/Detours/Private/Windows/UbaDetoursFileDecompressor.h
+++ b/Engine/Source/Programs/UnrealBuildAccelerator/Detours/Private/Windows/UbaDetoursFileDecompressor.h
@@ -234,7 +234,7 @@ namespace uba
 		{
 			DEBUG_LOG_DETOURED(L"ReturnMemory", L"%s", fi.originalName);
 			SCOPED_WRITE_LOCK(g_mappedFileTable.m_memLookupLock, lock);
-			auto res = g_mappedFileTable.m_memLookup.erase(fi.fileMapMem);
+			auto res = g_mappedFileTable.m_memLookup.erase(fi.fileMapMem);(void)res;
 			UBA_ASSERT(res);
 			lock.Leave();
 			//VirtualFree(fi.fileMapMem, 0, MEM_RELEASE);
