[UBA]
* Fixed so both x64 and arm64 detours library are sent to arm64 helpers on connect
* Fixed assert happening in debug when linking with compressed obj files

diff --git a/Engine/Source/Programs/UnrealBuildAccelerator/Common/Private/UbaProcess.cpp b/Engine/Source/Programs/UnrealBuildAccelerator/Common/Private/UbaProcess.cpp
--- a/Engine/Source/Programs/UnrealBuildAccelerator/Common/Private/UbaProcess.cpp
+++ b/Engine/Source/Programs/UnrealBuildAccelerator/Common/Private/UbaProcess.cpp
@@ -1008,8 +1008,8 @@ namespace uba
 		auto& process = *(ProcessImpl*)h.m_process;
 		process.m_echoOn = m_echoOn;
 
-		const char* detoursLib = m_session.m_detoursLibrary.c_str();
-		u32 detoursLibLen = u32(m_session.m_detoursLibrary.size());
+		const char* detoursLib = m_session.m_detoursLibrary[IsArm].c_str();
+		u32 detoursLibLen = u32(m_session.m_detoursLibrary[IsArm].size());
 
 		writer.WriteU32(childProcessId);
 		writer.WriteU32(info.rules->index);
@@ -1436,7 +1436,7 @@ namespace uba
 
 		if (!m_parentProcess)
 		{
-			const char* detoursLib = m_session.m_detoursLibrary.c_str();
+			const char* detoursLib = m_session.m_detoursLibrary[IsArm].c_str();
 			if (!*detoursLib)
 				detoursLib = UBA_DETOURS_LIBRARY_ANSI;
 
@@ -1828,7 +1828,7 @@ namespace uba
 
 			if (m_detourEnabled)
 			{
-				const char* detoursLib = m_session.m_detoursLibrary.c_str();
+				const char* detoursLib = m_session.m_detoursLibrary[IsArm].c_str();
 				if (*detoursLib)
 				{
 #if PLATFORM_LINUX
diff --git a/Engine/Source/Programs/UnrealBuildAccelerator/Common/Private/UbaSession.cpp b/Engine/Source/Programs/UnrealBuildAccelerator/Common/Private/UbaSession.cpp
--- a/Engine/Source/Programs/UnrealBuildAccelerator/Common/Private/UbaSession.cpp
+++ b/Engine/Source/Programs/UnrealBuildAccelerator/Common/Private/UbaSession.cpp
@@ -2752,7 +2752,7 @@ namespace uba
 		StringBuffer<256> detoursLibPath;
 		StringBuffer<256> alternativeLibPath;
 
-		detoursLibPath.Append(m_detoursLibrary.c_str()).Resize(detoursLibPath.Last(PathSeparator) - detoursLibPath.data);
+		detoursLibPath.Append(m_detoursLibrary[IsArm].c_str()).Resize(detoursLibPath.Last(PathSeparator) - detoursLibPath.data);
 		GetAlternativeUbaPath(m_logger, alternativeLibPath, detoursLibPath, IsArm);
 
 		StringView searchPaths[3] = { detoursLibPath, alternativeLibPath, {} };
diff --git a/Engine/Source/Programs/UnrealBuildAccelerator/Common/Private/UbaSessionClient.cpp b/Engine/Source/Programs/UnrealBuildAccelerator/Common/Private/UbaSessionClient.cpp
--- a/Engine/Source/Programs/UnrealBuildAccelerator/Common/Private/UbaSessionClient.cpp
+++ b/Engine/Source/Programs/UnrealBuildAccelerator/Common/Private/UbaSessionClient.cpp
@@ -1219,34 +1219,47 @@ namespace uba
 			return;
 		}
 
-		CasKey detoursBinaryKey = reader.ReadCasKey();
+		CasKey detoursBinaryKey[2];
+		detoursBinaryKey[0] = reader.ReadCasKey();
+		if (IsArm)
+			detoursBinaryKey[1] = reader.ReadCasKey();
+
+		StringBuffer<> detoursFile[2];
 
-		StringBuffer<> detoursFile;
 		if (m_downloadDetoursLib)
 		{
+			for (u32 i=0; i!=IsArm+1; ++i)
 			{
-				TimerScope s(m_stats.storageRetrieve);
-				Storage::RetrieveResult result;
-				if (!m_storage.RetrieveCasFile(result, AsCompressed(detoursBinaryKey, false), UBA_DETOURS_LIBRARY))
+				{
+					TimerScope s(m_stats.storageRetrieve);
+					Storage::RetrieveResult result;
+					if (!m_storage.RetrieveCasFile(result, AsCompressed(detoursBinaryKey[i], false), UBA_DETOURS_LIBRARY))
+						return;
+				}
+				StringKey key;
+				key.a = i;
+				KeyToString dir(key);
+
+				if (!WriteBinFile(detoursFile[i], UBA_DETOURS_LIBRARY, detoursBinaryKey[i], dir, DefaultAttributes()))
 					return;
 			}
-			KeyToString dir(StringKeyZero);
-			if (!WriteBinFile(detoursFile, UBA_DETOURS_LIBRARY, detoursBinaryKey, dir, DefaultAttributes()))
-				return;
 		}
 		else
 		{
-			GetDirectoryOfCurrentModule(m_logger, detoursFile);
-			detoursFile.EnsureEndsWithSlash().Append(UBA_DETOURS_LIBRARY);
+			GetDirectoryOfCurrentModule(m_logger, detoursFile[IsArm]);
+			detoursFile[IsArm].EnsureEndsWithSlash().Append(UBA_DETOURS_LIBRARY);
 		}
 
-		#if PLATFORM_WINDOWS
-		char dll[1024];
-		sprintf_s(dll, sizeof(dll), "%ls", detoursFile.data);
-		m_detoursLibrary = dll;
-		#else
-		m_detoursLibrary = detoursFile.data;
-		#endif
+		for (u32 i=0; i!=IsArm+1; ++i)
+		{
+			#if PLATFORM_WINDOWS
+			char dll[1024];
+			detoursFile[i].Parse(dll, sizeof_array(dll));
+			m_detoursLibrary[i] = dll;
+			#else
+			m_detoursLibrary[i] = detoursFile[i].data;
+			#endif
+		}
 
 		bool resetCas = reader.ReadBool();
 		if (resetCas)
@@ -2127,6 +2140,7 @@ namespace uba
 		StackBinaryWriter<SendMaxSize> msgWriter;
 		NetworkMessage msg(m_client, ServiceId, SessionMessageType_GetSymbols, msgWriter);
 		msgWriter.WriteU32(process.GetId());
+		msgWriter.WriteBool(IsArm);
 		msgWriter.WriteBytes(reader.GetPositionData(), reader.GetLeft());
 
 		BinaryReader responseReader(writer.GetData(), 0, writer.GetCapacityLeft());
diff --git a/Engine/Source/Programs/UnrealBuildAccelerator/Common/Private/UbaSessionServer.cpp b/Engine/Source/Programs/UnrealBuildAccelerator/Common/Private/UbaSessionServer.cpp
--- a/Engine/Source/Programs/UnrealBuildAccelerator/Common/Private/UbaSessionServer.cpp
+++ b/Engine/Source/Programs/UnrealBuildAccelerator/Common/Private/UbaSessionServer.cpp
@@ -103,6 +103,21 @@ namespace uba
 		}
 	}
 
+	bool GetCrossArchitectureDir(Logger& logger, StringBufferBase& dir, bool reportError)
+	{
+		// UBT has the path win-x64/native or win-arm64/native
+		bool isUbtPath = dir.EndsWith(TC("native"));
+		if (isUbtPath)
+			dir.Resize(dir.count - 7); // Remove native and slash
+		const tchar* archPath[2] = { TC("x64"), TC("arm64") };
+		if (!dir.EndsWith(archPath[IsArm]))
+			return reportError ? logger.Error(TC("Module dir is not under supported folder (%s) to be able to run cross architecutres, can't figure out matching x64/arm64 folder"), dir.data) : false;
+		dir.Resize(dir.count - TStrlen(archPath[IsArm])).Append(archPath[!IsArm]);
+		if (isUbtPath)
+			dir.Append(PathSeparator).Append(TC("native"));
+		return true;
+	}
+
 	SessionServer::SessionServer(const SessionServerCreateInfo& info, const u8* environment, u32 environmentSize)
 	:	Session(info, TC("UbaSessionServer"), false, &info.server)
 	,	m_server(info.server)
@@ -181,14 +196,21 @@ namespace uba
 			UBA_ASSERT(false);
 			return;
 		}
+		u32 dirLength = detoursFile.count;(void)dirLength;
+
 		detoursFile.Append(PathSeparator).Append(UBA_DETOURS_LIBRARY);
 		
 		#if PLATFORM_WINDOWS
 		char temp[1024];
-		sprintf_s(temp, sizeof_array(temp), "%ls", detoursFile.data);
-		m_detoursLibrary = temp;
+		detoursFile.Parse(temp, sizeof_array(temp));
+		m_detoursLibrary[IsArm] = temp;
+		if (GetCrossArchitectureDir(m_logger, detoursFile.Resize(dirLength), false))
+		{
+			detoursFile.Append(PathSeparator).Append(UBA_DETOURS_LIBRARY).Parse(temp, sizeof_array(temp));
+			m_detoursLibrary[!IsArm] = temp;
+		}
 		#else
-		m_detoursLibrary = detoursFile.data;
+		m_detoursLibrary[IsArm] = detoursFile.data;
 		#endif
 
 		if (!Create(info))
@@ -630,55 +652,51 @@ namespace uba
 		CasKey clientAgentKey = reader.ReadCasKey();
 		CasKey clientDetoursKey = reader.ReadCasKey();
 
-		CasKey& detoursBinaryKey = m_detoursBinaryKey[isClientArm];
+		CasKey detoursBinaryKey[2];
 		CasKey& agentBinaryKey = m_agentBinaryKey[isClientArm];
 
 		bool binAsVersion = clientAgentKey != CasKeyZero;
 		{
 			SCOPED_WRITE_LOCK(m_binKeysLock, lock);
-			if (detoursBinaryKey == CasKeyZero || (binAsVersion && agentBinaryKey == CasKeyZero))
+
+			detoursBinaryKey[0] = m_detoursBinaryKey[0];
+			detoursBinaryKey[1] = m_detoursBinaryKey[1];
+
+			StringBuffer<> detoursLib;
+			bool deferCreation = true;
+			bool fileIsCompressed = false;
+
+			// Handle x64 for both architectures and arm64 only for arm64 clients
+			for (u32 i=0; i!=(isClientArm ? 2u : 1u); ++i)
+			{
+				if (detoursBinaryKey[i] != CasKeyZero)
+					continue;
+				detoursLib.Clear().Append(m_detoursLibrary[i].c_str());
+				if (!m_storage.StoreCasFile(detoursBinaryKey[i], detoursLib.data, CasKeyZero, deferCreation, fileIsCompressed) || detoursBinaryKey[i] == CasKeyZero)
+					return m_logger.Error(TC("Failed to create cas for %s"), detoursLib.data);
+				m_detoursBinaryKey[i] = detoursBinaryKey[i];
+			}
+
+			if (binAsVersion && agentBinaryKey == CasKeyZero)
 			{
-				StringBuffer<> dir;
-				if (!GetDirectoryOfCurrentModule(m_logger, dir))
+				StringBuffer<> agentDir;
+				if (!GetDirectoryOfCurrentModule(m_logger, agentDir))
 					return false;
-		
 				if (IsArm != isClientArm)
+					if (!GetCrossArchitectureDir(m_logger, agentDir, true))
+						return false;
+				UBA_ASSERT(IsWindows);
+				agentDir.Append(PathSeparator).Append(UBA_AGENT_EXECUTABLE);
+				if (!m_storage.StoreCasFile(agentBinaryKey, agentDir.data, CasKeyZero, deferCreation, fileIsCompressed) || agentBinaryKey == CasKeyZero)
 				{
-					// UBT has the path win-x64/native or win-arm64/native
-					bool isUbtPath = dir.EndsWith(TC("native"));
-					if (isUbtPath)
-						dir.Resize(dir.count - 7); // Remove native and slash
-
-					const tchar* archPath[2] = { TC("x64"), TC("arm64") };
-					if (!dir.EndsWith(archPath[IsArm]))
-						return m_logger.Error(TC("Module dir is not under supported folder (%s) to be able to run cross architecutres, can't figure out matching x64/arm64 folder"), dir.data);
-					dir.Resize(dir.count - TStrlen(archPath[IsArm])).Append(archPath[isClientArm]);
-
-					if (isUbtPath)
-						dir.Append(PathSeparator).Append(TC("native"));
-				}
-
-				u64 dirCount = dir.count;
-				bool deferCreation = true;
-				bool fileIsCompressed = false;
-				if (binAsVersion && agentBinaryKey == CasKeyZero)
-				{
-					UBA_ASSERT(IsWindows);
-					dir.Append(PathSeparator).Append(UBA_AGENT_EXECUTABLE);
-					if (!m_storage.StoreCasFile(agentBinaryKey, dir.data, CasKeyZero, deferCreation, fileIsCompressed) || agentBinaryKey == CasKeyZero)
-					{
-						// This is hacky but uba binary is not copied by ubt anymore.
-						StringBuffer<> dir2;
-						if (!GetAlternativeUbaPath(m_logger, dir2, dir, IsWindows && isClientArm))
-							return false;
-						dir2.Append(UBA_AGENT_EXECUTABLE);
-						if (!m_storage.StoreCasFile(agentBinaryKey, dir2.data, CasKeyZero, deferCreation, fileIsCompressed) || agentBinaryKey == CasKeyZero)
-							return m_logger.Error(TC("Failed to create cas for %s"), dir2.data);
-					}
+					// This is hacky but uba binary is not copied by ubt anymore.
+					StringBuffer<> dir2;
+					if (!GetAlternativeUbaPath(m_logger, dir2, agentDir, IsWindows && isClientArm))
+						return false;
+					dir2.Append(UBA_AGENT_EXECUTABLE);
+					if (!m_storage.StoreCasFile(agentBinaryKey, dir2.data, CasKeyZero, deferCreation, fileIsCompressed) || agentBinaryKey == CasKeyZero)
+						return m_logger.Error(TC("Failed to create cas for %s"), dir2.data);
 				}
-				dir.Resize(dirCount).Append(PathSeparator).Append(UBA_DETOURS_LIBRARY);
-				if (!m_storage.StoreCasFile(detoursBinaryKey, dir.data, CasKeyZero, deferCreation, fileIsCompressed) || detoursBinaryKey == CasKeyZero)
-					return m_logger.Error(TC("Failed to create cas for %s"), dir.data);
 			}
 		}
 
@@ -701,7 +719,9 @@ namespace uba
 		{
 			writer.WriteString(disconnectResponse);
 			writer.WriteCasKey(agentBinaryKey);
-			writer.WriteCasKey(detoursBinaryKey);
+			writer.WriteCasKey(detoursBinaryKey[0]);
+			if (isClientArm)
+				writer.WriteCasKey(detoursBinaryKey[1]);
 			return true;
 		}
 
@@ -728,7 +748,9 @@ namespace uba
 		if (!InitializeNameToHashTable())
 			return false;
 
-		writer.WriteCasKey(detoursBinaryKey);
+		writer.WriteCasKey(m_detoursBinaryKey[0]);
+		if (isClientArm)
+			writer.WriteCasKey(m_detoursBinaryKey[1]);
 		writer.WriteBool(m_resetCas);
 		writer.WriteU32(sessionId);
 		writer.WriteU32(m_uiLanguage);
@@ -2268,10 +2290,17 @@ namespace uba
 	bool SessionServer::HandleGetSymbols(const ConnectionInfo& connectionInfo, BinaryReader& reader, BinaryWriter& writer)
 	{
 		u32 processId = reader.ReadU32();
+		bool isClientArm = reader.ReadBool();
 		ProcessHandle h(GetProcess(processId));
 		if (!h.IsValid())
 			return m_logger.Error(TC("Failed to find process for id %u when receiving update environment message"), processId);
 
+		if (IsArm != isClientArm)
+		{
+			writer.WriteString(TC("Can't resolve callstack on cross architectures"));
+			return true;
+		}
+
 		GetSymbols(*h.m_process, reader, writer);
 		return true;
 
diff --git a/Engine/Source/Programs/UnrealBuildAccelerator/Common/Public/UbaNetwork.h b/Engine/Source/Programs/UnrealBuildAccelerator/Common/Public/UbaNetwork.h
--- a/Engine/Source/Programs/UnrealBuildAccelerator/Common/Public/UbaNetwork.h
+++ b/Engine/Source/Programs/UnrealBuildAccelerator/Common/Public/UbaNetwork.h
@@ -16,7 +16,7 @@ namespace uba
 	
 	static constexpr u32 SystemNetworkVersion = 1339;
 	static constexpr u32 StorageNetworkVersion = 4;
-	static constexpr u32 SessionNetworkVersion = 41;
+	static constexpr u32 SessionNetworkVersion = 42;
 	static constexpr u32 CacheNetworkVersion = 5;
 
 	// Messages used over network between client/server (system, storage and session)
diff --git a/Engine/Source/Programs/UnrealBuildAccelerator/Common/Public/UbaSession.h b/Engine/Source/Programs/UnrealBuildAccelerator/Common/Public/UbaSession.h
--- a/Engine/Source/Programs/UnrealBuildAccelerator/Common/Public/UbaSession.h
+++ b/Engine/Source/Programs/UnrealBuildAccelerator/Common/Public/UbaSession.h
@@ -243,7 +243,7 @@ namespace uba
 		Atomic<u64> m_fileIndexCounter = 8000000000;
 
 		FileMappingAllocator m_processCommunicationAllocator;
-		std::string m_detoursLibrary;
+		std::string m_detoursLibrary[2]; // x64 and arm64
 
 		ReaderWriterLock m_processStatsLock;
 		ProcessStats m_processStats;
diff --git a/Engine/Source/Programs/UnrealBuildAccelerator/Core/Public/UbaBase.h b/Engine/Source/Programs/UnrealBuildAccelerator/Core/Public/UbaBase.h
--- a/Engine/Source/Programs/UnrealBuildAccelerator/Core/Public/UbaBase.h
+++ b/Engine/Source/Programs/UnrealBuildAccelerator/Core/Public/UbaBase.h
@@ -31,7 +31,7 @@ namespace uba
 	#define PERCENT_HS "%s"
 	#endif
 
-	#if defined(_M_ARM64)
+	#if PLATFORM_WINDOWS && defined(_M_ARM64) // For now we only care about windows arm..
 	constexpr bool IsArm = true;
 	#else
 	constexpr bool IsArm = false;
diff --git a/Engine/Source/Programs/UnrealBuildAccelerator/Detours/Private/Windows/UbaDetoursFileDecompressor.h b/Engine/Source/Programs/UnrealBuildAccelerator/Detours/Private/Windows/UbaDetoursFileDecompressor.h
--- a/Engine/Source/Programs/UnrealBuildAccelerator/Detours/Private/Windows/UbaDetoursFileDecompressor.h
+++ b/Engine/Source/Programs/UnrealBuildAccelerator/Detours/Private/Windows/UbaDetoursFileDecompressor.h
@@ -235,8 +235,8 @@ namespace uba
 			DEBUG_LOG_DETOURED(L"ReturnMemory", L"%s", fi.originalName);
 			SCOPED_WRITE_LOCK(g_mappedFileTable.m_memLookupLock, lock);
 			auto res = g_mappedFileTable.m_memLookup.erase(fi.fileMapMem);(void)res;
-			UBA_ASSERT(res);
 			lock.Leave();
+			//UBA_ASSERT(res);
 			//VirtualFree(fi.fileMapMem, 0, MEM_RELEASE);
 			_aligned_free(fi.fileMapMem);
 			fi.fileMapMem = nullptr;
@@ -386,8 +386,8 @@ namespace uba
 
 				{
 					SCOPED_WRITE_LOCK(g_mappedFileTable.m_memLookupLock, lock3);
-					auto insres = g_mappedFileTable.m_memLookup.try_emplace(destMem);
-					UBA_ASSERTF(insres.second, TC("Failed to add memory 0x%llx for %s"), u64(destMem), info.originalName);
+					auto insres = g_mappedFileTable.m_memLookup.try_emplace(info.fileMapMem);
+					UBA_ASSERTF(insres.second, TC("Failed to add memory 0x%llx for %s"), u64(info.fileMapMem), info.originalName);
 					insres.first->second = 1;
 				}
 
