[UBA]
* Fixed so x64 binaries can be detoured when running on windows arm64

diff --git a/Engine/Source/Programs/UnrealBuildAccelerator/Common/Private/UbaProcess.cpp b/Engine/Source/Programs/UnrealBuildAccelerator/Common/Private/UbaProcess.cpp
--- a/Engine/Source/Programs/UnrealBuildAccelerator/Common/Private/UbaProcess.cpp
+++ b/Engine/Source/Programs/UnrealBuildAccelerator/Common/Private/UbaProcess.cpp
@@ -90,6 +90,44 @@ namespace uba
 		TString currentString;
 	};
 
+	struct ImageInfo
+	{
+		CasKey key;
+		u64 fileSize = 0;
+		bool is64Bit = true;
+		bool isX64 = false;
+		bool isArm64 = false;
+	};
+
+	bool GetImageInfo(ImageInfo& out, Logger& logger, const tchar* application, bool calculateCas)
+	{
+		FileAccessor fa(logger, application);
+		if (!fa.OpenMemoryRead())
+			return false;
+		out.fileSize = fa.GetSize();
+		u8* data = fa.GetData();
+
+		if (calculateCas)
+		{
+			CasKeyHasher hasher;
+			hasher.Update(data, out.fileSize);
+			out.key = ToCasKey(hasher, false);
+		}
+
+		if (data[0] != 'M' || data[1] != 'Z')
+			out.is64Bit = false;
+		else
+		{
+			u32 offset = *(u32*)(data + 0x3c);
+			u32* signaturePos = (u32*)(data + offset);
+			out.is64Bit = *signaturePos == 0x00004550;
+			u16 machine = *(u16*)(signaturePos + 1);
+			out.isX64 = machine == 0x8664;
+			out.isArm64 = machine == 0xaa64;
+		}
+		return fa.Close();
+	}
+
 	ProcessImpl::ProcessImpl(Session& session, u32 id, ProcessImpl* parent, bool detourEnabled)
 	:	m_session(session)
 	,	m_parentProcess(parent)
@@ -253,46 +291,20 @@ namespace uba
 
 			if (m_messageCount == 0) // This is bad.. bad binaries?
 			{
-				bool imageIs64Bit = true;
-				bool imageIsX64 = false;
-				bool imageIsArm64 = false;
-				u64 fileSize = 0;
-				CasKey key;
-				
-				FileAccessor fa(m_session.m_logger, m_realApplication.c_str());
-				if (fa.OpenMemoryRead())
-				{
-					fileSize = fa.GetSize();
-					u8* data = fa.GetData();
-
-					CasKeyHasher hasher;
-					hasher.Update(data, fileSize);
-					key = ToCasKey(hasher, false);
-
-					if (data[0] != 'M' || data[1] != 'Z')
-						imageIs64Bit = false;
-					else
-					{
-						u32 offset = *(u32*)(data + 0x3c);
-						u32* signaturePos = (u32*)(data + offset);
-						imageIs64Bit = *signaturePos == 0x00004550;
-						u16 machine = *(u16*)(signaturePos + 1);
-						imageIsX64 = machine == 0x8664;
-						imageIsArm64 = machine == 0xaa64;
-					}
-				}
+				err.Append(TC("ERROR: Process did not start properly. "));
 
+				ImageInfo imageInfo;
 				bool machineIsArm = IsArm;
-				err.Append(TC("ERROR: Process did not start properly. "));
-				if (!imageIs64Bit)
+				if (!GetImageInfo(imageInfo, m_session.m_logger, m_realApplication.c_str(), true))
+					err.Appendf(TC("Failed to load %s"), m_realApplication.c_str());
+				else if (!imageInfo.is64Bit)
 					err.Append(TC("Doesn't seem to be a 64-bit executable"));
-				else if (!imageIsArm64 && !imageIsX64)
+				else if (!imageInfo.isArm64 && !imageInfo.isX64)
 					err.Append(TC("Unknown image architecture"));
-				else if (machineIsArm && imageIsX64)
-					err.Appendf(TC("Machine is arm64 and image is x64"));
-				else if (!machineIsArm && imageIsArm64)
+				else if (!machineIsArm && imageInfo.isArm64)
 					err.Appendf(TC("Machine is x64 and image is arm64"));
-				err.Appendf(TC(" (GetExitCodeProcess returned 0x%x (%s Size: %llu, CasKey: %s)"), exitCode, m_realApplication.c_str(), fileSize, CasKeyString(key).str);
+
+				err.Appendf(TC(" (GetExitCodeProcess returned 0x%x (%s Size: %llu, CasKey: %s)"), exitCode, m_realApplication.c_str(), imageInfo.fileSize, CasKeyString(imageInfo.key).str);
 			}
 
 			if (err.IsEmpty())
@@ -1434,9 +1446,33 @@ namespace uba
 		HANDLE readPipe = INVALID_HANDLE_VALUE;
 		auto readPipeGuard = MakeGuard([&]() { CloseHandle(readPipe); });
 
+		bool allowCustomAllocator = true;
+
 		if (!m_parentProcess)
 		{
-			const char* detoursLib = m_session.m_detoursLibrary[IsArm].c_str();
+			bool isArmApplication = IsArm;
+#if defined(_M_ARM64)
+			bool isX64Application;
+			{
+				SCOPED_WRITE_LOCK(m_session.m_isX64ApplicationLock, lock);
+				auto insres = m_session.m_isX64Application.try_emplace(m_realApplication);
+				if (insres.second)
+				{
+					ImageInfo imageInfo;
+					if (!GetImageInfo(imageInfo, m_session.m_logger, m_realApplication.c_str(), false))
+						return UBA_EXIT_CODE(19);
+					insres.first->second = imageInfo.isX64;
+				}
+				isX64Application = insres.first->second;
+			}
+			if (isX64Application)
+			{
+				isArmApplication = false;
+				allowCustomAllocator = false;
+			}
+#endif
+
+			const char* detoursLib = m_session.m_detoursLibrary[isArmApplication].c_str();
 			if (!*detoursLib)
 				detoursLib = UBA_DETOURS_LIBRARY_ANSI;
 
@@ -1666,7 +1702,7 @@ namespace uba
 			payload.suppressLogging = m_session.m_suppressLogging;
 			payload.isChild = m_parentProcess != nullptr;
 			payload.trackInputs = m_startInfo.trackInputs;
-			payload.useCustomAllocator = m_startInfo.useCustomAllocator && m_startInfo.rules->AllowMiMalloc();
+			payload.useCustomAllocator = allowCustomAllocator && m_startInfo.useCustomAllocator && m_startInfo.rules->AllowMiMalloc();
 			payload.isRunningWine = IsRunningWine();
 			payload.storeObjFilesCompressed = m_session.m_storeObjFilesCompressed;
 			payload.uiLanguage = m_startInfo.uiLanguage;
diff --git a/Engine/Source/Programs/UnrealBuildAccelerator/Common/Public/UbaSession.h b/Engine/Source/Programs/UnrealBuildAccelerator/Common/Public/UbaSession.h
--- a/Engine/Source/Programs/UnrealBuildAccelerator/Common/Public/UbaSession.h
+++ b/Engine/Source/Programs/UnrealBuildAccelerator/Common/Public/UbaSession.h
@@ -297,6 +297,11 @@ namespace uba
 		Logger* m_debugLogger = nullptr;
 		#endif
 
+		#if PLATFORM_WINDOWS && defined(_M_ARM64)
+		ReaderWriterLock m_isX64ApplicationLock;
+		UnorderedMap<TString, bool> m_isX64Application;
+		#endif
+
 		friend class ProcessImpl;
 	};
 
diff --git a/Engine/Source/Programs/UnrealBuildAccelerator/Detours/Private/Windows/UbaDetoursFunctionsNtDll.inl b/Engine/Source/Programs/UnrealBuildAccelerator/Detours/Private/Windows/UbaDetoursFunctionsNtDll.inl
--- a/Engine/Source/Programs/UnrealBuildAccelerator/Detours/Private/Windows/UbaDetoursFunctionsNtDll.inl
+++ b/Engine/Source/Programs/UnrealBuildAccelerator/Detours/Private/Windows/UbaDetoursFunctionsNtDll.inl
@@ -559,7 +559,7 @@ BOOLEAN Detoured_RtlFreeHeap(PVOID HeapHandle, ULONG Flags, PVOID BaseAddress)
 NTSTATUS Detoured_RtlAnsiStringToUnicodeString(PUNICODE_STRING DestinationString, PCANSI_STRING SourceString, BOOLEAN AllocateDestinationString)
 {
 #if UBA_USE_MIMALLOC
-	if (AllocateDestinationString)
+	if (AllocateDestinationString && g_useMiMalloc)
 	{
 		DestinationString->MaximumLength = SourceString->MaximumLength * 2;
 		DestinationString->Buffer = (wchar_t*)mi_malloc(DestinationString->MaximumLength);
@@ -573,7 +573,7 @@ NTSTATUS Detoured_RtlAnsiStringToUnicodeString(PUNICODE_STRING DestinationString
 NTSTATUS Detoured_RtlUnicodeStringToAnsiString(PANSI_STRING DestinationString, PCUNICODE_STRING SourceString, BOOLEAN AllocateDestinationString)
 {
 #if UBA_USE_MIMALLOC
-	if (AllocateDestinationString)
+	if (AllocateDestinationString && g_useMiMalloc)
 	{
 		DestinationString->MaximumLength = SourceString->MaximumLength / 2;
 		DestinationString->Buffer = (char*)mi_malloc(DestinationString->MaximumLength);
