[UBA]
* Changed so directory table memory is reserved and then committed based on need instead of commit everything up front
* Bumped directory table size to 128mb (crazy)...

diff --git a/Engine/Source/Programs/UnrealBuildAccelerator/Common/Private/UbaSession.cpp b/Engine/Source/Programs/UnrealBuildAccelerator/Common/Private/UbaSession.cpp
--- a/Engine/Source/Programs/UnrealBuildAccelerator/Common/Private/UbaSession.cpp
+++ b/Engine/Source/Programs/UnrealBuildAccelerator/Common/Private/UbaSession.cpp
@@ -332,6 +332,7 @@ namespace uba
 		u32 tableOffset;
 
 		SCOPED_WRITE_LOCK(dirTable.m_memoryLock, memoryLock);
+		EnsureDirectoryTableMemory();
 		u32 writePos = dirTable.m_memorySize;
 		BinaryWriter tableWriter(m_directoryTableMem + dirTable.m_memorySize);
 
@@ -359,7 +360,6 @@ namespace uba
 		u32 filesOffset = writePos + u32(tableWriter.GetPosition());
 		tableWriter.WriteBytes(memoryBlock.data(), written);
 		dirTable.m_memorySize += u32(tableWriter.GetPosition());
-		CheckDirectoryTableOverflow(dirTable);
 		
 		memoryLock.Leave();
 
@@ -857,6 +857,7 @@ namespace uba
 
 
 		SCOPED_WRITE_LOCK(dirTable.m_memoryLock, memoryLock);
+		EnsureDirectoryTableMemory();
 		u8* startPos = m_directoryTableMem + dirTable.m_memorySize;
 		BinaryWriter writer(startPos);
 		writer.Write7BitEncoded(written); // Storage size
@@ -865,7 +866,6 @@ namespace uba
 		writer.WriteBytes(temp, written);
 		dir.tableOffset = dirTable.m_memorySize + tableOffset;
 		dirTable.m_memorySize += u32(writer.GetPosition());
-		CheckDirectoryTableOverflow(dirTable);
 		return true;
 	}
 
@@ -933,6 +933,7 @@ namespace uba
 		#endif
 
 		SCOPED_WRITE_LOCK(dirTable.m_memoryLock, memoryLock);
+		EnsureDirectoryTableMemory();
 		u8* startPos = m_directoryTableMem + dirTable.m_memorySize;
 		BinaryWriter writer(startPos);
 		writer.Write7BitEncoded(written); // Storage size
@@ -940,7 +941,6 @@ namespace uba
 		writer.WriteBytes(temp, written);
 		dir.tableOffset = dirTable.m_memorySize + tableOffset;
 		dirTable.m_memorySize += u32(writer.GetPosition());
-		CheckDirectoryTableOverflow(dirTable);
 		return dirTable.m_memorySize;
 	}
 
@@ -959,13 +959,6 @@ namespace uba
 		return WithVfs(rootsHandle, !entry.locals.empty());
 	}
 
-	void Session::CheckDirectoryTableOverflow(DirectoryTable& dirTable)
-	{
-		if (dirTable.m_memorySize < DirTableMemSize - 1024*1024)
-			return;
-		static bool called = m_logger.Error(TC("Directory table overflow. DirTableMemSize need to be increased (Size: %llu)"), dirTable.m_memorySize);
-	}
-
 	bool Session::CopyImports(Vector<BinaryModule>& out, const tchar* library, tchar* applicationDir, tchar* applicationDirEnd, UnorderedSet<TString>& handledImports, const char* const* loaderPaths)
 	{
 		if (!handledImports.insert(library).second)
@@ -1152,7 +1145,7 @@ namespace uba
 		m_fileMappingTableMem = MapViewOfFile(m_logger, m_fileMappingTableHandle, FILE_MAP_WRITE, 0, FileMappingTableMemSize);
 		UBA_ASSERT(m_fileMappingTableMem);
 
-		m_directoryTableHandle = uba::CreateMemoryMappingW(m_logger, PAGE_READWRITE, DirTableMemSize, nullptr, TC("DirMappings"));
+		m_directoryTableHandle = uba::CreateMemoryMappingW(m_logger, PAGE_READWRITE|SEC_RESERVE, DirTableMemSize, nullptr, TC("DirMappings"));
 		UBA_ASSERT(m_directoryTableHandle.IsValid());
 		m_directoryTableMem = MapViewOfFile(m_logger, m_directoryTableHandle, FILE_MAP_WRITE, 0, DirTableMemSize);
 		UBA_ASSERT(m_directoryTableMem);
@@ -2872,6 +2865,24 @@ namespace uba
 		return true;
 	}
 
+	void Session::EnsureDirectoryTableMemory()
+	{
+		auto& dirTable = m_directoryTable;
+		u64 needed = dirTable.m_memorySize + 512*1024;
+		if (needed < m_directoryTableMemCommitted)
+			return;
+
+		if (dirTable.m_memorySize + 1024*1024 > DirTableMemSize)
+		{
+			static bool called = m_logger.Error(TC("Directory table overflow. DirTableMemSize need to be increased (Size: %llu)"), dirTable.m_memorySize);
+		}
+
+		u8* address = m_directoryTableMem + m_directoryTableMemCommitted;
+		u64 toCommit = 1024*1024;
+		if (!MapViewCommit(address, toCommit))
+			m_logger.Error(TC("Failed to commit memory for directory table (Committed: %llu) (%s)"), m_directoryTableMemCommitted, LastErrorToText().data);
+		m_directoryTableMemCommitted += toCommit;
+	}
 
 	void Session::GetSystemInfo(StringBufferBase& out)
 	{
diff --git a/Engine/Source/Programs/UnrealBuildAccelerator/Common/Private/UbaSessionClient.cpp b/Engine/Source/Programs/UnrealBuildAccelerator/Common/Private/UbaSessionClient.cpp
--- a/Engine/Source/Programs/UnrealBuildAccelerator/Common/Private/UbaSessionClient.cpp
+++ b/Engine/Source/Programs/UnrealBuildAccelerator/Common/Private/UbaSessionClient.cpp
@@ -1058,6 +1058,8 @@ namespace uba
 			if (m_directoryTableError)
 				return false;
 
+			EnsureDirectoryTableMemory();
+
 			if (toRead == 0)
 			{
 				// We wrote to lastWriteEnd and now we got an empty message where readPos is the same..
diff --git a/Engine/Source/Programs/UnrealBuildAccelerator/Common/Public/UbaSession.h b/Engine/Source/Programs/UnrealBuildAccelerator/Common/Public/UbaSession.h
--- a/Engine/Source/Programs/UnrealBuildAccelerator/Common/Public/UbaSession.h
+++ b/Engine/Source/Programs/UnrealBuildAccelerator/Common/Public/UbaSession.h
@@ -130,7 +130,6 @@ namespace uba
 
 		bool RegisterCreateFileForWrite(StringKey fileNameKey, const StringView& fileName, bool registerRealFile, u64 fileSize = 0, u64 lastWriteTime = 0, bool invalidateStorage = true);
 		u32 RegisterDeleteFile(StringKey fileNameKey, const StringView& fileName);
-		void CheckDirectoryTableOverflow(DirectoryTable& dirTable);
 
 		virtual bool PrepareProcess(ProcessStartInfoHolder& startInfo, bool isChild, StringBufferBase& outRealApplication, const tchar*& outRealWorkingDir);
 		virtual void* GetProcessEnvironmentVariables();
@@ -170,6 +169,7 @@ namespace uba
 		void WriteDirectoryEntriesRecursive(const StringKey& dirKey, StringView dirPath, u32& outTableOffset);
 		bool CopyImports(Vector<BinaryModule>& out, const tchar* library, tchar* applicationDir, tchar* applicationDirEnd, UnorderedSet<TString>& handledImports, const char* const* loaderPaths);
 		bool CreateProcessJobObject();
+		void EnsureDirectoryTableMemory();
 		void GetSystemInfo(StringBufferBase& out);
 		bool GetMemoryInfo(u64& outAvailable, u64& outTotal);
 		void WriteSummary(BinaryWriter& writer, const Function<void(Logger& logger)>& summaryFunc);
@@ -224,6 +224,7 @@ namespace uba
 
 		FileMappingHandle m_directoryTableHandle;
 		u8* m_directoryTableMem;
+		u64 m_directoryTableMemCommitted = 0;
 		DirectoryTable m_directoryTable;
 	
 		FileMappingHandle m_fileMappingTableHandle;
