[UBA]
* Fixed so WriteCallstackInfo returns true when it succeeds, not false :-)

diff --git a/Engine/Source/Programs/UnrealBuildAccelerator/Core/Private/UbaPlatform.cpp b/Engine/Source/Programs/UnrealBuildAccelerator/Core/Private/UbaPlatform.cpp
--- a/Engine/Source/Programs/UnrealBuildAccelerator/Core/Private/UbaPlatform.cpp
+++ b/Engine/Source/Programs/UnrealBuildAccelerator/Core/Private/UbaPlatform.cpp
@@ -366,7 +366,7 @@ namespace uba
 			out.Write7BitEncoded(rec->size);
 			out.WriteString(rec->name);
 		}
-		return false;
+		return true;
 #endif
 	}
 
