[UBA]
* Fixed so trace reserve size can be set in config file

diff --git a/Engine/Source/Programs/UnrealBuildAccelerator/Common/Private/UbaSession.cpp b/Engine/Source/Programs/UnrealBuildAccelerator/Common/Private/UbaSession.cpp
--- a/Engine/Source/Programs/UnrealBuildAccelerator/Common/Private/UbaSession.cpp
+++ b/Engine/Source/Programs/UnrealBuildAccelerator/Common/Private/UbaSession.cpp
@@ -228,6 +228,7 @@ namespace uba
 		table.GetValueAsBool(storeObjFilesCompressed, TC("StoreObjFilesCompressed"));
 		table.GetValueAsBool(allowLocalDetour, TC("AllowLocalDetour"));
 		table.GetValueAsBool(extractObjFilesSymbols, TC("ExtractObjFilesSymbols"));
+		table.GetValueAsU32(traceReserveSizeMb, TC("TraceReserveSizeMb"));
 	}
 
 	void Session::AddEnvironmentVariableNoLock(const tchar* key, const tchar* value)
@@ -1188,7 +1189,7 @@ namespace uba
 		}
 
 		if (!traceName.IsEmpty())
-			StartTrace(IsWindows ? traceName.data : nullptr); // non-windows named shared memory not implemented (only needed for UbaVisualizer which you can't run on linux either way)
+			StartTrace(IsWindows ? traceName.data : nullptr, info.traceReserveSizeMb); // non-windows named shared memory not implemented (only needed for UbaVisualizer which you can't run on linux either way)
 
 		#if PLATFORM_WINDOWS
 		if (info.launchVisualizer)
@@ -1835,14 +1836,18 @@ namespace uba
 		stats.Print(m_logger);
 	}
 
-	void Session::StartTrace(const tchar* traceName)
+	void Session::StartTrace(const tchar* traceName, u32 traceReserveSizeMb)
 	{
 		if (traceName)
 			LoggerWithWriter(m_logger.m_writer).Info(TC("---- Starting trace: %s ----"), traceName);
 		else
 			LoggerWithWriter(m_logger.m_writer).Info(TC("---- Starting trace ----"));
 
-		if (!m_trace.StartWrite(traceName, m_detailedTrace ? 512*1024*1024 : 128*1024*1024))
+		u64 traceReserveSize = u64(traceReserveSizeMb)*1024*1024;
+		if (m_detailedTrace)
+			traceReserveSize *= 2;
+
+		if (!m_trace.StartWrite(traceName, traceReserveSize))
 			return;
 
 		tchar buf[256];
diff --git a/Engine/Source/Programs/UnrealBuildAccelerator/Common/Private/UbaSessionClient.cpp b/Engine/Source/Programs/UnrealBuildAccelerator/Common/Private/UbaSessionClient.cpp
--- a/Engine/Source/Programs/UnrealBuildAccelerator/Common/Private/UbaSessionClient.cpp
+++ b/Engine/Source/Programs/UnrealBuildAccelerator/Common/Private/UbaSessionClient.cpp
@@ -1280,7 +1280,7 @@ namespace uba
 			{
 				m_client.SetWorkTracker(&m_trace);
 			}
-			StartTrace(nullptr);
+			StartTrace(nullptr, 256);
 		}
 
 		BuildEnvironmentVariables(reader);
diff --git a/Engine/Source/Programs/UnrealBuildAccelerator/Common/Public/UbaSession.h b/Engine/Source/Programs/UnrealBuildAccelerator/Common/Public/UbaSession.h
--- a/Engine/Source/Programs/UnrealBuildAccelerator/Common/Public/UbaSession.h
+++ b/Engine/Source/Programs/UnrealBuildAccelerator/Common/Public/UbaSession.h
@@ -105,7 +105,7 @@ namespace uba
 		void ProcessExited(ProcessImpl& process, u64 executionTime);
 		void FlushDeadProcesses();
 		void PrintProcessStats(ProcessStats& stats, const tchar* logName);
-		void StartTrace(const tchar* traceName);
+		void StartTrace(const tchar* traceName, u32 traceReserveSizeMb);
 		bool StopTrace(const tchar* writeFile);
 		void StopTraceThread();
 		u32 GetFileMappingSize();
diff --git a/Engine/Source/Programs/UnrealBuildAccelerator/Common/Public/UbaSessionCreateInfo.h b/Engine/Source/Programs/UnrealBuildAccelerator/Common/Public/UbaSessionCreateInfo.h
--- a/Engine/Source/Programs/UnrealBuildAccelerator/Common/Public/UbaSessionCreateInfo.h
+++ b/Engine/Source/Programs/UnrealBuildAccelerator/Common/Public/UbaSessionCreateInfo.h
@@ -39,5 +39,6 @@ namespace uba
 		bool extractObjFilesSymbols = false;	// Will extract import/export symbols and clear export table from obj file.
 		u64 deleteSessionsOlderThanSeconds = 12 * 60 * 60; // Delete session folders older than 12 hours by default . Set to 0 to not delete or 1 to delete all
 		u64 keepOutputFileMemoryMapsThreshold = 256 * 1024; // If allowMemoryMaps is true, output files will be kept in memory if smaller than this size
+		u32 traceReserveSizeMb = 256;			// Memory reserved for trace file (in mb)
 	};
 }
