[UBA]
* Moved Read/Write long string into BinaryReader/Writer

diff --git a/Engine/Source/Programs/UnrealBuildAccelerator/Common/Private/UbaProcess.cpp b/Engine/Source/Programs/UnrealBuildAccelerator/Common/Private/UbaProcess.cpp
--- a/Engine/Source/Programs/UnrealBuildAccelerator/Common/Private/UbaProcess.cpp
+++ b/Engine/Source/Programs/UnrealBuildAccelerator/Common/Private/UbaProcess.cpp
@@ -978,7 +978,7 @@ namespace uba
 		#endif
 
 		TString applicationStr = reader.ReadString();
-		TString commandLineWithoutApplication = ReadLongString(reader);
+		TString commandLineWithoutApplication = reader.ReadLongString();
 
 		UBA_ASSERTF(!applicationStr.empty() && (applicationStr[0] != '\"' || applicationStr[1] != '\"'), TC("Invalid application name: %s"), applicationStr.c_str());
 
diff --git a/Engine/Source/Programs/UnrealBuildAccelerator/Core/Private/UbaPlatform.cpp b/Engine/Source/Programs/UnrealBuildAccelerator/Core/Private/UbaPlatform.cpp
--- a/Engine/Source/Programs/UnrealBuildAccelerator/Core/Private/UbaPlatform.cpp
+++ b/Engine/Source/Programs/UnrealBuildAccelerator/Core/Private/UbaPlatform.cpp
@@ -2,6 +2,7 @@
 
 #include "UbaPlatform.h"
 #include "UbaProcessStats.h"
+#include <oodle2.h>
 
 #if PLATFORM_WINDOWS
 #include <conio.h>
@@ -592,4 +593,41 @@ namespace uba
 		#endif
 	}
 #endif
+
+	void BinaryWriter::WriteLongString(const StringView& str, u32 threshold)
+	{
+		if (str.count < threshold)
+		{
+			WriteBool(true);
+			WriteString(str);
+			return;
+		}
+			
+		WriteBool(false);
+		u32 uncompressedSize = str.count * sizeof(tchar);
+		u64 needed = (u64)OodleLZ_GetCompressedBufferSizeNeeded(OodleLZ_Compressor_Kraken, uncompressedSize);
+		Vector<u8> mem;
+		mem.resize(needed);
+		u64 compressedSize = OodleLZ_Compress(OodleLZ_Compressor_Kraken, str.data, (OO_SINTa)uncompressedSize, mem.data(), OodleLZ_CompressionLevel_Normal);
+		UBA_ASSERT(compressedSize != OODLELZ_FAILED);
+		Write7BitEncoded(str.count);
+		Write7BitEncoded(compressedSize);
+		WriteBytes(mem.data(), compressedSize);
+	}
+
+	TString BinaryReader::ReadLongString()
+	{
+		if (ReadBool())
+			return ReadString();
+
+		u64 stringLength = Read7BitEncoded();
+		u64 uncompressedSize = stringLength * sizeof(tchar);
+		u64 compressedSize = Read7BitEncoded();
+		TString str;
+		str.resize(stringLength);
+		OO_SINTa decompLen = OodleLZ_Decompress(GetPositionData(), (OO_SINTa)compressedSize, str.data(), (OO_SINTa)uncompressedSize);
+		UBA_ASSERT(decompLen == (OO_SINTa)uncompressedSize);(void)decompLen;
+		Skip(compressedSize);
+		return str;
+	}
 }
diff --git a/Engine/Source/Programs/UnrealBuildAccelerator/Core/Public/UbaBinaryReaderWriter.h b/Engine/Source/Programs/UnrealBuildAccelerator/Core/Public/UbaBinaryReaderWriter.h
--- a/Engine/Source/Programs/UnrealBuildAccelerator/Core/Public/UbaBinaryReaderWriter.h
+++ b/Engine/Source/Programs/UnrealBuildAccelerator/Core/Public/UbaBinaryReaderWriter.h
@@ -29,6 +29,8 @@ namespace uba
 		template<typename CharType>
 		inline void WriteUtf8String(const CharType* str, u64 len);
 
+		void WriteLongString(const StringView& str, u32 threshold = 10000);
+
 		inline u8* AllocWrite(u64 bytes);
 		inline u64 GetPosition() const { return u64(m_pos - m_begin); }
 		inline u64 GetCapacityLeft() const { return u64(m_end - m_pos); }
@@ -69,6 +71,7 @@ namespace uba
 		inline bool ReadBool() { return ReadByte() != 0; }
 		inline u64 Read7BitEncoded();
 		template<typename CharType> inline CharType ReadUtf8Char();
+		TString ReadLongString();
 		inline u32 PeekU32();
 		inline u64 PeekU64();
 		inline void Skip(u64 size) { m_pos += size; }
diff --git a/Engine/Source/Programs/UnrealBuildAccelerator/Core/Public/UbaProcessUtils.h b/Engine/Source/Programs/UnrealBuildAccelerator/Core/Public/UbaProcessUtils.h
--- a/Engine/Source/Programs/UnrealBuildAccelerator/Core/Public/UbaProcessUtils.h
+++ b/Engine/Source/Programs/UnrealBuildAccelerator/Core/Public/UbaProcessUtils.h
@@ -2,10 +2,7 @@
 
 #pragma once
 
-#include "UbaBinaryReaderWriter.h"
 #include "UbaPlatform.h"
-#include "UbaStringBuffer.h"
-#include <oodle2.h>
 
 namespace uba
 {
@@ -105,41 +102,4 @@ namespace uba
 	{
 		return ParseArguments(arguments, TStrlen(arguments), argumentFunc);
 	}
-
-	inline void WriteLongString(BinaryWriter& writer, const StringView& str)
-	{
-		if (str.count < 10000)
-		{
-			writer.WriteBool(true);
-			writer.WriteString(str);
-			return;
-		}
-			
-		writer.WriteBool(false);
-		u32 uncompressedSize = str.count * sizeof(tchar);
-		u64 needed = (u64)OodleLZ_GetCompressedBufferSizeNeeded(OodleLZ_Compressor_Kraken, uncompressedSize);
-		Vector<u8> mem;
-		mem.resize(needed);
-		u64 compressedSize = OodleLZ_Compress(OodleLZ_Compressor_Kraken, str.data, (OO_SINTa)uncompressedSize, mem.data(), OodleLZ_CompressionLevel_Normal);
-		UBA_ASSERT(compressedSize != OODLELZ_FAILED);
-		writer.Write7BitEncoded(str.count);
-		writer.Write7BitEncoded(compressedSize);
-		writer.WriteBytes(mem.data(), compressedSize);
-	}
-
-	inline TString ReadLongString(BinaryReader& reader)
-	{
-		if (reader.ReadBool())
-			return reader.ReadString();
-
-		u64 stringLength = reader.Read7BitEncoded();
-		u64 uncompressedSize = stringLength * sizeof(tchar);
-		u64 compressedSize = reader.Read7BitEncoded();
-		TString str;
-		str.resize(stringLength);
-		OO_SINTa decompLen = OodleLZ_Decompress(reader.GetPositionData(), (OO_SINTa)compressedSize, str.data(), (OO_SINTa)uncompressedSize);
-		UBA_ASSERT(decompLen == (OO_SINTa)uncompressedSize);(void)decompLen;
-		reader.Skip(compressedSize);
-		return str;
-	}
 }
diff --git a/Engine/Source/Programs/UnrealBuildAccelerator/Detours/Private/PosixOS/UbaDetoursFunctionsPosix.cpp b/Engine/Source/Programs/UnrealBuildAccelerator/Detours/Private/PosixOS/UbaDetoursFunctionsPosix.cpp
--- a/Engine/Source/Programs/UnrealBuildAccelerator/Detours/Private/PosixOS/UbaDetoursFunctionsPosix.cpp
+++ b/Engine/Source/Programs/UnrealBuildAccelerator/Detours/Private/PosixOS/UbaDetoursFunctionsPosix.cpp
@@ -2011,7 +2011,7 @@ int shared_posix_spawn(pid_t* pid, const char* path, const posix_spawn_file_acti
 		BinaryWriter writer;
 		writer.WriteByte(MessageType_CreateProcess);
 		writer.WriteString(path); // Application
-		WriteLongString(writer, cmdLineWithoutApplication);
+		writer.WriteLongString(cmdLineWithoutApplication);
 		writer.WriteString(pwd); // Current dir
 		writer.WriteBool(false); // Start suspended
 		writer.WriteBool(true); // Is child
diff --git a/Engine/Source/Programs/UnrealBuildAccelerator/Detours/Private/Windows/UbaDetoursFunctionsKernelBase.inl b/Engine/Source/Programs/UnrealBuildAccelerator/Detours/Private/Windows/UbaDetoursFunctionsKernelBase.inl
--- a/Engine/Source/Programs/UnrealBuildAccelerator/Detours/Private/Windows/UbaDetoursFunctionsKernelBase.inl
+++ b/Engine/Source/Programs/UnrealBuildAccelerator/Detours/Private/Windows/UbaDetoursFunctionsKernelBase.inl
@@ -2861,7 +2861,7 @@ BOOL Detoured_CreateProcessW(LPCWSTR lpApplicationName, LPWSTR lpCommandLine, LP
 		BinaryWriter writer;
 		writer.WriteByte(MessageType_CreateProcess);
 		writer.WriteString(application);
-		WriteLongString(writer, commandLineWithoutApplication);
+		writer.WriteLongString(commandLineWithoutApplication);
 		writer.WriteString(workingDir);
 		writer.WriteBool(startSuspended);
 		writer.WriteBool(isChild);
