[UBA]
* Fixed so breadcrumbs can be copied from process

diff --git a/Engine/Source/Programs/UnrealBuildAccelerator/Visualizer/Private/UbaVisualizer.cpp b/Engine/Source/Programs/UnrealBuildAccelerator/Visualizer/Private/UbaVisualizer.cpp
--- a/Engine/Source/Programs/UnrealBuildAccelerator/Visualizer/Private/UbaVisualizer.cpp
+++ b/Engine/Source/Programs/UnrealBuildAccelerator/Visualizer/Private/UbaVisualizer.cpp
@@ -21,6 +21,7 @@ namespace uba
 		Popup_CopySessionInfo = 3,
 		Popup_CopyProcessInfo,
 		Popup_CopyProcessLog,
+		Popup_CopyProcessBreadcrumbs,
 		Popup_Replay,
 		Popup_Pause,
 		Popup_Play,
@@ -3477,6 +3478,8 @@ namespace uba
 				AppendMenuW(hMenu, MF_STRING, Popup_CopyProcessInfo, L"&Copy Process Info");
 				if (!process.logLines.empty())
 					AppendMenuW(hMenu, MF_STRING, Popup_CopyProcessLog, L"Copy Process &Log");
+				if (!process.breadcrumbs.empty())
+					AppendMenuW(hMenu, MF_STRING, Popup_CopyProcessBreadcrumbs, L"Copy Process &Breadcrumbs");
 				AppendMenuW(hMenu, MF_SEPARATOR, 0, NULL);
 			}
 
@@ -3645,6 +3648,12 @@ namespace uba
 				CopyTextToClipboard(str);
 				break;
 			}
+			case Popup_CopyProcessBreadcrumbs:
+			{
+				const TraceView::Process& process = m_traceView.GetProcess(m_processSelectedLocation);
+				CopyTextToClipboard(process.breadcrumbs);
+				break;
+			}
 			}
 
 			DestroyMenu(hMenu);
