[UBA]
* Fixed crash in UbaVisualizer when hovering over processes that were being reused

diff --git a/Engine/Source/Programs/UnrealBuildAccelerator/Visualizer/Private/UbaVisualizer.cpp b/Engine/Source/Programs/UnrealBuildAccelerator/Visualizer/Private/UbaVisualizer.cpp
--- a/Engine/Source/Programs/UnrealBuildAccelerator/Visualizer/Private/UbaVisualizer.cpp
+++ b/Engine/Source/Programs/UnrealBuildAccelerator/Visualizer/Private/UbaVisualizer.cpp
@@ -1845,7 +1845,7 @@ namespace uba
 
 					if (reader.GetLeft())
 					{
-						if (process.isRemote || m_traceView.version >= 36)
+						if (process.isRemote || m_traceView.version < 35)
 							sessionStats.Read(reader, m_traceView.version);
 						storageStats.Read(reader, m_traceView.version);
 						kernelStats.Read(reader, m_traceView.version);
@@ -2839,7 +2839,7 @@ namespace uba
 			processStats.Read(reader, m_traceView.version);
 			if (reader.GetLeft())
 			{
-				if (process.isRemote || m_traceView.version >= 36)
+				if (process.isRemote || m_traceView.version < 35)
 					sessionStats.Read(reader, m_traceView.version);
 				storageStats.Read(reader, m_traceView.version);
 				kernelStats.Read(reader, m_traceView.version);
