[UBA]
* Changed so breadcrumbs are compressed when written to trace stream

diff --git a/Engine/Source/Programs/UnrealBuildAccelerator/Common/Private/UbaTrace.cpp b/Engine/Source/Programs/UnrealBuildAccelerator/Common/Private/UbaTrace.cpp
--- a/Engine/Source/Programs/UnrealBuildAccelerator/Common/Private/UbaTrace.cpp
+++ b/Engine/Source/Programs/UnrealBuildAccelerator/Common/Private/UbaTrace.cpp
@@ -68,6 +68,7 @@ namespace uba
 
 		if (!m_memoryBegin)
 		{
+			m_logger.Warning(TC("Failed to map view of trace mapping (%s)"), LastErrorToText().data);
 			CloseFileMapping(m_logger, m_memoryHandle, TC("Trace"));
 			m_memoryHandle = {};
 			return false;
@@ -253,7 +254,7 @@ namespace uba
 		writer.WriteU32(sessionId);
 		writer.WriteU32(processId);
 		writer.WriteString(description);
-		writer.WriteString(breadcrumbs);
+		writer.WriteLongString(breadcrumbs, 20);
 	}
 
 	void Trace::ProcessEnvironmentUpdated(u32 processId, const StringView& reason, const u8* data, u64 dataSize, const StringView& breadcrumbs)
@@ -262,7 +263,7 @@ namespace uba
 		writer.WriteU32(processId);
 		writer.WriteString(reason);
 		writer.WriteBytes(data, dataSize);
-		writer.WriteString(breadcrumbs);
+		writer.WriteLongString(breadcrumbs, 20);
 	}
 
 	void Trace::ProcessExited(u32 processId, u32 exitCode, const u8* data, u64 dataSize, const Vector<ProcessLogLine>& logLines)
@@ -296,7 +297,7 @@ namespace uba
 	{
 		BEGIN_TRACE_ENTRY(ProcessBreadcrumbs);
 		writer.WriteU32(processId);
-		writer.WriteString(breadcrumbs);
+		writer.WriteLongString(breadcrumbs, 20);
 		writer.WriteBool(deleteOld);
 	}
 
diff --git a/Engine/Source/Programs/UnrealBuildAccelerator/Common/Public/UbaTrace.h b/Engine/Source/Programs/UnrealBuildAccelerator/Common/Public/UbaTrace.h
--- a/Engine/Source/Programs/UnrealBuildAccelerator/Common/Public/UbaTrace.h
+++ b/Engine/Source/Programs/UnrealBuildAccelerator/Common/Public/UbaTrace.h
@@ -70,7 +70,7 @@ namespace uba
 	using Color = u32;
 	inline Color toColor(u8 r, u8 g, u8 b) { return (r << 16) + (g << 8) + b; }
 
-	static constexpr u32 TraceVersion = 37;
+	static constexpr u32 TraceVersion = 38;
 	static constexpr u32 TraceReadCompatibilityVersion = 6;
 
 	class Trace : public WorkTracker
diff --git a/Engine/Source/Programs/UnrealBuildAccelerator/Visualizer/Private/UbaTraceReader.cpp b/Engine/Source/Programs/UnrealBuildAccelerator/Visualizer/Private/UbaTraceReader.cpp
--- a/Engine/Source/Programs/UnrealBuildAccelerator/Visualizer/Private/UbaTraceReader.cpp
+++ b/Engine/Source/Programs/UnrealBuildAccelerator/Visualizer/Private/UbaTraceReader.cpp
@@ -543,7 +543,12 @@ namespace uba
 
 			TString breadcrumbs;
 			if (out.version >= 35)
-				breadcrumbs = reader.ReadString();
+			{
+				if (out.version < 38)
+					breadcrumbs = reader.ReadString();
+				else
+					breadcrumbs = reader.ReadLongString();
+			}
 
 			if (out.version < 15)
 			{
@@ -672,7 +677,12 @@ namespace uba
 
 			TString breadcrumbs;
 			if (out.version >= 35)
-				breadcrumbs = reader.ReadString();
+			{
+				if (out.version < 38)
+					breadcrumbs = reader.ReadString();
+				else
+					breadcrumbs = reader.ReadLongString();
+			}
 
 			process.exitCode = 0u;
 			process.stop = time;
@@ -1038,7 +1048,11 @@ namespace uba
 		case TraceType_ProcessBreadcrumbs:
 		{
 			u32 processId = reader.ReadU32();
-			TString breadcrumbs = reader.ReadString();
+			TString breadcrumbs;
+			if (out.version < 38)
+				breadcrumbs = reader.ReadString();
+			else
+				breadcrumbs = reader.ReadLongString();
 			bool deleteOld = reader.ReadBool();
 
 			for (auto& session : out.sessions)
