[UBA]
* Fixed logic that commit memory for directory table. It didn't handle really big writes

diff --git a/Engine/Source/Programs/UnrealBuildAccelerator/Common/Private/UbaSession.cpp b/Engine/Source/Programs/UnrealBuildAccelerator/Common/Private/UbaSession.cpp
--- a/Engine/Source/Programs/UnrealBuildAccelerator/Common/Private/UbaSession.cpp
+++ b/Engine/Source/Programs/UnrealBuildAccelerator/Common/Private/UbaSession.cpp
@@ -252,8 +252,6 @@ namespace uba
 		u32 dirAttributes = 0;
 		u64 fileIndex = 0;
 
-		Vector<u8> memoryBlock;
-		memoryBlock.resize(4096);
 		u64 written = 0;
 		u32 itemCount = 0;
 
@@ -271,10 +269,17 @@ namespace uba
 		#if UBA_DEBUG_TRACK_DIR
 		m_debugLogger->BeginScope();
 		auto dg = MakeGuard([&]() { m_debugLogger->EndScope(); });
-		m_debugLogger->Info(TC("TRACKDIR %s\n"), StringBuffer<>(dirPath).data);
+		StringBuffer<> str;
+		str.Append(TC("TRACKDIR "));
+		if (isRefresh)
+			str.Append(TC("(Refresh) "));
+		str.Append(dirPath).Append('\n');
+		m_debugLogger->Log(LogEntryType_Info, str.data, str.count);
 		#endif
 
 
+		Vector<u8> memoryBlock;
+		memoryBlock.resize(4096);
 		BinaryWriter memoryWriter(memoryBlock.data(), 0, memoryBlock.size());
 
 		bool res = TraverseDir(m_logger, dirPathLen ? dirPath : AsView(TC("/")),
@@ -333,8 +338,8 @@ namespace uba
 		u32 tableOffset;
 
 		SCOPED_WRITE_LOCK(dirTable.m_memoryLock, memoryLock);
-		EnsureDirectoryTableMemory();
 		u32 writePos = dirTable.m_memorySize;
+		EnsureDirectoryTableMemory(writePos + 128 + storageSize);
 		BinaryWriter tableWriter(m_directoryTableMem + dirTable.m_memorySize);
 
 		if (isRefresh)
@@ -858,9 +863,9 @@ namespace uba
 
 
 		SCOPED_WRITE_LOCK(dirTable.m_memoryLock, memoryLock);
-		EnsureDirectoryTableMemory();
-		u8* startPos = m_directoryTableMem + dirTable.m_memorySize;
-		BinaryWriter writer(startPos);
+		u32 writePos = dirTable.m_memorySize;
+		EnsureDirectoryTableMemory(writePos + 8 + written);
+		BinaryWriter writer(m_directoryTableMem + writePos);
 		writer.Write7BitEncoded(written); // Storage size
 		insres.first->second = dirTable.m_memorySize + u32(writer.GetPosition() + entryPos); // Storing position to last write time
 		u32 tableOffset = u32(writer.GetPosition()) + sizeof(StringKey);
@@ -934,9 +939,9 @@ namespace uba
 		#endif
 
 		SCOPED_WRITE_LOCK(dirTable.m_memoryLock, memoryLock);
-		EnsureDirectoryTableMemory();
-		u8* startPos = m_directoryTableMem + dirTable.m_memorySize;
-		BinaryWriter writer(startPos);
+		u32 writePos = dirTable.m_memorySize;
+		EnsureDirectoryTableMemory(writePos + 8 + written);
+		BinaryWriter writer(m_directoryTableMem + writePos);
 		writer.Write7BitEncoded(written); // Storage size
 		u32 tableOffset = u32(writer.GetPosition()) + sizeof(StringKey);
 		writer.WriteBytes(temp, written);
@@ -2870,20 +2875,20 @@ namespace uba
 		return true;
 	}
 
-	void Session::EnsureDirectoryTableMemory()
+	void Session::EnsureDirectoryTableMemory(u64 neededSize)
 	{
 		auto& dirTable = m_directoryTable;
-		u64 needed = dirTable.m_memorySize + 512*1024;
-		if (needed < m_directoryTableMemCommitted)
+		if (neededSize < m_directoryTableMemCommitted)
 			return;
 
-		if (dirTable.m_memorySize + 1024*1024 > DirTableMemSize)
+		u64 newSize = AlignUp(neededSize, u64(1024*1024));
+		if (newSize > DirTableMemSize)
 		{
 			static bool called = m_logger.Error(TC("Directory table overflow. DirTableMemSize need to be increased (Size: %llu)"), dirTable.m_memorySize);
 		}
 
+		u64 toCommit = newSize - m_directoryTableMemCommitted;
 		u8* address = m_directoryTableMem + m_directoryTableMemCommitted;
-		u64 toCommit = 1024*1024;
 		if (!MapViewCommit(address, toCommit))
 			m_logger.Error(TC("Failed to commit memory for directory table (Committed: %llu) (%s)"), m_directoryTableMemCommitted, LastErrorToText().data);
 		m_directoryTableMemCommitted += toCommit;
diff --git a/Engine/Source/Programs/UnrealBuildAccelerator/Common/Private/UbaSessionClient.cpp b/Engine/Source/Programs/UnrealBuildAccelerator/Common/Private/UbaSessionClient.cpp
--- a/Engine/Source/Programs/UnrealBuildAccelerator/Common/Private/UbaSessionClient.cpp
+++ b/Engine/Source/Programs/UnrealBuildAccelerator/Common/Private/UbaSessionClient.cpp
@@ -1058,7 +1058,7 @@ namespace uba
 			if (m_directoryTableError)
 				return false;
 
-			EnsureDirectoryTableMemory();
+			EnsureDirectoryTableMemory(readPos + toRead);
 
 			if (toRead == 0)
 			{
diff --git a/Engine/Source/Programs/UnrealBuildAccelerator/Common/Public/UbaSession.h b/Engine/Source/Programs/UnrealBuildAccelerator/Common/Public/UbaSession.h
--- a/Engine/Source/Programs/UnrealBuildAccelerator/Common/Public/UbaSession.h
+++ b/Engine/Source/Programs/UnrealBuildAccelerator/Common/Public/UbaSession.h
@@ -169,7 +169,7 @@ namespace uba
 		void WriteDirectoryEntriesRecursive(const StringKey& dirKey, StringView dirPath, u32& outTableOffset);
 		bool CopyImports(Vector<BinaryModule>& out, const tchar* library, tchar* applicationDir, tchar* applicationDirEnd, UnorderedSet<TString>& handledImports, const char* const* loaderPaths);
 		bool CreateProcessJobObject();
-		void EnsureDirectoryTableMemory();
+		void EnsureDirectoryTableMemory(u64 neededSize);
 		void GetSystemInfo(StringBufferBase& out);
 		bool GetMemoryInfo(u64& outAvailable, u64& outTotal);
 		void WriteSummary(BinaryWriter& writer, const Function<void(Logger& logger)>& summaryFunc);
