[UBA]
* ShaderCompileWorker is allowed to run in same console as parent

diff --git a/Engine/Source/Programs/UnrealBuildAccelerator/Core/Private/UbaApplicationRules.cpp b/Engine/Source/Programs/UnrealBuildAccelerator/Core/Private/UbaApplicationRules.cpp
--- a/Engine/Source/Programs/UnrealBuildAccelerator/Core/Private/UbaApplicationRules.cpp
+++ b/Engine/Source/Programs/UnrealBuildAccelerator/Core/Private/UbaApplicationRules.cpp
@@ -542,6 +542,11 @@ namespace uba
 		{
 			return file.Contains(TC(".uba."));
 		}
+
+		virtual bool AllowDetach() const
+		{
+			return true;
+		}
 	};
 
 	class ApplicationRulesUbaObjTool : public ApplicationRules
