[UBA]
* Fixed so temp dir is not parsed when creating directory table
* Fixed so directories are not readded to directory table even if they last written time has changed (which I don't understand how it can be)
* Fixed so drive letter without backslash is fixed up before calling traversedir. This prevents reporting error and make sure drive root end up properly in directory table

diff --git a/Engine/Source/Programs/UnrealBuildAccelerator/Common/Private/UbaSession.cpp b/Engine/Source/Programs/UnrealBuildAccelerator/Common/Private/UbaSession.cpp
--- a/Engine/Source/Programs/UnrealBuildAccelerator/Common/Private/UbaSession.cpp
+++ b/Engine/Source/Programs/UnrealBuildAccelerator/Common/Private/UbaSession.cpp
@@ -239,7 +239,7 @@ namespace uba
 		m_environmentVariables.push_back(0);
 	}
 
-	bool Session::WriteDirectoryEntriesInternal(DirectoryTable::Directory& dir, const StringKey& dirKey, const StringView& dirPath, bool isRefresh, u32& outTableOffset)
+	bool Session::WriteDirectoryEntriesInternal(DirectoryTable::Directory& dir, const StringKey& dirKey, StringView dirPath, bool isRefresh, u32& outTableOffset)
 	{
 		if (dir.tableOffset != InvalidTableOffset && !isRefresh)
 		{
@@ -256,8 +256,7 @@ namespace uba
 		u32 itemCount = 0;
 
 		StringKeyHasher hasher;
-		u32 dirPathLen = dirPath.count;
-		if (dirPathLen)
+		if (dirPath.count)
 		{
 			StringBuffer<> forHash;
 			forHash.Append(dirPath);
@@ -282,54 +281,86 @@ namespace uba
 		memoryBlock.resize(4096);
 		BinaryWriter memoryWriter(memoryBlock.data(), 0, memoryBlock.size());
 
-		bool res = TraverseDir(m_logger, dirPathLen ? dirPath : AsView(TC("/")),
-			[&](const DirectoryEntry& e)
+		if (dirKey != m_directoryForcedEmpty)
+		{
+			StringBuffer<4> realPath;
+			if (IsWindows)
 			{
-				StringBuffer<256> fileNameForHash;
-				fileNameForHash.Append(PathSeparator).Append(e.name, e.nameLen);
-				if (CaseInsensitiveFs)
-					fileNameForHash.MakeLower();
-
-				StringKey fileKey = ToStringKey(hasher, fileNameForHash.data, fileNameForHash.count);
-				auto res = dir.files.try_emplace(fileKey, ~0u);
-				if (!res.second)
-					return;
-				UBA_ASSERT(e.attributes);
-				memoryWriter.WriteString(e.name, e.nameLen);
+				if (dirPath.count == 2)
+					dirPath = realPath.Append(dirPath).Append(PathSeparator);
+			}
+			else
+			{
+				if (!dirPath.count)
+					dirPath = realPath.Append(PathSeparator);
+			}
 
-				#if UBA_DEBUG_TRACK_DIR
-				m_debugLogger->Info(TC("    %s (Size: %llu, Key: %s, Id: %llu)\n"), e.name, e.size, KeyToString(fileKey).data, e.id);
-				#endif
+			bool res = TraverseDir(m_logger, dirPath,
+				[&](const DirectoryEntry& e)
+				{
+					StringBuffer<256> fileNameForHash;
+					fileNameForHash.Append(PathSeparator).Append(e.name, e.nameLen);
+					if (CaseInsensitiveFs)
+						fileNameForHash.MakeLower();
+
+					StringKey fileKey = ToStringKey(hasher, fileNameForHash.data, fileNameForHash.count);
+					auto res = dir.files.try_emplace(fileKey, ~0u);
+					if (!res.second)
+						return;
+					UBA_ASSERT(e.attributes);
+					memoryWriter.WriteString(e.name, e.nameLen);
+
+					#if UBA_DEBUG_TRACK_DIR
+					m_debugLogger->Info(TC("    %s (Size: %llu, Key: %s, Id: %llu)\n"), e.name, e.size, KeyToString(fileKey).data, e.id);
+					#endif
 
-				u64 id = e.id;
-				if (id == 0xffffffffffffffffllu) // When using projfs we might not have the file yet and in that case we need to make this up.
-					id = ++m_fileIndexCounter;
+					u64 id = e.id;
+					if (id == 0xffffffffffffffffllu) // When using projfs we might not have the file yet and in that case we need to make this up.
+						id = ++m_fileIndexCounter;
 
-				res.first->second = u32(memoryWriter.GetPosition()); // Temporary offset that will be used further down to calculate the real offset
-				memoryWriter.WriteU64(e.lastWritten);
-				memoryWriter.WriteU32(e.attributes);
-				memoryWriter.WriteU32(e.volumeSerial);
-				memoryWriter.WriteU64(id);
-				memoryWriter.WriteU64(e.size);
+					res.first->second = u32(memoryWriter.GetPosition()); // Temporary offset that will be used further down to calculate the real offset
+					memoryWriter.WriteU64(e.lastWritten);
+					memoryWriter.WriteU32(e.attributes);
+					memoryWriter.WriteU32(e.volumeSerial);
+					memoryWriter.WriteU64(id);
+					memoryWriter.WriteU64(e.size);
 
-				FileEntryAdded(res.first->first, e.lastWritten, e.size);
+					FileEntryAdded(res.first->first, e.lastWritten, e.size);
 
 
-				++itemCount;
-				if (memoryWriter.GetPosition() > memoryBlock.size() - MaxPath)
+					++itemCount;
+					if (memoryWriter.GetPosition() > memoryBlock.size() - MaxPath)
+					{
+						memoryBlock.resize(memoryBlock.size() * 2);
+						memoryWriter.ChangeData(memoryBlock.data(), memoryBlock.size());
+					}
+				}, true,
+				[&](const DirectoryInfo& e)
 				{
-					memoryBlock.resize(memoryBlock.size() * 2);
-					memoryWriter.ChangeData(memoryBlock.data(), memoryBlock.size());
-				}
-			}, true,
-			[&](const DirectoryInfo& e)
+					volumeSerial = e.volumeSerial;
+					dirAttributes = e.attributes;
+					fileIndex = e.id;
+				});
+			if (!res)
 			{
-				volumeSerial = e.volumeSerial;
-				dirAttributes = e.attributes;
-				fileIndex = e.id;
-			});
-		if (!res)
-			return false;
+				#if UBA_DEBUG_TRACK_DIR
+				m_debugLogger->Info(TC("    FAILED (not existing?)\n"));
+				#endif
+				return false;
+			}
+		}
+		else
+		{
+			#if UBA_DEBUG_TRACK_DIR
+			m_debugLogger->Info(TC("    FORCED EMPTY\n"));
+			#endif
+
+			#if PLATFORM_WINDOWS
+			dirAttributes = FILE_ATTRIBUTE_DIRECTORY;
+			#else
+			UBA_ASSERTF(false, TC("Not implemented"));
+			#endif
+		}
 
 		written = memoryWriter.GetPosition();
 
@@ -790,6 +821,7 @@ namespace uba
 		u64 fileIndex = InvalidValue;
 		u32 attributes = 0;
 		u32 volumeSerial = 0;
+		bool isDirectory = false;
 
 		if (shouldWriteToDisk)
 		{
@@ -800,7 +832,8 @@ namespace uba
 			attributes = info.attributes;
 			volumeSerial = info.volumeSerialNumber;
 			lastWriteTime = info.lastWriteTime;
-			if (IsDirectory(attributes))
+			isDirectory = IsDirectory(attributes);
+			if (isDirectory)
 				fileSize = 0;
 			else
 				fileSize = info.size;
@@ -821,7 +854,7 @@ namespace uba
 			BinaryReader reader(m_directoryTableMem + insres.first->second);
 			u64 oldLastWriteTime = reader.ReadU64();
 
-			if (lastWriteTime == oldLastWriteTime)
+			if (lastWriteTime == oldLastWriteTime || isDirectory) // Ignore updating directories.. they should always be the same regardless
 			{
 				#if !PLATFORM_WINDOWS
 				reader.Skip(sizeof(u32) * 2);
@@ -1130,6 +1163,19 @@ namespace uba
 		m_storage.CreateDirectory(m_sessionLogDir.data);
 		m_sessionLogDir.EnsureEndsWithSlash();
 
+		// We never want to populate files in Temp
+		#if PLATFORM_WINDOWS
+		if (info.treatTempDirAsEmpty)
+		{
+			wchar_t systemTemp[256];
+			GetEnvironmentVariableW(L"TEMP", systemTemp, 256);
+			StringBuffer<> temp;
+			FixPath(systemTemp, nullptr, 0, temp);
+			temp.MakeLower();
+			m_directoryForcedEmpty = ToStringKey(temp);
+		}
+		#endif
+
 		if (info.traceOutputFile)
 			m_traceOutputFile.Append(info.traceOutputFile);
 	}
diff --git a/Engine/Source/Programs/UnrealBuildAccelerator/Common/Public/UbaSession.h b/Engine/Source/Programs/UnrealBuildAccelerator/Common/Public/UbaSession.h
--- a/Engine/Source/Programs/UnrealBuildAccelerator/Common/Public/UbaSession.h
+++ b/Engine/Source/Programs/UnrealBuildAccelerator/Common/Public/UbaSession.h
@@ -165,7 +165,7 @@ namespace uba
 		static constexpr CasKey CasKeyIsDirectory = { ~u64(0), ~u64(0), ~u32(0) };
 
 		void AddEnvironmentVariableNoLock(const tchar* key, const tchar* value);
-		bool WriteDirectoryEntriesInternal(DirectoryTable::Directory& dir, const StringKey& dirKey, const StringView& dirPath, bool isRefresh, u32& outTableOffset);
+		bool WriteDirectoryEntriesInternal(DirectoryTable::Directory& dir, const StringKey& dirKey, StringView dirPath, bool isRefresh, u32& outTableOffset);
 		void WriteDirectoryEntriesRecursive(const StringKey& dirKey, StringView dirPath, u32& outTableOffset);
 		bool CopyImports(Vector<BinaryModule>& out, const tchar* library, tchar* applicationDir, tchar* applicationDirEnd, UnorderedSet<TString>& handledImports, const char* const* loaderPaths);
 		bool CreateProcessJobObject();
@@ -226,6 +226,7 @@ namespace uba
 		u8* m_directoryTableMem;
 		u64 m_directoryTableMemCommitted = 0;
 		DirectoryTable m_directoryTable;
+		StringKey m_directoryForcedEmpty;
 	
 		FileMappingHandle m_fileMappingTableHandle;
 		FileMappingBuffer m_fileMappingBuffer;
