[UBA]
* Fixed bug in EnsureDirectoryTableMemory where if neededsize is exactly the same as committed size it will try to map 0 bytes

diff --git a/Engine/Source/Programs/UnrealBuildAccelerator/Common/Private/UbaSession.cpp b/Engine/Source/Programs/UnrealBuildAccelerator/Common/Private/UbaSession.cpp
--- a/Engine/Source/Programs/UnrealBuildAccelerator/Common/Private/UbaSession.cpp
+++ b/Engine/Source/Programs/UnrealBuildAccelerator/Common/Private/UbaSession.cpp
@@ -2924,7 +2924,7 @@ namespace uba
 	void Session::EnsureDirectoryTableMemory(u64 neededSize)
 	{
 		auto& dirTable = m_directoryTable;
-		if (neededSize < m_directoryTableMemCommitted)
+		if (neededSize <= m_directoryTableMemCommitted)
 			return;
 
 		u64 newSize = AlignUp(neededSize, u64(1024*1024));
@@ -2936,7 +2936,7 @@ namespace uba
 		u64 toCommit = newSize - m_directoryTableMemCommitted;
 		u8* address = m_directoryTableMem + m_directoryTableMemCommitted;
 		if (!MapViewCommit(address, toCommit))
-			m_logger.Error(TC("Failed to commit memory for directory table (Committed: %llu) (%s)"), m_directoryTableMemCommitted, LastErrorToText().data);
+			m_logger.Error(TC("Failed to commit memory for directory table (Committed: %llu, ToCommit: %llu) (%s)"), m_directoryTableMemCommitted, toCommit, LastErrorToText().data);
 		m_directoryTableMemCommitted += toCommit;
 	}
 
