[UBA]
* Added option "-noaws" to ubaagent which prevents it from trying to communicate with aws meta data service

diff --git a/Engine/Source/Programs/UnrealBuildAccelerator/Agent/Private/UbaAgent.cpp b/Engine/Source/Programs/UnrealBuildAccelerator/Agent/Private/UbaAgent.cpp
--- a/Engine/Source/Programs/UnrealBuildAccelerator/Agent/Private/UbaAgent.cpp
+++ b/Engine/Source/Programs/UnrealBuildAccelerator/Agent/Private/UbaAgent.cpp
@@ -103,6 +103,7 @@ namespace uba
 		logger.Info(TC("  -zone                   Set the zone this machine exists in. This info is used to figure out if proxies should be created."));
 		logger.Info(TC("  -version                Prints the version for this executable."));
 		logger.Info(TC("  -noproxy                Does not allow this agent to be a storage proxy for other agents"));
+		logger.Info(TC("  -noaws                  Will not try to connect to aws meta data server (this can take time during first startup)"));
 		logger.Info(TC("  -killrandom             Kills random process and exit session"));
 		logger.Info(TC("  -memwait=<percent>      The amount of memory needed to spawn a process. Set this to 100 to disable. Defaults to 80%%"));
 		logger.Info(TC("  -memkill=<percent>      The amount of memory needed before processes starts to be killed. Set this to 100 to disable. Defaults to 90%%"));
@@ -387,6 +388,7 @@ namespace uba
 		bool useQuic = false;
 		bool poll = true;
 		bool allowProxy = true;
+		bool couldBeAws = true;
 		bool useStorage = true;
 		bool resetStore = false;
 		bool quiet = false;
@@ -674,6 +676,10 @@ namespace uba
 			{
 				allowProxy = false;
 			}
+			else if (name.Equals(TC("-noaws")))
+			{
+				couldBeAws = false;
+			}
 			else if (name.Equals(TC("-command")))
 			{
 				if (value.IsEmpty())
@@ -754,6 +760,7 @@ namespace uba
 		// Check if AWS
 		#if UBA_USE_AWS
 		AWS aws;
+		if (couldBeAws)
 		{
 			DirectoryCache dirCache;
 			dirCache.CreateDirectory(logger, g_rootDir.data);
@@ -975,7 +982,7 @@ namespace uba
 		u64 terminationTimeMs = 0;
 
 		#if UBA_USE_AWS
-		if (terminationReason.IsEmpty())
+		if (couldBeAws && terminationReason.IsEmpty())
 			aws.IsTerminating(logger, terminationReason, terminationTimeMs);
 		#endif
 
@@ -1336,7 +1343,7 @@ namespace uba
 					if (IsTerminating(logger, eventFile.data, terminationReason, terminationTimeMs))
 						isTerminating = true;
 					#if UBA_USE_AWS
-					else if (aws.IsTerminating(logger, terminationReason, terminationTimeMs))
+					else if (couldBeAws && aws.IsTerminating(logger, terminationReason, terminationTimeMs))
 						isTerminating = true;
 					#endif
 
