[UBA]
* Added support for setting timeout on HttpConnections
* Fixed so http query to aws metadata service times out in 200ms

diff --git a/Engine/Source/Programs/UnrealBuildAccelerator/Common/Private/UbaNetworkBackendTcp.cpp b/Engine/Source/Programs/UnrealBuildAccelerator/Common/Private/UbaNetworkBackendTcp.cpp
--- a/Engine/Source/Programs/UnrealBuildAccelerator/Common/Private/UbaNetworkBackendTcp.cpp
+++ b/Engine/Source/Programs/UnrealBuildAccelerator/Common/Private/UbaNetworkBackendTcp.cpp
@@ -93,6 +93,7 @@ namespace uba
 
 	bool SetKeepAlive(Logger& logger, SOCKET socket);
 	bool SetBlocking(Logger& logger, SOCKET socket, bool blocking);
+	bool SetTimeout(Logger& logger, SOCKET socket, u32 timeoutMs);
 	bool DisableNagle(Logger& logger, SOCKET socket);
 	bool EnableFastLoopback(Logger& logger, SOCKET socket);
 	bool SendSocket(Logger& logger, SOCKET socket, const void* b, u64 bufferLen);
@@ -752,6 +753,20 @@ namespace uba
 		return true;
 	}
 
+	bool SetTimeout(Logger& logger, SOCKET socket, u32 timeoutMs)
+	{
+		u32 timeoutSeconds = timeoutMs / 1000;
+		timeoutMs = timeoutMs % 1000;
+		struct timeval timeout;
+		timeout.tv_sec = timeoutSeconds;
+		timeout.tv_usec = timeoutMs*1000;
+		if (setsockopt(socket, SOL_SOCKET, SO_SNDTIMEO, (const char*)&timeout, sizeof timeout) == SOCKET_ERROR)
+			return logger.Error(TC("setsockopt SO_SNDTIMEO error: (error: %s)"), LastErrorToText(WSAGetLastError()).data);
+		if (setsockopt(socket, SOL_SOCKET, SO_RCVTIMEO, (const char*)&timeout, sizeof timeout) == SOCKET_ERROR)
+			return logger.Error(TC("setsockopt SO_RCVTIMEO error: (error: %s)"), LastErrorToText(WSAGetLastError()).data);
+		return true;
+	}
+
 	bool EnableFastLoopback(Logger& logger, SOCKET socket)
 	{
 #if 0 // PLATFORM_WINDOWS // Disabled for now because it seems like it is not supported on windows 11
@@ -992,6 +1007,14 @@ namespace uba
 		m_wsaInitDone = true;
 		#endif
 
+		hostent* hostent = gethostbyname(host);
+		if (hostent == NULL)
+			return logger.Error(TC("HttpRequest: gethostbyname error"));
+
+		unsigned long in_addr = inet_addr(inet_ntoa(*(struct in_addr*)*(hostent->h_addr_list)));
+		if (in_addr == INADDR_NONE)
+			return logger.Error(TC("HttpRequest: inet_addr error"));
+
 		protoent* protoent = getprotobyname("tcp");
 		if (protoent == NULL)
 			return logger.Error(TC("HttpRequest: socket error"));
@@ -1001,13 +1024,8 @@ namespace uba
 			return logger.Error(TC("HttpRequest: socket error"));
 		auto socketClose = MakeGuard([sock]() { closesocket(sock); });
 
-		hostent* hostent = gethostbyname(host);
-		if (hostent == NULL)
-			return logger.Error(TC("HttpRequest: gethostbyname error"));
-
-		unsigned long in_addr = inet_addr(inet_ntoa(*(struct in_addr*)*(hostent->h_addr_list)));
-		if (in_addr == INADDR_NONE)
-			return logger.Error(TC("HttpRequest: inet_addr error"));
+		if (m_connectTimeOutMs)
+			SetTimeout(logger, sock, m_connectTimeOutMs);
 
 		sockaddr_in sockaddr_in;
 		sockaddr_in.sin_addr.s_addr = in_addr;
@@ -1019,6 +1037,9 @@ namespace uba
 
 		socketClose.Cancel();
 
+		if (m_connectTimeOutMs)
+			SetTimeout(logger, sock, 10000); // 10 seconds timeout
+
 		strcpy_s(m_host, sizeof_array(m_host), host);
 		m_socket = sock;
 		return true;
@@ -1087,4 +1108,9 @@ namespace uba
 		outResponse.Append(bodyStart + 4);
 		return true;
 	}
+
+	void HttpConnection::SetConnectTimeout(u32 timeOutMs)
+	{
+		m_connectTimeOutMs = timeOutMs;
+	}
 }
\ No newline at end of file
diff --git a/Engine/Source/Programs/UnrealBuildAccelerator/Common/Public/UbaAWS.h b/Engine/Source/Programs/UnrealBuildAccelerator/Common/Public/UbaAWS.h
--- a/Engine/Source/Programs/UnrealBuildAccelerator/Common/Public/UbaAWS.h
+++ b/Engine/Source/Programs/UnrealBuildAccelerator/Common/Public/UbaAWS.h
@@ -31,6 +31,7 @@ namespace uba
 				return false;
 
 			HttpConnection http;
+			http.SetConnectTimeout(200);
 
 			u32 statusCode = 0;
 
@@ -81,6 +82,7 @@ namespace uba
 				return false;
 
 			HttpConnection http;
+			http.SetConnectTimeout(200);
 
 			StringBuffer<128> availabilityZone;
 			u32 statusCode = 0;
diff --git a/Engine/Source/Programs/UnrealBuildAccelerator/Common/Public/UbaNetworkBackend.h b/Engine/Source/Programs/UnrealBuildAccelerator/Common/Public/UbaNetworkBackend.h
--- a/Engine/Source/Programs/UnrealBuildAccelerator/Common/Public/UbaNetworkBackend.h
+++ b/Engine/Source/Programs/UnrealBuildAccelerator/Common/Public/UbaNetworkBackend.h
@@ -87,6 +87,8 @@ namespace uba
 		~HttpConnection();
 		bool Query(Logger& logger, const char* type, StringBufferBase& outResponse, u32& outStatusCode, const char* host, const char* path, const char* header = "");
 
+		void SetConnectTimeout(u32 timeOutMs);
+
 	private:
 		bool Connect(Logger& logger, const char* host);
 
@@ -97,5 +99,6 @@ namespace uba
 		#else
 		int m_socket;
 		#endif
+		u32 m_connectTimeOutMs = 0;
 	};
 }
