[UBA]
* Improved error message for when something goes wrong around server decompressing cas from client

diff --git a/Engine/Source/Programs/UnrealBuildAccelerator/Common/Private/UbaSession.cpp b/Engine/Source/Programs/UnrealBuildAccelerator/Common/Private/UbaSession.cpp
--- a/Engine/Source/Programs/UnrealBuildAccelerator/Common/Private/UbaSession.cpp
+++ b/Engine/Source/Programs/UnrealBuildAccelerator/Common/Private/UbaSession.cpp
@@ -693,7 +693,7 @@ namespace uba
 					return false;
 				auto unmapGuard = MakeGuard([&](){ m_fileMappingBuffer.UnmapView(mappedViewWrite, fileName); });
 
-				if (!m_storage.DecompressMemoryToMemory(readMemory, mappedViewWrite.memory, size, fileName, TC("TransientMapping")))
+				if (!m_storage.DecompressMemoryToMemory(readMemory, mappedViewRead.size, mappedViewWrite.memory, size, fileName, TC("TransientMapping")))
 					return false;
 				unmapGuard.Execute();
 
diff --git a/Engine/Source/Programs/UnrealBuildAccelerator/Common/Private/UbaSessionClient.cpp b/Engine/Source/Programs/UnrealBuildAccelerator/Common/Private/UbaSessionClient.cpp
--- a/Engine/Source/Programs/UnrealBuildAccelerator/Common/Private/UbaSessionClient.cpp
+++ b/Engine/Source/Programs/UnrealBuildAccelerator/Common/Private/UbaSessionClient.cpp
@@ -636,7 +636,8 @@ namespace uba
 			if (!m_storage.StoreCasFileClient(casKey, source.key, source.name.c_str(), source.mappingHandle, 0, source.mappingWritten, destination, keepMappingInMemory, compressed))
 				return false;
 		}
-		UBA_ASSERTF(casKey != CasKeyZero, TC("Failed to store cas file for %s (destination %s)"), source.name.c_str(), destination);
+		if (casKey == CasKeyZero)
+			return m_logger.Error(TC("Failed to store cas on server for local file %s (size %llu destination %s)"), source.name.c_str(), source.mappingWritten, destination);
 	
 		CloseFileMapping(m_logger, source.mappingHandle, source.name.c_str());
 		source.mappingHandle = {};
@@ -659,7 +660,7 @@ namespace uba
 		#endif
 
 		if (!reader.ReadBool())
-			return m_logger.Error(TC("Server failed to receive file %s (%s)"), source.name.c_str(), destination);
+			return m_logger.Error(TC("Server failed to copy cas %s to %s (local source %s)"), CasKeyString(casKey).str, destination, source.name.c_str());
 		return true;
 	}
 
diff --git a/Engine/Source/Programs/UnrealBuildAccelerator/Common/Private/UbaStorage.cpp b/Engine/Source/Programs/UnrealBuildAccelerator/Common/Private/UbaStorage.cpp
--- a/Engine/Source/Programs/UnrealBuildAccelerator/Common/Private/UbaStorage.cpp
+++ b/Engine/Source/Programs/UnrealBuildAccelerator/Common/Private/UbaStorage.cpp
@@ -683,7 +683,7 @@ namespace uba
 		return false;
 	}
 
-	bool StorageImpl::DecompressMemoryToMemory(u8* compressedData, u8* writeData, u64 decompressedSize, const tchar* readHint, const tchar* writeHint)
+	bool StorageImpl::DecompressMemoryToMemory(u8* compressedData, u64 compressedSize, u8* writeData, u64 decompressedSize, const tchar* readHint, const tchar* writeHint)
 	{
 		UBA_ASSERTF(compressedData, TC("DecompressMemoryToMemory got readmem nullptr (%s)"), readHint);
 		UBA_ASSERTF(writeData, TC("DecompressMemoryToMemory got writemem nullptr (%s)"), writeHint);
@@ -811,7 +811,7 @@ namespace uba
 				TimerScope ts(stats.decompressToMem);
 				OO_SINTa decompLen = OodleLZ_Decompress(readPos, (OO_SINTa)compressedBlockSize, writePos, (OO_SINTa)decompressedBlockSize); (void)decompLen;
 				if (decompLen != decompressedBlockSize)
-					return m_logger.Error(TC("Expecting to be able to decompress %u to %u bytes at pos %llu but got %llu (%s -> %s)"), compressedBlockSize, decompressedBlockSize, decompressedSize - left, decompLen, readHint, writeHint);
+					return m_logger.Error(TC("Expecting to be able to decompress %u to %u bytes at pos %llu but got %llu. File compressed size: %llu (%s -> %s)"), compressedBlockSize, decompressedBlockSize, decompressedSize - left, decompLen, compressedSize, readHint, writeHint);
 				writePos += decompressedBlockSize;
 				readPos += compressedBlockSize;
 				left -= decompressedBlockSize;
@@ -1617,7 +1617,7 @@ namespace uba
 			for (auto& pair : m_fileTableLookup)
 			{
 				FileEntry& fileEntry = pair.second;
-				SCOPED_READ_LOCK(fileEntry.lock, entryLock);
+				SCOPED_READ_LOCK(fileEntry.lock, entryLock); // TODO: Remove this? tsan complains about it
 				if (fileEntry.casKey == CasKeyZero || fileEntry.lastWritten == 0)
 					continue;
 				writer.WriteStringKey(pair.first);
@@ -1753,14 +1753,14 @@ namespace uba
 							FileAccessor file(m_logger, filePath.c_str());
 							if (!file.OpenMemoryRead())
 								return;
-							if (file.GetSize())
+							if (u64 fileSize = file.GetSize())
 							{
 								u8* mem = file.GetData();
 								u64 decompressedSize = *(u64*)mem;
 								mem += sizeof(u64);
 								u8* dest = new u8[decompressedSize];
 								auto g = MakeGuard([dest]() { delete[] dest; });
-								if (!DecompressMemoryToMemory(mem, dest, decompressedSize, filePath.c_str(), TC("Memory")))
+								if (!DecompressMemoryToMemory(mem, fileSize, dest, decompressedSize, filePath.c_str(), TC("Memory")))
 									return;
 								checkedKey = CalculateCasKey(dest, decompressedSize, true);
 							}
@@ -2546,9 +2546,9 @@ namespace uba
 
 				auto decompressToMemory = [&](u8* destinationMemory)
 					{
-						if (casEntry->mappingHandle.IsValid())
+						if (readData)
 						{
-							if (!DecompressMemoryToMemory(readData, destinationMemory, decompressedSize, casFile.data, destination))
+							if (!DecompressMemoryToMemory(readData, mappedView.size, destinationMemory, decompressedSize, casFile.data, destination))
 								return false;
 						}
 						else
@@ -2889,7 +2889,7 @@ namespace uba
 				return m_logger.Error(TC("Failed to map view of file mapping for %s (%s)"), fileName, LastErrorToText().data);
 			auto udg = MakeGuard([&]() { UnmapViewOfFile(m_logger, fileData, compressedSize, fileName); });
 			
-			if (!DecompressMemoryToMemory(fileData + 8, dest, decompressedSize, fileName, writeHint))
+			if (!DecompressMemoryToMemory(fileData + 8, compressedSize, dest, decompressedSize, fileName, writeHint))
 				return false;
 		}
 		else
diff --git a/Engine/Source/Programs/UnrealBuildAccelerator/Common/Public/UbaStorage.h b/Engine/Source/Programs/UnrealBuildAccelerator/Common/Public/UbaStorage.h
--- a/Engine/Source/Programs/UnrealBuildAccelerator/Common/Public/UbaStorage.h
+++ b/Engine/Source/Programs/UnrealBuildAccelerator/Common/Public/UbaStorage.h
@@ -33,7 +33,7 @@ namespace uba
 		virtual bool HasProxy(u32 clientId) { return false; }
 
 		virtual bool DecompressFileToMemory(const tchar* fileName, FileHandle fileHandle, u8* dest, u64 decompressedSize, const tchar* writeHint) = 0;
-		virtual bool DecompressMemoryToMemory(u8* compressedData, u8* writeData, u64 decompressedSize, const tchar* readHint, const tchar* writeHint) = 0;
+		virtual bool DecompressMemoryToMemory(u8* compressedData, u64 compressedSize, u8* writeData, u64 decompressedSize, const tchar* readHint, const tchar* writeHint) = 0;
 		virtual bool CreateDirectory(const tchar* dir) = 0;
 		virtual bool DeleteCasForFile(const tchar* file) = 0;
 
@@ -171,7 +171,7 @@ namespace uba
 		bool AddCasFile(StringKey fileNameKey, const tchar* fileName, const CasKey& casKey, bool deferCreation, bool fileIsCompressed);
 		void CasEntryAccessed(const CasKey& casKey);
 		virtual bool IsDisallowedPath(const tchar* fileName);
-		virtual bool DecompressMemoryToMemory(u8* compressedData, u8* writeData, u64 decompressedSize, const tchar* readHint, const tchar* writeHint) override;
+		virtual bool DecompressMemoryToMemory(u8* compressedData, u64 compressedSize, u8* writeData, u64 decompressedSize, const tchar* readHint, const tchar* writeHint) override;
 		bool DecompressMemoryToFile(u8* compressedData, FileAccessor& destination, u64 decompressedSize, bool useNoBuffering);
 
 		void CasEntryAccessed(CasEntry& entry);
