[UBA]
* Fixed bugs related to segmentation violation on linux.. symbol file was not properly found. Also added signal value to output

diff --git a/Engine/Source/Programs/UnrealBuildAccelerator/Common/Private/UbaFile.cpp b/Engine/Source/Programs/UnrealBuildAccelerator/Common/Private/UbaFile.cpp
--- a/Engine/Source/Programs/UnrealBuildAccelerator/Common/Private/UbaFile.cpp
+++ b/Engine/Source/Programs/UnrealBuildAccelerator/Common/Private/UbaFile.cpp
@@ -923,7 +923,8 @@ namespace uba
 		const tchar* platformStr = TC("Mac");
 		#endif
 		out.Resize(engineDirPos - out.data + 8).Append(TC("Binaries")).Append(PathSeparator).Append(platformStr).Append(PathSeparator).Append("UnrealBuildAccelerator").Append(PathSeparator);
-		out.Append(isWindowsArm ? TC("arm64") : TC("x64")).Append(PathSeparator);
+		if (IsWindows)
+			out.Append(isWindowsArm ? TC("arm64") : TC("x64")).Append(PathSeparator);
 		return true;
 	};
 }
diff --git a/Engine/Source/Programs/UnrealBuildAccelerator/Common/Private/UbaLogger.cpp b/Engine/Source/Programs/UnrealBuildAccelerator/Common/Private/UbaLogger.cpp
--- a/Engine/Source/Programs/UnrealBuildAccelerator/Common/Private/UbaLogger.cpp
+++ b/Engine/Source/Programs/UnrealBuildAccelerator/Common/Private/UbaLogger.cpp
@@ -32,16 +32,17 @@ namespace uba
 		WriteCallstackInfo(writer, 2, context);
 
 		static auto& sb = *new StringBuffer<16*1024>();
-		WriteAssertInfo(sb.Clear(), text, file, line, expr);
+		WriteAssertInfo(sb.Clear(), text, file, line, expr, context);
 
 		StringView searchPaths[3];
 		StringBuffer<512> currentModuleDir;
 		LoggerWithWriter logger(g_nullLogWriter);
 		GetDirectoryOfCurrentModule(logger, currentModuleDir);
-		searchPaths[0] = currentModuleDir;
 		StringBuffer<512> alternativePath;
+		u32 searchPathIndex = 0;
 		if (GetAlternativeUbaPath(logger, alternativePath, currentModuleDir, IsWindows && IsArm))
-			searchPaths[1] = alternativePath;
+			searchPaths[searchPathIndex++] = alternativePath;
+		searchPaths[searchPathIndex] = currentModuleDir;
 
 		tchar executable[512] = { 0 };
 		#if !PLATFORM_WINDOWS
@@ -514,10 +515,13 @@ namespace uba
 			mod.handled = false;
 		}
 
+		out.Append(TC("\n CALLSTACK")).Append(isRunningWine ? TC(" (Wine)") : TC("")).Append(':');
+
 		if (entries.empty())
+		{
+			out.Appendf(TC("\n   <No entries available>"));
 			return;
-
-		out.Appendf(TC("\n CALLSTACK%s:"), isRunningWine ? TC(" (Wine)") : TC(""));
+		}
 
 #if PLATFORM_WINDOWS
 		static u64 processHandleCounter = 45234523;
@@ -617,13 +621,14 @@ namespace uba
 			{
 				if (mod.name.empty())
 					mod.name = executable;
-				StringBuffer<256> buf(mod.name);
-				if (auto lastSlash = buf.Last(PathSeparator))
-					buf.Resize(lastSlash - buf.data);
+				StringBuffer<128> name(mod.name);
+				if (auto lastSlash = name.Last(PathSeparator))
+					name.Clear().Append(lastSlash + 1);
 				for (auto it=searchPaths; it->count; ++it)
 				{
 					struct stat attr;
-					buf.Clear().Append(*it).EnsureEndsWithSlash().Append(mod.name);
+					StringBuffer<512> buf;
+					buf.Clear().Append(*it).EnsureEndsWithSlash().Append(name);
 					if (stat(buf.data, &attr) == -1)
 						continue;
 					mod.name = buf.data;
diff --git a/Engine/Source/Programs/UnrealBuildAccelerator/Core/Private/UbaPlatform.cpp b/Engine/Source/Programs/UnrealBuildAccelerator/Core/Private/UbaPlatform.cpp
--- a/Engine/Source/Programs/UnrealBuildAccelerator/Core/Private/UbaPlatform.cpp
+++ b/Engine/Source/Programs/UnrealBuildAccelerator/Core/Private/UbaPlatform.cpp
@@ -142,13 +142,21 @@ namespace uba
 		#endif
 	}
 
-	void WriteAssertInfo(StringBufferBase& out, const tchar* text, const char* file, u32 line, const char* expr)
+	void WriteAssertInfo(StringBufferBase& out, const tchar* text, const char* file, u32 line, const char* expr, void* context)
 	{
+		int signal = 0;
+		#if !PLATFORM_WINDOWS
+		if (context)
+			signal = *(int*)context;
+		#endif
+
 		out.Append(TC("ASSERT: ")).Append(*text ? text : TC("Unknown"));
 		if (!*text && expr && *expr && strcmp(expr, "false") != 0)
 			out.Appendf(TC("\n EXPR: ")).Append(expr);
 		if (file && *file)
 			out.Append(TC("\n LOCATION: ")).Append(file).Append(':').AppendValue(line);
+		if (signal)
+			out.Append(TC("\n SIGNAL: ")).AppendValue(signal);
 	}
 
 	#if defined(_M_ARM64)
@@ -334,9 +342,17 @@ namespace uba
 		Vector<ModuleRec*> usedModules;
 
 		if (contextPtr) // Use this to know if we come from a signal
-			skipCallstack +=2;
+			skipCallstack += 2;
 
-		out.Write7BitEncoded(nptrs < skipCallstack ? 0 : nptrs - skipCallstack);
+		if (nptrs <= skipCallstack)
+		{
+			if (!nptrs)
+				skipCallstack = 0;
+			else
+				skipCallstack = nptrs - 1;
+		}
+
+		out.Write7BitEncoded(nptrs - skipCallstack);
 
 		for (int i = skipCallstack; i < nptrs; ++i)
 		{
@@ -551,7 +567,11 @@ namespace uba
 		SIGSEGV,
 		//SIGTERM 
 	};
-	static void segfault_sigaction(int signal) { uba::UbaAssert("Segmentation fault", "", 0, "", true, -1, &signal); }
+	static void segfault_sigaction(int signal)
+	{
+		const char* desc = signal == SIGSEGV ? "ERROR: Segmentation fault" : "ERROR: Unhandled signal";
+		uba::UbaAssert(desc, "", 0, "", true, -1, &signal);
+	}
 #endif
 
 	void AddExceptionHandler()
diff --git a/Engine/Source/Programs/UnrealBuildAccelerator/Detours/Private/PosixOS/UbaDetoursFunctionsPosix.cpp b/Engine/Source/Programs/UnrealBuildAccelerator/Detours/Private/PosixOS/UbaDetoursFunctionsPosix.cpp
--- a/Engine/Source/Programs/UnrealBuildAccelerator/Detours/Private/PosixOS/UbaDetoursFunctionsPosix.cpp
+++ b/Engine/Source/Programs/UnrealBuildAccelerator/Detours/Private/PosixOS/UbaDetoursFunctionsPosix.cpp
@@ -2729,7 +2729,7 @@ namespace uba
 		static CriticalSection cs;
 		ScopedCriticalSection scs(cs);
 		static auto& sb =  *new StringBuffer<8*1024>;
-		WriteAssertInfo(sb, text, file, line, expr);
+		WriteAssertInfo(sb, text, file, line, expr, context);
 		Rpc_ResolveCallstack(sb, 3, context);
 		Rpc_WriteLog(sb.data, sb.count, true, true);
 
