[UBA]
* Renamed "IsArm" to "IsArmBinary" after realizing user could actually run x64 binaries on uba on arm windows
* Added IsRunningArm64 function

diff --git a/Engine/Source/Programs/UnrealBuildAccelerator/Agent/Private/UbaAgent.cpp b/Engine/Source/Programs/UnrealBuildAccelerator/Agent/Private/UbaAgent.cpp
--- a/Engine/Source/Programs/UnrealBuildAccelerator/Agent/Private/UbaAgent.cpp
+++ b/Engine/Source/Programs/UnrealBuildAccelerator/Agent/Private/UbaAgent.cpp
@@ -64,7 +64,7 @@ namespace uba
 		}
 		logger.Info(TC(""));
 		logger.Info(TC("-------------------------------------------"));
-		logger.Info(TC("   UbaAgent v%s%s"), Version, (IsArm ? TC(" (ARM64)") : TC("")));
+		logger.Info(TC("   UbaAgent v%s%s"), Version, (IsArmBinary ? TC(" (ARM64)") : TC("")));
 		logger.Info(TC("-------------------------------------------"));
 		logger.Info(TC(""));
 		logger.Info(TC("  When started UbaAgent will keep trying to connect to provided host address."));
@@ -792,7 +792,7 @@ namespace uba
 		#if UBA_DEBUG
 		dbgStr = TC(" (DEBUG)");
 		#endif
-		logger.Info(TC("UbaAgent v%s%s (Cpu%s: %u, MaxCon: %u, Dir: \"%s\", StoreCapacity: %uGb%s)"), Version, dbgStr, (IsArm ? TC("[Arm]") : TC("")), maxProcessCount, maxConnectionCount, g_rootDir.data, storageCapacityGb, extraInfo.data);
+		logger.Info(TC("UbaAgent v%s%s (Cpu%s: %u, MaxCon: %u, Dir: \"%s\", StoreCapacity: %uGb%s)"), Version, dbgStr, (IsArmBinary ? TC("[Arm]") : TC("")), maxProcessCount, maxConnectionCount, g_rootDir.data, storageCapacityGb, extraInfo.data);
 		
 		Config config;
 		if (!configFile.IsEmpty())
diff --git a/Engine/Source/Programs/UnrealBuildAccelerator/Common/Private/UbaLogger.cpp b/Engine/Source/Programs/UnrealBuildAccelerator/Common/Private/UbaLogger.cpp
--- a/Engine/Source/Programs/UnrealBuildAccelerator/Common/Private/UbaLogger.cpp
+++ b/Engine/Source/Programs/UnrealBuildAccelerator/Common/Private/UbaLogger.cpp
@@ -40,7 +40,7 @@ namespace uba
 		GetDirectoryOfCurrentModule(logger, currentModuleDir);
 		StringBuffer<512> alternativePath;
 		u32 searchPathIndex = 0;
-		if (GetAlternativeUbaPath(logger, alternativePath, currentModuleDir, IsWindows && IsArm))
+		if (GetAlternativeUbaPath(logger, alternativePath, currentModuleDir, IsWindows && IsArmBinary))
 			searchPaths[searchPathIndex++] = alternativePath;
 		searchPaths[searchPathIndex] = currentModuleDir;
 
diff --git a/Engine/Source/Programs/UnrealBuildAccelerator/Common/Private/UbaProcess.cpp b/Engine/Source/Programs/UnrealBuildAccelerator/Common/Private/UbaProcess.cpp
--- a/Engine/Source/Programs/UnrealBuildAccelerator/Common/Private/UbaProcess.cpp
+++ b/Engine/Source/Programs/UnrealBuildAccelerator/Common/Private/UbaProcess.cpp
@@ -294,7 +294,7 @@ namespace uba
 				err.Append(TC("ERROR: Process did not start properly. "));
 
 				ImageInfo imageInfo;
-				bool machineIsArm = IsArm;
+				bool machineIsArm = IsArmBinary;
 				if (!GetImageInfo(imageInfo, m_session.m_logger, m_realApplication.c_str(), true))
 					err.Appendf(TC("Failed to load %s"), m_realApplication.c_str());
 				else if (!imageInfo.is64Bit)
@@ -1020,8 +1020,8 @@ namespace uba
 		auto& process = *(ProcessImpl*)h.m_process;
 		process.m_echoOn = m_echoOn;
 
-		const char* detoursLib = m_session.m_detoursLibrary[IsArm].c_str();
-		u32 detoursLibLen = u32(m_session.m_detoursLibrary[IsArm].size());
+		const char* detoursLib = m_session.m_detoursLibrary[IsArmBinary].c_str();
+		u32 detoursLibLen = u32(m_session.m_detoursLibrary[IsArmBinary].size());
 
 		writer.WriteU32(childProcessId);
 		writer.WriteU32(info.rules->index);
@@ -1450,7 +1450,7 @@ namespace uba
 
 		if (!m_parentProcess)
 		{
-			bool isArmApplication = IsArm;
+			bool isArmApplication = IsArmBinary;
 #if defined(_M_ARM64)
 			bool isX64Application;
 			{
diff --git a/Engine/Source/Programs/UnrealBuildAccelerator/Common/Private/UbaSession.cpp b/Engine/Source/Programs/UnrealBuildAccelerator/Common/Private/UbaSession.cpp
--- a/Engine/Source/Programs/UnrealBuildAccelerator/Common/Private/UbaSession.cpp
+++ b/Engine/Source/Programs/UnrealBuildAccelerator/Common/Private/UbaSession.cpp
@@ -2801,8 +2801,8 @@ namespace uba
 		StringBuffer<256> detoursLibPath;
 		StringBuffer<256> alternativeLibPath;
 
-		detoursLibPath.Append(m_detoursLibrary[IsArm].c_str()).Resize(detoursLibPath.Last(PathSeparator) - detoursLibPath.data);
-		GetAlternativeUbaPath(m_logger, alternativeLibPath, detoursLibPath, IsArm);
+		detoursLibPath.Append(m_detoursLibrary[IsArmBinary].c_str()).Resize(detoursLibPath.Last(PathSeparator) - detoursLibPath.data);
+		GetAlternativeUbaPath(m_logger, alternativeLibPath, detoursLibPath, IsArmBinary);
 
 		StringView searchPaths[3] = { detoursLibPath, alternativeLibPath, {} };
 
@@ -2985,7 +2985,7 @@ namespace uba
 		BytesToText temp(capacity);
 		if (capacity)
 			capacityStr = temp.str;
-		out.Appendf(TC("CPU%s:%s%s Mem:%ugb Cas:%s"), (IsArm ? TC("[Arm]") : TC("")), cpuCountStr.data, hzStr.data, u32(totalMemoryInKilobytes/(1024*1024)), capacityStr);
+		out.Appendf(TC("CPU%s:%s%s Mem:%ugb Cas:%s"), (IsArmBinary ? TC("[Arm]") : TC("")), cpuCountStr.data, hzStr.data, u32(totalMemoryInKilobytes/(1024*1024)), capacityStr);
 
 		StringBuffer<128> zone;
 		if (m_storage.GetZone(zone))
diff --git a/Engine/Source/Programs/UnrealBuildAccelerator/Common/Private/UbaSessionClient.cpp b/Engine/Source/Programs/UnrealBuildAccelerator/Common/Private/UbaSessionClient.cpp
--- a/Engine/Source/Programs/UnrealBuildAccelerator/Common/Private/UbaSessionClient.cpp
+++ b/Engine/Source/Programs/UnrealBuildAccelerator/Common/Private/UbaSessionClient.cpp
@@ -208,7 +208,7 @@ namespace uba
 			UBA_ASSERT(fileNameKey == StringKeyZero);
 			StackBinaryWriter<1024> writer;
 			NetworkMessage msg(m_client, ServiceId, SessionMessageType_EnsureBinaryFile, writer);
-			writer.WriteBool(IsArm);
+			writer.WriteBool(IsArmBinary);
 			writer.WriteString(fileName);
 			writer.WriteStringKey(fileNameKey);
 			writer.WriteString(applicationDir);
@@ -1180,7 +1180,7 @@ namespace uba
 		NetworkMessage msg(m_client, ServiceId, SessionMessageType_Connect, writer);
 		writer.WriteString(m_name.data);
 		writer.WriteU32(SessionNetworkVersion);
-		writer.WriteBool(IsArm);
+		writer.WriteBool(IsArmBinary);
 
 		CasKey keys[2];
 		if (m_useBinariesAsVersion)
@@ -1224,14 +1224,14 @@ namespace uba
 
 		CasKey detoursBinaryKey[2];
 		detoursBinaryKey[0] = reader.ReadCasKey();
-		if (IsArm)
+		if (IsArmBinary)
 			detoursBinaryKey[1] = reader.ReadCasKey();
 
 		StringBuffer<> detoursFile[2];
 
 		if (m_downloadDetoursLib)
 		{
-			for (u32 i=0; i!=IsArm+1; ++i)
+			for (u32 i=0; i!=IsArmBinary+1; ++i)
 			{
 				{
 					TimerScope s(m_stats.storageRetrieve);
@@ -1249,11 +1249,11 @@ namespace uba
 		}
 		else
 		{
-			GetDirectoryOfCurrentModule(m_logger, detoursFile[IsArm]);
-			detoursFile[IsArm].EnsureEndsWithSlash().Append(UBA_DETOURS_LIBRARY);
+			GetDirectoryOfCurrentModule(m_logger, detoursFile[IsArmBinary]);
+			detoursFile[IsArmBinary].EnsureEndsWithSlash().Append(UBA_DETOURS_LIBRARY);
 		}
 
-		for (u32 i=0; i!=IsArm+1; ++i)
+		for (u32 i=0; i!=IsArmBinary+1; ++i)
 		{
 			#if PLATFORM_WINDOWS
 			char dll[1024];
@@ -2143,7 +2143,7 @@ namespace uba
 		StackBinaryWriter<SendMaxSize> msgWriter;
 		NetworkMessage msg(m_client, ServiceId, SessionMessageType_GetSymbols, msgWriter);
 		msgWriter.WriteU32(process.GetId());
-		msgWriter.WriteBool(IsArm);
+		msgWriter.WriteBool(IsArmBinary);
 		msgWriter.WriteBytes(reader.GetPositionData(), reader.GetLeft());
 
 		BinaryReader responseReader(writer.GetData(), 0, writer.GetCapacityLeft());
diff --git a/Engine/Source/Programs/UnrealBuildAccelerator/Common/Private/UbaSessionServer.cpp b/Engine/Source/Programs/UnrealBuildAccelerator/Common/Private/UbaSessionServer.cpp
--- a/Engine/Source/Programs/UnrealBuildAccelerator/Common/Private/UbaSessionServer.cpp
+++ b/Engine/Source/Programs/UnrealBuildAccelerator/Common/Private/UbaSessionServer.cpp
@@ -110,9 +110,9 @@ namespace uba
 		if (isUbtPath)
 			dir.Resize(dir.count - 7); // Remove native and slash
 		const tchar* archPath[2] = { TC("x64"), TC("arm64") };
-		if (!dir.EndsWith(archPath[IsArm]))
+		if (!dir.EndsWith(archPath[IsArmBinary]))
 			return reportError ? logger.Error(TC("Module dir is not under supported folder (%s) to be able to run cross architecutres, can't figure out matching x64/arm64 folder"), dir.data) : false;
-		dir.Resize(dir.count - TStrlen(archPath[IsArm])).Append(archPath[!IsArm]);
+		dir.Resize(dir.count - TStrlen(archPath[IsArmBinary])).Append(archPath[!IsArmBinary]);
 		if (isUbtPath)
 			dir.Append(PathSeparator).Append(TC("native"));
 		return true;
@@ -203,14 +203,14 @@ namespace uba
 		#if PLATFORM_WINDOWS
 		char temp[1024];
 		detoursFile.Parse(temp, sizeof_array(temp));
-		m_detoursLibrary[IsArm] = temp;
+		m_detoursLibrary[IsArmBinary] = temp;
 		if (GetCrossArchitectureDir(m_logger, detoursFile.Resize(dirLength), false))
 		{
 			detoursFile.Append(PathSeparator).Append(UBA_DETOURS_LIBRARY).Parse(temp, sizeof_array(temp));
-			m_detoursLibrary[!IsArm] = temp;
+			m_detoursLibrary[!IsArmBinary] = temp;
 		}
 		#else
-		m_detoursLibrary[IsArm] = detoursFile.data;
+		m_detoursLibrary[IsArmBinary] = detoursFile.data;
 		#endif
 
 		if (!Create(info))
@@ -682,7 +682,7 @@ namespace uba
 				StringBuffer<> agentDir;
 				if (!GetDirectoryOfCurrentModule(m_logger, agentDir))
 					return false;
-				if (IsArm != isClientArm)
+				if (IsArmBinary != isClientArm)
 					if (!GetCrossArchitectureDir(m_logger, agentDir, true))
 						return false;
 				UBA_ASSERT(IsWindows);
@@ -801,7 +801,7 @@ namespace uba
 
 		auto FixCrossArchitecture = [this, clientIsArm](StringBuffer<>& absoluteFile)
 			{
-				if (clientIsArm == IsArm)
+				if (clientIsArm == IsArmBinary)
 					return;
 				for (auto& mapping : m_crossArchitectureMappings)
 					if (absoluteFile.StartsWith(mapping.from.c_str()))
@@ -1185,7 +1185,7 @@ namespace uba
 		ClientSession& session = *m_clientSessions[sessionIndex];
 		sessionsLock.Leave();
 
-		bool isCrossArchitecture = IsArm != session.isArm;
+		bool isCrossArchitecture = IsArmBinary != session.isArm;
 
 		u32 weight32 = reader.ReadU32();
 		float availableWeight = *(float*)&weight32;
@@ -1897,7 +1897,7 @@ namespace uba
 			if (hasCalledCallback)
 				return nullptr;
 
-			m_remoteProcessSlotAvailableEvent(IsArm != session.isArm);
+			m_remoteProcessSlotAvailableEvent(IsArmBinary != session.isArm);
 			hasCalledCallback = true;
 		}
 		return nullptr;
@@ -2295,7 +2295,7 @@ namespace uba
 		if (!h.IsValid())
 			return m_logger.Error(TC("Failed to find process for id %u when receiving update environment message"), processId);
 
-		if (IsArm != isClientArm)
+		if (IsArmBinary != isClientArm)
 		{
 			writer.WriteString(TC("Can't resolve callstack on cross architectures"));
 			return true;
diff --git a/Engine/Source/Programs/UnrealBuildAccelerator/Core/Private/UbaPlatform.cpp b/Engine/Source/Programs/UnrealBuildAccelerator/Core/Private/UbaPlatform.cpp
--- a/Engine/Source/Programs/UnrealBuildAccelerator/Core/Private/UbaPlatform.cpp
+++ b/Engine/Source/Programs/UnrealBuildAccelerator/Core/Private/UbaPlatform.cpp
@@ -86,6 +86,21 @@ namespace uba
 		return GetWineVersion() != nullptr;
 	}
 
+	bool IsRunningArm64()
+	{
+		#if PLATFORM_WINDOWS
+		static bool isArm = []()
+		{
+			SYSTEM_INFO SystemInfo;
+			GetSystemInfo(&SystemInfo);
+			return SystemInfo.wProcessorArchitecture == PROCESSOR_ARCHITECTURE_ARM;
+		}();
+		return isArm;
+		#else
+		return IsArmBinary;
+		#endif
+	}
+
 	void Sleep(u32 milliseconds)
 	{
 		#if PLATFORM_WINDOWS
diff --git a/Engine/Source/Programs/UnrealBuildAccelerator/Core/Public/UbaBase.h b/Engine/Source/Programs/UnrealBuildAccelerator/Core/Public/UbaBase.h
--- a/Engine/Source/Programs/UnrealBuildAccelerator/Core/Public/UbaBase.h
+++ b/Engine/Source/Programs/UnrealBuildAccelerator/Core/Public/UbaBase.h
@@ -32,9 +32,9 @@ namespace uba
 	#endif
 
 	#if PLATFORM_WINDOWS && defined(_M_ARM64) // For now we only care about windows arm..
-	constexpr bool IsArm = true;
+	constexpr bool IsArmBinary = true;
 	#else
-	constexpr bool IsArm = false;
+	constexpr bool IsArmBinary = false;
 	#endif
 
 #if PLATFORM_WINDOWS && defined(UBA_BUILD)
diff --git a/Engine/Source/Programs/UnrealBuildAccelerator/Core/Public/UbaPlatform.h b/Engine/Source/Programs/UnrealBuildAccelerator/Core/Public/UbaPlatform.h
--- a/Engine/Source/Programs/UnrealBuildAccelerator/Core/Public/UbaPlatform.h
+++ b/Engine/Source/Programs/UnrealBuildAccelerator/Core/Public/UbaPlatform.h
@@ -81,7 +81,7 @@ namespace uba
 	#define UBA_NOT_IMPLEMENTED(x) UBA_ASSERTF(false, TC("%s not implemented!"), #x);
 
 	bool WriteCallstackInfo(BinaryWriter& out, u32 skipCallstack, void* context);
-	void WriteAssertInfo(StringBufferBase& out, const tchar* text, const char* file, u32 line, const char* expr);
+	void WriteAssertInfo(StringBufferBase& out, const tchar* text, const char* file, u32 line, const char* expr, void* context = nullptr);
 	ANALYSIS_NORETURN void UbaAssert(const tchar* text, const char* file, u32 line, const char* expr, bool allowTerminate, u32 terminateCode, void* context);
 	ANALYSIS_NORETURN void FatalError(u32 code, const tchar* format, ...);
 	using CustomAssertHandler = void(const tchar* text);
@@ -89,6 +89,7 @@ namespace uba
 	bool CreateGuid(Guid& out);
 	const char* GetWineVersion();
 	bool IsRunningWine();
+	bool IsRunningArm64();
 	bool IsEscapePressed();
 	void Sleep(u32 milliseconds);
 	u32 GetUserDefaultUILanguage();
