[UBA]
* Fixed so x64 uba binaries can be used on arm64 windows

diff --git a/Engine/Source/Programs/UnrealBuildAccelerator/Common/Private/UbaProcess.cpp b/Engine/Source/Programs/UnrealBuildAccelerator/Common/Private/UbaProcess.cpp
--- a/Engine/Source/Programs/UnrealBuildAccelerator/Common/Private/UbaProcess.cpp
+++ b/Engine/Source/Programs/UnrealBuildAccelerator/Common/Private/UbaProcess.cpp
@@ -1450,9 +1450,7 @@ namespace uba
 
 		if (!m_parentProcess)
 		{
-			bool isArmApplication = IsArmBinary;
-#if defined(_M_ARM64)
-			bool isX64Application;
+			if (IsRunningArm())
 			{
 				SCOPED_WRITE_LOCK(m_session.m_isX64ApplicationLock, lock);
 				auto insres = m_session.m_isX64Application.try_emplace(m_realApplication);
@@ -1463,16 +1461,10 @@ namespace uba
 						return UBA_EXIT_CODE(19);
 					insres.first->second = imageInfo.isX64;
 				}
-				isX64Application = insres.first->second;
+				m_isArm = !insres.first->second;
 			}
-			if (isX64Application)
-			{
-				isArmApplication = false;
-				allowCustomAllocator = false;
-			}
-#endif
 
-			const char* detoursLib = m_session.m_detoursLibrary[isArmApplication].c_str();
+			const char* detoursLib = m_session.m_detoursLibrary[m_isArm].c_str();
 			if (!*detoursLib)
 				detoursLib = UBA_DETOURS_LIBRARY_ANSI;
 
@@ -1864,7 +1856,7 @@ namespace uba
 
 			if (m_detourEnabled)
 			{
-				const char* detoursLib = m_session.m_detoursLibrary[IsArm].c_str();
+				const char* detoursLib = m_session.m_detoursLibrary[IsArmBinary].c_str();
 				if (*detoursLib)
 				{
 #if PLATFORM_LINUX
diff --git a/Engine/Source/Programs/UnrealBuildAccelerator/Common/Private/UbaSessionClient.cpp b/Engine/Source/Programs/UnrealBuildAccelerator/Common/Private/UbaSessionClient.cpp
--- a/Engine/Source/Programs/UnrealBuildAccelerator/Common/Private/UbaSessionClient.cpp
+++ b/Engine/Source/Programs/UnrealBuildAccelerator/Common/Private/UbaSessionClient.cpp
@@ -208,7 +208,7 @@ namespace uba
 			UBA_ASSERT(fileNameKey == StringKeyZero);
 			StackBinaryWriter<1024> writer;
 			NetworkMessage msg(m_client, ServiceId, SessionMessageType_EnsureBinaryFile, writer);
-			writer.WriteBool(IsArmBinary);
+			writer.WriteBool(IsRunningArm());
 			writer.WriteString(fileName);
 			writer.WriteStringKey(fileNameKey);
 			writer.WriteString(applicationDir);
@@ -1180,7 +1180,7 @@ namespace uba
 		NetworkMessage msg(m_client, ServiceId, SessionMessageType_Connect, writer);
 		writer.WriteString(m_name.data);
 		writer.WriteU32(SessionNetworkVersion);
-		writer.WriteBool(IsArmBinary);
+		writer.WriteBool(IsRunningArm());
 
 		CasKey keys[2];
 		if (m_useBinariesAsVersion)
@@ -1222,16 +1222,18 @@ namespace uba
 			return;
 		}
 
+		u32 isArm = IsRunningArm();
+
 		CasKey detoursBinaryKey[2];
 		detoursBinaryKey[0] = reader.ReadCasKey();
-		if (IsArmBinary)
+		if (isArm)
 			detoursBinaryKey[1] = reader.ReadCasKey();
 
 		StringBuffer<> detoursFile[2];
 
 		if (m_downloadDetoursLib)
 		{
-			for (u32 i=0; i!=IsArmBinary+1; ++i)
+			for (u32 i=0; i!=isArm+1; ++i)
 			{
 				{
 					TimerScope s(m_stats.storageRetrieve);
@@ -1249,11 +1251,11 @@ namespace uba
 		}
 		else
 		{
-			GetDirectoryOfCurrentModule(m_logger, detoursFile[IsArmBinary]);
-			detoursFile[IsArmBinary].EnsureEndsWithSlash().Append(UBA_DETOURS_LIBRARY);
+			GetDirectoryOfCurrentModule(m_logger, detoursFile[isArm]);
+			detoursFile[isArm].EnsureEndsWithSlash().Append(UBA_DETOURS_LIBRARY);
 		}
 
-		for (u32 i=0; i!=IsArmBinary+1; ++i)
+		for (u32 i=0; i!=isArm+1; ++i)
 		{
 			#if PLATFORM_WINDOWS
 			char dll[1024];
@@ -2143,7 +2145,7 @@ namespace uba
 		StackBinaryWriter<SendMaxSize> msgWriter;
 		NetworkMessage msg(m_client, ServiceId, SessionMessageType_GetSymbols, msgWriter);
 		msgWriter.WriteU32(process.GetId());
-		msgWriter.WriteBool(IsArmBinary);
+		msgWriter.WriteBool(process.IsArm());
 		msgWriter.WriteBytes(reader.GetPositionData(), reader.GetLeft());
 
 		BinaryReader responseReader(writer.GetData(), 0, writer.GetCapacityLeft());
diff --git a/Engine/Source/Programs/UnrealBuildAccelerator/Common/Public/UbaProcess.h b/Engine/Source/Programs/UnrealBuildAccelerator/Common/Public/UbaProcess.h
--- a/Engine/Source/Programs/UnrealBuildAccelerator/Common/Public/UbaProcess.h
+++ b/Engine/Source/Programs/UnrealBuildAccelerator/Common/Public/UbaProcess.h
@@ -27,6 +27,7 @@ namespace uba
 		virtual bool IsRemote() const { UBA_ASSERT(false); return false; }
 		virtual ProcessExecutionType GetExecutionType() const { UBA_ASSERT(false); return ProcessExecutionType_Native; }
 		virtual bool IsChild() { UBA_ASSERT(false); return false; }
+		virtual bool IsArm() { UBA_ASSERT(false); return false; }
 
 	protected:
 		void AddRef();
@@ -55,6 +56,7 @@ namespace uba
 		virtual bool IsRemote() const override { return false; }
 		virtual ProcessExecutionType GetExecutionType() const override { return m_detourEnabled ? ProcessExecutionType_Detoured : ProcessExecutionType_Native; }
 		virtual bool IsChild() override { return m_parentProcess != nullptr; }
+		virtual bool IsArm() override { return m_isArm; }
 
 		ProcessImpl(Session& session, u32 id, ProcessImpl* parent, bool detourEnabled);
 		~ProcessImpl();
@@ -138,6 +140,7 @@ namespace uba
 		bool m_gotExitMessage = false;
 		bool m_parentReportedExit = false;
 		bool m_detourEnabled = true;
+		bool m_isArm = false;
 		TString m_realApplication;
 		const tchar* m_realWorkingDir = nullptr;
 		u64 m_startTime = 0;
diff --git a/Engine/Source/Programs/UnrealBuildAccelerator/Common/Public/UbaSession.h b/Engine/Source/Programs/UnrealBuildAccelerator/Common/Public/UbaSession.h
--- a/Engine/Source/Programs/UnrealBuildAccelerator/Common/Public/UbaSession.h
+++ b/Engine/Source/Programs/UnrealBuildAccelerator/Common/Public/UbaSession.h
@@ -299,7 +299,7 @@ namespace uba
 		Logger* m_debugLogger = nullptr;
 		#endif
 
-		#if PLATFORM_WINDOWS && defined(_M_ARM64)
+		#if PLATFORM_WINDOWS
 		ReaderWriterLock m_isX64ApplicationLock;
 		UnorderedMap<TString, bool> m_isX64Application;
 		#endif
diff --git a/Engine/Source/Programs/UnrealBuildAccelerator/Core/Private/UbaPlatform.cpp b/Engine/Source/Programs/UnrealBuildAccelerator/Core/Private/UbaPlatform.cpp
--- a/Engine/Source/Programs/UnrealBuildAccelerator/Core/Private/UbaPlatform.cpp
+++ b/Engine/Source/Programs/UnrealBuildAccelerator/Core/Private/UbaPlatform.cpp
@@ -86,7 +86,7 @@ namespace uba
 		return GetWineVersion() != nullptr;
 	}
 
-	bool IsRunningArm64()
+	bool IsRunningArm()
 	{
 		#if PLATFORM_WINDOWS
 		static bool isArm = []()
diff --git a/Engine/Source/Programs/UnrealBuildAccelerator/Core/Public/UbaPlatform.h b/Engine/Source/Programs/UnrealBuildAccelerator/Core/Public/UbaPlatform.h
--- a/Engine/Source/Programs/UnrealBuildAccelerator/Core/Public/UbaPlatform.h
+++ b/Engine/Source/Programs/UnrealBuildAccelerator/Core/Public/UbaPlatform.h
@@ -89,7 +89,7 @@ namespace uba
 	bool CreateGuid(Guid& out);
 	const char* GetWineVersion();
 	bool IsRunningWine();
-	bool IsRunningArm64();
+	bool IsRunningArm();
 	bool IsEscapePressed();
 	void Sleep(u32 milliseconds);
 	u32 GetUserDefaultUILanguage();
