[UBA]
* Fixed so IsRunningArm checks the right processor architecture

diff --git a/Engine/Source/Programs/UnrealBuildAccelerator/Core/Private/UbaPlatform.cpp b/Engine/Source/Programs/UnrealBuildAccelerator/Core/Private/UbaPlatform.cpp
--- a/Engine/Source/Programs/UnrealBuildAccelerator/Core/Private/UbaPlatform.cpp
+++ b/Engine/Source/Programs/UnrealBuildAccelerator/Core/Private/UbaPlatform.cpp
@@ -93,7 +93,7 @@ namespace uba
 		{
 			SYSTEM_INFO SystemInfo;
 			GetSystemInfo(&SystemInfo);
-			return SystemInfo.wProcessorArchitecture == PROCESSOR_ARCHITECTURE_ARM;
+			return SystemInfo.wProcessorArchitecture == PROCESSOR_ARCHITECTURE_ARM64;
 		}();
 		return isArm;
 		#else
