[UBA]
* Disabled mimalloc for detoured x64 binaries running on arm64 system

diff --git a/Engine/Source/Programs/UnrealBuildAccelerator/Common/Private/UbaProcess.cpp b/Engine/Source/Programs/UnrealBuildAccelerator/Common/Private/UbaProcess.cpp
--- a/Engine/Source/Programs/UnrealBuildAccelerator/Common/Private/UbaProcess.cpp
+++ b/Engine/Source/Programs/UnrealBuildAccelerator/Common/Private/UbaProcess.cpp
@@ -1461,10 +1461,13 @@ namespace uba
 						return UBA_EXIT_CODE(19);
 					insres.first->second = imageInfo.isX64;
 				}
-				m_isArm = !insres.first->second;
+				m_isArmBinary = !insres.first->second;
+
+				if (!m_isArmBinary)
+					allowCustomAllocator = false;
 			}
 
-			const char* detoursLib = m_session.m_detoursLibrary[m_isArm].c_str();
+			const char* detoursLib = m_session.m_detoursLibrary[m_isArmBinary].c_str();
 			if (!*detoursLib)
 				detoursLib = UBA_DETOURS_LIBRARY_ANSI;
 
diff --git a/Engine/Source/Programs/UnrealBuildAccelerator/Common/Public/UbaProcess.h b/Engine/Source/Programs/UnrealBuildAccelerator/Common/Public/UbaProcess.h
--- a/Engine/Source/Programs/UnrealBuildAccelerator/Common/Public/UbaProcess.h
+++ b/Engine/Source/Programs/UnrealBuildAccelerator/Common/Public/UbaProcess.h
@@ -56,7 +56,7 @@ namespace uba
 		virtual bool IsRemote() const override { return false; }
 		virtual ProcessExecutionType GetExecutionType() const override { return m_detourEnabled ? ProcessExecutionType_Detoured : ProcessExecutionType_Native; }
 		virtual bool IsChild() override { return m_parentProcess != nullptr; }
-		virtual bool IsArm() override { return m_isArm; }
+		virtual bool IsArm() override { return m_isArmBinary; }
 
 		ProcessImpl(Session& session, u32 id, ProcessImpl* parent, bool detourEnabled);
 		~ProcessImpl();
@@ -140,7 +140,7 @@ namespace uba
 		bool m_gotExitMessage = false;
 		bool m_parentReportedExit = false;
 		bool m_detourEnabled = true;
-		bool m_isArm = false;
+		bool m_isArmBinary = false;
 		TString m_realApplication;
 		const tchar* m_realWorkingDir = nullptr;
 		u64 m_startTime = 0;
