[UBA]
* Improve error output when getting a seh/segfault

diff --git a/Engine/Source/Programs/UnrealBuildAccelerator/Core/Private/UbaPlatform.cpp b/Engine/Source/Programs/UnrealBuildAccelerator/Core/Private/UbaPlatform.cpp
--- a/Engine/Source/Programs/UnrealBuildAccelerator/Core/Private/UbaPlatform.cpp
+++ b/Engine/Source/Programs/UnrealBuildAccelerator/Core/Private/UbaPlatform.cpp
@@ -165,7 +165,9 @@ namespace uba
 			signal = *(int*)context;
 		#endif
 
-		out.Append(TC("ASSERT: ")).Append(*text ? text : TC("Unknown"));
+		if (!context)
+			out.Append(TC("ASSERT: "));
+		out.Append(*text ? text : TC("Unknown"));
 		if (!*text && expr && *expr && strcmp(expr, "false") != 0)
 			out.Appendf(TC("\n EXPR: ")).Append(expr);
 		if (file && *file)
@@ -546,7 +548,7 @@ namespace uba
 		CONTEXT* contextPtr = ExceptionInfo->ContextRecord;
 
 		StringBuffer<> text;
-		text.Appendf(L"Unhandled Exception (Code: 0x%x)", exceptionCode);
+		text.Appendf(L"ERROR: Unhandled Exception (Code: 0x%x)", exceptionCode);
 		UbaAssert(text.data, nullptr, 0, nullptr, true, exceptionCode, contextPtr);
 		return EXCEPTION_EXECUTE_HANDLER;
 	}
@@ -556,7 +558,7 @@ namespace uba
 		if (code != EXCEPTION_ACCESS_VIOLATION)
 			return EXCEPTION_CONTINUE_SEARCH;
 		StringBuffer<> buf;
-		buf.Appendf(L"Access violation reading %s", fileName);
+		buf.Appendf(L"ERROR: Access violation reading %s", fileName);
 		UbaAssert(buf.data, __FILE__, __LINE__, "", true, code, info->ContextRecord);
 		return EXCEPTION_EXECUTE_HANDLER;
 	}
