[UBA]
* Changed so it is fine to not find cross architecture mappings since x64 binaries now work to detour on remotes

diff --git a/Engine/Source/Programs/UnrealBuildAccelerator/Common/Private/UbaSessionServer.cpp b/Engine/Source/Programs/UnrealBuildAccelerator/Common/Private/UbaSessionServer.cpp
--- a/Engine/Source/Programs/UnrealBuildAccelerator/Common/Private/UbaSessionServer.cpp
+++ b/Engine/Source/Programs/UnrealBuildAccelerator/Common/Private/UbaSessionServer.cpp
@@ -1227,19 +1227,19 @@ namespace uba
 						break;
 					}
 
-				if (!applicationOverride.count)
+				if (applicationOverride.count)
 				{
-					m_logger.Info(TC("Couldn't find cross architecture mapping for %s"), startInfo.application);
-					returnProcess();
-					break;
+					//m_logger.Info(TC("Couldn't find cross architecture mapping for %s"), startInfo.application);
+					//returnProcess();
+					//break;
+					if (!FileExists(m_logger, applicationOverride.data))
+					{
+						m_logger.Info(TC("Couldn't find cross architecture executable %s"), applicationOverride.data);
+						returnProcess();
+						break;
+					}
 				}
 
-				if (!FileExists(m_logger, applicationOverride.data))
-				{
-					m_logger.Info(TC("Couldn't find cross architecture executable %s"), applicationOverride.data);
-					returnProcess();
-					break;
-				}
 			}
 
 			ProcessAdded(*process, sessionId);
