[UBA]
* Fixed so IsRunningArm uses IsWow64Process2 since GetSystemInfo returned the wrong architecture

diff --git a/Engine/Source/Programs/UnrealBuildAccelerator/Core/Private/UbaPlatform.cpp b/Engine/Source/Programs/UnrealBuildAccelerator/Core/Private/UbaPlatform.cpp
--- a/Engine/Source/Programs/UnrealBuildAccelerator/Core/Private/UbaPlatform.cpp
+++ b/Engine/Source/Programs/UnrealBuildAccelerator/Core/Private/UbaPlatform.cpp
@@ -91,9 +91,10 @@ namespace uba
 		#if PLATFORM_WINDOWS
 		static bool isArm = []()
 		{
-			SYSTEM_INFO SystemInfo;
-			GetSystemInfo(&SystemInfo);
-			return SystemInfo.wProcessorArchitecture == PROCESSOR_ARCHITECTURE_ARM64;
+			USHORT processMachine, nativeMachine;
+			if (!IsWow64Process2(GetCurrentProcess(), &processMachine, &nativeMachine))
+				return false;
+			return nativeMachine == IMAGE_FILE_MACHINE_ARM64;
 		}();
 		return isArm;
 		#else
