[UBA]
* Added warnings when running x64 versions of uba on arm64 system

diff --git a/Engine/Source/Programs/UnrealBuildAccelerator/Agent/Private/UbaAgent.cpp b/Engine/Source/Programs/UnrealBuildAccelerator/Agent/Private/UbaAgent.cpp
--- a/Engine/Source/Programs/UnrealBuildAccelerator/Agent/Private/UbaAgent.cpp
+++ b/Engine/Source/Programs/UnrealBuildAccelerator/Agent/Private/UbaAgent.cpp
@@ -800,6 +800,10 @@ namespace uba
 				
 		if (!eventFile.IsEmpty())
 			logger.Info(TC("  Will poll for external events in file %s"), eventFile.data);
+
+		if (!IsArmBinary && IsRunningArm())
+			logger.Warning(TC("  Running x64 binary on arm64 system. Use arm binaries instead"));
+
 		logger.Info(TC(""));
 
 
diff --git a/Engine/Source/Programs/UnrealBuildAccelerator/Cli/Private/UbaCli.cpp b/Engine/Source/Programs/UnrealBuildAccelerator/Cli/Private/UbaCli.cpp
--- a/Engine/Source/Programs/UnrealBuildAccelerator/Cli/Private/UbaCli.cpp
+++ b/Engine/Source/Programs/UnrealBuildAccelerator/Cli/Private/UbaCli.cpp
@@ -442,6 +442,9 @@ namespace uba
 		FilteredLogWriter logWriter(g_consoleLogWriter, quiet ? LogEntryType_Warning : LogEntryType_Detail);
 		LoggerWithWriter logger(logWriter, TC(""));
 
+		if (!IsArmBinary && IsRunningArm())
+			logger.Warning(TC("  Running x64 binary on arm64 system. Use arm binaries instead"));
+
 		bool exit = false;
 		if (deleteCas)
 		{
diff --git a/Engine/Source/Programs/UnrealBuildAccelerator/Host/Private/UbaExports.cpp b/Engine/Source/Programs/UnrealBuildAccelerator/Host/Private/UbaExports.cpp
--- a/Engine/Source/Programs/UnrealBuildAccelerator/Host/Private/UbaExports.cpp
+++ b/Engine/Source/Programs/UnrealBuildAccelerator/Host/Private/UbaExports.cpp
@@ -124,6 +124,11 @@ extern "C"
 	uba::NetworkServer* NetworkServer_Create(uba::LogWriter& writer, uba::u32 workerCount, uba::u32 sendSize, uba::u32 receiveTimeoutSeconds, bool useQuic)
 	{
 		using namespace uba;
+
+		if (!IsArmBinary && IsRunningArm())
+			LoggerWithWriter(writer, TC("")).Warning(TC("  Running x64 binary on arm64 system. Use arm binaries instead"));
+
+
 		NetworkBackend* networkBackend;
 		#if UBA_USE_QUIC
 		if (useQuic)
