[UBA]
* Improved detouring log entries

diff --git a/Engine/Source/Programs/UnrealBuildAccelerator/Detours/Private/Windows/UbaDetoursFunctionsKernelBase.inl b/Engine/Source/Programs/UnrealBuildAccelerator/Detours/Private/Windows/UbaDetoursFunctionsKernelBase.inl
--- a/Engine/Source/Programs/UnrealBuildAccelerator/Detours/Private/Windows/UbaDetoursFunctionsKernelBase.inl
+++ b/Engine/Source/Programs/UnrealBuildAccelerator/Detours/Private/Windows/UbaDetoursFunctionsKernelBase.inl
@@ -519,7 +519,7 @@ BOOL Detoured_GetVolumeInformationByHandleW(HANDLE hFile, LPWSTR lpVolumeNameBuf
 			wcscpy_s(lpFileSystemNameBuffer, nFileSystemNameSize, L"NTFS"); // TODO: Not everyone has NTFS?
 		}
 		SetLastError(ERROR_SUCCESS);
-		DEBUG_LOG_DETOURED(L"GetVolumeInformationByHandleW", L"%llu %ls", uintptr_t(hFile), HandleToName(hFile));
+		DEBUG_LOG_DETOURED(L"GetVolumeInformationByHandleW", L"%llu (Serial: %u) (%ls) -> Success", uintptr_t(hFile), entryInfo.volumeSerial, HandleToName(hFile));
 		return true;
 	}
 	bool res = True_GetVolumeInformationByHandleW(trueHandle, lpVolumeNameBuffer, nVolumeNameSize, lpVolumeSerialNumber, lpMaximumComponentLength, lpFileSystemFlags, lpFileSystemNameBuffer, nFileSystemNameSize);
@@ -1881,7 +1881,6 @@ BOOL Detoured_GetFileInformationByHandleEx(HANDLE hFile, FILE_INFO_BY_HANDLE_CLA
 
 	if (entryOffset != ~u32(0))
 	{
-		DEBUG_LOG_DETOURED(L"GetFileInformationByHandleEx", L"(DIRTABLE) (%u) %llu (%ls)", fileInformationClass, uintptr_t(hFile), HandleToName(hFile));
 		DirectoryTable::EntryInformation entryInfo;
 		g_directoryTable.GetEntryInformation(entryInfo, entryOffset);
 		if (fileInformationClass == FileBasicInfo)
@@ -1892,6 +1891,7 @@ BOOL Detoured_GetFileInformationByHandleEx(HANDLE hFile, FILE_INFO_BY_HANDLE_CLA
 			data.LastWriteTime = ToLargeInteger(entryInfo.lastWrite);
 			data.ChangeTime = ToLargeInteger(entryInfo.lastWrite);
 			data.FileAttributes = entryInfo.attributes;
+			DEBUG_LOG_DETOURED(L"GetFileInformationByHandleEx", L"(DIRTABLE) (FileBasicInfo) %llu (%ls)", uintptr_t(hFile), HandleToName(hFile));
 			return TRUE;
 		}
 		else if (fileInformationClass == FileIdInfo)
@@ -1901,6 +1901,7 @@ BOOL Detoured_GetFileInformationByHandleEx(HANDLE hFile, FILE_INFO_BY_HANDLE_CLA
 			u64* id = (u64*)&data.FileId;
 			id[0] = 0;
 			id[1] = entryInfo.fileIndex;
+			DEBUG_LOG_DETOURED(L"GetFileInformationByHandleEx", L"(DIRTABLE) (FileIdInfo) %llu (VolumeSerial: %u FileIndex: %llu), (%ls)", uintptr_t(hFile), entryInfo.volumeSerial, entryInfo.fileIndex, HandleToName(hFile));
 			return TRUE;
 		}
 		else if (fileInformationClass == FileStandardInfo)
@@ -1932,7 +1933,7 @@ BOOL Detoured_GetFileInformationByHandleEx(HANDLE hFile, FILE_INFO_BY_HANDLE_CLA
 			}
 #endif
 
-
+			DEBUG_LOG_DETOURED(L"GetFileInformationByHandleEx", L"(DIRTABLE) (FileStandardInfo) %llu (%ls)", uintptr_t(hFile), HandleToName(hFile));
 			return TRUE;
 		}
 		else if (fileInformationClass == FileRemoteProtocolInfo)
@@ -1955,6 +1956,7 @@ BOOL Detoured_GetFileInformationByHandleEx(HANDLE hFile, FILE_INFO_BY_HANDLE_CLA
 			auto& data = *(FILE_ATTRIBUTE_TAG_INFO*)lpFileInformation;
 			data.FileAttributes = entryInfo.attributes;
 			data.ReparseTag = 0;
+			DEBUG_LOG_DETOURED(L"GetFileInformationByHandleEx", L"(DIRTABLE) (FileAttributeTagInfo) %llu (%ls)", uintptr_t(hFile), HandleToName(hFile));
 			return TRUE;
 		}
 		else
