[UBA]
* Set bRetainFramePointers to true for all targets when building UBA.

diff --git a/Engine/Source/Programs/UnrealBuildAccelerator/UbaAgent.Target.cs b/Engine/Source/Programs/UnrealBuildAccelerator/UbaAgent.Target.cs
--- a/Engine/Source/Programs/UnrealBuildAccelerator/UbaAgent.Target.cs
+++ b/Engine/Source/Programs/UnrealBuildAccelerator/UbaAgent.Target.cs
@@ -19,6 +19,8 @@ public class UbaAgentTarget : TargetRules
 		Rules.UndecoratedConfiguration = UnrealTargetConfiguration.Shipping;
 		Rules.bHasExports = false;
 
+		Rules.bRetainFramePointers = true;
+
 		// Lean and mean
 		Rules.bBuildDeveloperTools = false;
 
