[UBA]
* Fixed so processor frequency shows correctly on posix

diff --git a/Engine/Source/Programs/UnrealBuildAccelerator/Common/Private/UbaSession.cpp b/Engine/Source/Programs/UnrealBuildAccelerator/Common/Private/UbaSession.cpp
--- a/Engine/Source/Programs/UnrealBuildAccelerator/Common/Private/UbaSession.cpp
+++ b/Engine/Source/Programs/UnrealBuildAccelerator/Common/Private/UbaSession.cpp
@@ -2976,8 +2976,19 @@ namespace uba
 		u64 throwAway;
 		GetMemoryInfo(throwAway, totalMemoryInKilobytes);
 		totalMemoryInKilobytes /= 1024;
-		// TODO get freq from /proc/cpuinfo
-		hzStr.Append(" @ ?GHz");
+		double processorMhz = 0.0;
+		if (FILE *fp = fopen("/proc/cpuinfo", "r"))
+		{
+			char line[256];
+			while (processorMhz == 0.0 && fgets(line, sizeof(line), fp))
+				if (strncmp(line, "cpu MHz", 7) == 0)
+					sscanf(line, "cpu MHz : %lf", &processorMhz);
+			fclose(fp);
+		}
+		if (processorMhz != 0.0)
+			hzStr.Appendf(" @ %.1fGHz", processorMhz/1000.0);
+		else
+			hzStr.Append(" @ ?GHz");
 		#endif
 
 		const tchar* capacityStr = TC("NoLimit");
