[UBA]
* Implemented osx code for getting cpu frequency and physical memory
* Improved trace machine info

diff --git a/Engine/Source/Programs/UnrealBuildAccelerator/Agent/Private/UbaAgent.cpp b/Engine/Source/Programs/UnrealBuildAccelerator/Agent/Private/UbaAgent.cpp
--- a/Engine/Source/Programs/UnrealBuildAccelerator/Agent/Private/UbaAgent.cpp
+++ b/Engine/Source/Programs/UnrealBuildAccelerator/Agent/Private/UbaAgent.cpp
@@ -776,9 +776,6 @@ namespace uba
 		if (!zone.count)
 			GetZone(zone);
 
-		if (zone.count)
-			extraInfo.Append(TC(", ")).Append(zone);
-
 		if (IsRunningWine())
 			extraInfo.Append(TC(", Linux/Wine-")).Append(GetWineVersion());
 
@@ -792,7 +789,7 @@ namespace uba
 		#if UBA_DEBUG
 		dbgStr = TC(" (DEBUG)");
 		#endif
-		logger.Info(TC("UbaAgent v%s%s (Cpu%s: %u, MaxCon: %u, Dir: \"%s\", StoreCapacity: %uGb%s)"), Version, dbgStr, (IsArmBinary ? TC("[Arm]") : TC("")), maxProcessCount, maxConnectionCount, g_rootDir.data, storageCapacityGb, extraInfo.data);
+		logger.Info(TC("UbaAgent v%s%s (Cpu%s: %u, MaxCon: %u, Dir: \"%s\", StoreCapacity: %uGb, Zone: %s%s)"), Version, dbgStr, (IsArmBinary ? TC("[Arm]") : TC("")), maxProcessCount, maxConnectionCount, g_rootDir.data, storageCapacityGb, zone.count ? zone.data : TC("none"), extraInfo.data);
 		
 		Config config;
 		if (!configFile.IsEmpty())
diff --git a/Engine/Source/Programs/UnrealBuildAccelerator/Common/Private/UbaSession.cpp b/Engine/Source/Programs/UnrealBuildAccelerator/Common/Private/UbaSession.cpp
--- a/Engine/Source/Programs/UnrealBuildAccelerator/Common/Private/UbaSession.cpp
+++ b/Engine/Source/Programs/UnrealBuildAccelerator/Common/Private/UbaSession.cpp
@@ -26,6 +26,7 @@
 #include <mach/mach_init.h>
 #include <mach/mach_host.h>
 #include <mach/mach.h>
+#include <sys/sysctl.h>
 extern char **environ;
 #else // PLATFORM_LINUX
 #include "UbaLinuxBinDependencyParser.h"
@@ -2945,15 +2946,17 @@ namespace uba
 		u32 cpuCount = GetLogicalProcessorCount();
 		u32 cpuGroupCount = GetProcessorGroupCount();
 
-		StringBuffer<128> cpuCountStr;
+		StringBuffer<128> cpuStr(TC("CPU"));
+		if (IsRunningArm())
+			cpuStr.Append(TC("[Arm]"));
+		cpuStr.Append(':');
 		if (cpuGroupCount != 1)
-			cpuCountStr.AppendValue(cpuGroupCount).Append('x');
-		cpuCountStr.AppendValue(cpuCount/cpuGroupCount);
+			cpuStr.AppendValue(cpuGroupCount).Append('x');
+		cpuStr.AppendValue(cpuCount/cpuGroupCount);
 
 		u64 totalMemoryInKilobytes = 0;
-		StringBuffer<128> hzStr;
 
-		#if PLATFORM_WINDOWS
+#if PLATFORM_WINDOWS
 		GetPhysicallyInstalledSystemMemory(&totalMemoryInKilobytes);
 
 		{
@@ -2969,14 +2972,15 @@ namespace uba
 				if (CallNtPowerInformation(ProcessorInformation, NULL, 0, procInfos.data(), cpuCount*sizeof(PROCESSOR_POWER_INFORMATION)) == STATUS_SUCCESS)
 					maxMHz = procInfos[0].MaxMhz;
 			}
-			hzStr.Appendf(TC(" @ %.1fGHz"), float(maxMHz) / 1000.0f);
+			cpuStr.Appendf(TC(" @ %.1fGHz"), float(maxMHz) / 1000.0f);
 		}
 
-		#else
+#else
 		u64 throwAway;
 		GetMemoryInfo(throwAway, totalMemoryInKilobytes);
 		totalMemoryInKilobytes /= 1024;
 		double processorMhz = 0.0;
+	#if PLATFORM_LINUX
 		if (FILE *fp = fopen("/proc/cpuinfo", "r"))
 		{
 			char line[256];
@@ -2984,19 +2988,22 @@ namespace uba
 				if (strncmp(line, "cpu MHz", 7) == 0)
 					sscanf(line, "cpu MHz : %lf", &processorMhz);
 			fclose(fp);
-		}
-		if (processorMhz != 0.0)
-			hzStr.Appendf(" @ %.1fGHz", processorMhz/1000.0);
-		else
-			hzStr.Append(" @ ?GHz");
-		#endif
+			cpuStr.Appendf(" @ %.1fGHz", processorMhz/1000.0);
+		}
+	#else
+		char brand[128];
+		size_t size = sizeof(brand);
+		if (sysctlbyname("machdep.cpu.brand_string", brand, &size, NULL, 0) == 0)
+			cpuStr.Clear().Append(brand).Appendf(" CPU:%u", cpuCount);
+	#endif
+#endif
 
 		const tchar* capacityStr = TC("NoLimit");
 		u64 capacity = m_storage.GetStorageCapacity();
 		BytesToText temp(capacity);
 		if (capacity)
 			capacityStr = temp.str;
-		out.Appendf(TC("CPU%s:%s%s Mem:%ugb Cas:%s"), (IsArmBinary ? TC("[Arm]") : TC("")), cpuCountStr.data, hzStr.data, u32(totalMemoryInKilobytes/(1024*1024)), capacityStr);
+		out.Appendf(TC("%s Mem:%ugb Cas:%s"), cpuStr.data, u32(totalMemoryInKilobytes/(1024*1024)), capacityStr);
 
 		StringBuffer<128> zone;
 		if (m_storage.GetZone(zone))
@@ -3041,7 +3048,7 @@ namespace uba
 			out.Append(m_extraInfo);
 
 		#if UBA_DEBUG
-		out.Append(TC(" DEBUG"));
+		out.Append(TC(" - DEBUG"));
 		#endif
 	}
 
@@ -3155,29 +3162,11 @@ namespace uba
 			outTotal = memStatus.ullTotalPageFile;
 			outAvailable = memStatus.ullAvailPageFile;
 		}
-#elif PLATFORM_MAC
-        vm_size_t page_size;
-            mach_port_t mach_port;
-            mach_msg_type_number_t count;
-            vm_statistics64_data_t vm_stats;
-
-            mach_port = mach_host_self();
-            count = sizeof(vm_stats) / sizeof(natural_t);
-            if (KERN_SUCCESS == host_page_size(mach_port, &page_size) &&
-                KERN_SUCCESS == host_statistics64(mach_port, HOST_VM_INFO,
-                                                (host_info64_t)&vm_stats, &count))
-            {
-                outAvailable = (int64_t)vm_stats.free_count * (int64_t)page_size;
-                outTotal = vm_stats.wire_count + vm_stats.active_count + vm_stats.inactive_count + vm_stats.free_count;
-//                long long used_memory = ((int64_t)vm_stats.active_count +
-//                                         (int64_t)vm_stats.inactive_count +
-//                                         (int64_t)vm_stats.wire_count) *  (int64_t)page_size;
-            }
-
 #else
-		static long mem = sysconf(_SC_PHYS_PAGES) * sysconf(_SC_PAGE_SIZE);
-		outTotal = mem;
-		outAvailable = mem;
+		u64 memKb;
+		GetPhysicallyInstalledSystemMemory(memKb);
+		outTotal = memKb*1024*1024;
+		outAvailable = outTotal;
 #endif
 		return true;
 	}
diff --git a/Engine/Source/Programs/UnrealBuildAccelerator/Core/Private/UbaPlatform.cpp b/Engine/Source/Programs/UnrealBuildAccelerator/Core/Private/UbaPlatform.cpp
--- a/Engine/Source/Programs/UnrealBuildAccelerator/Core/Private/UbaPlatform.cpp
+++ b/Engine/Source/Programs/UnrealBuildAccelerator/Core/Private/UbaPlatform.cpp
@@ -14,6 +14,7 @@
 #if PLATFORM_LINUX
 #include <link.h>
 #else
+#include <sys/sysctl.h>
 #include <mach-o/dyld.h>
 #endif
 #endif
@@ -630,6 +631,26 @@ namespace uba
 		out.Append("/uba_").AppendHex(uid);
 		#endif
 	}
+
+	bool GetPhysicallyInstalledSystemMemory(u64& outKb)
+	{
+	#if PLATFORM_MAC
+		static long mem = []() -> int64_t
+		{
+			int64_t mem;
+			size_t len = sizeof(mem);
+			int mib[] = {CTL_HW, HW_MEMSIZE};
+			if (sysctl(mib, 2, &mem, &len, NULL, 0) == 0)
+				return mem;
+			return 0;
+		}();
+	#else
+		static long mem = sysconf(_SC_PHYS_PAGES) * sysconf(_SC_PAGE_SIZE);
+	#endif
+		outKb = mem/(1024*1024);
+		return true;
+	}
+
 #endif
 
 	void BinaryWriter::WriteLongString(const StringView& str, u32 threshold)
diff --git a/Engine/Source/Programs/UnrealBuildAccelerator/Core/Public/UbaPlatform.h b/Engine/Source/Programs/UnrealBuildAccelerator/Core/Public/UbaPlatform.h
--- a/Engine/Source/Programs/UnrealBuildAccelerator/Core/Public/UbaPlatform.h
+++ b/Engine/Source/Programs/UnrealBuildAccelerator/Core/Public/UbaPlatform.h
@@ -157,6 +157,7 @@ namespace uba
 	void strcpy_s(tchar* dest, u64 destCapacity, const tchar* source);
 	void strcat_s(tchar* dest, u64 destCapacity, const tchar* source);
 	void GetMappingHandleName(StringBufferBase& out, u64 uid);
+	bool GetPhysicallyInstalledSystemMemory(u64& outKb);
 	inline u64 FromTimeSpec(const timespec& ts) { return u64(ts.tv_sec) * 10'000'000ull + u64(ts.tv_nsec/100); }
 	inline timespec ToTimeSpec(u64 time) { timespec ts; ts.tv_sec = time / 10'000'000ull; ts.tv_nsec = (time - (u64(ts.tv_sec) * 10'000'000ull)) * 100; return ts; }
 	#if PLATFORM_LINUX
