[UBA]
* Changed so GetComputerName uses dns provided name and fallback to NetBIOS if not existing.

diff --git a/Engine/Source/Programs/UnrealBuildAccelerator/Agent/Private/UbaAgent.cpp b/Engine/Source/Programs/UnrealBuildAccelerator/Agent/Private/UbaAgent.cpp
--- a/Engine/Source/Programs/UnrealBuildAccelerator/Agent/Private/UbaAgent.cpp
+++ b/Engine/Source/Programs/UnrealBuildAccelerator/Agent/Private/UbaAgent.cpp
@@ -51,7 +51,6 @@ namespace uba
 		return buf;
 	}();
 	u32				DefaultProcessorCount = []() { return GetLogicalProcessorCount(); }();
-	const tchar*	DefaultAgentName = []() { static tchar buf[256]; GetComputerNameW(buf, sizeof_array(buf)); return buf; }();
 	u32				DefaultMaxConnectionCount = 4;
 
 	int PrintHelp(const tchar* message)
@@ -62,6 +61,8 @@ namespace uba
 			logger.Info(TC(""));
 			logger.Error(TC("%s"), message);
 		}
+		StringBuffer<> name;
+		GetComputerNameW(name);
 		logger.Info(TC(""));
 		logger.Info(TC("-------------------------------------------"));
 		logger.Info(TC("   UbaAgent v%s%s"), Version, (IsArmBinary ? TC(" (ARM64)") : TC("")));
@@ -82,7 +83,7 @@ namespace uba
 		logger.Info(TC("  -capacity=<gigaby>      Capacity of local store. Defaults to %u gigabytes"), DefaultCapacityGb);
 		logger.Info(TC("  -config=<file>          Config file that contains options for various systems"));
 		logger.Info(TC("  -quic                   Use Quic instead of tcp backend."));
-		logger.Info(TC("  -name=<name>            The identifier of this agent. Defaults to \"%s\" on this machine"), DefaultAgentName);
+		logger.Info(TC("  -name=<name>            The identifier of this agent. Defaults to \"%s\" on this machine"), name.data);
 		logger.Info(TC("  -verbose                Print debug information to console"));
 		logger.Info(TC("  -log                    Log all processes detouring information to file (only works with debug builds)"));
 		logger.Info(TC("  -nocustomalloc          Disable custom allocator for processes. If you see odd crashes this can be tested"));
@@ -771,7 +772,7 @@ namespace uba
 		#endif
 
 		if (agentName.IsEmpty())
-			agentName.Append(DefaultAgentName);
+			GetComputerNameW(agentName);
 
 		if (!zone.count)
 			GetZone(zone);
diff --git a/Engine/Source/Programs/UnrealBuildAccelerator/Common/Private/UbaSession.cpp b/Engine/Source/Programs/UnrealBuildAccelerator/Common/Private/UbaSession.cpp
--- a/Engine/Source/Programs/UnrealBuildAccelerator/Common/Private/UbaSession.cpp
+++ b/Engine/Source/Programs/UnrealBuildAccelerator/Common/Private/UbaSession.cpp
@@ -1902,12 +1902,12 @@ namespace uba
 		if (!m_trace.StartWrite(traceName, traceReserveSize))
 			return;
 
-		tchar buf[256];
-		if (!GetComputerNameW(buf, sizeof_array(buf)))
-			TStrcpy_s(buf, 256, TC("LOCAL"));
+		StringBuffer<256> buf;
+		if (!GetComputerNameW(buf))
+			buf.Append(TC("LOCAL"));
 		StringBuffer<> systemInfo;
 		GetSystemInfo(systemInfo);
-		m_trace.SessionAdded(0, {}, buf, systemInfo.data);
+		m_trace.SessionAdded(0, {}, buf, systemInfo);
 		m_traceThreadEvent.Create(true);
 		m_traceThread.Start([this]() { ThreadTraceLoop(); return 0; }, TC("UbaTraceLoop"));
 	}
@@ -3512,10 +3512,9 @@ namespace uba
 
 		// TODO: Remove.
 		#if PLATFORM_MAC
-		if (!GetComputerNameW(outZone.data, outZone.capacity))
+		if (!GetComputerNameW(outZone))
 			return false;
 
-		outZone.count = TStrlen(outZone.data);
 		if (outZone.StartsWith(TC("dc4-mac")) || outZone.StartsWith(TC("rdu-mac")))
 		{
 			outZone.Resize(7);
diff --git a/Engine/Source/Programs/UnrealBuildAccelerator/Common/Private/UbaSessionClient.cpp b/Engine/Source/Programs/UnrealBuildAccelerator/Common/Private/UbaSessionClient.cpp
--- a/Engine/Source/Programs/UnrealBuildAccelerator/Common/Private/UbaSessionClient.cpp
+++ b/Engine/Source/Programs/UnrealBuildAccelerator/Common/Private/UbaSessionClient.cpp
@@ -48,11 +48,7 @@ namespace uba
 		m_processIdCounter = ~0u / 2; // We set this value to a very high value.. because it will be used by child processes and we don't want id from server and child process id to collide
 
 		if (m_name.IsEmpty())
-		{
-			tchar buf[256];
-			if (GetComputerNameW(buf, sizeof_array(buf)))
-				m_name.Appendf(TC("%s"), buf);
-		}
+			GetComputerNameW(m_name);
 
 		m_processWorkingDir.Append(m_rootDir).Append(TC("empty"));// + TC("empty\\");
 		m_storage.CreateDirectory(m_processWorkingDir.data);
diff --git a/Engine/Source/Programs/UnrealBuildAccelerator/Common/Private/UbaSessionServer.cpp b/Engine/Source/Programs/UnrealBuildAccelerator/Common/Private/UbaSessionServer.cpp
--- a/Engine/Source/Programs/UnrealBuildAccelerator/Common/Private/UbaSessionServer.cpp
+++ b/Engine/Source/Programs/UnrealBuildAccelerator/Common/Private/UbaSessionServer.cpp
@@ -760,7 +760,7 @@ namespace uba
 		writer.WriteBool(m_remoteTraceEnabled);
 		WriteRemoteEnvironmentVariables(writer);
 
-		m_trace.SessionAdded(sessionId, connectionInfo.GetId(), name.data, info.data); // Must be inside lock for TraceSessionUpdate() to not include
+		m_trace.SessionAdded(sessionId, connectionInfo.GetId(), name, info); // Must be inside lock for TraceSessionUpdate() to not include
 
 		lock.Leave();
 		return true;
diff --git a/Engine/Source/Programs/UnrealBuildAccelerator/Common/Private/UbaTrace.cpp b/Engine/Source/Programs/UnrealBuildAccelerator/Common/Private/UbaTrace.cpp
--- a/Engine/Source/Programs/UnrealBuildAccelerator/Common/Private/UbaTrace.cpp
+++ b/Engine/Source/Programs/UnrealBuildAccelerator/Common/Private/UbaTrace.cpp
@@ -206,7 +206,7 @@ namespace uba
 		writer.WriteByte(TraceType_##x); \
 		writer.Write7BitEncoded(time);
 
-	void Trace::SessionAdded(u32 sessionId, u32 clientId, const tchar* name, const tchar* info)
+	void Trace::SessionAdded(u32 sessionId, u32 clientId, const StringView& name, const StringView& info)
 	{
 		BEGIN_TRACE_ENTRY(SessionAdded);
 		writer.WriteString(name);
diff --git a/Engine/Source/Programs/UnrealBuildAccelerator/Common/Public/UbaTrace.h b/Engine/Source/Programs/UnrealBuildAccelerator/Common/Public/UbaTrace.h
--- a/Engine/Source/Programs/UnrealBuildAccelerator/Common/Public/UbaTrace.h
+++ b/Engine/Source/Programs/UnrealBuildAccelerator/Common/Public/UbaTrace.h
@@ -82,7 +82,7 @@ namespace uba
 		bool IsWriting() const { return m_memoryBegin != nullptr; }
 
 		bool StartWrite(const tchar* namedTrace, u64 traceMemCapacity = 64*1024*1024);
-		void SessionAdded(u32 sessionId, u32 clientId, const tchar* name, const tchar* info);
+		void SessionAdded(u32 sessionId, u32 clientId, const StringView& name, const StringView& info);
 		void SessionUpdate(u32 sessionId, u32 connectionCount, u64 send, u64 recv, u64 lastPing, u64 memAvail, u64 memTotal, float cpuLoad);
 		void SessionNotification(u32 sessionId, const tchar* text);
 		void SessionSummary(u32 sessionId, const u8* data, u64 dataSize);
diff --git a/Engine/Source/Programs/UnrealBuildAccelerator/Core/Private/UbaPlatform.cpp b/Engine/Source/Programs/UnrealBuildAccelerator/Core/Private/UbaPlatform.cpp
--- a/Engine/Source/Programs/UnrealBuildAccelerator/Core/Private/UbaPlatform.cpp
+++ b/Engine/Source/Programs/UnrealBuildAccelerator/Core/Private/UbaPlatform.cpp
@@ -149,14 +149,20 @@ namespace uba
 		#endif
 	}
 
-	bool GetComputerNameW(tchar* buffer, u32 bufferLen)
+	bool GetComputerNameW(StringBufferBase& out)
 	{
 		#if PLATFORM_WINDOWS
-		DWORD nSize = bufferLen;
-		return ::GetComputerNameW(buffer, &nSize);
+		DWORD size = out.capacity;
+		if (!::GetComputerNameExW(ComputerNameDnsHostname, out.data, &size))
+			if (!::GetComputerNameExW(ComputerNameNetBIOS, out.data, &size))
+				return false;
+		out.count = size;
 		#else
-		return gethostname(buffer, bufferLen) == 0;
+		if (gethostname(out.data, out.capacity) == -1)
+			return false;
+		out.count = TStrlen(out.data);
 		#endif
+		return true;
 	}
 
 	void WriteAssertInfo(StringBufferBase& out, const tchar* text, const char* file, u32 line, const char* expr, void* context)
diff --git a/Engine/Source/Programs/UnrealBuildAccelerator/Core/Public/UbaPlatform.h b/Engine/Source/Programs/UnrealBuildAccelerator/Core/Public/UbaPlatform.h
--- a/Engine/Source/Programs/UnrealBuildAccelerator/Core/Public/UbaPlatform.h
+++ b/Engine/Source/Programs/UnrealBuildAccelerator/Core/Public/UbaPlatform.h
@@ -95,7 +95,7 @@ namespace uba
 	u32 GetUserDefaultUILanguage();
 	u32 GetLastError();
 	void SetLastError(u32 error);
-	bool GetComputerNameW(tchar* buffer, u32 bufferLen);
+	bool GetComputerNameW(StringBufferBase& out);
 	u32 GetCurrentProcessId();
 
 	enum FileHandle : u64 {};
diff --git a/Engine/Source/Programs/UnrealBuildAccelerator/StorageProxy/Private/UbaStorageProxyApp.cpp b/Engine/Source/Programs/UnrealBuildAccelerator/StorageProxy/Private/UbaStorageProxyApp.cpp
--- a/Engine/Source/Programs/UnrealBuildAccelerator/StorageProxy/Private/UbaStorageProxyApp.cpp
+++ b/Engine/Source/Programs/UnrealBuildAccelerator/StorageProxy/Private/UbaStorageProxyApp.cpp
@@ -12,7 +12,6 @@ namespace uba
 {
 	const tchar*	Version = GetVersionString();
 	u16				DefaultListenPort = DefaultPort + 1;
-	const wchar_t*	DefaultName = []() { static wchar_t buf[256]; GetComputerNameW(buf, sizeof_array(buf)); return buf; }();
 
 	int PrintHelp(const wchar_t* message)
 	{
@@ -43,7 +42,7 @@ int wmain(int argc, wchar_t *argv[])
 {
 	using namespace uba;
 
-	StringBuffer<> proxyName(DefaultName);
+	StringBuffer<> proxyName;
 	StringBuffer<> host;
 	u16 hostPort = DefaultPort;
 	StringBuffer<> listenHost;
@@ -51,6 +50,7 @@ int wmain(int argc, wchar_t *argv[])
 	bool poll = true;
 	bool quiet = false;
 
+	GetComputerNameW(proxyName);
 	for (int i = 1; i != argc; ++i)
 	{
 		StringBuffer<> name;
