[UBA]
* Moved deprecated functions to bottom of file

diff --git a/Engine/Source/Programs/UnrealBuildAccelerator/Host/Public/UbaExports.h b/Engine/Source/Programs/UnrealBuildAccelerator/Host/Public/UbaExports.h
--- a/Engine/Source/Programs/UnrealBuildAccelerator/Host/Public/UbaExports.h
+++ b/Engine/Source/Programs/UnrealBuildAccelerator/Host/Public/UbaExports.h
@@ -63,12 +63,6 @@ extern "C"
 	UBA_API void StorageServer_RegisterDisallowedPath(uba::StorageServer* storageServer, const uba::tchar* path);
 	UBA_API void StorageServer_DeleteFile(uba::StorageServer* storage, const uba::tchar* file);
 
-	// RootPaths
-	UBA_API uba::RootPaths* RootPaths_Create(uba::LogWriter& writer);
-	UBA_API bool RootPaths_RegisterRoot(uba::RootPaths* rootPaths, const uba::tchar* path, bool includeInKey, uba::u8 id = 0);
-	UBA_API bool RootPaths_RegisterSystemRoots(uba::RootPaths* rootPaths, uba::u8 startId = 0);
-	UBA_API void RootPaths_Destroy(uba::RootPaths* rootPaths);
-
 	// ProcessStartInfo
 	using ProcessHandle_ExitCallback = void(void* userData, const uba::ProcessHandle&);
 	UBA_API uba::ProcessStartInfo* ProcessStartInfo_Create(const uba::tchar* application, const uba::tchar* arguments, const uba::tchar* workingDir, const uba::tchar* description, uba::u32 priorityClass, uba::u64 outputStatsThresholdMs, bool trackInputs, const uba::tchar* logFile, ProcessHandle_ExitCallback* exit);
@@ -135,10 +129,7 @@ extern "C"
 	// Cache
 	UBA_API uba::CacheClient* CacheClient_Create(uba::SessionServer* session, bool reportMissReason = false, const uba::tchar* crypto = nullptr);
 	UBA_API bool CacheClient_Connect(uba::CacheClient* cacheClient, const uba::tchar* host, int port);
-	UBA_API bool CacheClient_WriteToCache(uba::CacheClient* cacheClient, uba::RootPaths* rootPaths, uba::u32 bucket, const uba::ProcessHandle* process, const uba::u8* inputs, uba::u32 inputsSize, const uba::u8* outputs, uba::u32 outputsSize);
 	UBA_API bool CacheClient_WriteToCache2(uba::CacheClient* cacheClient, uba::u32 bucket, const uba::ProcessHandle* process, const uba::u8* inputs, uba::u32 inputsSize, const uba::u8* outputs, uba::u32 outputsSize, const uba::u8* logLines, uba::u32 logLinesSize);
-	UBA_API uba::u32 CacheClient_FetchFromCache(uba::CacheClient* cacheClient, uba::RootPaths* rootPaths, uba::u32 bucket, const uba::ProcessStartInfo& info);
-	UBA_API uba::CacheResult* CacheClient_FetchFromCache2(uba::CacheClient* cacheClient, uba::RootPaths* rootPaths, uba::u32 bucket, const uba::ProcessStartInfo& info);
 	UBA_API uba::CacheResult* CacheClient_FetchFromCache3(uba::CacheClient* cacheClient, uba::u64 rootsHandle, uba::u32 bucket, const uba::ProcessStartInfo& info);
 	UBA_API void CacheClient_RequestServerShutdown(uba::CacheClient* cacheClient, const uba::tchar* reason);
 	UBA_API void CacheClient_Destroy(uba::CacheClient* cacheClient);
@@ -160,6 +151,13 @@ extern "C"
 	UBA_API void Uba_Destroy(void* uba);
 
 
-	// DEPRECATED
+	// DEPRECATED, don't use
 	UBA_API void DestroyProcessHandle(uba::ProcessHandle* handle);
+	UBA_API uba::RootPaths* RootPaths_Create(uba::LogWriter& writer);
+	UBA_API bool RootPaths_RegisterRoot(uba::RootPaths* rootPaths, const uba::tchar* path, bool includeInKey, uba::u8 id = 0);
+	UBA_API bool RootPaths_RegisterSystemRoots(uba::RootPaths* rootPaths, uba::u8 startId = 0);
+	UBA_API void RootPaths_Destroy(uba::RootPaths* rootPaths);
+	UBA_API bool CacheClient_WriteToCache(uba::CacheClient* cacheClient, uba::RootPaths* rootPaths, uba::u32 bucket, const uba::ProcessHandle* process, const uba::u8* inputs, uba::u32 inputsSize, const uba::u8* outputs, uba::u32 outputsSize);
+	UBA_API uba::u32 CacheClient_FetchFromCache(uba::CacheClient* cacheClient, uba::RootPaths* rootPaths, uba::u32 bucket, const uba::ProcessStartInfo& info);
+	UBA_API uba::CacheResult* CacheClient_FetchFromCache2(uba::CacheClient* cacheClient, uba::RootPaths* rootPaths, uba::u32 bucket, const uba::ProcessStartInfo& info);
 }
