[UBA]
* Fixed so GetFileSizeEx falls back to directory table information if size is not known in the other paths
* Fixed so NtQueryObject handles detoured handle for ObjectNameInformation class. Needed to fix issues when machine has symantec antivirus installed

diff --git a/Engine/Source/Programs/UnrealBuildAccelerator/Core/Private/UbaStringBuffer.cpp b/Engine/Source/Programs/UnrealBuildAccelerator/Core/Private/UbaStringBuffer.cpp
--- a/Engine/Source/Programs/UnrealBuildAccelerator/Core/Private/UbaStringBuffer.cpp
+++ b/Engine/Source/Programs/UnrealBuildAccelerator/Core/Private/UbaStringBuffer.cpp
@@ -253,12 +253,12 @@ namespace uba
 		return *this;
 	}
 
-	StringBufferBase& StringBufferBase::Prepend(const StringView& view)
+	StringBufferBase& StringBufferBase::Prepend(const StringView& view, u32 overwriteCount)
 	{
-		UBA_ASSERT(count + view.count < capacity);
-		memmove(data+view.count, data, (count+1)*sizeof(tchar));
+		UBA_ASSERT(count + view.count - overwriteCount < capacity);
+		memmove(data+view.count - overwriteCount, data, (count+1)*sizeof(tchar));
 		memcpy(data, view.data, view.count*sizeof(tchar));
-		count += view.count;
+		count += view.count - overwriteCount;
 		return *this;
 	}
 
diff --git a/Engine/Source/Programs/UnrealBuildAccelerator/Core/Public/UbaStringBuffer.h b/Engine/Source/Programs/UnrealBuildAccelerator/Core/Public/UbaStringBuffer.h
--- a/Engine/Source/Programs/UnrealBuildAccelerator/Core/Public/UbaStringBuffer.h
+++ b/Engine/Source/Programs/UnrealBuildAccelerator/Core/Public/UbaStringBuffer.h
@@ -36,7 +36,7 @@ namespace uba
 		StringBufferBase& AppendHex(u64 v);
 		StringBufferBase& AppendValue(u64 v);
 		StringBufferBase& Append(const tchar* format, va_list& args);
-		StringBufferBase& Prepend(const StringView& view);
+		StringBufferBase& Prepend(const StringView& view, u32 overwriteCount = 0);
 		StringBufferBase& Resize(u64 newSize);
 		StringBufferBase& Clear();
 
diff --git a/Engine/Source/Programs/UnrealBuildAccelerator/Detours/Private/Windows/UbaDetoursFunctionsKernelBase.inl b/Engine/Source/Programs/UnrealBuildAccelerator/Detours/Private/Windows/UbaDetoursFunctionsKernelBase.inl
--- a/Engine/Source/Programs/UnrealBuildAccelerator/Detours/Private/Windows/UbaDetoursFunctionsKernelBase.inl
+++ b/Engine/Source/Programs/UnrealBuildAccelerator/Detours/Private/Windows/UbaDetoursFunctionsKernelBase.inl
@@ -791,7 +791,21 @@ DWORD Detoured_GetFileSizeEx(HANDLE hFile, PLARGE_INTEGER lpFileSize)
 		}
 
 		trueHandle = dh.trueHandle;
-		UBA_ASSERT(trueHandle != INVALID_HANDLE_VALUE);
+		
+		if (trueHandle == INVALID_HANDLE_VALUE)
+		{
+			u32 entryOffset = dh.dirTableOffset;
+			if (entryOffset != ~u32(0))
+			{
+				DirectoryTable::EntryInformation entryInfo;
+				g_directoryTable.GetEntryInformation(entryInfo, entryOffset);
+				*lpFileSize = ToLargeInteger(entryInfo.size);
+				SetLastError(ERROR_SUCCESS);
+				DEBUG_LOG_DETOURED(L"GetFileSizeEx", L"%llu (%ls) (Size:%llu) -> 1", uintptr_t(hFile), HandleToName(hFile), entryInfo.size);
+				return 1;
+			}
+			UBA_ASSERTF(false, L"GetFileSizeEx (%s)", HandleToName(hFile));
+		}
 	}
 	DEBUG_LOG_TRUE(L"GetFileSizeEx", L"%llu (%ls)", uintptr_t(hFile), HandleToName(hFile));
 	return True_GetFileSizeEx(trueHandle, lpFileSize); // This ends up in Detoured_NtQueryInformationFile
diff --git a/Engine/Source/Programs/UnrealBuildAccelerator/Detours/Private/Windows/UbaDetoursFunctionsNtDll.inl b/Engine/Source/Programs/UnrealBuildAccelerator/Detours/Private/Windows/UbaDetoursFunctionsNtDll.inl
--- a/Engine/Source/Programs/UnrealBuildAccelerator/Detours/Private/Windows/UbaDetoursFunctionsNtDll.inl
+++ b/Engine/Source/Programs/UnrealBuildAccelerator/Detours/Private/Windows/UbaDetoursFunctionsNtDll.inl
@@ -1497,13 +1497,70 @@ NTSTATUS Detoured_NtQueryObject(HANDLE Handle, OBJECT_INFORMATION_CLASS ObjectIn
 {
 	DETOURED_CALL(NtQueryObject);
 
+	HANDLE trueHandle = Handle;
+
 	// This can be other things than FILES.. Is used by GetHandleInformation
 	if (isDetouredHandle(Handle))
 	{
-		Handle = asDetouredHandle(Handle).trueHandle;
-		UBA_ASSERTF(Handle != INVALID_HANDLE_VALUE, L"NtQueryObject");
+		DetouredHandle& dh = asDetouredHandle(Handle);
+		trueHandle = dh.trueHandle;
+
+		if (trueHandle == INVALID_HANDLE_VALUE)
+		{
+			if (ObjectInformationClass == 1) // ObjectNameInformation
+			{
+				auto fo = dh.fileObject;
+				UBA_ASSERT(fo);
+				auto& fi = *fo->fileInfo;
+				UBA_ASSERT(fi.originalName);
+				const wchar_t* fileName = fi.originalName;
+
+				StringBuffer<> fixedPath;
+				FixPath(fileName, g_virtualWorkingDir.data, g_virtualWorkingDir.count, fixedPath);
+
+				StringBuffer<> buffer;
+
+				g_directoryTable.GetFinalPath(buffer, fixedPath.data);
+				VirtualizePath(buffer);
+
+				if (g_runningRemote || buffer[0] != fixedPath[0])
+				{
+					// It was remote or virtualized.. let's make up device drive.
+					// TODO: Maybe we need to replicate device drive names to remotes
+					buffer.Prepend(AsView(L"\\Device\\HarddiskVolume100"), 2);
+				}
+				else
+				{
+					wchar_t drive[3] = { buffer[0], ':', 0};
+					wchar_t device[256];
+					DWORD deviceLen = QueryDosDeviceW(drive, device, sizeof(device));
+					UBA_ASSERT(deviceLen);
+					buffer.Prepend(StringView(device, deviceLen), 2);
+				}
+
+				u64 bufferSize = (buffer.count+1)*sizeof(tchar);
+				u64 totalSize = sizeof(UNICODE_STRING) + bufferSize;
+
+				if (ObjectInformationLength < totalSize)
+				{
+					DEBUG_LOG_DETOURED(L"NtQueryObject", L"(ObjectNameInformation) %s -> STATUS_BUFFER_OVERFLOW", HandleToName(Handle));
+					return STATUS_BUFFER_OVERFLOW;
+				}
+				auto& ustr = *(PUNICODE_STRING)ObjectInformation;
+				ustr.Length = u16(bufferSize - sizeof(tchar));
+				ustr.MaximumLength = u16(bufferSize);
+				ustr.Buffer = (wchar_t*)((&ustr) + 1);
+				memcpy(ustr.Buffer, buffer.data, bufferSize);
+				*ReturnLength = (ULONG)totalSize;
+
+				DEBUG_LOG_DETOURED(L"NtQueryObject", L"(ObjectNameInformation) %llu -> Success (%s)", uintptr_t(Handle), buffer.data);
+				return STATUS_SUCCESS;
+			}
+
+			UBA_ASSERTF(false, L"NtQueryObject NOT_IMPLEMENTED (class %i) (%s)", ObjectInformationClass, HandleToName(Handle));
+		}
 	}
-	auto res = True_NtQueryObject(Handle, ObjectInformationClass, ObjectInformation, ObjectInformationLength, ReturnLength);
+	auto res = True_NtQueryObject(trueHandle, ObjectInformationClass, ObjectInformation, ObjectInformationLength, ReturnLength);
 	DEBUG_LOG_TRUE(L"NtQueryObject", L"(%i) %llu -> %ls", ObjectInformationClass, uintptr_t(Handle), ToString(res).data);
 	return res;
 }
