[UBA]
* Bumped cas and path table max size to be much larger than needed.. so we never hit this limit again

diff --git a/Engine/Source/Programs/UnrealBuildAccelerator/Common/Private/UbaCompactTables.cpp b/Engine/Source/Programs/UnrealBuildAccelerator/Common/Private/UbaCompactTables.cpp
--- a/Engine/Source/Programs/UnrealBuildAccelerator/Common/Private/UbaCompactTables.cpp
+++ b/Engine/Source/Programs/UnrealBuildAccelerator/Common/Private/UbaCompactTables.cpp
@@ -5,8 +5,8 @@
 
 namespace uba
 {
-	static constexpr u32 PathTableMaxSize = 20*1024*1024; // Max size for string memory
-	static constexpr u32 CasKeyTableMaxSize = 64*1024*1024;
+	static constexpr u32 PathTableMaxSize = 128*1024*1024; // Max size for string memory
+	static constexpr u32 CasKeyTableMaxSize = 256*1024*1024;
 	static constexpr u32 CasKeyLookupInitialCount = 1024;
 	static constexpr u32 CasKeyArrayMaxSize = 16;
 
