[UBA]
* Added hints to path/caskey table allocations to improve error logging

diff --git a/Engine/Source/Programs/UnrealBuildAccelerator/Common/Private/UbaCompactTables.cpp b/Engine/Source/Programs/UnrealBuildAccelerator/Common/Private/UbaCompactTables.cpp
--- a/Engine/Source/Programs/UnrealBuildAccelerator/Common/Private/UbaCompactTables.cpp
+++ b/Engine/Source/Programs/UnrealBuildAccelerator/Common/Private/UbaCompactTables.cpp
@@ -38,7 +38,7 @@ namespace uba
 		if (!m_pathTableMem.memory)
 			m_pathTableMem.Init(PathTableMaxSize);
 		if (!m_pathTableMem.writtenSize)
-			m_pathTableMem.AllocateNoLock(1, 1, TC(""));
+			m_pathTableMem.AllocateNoLock(1, 1, TC("CompactPathTable"));
 
 		const tchar* stringKeyString = str;
 		StringBuffer<MaxPath> tempStringKeyStr;
@@ -82,7 +82,7 @@ namespace uba
 			// Put string directly after current element and set segment offset to 0
 			u64 bytesForString = GetStringWriteSize(seg, segLen);
 			u64 memSize = bytesForParent + 1 + bytesForString;
-			u8* mem = (u8*)m_pathTableMem.AllocateNoLock(memSize, 1, TC(""));
+			u8* mem = (u8*)m_pathTableMem.AllocateNoLock(memSize, 1, TC("CompactPathTable"));
 			BinaryWriter writer(mem, 0, memSize);
 			writer.Write7BitEncoded(parentOffset);
 			writer.Write7BitEncoded(0);
@@ -100,7 +100,7 @@ namespace uba
 		#endif
 
 		u64 memSize = bytesForParent + Get7BitEncodedCount(strOffset);
-		u8* mem = (u8*)m_pathTableMem.AllocateNoLock(memSize, 1, TC(""));
+		u8* mem = (u8*)m_pathTableMem.AllocateNoLock(memSize, 1, TC("CompactPathTable"));
 		BinaryWriter writer(mem, 0, memSize);
 		writer.Write7BitEncoded(parentOffset);
 		writer.Write7BitEncoded(strOffset);
@@ -173,7 +173,7 @@ namespace uba
 
 		u64 writtenSize = m_pathTableMem.writtenSize;
 		u64 left = reader.GetLeft();
-		void* mem = m_pathTableMem.AllocateNoLock(left, 1, TC(""));
+		void* mem = m_pathTableMem.AllocateNoLock(left, 1, TC("CompactPathTable"));
 		reader.ReadBytes(mem, left);
 
 		if (!populateLookup)
@@ -245,7 +245,7 @@ namespace uba
 		if (added)
 		{
 			u8 bytesForStringOffset = Get7BitEncodedCount(stringOffset);
-			u8* mem = (u8*)m_casKeyTableMem.AllocateNoLock(bytesForStringOffset + sizeof(CasKey), 1, TC(""));
+			u8* mem = (u8*)m_casKeyTableMem.AllocateNoLock(bytesForStringOffset + sizeof(CasKey), 1, TC("CompactCasKeyTable"));
 			BinaryWriter writer(mem, 0, 1000);
 			writer.Write7BitEncoded(stringOffset);
 			writer.WriteCasKey(casKey);
@@ -273,7 +273,7 @@ namespace uba
 		if (added)
 		{
 			u8 bytesForStringOffset = Get7BitEncodedCount(stringOffset);
-			u8* mem = (u8*)m_casKeyTableMem.AllocateNoLock(bytesForStringOffset + sizeof(CasKey), 1, TC(""));
+			u8* mem = (u8*)m_casKeyTableMem.AllocateNoLock(bytesForStringOffset + sizeof(CasKey), 1, TC("CompactCasKeyTable"));
 			BinaryWriter writer(mem, 0, 1000);
 			writer.Write7BitEncoded(stringOffset);
 			writer.WriteCasKey(casKey);
@@ -415,7 +415,7 @@ namespace uba
 		u64 writtenSize = m_casKeyTableMem.writtenSize;
 
 		u64 left = reader.GetLeft();
-		void* mem = m_casKeyTableMem.AllocateNoLock(left, 1, TC(""));
+		void* mem = m_casKeyTableMem.AllocateNoLock(left, 1, TC("CompactCasKeyTable"));
 		reader.ReadBytes(mem, left);
 
 		if (!populateLookup)
