[UBA]
* Added cache server command "pathtable" which outputs content of path table
* Minor cleanups

diff --git a/Engine/Source/Programs/UnrealBuildAccelerator/Common/Private/UbaCacheServer.cpp b/Engine/Source/Programs/UnrealBuildAccelerator/Common/Private/UbaCacheServer.cpp
--- a/Engine/Source/Programs/UnrealBuildAccelerator/Common/Private/UbaCacheServer.cpp
+++ b/Engine/Source/Programs/UnrealBuildAccelerator/Common/Private/UbaCacheServer.cpp
@@ -1669,11 +1669,11 @@ namespace uba
 		u8 bom[] = {0xEF,0xBB,0xBF}; 
 		Write(bom, sizeof(bom));
 
-		auto writeLine = [&](const tchar* text)
+		auto writeLine = [&](const StringView& text)
 			{
 				u8 buffer[1024];
 				BinaryWriter w(buffer, 0, sizeof(buffer));
-				w.WriteUtf8String(text, TStrlen(text));
+				w.WriteUtf8String(text.data, text.count);
 				w.WriteUtf8String(TC("\n"), 1);
 				Write(buffer, w.GetPosition());
 			};
@@ -1682,7 +1682,7 @@ namespace uba
 
 		if (command.Equals(TC("content")))
 		{
-			writeLine(TC("UbaCache server summary"));
+			writeLine(AsView(TC("UbaCache server summary")));
 
 			StringBufferBase& filterString = additionalInfo;
 
@@ -1734,7 +1734,7 @@ namespace uba
 					}
 
 
-					writeLine(CasKeyString(kv2.first).str);
+					writeLine(ToView(CasKeyString(kv2.first).str));
 					u32 index = 0;
 					for (auto& entry : entries.entries)
 					{
@@ -1744,8 +1744,8 @@ namespace uba
 							continue;
 						}
 
-						u64 age = GetFileTimeAsTime(now - entry.creationTime);
-						writeLine(line.Clear().Appendf(TC("  #%u (%s ago)"), index, TimeToText(age, true).str).data);
+						u64 age = GetFileTimeAsTime(now - (m_creationTime + entry.creationTime));
+						writeLine(line.Clear().Appendf(TC("  #%u (%s ago)"), index, TimeToText(age, true).str));
 
 						auto writeOffsets = [&](const Vector<u8>& offsets)
 							{
@@ -1756,14 +1756,14 @@ namespace uba
 									CasKey casKey;
 									StringBuffer<> path;
 									bucket.m_casKeyTable.GetPathAndKey(path, casKey, bucket.m_pathTable, offset);
-									writeLine(line.Clear().Appendf(TC("    %s - %s"), path.data, CasKeyString(casKey).str).data);
+									writeLine(line.Clear().Appendf(TC("    %s - %s"), path.data, CasKeyString(casKey).str));
 								}
 							};
 
-						writeLine(line.Clear().Append(TC("   Inputs:")).data);
+						writeLine(line.Clear().Append(TC("   Inputs:")));
 						entries.Flatten(temp, entry);
 						writeOffsets(temp);
-						writeLine(line.Clear().Append(TC("   Outputs:")).data);
+						writeLine(line.Clear().Append(TC("   Outputs:")));
 						writeOffsets(entry.outputCasKeyOffsets);
 						++index;
 					}
@@ -1772,11 +1772,11 @@ namespace uba
 		}
 		else if (command.Equals(TC("status")))
 		{
-			writeLine(TC("UbaCacheServer status"));
-			writeLine(line.Clear().Appendf(TC("  CreationTime: %s ago"), TimeToText(GetFileTimeAsTime(GetSystemTimeAsFileTime() - m_creationTime), true).str).data);
-			writeLine(line.Clear().Appendf(TC("  UpTime: %s"), TimeToText(GetTime() - m_bootTime, true).str).data);
-			writeLine(line.Clear().Appendf(TC("  Longest maintenance: %s"), TimeToText(m_longestMaintenance).str).data);
-			writeLine(line.Clear().Appendf(TC("  Buckets:")).data);
+			writeLine(AsView(TC("UbaCacheServer status")));
+			writeLine(line.Clear().Appendf(TC("  CreationTime: %s ago"), TimeToText(GetFileTimeAsTime(GetSystemTimeAsFileTime() - m_creationTime), true).str));
+			writeLine(line.Clear().Appendf(TC("  UpTime: %s"), TimeToText(GetTime() - m_bootTime, true).str));
+			writeLine(line.Clear().Appendf(TC("  Longest maintenance: %s"), TimeToText(m_longestMaintenance).str));
+			writeLine(line.Clear().Appendf(TC("  Buckets:")));
 			u32 index = 0;
 
 			{
@@ -1802,36 +1802,55 @@ namespace uba
 					if (lastUsed)
 						lastUsedTime = GetFileTimeAsTime(GetSystemTimeAsFileTime() - (m_creationTime + lastUsed));
 
-					writeLine(line.Clear().Appendf(TC("    #%u - %llu"), index++, kv.first).data);
-					writeLine(line.Clear().Appendf(TC("      PathTable: %lls (%s)"), bucket.m_pathTable.GetPathCount(), BytesToText(bucket.m_pathTable.GetSize()).str).data);
-					writeLine(line.Clear().Appendf(TC("      CasKeyTable: %llu (%s)"), bucket.m_cacheEntryLookup.size(), BytesToText(bucket.m_casKeyTable.GetSize()).str).data);
-					writeLine(line.Clear().Appendf(TC("      TotalEntries: %llu"), totalEntryCount).data);
-					writeLine(line.Clear().Appendf(TC("      KeyMostEntries: %llu"), mostEntries).data);
-					writeLine(line.Clear().Appendf(TC("      LastEntryUsed: %s ago"), TimeToText(lastUsedTime, true).str).data);
+					writeLine(line.Clear().Appendf(TC("    #%u - %llu"), index++, kv.first));
+					writeLine(line.Clear().Appendf(TC("      PathTable: %llu (%s)"), bucket.m_pathTable.GetPathCount(), BytesToText(bucket.m_pathTable.GetSize()).str));
+					writeLine(line.Clear().Appendf(TC("      CasKeyTable: %llu (%s)"), bucket.m_cacheEntryLookup.size(), BytesToText(bucket.m_casKeyTable.GetSize()).str));
+					writeLine(line.Clear().Appendf(TC("      TotalEntries: %llu"), totalEntryCount));
+					writeLine(line.Clear().Appendf(TC("      KeyMostEntries: %llu"), mostEntries));
+					writeLine(line.Clear().Appendf(TC("      LastEntryUsed: %s ago"), TimeToText(lastUsedTime, true).str));
 				}
 			}
 			u64 totalCasSize = 0;
 			u64 totalCasCount = 0;
 			m_storage.TraverseAllCasFiles([&](const CasKey& casKey, u64 size) { ++totalCasCount; totalCasSize += size; });
-			writeLine(line.Clear().Appendf(TC("  CasDb:")).data);
-			writeLine(line.Clear().Appendf(TC("    Count: %llu"), totalCasCount).data);
-			writeLine(line.Clear().Appendf(TC("    Size: %s"), BytesToText(totalCasSize).str).data);
+			writeLine(line.Clear().Appendf(TC("  CasDb:")));
+			writeLine(line.Clear().Appendf(TC("    Count: %llu"), totalCasCount));
+			writeLine(line.Clear().Appendf(TC("    Size: %s"), BytesToText(totalCasSize).str));
+		}
+		else if (command.Equals(TC("pathtable")))
+		{
+				SCOPED_READ_LOCK(m_bucketsLock, bucketsLock);
+				for (auto& kv : m_buckets)
+				{
+					Bucket& bucket = kv.second;
+					SCOPED_READ_LOCK(bucket.m_cacheEntryLookupLock, lock2);
+					Set<TString> sortedPaths;
+
+					bucket.m_pathTable.TraversePaths([&](const StringView& path)
+						{
+							if (!sortedPaths.emplace(path.data).second)
+								writeLine(AsView(TC("EEEEERRRRRROOOORRRR!!!!!!")));
+						});
+
+					for (auto& path : sortedPaths)
+						writeLine(path);
+				}
 		}
 		else if (command.Equals(TC("obliterate")))
 		{
 			m_shouldWipe = true;
 			m_addsSinceMaintenance = 1;
-			writeLine(line.Clear().Appendf(TC("Cache server database obliteration queued!")).data);
+			writeLine(line.Clear().Appendf(TC("Cache server database obliteration queued!")));
 		}
 		else if (command.Equals(TC("maintenance")))
 		{
 			m_forceAllSteps = true;
 			m_addsSinceMaintenance = 1;
-			writeLine(line.Clear().Appendf(TC("Cache server maintenance queued!")).data);
+			writeLine(line.Clear().Appendf(TC("Cache server maintenance queued!")));
 		}
 		else
 		{
-			writeLine(line.Clear().Appendf(TC("Unknown command: %s"), command.data).data);
+			writeLine(line.Clear().Appendf(TC("Unknown command: %s"), command.data));
 		}
 
 		Write("", 1);
diff --git a/Engine/Source/Programs/UnrealBuildAccelerator/Common/Public/UbaCompactTables.h b/Engine/Source/Programs/UnrealBuildAccelerator/Common/Public/UbaCompactTables.h
--- a/Engine/Source/Programs/UnrealBuildAccelerator/Common/Public/UbaCompactTables.h
+++ b/Engine/Source/Programs/UnrealBuildAccelerator/Common/Public/UbaCompactTables.h
@@ -33,6 +33,9 @@ namespace uba
 		u64 GetPathCount() { return m_offsets.Size(); }
 		u64 GetSegmentCount() { return m_segmentOffsets.Size(); }
 
+		template<typename Func>
+		void TraversePaths(const Func& func) const;
+
 	private:
 		u32 InternalAdd(const tchar* str, const tchar* stringKeyString, u64 strLen);
 		ReaderWriterLock m_lock;
@@ -94,6 +97,16 @@ namespace uba
 		HashMap<CasKey, Value, true> m_offsets;
 	};
 
+	template<typename Func>
+	void CompactPathTable::TraversePaths(const Func& func) const
+	{
+		for (auto i = m_offsets.ValuesBegin(), e = m_offsets.ValuesEnd(); i!=e; ++i)
+		{
+			StringBuffer path;
+			GetString(path, *i);
+			func(path);
+		}
+	}
 
 	template<typename Func>
 	void CompactCasKeyTable::TraverseOffsets(const CasKey& casKey, const Func& func) const
diff --git a/Engine/Source/Programs/UnrealBuildAccelerator/Common/Public/UbaHashMap.h b/Engine/Source/Programs/UnrealBuildAccelerator/Common/Public/UbaHashMap.h
--- a/Engine/Source/Programs/UnrealBuildAccelerator/Common/Public/UbaHashMap.h
+++ b/Engine/Source/Programs/UnrealBuildAccelerator/Common/Public/UbaHashMap.h
@@ -185,8 +185,8 @@ namespace uba
 			}
 		}
 
-		Value* ValuesBegin() { return m_values + 1; }
-		Value* ValuesEnd() { return m_values + m_nextAvailableEntry; }
+		Value* ValuesBegin() const { return m_values + 1; }
+		Value* ValuesEnd() const { return m_values + m_nextAvailableEntry; }
 
 		void Swap(HashMap& o)
 		{
