[UBA]
* Added so CasKeyString inherits from StringView

diff --git a/Engine/Source/Programs/UnrealBuildAccelerator/Core/Public/UbaHash.h b/Engine/Source/Programs/UnrealBuildAccelerator/Core/Public/UbaHash.h
--- a/Engine/Source/Programs/UnrealBuildAccelerator/Core/Public/UbaHash.h
+++ b/Engine/Source/Programs/UnrealBuildAccelerator/Core/Public/UbaHash.h
@@ -43,9 +43,9 @@ namespace uba
 		TSprintf_s(out, capacity, TC("%08x-%04hx-%04hx-%02hhx%02hhx-%02hhx%02hhx%02hhx%02hhx%02hhx%02hhx"), g.data1, g.data2, g.data3, g.data4[0], g.data4[1], g.data4[2], g.data4[3], g.data4[4], g.data4[5], g.data4[6], g.data4[7]);
 	}
 
-	struct GuidToString
+	struct GuidToString : StringView
 	{
-		GuidToString(const Guid& g) { GuidToStr(str, 37, g); }
+		GuidToString(const Guid& g) : StringView(str, 36) { GuidToStr(str, 37, g); }
 		tchar str[37];
 	};
 
@@ -170,9 +170,9 @@ namespace uba
 		*pos = 0;
 	}
 
-	struct CasKeyString
+	struct CasKeyString : StringView
 	{
-		CasKeyString(const CasKey& g) { ToString(str, 41, g); }
+		CasKeyString(const CasKey& g) : StringView(str, 40) { ToString(str, 41, g); }
 		tchar str[41];
 	};
 
