[UBA]
* Added some documentation around how -http works when storing crypto keys

diff --git a/Engine/Source/Programs/UnrealBuildAccelerator/CacheService/Private/UbaCacheService.cpp b/Engine/Source/Programs/UnrealBuildAccelerator/CacheService/Private/UbaCacheService.cpp
--- a/Engine/Source/Programs/UnrealBuildAccelerator/CacheService/Private/UbaCacheService.cpp
+++ b/Engine/Source/Programs/UnrealBuildAccelerator/CacheService/Private/UbaCacheService.cpp
@@ -50,6 +50,9 @@ namespace uba
 		logger.Info(TC("  -expiration=<seconds>   Time until unused cache entries get deleted. Defaults to %s (%u seconds)"), TimeToText(MsToTime(DefaultExpiration*1000)).str, DefaultExpiration);
 		logger.Info(TC("  -http=<port>            If set, a http server will be started and listen on <port>"));
 		logger.Info(TC("  -fullmaintenance        Force a full maintenance"));
+
+		logger.Info(TC(" Example of how to register crypto key to cache server (when -http=80 is provided)"));
+		logger.Info(TC("   curl http://localhost/addcrypto?3f58aa57466db9999213456789123445"));
 		logger.Info(TC(""));
 		return -1;
 	}
@@ -266,7 +269,7 @@ namespace uba
 
 					u8 crypto128Data[16];
 					if (!CryptoFromString(crypto128Data, 16, arguments))
-						return "Failed to read crypto argument";
+						return "Failed to read crypto argument (Needs to be 32 characters long)";
 					u64 expirationTime = GetTime() + MsToTime(expirationSeconds*1000);
 					networkServer.RegisterCryptoKey(crypto128Data, expirationTime);
 					return (const char*)nullptr;
