[UBA]
* Added versioning to bucket
* Added versioning to compact path table
* Changed so compact path table version 2 includes dot (.) as separator for segments..

This change reduces path table size in half at least

diff --git a/Engine/Source/Programs/UnrealBuildAccelerator/Common/Private/UbaCacheClient.cpp b/Engine/Source/Programs/UnrealBuildAccelerator/Common/Private/UbaCacheClient.cpp
--- a/Engine/Source/Programs/UnrealBuildAccelerator/Common/Private/UbaCacheClient.cpp
+++ b/Engine/Source/Programs/UnrealBuildAccelerator/Common/Private/UbaCacheClient.cpp
@@ -40,8 +40,8 @@ namespace uba
 	{
 		Bucket(u32 id_)
 		:	id(id_)
-		,	serverPathTable(CaseInsensitiveFs, 0, 0)
-		,	sendPathTable(CaseInsensitiveFs, 0, 0)
+		,	serverPathTable(CaseInsensitiveFs, 0, 0, CacheBucketVersion)
+		,	sendPathTable(CaseInsensitiveFs, 0, 0, CacheBucketVersion)
 		{
 		}
 
@@ -388,7 +388,8 @@ namespace uba
 
 	u64 CacheClient::MakeId(u32 bucketId)
 	{
-		return u64(bucketId) | ((u64(!CaseInsensitiveFs) + (RootPathsVersion << 1u) + (u64(!m_useRoots) << 2u)) << 32u);
+		constexpr u64 ForBackwardsCompatibility = 1;
+		return u64(bucketId) | ((u64(!CaseInsensitiveFs) + (ForBackwardsCompatibility << 1u) + (u64(!m_useRoots) << 2u) + (CacheBucketVersion << 3u)) << 32u);
 	}
 
 	bool CacheClient::FetchFromCache(CacheResult& outResult, RootsHandle rootsHandle, u32 bucketId, const ProcessStartInfo& info)
@@ -833,7 +834,7 @@ namespace uba
 		}
 		else
 		{
-			MemoryBlock block(4*1024*1024);
+			MemoryBlock block(512*1024*1024);
 			if (!fetcher.RetrieveFile(m_logger, m_client, statusFileCasKey, TC("CommandString"), destinationIsCompressed, &block))
 				return false;
 			BinaryReader reader(block.memory, 3, block.writtenSize); // Skipping bom
diff --git a/Engine/Source/Programs/UnrealBuildAccelerator/Common/Private/UbaCacheServer.cpp b/Engine/Source/Programs/UnrealBuildAccelerator/Common/Private/UbaCacheServer.cpp
--- a/Engine/Source/Programs/UnrealBuildAccelerator/Common/Private/UbaCacheServer.cpp
+++ b/Engine/Source/Programs/UnrealBuildAccelerator/Common/Private/UbaCacheServer.cpp
@@ -124,7 +124,7 @@ namespace uba
 
 	struct CacheServer::ConnectionBucket
 	{
-		ConnectionBucket(u64 i) : pathTable(IsCaseInsensitive(i), 0, 0), id(i) {}
+		ConnectionBucket(u64 i, u32 version) : pathTable(IsCaseInsensitive(i), 0, 0, version), id(i) {}
 		CompactPathTable pathTable;
 		CompactCasKeyTable casKeyTable;
 
@@ -141,7 +141,7 @@ namespace uba
 
 	struct CacheServer::Bucket
 	{
-		Bucket(u64 id) : m_pathTable(IsCaseInsensitive(id), 0, 0) {}
+		Bucket(u64 id, u32 version) : m_pathTable(IsCaseInsensitive(id), 0, 0, version) {}
 		ReaderWriterLock m_cacheEntryLookupLock;
 		UnorderedMap<CasKey, CacheEntries> m_cacheEntryLookup;
 
@@ -1071,7 +1071,7 @@ namespace uba
 			{
 				u64 pathReserveCount = pathTable.GetPathCount() + 1000; // Add extra reservation space for growth
 				u64 segmentReserveCount = pathTable.GetSegmentCount() + 1000;  // Add extra reservation space for growth
-				CompactPathTable newPathTable(pathTable.GetCaseInsensitive(), pathReserveCount, segmentReserveCount);
+				CompactPathTable newPathTable(pathTable.GetCaseInsensitive(), pathReserveCount, segmentReserveCount, pathTable.GetVersion());
 				oldToNewPathOffset.Init(memoryBlock, usedPathOffsets.CountSetBits());
 
 				StringBuffer<> temp;
@@ -1210,11 +1210,12 @@ namespace uba
 	CacheServer::ConnectionBucket& CacheServer::GetConnectionBucket(const ConnectionInfo& connectionInfo, BinaryReader& reader, u32* outClientVersion)
 	{
 		u64 id = reader.Read7BitEncoded();
+		u32 bucketVersion = u32(id >> 35u);
 		SCOPED_WRITE_LOCK(m_connectionsLock, lock);
 		auto& connection = m_connections[connectionInfo.GetId()];
 		if (outClientVersion)
 			*outClientVersion = connection.clientVersion;
-		return connection.buckets.try_emplace(id, id).first->second;
+		return connection.buckets.try_emplace(id, id, bucketVersion).first->second;
 	}
 
 	CacheServer::Bucket& CacheServer::GetBucket(BinaryReader& reader)
@@ -1226,8 +1227,9 @@ namespace uba
 
 	CacheServer::Bucket& CacheServer::GetBucket(u64 id)
 	{
+		u32 bucketVersion = u32(id >> 35u);
 		SCOPED_WRITE_LOCK(m_bucketsLock, bucketsLock);
-		auto insres = m_buckets.try_emplace(id, id);
+		auto insres = m_buckets.try_emplace(id, id, bucketVersion);
 		auto& bucket = insres.first->second;
 		if (!insres.second)
 			return bucket;
@@ -1802,7 +1804,7 @@ namespace uba
 					if (lastUsed)
 						lastUsedTime = GetFileTimeAsTime(GetSystemTimeAsFileTime() - (m_creationTime + lastUsed));
 
-					writeLine(line.Clear().Appendf(TC("    #%u - %llu"), index++, kv.first));
+					writeLine(line.Clear().Appendf(TC("    #%u - %llu (Version %u)"), index++, kv.first, bucket.m_pathTable.GetVersion()));
 					writeLine(line.Clear().Appendf(TC("      PathTable: %llu (%s)"), bucket.m_pathTable.GetPathCount(), BytesToText(bucket.m_pathTable.GetSize()).str));
 					writeLine(line.Clear().Appendf(TC("      CasKeyTable: %llu (%s)"), bucket.m_cacheEntryLookup.size(), BytesToText(bucket.m_casKeyTable.GetSize()).str));
 					writeLine(line.Clear().Appendf(TC("      TotalEntries: %llu"), totalEntryCount));
@@ -1824,6 +1826,9 @@ namespace uba
 				{
 					Bucket& bucket = kv.second;
 					SCOPED_READ_LOCK(bucket.m_cacheEntryLookupLock, lock2);
+
+					writeLine(line.Clear().Appendf(TC("Bucket #%u"), bucket.index));
+
 					Set<TString> sortedPaths;
 
 					bucket.m_pathTable.TraversePaths([&](const StringView& path)
@@ -1833,7 +1838,7 @@ namespace uba
 						});
 
 					for (auto& path : sortedPaths)
-						writeLine(path);
+						writeLine(line.Clear().Append(TC("   ")).Append(path));
 				}
 		}
 		else if (command.Equals(TC("obliterate")))
diff --git a/Engine/Source/Programs/UnrealBuildAccelerator/Common/Private/UbaCompactTables.cpp b/Engine/Source/Programs/UnrealBuildAccelerator/Common/Private/UbaCompactTables.cpp
--- a/Engine/Source/Programs/UnrealBuildAccelerator/Common/Private/UbaCompactTables.cpp
+++ b/Engine/Source/Programs/UnrealBuildAccelerator/Common/Private/UbaCompactTables.cpp
@@ -11,8 +11,9 @@ namespace uba
 	static constexpr u32 CasKeyArrayMaxSize = 16;
 
 
-	CompactPathTable::CompactPathTable(bool caseInsensitive, u64 reservePathCount, u64 reserveSegmentCount)
+	CompactPathTable::CompactPathTable(bool caseInsensitive, u64 reservePathCount, u64 reserveSegmentCount, u32 version)
 	:	m_caseInsensitive(caseInsensitive)
+	,	m_version(version)
 	{
 		m_offsets.Init(reservePathCount);
 		m_segmentOffsets.Init(reserveSegmentCount);
@@ -59,14 +60,34 @@ namespace uba
 		const tchar* seg = str;
 		u32 parentOffset = 0;
 		
-		for (const tchar* it = str + strLen - 1; it > str; --it)
+		if (m_version >= 2)
 		{
-			if (*it != PathSeparator)
-				continue;
+			for (const tchar* it = str + strLen - 1; it > str; --it)
+			{
+				if (*it != PathSeparator && *it != '.')
+					continue;
+
+				parentOffset = InternalAdd(str, stringKeyString, it - str);
+
+				parentOffset = (parentOffset << 1);
+				if (*it == '.')
+					++parentOffset;
 
-			parentOffset = InternalAdd(str, stringKeyString, it - str);
-			seg = it + 1;
-			break;
+
+				seg = it + 1;
+				break;
+			}
+		}
+		else
+		{
+			for (const tchar* it = str + strLen - 1; it > str; --it)
+			{
+				if (*it != PathSeparator)
+					continue;
+				parentOffset = InternalAdd(str, stringKeyString, it - str);
+				seg = it + 1;
+				break;
+			}
 		}
 
 		u64 segLen = strLen - (seg - str);
@@ -118,6 +139,7 @@ namespace uba
 		#endif
 
 		u32 offsets[256];
+		bool separators[256];
 		offsets[0] = u32(offset);
 		u32 offsetCount = 0;
 
@@ -129,6 +151,12 @@ namespace uba
 			UBA_ASSERT(offsetCount < sizeof_array(offsets));
 			reader.SetPosition(offset);
 			offset = (u32)reader.Read7BitEncoded();
+			separators[offsetCount] = 0;
+			if (m_version >= 2)
+			{
+				separators[offsetCount] = offset & 1;
+				offset >>= 1;
+			}
 			offsets[offsetCount] = u32(offset);
 		}
 
@@ -143,7 +171,7 @@ namespace uba
 					reader.SetPosition(strOffset);
 
 				if (!isFirst)
-					out.Append(PathSeparator);
+					out.Append(separators[i] ? '.' : PathSeparator);
 				isFirst = false;
 				reader.ReadString(out);
 			}
@@ -166,6 +194,11 @@ namespace uba
 		return m_caseInsensitive;
 	}
 
+	u32 CompactPathTable::GetVersion()
+	{
+		return m_version;
+	}
+
 	void CompactPathTable::ReadMem(BinaryReader& reader, bool populateLookup)
 	{
 		if (!m_pathTableMem.memory)
diff --git a/Engine/Source/Programs/UnrealBuildAccelerator/Common/Private/UbaTrace.cpp b/Engine/Source/Programs/UnrealBuildAccelerator/Common/Private/UbaTrace.cpp
--- a/Engine/Source/Programs/UnrealBuildAccelerator/Common/Private/UbaTrace.cpp
+++ b/Engine/Source/Programs/UnrealBuildAccelerator/Common/Private/UbaTrace.cpp
@@ -68,7 +68,7 @@ namespace uba
 
 		if (!m_memoryBegin)
 		{
-			m_logger.Warning(TC("Failed to map view of trace mapping (%s)"), LastErrorToText().data);
+			m_logger.Warning(TC("Failed to map view of trace mapping '%s' (%s)"), namedTrace, LastErrorToText().data);
 			CloseFileMapping(m_logger, m_memoryHandle, TC("Trace"));
 			m_memoryHandle = {};
 			return false;
diff --git a/Engine/Source/Programs/UnrealBuildAccelerator/Common/Public/UbaCompactTables.h b/Engine/Source/Programs/UnrealBuildAccelerator/Common/Public/UbaCompactTables.h
--- a/Engine/Source/Programs/UnrealBuildAccelerator/Common/Public/UbaCompactTables.h
+++ b/Engine/Source/Programs/UnrealBuildAccelerator/Common/Public/UbaCompactTables.h
@@ -15,7 +15,7 @@ namespace uba
 	class CompactPathTable
 	{
 	public:
-		CompactPathTable(bool caseSensitive, u64 reservePathCount, u64 reserveSegmentCount);
+		CompactPathTable(bool caseSensitive, u64 reservePathCount, u64 reserveSegmentCount, u32 version);
 
 		u32 Add(const tchar* str, u64 strLen);
 		u32 Add(const tchar* str, u64 strLen, u32& outRequiredPathTableSize);
@@ -26,6 +26,7 @@ namespace uba
 		u8* GetMemory();
 		u32 GetSize();
 		bool GetCaseInsensitive();
+		u32 GetVersion();
 
 		void ReadMem(BinaryReader& reader, bool populateLookup);
 		void Swap(CompactPathTable& other);
@@ -42,6 +43,7 @@ namespace uba
 		MemoryBlock m_pathTableMem;
 		HashMap<StringKey, u32, true> m_offsets;
 		HashMap<StringKey, u32, true> m_segmentOffsets;
+		u32 m_version;
 		bool m_caseInsensitive;
 	};
 
diff --git a/Engine/Source/Programs/UnrealBuildAccelerator/Common/Public/UbaNetwork.h b/Engine/Source/Programs/UnrealBuildAccelerator/Common/Public/UbaNetwork.h
--- a/Engine/Source/Programs/UnrealBuildAccelerator/Common/Public/UbaNetwork.h
+++ b/Engine/Source/Programs/UnrealBuildAccelerator/Common/Public/UbaNetwork.h
@@ -18,6 +18,7 @@ namespace uba
 	static constexpr u32 StorageNetworkVersion = 4;
 	static constexpr u32 SessionNetworkVersion = 42;
 	static constexpr u32 CacheNetworkVersion = 5;
+	static constexpr u32 CacheBucketVersion = 2;
 
 	// Messages used over network between client/server (system, storage and session)
 
diff --git a/Engine/Source/Programs/UnrealBuildAccelerator/Common/Public/UbaRootPaths.h b/Engine/Source/Programs/UnrealBuildAccelerator/Common/Public/UbaRootPaths.h
--- a/Engine/Source/Programs/UnrealBuildAccelerator/Common/Public/UbaRootPaths.h
+++ b/Engine/Source/Programs/UnrealBuildAccelerator/Common/Public/UbaRootPaths.h
@@ -7,7 +7,6 @@
 
 namespace uba
 {
-	static constexpr u8 RootPathsVersion = 1;
 	static constexpr u32 PathsPerRoot = IsWindows ? 4u : 1;
 
 	class RootPaths
diff --git a/Engine/Source/Programs/UnrealBuildAccelerator/Test/Private/UbaTestCache.h b/Engine/Source/Programs/UnrealBuildAccelerator/Test/Private/UbaTestCache.h
--- a/Engine/Source/Programs/UnrealBuildAccelerator/Test/Private/UbaTestCache.h
+++ b/Engine/Source/Programs/UnrealBuildAccelerator/Test/Private/UbaTestCache.h
@@ -212,29 +212,32 @@ namespace uba
 
 	bool TestCompactPathTable(Logger& logger, const StringBufferBase& rootDir)
 	{
-		CompactPathTable table(CaseInsensitiveFs, 0, 0);
+		for (u32 version=0; version<CacheBucketVersion; ++version)
+		{
+			CompactPathTable table(CaseInsensitiveFs, 0, 0, version);
 
-		StringBuffer<> str;
-		str.Append("foo").EnsureEndsWithSlash().Append("bar");
-		u32 offset = table.Add(str.data, str.count);
+			StringBuffer<> str;
+			str.Append("foo").EnsureEndsWithSlash().Append("bar.h");
+			u32 offset = table.Add(str.data, str.count);
 
-		StringBuffer<> str2;
-		table.GetString(str2, offset);
-		if (!str.Equals(str2.data))
-			return false;
+			StringBuffer<> str2;
+			table.GetString(str2, offset);
+			if (!str.Equals(str2.data))
+				return false;
 
-		str.Clear().Append(PathSeparator).Append("foo").Append(PathSeparator).Append("bar");
-		offset = table.Add(str.data, str.count);
-		table.GetString(str2.Clear(), offset);
-		if (!str.Equals(str2.data))
-			return false;
+			str.Clear().Append(PathSeparator).Append("foo").Append(PathSeparator).Append("bar.h");
+			offset = table.Add(str.data, str.count);
+			table.GetString(str2.Clear(), offset);
+			if (!str.Equals(str2.data))
+				return false;
 
-		CompactPathTable table2(CaseInsensitiveFs, 0, 0);
-		BinaryReader reader(table.GetMemory(), 0, table.GetSize());
-		table2.ReadMem(reader, true);
-		u32 offset2 = table2.Add(str.data, str.count);
-		if (offset != offset2)
-			return false;
+			CompactPathTable table2(CaseInsensitiveFs, 0, 0, version);
+			BinaryReader reader(table.GetMemory(), 0, table.GetSize());
+			table2.ReadMem(reader, true);
+			u32 offset2 = table2.Add(str.data, str.count);
+			if (offset != offset2)
+				return false;
+		}
 		return true;
 	}
 
