[UBA]
* Fixed linux compile warning

diff --git a/Engine/Source/Programs/UnrealBuildAccelerator/Common/Private/UbaCompactTables.cpp b/Engine/Source/Programs/UnrealBuildAccelerator/Common/Private/UbaCompactTables.cpp
--- a/Engine/Source/Programs/UnrealBuildAccelerator/Common/Private/UbaCompactTables.cpp
+++ b/Engine/Source/Programs/UnrealBuildAccelerator/Common/Private/UbaCompactTables.cpp
@@ -12,8 +12,8 @@ namespace uba
 
 
 	CompactPathTable::CompactPathTable(bool caseInsensitive, u64 reservePathCount, u64 reserveSegmentCount, u32 version)
-	:	m_caseInsensitive(caseInsensitive)
-	,	m_version(version)
+	:	m_version(version)
+	,	m_caseInsensitive(caseInsensitive)
 	{
 		m_offsets.Init(reservePathCount);
 		m_segmentOffsets.Init(reserveSegmentCount);
