[UBA]
* Added version to log entry when saving bucket to disk

diff --git a/Engine/Source/Programs/UnrealBuildAccelerator/Common/Private/UbaCacheServer.cpp b/Engine/Source/Programs/UnrealBuildAccelerator/Common/Private/UbaCacheServer.cpp
--- a/Engine/Source/Programs/UnrealBuildAccelerator/Common/Private/UbaCacheServer.cpp
+++ b/Engine/Source/Programs/UnrealBuildAccelerator/Common/Private/UbaCacheServer.cpp
@@ -478,7 +478,7 @@ namespace uba
 
 		bucket.lastSavedTime = GetSystemTimeAsFileTime() - m_creationTime;
 
-		m_logger.Detail(TC("    Bucket %u saved - %s (%s)"), bucket.index, BytesToText(file.written).str, TimeToText(GetTime() - saveStart).str);
+		m_logger.Detail(TC("    Bucket %u saved (v%u) - %s (%s)"), bucket.index, bucket.m_pathTable.GetVersion(), BytesToText(file.written).str, TimeToText(GetTime() - saveStart).str);
 		return true;
 	}
 
