[UBA]
* Minor cosmetic log change

diff --git a/Engine/Source/Programs/UnrealBuildAccelerator/Common/Private/UbaCacheServer.cpp b/Engine/Source/Programs/UnrealBuildAccelerator/Common/Private/UbaCacheServer.cpp
--- a/Engine/Source/Programs/UnrealBuildAccelerator/Common/Private/UbaCacheServer.cpp
+++ b/Engine/Source/Programs/UnrealBuildAccelerator/Common/Private/UbaCacheServer.cpp
@@ -1804,7 +1804,7 @@ namespace uba
 					if (lastUsed)
 						lastUsedTime = GetFileTimeAsTime(GetSystemTimeAsFileTime() - (m_creationTime + lastUsed));
 
-					writeLine(line.Clear().Appendf(TC("    #%u - %llu (Version %u)"), index++, kv.first, bucket.m_pathTable.GetVersion()));
+					writeLine(line.Clear().Appendf(TC("    #%u - %llu (v%u)"), index++, kv.first, bucket.m_pathTable.GetVersion()));
 					writeLine(line.Clear().Appendf(TC("      PathTable: %llu (%s)"), bucket.m_pathTable.GetPathCount(), BytesToText(bucket.m_pathTable.GetSize()).str));
 					writeLine(line.Clear().Appendf(TC("      CasKeyTable: %llu (%s)"), bucket.m_cacheEntryLookup.size(), BytesToText(bucket.m_casKeyTable.GetSize()).str));
 					writeLine(line.Clear().Appendf(TC("      TotalEntries: %llu"), totalEntryCount));
