[UBA]
* Fixed uba visualizer stats bug caused by fixing process reuse

diff --git a/Engine/Source/Programs/UnrealBuildAccelerator/Common/Private/UbaFileMapping.cpp b/Engine/Source/Programs/UnrealBuildAccelerator/Common/Private/UbaFileMapping.cpp
--- a/Engine/Source/Programs/UnrealBuildAccelerator/Common/Private/UbaFileMapping.cpp
+++ b/Engine/Source/Programs/UnrealBuildAccelerator/Common/Private/UbaFileMapping.cpp
@@ -274,6 +274,7 @@ namespace uba
 	bool MapViewCommit(void* address, u64 size)
 	{
 #if PLATFORM_WINDOWS
+		ExtendedTimerScope ts(KernelStats::GetCurrent().virtualAlloc);
 		return ::VirtualAlloc(address, size, MEM_COMMIT, PAGE_READWRITE);
 #else
 		return true;
diff --git a/Engine/Source/Programs/UnrealBuildAccelerator/Visualizer/Private/UbaTraceReader.cpp b/Engine/Source/Programs/UnrealBuildAccelerator/Visualizer/Private/UbaTraceReader.cpp
--- a/Engine/Source/Programs/UnrealBuildAccelerator/Visualizer/Private/UbaTraceReader.cpp
+++ b/Engine/Source/Programs/UnrealBuildAccelerator/Visualizer/Private/UbaTraceReader.cpp
@@ -684,6 +684,7 @@ namespace uba
 					breadcrumbs = reader.ReadLongString();
 			}
 
+			process.isReuse = true;
 			process.exitCode = 0u;
 			process.stop = time;
 			process.bitmapDirty = true;
diff --git a/Engine/Source/Programs/UnrealBuildAccelerator/Visualizer/Private/UbaTraceReader.h b/Engine/Source/Programs/UnrealBuildAccelerator/Visualizer/Private/UbaTraceReader.h
--- a/Engine/Source/Programs/UnrealBuildAccelerator/Visualizer/Private/UbaTraceReader.h
+++ b/Engine/Source/Programs/UnrealBuildAccelerator/Visualizer/Private/UbaTraceReader.h
@@ -27,6 +27,7 @@ namespace uba
 			bool bitmapDirty = true;
 			bool cacheFetch = false;
 			bool isRemote = false;
+			bool isReuse = false;
 			u64 createFilesTime = 0;
 			u64 writeFilesTime = 0;
 			Vector<u8> stats;
diff --git a/Engine/Source/Programs/UnrealBuildAccelerator/Visualizer/Private/UbaVisualizer.cpp b/Engine/Source/Programs/UnrealBuildAccelerator/Visualizer/Private/UbaVisualizer.cpp
--- a/Engine/Source/Programs/UnrealBuildAccelerator/Visualizer/Private/UbaVisualizer.cpp
+++ b/Engine/Source/Programs/UnrealBuildAccelerator/Visualizer/Private/UbaVisualizer.cpp
@@ -1845,7 +1845,7 @@ namespace uba
 
 					if (reader.GetLeft())
 					{
-						if (process.isRemote || m_traceView.version < 35)
+						if (process.isRemote || (m_traceView.version >= 36 && !process.isReuse))
 							sessionStats.Read(reader, m_traceView.version);
 						storageStats.Read(reader, m_traceView.version);
 						kernelStats.Read(reader, m_traceView.version);
@@ -2839,7 +2839,7 @@ namespace uba
 			processStats.Read(reader, m_traceView.version);
 			if (reader.GetLeft())
 			{
-				if (process.isRemote || m_traceView.version < 35)
+				if (process.isRemote || (m_traceView.version >= 36 && !process.isReuse))
 					sessionStats.Read(reader, m_traceView.version);
 				storageStats.Read(reader, m_traceView.version);
 				kernelStats.Read(reader, m_traceView.version);
