[UBA]
* Added cache write stats to UbaVisualizer process copy-paste stats

diff --git a/Engine/Source/Programs/UnrealBuildAccelerator/Visualizer/Private/UbaVisualizer.cpp b/Engine/Source/Programs/UnrealBuildAccelerator/Visualizer/Private/UbaVisualizer.cpp
--- a/Engine/Source/Programs/UnrealBuildAccelerator/Visualizer/Private/UbaVisualizer.cpp
+++ b/Engine/Source/Programs/UnrealBuildAccelerator/Visualizer/Private/UbaVisualizer.cpp
@@ -757,6 +757,19 @@ namespace uba
 		m_lastBitmap = 0;
 	}
 
+	void Visualizer::PrintCacheWriteStats(Logger& logger, u32 processId)
+	{
+		auto findIt = m_traceView.cacheWrites.find(processId);
+		if (findIt == m_traceView.cacheWrites.end())
+			return;
+		TraceView::CacheWrite& write = findIt->second;
+		logger.Info(L"");
+		logger.Info(L"  -------- Cache write stats ----------");
+		logger.Info(L"  Duration                    %9s", TimeToText(write.end - write.start).str);
+		logger.Info(L"  Success                     %9s", write.success ? L"true" : L"false");
+		logger.Info(L"  Bytes sent                  %9s", BytesToText(write.bytesSent).str);
+	}
+
 	void Visualizer::UpdateFont(Font& font, int height, bool createUnderline)
 	{
 		font.height = height;
@@ -1893,17 +1906,7 @@ namespace uba
 					kernelStats.Print(logger, false, m_traceView.frequency);
 				}
 
-				auto findIt = m_traceView.cacheWrites.find(process.id);
-				if (findIt != m_traceView.cacheWrites.end())
-				{
-					TraceView::CacheWrite& write = findIt->second;
-					logger.Info(L"");
-					logger.Info(L"  -------- Cache write stats ----------");
-					logger.Info(L"  Duration                    %9s", TimeToText(write.end - write.start).str);
-					logger.Info(L"  Success                     %9s", write.success ? L"true" : L"false");
-					logger.Info(L"  Bytes sent                  %9s", BytesToText(write.bytesSent).str);
-
-				}
+				PrintCacheWriteStats(logger, process.id);
 
 				if (!logLines.empty())
 				{
@@ -2868,6 +2871,8 @@ namespace uba
 				out.Info(L"  ----------- Kernel stats ------------");
 				kernelStats.Print(out, false, m_traceView.frequency);
 			}
+
+			PrintCacheWriteStats(out, process.id);
 		}
 	}
 
