[UBA]
* Added cache write stats to UbaVisualizer process copy-paste stats

diff --git a/Engine/Source/Programs/UnrealBuildAccelerator/Visualizer/Private/UbaVisualizer.h b/Engine/Source/Programs/UnrealBuildAccelerator/Visualizer/Private/UbaVisualizer.h
--- a/Engine/Source/Programs/UnrealBuildAccelerator/Visualizer/Private/UbaVisualizer.h
+++ b/Engine/Source/Programs/UnrealBuildAccelerator/Visualizer/Private/UbaVisualizer.h
@@ -135,6 +135,7 @@ namespace uba
 		void StopDragToScroll();
 		void SaveSettings();
 		void DirtyBitmaps(bool full);
+		void PrintCacheWriteStats(Logger& logger, u32 processId);
 
 		struct Font
 		{
