[UBA]
* Added option to disable connect/disconnect logging on network server
* Moved connect/disconnect logging to cache server and added useful information to log entries

diff --git a/Engine/Source/Programs/UnrealBuildAccelerator/CacheService/Private/UbaCacheService.cpp b/Engine/Source/Programs/UnrealBuildAccelerator/CacheService/Private/UbaCacheService.cpp
--- a/Engine/Source/Programs/UnrealBuildAccelerator/CacheService/Private/UbaCacheService.cpp
+++ b/Engine/Source/Programs/UnrealBuildAccelerator/CacheService/Private/UbaCacheService.cpp
@@ -221,6 +221,7 @@ namespace uba
 		NetworkBackendTcp networkBackend(logWriter);
 		NetworkServerCreateInfo nsci(logWriter);
 		//nsci.workerCount = 4;
+		nsci.logConnections = false; // Let the cache server report instead
 		bool ctorSuccess = true;
 		NetworkServer networkServer(ctorSuccess, nsci);
 		if (!ctorSuccess)
diff --git a/Engine/Source/Programs/UnrealBuildAccelerator/Common/Private/UbaCacheClient.cpp b/Engine/Source/Programs/UnrealBuildAccelerator/Common/Private/UbaCacheClient.cpp
--- a/Engine/Source/Programs/UnrealBuildAccelerator/Common/Private/UbaCacheClient.cpp
+++ b/Engine/Source/Programs/UnrealBuildAccelerator/Common/Private/UbaCacheClient.cpp
@@ -463,7 +463,7 @@ namespace uba
 		if (!entryCount)
 		{
 			if (m_reportMissReason)
-				m_logger.Info(TC("Cache miss on %s because no entry with key %s was found in bucket %u"), info.GetDescription(), CasKeyString(cmdKey).str, bucketId);
+				m_logger.Info(TC("Cache miss on %s because no entry with key %s was found in bucket %u (%llu)"), info.GetDescription(), CasKeyString(cmdKey).str, bucketId, MakeId(bucketId));
 			return false;
 		}
 
diff --git a/Engine/Source/Programs/UnrealBuildAccelerator/Common/Private/UbaCacheServer.cpp b/Engine/Source/Programs/UnrealBuildAccelerator/Common/Private/UbaCacheServer.cpp
--- a/Engine/Source/Programs/UnrealBuildAccelerator/Common/Private/UbaCacheServer.cpp
+++ b/Engine/Source/Programs/UnrealBuildAccelerator/Common/Private/UbaCacheServer.cpp
@@ -136,7 +136,11 @@ namespace uba
 	struct CacheServer::Connection
 	{
 		u32 clientVersion;
-		UnorderedMap<u64, ConnectionBucket> buckets;
+		UnorderedMap<u64, ConnectionBucket> storeBuckets;
+		u64 storeEntryCount = 0;
+		UnorderedSet<u64> fetchBuckets;
+		u64 fetchEntryCount = 0;
+		u64 fetchEntryHitCount = 0;
 	};
 
 	struct CacheServer::Bucket
@@ -1202,9 +1206,35 @@ namespace uba
 
 	void CacheServer::OnDisconnected(u32 clientId)
 	{
+		u64 fetchBucketsCount = 0;
+		u64 fetchEntryCount = 0;
+		u64 storeBucketsCount = 0;
+		u64 storeEntryCount = 0;
+		u64 fetchEntryHitCount = 0;
 		SCOPED_WRITE_LOCK(m_connectionsLock, lock);
-		m_connections.erase(clientId);
+		auto it = m_connections.find(clientId);
+		if (it != m_connections.end())
+		{
+			auto& conn = it->second;
+			fetchBucketsCount = conn.fetchBuckets.size();
+			fetchEntryCount = conn.fetchEntryCount;
+			fetchEntryHitCount = conn.fetchEntryHitCount;
+			storeBucketsCount = conn.storeBuckets.size();
+			storeEntryCount = conn.storeEntryCount;
+			m_connections.erase(it);
+		}
+		u64 activeCount = m_connections.size();
 		lock.Leave();
+
+		StringBuffer<> log;
+		log.Appendf(TC("Client %u disconnected (%llu active). "), clientId, activeCount);
+		if (fetchBucketsCount)
+			log.Appendf(TC("Fetched %llu (%llu hits) from %llu bucket(s)"), fetchEntryCount, fetchEntryHitCount, fetchBucketsCount);
+		if (fetchBucketsCount && storeBucketsCount)
+			log.Append(TC(", "));
+		if (storeBucketsCount)
+			log.Appendf(TC("Stored %llu to %llu bucket(s)"), storeEntryCount, storeBucketsCount);
+		m_logger.Log(LogEntryType_Info, log.data, log.count);
 	}
 
 	CacheServer::ConnectionBucket& CacheServer::GetConnectionBucket(const ConnectionInfo& connectionInfo, BinaryReader& reader, u32* outClientVersion)
@@ -1214,8 +1244,11 @@ namespace uba
 		SCOPED_WRITE_LOCK(m_connectionsLock, lock);
 		auto& connection = m_connections[connectionInfo.GetId()];
 		if (outClientVersion)
+		{
+			++connection.storeEntryCount; // We know this is the store entry call
 			*outClientVersion = connection.clientVersion;
-		return connection.buckets.try_emplace(id, id, bucketVersion).first->second;
+		}
+		return connection.storeBuckets.try_emplace(id, id, bucketVersion).first->second;
 	}
 
 	CacheServer::Bucket& CacheServer::GetBucket(BinaryReader& reader)
@@ -1269,6 +1302,9 @@ namespace uba
 			auto& connection = insres.first->second;
 			connection.clientVersion = clientVersion;
 			m_peakConnectionCount = Max(m_peakConnectionCount, u32(m_connections.size()));
+
+			m_logger.Info(TC("Client %u connected (v%u)"), connectionInfo.GetId(), clientVersion);
+
 			return true;
 		}
 		case CacheMessageType_StorePathTable:
@@ -1323,22 +1359,14 @@ namespace uba
 			return HandleFetchCasTable(reader, writer);
 
 		case CacheMessageType_FetchEntries:
-		{
-			SCOPED_READ_LOCK(m_connectionsLock, lock);
-			u32 clientVersion = m_connections[connectionInfo.GetId()].clientVersion;
-			lock.Leave();
-			return HandleFetchEntries(reader, writer, clientVersion);
-		}
+			return HandleFetchEntries(reader, writer, connectionInfo.GetId());
+
 		case CacheMessageType_ExecuteCommand:
 			return HandleExecuteCommand(reader, writer);
 
 		case CacheMessageType_ReportUsedEntry:
-		{
-			SCOPED_READ_LOCK(m_connectionsLock, lock);
-			u32 clientVersion = m_connections[connectionInfo.GetId()].clientVersion;
-			lock.Leave();
-			return HandleReportUsedEntry(reader, writer, clientVersion);
-		}
+			return HandleReportUsedEntry(reader, writer, connectionInfo.GetId());
+
 		case CacheMessageType_RequestShutdown:
 		{
 			TString reason = reader.ReadString();
@@ -1595,11 +1623,20 @@ namespace uba
 		return true;
 	}
 
-	bool CacheServer::HandleFetchEntries(BinaryReader& reader, BinaryWriter& writer, u32 clientVersion)
+	bool CacheServer::HandleFetchEntries(BinaryReader& reader, BinaryWriter& writer, u32 clientId)
 	{
 		Bucket& bucket = GetBucket(reader);
 		CasKey cmdKey = reader.ReadCasKey();
 
+		u32 clientVersion;
+		{
+			SCOPED_WRITE_LOCK(m_connectionsLock, lock);
+			auto& conn = m_connections[clientId];
+			conn.fetchBuckets.insert(bucket.index);
+			++conn.fetchEntryCount;
+			clientVersion = conn.clientVersion;
+		}
+
 		++m_cacheKeyFetchCount;
 
 		SCOPED_READ_LOCK(bucket.m_cacheEntryLookupLock, lock);
@@ -1616,12 +1653,20 @@ namespace uba
 		return cacheEntries.Write(writer, clientVersion, false);
 	}
 
-	bool CacheServer::HandleReportUsedEntry(BinaryReader& reader, BinaryWriter& writer, u32 clientVersion)
+	bool CacheServer::HandleReportUsedEntry(BinaryReader& reader, BinaryWriter& writer, u32 clientId)
 	{
 		Bucket& bucket = GetBucket(reader);
 		CasKey cmdKey = reader.ReadCasKey();
 		u64 entryId = reader.Read7BitEncoded();
 
+		u32 clientVersion;
+		{
+			SCOPED_WRITE_LOCK(m_connectionsLock, lock);
+			auto& conn = m_connections[clientId];
+			++conn.fetchEntryHitCount;
+			clientVersion = conn.clientVersion;
+		}
+
 		++m_cacheKeyHitCount;
 
 		SCOPED_READ_LOCK(bucket.m_cacheEntryLookupLock, lock);
diff --git a/Engine/Source/Programs/UnrealBuildAccelerator/Common/Private/UbaNetworkServer.cpp b/Engine/Source/Programs/UnrealBuildAccelerator/Common/Private/UbaNetworkServer.cpp
--- a/Engine/Source/Programs/UnrealBuildAccelerator/Common/Private/UbaNetworkServer.cpp
+++ b/Engine/Source/Programs/UnrealBuildAccelerator/Common/Private/UbaNetworkServer.cpp
@@ -302,11 +302,14 @@ namespace uba
 			{
 				if (server.m_onConnectionFunction)
 					server.m_onConnectionFunction(clientUid, clientId);
-				server.m_logger.Detail(TC("Client %u (%s) connected on connection %s"), clientId, GuidToString(clientUid).str, GuidToString(connectionUid).str);
+				if (server.m_logConnections)
+					server.m_logger.Detail(TC("Client %u (%s) connected on connection %s"), clientId, GuidToString(clientUid).str, GuidToString(connectionUid).str);
 			}
 			else
-				server.m_logger.Detail(TC("Client %u (%s) additional connection %s connected"), clientId, GuidToString(clientUid).str, GuidToString(connectionUid).str);
-
+			{
+				if (server.m_logConnections)
+					server.m_logger.Detail(TC("Client %u (%s) additional connection %s connected"), clientId, GuidToString(clientUid).str, GuidToString(connectionUid).str);
+			}
 
 			return true;
 		}
@@ -398,7 +401,8 @@ namespace uba
 				SCOPED_READ_LOCK(m_server.m_onDisconnectFunctionsLock, l);
 				for (auto& entry : m_server.m_onDisconnectFunctions)
 					entry.function(m_client->uid, m_client->id);
-				m_server.m_logger.Detail(TC("Client %u (%s) disconnected"), m_client->id, GuidToString(m_client->uid).str);
+				if (m_server.m_logConnections)
+					m_server.m_logger.Detail(TC("Client %u (%s) disconnected"), m_client->id, GuidToString(m_client->uid).str);
 			}
 			m_disconnected = true;
 		}
@@ -633,6 +637,7 @@ namespace uba
 			m_logger.Detail(TC("Adjusted msg size to %u to stay inside limits"), fixedSendSize);
 		m_sendSize = fixedSendSize;
 		m_receiveTimeoutMs = info.receiveTimeoutSeconds * 1000;
+		m_logConnections = info.logConnections;
 
 		m_workerFunctions[SystemServiceId].toString = GetMessageTypeToName;
 		m_workerFunctions[SystemServiceId].func = [this](const ConnectionInfo& connectionInfo, MessageInfo& messageInfo, BinaryReader& reader, BinaryWriter& writer)
diff --git a/Engine/Source/Programs/UnrealBuildAccelerator/Common/Public/UbaCacheServer.h b/Engine/Source/Programs/UnrealBuildAccelerator/Common/Public/UbaCacheServer.h
--- a/Engine/Source/Programs/UnrealBuildAccelerator/Common/Public/UbaCacheServer.h
+++ b/Engine/Source/Programs/UnrealBuildAccelerator/Common/Public/UbaCacheServer.h
@@ -80,8 +80,8 @@ namespace uba
 		bool HandleStoreEntry(ConnectionBucket& bucket, BinaryReader& reader, BinaryWriter& writer, u32 clientVersion);
 		bool HandleFetchPathTable(BinaryReader& reader, BinaryWriter& writer);
 		bool HandleFetchCasTable(BinaryReader& reader, BinaryWriter& writer);
-		bool HandleFetchEntries(BinaryReader& reader, BinaryWriter& writer, u32 clientVersion);
-		bool HandleReportUsedEntry(BinaryReader& reader, BinaryWriter& writer, u32 clientVersion);
+		bool HandleFetchEntries(BinaryReader& reader, BinaryWriter& writer, u32 clientId);
+		bool HandleReportUsedEntry(BinaryReader& reader, BinaryWriter& writer, u32 clientId);
 		bool HandleExecuteCommand(BinaryReader& reader, BinaryWriter& writer);
 
 		MutableLogger m_logger;
diff --git a/Engine/Source/Programs/UnrealBuildAccelerator/Common/Public/UbaNetworkServer.h b/Engine/Source/Programs/UnrealBuildAccelerator/Common/Public/UbaNetworkServer.h
--- a/Engine/Source/Programs/UnrealBuildAccelerator/Common/Public/UbaNetworkServer.h
+++ b/Engine/Source/Programs/UnrealBuildAccelerator/Common/Public/UbaNetworkServer.h
@@ -36,6 +36,7 @@ namespace uba
 		u32 workerCount = 0; // Zero means it will use the number of logical cores as worker count
 		u32 sendSize = SendDefaultSize;
 		u32 receiveTimeoutSeconds = 0;
+		bool logConnections = true;
 	};
 
 	struct MessageInfo
@@ -154,6 +155,7 @@ namespace uba
 		u32 m_createdWorkerCount = 0;
 		u32 m_maxCreatedWorkerCount = 0;
 		bool m_workersEnabled = true;
+		bool m_logConnections = true;
 
 		ReaderWriterLock m_addConnectionsLock;
 		List<Thread> m_addConnections;
