[UBA]
* Fixed so :// paths are not seen as paths to normalize (http://, https://, file:// etc)

diff --git a/Engine/Source/Programs/UnrealBuildAccelerator/Common/Public/UbaRootPaths.h b/Engine/Source/Programs/UnrealBuildAccelerator/Common/Public/UbaRootPaths.h
--- a/Engine/Source/Programs/UnrealBuildAccelerator/Common/Public/UbaRootPaths.h
+++ b/Engine/Source/Programs/UnrealBuildAccelerator/Common/Public/UbaRootPaths.h
@@ -86,7 +86,7 @@ namespace uba
 		
 			if (absPathChars == strEnd)
 			{
-				func(searchPos, strEnd - searchPos, ~0u);
+				func(searchPos, strEnd - searchPos, ~0u); // Yes, this contains the terminating character too which is a bit weird
 				return true;
 			}
 
@@ -104,7 +104,23 @@ namespace uba
 			auto root = FindRoot(path);
 			if (!root)
 			{
-				if (allowPathsWithoutRoot || InternalFindRoot(m_ignoredRoots, path))
+				bool skip = allowPathsWithoutRoot;
+#if PLATFORM_WINDOWS
+				if (str + 2 <= absPathChars)
+				{
+					CharType driveLetter = absPathChars[-2];
+					if (!(driveLetter >= 'a' && driveLetter <= 'z' || driveLetter >= 'A' && driveLetter <= 'Z')) // Not a drive
+						skip = true;
+					else if (str + 3 <= absPathChars && absPathChars[0] == '/' && absPathChars[1] == '/')
+					{
+						CharType preDrive = absPathChars[-3];
+						if (preDrive >= 'a' && preDrive <= 'z' || preDrive >= 'A' && preDrive <= 'Z')  // http://, https:// or file://
+							skip = true;
+					}
+				}
+#endif
+
+				if (skip || InternalFindRoot(m_ignoredRoots, path))
 				{
 					u32 len = u32(absPathChars - searchPos) + 1;
 					destPos += len;
diff --git a/Engine/Source/Programs/UnrealBuildAccelerator/Test/Private/UbaTestBasics.h b/Engine/Source/Programs/UnrealBuildAccelerator/Test/Private/UbaTestBasics.h
--- a/Engine/Source/Programs/UnrealBuildAccelerator/Test/Private/UbaTestBasics.h
+++ b/Engine/Source/Programs/UnrealBuildAccelerator/Test/Private/UbaTestBasics.h
@@ -396,7 +396,7 @@ namespace uba
 		bool success = true;
 		StringBuffer<> temp;
 		u32 rootPos = ~0u;
-		paths.NormalizeString(logger, str, sizeof_array(str), [&](const tchar* str, u64 strLen, u32 rp)
+		bool res = paths.NormalizeString(logger, str, sizeof_array(str), [&](const tchar* str, u64 strLen, u32 rp)
 			{
 				if (rp != ~0u)
 				{
@@ -414,6 +414,8 @@ namespace uba
 				}
 			}, false, TC(""));
 
+		if (!res)
+			return false;
 		if (!success)
 			return false;
 
@@ -423,6 +425,28 @@ namespace uba
 		if (!newStr.Equals(str))
 			return false;
 
+		#if PLATFORM_WINDOWS
+		const tchar str2[] = TC("file://e:/temp/");
+		bool foundPath = false;
+		res = paths.NormalizeString(logger, str2, sizeof_array(str2), [&](const tchar* str, u64 strLen, u32 rp)
+			{
+				if (rp != ~0)
+				{
+					if (str[0] != RootPaths::RootStartByte + PathsPerRoot) // Add one for windows because second entry is backslash
+						success = false;
+					foundPath = true;
+				}
+				else
+				{
+					if (!((strLen == 1 && (!str[0] || str[0] == '/')) || (strLen == 6 && Equals(str, TC("file:/"), 6, false))))
+						success = false;
+				}
+
+			}, false, TC(""));
+		if (!res || !foundPath || !success)
+			return false;
+		#endif
+
 		return true;
 	}
 
