[UBA]
* Added possibility to provide hint in CacheClientCreateInfo that will be outputted in cache server log

diff --git a/Engine/Source/Programs/UnrealBuildAccelerator/Common/Private/UbaCacheClient.cpp b/Engine/Source/Programs/UnrealBuildAccelerator/Common/Private/UbaCacheClient.cpp
--- a/Engine/Source/Programs/UnrealBuildAccelerator/Common/Private/UbaCacheClient.cpp
+++ b/Engine/Source/Programs/UnrealBuildAccelerator/Common/Private/UbaCacheClient.cpp
@@ -80,7 +80,7 @@ namespace uba
 		m_useCacheHit = info.useCacheHit;
 		m_useRoots = info.useRoots;
 
-		m_client.RegisterOnConnected([this]()
+		m_client.RegisterOnConnected([this, hint = info.hint]()
 			{
 				u32 retryCount = 0;
 				while (retryCount < 10)
@@ -88,6 +88,7 @@ namespace uba
 					StackBinaryWriter<1024> writer;
 					NetworkMessage msg(m_client, CacheServiceId, CacheMessageType_Connect, writer);
 					writer.WriteU32(CacheNetworkVersion);
+					writer.WriteString(hint);
 					StackBinaryReader<1024> reader;
 					u64 sendTime = GetTime();
 					if (!msg.Send(reader))
diff --git a/Engine/Source/Programs/UnrealBuildAccelerator/Common/Private/UbaCacheServer.cpp b/Engine/Source/Programs/UnrealBuildAccelerator/Common/Private/UbaCacheServer.cpp
--- a/Engine/Source/Programs/UnrealBuildAccelerator/Common/Private/UbaCacheServer.cpp
+++ b/Engine/Source/Programs/UnrealBuildAccelerator/Common/Private/UbaCacheServer.cpp
@@ -1290,6 +1290,10 @@ namespace uba
 			if (clientVersion < 3 || clientVersion > CacheNetworkVersion)
 				return m_logger.Error(TC("Different network versions. Client: %u, Server: %u. Disconnecting"), clientVersion, CacheNetworkVersion);
 
+			TString hint;
+			if (reader.GetLeft())
+				hint = reader.ReadString();
+
 			SCOPED_WRITE_LOCK(m_connectionsLock, lock);
 			if (m_isRunningMaintenance)
 			{
@@ -1303,7 +1307,7 @@ namespace uba
 			connection.clientVersion = clientVersion;
 			m_peakConnectionCount = Max(m_peakConnectionCount, u32(m_connections.size()));
 
-			m_logger.Info(TC("Client %u connected (v%u)"), connectionInfo.GetId(), clientVersion);
+			m_logger.Info(TC("Client %u connected (v%u) %s"), connectionInfo.GetId(), clientVersion, hint.c_str());
 
 			return true;
 		}
diff --git a/Engine/Source/Programs/UnrealBuildAccelerator/Common/Public/UbaCacheClient.h b/Engine/Source/Programs/UnrealBuildAccelerator/Common/Public/UbaCacheClient.h
--- a/Engine/Source/Programs/UnrealBuildAccelerator/Common/Public/UbaCacheClient.h
+++ b/Engine/Source/Programs/UnrealBuildAccelerator/Common/Public/UbaCacheClient.h
@@ -38,6 +38,7 @@ namespace uba
 		bool validateCacheWritesOutput = false; // Set to true to validate cas of all output files before sent to cache
 		bool useRoots = true; // Set this to false to allow paths that are not under roots and to not fix them up
 		bool useCacheHit = true; // Set this to false to ignore found cache hits.. this is for debugging/testing only
+		const tchar* hint = TC(""); // Hint will show up in cache server log
 	};
 
 	struct CacheResult
