[UBA]
* Fixed so contention summary is logged to trace when enabled
* Fixed session info box in visualizer
* Moved RountUpPow2 function to UbaMemory

diff --git a/Engine/Source/Programs/UnrealBuildAccelerator/Common/Private/UbaCacheServer.cpp b/Engine/Source/Programs/UnrealBuildAccelerator/Common/Private/UbaCacheServer.cpp
--- a/Engine/Source/Programs/UnrealBuildAccelerator/Common/Private/UbaCacheServer.cpp
+++ b/Engine/Source/Programs/UnrealBuildAccelerator/Common/Private/UbaCacheServer.cpp
@@ -565,6 +565,8 @@ namespace uba
 			m_isRunningMaintenance = true;
 			lock2.Leave();
 
+			PrintContentionSummary(m_logger);
+
 			auto g = MakeGuard([&]()
 				{
 					SCOPED_WRITE_LOCK(m_connectionsLock, lock2);
diff --git a/Engine/Source/Programs/UnrealBuildAccelerator/Common/Private/UbaLogger.cpp b/Engine/Source/Programs/UnrealBuildAccelerator/Common/Private/UbaLogger.cpp
--- a/Engine/Source/Programs/UnrealBuildAccelerator/Common/Private/UbaLogger.cpp
+++ b/Engine/Source/Programs/UnrealBuildAccelerator/Common/Private/UbaLogger.cpp
@@ -686,7 +686,7 @@ namespace uba
 	void PrintContentionSummary(Logger& logger)
 	{
 	#if UBA_TRACK_CONTENTION
-		logger.Info(TC("Contention summary:"));
+		logger.Info(TC("  ------- Contention summary -------"));
 		List<ContentionTracker*> list;
 		for (auto& ct : GetContentionTrackerList())
 			if (TimeToMs(ct.time) > 1)
@@ -703,8 +703,10 @@ namespace uba
 			StringBuffer<512> fn;
 			fn.Append(ct->file);
 			StringBuffer<256> s;
-			s.Append(TC("  ")).AppendFileName(fn.data).Append(':').AppendValue(ct->line).Append(TC(" - ")).AppendValue(ct->count).Append(TC(" calls in ")).Append(TimeToText(ct->time).str);
+			s.Append(TC("  ")).AppendFileName(fn.data).Append(':').AppendValue(ct->line).Append(TC(" - ")).AppendValue(ct->count).Append(TC(" (")).Append(TimeToText(ct->time).str).Append(')');
 			logger.Info(s.data);
+			ct->count = 0;
+			ct->time = 0;
 		}
 	#endif
 	}
diff --git a/Engine/Source/Programs/UnrealBuildAccelerator/Common/Private/UbaSessionClient.cpp b/Engine/Source/Programs/UnrealBuildAccelerator/Common/Private/UbaSessionClient.cpp
--- a/Engine/Source/Programs/UnrealBuildAccelerator/Common/Private/UbaSessionClient.cpp
+++ b/Engine/Source/Programs/UnrealBuildAccelerator/Common/Private/UbaSessionClient.cpp
@@ -1560,6 +1560,7 @@ namespace uba
 				m_storage.PrintSummary(logger);
 				m_client.PrintSummary(logger);
 				KernelStats::GetGlobal().Print(logger, true);
+				PrintContentionSummary(logger);
 				if (extraInfo)
 					extraInfo(logger);
 			});
@@ -1970,6 +1971,7 @@ namespace uba
 					m_storage.PrintSummary(logger);
 					m_client.PrintSummary(logger);
 					KernelStats::GetGlobal().Print(logger, true);
+					PrintContentionSummary(logger);
 				});
 			m_trace.SessionSummary(0, writer.GetData(), writer.GetPosition());
 
diff --git a/Engine/Source/Programs/UnrealBuildAccelerator/Common/Private/UbaSessionServer.cpp b/Engine/Source/Programs/UnrealBuildAccelerator/Common/Private/UbaSessionServer.cpp
--- a/Engine/Source/Programs/UnrealBuildAccelerator/Common/Private/UbaSessionServer.cpp
+++ b/Engine/Source/Programs/UnrealBuildAccelerator/Common/Private/UbaSessionServer.cpp
@@ -254,6 +254,7 @@ namespace uba
 					m_storage.PrintSummary(logger);
 					m_server.PrintSummary(logger);
 					KernelStats::GetGlobal().Print(logger, true);
+					PrintContentionSummary(logger);
 				});
 			m_trace.SessionSummary(0, writer.GetData(), writer.GetPosition());
 		}
@@ -2061,6 +2062,8 @@ namespace uba
 	{
 		Session::PrintSessionStats(logger);
 
+		if (m_nameToHashLookup.size())
+			logger.Info(TC("  NameToHashLookup    %7u %9s"), u32(m_nameToHashLookup.size()), BytesToText(m_nameToHashTableMem.writtenSize).str);
 		logger.Info(TC("  Remote processes finished    %8u"), m_finishedRemoteProcessCount);
 		logger.Info(TC("  Remote processes returned    %8u"), m_returnedRemoteProcessCount);
 		logger.Info(TC(""));
diff --git a/Engine/Source/Programs/UnrealBuildAccelerator/Common/Public/UbaHashMap.h b/Engine/Source/Programs/UnrealBuildAccelerator/Common/Public/UbaHashMap.h
--- a/Engine/Source/Programs/UnrealBuildAccelerator/Common/Public/UbaHashMap.h
+++ b/Engine/Source/Programs/UnrealBuildAccelerator/Common/Public/UbaHashMap.h
@@ -6,6 +6,11 @@
 
 namespace uba
 {
+	// Hash map that contains three parts
+	//  1. Lookup table indexed with hashed key. Points to index in entry table
+	//  2. Entry table. Contains actual key and an index to next entry. Entry table order is stable during rehash
+	//  3. Value table. Matches entry table ordering.
+
 	template<typename Key, typename Value, bool AllowGrow = false>
 	struct HashMap
 	{
@@ -26,13 +31,13 @@ namespace uba
 			reserveCount = Max(reserveCount + 1, 4ull);
 			m_memory = &mem;
 			UBA_ASSERT(!AllowGrow);
-			InternalInit(reserveCount, GetLookupCount(reserveCount));
+			InternalInit(reserveCount, RoundUpPow2(reserveCount));
 		}
 
 		void Init(u64 reserveCount = 4)
 		{
 			reserveCount = Max(reserveCount + 1, 4ull);
-			u64 lookupCount = GetLookupCount(reserveCount);
+			u64 lookupCount = RoundUpPow2(reserveCount);
 
 			u64 lookupSize = lookupCount*sizeof(u32);
 			u64 entriesSize = sizeof(Entry)*(reserveCount+1);
@@ -208,19 +213,6 @@ private:
 			m_values = (Value*)m_memory->AllocateNoLock(m_reserveCount*sizeof(Value), alignof(Value), TC(""));
 		}
 
-		u32 GetLookupCount(u64 reserveCount)
-		{
-			u32 v = u32(reserveCount);
-			v--;
-			v |= v >> 1;
-			v |= v >> 2;
-			v |= v >> 4;
-			v |= v >> 8;
-			v |= v >> 16;
-			v++;
-			return v;
-		}
-
 		u64 GetNextReserveCount(u64 reserveCount)
 		{
 			if (reserveCount < 50'000)
@@ -240,6 +232,10 @@ private:
 		bool m_ownsMemory = false;
 	};
 
+	// Hash map that contains two parts
+	//  1. Lookup table indexed with hashed key. Points to index in entry table
+	//  2. Entry table. Contains actual key and value and an index to next entry. Entry table order is stable during rehash
+
 	template<typename Key, typename Value>
 	struct HashMap2
 	{
@@ -252,14 +248,7 @@ private:
 
 		void Init(MemoryBlock& memory, u64 maxSize)
 		{
-			u32 v = u32(maxSize);
-			v--;
-			v |= v >> 1;
-			v |= v >> 2;
-			v |= v >> 4;
-			v |= v >> 8;
-			v |= v >> 16;
-			v++;
+			u32 v = RoundUpPow2(maxSize);
 			u32 lookupSize = v;// << 1u;
 			mask = lookupSize - 1;
 			nextAvailableEntry = 1;
diff --git a/Engine/Source/Programs/UnrealBuildAccelerator/Common/Public/UbaSessionServer.h b/Engine/Source/Programs/UnrealBuildAccelerator/Common/Public/UbaSessionServer.h
--- a/Engine/Source/Programs/UnrealBuildAccelerator/Common/Public/UbaSessionServer.h
+++ b/Engine/Source/Programs/UnrealBuildAccelerator/Common/Public/UbaSessionServer.h
@@ -2,6 +2,7 @@
 
 #pragma once
 
+#include "UbaHashMap.h"
 #include "UbaNetwork.h"
 #include "UbaSession.h"
 #include "UbaSessionServerCreateInfo.h"
diff --git a/Engine/Source/Programs/UnrealBuildAccelerator/Core/Public/UbaMemory.h b/Engine/Source/Programs/UnrealBuildAccelerator/Core/Public/UbaMemory.h
--- a/Engine/Source/Programs/UnrealBuildAccelerator/Core/Public/UbaMemory.h
+++ b/Engine/Source/Programs/UnrealBuildAccelerator/Core/Public/UbaMemory.h
@@ -277,6 +277,19 @@ namespace uba
 		return v;
 	}
 
+	inline u32 RoundUpPow2(u64 value)
+	{
+		u32 v = u32(value);
+		v--;
+		v |= v >> 1;
+		v |= v >> 2;
+		v |= v >> 4;
+		v |= v >> 8;
+		v |= v >> 16;
+		v++;
+		return v;
+	}
+
 	bool SupportsHugePages();
 	u64 GetHugePageCount();
 }
diff --git a/Engine/Source/Programs/UnrealBuildAccelerator/Visualizer/Private/UbaVisualizer.cpp b/Engine/Source/Programs/UnrealBuildAccelerator/Visualizer/Private/UbaVisualizer.cpp
--- a/Engine/Source/Programs/UnrealBuildAccelerator/Visualizer/Private/UbaVisualizer.cpp
+++ b/Engine/Source/Programs/UnrealBuildAccelerator/Visualizer/Private/UbaVisualizer.cpp
@@ -1940,48 +1940,24 @@ namespace uba
 		}
 		else if (m_sessionSelectedIndex != ~0u)
 		{
-			int width = 290;
+			SetActiveFont(m_popupFont);
+			DrawTextLogger logger(m_hwnd, hdc, m_popupFont.height, m_tooltipBackgroundBrush);
+			logger.AddTextOffset(-10); // Remove spaces in the front
+			logger.AddWidth(3);
+			logger.AddSpace(2);
 
-			Vector<TString> summary = m_traceView.sessions[m_sessionSelectedIndex].summary;
+			Vector<TString>& summary = m_traceView.sessions[m_sessionSelectedIndex].summary;
 			if (summary.empty())
 			{
 				if (m_traceView.finished)
-					summary.push_back(L" Session summary not available on this trace version");
+					logger.Info(L" Session summary not available on this trace version");
 				else
-					summary.push_back(L" Session summary not available until session is done");
-				summary.push_back(L"");
-				width = 380;
-			}
-
-			int height = int(summary.size())*m_popupFont.height;
-
-			POINT p;
-			GetCursorPos(&p);
-			ScreenToClient(m_hwnd, &p);
-			RECT r;
-			r.left = p.x;
-			r.top = p.y;
-			r.right = r.left + width;
-			r.bottom = r.top + height;
-
-			if (r.right > clientRect.right)
-				OffsetRect(&r, -width, 0);
-			if (r.bottom > clientRect.bottom)
-			{
-				OffsetRect(&r, 0, -height);
-				if (r.top < 0)
-					OffsetRect(&r, 0, -r.top);
+					logger.Info(L" Session summary not available until session is done");
 			}
-			FillRect(hdc, &r, m_tooltipBackgroundBrush);
-
-			r.top += 5;
-			SetActiveFont(m_popupFont);
-
 			for (auto& line : summary)
-			{
-				DrawTextW(hdc, line.c_str(), int(line.size()), &r, DT_SINGLELINE);
-				r.top += m_popupFont.height;
-			}
+				logger.Log(LogEntryType_Info, line.c_str(), u32(line.size()));
+			logger.AddSpace(3);
+			logger.DrawAtCursor();
 		}
 		else if (m_statsSelected)
 		{
