[UBA]
* Bumped dir table reserve capacity to 128mb (this does not affect committed memory until it is used)

diff --git a/Engine/Source/Programs/UnrealBuildAccelerator/Core/Public/UbaProtocol.h b/Engine/Source/Programs/UnrealBuildAccelerator/Core/Public/UbaProtocol.h
--- a/Engine/Source/Programs/UnrealBuildAccelerator/Core/Public/UbaProtocol.h
+++ b/Engine/Source/Programs/UnrealBuildAccelerator/Core/Public/UbaProtocol.h
@@ -53,7 +53,7 @@ namespace uba
 	inline constexpr u32 CommunicationMemSize = IsWindows ? 64*1024 : 64*1024*2; // Macos expands some commandlines to be crazy long
 
 	inline constexpr u32 FileMappingTableMemSize = 16 * 1024 * 1024;
-	inline constexpr u32 DirTableMemSize = 64 * 1024 * 1024;
+	inline constexpr u32 DirTableMemSize = 128 * 1024 * 1024;
 }
 
 // Currently only used for detoured process
