[UBA]
* Added SCOPED_CRITICAL_SECTION to be able to measure contention

diff --git a/Engine/Source/Programs/UnrealBuildAccelerator/Core/Public/UbaSynchronization.h b/Engine/Source/Programs/UnrealBuildAccelerator/Core/Public/UbaSynchronization.h
--- a/Engine/Source/Programs/UnrealBuildAccelerator/Core/Public/UbaSynchronization.h
+++ b/Engine/Source/Programs/UnrealBuildAccelerator/Core/Public/UbaSynchronization.h
@@ -91,19 +91,21 @@ namespace uba
 	#if !UBA_TRACK_CONTENTION
 	#define SCOPED_READ_LOCK(readerWriterLock, name) ScopedReadLock name(readerWriterLock);
 	#define SCOPED_WRITE_LOCK(readerWriterLock, name) ScopedWriteLock name(readerWriterLock);
+	#define SCOPED_CRITICAL_SECTION(cs, name) ScopedCriticalSection name(cs);
 	#else
 	u64 GetTime();
 	struct ContentionTracker { void Add(u64 t) { time += t; ++count; }; Atomic<u64> time; Atomic<u64> count; const char* file; u64 line; };
 	ContentionTracker& GetContentionTracker(const char* file, u64 line);
 
-	#define SCOPED_LOCK(readerWriterLock, lockType, name) \
+	#define SCOPED_LOCK_INTERNAL(lock, lockType, name) \
 		u64 STRING_JOIN(contentionStart, __LINE__) = GetTime(); \
-		lockType name(readerWriterLock); \
+		lockType name(lock); \
 		static ContentionTracker& STRING_JOIN(tracker, __LINE__) = GetContentionTracker(__FILE__, __LINE__); \
 		STRING_JOIN(tracker, __LINE__).Add(GetTime() - STRING_JOIN(contentionStart, __LINE__));
 
-	#define SCOPED_READ_LOCK(readerWriterLock, name) SCOPED_LOCK(readerWriterLock, ScopedReadLock, name)
-	#define SCOPED_WRITE_LOCK(readerWriterLock, name) SCOPED_LOCK(readerWriterLock, ScopedWriteLock, name)
+	#define SCOPED_READ_LOCK(readerWriterLock, name) SCOPED_LOCK_INTERNAL(readerWriterLock, ScopedReadLock, name)
+	#define SCOPED_WRITE_LOCK(readerWriterLock, name) SCOPED_LOCK_INTERNAL(readerWriterLock, ScopedWriteLock, name)
+	#define SCOPED_CRITICAL_SECTION(cs, name) SCOPED_LOCK_INTERNAL(cs, ScopedCriticalSection,  name)
 	#endif
 
 
diff --git a/Engine/Source/Programs/UnrealBuildAccelerator/ObjTool/Private/UbaObjTool.cpp b/Engine/Source/Programs/UnrealBuildAccelerator/ObjTool/Private/UbaObjTool.cpp
--- a/Engine/Source/Programs/UnrealBuildAccelerator/ObjTool/Private/UbaObjTool.cpp
+++ b/Engine/Source/Programs/UnrealBuildAccelerator/ObjTool/Private/UbaObjTool.cpp
@@ -273,7 +273,7 @@ namespace uba
 						success = false;
 						return;
 					}
-					ScopedCriticalSection _(cs);
+					SCOPED_CRITICAL_SECTION(cs, _);
 					allExternalImports.insert(symbolFile.imports.begin(), symbolFile.imports.end());
 				});
 			if (!success)
@@ -294,7 +294,7 @@ namespace uba
 						success = false;
 						return;
 					}
-					ScopedCriticalSection _(cs);
+					SCOPED_CRITICAL_SECTION(cs, _);
 					allInternalImports.insert(symbolFile.imports.begin(), symbolFile.imports.end());
 					allExports.insert(symbolFile.exports.begin(), symbolFile.exports.end());
 				});
