[UBA]
* Removed lock around socket send for windows. Want to see if socket send times go up on the farm

diff --git a/Engine/Source/Programs/UnrealBuildAccelerator/Common/Private/UbaNetworkBackendTcp.cpp b/Engine/Source/Programs/UnrealBuildAccelerator/Common/Private/UbaNetworkBackendTcp.cpp
--- a/Engine/Source/Programs/UnrealBuildAccelerator/Common/Private/UbaNetworkBackendTcp.cpp
+++ b/Engine/Source/Programs/UnrealBuildAccelerator/Common/Private/UbaNetworkBackendTcp.cpp
@@ -41,7 +41,7 @@
 #define WSAPoll poll
 #endif
 
-#define UBA_LOCK_AROUND_SEND 1  // !PLATFORM_WINDOWS
+#define UBA_LOCK_AROUND_SEND !PLATFORM_WINDOWS
 
 namespace uba
 {
@@ -155,7 +155,7 @@ namespace uba
 		SCOPED_WRITE_LOCK(m_connectionsLock, lock);
 		for (auto& conn : m_connections)
 		{
-			ScopedCriticalSection lock2(conn.shutdownLock);
+			SCOPED_CRITICAL_SECTION(conn.shutdownLock, lock2);
 			if (conn.socket == INVALID_SOCKET)
 				continue;
 			SOCKET s = conn.socket;
@@ -176,7 +176,7 @@ namespace uba
 	void NetworkBackendTcp::Shutdown(void* connection)
 	{
 		auto& conn = *(Connection*)connection;
-		ScopedCriticalSection lock(conn.shutdownLock);
+		SCOPED_CRITICAL_SECTION(conn.shutdownLock, lock);
 		if (conn.socket == INVALID_SOCKET)
 			return;
 		ShutdownSocket(conn.logger, conn.socket, TC("Shutdown"));
@@ -222,7 +222,7 @@ namespace uba
 		UBA_ASSERT(headerSize <= MaxHeaderSize);
 		auto& conn = *(Connection*)connection;
 
-		ScopedCriticalSection lock(conn.shutdownLock);
+		SCOPED_CRITICAL_SECTION(conn.shutdownLock, lock);
 		UBA_ASSERTF(conn.disconnectCallback, TC("SetDisconnectCallback must be called before SetRecvCallbacks"));
 		conn.recvContext = context;
 		conn.headerSize = headerSize;
@@ -241,7 +241,7 @@ namespace uba
 	void NetworkBackendTcp::SetDisconnectCallback(void* connection, void* context, DisconnectCallback* callback)
 	{
 		auto& conn = *(Connection*)connection;
-		ScopedCriticalSection lock(conn.shutdownLock);
+		SCOPED_CRITICAL_SECTION(conn.shutdownLock, lock);
 		conn.disconnectCallback = callback;
 		conn.disconnectContext = context;
 	}
@@ -249,7 +249,7 @@ namespace uba
 	void NetworkBackendTcp::SetAllowLessThanBodySize(void* connection, bool allow)
 	{
 		auto& conn = *(Connection*)connection;
-		ScopedCriticalSection lock(conn.shutdownLock);
+		SCOPED_CRITICAL_SECTION(conn.shutdownLock, lock);
 		conn.allowLess = allow;
 	}
 
@@ -509,7 +509,7 @@ namespace uba
 			logger.Warning(TC("Tcp connection timed out waiting for recv thread to be ready"));
 		}
 
-		ScopedCriticalSection lock2(connection.shutdownLock);
+		SCOPED_CRITICAL_SECTION(connection.shutdownLock, lock2);
 		SOCKET s = connection.socket;
 
 		{
@@ -817,8 +817,9 @@ namespace uba
 
 			left -= sent;
 
-			#if PLATFORM_WINDOWS
-			UBA_ASSERTF(left == 0, L"Failed to send all data in one call. Wanted to send %llu, sent %i", bufferLen, sent);
+			#if !UBA_LOCK_AROUND_SEND
+			if (left != 0)
+				return logger.Error(TC("Failed to send all data in one call. Wanted to send %llu, sent %i"), bufferLen, sent);
 			#endif
 		}
 		return true;
