[UBA]
* Added test to be able to see how much memory is taken by detoured processes (disabled by default)... it seems like mimalloc commit 32mb at startup

diff --git a/Engine/Source/Programs/UnrealBuildAccelerator/Test/Private/UbaTestAll.h b/Engine/Source/Programs/UnrealBuildAccelerator/Test/Private/UbaTestAll.h
--- a/Engine/Source/Programs/UnrealBuildAccelerator/Test/Private/UbaTestAll.h
+++ b/Engine/Source/Programs/UnrealBuildAccelerator/Test/Private/UbaTestAll.h
@@ -20,6 +20,14 @@
 namespace uba
 {
 
+#if 0
+	#define UBA_TEST_EXTRA \
+		UBA_TEST(TestSharedReservedMemory)
+#else
+	#define UBA_TEST_EXTRA
+#endif
+
+
 #define UBA_ALLPLATFORM_TESTS \
 		UBA_TEST(TestTime) \
 		UBA_TEST(TestEvents) \
@@ -47,7 +55,7 @@ namespace uba
 		UBA_TEST(TestRootPaths) \
 		UBA_TEST(TestRegisterChanges) \
 		UBA_TEST(TestRegisterChangesRemote) \
-
+		UBA_TEST_EXTRA \
 
 #define UBA_POSIX_TESTS \
 		UBA_TEST(TestDetouredClang) \
diff --git a/Engine/Source/Programs/UnrealBuildAccelerator/Test/Private/UbaTestSession.h b/Engine/Source/Programs/UnrealBuildAccelerator/Test/Private/UbaTestSession.h
--- a/Engine/Source/Programs/UnrealBuildAccelerator/Test/Private/UbaTestSession.h
+++ b/Engine/Source/Programs/UnrealBuildAccelerator/Test/Private/UbaTestSession.h
@@ -543,4 +543,34 @@ namespace uba
 	{
 		return RunRemote(logger, testRootDir, CheckAttributes);
 	}
+
+	bool TestSharedReservedMemory(LoggerWithWriter& logger, const StringBufferBase& testRootDir)
+	{
+		return RunLocal(logger, testRootDir, [](LoggerWithWriter& logger, SessionServer& session, const tchar* workingDir, const RunProcessFunction& runProcess)
+			{
+				StringBuffer<MaxPath> testApp;
+				GetTestAppPath(logger, testApp);
+
+				ProcessStartInfo processInfo;
+				processInfo.application = testApp.data;
+				processInfo.workingDir = workingDir;
+				processInfo.arguments = L"-sleep=100000";
+				Vector<ProcessHandle> processes;
+
+				for (u32 i=0; i!=128; ++i)
+					processes.push_back(runProcess(processInfo));
+
+				for (auto& process : processes)
+				{
+					if (!process.WaitForExit(100000))
+						return logger.Error(TC("UbaTestApp did not exit in 10 seconds"));
+					u32 exitCode = process.GetExitCode();
+					if (exitCode != 0)
+						return false;
+				}
+
+				return true;
+			});
+	}
+
 }
