[UBA]
* Removed lock in CacheClient that should not be needed (code is just accessing mapped memory, no table accesses)

diff --git a/Engine/Source/Programs/UnrealBuildAccelerator/Common/Private/UbaCacheClient.cpp b/Engine/Source/Programs/UnrealBuildAccelerator/Common/Private/UbaCacheClient.cpp
--- a/Engine/Source/Programs/UnrealBuildAccelerator/Common/Private/UbaCacheClient.cpp
+++ b/Engine/Source/Programs/UnrealBuildAccelerator/Common/Private/UbaCacheClient.cpp
@@ -1268,8 +1268,6 @@ namespace uba
 
 	bool CacheClient::GetLocalPathAndCasKey(Bucket& bucket, const RootPaths& rootPaths, StringBufferBase& outPath, CasKey& outKey, CompactCasKeyTable& casKeyTable, CompactPathTable& pathTable, u32 offset)
 	{
-		SCOPED_READ_LOCK(bucket.casKeyTableNetworkLock, lock); // TODO: Is this needed?
-
 		StringBuffer<MaxPath> normalizedPath;
 		casKeyTable.GetPathAndKey(normalizedPath, outKey, pathTable, offset);
 		UBA_ASSERT(normalizedPath.count);
