[UBA]
* Changed hash map type for name-to-hash lookup table to use uba hash map. This reduced contention in some code by quite a lot

diff --git a/Engine/Source/Programs/UnrealBuildAccelerator/Common/Private/UbaSessionServer.cpp b/Engine/Source/Programs/UnrealBuildAccelerator/Common/Private/UbaSessionServer.cpp
--- a/Engine/Source/Programs/UnrealBuildAccelerator/Common/Private/UbaSessionServer.cpp
+++ b/Engine/Source/Programs/UnrealBuildAccelerator/Common/Private/UbaSessionServer.cpp
@@ -231,7 +231,7 @@ namespace uba
 		m_server.UnregisterOnClientDisconnected(ServiceId);
 		m_server.UnregisterService(ServiceId);
 
-		ScopedCriticalSection lock(m_remoteProcessAndSessionLock);
+		SCOPED_CRITICAL_SECTION(m_remoteProcessAndSessionLock, lock);
 		for (ProcessHandle& p : m_queuedRemoteProcesses)
 		{
 			((RemoteProcess*)p.m_process)->m_server = nullptr;
@@ -329,7 +329,7 @@ namespace uba
 
 				// Update name to hash table
 				SCOPED_WRITE_LOCK(m_nameToHashLookupLock, lock);
-				CasKey& lookupCasKey = m_nameToHashLookup[fileNameKey];
+				CasKey& lookupCasKey = m_nameToHashLookup.Insert(fileNameKey);
 				if (lookupCasKey != casKey)
 				{
 					lookupCasKey = casKey;
@@ -344,7 +344,7 @@ namespace uba
 
 		ProcessHandle h(remoteProcess); // Keep ref count up even if process is removed by callbacks etc.
 
-		ScopedCriticalSection lock(m_remoteProcessAndSessionLock);
+		SCOPED_CRITICAL_SECTION(m_remoteProcessAndSessionLock, lock);
 		m_queuedRemoteProcesses.push_back(remoteProcess);
 
 		SCOPED_READ_LOCK(m_remoteProcessReturnedEventLock, lock2);
@@ -366,7 +366,7 @@ namespace uba
 
 	void SessionServer::DisableRemoteExecution()
 	{
-		ScopedCriticalSection lock(m_remoteProcessAndSessionLock);
+		SCOPED_CRITICAL_SECTION(m_remoteProcessAndSessionLock, lock);
 		if (m_remoteExecutionEnabled)
 			m_logger.Info(TC("Disable remote execution (remote sessions will finish current processes)"));
 		m_remoteExecutionEnabled = false;
@@ -457,7 +457,7 @@ namespace uba
 	{
 		while (true)
 		{
-			ScopedCriticalSection lock(m_remoteProcessAndSessionLock);
+			SCOPED_CRITICAL_SECTION(m_remoteProcessAndSessionLock, lock);
 			if (m_activeRemoteProcesses.empty() && m_queuedRemoteProcesses.empty())
 				break;
 			lock.Leave();
@@ -525,7 +525,7 @@ namespace uba
 
 	void SessionServer::RegisterNetworkTrafficProvider(const NetworkTrafficProvider& provider)
 	{
-		ScopedCriticalSection lock(m_remoteProcessAndSessionLock);
+		SCOPED_CRITICAL_SECTION(m_remoteProcessAndSessionLock, lock);
 		m_provider = provider;
 	}
 
@@ -537,7 +537,7 @@ namespace uba
 	void SessionServer::OnDisconnected(const Guid& clientUid, u32 clientId)
 	{
 		u32 returnCount = 0;
-		ScopedCriticalSection queueLock(m_remoteProcessAndSessionLock);
+		SCOPED_CRITICAL_SECTION(m_remoteProcessAndSessionLock, queueLock);
 		for (auto it=m_activeRemoteProcesses.begin(); it!=m_activeRemoteProcesses.end();)
 		{
 			RemoteProcess* remoteProcess = (RemoteProcess*)it->m_process;
@@ -735,7 +735,7 @@ namespace uba
 		// I have no explanation for this. On linux we get a shutdown crash when running through UBT if session is allocated with normal new
 		// For now we will work around it by using aligned_alloc which seems to be working on all platforms
 		auto& session = *new (aligned_alloc(alignof(ClientSession), sizeof(ClientSession))) ClientSession();
-		ScopedCriticalSection lock(m_remoteProcessAndSessionLock);
+		SCOPED_CRITICAL_SECTION(m_remoteProcessAndSessionLock, lock);
 		m_clientSessions.push_back(&session);
 		u32 sessionId = u32(m_clientSessions.size());
 		session.name = name.data;
@@ -928,7 +928,7 @@ namespace uba
 		{
 			SCOPED_WRITE_LOCK(m_nameToHashLookupLock, lock);
 			serverTime = GetTime();
-			CasKey& lookupCasKey = m_nameToHashLookup[fileNameKey];
+			CasKey& lookupCasKey = m_nameToHashLookup.Insert(fileNameKey);
 			if (lookupCasKey != casKey)
 			{
 				lookupCasKey = casKey;
@@ -1129,7 +1129,7 @@ namespace uba
 	{
 		u32 sessionId = reader.ReadU32();
 		u32 sessionIndex = sessionId - 1;
-		ScopedCriticalSection lock(m_remoteProcessAndSessionLock);
+		SCOPED_CRITICAL_SECTION(m_remoteProcessAndSessionLock, lock);
 		if (sessionIndex >= m_clientSessions.size())
 			return m_logger.Error(TC("Got ListDirectory message from connection using bad sessionid (%u/%llu)"), sessionIndex, m_clientSessions.size());
 		ClientSession& session = *m_clientSessions[sessionIndex];
@@ -1149,7 +1149,7 @@ namespace uba
 	{
 		u32 sessionId = reader.ReadU32();
 		u32 sessionIndex = sessionId - 1;
-		ScopedCriticalSection lock(m_remoteProcessAndSessionLock);
+		SCOPED_CRITICAL_SECTION(m_remoteProcessAndSessionLock, lock);
 		if (sessionIndex >= m_clientSessions.size())
 			return m_logger.Error(TC("Got GetDirectories message from connection using bad sessionid (%u/%llu)"), sessionIndex, m_clientSessions.size());
 		ClientSession& session = *m_clientSessions[sessionIndex];
@@ -1180,7 +1180,7 @@ namespace uba
 		u32 sessionId = reader.ReadU32();
 		u32 sessionIndex = sessionId - 1;
 
-		ScopedCriticalSection sessionsLock(m_remoteProcessAndSessionLock);
+		SCOPED_CRITICAL_SECTION(m_remoteProcessAndSessionLock, sessionsLock);
 		if (sessionIndex >= m_clientSessions.size())
 			return m_logger.Error(TC("Got ProcessAvailable message from connection using bad sessionid (%u/%llu)"), sessionIndex, m_clientSessions.size());
 		ClientSession& session = *m_clientSessions[sessionIndex];
@@ -1348,7 +1348,7 @@ namespace uba
 		}
 		auto& process = *(RemoteProcess*)h.m_process;
 
-		ScopedCriticalSection cs2(m_remoteProcessAndSessionLock);
+		SCOPED_CRITICAL_SECTION(m_remoteProcessAndSessionLock, cs2);
 		if (!m_activeRemoteProcesses.erase(&process))
 		{
 			cs2.Leave();
@@ -1417,7 +1417,7 @@ namespace uba
 			return true;
 		}
 
-		ScopedCriticalSection cs2(m_remoteProcessAndSessionLock);
+		SCOPED_CRITICAL_SECTION(m_remoteProcessAndSessionLock, cs2);
 		if (!m_activeRemoteProcesses.erase(process))
 		{
 			cs2.Leave();
@@ -1488,7 +1488,7 @@ namespace uba
 
 		u64 pingTime = GetTime();
 		u32 sessionIndex = sessionId - 1;
-		ScopedCriticalSection lock(m_remoteProcessAndSessionLock);
+		SCOPED_CRITICAL_SECTION(m_remoteProcessAndSessionLock, lock);
 		if (sessionIndex >= m_clientSessions.size())
 			return m_logger.Error(TC("Got Pingmessage from connection using bad sessionid (%u/%llu)"), sessionIndex, m_clientSessions.size());
 		auto& session = *m_clientSessions[sessionIndex];
@@ -1508,7 +1508,7 @@ namespace uba
 
 		u32 sessionIndex = sessionId - 1;
 		{
-			ScopedCriticalSection lock(m_remoteProcessAndSessionLock);
+			SCOPED_CRITICAL_SECTION(m_remoteProcessAndSessionLock, lock);
 			if (sessionIndex < m_clientSessions.size())
 				m_clientSessions[sessionIndex]->hasNotification = true;
 		}
@@ -1592,7 +1592,7 @@ namespace uba
 		{
 			u32 totalUsed = 0;
 			u32 totalSlots = 0;
-			ScopedCriticalSection queueLock(m_remoteProcessAndSessionLock);
+			SCOPED_CRITICAL_SECTION(m_remoteProcessAndSessionLock, queueLock);
 			u64 time = GetTime();
 			for (auto& s : m_clientSessions)
 			{
@@ -1616,7 +1616,7 @@ namespace uba
 				abortWithProxy = true;
 				abortUseProxy = true;
 			}
-			ScopedCriticalSection queueLock(m_remoteProcessAndSessionLock);
+			SCOPED_CRITICAL_SECTION(m_remoteProcessAndSessionLock, queueLock);
 			u32 abortCount = 0;
 			for (auto& s : m_clientSessions)
 			{
@@ -1744,7 +1744,7 @@ namespace uba
 				break;
 
 			#if 0
-			ScopedCriticalSection queueLock(m_remoteProcessAndSessionLock);
+			SCOPED_CRITICAL_SECTION(m_remoteProcessAndSessionLock, queueLock);
 			m_logger.Info(TC("RemoteQueue: %llu Active: %llu ConnectionCount: %u"), m_queuedRemoteProcesses.size(), m_activeRemoteProcesses.size(), m_connectionCount);
 
 			if (!m_connectionCount && !m_activeRemoteProcesses.empty())
@@ -1866,7 +1866,7 @@ namespace uba
 
 		while (true)
 		{
-			ScopedCriticalSection queueLock(m_remoteProcessAndSessionLock);
+			SCOPED_CRITICAL_SECTION(m_remoteProcessAndSessionLock, queueLock);
 
 			if (!session.connected) // This should not be possible
 			{
@@ -1908,7 +1908,7 @@ namespace uba
 	{
 		ProcessHandle h(process);
 
-		ScopedCriticalSection queueLock(m_remoteProcessAndSessionLock);
+		SCOPED_CRITICAL_SECTION(m_remoteProcessAndSessionLock, queueLock);
 		if (process->m_clientId == ~0u)
 		{
 			for (auto it=m_queuedRemoteProcesses.begin(); it!=m_queuedRemoteProcesses.end(); ++it)
@@ -2041,9 +2041,9 @@ namespace uba
 		
 		Storage::CachedFileInfo cachedInfo;
 		if (!m_storage.VerifyAndGetCachedFileInfo(cachedInfo, fileNameKey, lastWritten, size))
-			if (m_nameToHashLookup.find(fileNameKey) == m_nameToHashLookup.end())
+			if (!m_nameToHashLookup.Find(fileNameKey))
 				return;
-		CasKey& lookupCasKey = m_nameToHashLookup[fileNameKey];
+		CasKey& lookupCasKey = m_nameToHashLookup.Insert(fileNameKey);
 		if (lookupCasKey == cachedInfo.casKey)
 			return;
 
@@ -2062,8 +2062,8 @@ namespace uba
 	{
 		Session::PrintSessionStats(logger);
 
-		if (m_nameToHashLookup.size())
-			logger.Info(TC("  NameToHashLookup    %7u %9s"), u32(m_nameToHashLookup.size()), BytesToText(m_nameToHashTableMem.writtenSize).str);
+		if (m_nameToHashLookup.Size())
+			logger.Info(TC("  NameToHashLookup    %7u %9s"), u32(m_nameToHashLookup.Size()), BytesToText(m_nameToHashTableMem.writtenSize).str);
 		logger.Info(TC("  Remote processes finished    %8u"), m_finishedRemoteProcessCount);
 		logger.Info(TC("  Remote processes returned    %8u"), m_returnedRemoteProcessCount);
 		logger.Info(TC(""));
@@ -2076,7 +2076,7 @@ namespace uba
 		u64 serverSend = m_server.GetTotalSentBytes();
 		u64 serverRecv = m_server.GetTotalRecvBytes();
 
-		ScopedCriticalSection lock(m_remoteProcessAndSessionLock);
+		SCOPED_CRITICAL_SECTION(m_remoteProcessAndSessionLock, lock);
 		for (auto sptr : m_clientSessions)
 		{
 			auto& s = *sptr;
@@ -2141,6 +2141,7 @@ namespace uba
 			return true;
 
 		SCOPED_WRITE_LOCK(m_nameToHashLookupLock, lock);
+		m_nameToHashLookup.Init(256*1024);
 		m_nameToHashTableMem.Init(NameToHashMemSize);
 		m_nameToHashInitialized = true;
 		lock.Leave();
@@ -2176,7 +2177,7 @@ namespace uba
 			}
 		}
 		SCOPED_WRITE_LOCK(m_nameToHashLookupLock, lock2);
-		u64 entryCount = m_nameToHashLookup.size();
+		u64 entryCount = m_nameToHashLookup.Size();
 		lock2.Leave();
 
 		m_logger.Debug(TC("Prepopulated NameToHash table with %u entries"), entryCount);
@@ -2258,9 +2259,9 @@ namespace uba
 				}
 				{
 					SCOPED_READ_LOCK(m_nameToHashLookupLock, hlock);
-					auto findIt = m_nameToHashLookup.find(fileNameKey);
-					if (findIt != m_nameToHashLookup.end())
-						logStr.Clear().Appendf(TC("File %s found in name-to-hash lookup. CasKey is %s"), path.data, CasKeyString(findIt->second).str);
+					auto valuePtr = m_nameToHashLookup.Find(fileNameKey);
+					if (valuePtr)
+						logStr.Clear().Appendf(TC("File %s found in name-to-hash lookup. CasKey is %s"), path.data, CasKeyString(*valuePtr).str);
 					else
 						logStr.Clear().Appendf(TC("File %s not found in name-to-hash lookup"), path.data);
 					LogLine(logStr);
diff --git a/Engine/Source/Programs/UnrealBuildAccelerator/Common/Public/UbaSessionServer.h b/Engine/Source/Programs/UnrealBuildAccelerator/Common/Public/UbaSessionServer.h
--- a/Engine/Source/Programs/UnrealBuildAccelerator/Common/Public/UbaSessionServer.h
+++ b/Engine/Source/Programs/UnrealBuildAccelerator/Common/Public/UbaSessionServer.h
@@ -171,7 +171,7 @@ namespace uba
 		ReaderWriterLock m_customCasKeysLock;
 		UnorderedMap<StringKey, CustomCasKey> m_customCasKeys;
 
-		UnorderedMap<StringKey, CasKey> m_nameToHashLookup;
+		HashMap<StringKey, CasKey, true> m_nameToHashLookup;
 		ReaderWriterLock m_nameToHashLookupLock;
 		Atomic<bool> m_nameToHashInitialized;
 
