[UBA]
* Fixed linux compile error in test module

diff --git a/Engine/Source/Programs/UnrealBuildAccelerator/Test/Private/UbaTestSession.h b/Engine/Source/Programs/UnrealBuildAccelerator/Test/Private/UbaTestSession.h
--- a/Engine/Source/Programs/UnrealBuildAccelerator/Test/Private/UbaTestSession.h
+++ b/Engine/Source/Programs/UnrealBuildAccelerator/Test/Private/UbaTestSession.h
@@ -554,7 +554,7 @@ namespace uba
 				ProcessStartInfo processInfo;
 				processInfo.application = testApp.data;
 				processInfo.workingDir = workingDir;
-				processInfo.arguments = L"-sleep=100000";
+				processInfo.arguments = TC("-sleep=100000");
 				Vector<ProcessHandle> processes;
 
 				for (u32 i=0; i!=128; ++i)
