[UBA]
* Moved functions that use large stack objects to an inner function to be able to early out without having to check stack

diff --git a/Engine/Source/Programs/UnrealBuildAccelerator/Common/Private/UbaCacheClient.cpp b/Engine/Source/Programs/UnrealBuildAccelerator/Common/Private/UbaCacheClient.cpp
--- a/Engine/Source/Programs/UnrealBuildAccelerator/Common/Private/UbaCacheClient.cpp
+++ b/Engine/Source/Programs/UnrealBuildAccelerator/Common/Private/UbaCacheClient.cpp
@@ -1060,20 +1060,13 @@ namespace uba
 
 	bool CacheClient::FetchCasTable(Bucket& bucket, CacheStats& stats, u32 requiredCasTableOffset)
 	{
-		auto hasEnoughData = [&bucket, requiredCasTableOffset](u32 tableSize)
-			{
-				u32 neededSize = requiredCasTableOffset + 4;
-				if (neededSize > tableSize)
-					return false;
-				BinaryReader r(bucket.serverCasKeyTable.GetMemory(), requiredCasTableOffset, tableSize);
-				u8 bytesNeeded = Get7BitEncodedCount(r.Read7BitEncoded());
-				neededSize = requiredCasTableOffset + bytesNeeded + sizeof(CasKey);
-				return neededSize <= tableSize;
-			};
-
-		if (hasEnoughData(bucket.availableCasKeyTableSize))
+		if (HasEnoughData(bucket, requiredCasTableOffset, bucket.availableCasKeyTableSize))
 			return true;
+		return FetchCasTable2(bucket, stats, requiredCasTableOffset);
+	}
 
+	bool CacheClient::FetchCasTable2(Bucket& bucket, CacheStats& stats, u32 requiredCasTableOffset)
+	{
 		TimerScope ts2(stats.fetchCasTable);
 
 		StackBinaryReader<SendMaxSize> reader;
@@ -1084,7 +1077,7 @@ namespace uba
 			while (true)
 			{
 				u32 tableSize = bucket.serverCasKeyTable.GetSize();
-				if (hasEnoughData(tableSize))
+				if (HasEnoughData(bucket, requiredCasTableOffset, tableSize))
 				{
 					if (!messageSent)
 						return true;
@@ -1128,6 +1121,17 @@ namespace uba
 		return true;
 	}
 
+	bool CacheClient::HasEnoughData(Bucket& bucket, u32 requiredCasTableOffset, u32 tableSize)
+	{
+		u32 neededSize = requiredCasTableOffset + 4;
+		if (neededSize > tableSize)
+			return false;
+		BinaryReader r(bucket.serverCasKeyTable.GetMemory(), requiredCasTableOffset, tableSize);
+		u8 bytesNeeded = Get7BitEncodedCount(r.Read7BitEncoded());
+		neededSize = requiredCasTableOffset + bytesNeeded + sizeof(CasKey);
+		return neededSize <= tableSize;
+	}
+
 	bool CacheClient::ReportUsedEntry(Vector<ProcessLogLine>& outLogLines, bool ownedLogLines, Bucket& bucket, const CasKey& cmdKey, u32 entryId)
 	{
 		StackBinaryWriter<128> writer;
diff --git a/Engine/Source/Programs/UnrealBuildAccelerator/Common/Private/UbaNetworkBackendTcp.cpp b/Engine/Source/Programs/UnrealBuildAccelerator/Common/Private/UbaNetworkBackendTcp.cpp
--- a/Engine/Source/Programs/UnrealBuildAccelerator/Common/Private/UbaNetworkBackendTcp.cpp
+++ b/Engine/Source/Programs/UnrealBuildAccelerator/Common/Private/UbaNetworkBackendTcp.cpp
@@ -41,6 +41,9 @@
 #define WSAPoll poll
 #endif
 
+// Hard to say if this lock should be enabled on windows or not.
+// Local tests with cache server and 100% cache hits make no-lock 35s vs lock 40s.
+// But by looking at the farm jobs it looks like more time is spent in send without the lock
 #define UBA_LOCK_AROUND_SEND !PLATFORM_WINDOWS
 
 namespace uba
diff --git a/Engine/Source/Programs/UnrealBuildAccelerator/Common/Private/UbaProcessStartInfoHolder.cpp b/Engine/Source/Programs/UnrealBuildAccelerator/Common/Private/UbaProcessStartInfoHolder.cpp
--- a/Engine/Source/Programs/UnrealBuildAccelerator/Common/Private/UbaProcessStartInfoHolder.cpp
+++ b/Engine/Source/Programs/UnrealBuildAccelerator/Common/Private/UbaProcessStartInfoHolder.cpp
@@ -12,7 +12,6 @@ namespace uba
 			return false;
 
 		const tchar* argsBegin = argumentsStr.c_str();
-		const tchar* argsEnd = argsBegin + argumentsStr.size();
 
 		// Check if application is repeated as first argument, in that case consume it
 		const tchar* firstArgBegin = argsBegin;
@@ -26,6 +25,17 @@ namespace uba
 			firstArgEnd = TStrchr(firstArgBegin, ' ');
 		if (!firstArgEnd)
 			return false;
+		return InternalExpand(firstArgBegin, firstArgEnd);
+	#endif
+	}
+
+	bool ProcessStartInfoHolder::InternalExpand(const tchar* firstArgBegin, const tchar* firstArgEnd)
+	{
+		// Separate function to hide away large stack object
+	#if PLATFORM_WINDOWS
+		const tchar* argsBegin = argumentsStr.c_str();
+		const tchar* argsEnd = argsBegin + argumentsStr.size();
+
 		StringBuffer<32*1024> commands;
 		commands.Append(firstArgBegin, firstArgEnd - firstArgBegin);
 		if (commands.Contains(application))
diff --git a/Engine/Source/Programs/UnrealBuildAccelerator/Common/Public/UbaCacheClient.h b/Engine/Source/Programs/UnrealBuildAccelerator/Common/Public/UbaCacheClient.h
--- a/Engine/Source/Programs/UnrealBuildAccelerator/Common/Public/UbaCacheClient.h
+++ b/Engine/Source/Programs/UnrealBuildAccelerator/Common/Public/UbaCacheClient.h
@@ -74,6 +74,8 @@ namespace uba
 		bool SendCasTable(Bucket& bucket, u32 requiredCasTableSize);
 		bool SendCacheEntry(Bucket& bucket, const RootPaths& rootPaths, const CasKey& cmdKey, const Map<u32, u32>& inputsStringToCasKey, const Map<u32, u32>& outputsStringToCasKey, const u8* logLines, u64 logLinesSize, u64& outBytesSent);
 		bool FetchCasTable(Bucket& bucket, CacheStats& stats, u32 requiredCasTableOffset);
+		bool FetchCasTable2(Bucket& bucket, CacheStats& stats, u32 requiredCasTableOffset);
+		bool HasEnoughData(Bucket& bucket, u32 requiredCasTableOffset, u32 tableSize);
 		bool ReportUsedEntry(Vector<ProcessLogLine>& outLogLines, bool ownedLogLines, Bucket& bucket, const CasKey& cmdKey, u32 entryId);
 		bool PopulateLogLines(Vector<ProcessLogLine>& outLogLines, const u8* mem, u64 memLen);
 
diff --git a/Engine/Source/Programs/UnrealBuildAccelerator/Common/Public/UbaProcessStartInfoHolder.h b/Engine/Source/Programs/UnrealBuildAccelerator/Common/Public/UbaProcessStartInfoHolder.h
--- a/Engine/Source/Programs/UnrealBuildAccelerator/Common/Public/UbaProcessStartInfoHolder.h
+++ b/Engine/Source/Programs/UnrealBuildAccelerator/Common/Public/UbaProcessStartInfoHolder.h
@@ -127,5 +127,7 @@ namespace uba
 		TString logFileStr;
 		TString breadcrumbsStr;
 		float weight = 1.0f;
+
+		bool InternalExpand(const tchar* firstArgBegin, const tchar* firstArgEnd);
 	};
 }
\ No newline at end of file
