[UBA]
* Fixed linux compile error

diff --git a/Engine/Source/Programs/UnrealBuildAccelerator/Common/Private/UbaProcessStartInfoHolder.cpp b/Engine/Source/Programs/UnrealBuildAccelerator/Common/Private/UbaProcessStartInfoHolder.cpp
--- a/Engine/Source/Programs/UnrealBuildAccelerator/Common/Private/UbaProcessStartInfoHolder.cpp
+++ b/Engine/Source/Programs/UnrealBuildAccelerator/Common/Private/UbaProcessStartInfoHolder.cpp
@@ -26,6 +26,8 @@ namespace uba
 		if (!firstArgEnd)
 			return false;
 		return InternalExpand(firstArgBegin, firstArgEnd);
+	#else
+		return false;
 	#endif
 	}
 
