[UBA]
* Changed container for file-name-to-hash back to std unordered_map after seeing two hangs on linux.. don't know if it is related since there were more changes but will start by testing this

diff --git a/Engine/Source/Programs/UnrealBuildAccelerator/Common/Private/UbaSessionServer.cpp b/Engine/Source/Programs/UnrealBuildAccelerator/Common/Private/UbaSessionServer.cpp
--- a/Engine/Source/Programs/UnrealBuildAccelerator/Common/Private/UbaSessionServer.cpp
+++ b/Engine/Source/Programs/UnrealBuildAccelerator/Common/Private/UbaSessionServer.cpp
@@ -329,7 +329,7 @@ namespace uba
 
 				// Update name to hash table
 				SCOPED_WRITE_LOCK(m_nameToHashLookupLock, lock);
-				CasKey& lookupCasKey = m_nameToHashLookup.Insert(fileNameKey);
+				CasKey& lookupCasKey = m_nameToHashLookup[fileNameKey];
 				if (lookupCasKey != casKey)
 				{
 					lookupCasKey = casKey;
@@ -928,7 +928,7 @@ namespace uba
 		{
 			SCOPED_WRITE_LOCK(m_nameToHashLookupLock, lock);
 			serverTime = GetTime();
-			CasKey& lookupCasKey = m_nameToHashLookup.Insert(fileNameKey);
+			CasKey& lookupCasKey = m_nameToHashLookup[fileNameKey];
 			if (lookupCasKey != casKey)
 			{
 				lookupCasKey = casKey;
@@ -2041,9 +2041,9 @@ namespace uba
 		
 		Storage::CachedFileInfo cachedInfo;
 		if (!m_storage.VerifyAndGetCachedFileInfo(cachedInfo, fileNameKey, lastWritten, size))
-			if (!m_nameToHashLookup.Find(fileNameKey))
+			if (m_nameToHashLookup.find(fileNameKey) == m_nameToHashLookup.end())
 				return;
-		CasKey& lookupCasKey = m_nameToHashLookup.Insert(fileNameKey);
+		CasKey& lookupCasKey = m_nameToHashLookup[fileNameKey];
 		if (lookupCasKey == cachedInfo.casKey)
 			return;
 
@@ -2062,8 +2062,8 @@ namespace uba
 	{
 		Session::PrintSessionStats(logger);
 
-		if (m_nameToHashLookup.Size())
-			logger.Info(TC("  NameToHashLookup    %7u %9s"), u32(m_nameToHashLookup.Size()), BytesToText(m_nameToHashTableMem.writtenSize).str);
+		if (m_nameToHashLookup.size())
+			logger.Info(TC("  NameToHashLookup    %7u %9s"), u32(m_nameToHashLookup.size()), BytesToText(m_nameToHashTableMem.writtenSize).str);
 		logger.Info(TC("  Remote processes finished    %8u"), m_finishedRemoteProcessCount);
 		logger.Info(TC("  Remote processes returned    %8u"), m_returnedRemoteProcessCount);
 		logger.Info(TC(""));
@@ -2141,7 +2141,6 @@ namespace uba
 			return true;
 
 		SCOPED_WRITE_LOCK(m_nameToHashLookupLock, lock);
-		m_nameToHashLookup.Init(256*1024);
 		m_nameToHashTableMem.Init(NameToHashMemSize);
 		m_nameToHashInitialized = true;
 		lock.Leave();
@@ -2177,7 +2176,7 @@ namespace uba
 			}
 		}
 		SCOPED_WRITE_LOCK(m_nameToHashLookupLock, lock2);
-		u64 entryCount = m_nameToHashLookup.Size();
+		u64 entryCount = m_nameToHashLookup.size();
 		lock2.Leave();
 
 		m_logger.Debug(TC("Prepopulated NameToHash table with %u entries"), entryCount);
@@ -2259,9 +2258,9 @@ namespace uba
 				}
 				{
 					SCOPED_READ_LOCK(m_nameToHashLookupLock, hlock);
-					auto valuePtr = m_nameToHashLookup.Find(fileNameKey);
-					if (valuePtr)
-						logStr.Clear().Appendf(TC("File %s found in name-to-hash lookup. CasKey is %s"), path.data, CasKeyString(*valuePtr).str);
+					auto findIt = m_nameToHashLookup.find(fileNameKey);
+					if (findIt != m_nameToHashLookup.end())
+						logStr.Clear().Appendf(TC("File %s found in name-to-hash lookup. CasKey is %s"), path.data, CasKeyString(findIt->second).str);
 					else
 						logStr.Clear().Appendf(TC("File %s not found in name-to-hash lookup"), path.data);
 					LogLine(logStr);
diff --git a/Engine/Source/Programs/UnrealBuildAccelerator/Common/Public/UbaSessionServer.h b/Engine/Source/Programs/UnrealBuildAccelerator/Common/Public/UbaSessionServer.h
--- a/Engine/Source/Programs/UnrealBuildAccelerator/Common/Public/UbaSessionServer.h
+++ b/Engine/Source/Programs/UnrealBuildAccelerator/Common/Public/UbaSessionServer.h
@@ -171,7 +171,7 @@ namespace uba
 		ReaderWriterLock m_customCasKeysLock;
 		UnorderedMap<StringKey, CustomCasKey> m_customCasKeys;
 
-		HashMap<StringKey, CasKey, true> m_nameToHashLookup;
+		UnorderedMap<StringKey, CasKey> m_nameToHashLookup;
 		ReaderWriterLock m_nameToHashLookupLock;
 		Atomic<bool> m_nameToHashInitialized;
 
