[UBA]
* CacheService - Set recv timeout to 2 hours for connections
* CacheServer - Added connection time to disconnect log entry
*

diff --git a/Engine/Source/Programs/UnrealBuildAccelerator/CacheService/Private/UbaCacheService.cpp b/Engine/Source/Programs/UnrealBuildAccelerator/CacheService/Private/UbaCacheService.cpp
--- a/Engine/Source/Programs/UnrealBuildAccelerator/CacheService/Private/UbaCacheService.cpp
+++ b/Engine/Source/Programs/UnrealBuildAccelerator/CacheService/Private/UbaCacheService.cpp
@@ -222,6 +222,7 @@ namespace uba
 		NetworkServerCreateInfo nsci(logWriter);
 		//nsci.workerCount = 4;
 		nsci.logConnections = false; // Let the cache server report instead
+		nsci.receiveTimeoutSeconds = 2*60*60; // Two hours timeout
 		bool ctorSuccess = true;
 		NetworkServer networkServer(ctorSuccess, nsci);
 		if (!ctorSuccess)
diff --git a/Engine/Source/Programs/UnrealBuildAccelerator/Common/Private/UbaCacheClient.cpp b/Engine/Source/Programs/UnrealBuildAccelerator/Common/Private/UbaCacheClient.cpp
--- a/Engine/Source/Programs/UnrealBuildAccelerator/Common/Private/UbaCacheClient.cpp
+++ b/Engine/Source/Programs/UnrealBuildAccelerator/Common/Private/UbaCacheClient.cpp
@@ -1065,7 +1065,7 @@ namespace uba
 		return FetchCasTable2(bucket, stats, requiredCasTableOffset);
 	}
 
-	bool CacheClient::FetchCasTable2(Bucket& bucket, CacheStats& stats, u32 requiredCasTableOffset)
+	UBA_NOINLINE bool CacheClient::FetchCasTable2(Bucket& bucket, CacheStats& stats, u32 requiredCasTableOffset)
 	{
 		TimerScope ts2(stats.fetchCasTable);
 
diff --git a/Engine/Source/Programs/UnrealBuildAccelerator/Common/Private/UbaCacheServer.cpp b/Engine/Source/Programs/UnrealBuildAccelerator/Common/Private/UbaCacheServer.cpp
--- a/Engine/Source/Programs/UnrealBuildAccelerator/Common/Private/UbaCacheServer.cpp
+++ b/Engine/Source/Programs/UnrealBuildAccelerator/Common/Private/UbaCacheServer.cpp
@@ -141,6 +141,7 @@ namespace uba
 		UnorderedSet<u64> fetchBuckets;
 		u64 fetchEntryCount = 0;
 		u64 fetchEntryHitCount = 0;
+		u64 connectTime = 0;
 	};
 
 	struct CacheServer::Bucket
@@ -1213,6 +1214,7 @@ namespace uba
 		u64 storeBucketsCount = 0;
 		u64 storeEntryCount = 0;
 		u64 fetchEntryHitCount = 0;
+		u64 connectTime = 0;
 		SCOPED_WRITE_LOCK(m_connectionsLock, lock);
 		auto it = m_connections.find(clientId);
 		if (it != m_connections.end())
@@ -1223,13 +1225,14 @@ namespace uba
 			fetchEntryHitCount = conn.fetchEntryHitCount;
 			storeBucketsCount = conn.storeBuckets.size();
 			storeEntryCount = conn.storeEntryCount;
+			connectTime = conn.connectTime;
 			m_connections.erase(it);
 		}
 		u64 activeCount = m_connections.size();
 		lock.Leave();
 
 		StringBuffer<> log;
-		log.Appendf(TC("Client %u disconnected (%llu active). "), clientId, activeCount);
+		log.Appendf(TC("Client %u disconnected after %s (%llu active). "), clientId, TimeToText(connectTime - GetTime()).str, activeCount);
 		if (fetchBucketsCount)
 			log.Appendf(TC("Fetched %llu (%llu hits) from %llu bucket(s)"), fetchEntryCount, fetchEntryHitCount, fetchBucketsCount);
 		if (fetchBucketsCount && storeBucketsCount)
@@ -1307,6 +1310,7 @@ namespace uba
 			auto insres = m_connections.try_emplace(connectionInfo.GetId());
 			auto& connection = insres.first->second;
 			connection.clientVersion = clientVersion;
+			connection.connectTime = GetTime();
 			m_peakConnectionCount = Max(m_peakConnectionCount, u32(m_connections.size()));
 
 			m_logger.Info(TC("Client %u connected (v%u) %s"), connectionInfo.GetId(), clientVersion, hint.c_str());
diff --git a/Engine/Source/Programs/UnrealBuildAccelerator/Common/Private/UbaProcessStartInfoHolder.cpp b/Engine/Source/Programs/UnrealBuildAccelerator/Common/Private/UbaProcessStartInfoHolder.cpp
--- a/Engine/Source/Programs/UnrealBuildAccelerator/Common/Private/UbaProcessStartInfoHolder.cpp
+++ b/Engine/Source/Programs/UnrealBuildAccelerator/Common/Private/UbaProcessStartInfoHolder.cpp
@@ -31,7 +31,7 @@ namespace uba
 	#endif
 	}
 
-	bool ProcessStartInfoHolder::InternalExpand(const tchar* firstArgBegin, const tchar* firstArgEnd)
+	UBA_NOINLINE bool ProcessStartInfoHolder::InternalExpand(const tchar* firstArgBegin, const tchar* firstArgEnd)
 	{
 		// Separate function to hide away large stack object
 	#if PLATFORM_WINDOWS
