[UBA]
* Added optimization where cache entries track if they have any inputs that has been seen in the cas storage.. if not, then we don't have to check inputs when running maintenance and caskeys are deleted.
* Fixed bug related to cache and compressed obj files (sometimes wrong caskey of compressed obj file was used when it should have used the caskey of the uncompressed data)
* Added so uncompressed caskey is outputted when running UbaCli -getcas and file is compressed

diff --git a/Engine/Source/Programs/UnrealBuildAccelerator/Cli/Private/UbaCli.cpp b/Engine/Source/Programs/UnrealBuildAccelerator/Cli/Private/UbaCli.cpp
--- a/Engine/Source/Programs/UnrealBuildAccelerator/Cli/Private/UbaCli.cpp
+++ b/Engine/Source/Programs/UnrealBuildAccelerator/Cli/Private/UbaCli.cpp
@@ -733,6 +733,13 @@ namespace uba
 				bool is64Bit = true;
 
 				CasKey key = CalculateCasKey(data, fileSize, false, nullptr, application.c_str());
+				CasKey uncompressedKey;
+				if (fileSize > sizeof(CompressedObjFileHeader))
+				{
+					auto& hdr = *(CompressedObjFileHeader*)data;
+					if (hdr.IsValid())
+						uncompressedKey = hdr.casKey;
+				}
 
 				if (data[0] != 'M' || data[1] != 'Z')
 					is64Bit = false;
@@ -745,6 +752,8 @@ namespace uba
 				logger.Info(TC("  Is64Bit: %s"), (is64Bit ? TC("true") : TC("false")));
 				logger.Info(TC("  Size: %llu"), fileSize);
 				logger.Info(TC("  CasKey: %s"), CasKeyString(key).str);
+				if (uncompressedKey != CasKeyZero)
+					logger.Info(TC("  CasKey (uncompressed): %s"), CasKeyString(uncompressedKey).str);
 				return true;
 			}
 		}
@@ -1176,6 +1185,7 @@ namespace uba
 			storageServer.PrintSummary(logger);
 			networkServer.PrintSummary(logger);
 			KernelStats::GetGlobal().Print(logger, true);
+			PrintContentionSummary(logger);
 		}
 		logger.EndScope();
 
diff --git a/Engine/Source/Programs/UnrealBuildAccelerator/Common/Private/UbaCacheClient.cpp b/Engine/Source/Programs/UnrealBuildAccelerator/Common/Private/UbaCacheClient.cpp
--- a/Engine/Source/Programs/UnrealBuildAccelerator/Common/Private/UbaCacheClient.cpp
+++ b/Engine/Source/Programs/UnrealBuildAccelerator/Common/Private/UbaCacheClient.cpp
@@ -274,7 +274,7 @@ namespace uba
 			{
 				bool shouldValidate = (m_validateCacheWritesInput && !isOutput) || (m_validateCacheWritesOutput && isOutput);
 				bool deferCreation = true;
-				bool fileIsCompressed = IsFileCompressed(info, path);
+				bool fileIsCompressed = IsFileCompressed(info, path, isOutput);
 
 				if (isOutput)
 				{
@@ -527,7 +527,7 @@ namespace uba
 					if (m_useDirectoryPreParsing)
 						PreparseDirectory(fileNameKey, path);
 
-					bool fileIsCompressed = IsFileCompressed(info, path);
+					bool fileIsCompressed = IsFileCompressed(info, path, false);
 					m_storage.StoreCasKey(localCasKey, fileNameKey, path.data, CasKeyZero, fileIsCompressed);
 					UBA_ASSERT(localCasKey == CasKeyZero || IsCompressed(localCasKey));
 				}
@@ -772,7 +772,7 @@ namespace uba
 					else
 					{
 						DowngradedLogger logger(m_logger.m_writer, TC("UbaCacheClientDownload"));
-						bool destinationIsCompressed = IsFileCompressed(info, path);
+						bool destinationIsCompressed = IsFileCompressed(info, path, true);
 						if (!fetcher.RetrieveFile(logger, m_client, casKey, path.data, destinationIsCompressed))
 							return logger.Error(TC("Failed to download cache output %s for %s"), path.data, info.GetDescription());
 					}
@@ -1283,14 +1283,17 @@ namespace uba
 		return true;
 	}
 
-	bool CacheClient::IsFileCompressed(const ProcessStartInfo& info, const StringView& filename)
+	bool CacheClient::IsFileCompressed(const ProcessStartInfo& info, const StringView& filename, bool isOutput)
 	{
 		if (!m_session.ShouldStoreObjFilesCompressed())
 			return false;
 		auto rules = info.rules;
 		if (!rules)
 			rules = m_session.GetRules(info);
-		return rules->StoreFileCompressed(filename);
+		if (isOutput)
+			return rules->StoreFileCompressed(filename);
+		else
+			return rules->ShouldDecompressFiles(filename);
 	}
 
 	void CacheClient::PreparseDirectory(const StringKey& fileNameKey, const StringBufferBase& filePath)
diff --git a/Engine/Source/Programs/UnrealBuildAccelerator/Common/Private/UbaCacheEntry.cpp b/Engine/Source/Programs/UnrealBuildAccelerator/Common/Private/UbaCacheEntry.cpp
--- a/Engine/Source/Programs/UnrealBuildAccelerator/Common/Private/UbaCacheEntry.cpp
+++ b/Engine/Source/Programs/UnrealBuildAccelerator/Common/Private/UbaCacheEntry.cpp
@@ -48,6 +48,8 @@ namespace uba
 		u64 size = GetSharedSize();
 		for (auto& entry : entries)
 			size += GetEntrySize(entry, clientVersion, toDisk);
+		if (toDisk)
+			size += sizeof(hasInputsThatAreOutputs);
 		return size;
 	}
 
@@ -151,6 +153,9 @@ namespace uba
 			++entryCount;
 		}
 
+		if (toDisk)
+			writer.WriteBool(hasInputsThatAreOutputs);
+
 		return true;
 	}
 
@@ -228,6 +233,9 @@ namespace uba
 
 		}
 
+		if (databaseVersion >= 7)
+			hasInputsThatAreOutputs = reader.ReadBool();
+
 		return true;
 	}
 
diff --git a/Engine/Source/Programs/UnrealBuildAccelerator/Common/Private/UbaCacheServer.cpp b/Engine/Source/Programs/UnrealBuildAccelerator/Common/Private/UbaCacheServer.cpp
--- a/Engine/Source/Programs/UnrealBuildAccelerator/Common/Private/UbaCacheServer.cpp
+++ b/Engine/Source/Programs/UnrealBuildAccelerator/Common/Private/UbaCacheServer.cpp
@@ -27,7 +27,7 @@
 
 namespace uba
 {
-	static constexpr u32 CacheFileVersion = 6;
+	static constexpr u32 CacheFileVersion = 7;
 	static constexpr u32 CacheFileCompatibilityVersion = 3;
 
 	bool IsCaseInsensitive(u64 id) { return (id & (1ull << 32)) == 0; }
@@ -343,6 +343,13 @@ namespace uba
 			auto& cacheEntries = insres.first->second;
 			cacheEntries.Read(m_logger, reader, databaseVersion);
 			outStats.totalCacheEntryCount += cacheEntries.entries.size();
+
+			if (databaseVersion < 7)
+			{
+				cacheEntries.hasInputsThatAreOutputs |= CheckIfInputsAreOutputs(bucket.m_casKeyTable, cacheEntries.sharedInputCasKeyOffsets);
+				for (auto& entry : cacheEntries.entries)
+					cacheEntries.hasInputsThatAreOutputs |= CheckIfInputsAreOutputs(bucket.m_casKeyTable, entry.extraInputCasKeyOffsets);
+			}
 		}
 		return true;
 	}
@@ -746,12 +753,11 @@ namespace uba
 		u64 deleteCacheEntriesStartTime = GetTime();
 		do
 		{
-			bool checkInputsForDeletes = m_checkInputsForDeletedCas && !deletedCasFiles.empty();
-
 			oldest = 0;
 			oldestUsedTime = 0;
 			totalEntryCount = 0;
 
+			// Traverse all buckets in parallel
 			m_server.ParallelFor(workerCountToUseForBuckets, sortedBuckets, [&](auto& it)
 			{
 				Bucket& bucket = **it;
@@ -759,23 +765,25 @@ namespace uba
 				if (!context)
 					context = bucket.m_maintenanceContext = new Bucket::MaintenanceContext;
 
+				// Traverse all deleted cas files and create a bit table which has the bit of caskeytable offset set to true if deleted
+				// This is created to make fast lookup further down
 				bool foundDeletedCasKey = false;
 				for (auto& cas : deletedCasFiles)
-					bucket.m_casKeyTable.TraverseOffsets(cas, [&](u32 casKeyOffset)
+					bucket.m_casKeyTable.TraverseOffsets(cas, [&](u32 casKeyOffset) // There can be multiple offsets since caskey can be the same but path different
 						{
 							EnsureBucketContextInitialized(bucket);
 							foundDeletedCasKey = true;
 							context->deletedOffsets.Set(casKeyOffset);
 						});
 
-				if (!foundDeletedCasKey)
-					checkInputsForDeletes = false;
-
 				auto& deletedOffsets = context->deletedOffsets;
 
 				ReaderWriterLock keysToEraseLock;
 				Vector<CasKey> keysToErase;
 
+
+				// Check if we need to change expiration time. This is only done first iteration (to prevent very long maintenance times)
+
 				u64 bucketExpirationTimeSeconds = m_expirationTimeSeconds;
 
 				u64 lastUseTimeLimit = 0; // This is the time relative to server startup time
@@ -806,11 +814,14 @@ namespace uba
 				bucket.totalEntryCount = 0;
 				bucket.totalEntrySize = 0;
 
+				// Loop through all cache entries in parallel
 				m_server.ParallelFor(workerCountToUse, bucket.m_cacheEntryLookup, [&, touchedCas = Vector<Atomic<bool>*>()](auto& li) mutable
 				{
 					CacheEntries& entries = li->second;
 
-					// There is currently no idea saving more than 256kb worth of entries per lookup key (because that is what fetch max returns).. so let's wipe out
+					bool checkInputsForDeletes = foundDeletedCasKey && m_checkInputsForDeletedCas && entries.hasInputsThatAreOutputs;
+
+					// There is currently no idea keeping more than 256kb worth of entries per lookup key (because that is what fetch max returns).. so let's wipe out
 					// all the entries that overflow that number
 					u64 capacityLeft = SendMaxSize - 32 - entries.GetSharedSize();
 
@@ -887,6 +898,7 @@ namespace uba
 							}
 						}
 
+						// We must always traverse this to collect output cas keys.. since if this entry is not deleted we will set cas keys as used
 						if (!deleteEntry)
 						{
 							// Traverse outputs and check if cas files exists for each output, if not, delete entry.
@@ -1474,6 +1486,7 @@ namespace uba
 		CacheEntry newEntry;
 		cacheEntries.BuildInputs(newEntry, inputs);
 
+		// Check if there already is an entry with exactly the same inputs
 		List<CacheEntry>::iterator matchingEntry = cacheEntries.entries.end();
 		for (auto i=cacheEntries.entries.begin(), e=cacheEntries.entries.end(); i!=e; ++i)
 		{
@@ -1489,6 +1502,7 @@ namespace uba
 			bool shouldOverwrite = false;
 			Map<TString, CasKey> existing;
 
+			// Collect all outputs from entry with matching inputs
 			BinaryReader r2(matchingEntry->outputCasKeyOffsets.data(), 0, matchingEntry->outputCasKeyOffsets.size());
 			while (r2.GetLeft())
 			{
@@ -1499,6 +1513,7 @@ namespace uba
 				existing.try_emplace(path.data, casKey);
 			}
 
+			// Traverse outputs and log out the mismatching output, then remove the old entry
 			reader.SetPosition(outputStartOffset);
 			u64 left = outputCount;
 			while (left--)
@@ -1531,6 +1546,7 @@ namespace uba
 		Set<u32> outputs;
 		u64 bytesForOutput = 0;
 
+		// Check if all content for output caskeys exist on the server.. We don't want to publish the entry before server has all cas keys
 		bool hasAllContent = true;
 		reader.SetPosition(outputStartOffset);
 		u64 left = outputCount;
@@ -1560,6 +1576,7 @@ namespace uba
 			}
 		}
 
+		// Write outputs into buffer for entry
 		newEntry.outputCasKeyOffsets.resize(bytesForOutput);
 		BinaryWriter w2(newEntry.outputCasKeyOffsets.data(), 0, newEntry.outputCasKeyOffsets.size());
 		for (u32 output : outputs)
@@ -1591,6 +1608,14 @@ namespace uba
 			}
 		}
 
+		// Let's check if any inputs are cas files.. this is an optimization for maintenance...
+		// We never have to check inputs against deleted cas files if they never existed
+		{
+			if (cacheEntries.entries.empty())
+				cacheEntries.hasInputsThatAreOutputs |= CheckIfInputsAreOutputs(bucket.m_casKeyTable, cacheEntries.sharedInputCasKeyOffsets);
+			cacheEntries.hasInputsThatAreOutputs |= CheckIfInputsAreOutputs(bucket.m_casKeyTable, newEntry.extraInputCasKeyOffsets);
+		}
+
 		// If cache server has all content we can put the new cache entry directly in the lookup.. otherwise we'll have to wait until client has uploaded content
 		if (hasAllContent)
 		{
@@ -1845,6 +1870,7 @@ namespace uba
 					u64 mostEntries = 0;
 					u64 lastUsed = 0;
 					u64 totalEntryCount = 0;
+					u64 totalCountHasOutputAsInput = 0;
 					for (auto& kv2 : bucket.m_cacheEntryLookup)
 					{
 						CacheEntries& entries = kv2.second;
@@ -1853,6 +1879,8 @@ namespace uba
 						for (auto& entry : entries.entries)
 							lastUsed = Max(lastUsed, entry.lastUsedTime);
 						totalEntryCount += entries.entries.size();
+						if (entries.hasInputsThatAreOutputs)
+							++totalCountHasOutputAsInput;
 					}
 					lock2.Leave();
 					u64 lastUsedTime = 0;
@@ -1861,9 +1889,11 @@ namespace uba
 
 					writeLine(line.Clear().Appendf(TC("    #%u - %llu (v%u)"), index++, kv.first, bucket.m_pathTable.GetVersion()));
 					writeLine(line.Clear().Appendf(TC("      PathTable: %llu (%s)"), bucket.m_pathTable.GetPathCount(), BytesToText(bucket.m_pathTable.GetSize()).str));
-					writeLine(line.Clear().Appendf(TC("      CasKeyTable: %llu (%s)"), bucket.m_cacheEntryLookup.size(), BytesToText(bucket.m_casKeyTable.GetSize()).str));
-					writeLine(line.Clear().Appendf(TC("      TotalEntries: %llu"), totalEntryCount));
+					writeLine(line.Clear().Appendf(TC("      CasKeyTable: %llu (%s)"), bucket.m_casKeyTable.GetKeyCount(), BytesToText(bucket.m_casKeyTable.GetSize()).str));
+					writeLine(line.Clear().Appendf(TC("      Keys: %llu"), bucket.m_cacheEntryLookup.size()));
+					writeLine(line.Clear().Appendf(TC("      KeysWithInputsFromOutputs: %llu"), totalCountHasOutputAsInput));
 					writeLine(line.Clear().Appendf(TC("      KeyMostEntries: %llu"), mostEntries));
+					writeLine(line.Clear().Appendf(TC("      TotalEntries: %llu"), totalEntryCount));
 					writeLine(line.Clear().Appendf(TC("      LastEntryUsed: %s ago"), TimeToText(lastUsedTime, true).str));
 				}
 			}
@@ -1929,4 +1959,18 @@ namespace uba
 		DeleteFileW(tempFile.data);
 		return true;
 	}
+
+	bool CacheServer::CheckIfInputsAreOutputs(CompactCasKeyTable& table, const Vector<u8>& inputData)
+	{
+		bool inputsAreOutputs = false;
+		BinaryReader inputReader(inputData.data(), 0, inputData.size());
+		while (inputReader.GetLeft())
+		{
+			u32 casKeyOffset = u32(inputReader.Read7BitEncoded());
+			CasKey casKey;
+			table.GetKey(casKey, casKeyOffset);
+			inputsAreOutputs |= m_storage.HasBeenSeen(casKey);
+		}
+		return inputsAreOutputs;
+	}
 }
diff --git a/Engine/Source/Programs/UnrealBuildAccelerator/Common/Private/UbaStorage.cpp b/Engine/Source/Programs/UnrealBuildAccelerator/Common/Private/UbaStorage.cpp
--- a/Engine/Source/Programs/UnrealBuildAccelerator/Common/Private/UbaStorage.cpp
+++ b/Engine/Source/Programs/UnrealBuildAccelerator/Common/Private/UbaStorage.cpp
@@ -2063,15 +2063,18 @@ namespace uba
 				u8* fileMem = scope.MapView(TC("StoreCasFile"), externalMapping.handle, externalMapping.offset, externalMapping.size, fileName);
 				if (!fileMem)
 					return false;
-				fileEntry.casKey = CalculateCasKey(fileMem, externalMapping.size, m_storeCompressed);
+				if (fileSize > sizeof(CompressedObjFileHeader) && ((CompressedObjFileHeader*)fileMem)->IsValid())
+					fileEntry.casKey = ((CompressedObjFileHeader*)fileMem)->casKey;
+				else
+					fileEntry.casKey = CalculateCasKey(fileMem, externalMapping.size, m_storeCompressed);
 			}
-			else if (fileIsCompressed)
+			else if (fileIsCompressed && fileSize > sizeof(CompressedObjFileHeader))
 			{
-				UBA_ASSERT(fileSize > sizeof(CompressedObjFileHeader));
 				CompressedObjFileHeader header(CasKeyZero);
 				if (!ReadFile(m_logger, fileName, fileHandle, &header, sizeof(header)))
 					return m_logger.Error(TC("Failed to read header of compressed file %s (%s)"), fileName, LastErrorToText().data);
-				UBA_ASSERT(header.IsValid());
+				if (!header.IsValid())
+					return m_logger.Error(TC("File was expected to be compressed but was not %s (%s)"), fileName, LastErrorToText().data);
 				fileEntry.casKey = AsCompressed(header.casKey, m_storeCompressed);
 			}
 			else
@@ -2117,6 +2120,12 @@ namespace uba
 		fileEntry.verified = fileEntry.lastWritten == verifiedLastWriteTime && fileEntry.size == verifiedSize && fileEntry.casKey != CasKeyInvalid;
 	}
 
+	bool StorageImpl::HasBeenSeen(const CasKey& casKey)
+	{
+		SCOPED_READ_LOCK(m_casLookupLock, lookupLock);
+		return m_casLookup.find(casKey) != m_casLookup.end();
+	}
+
 	bool StorageImpl::StoreCasKey(CasKey& out, const tchar* fileName, const CasKey& casKeyOverride, bool fileIsCompressed)
 	{
 		StringBuffer<> forKey;
diff --git a/Engine/Source/Programs/UnrealBuildAccelerator/Common/Public/UbaCacheClient.h b/Engine/Source/Programs/UnrealBuildAccelerator/Common/Public/UbaCacheClient.h
--- a/Engine/Source/Programs/UnrealBuildAccelerator/Common/Public/UbaCacheClient.h
+++ b/Engine/Source/Programs/UnrealBuildAccelerator/Common/Public/UbaCacheClient.h
@@ -84,7 +84,7 @@ namespace uba
 		bool ShouldNormalize(const StringBufferBase& path);
 
 		bool GetLocalPathAndCasKey(Bucket& bucket, const RootPaths& rootPaths, StringBufferBase& outPath, CasKey& outKey, CompactCasKeyTable& casKeyTable, CompactPathTable& pathTable, u32 offset);
-		bool IsFileCompressed(const ProcessStartInfo& info, const StringView& filename);
+		bool IsFileCompressed(const ProcessStartInfo& info, const StringView& filename, bool isOutput);
 		void PreparseDirectory(const StringKey& fileNameKey, const StringBufferBase& filePath);
 
 		MutableLogger m_logger;
diff --git a/Engine/Source/Programs/UnrealBuildAccelerator/Common/Public/UbaCacheEntry.h b/Engine/Source/Programs/UnrealBuildAccelerator/Common/Public/UbaCacheEntry.h
--- a/Engine/Source/Programs/UnrealBuildAccelerator/Common/Public/UbaCacheEntry.h
+++ b/Engine/Source/Programs/UnrealBuildAccelerator/Common/Public/UbaCacheEntry.h
@@ -40,6 +40,7 @@ namespace uba
 		Vector<u8> sharedLogLines;
 		u32 idCounter = 0;
 		u32 primaryId = ~0u; // Id of entry that shared offsets was made from
+		bool hasInputsThatAreOutputs = false;
 
 		u64 GetSharedSize();
 		u64 GetEntrySize(CacheEntry& entry, u32 clientVersion, bool toDisk);
diff --git a/Engine/Source/Programs/UnrealBuildAccelerator/Common/Public/UbaCacheServer.h b/Engine/Source/Programs/UnrealBuildAccelerator/Common/Public/UbaCacheServer.h
--- a/Engine/Source/Programs/UnrealBuildAccelerator/Common/Public/UbaCacheServer.h
+++ b/Engine/Source/Programs/UnrealBuildAccelerator/Common/Public/UbaCacheServer.h
@@ -8,6 +8,7 @@
 
 namespace uba
 {
+	class CompactCasKeyTable;
 	class NetworkServer;
 	class StorageServer;
 	struct BinaryReader;
@@ -84,6 +85,8 @@ namespace uba
 		bool HandleReportUsedEntry(BinaryReader& reader, BinaryWriter& writer, u32 clientId);
 		bool HandleExecuteCommand(BinaryReader& reader, BinaryWriter& writer);
 
+		bool CheckIfInputsAreOutputs(CompactCasKeyTable& table, const Vector<u8>& inputData);
+
 		MutableLogger m_logger;
 		NetworkServer& m_server;
 		StorageServer& m_storage;
diff --git a/Engine/Source/Programs/UnrealBuildAccelerator/Common/Public/UbaStorage.h b/Engine/Source/Programs/UnrealBuildAccelerator/Common/Public/UbaStorage.h
--- a/Engine/Source/Programs/UnrealBuildAccelerator/Common/Public/UbaStorage.h
+++ b/Engine/Source/Programs/UnrealBuildAccelerator/Common/Public/UbaStorage.h
@@ -160,6 +160,7 @@ namespace uba
 		bool StoreCasKey(CasKey& out, const StringKey& fileNameKey, const tchar* fileName, const CasKey& casKeyOverride, bool fileIsCompressed);
 		bool IsFileVerified(const StringKey& fileNameKey);
 		void ReportFileInfoWeak(const StringKey& fileNameKey, u64 verifiedLastWriteTime, u64 verifiedSize);
+		bool HasBeenSeen(const CasKey& casKey);
 
 		virtual void RegisterExternalFileMappingsProvider(ExternalFileMappingsProvider&& provider) override;
 		bool WriteCompressed(WriteResult& out, const StringKey& fromKey, const tchar* from, const tchar* toFile);
