[UBA]
* Added hints to some functions that does allocations that can fail

diff --git a/Engine/Source/Programs/UnrealBuildAccelerator/Common/Private/UbaCacheServer.cpp b/Engine/Source/Programs/UnrealBuildAccelerator/Common/Private/UbaCacheServer.cpp
--- a/Engine/Source/Programs/UnrealBuildAccelerator/Common/Private/UbaCacheServer.cpp
+++ b/Engine/Source/Programs/UnrealBuildAccelerator/Common/Private/UbaCacheServer.cpp
@@ -34,10 +34,10 @@ namespace uba
 
 	struct BitArray
 	{
-		void Init(MemoryBlock& memoryBlock, u32 bitCount)
+		void Init(MemoryBlock& memoryBlock, u32 bitCount, const tchar* hint)
 		{
 			u32 bytes = AlignUp((bitCount+7) / 8, 8u); // Align up to 64 bits
-			data = (u64*)memoryBlock.Allocate(bytes, 8, TC(""));
+			data = (u64*)memoryBlock.Allocate(bytes, 8, hint);
 			memset(data, 0, bytes);
 			count = bytes / 8;
 		}
@@ -656,7 +656,7 @@ namespace uba
 		struct CasFileInfo { CasFileInfo(u32 s = 0) : size(s) {} u32 size; Atomic<bool> isUsed; }; // These are compressed cas, should never be over 4gb
 
 		// Existing cas entries can be more than 2 million entries.. which uses a lot of memory
-		constexpr u64 existingCasMemoryReserveSize = 192*1024*1024;
+		constexpr u64 existingCasMemoryReserveSize = 512*1024*1024;
 		MemoryBlock existingCasMemoryBlock;
 		if (!existingCasMemoryBlock.Init(existingCasMemoryReserveSize, nullptr, true))
 			existingCasMemoryBlock.Init(existingCasMemoryReserveSize);
@@ -676,7 +676,7 @@ namespace uba
 			SCOPED_WRITE_LOCK(m_storage.m_casLookupLock, lookupLock);
 			
 			totalCasCount = m_storage.m_casLookup.size();
-			existingCas.Init(existingCasMemoryBlock, totalCasCount);
+			existingCas.Init(existingCasMemoryBlock, totalCasCount, TC("ExistingCas"));
 
 			for (auto i=m_storage.m_casLookup.begin(), e=m_storage.m_casLookup.end(); i!=e;)
 			{
@@ -744,7 +744,7 @@ namespace uba
 				{
 					if (!context.memoryBlock.Init(m_maintenanceReserveSize, nullptr, true)) // Try to use large blocks
 						context.memoryBlock.Init(m_maintenanceReserveSize, nullptr, false);
-					context.deletedOffsets.Init(context.memoryBlock, bucket.m_casKeyTable.GetSize());
+					context.deletedOffsets.Init(context.memoryBlock, bucket.m_casKeyTable.GetSize(), TC("DeletedOffsets"));
 					context.isInitialized = true;
 				}
 			};
@@ -1041,7 +1041,7 @@ namespace uba
 			MemoryBlock& memoryBlock = bucket.m_maintenanceContext->memoryBlock;
 
 			BitArray usedCasKeyOffsets;
-			usedCasKeyOffsets.Init(memoryBlock, bucket.m_casKeyTable.GetSize());
+			usedCasKeyOffsets.Init(memoryBlock, bucket.m_casKeyTable.GetSize(), TC("UsedCasKeyOffsets"));
 
 			u64 collectUsedCasKeysStart = GetTime();
 
@@ -1073,7 +1073,7 @@ namespace uba
 
 			// Traverse all caskeys in caskey table and figure out which ones we can delete
 			BitArray usedPathOffsets;
-			usedPathOffsets.Init(memoryBlock, bucket.m_pathTable.GetSize());
+			usedPathOffsets.Init(memoryBlock, bucket.m_pathTable.GetSize(), TC("UsedPathOffsets"));
 
 			BinaryReader casKeyTableReader(bucket.m_casKeyTable.GetMemory(), 0, bucket.m_casKeyTable.GetSize());
 			usedCasKeyOffsets.Traverse([&](u32 casKeyOffset)
@@ -1091,7 +1091,7 @@ namespace uba
 				u64 pathReserveCount = pathTable.GetPathCount() + 1000; // Add extra reservation space for growth
 				u64 segmentReserveCount = pathTable.GetSegmentCount() + 1000;  // Add extra reservation space for growth
 				CompactPathTable newPathTable(pathTable.GetCaseInsensitive(), pathReserveCount, segmentReserveCount, pathTable.GetVersion());
-				oldToNewPathOffset.Init(memoryBlock, usedPathOffsets.CountSetBits());
+				oldToNewPathOffset.Init(memoryBlock, usedPathOffsets.CountSetBits(), TC("OldToNewPathOffset"));
 
 				StringBuffer<> temp;
 				usedPathOffsets.Traverse([&](u32 pathOffset)
@@ -1118,7 +1118,7 @@ namespace uba
 			auto& casKeyTable = bucket.m_casKeyTable;
 			oldSize = casKeyTable.GetSize();
 			{
-				oldToNewCasKeyOffset.Init(memoryBlock, usedCasKeyOffsetsCount);
+				oldToNewCasKeyOffset.Init(memoryBlock, usedCasKeyOffsetsCount, TC("OldToNewCasKeyOffset"));
 				CompactCasKeyTable newCasKeyTable(usedCasKeyOffsetsCount + Min(usedCasKeyOffsetsCount/2, 10'000ull)); // Add extra reservation space for growth
 				BinaryReader reader2(casKeyTable.GetMemory(), 0, oldSize);
 				usedCasKeyOffsets.Traverse([&](u32 casKeyOffset)
diff --git a/Engine/Source/Programs/UnrealBuildAccelerator/Common/Public/UbaHashMap.h b/Engine/Source/Programs/UnrealBuildAccelerator/Common/Public/UbaHashMap.h
--- a/Engine/Source/Programs/UnrealBuildAccelerator/Common/Public/UbaHashMap.h
+++ b/Engine/Source/Programs/UnrealBuildAccelerator/Common/Public/UbaHashMap.h
@@ -26,15 +26,15 @@ namespace uba
 				delete m_memory;
 		}
 
-		void Init(MemoryBlock& mem, u64 reserveCount)
+		void Init(MemoryBlock& mem, u64 reserveCount, const tchar* hint = TC(""))
 		{
 			reserveCount = Max(reserveCount + 1, 4ull);
 			m_memory = &mem;
 			UBA_ASSERT(!AllowGrow);
-			InternalInit(reserveCount, RoundUpPow2(reserveCount));
+			InternalInit(reserveCount, RoundUpPow2(reserveCount), hint);
 		}
 
-		void Init(u64 reserveCount = 4)
+		void Init(u64 reserveCount = 4, const tchar* hint = TC(""))
 		{
 			reserveCount = Max(reserveCount + 1, 4ull);
 			u64 lookupCount = RoundUpPow2(reserveCount);
@@ -46,7 +46,7 @@ namespace uba
 			m_memory->Init(lookupSize + entriesSize + valuesSize);
 			m_ownsMemory = true;
 
-			InternalInit(reserveCount, lookupCount);
+			InternalInit(reserveCount, lookupCount, hint);
 		}
 
 		bool IsInitialized()
@@ -202,15 +202,15 @@ namespace uba
 		}
 
 private:
-		void InternalInit(u64 reserveCount, u64 lookupCount)
+		void InternalInit(u64 reserveCount, u64 lookupCount, const tchar* hint)
 		{
 			m_mask = u32(lookupCount - 1);
 			m_nextAvailableEntry = 1;
 			m_reserveCount = u32(reserveCount);
 			// This code relies on that allocated memory is zeroed out
-			m_lookup = (u32*)m_memory->AllocateNoLock(lookupCount*sizeof(u32), 1, TC(""));
-			m_entries = (Entry*)m_memory->AllocateNoLock(m_reserveCount*sizeof(Entry), alignof(Entry), TC(""));
-			m_values = (Value*)m_memory->AllocateNoLock(m_reserveCount*sizeof(Value), alignof(Value), TC(""));
+			m_lookup = (u32*)m_memory->AllocateNoLock(lookupCount*sizeof(u32), 1, hint);
+			m_entries = (Entry*)m_memory->AllocateNoLock(m_reserveCount*sizeof(Entry), alignof(Entry), hint);
+			m_values = (Value*)m_memory->AllocateNoLock(m_reserveCount*sizeof(Value), alignof(Value), hint);
 		}
 
 		u64 GetNextReserveCount(u64 reserveCount)
@@ -246,7 +246,7 @@ private:
 			u32 next; // If this is 0 it means that it is unused, ~0u means it is the end of a chain, otherwise offset into entry array
 		};
 
-		void Init(MemoryBlock& memory, u64 maxSize)
+		void Init(MemoryBlock& memory, u64 maxSize, const tchar* hint = TC(""))
 		{
 			u32 v = RoundUpPow2(maxSize);
 			u32 lookupSize = v;// << 1u;
@@ -254,8 +254,8 @@ private:
 			nextAvailableEntry = 1;
 
 			// This code relies on that AllocateNoLockd memory is zeroed out
-			lookup = (u32*)memory.AllocateNoLock(lookupSize*sizeof(u32), 1, TC(""));
-			entries = (Entry*)memory.AllocateNoLock((maxSize+1)*sizeof(Entry), alignof(Entry), TC(""));
+			lookup = (u32*)memory.AllocateNoLock(lookupSize*sizeof(u32), 1, hint);
+			entries = (Entry*)memory.AllocateNoLock((maxSize+1)*sizeof(Entry), alignof(Entry), hint);
 		}
 
 		Value& Insert(const Key& key)
