[UBA]
* Fixed connect time bug in disconnect log entry

diff --git a/Engine/Source/Programs/UnrealBuildAccelerator/Common/Private/UbaCacheServer.cpp b/Engine/Source/Programs/UnrealBuildAccelerator/Common/Private/UbaCacheServer.cpp
--- a/Engine/Source/Programs/UnrealBuildAccelerator/Common/Private/UbaCacheServer.cpp
+++ b/Engine/Source/Programs/UnrealBuildAccelerator/Common/Private/UbaCacheServer.cpp
@@ -1244,7 +1244,7 @@ namespace uba
 		lock.Leave();
 
 		StringBuffer<> log;
-		log.Appendf(TC("Client %u disconnected after %s (%llu active). "), clientId, TimeToText(connectTime - GetTime()).str, activeCount);
+		log.Appendf(TC("Client %u disconnected after %s (%llu active). "), clientId, TimeToText(GetTime() - connectTime).str, activeCount);
 		if (fetchBucketsCount)
 			log.Appendf(TC("Fetched %llu (%llu hits) from %llu bucket(s)"), fetchEntryCount, fetchEntryHitCount, fetchBucketsCount);
 		if (fetchBucketsCount && storeBucketsCount)
