[UBA]
* Added so CacheClient_Create takes in hint for name

diff --git a/Engine/Source/Programs/UnrealBuildAccelerator/Host/Private/UbaExports.cpp b/Engine/Source/Programs/UnrealBuildAccelerator/Host/Private/UbaExports.cpp
--- a/Engine/Source/Programs/UnrealBuildAccelerator/Host/Private/UbaExports.cpp
+++ b/Engine/Source/Programs/UnrealBuildAccelerator/Host/Private/UbaExports.cpp
@@ -579,7 +579,7 @@ extern "C"
 		scheduler->GetStats(outQueued, outActiveLocal, outActiveRemote, outFinished);
 	}
 
-	uba::CacheClient* CacheClient_Create(uba::SessionServer* session, bool reportMissReason, const uba::tchar* crypto)
+	uba::CacheClient* CacheClient_Create(uba::SessionServer* session, bool reportMissReason, const uba::tchar* crypto, const uba::tchar* hint)
 	{
 		using namespace uba;
 		LogWriter& writer = session->GetLogWriter();
@@ -617,6 +617,8 @@ extern "C"
 
 		CacheClientCreateInfo info{writer, storage, *networkClient, *session};
 		info.Apply(GetConfig());
+		if (hint && *hint)
+			info.hint = hint;
 
 		info.reportMissReason |= reportMissReason;
 		return new CacheClient(info);
diff --git a/Engine/Source/Programs/UnrealBuildAccelerator/Host/Public/UbaExports.h b/Engine/Source/Programs/UnrealBuildAccelerator/Host/Public/UbaExports.h
--- a/Engine/Source/Programs/UnrealBuildAccelerator/Host/Public/UbaExports.h
+++ b/Engine/Source/Programs/UnrealBuildAccelerator/Host/Public/UbaExports.h
@@ -127,7 +127,7 @@ extern "C"
 	UBA_API void Scheduler_GetStats(uba::Scheduler* scheduler, uba::u32& outQueued, uba::u32& outActiveLocal, uba::u32& outActiveRemote, uba::u32& outFinished);
 
 	// Cache
-	UBA_API uba::CacheClient* CacheClient_Create(uba::SessionServer* session, bool reportMissReason = false, const uba::tchar* crypto = nullptr);
+	UBA_API uba::CacheClient* CacheClient_Create(uba::SessionServer* session, bool reportMissReason = false, const uba::tchar* crypto = nullptr, const uba::tchar* hint = nullptr);
 	UBA_API bool CacheClient_Connect(uba::CacheClient* cacheClient, const uba::tchar* host, int port);
 	UBA_API bool CacheClient_WriteToCache2(uba::CacheClient* cacheClient, uba::u32 bucket, const uba::ProcessHandle* process, const uba::u8* inputs, uba::u32 inputsSize, const uba::u8* outputs, uba::u32 outputsSize, const uba::u8* logLines, uba::u32 logLinesSize);
 	UBA_API uba::CacheResult* CacheClient_FetchFromCache3(uba::CacheClient* cacheClient, uba::u64 rootsHandle, uba::u32 bucket, const uba::ProcessStartInfo& info);
