[UBA]
* Some cosmetic fixes of log entries on cache server
* Added Logger.Log(type, stringview)

diff --git a/Engine/Source/Programs/UnrealBuildAccelerator/Common/Private/UbaCacheServer.cpp b/Engine/Source/Programs/UnrealBuildAccelerator/Common/Private/UbaCacheServer.cpp
--- a/Engine/Source/Programs/UnrealBuildAccelerator/Common/Private/UbaCacheServer.cpp
+++ b/Engine/Source/Programs/UnrealBuildAccelerator/Common/Private/UbaCacheServer.cpp
@@ -1251,7 +1251,7 @@ namespace uba
 			log.Append(TC(", "));
 		if (storeBucketsCount)
 			log.Appendf(TC("Stored %llu to %llu bucket(s)"), storeEntryCount, storeBucketsCount);
-		m_logger.Log(LogEntryType_Info, log.data, log.count);
+		m_logger.Log(LogEntryType_Info, log);
 	}
 
 	CacheServer::ConnectionBucket& CacheServer::GetConnectionBucket(const ConnectionInfo& connectionInfo, BinaryReader& reader, u32* outClientVersion)
@@ -1325,7 +1325,13 @@ namespace uba
 			connection.connectTime = GetTime();
 			m_peakConnectionCount = Max(m_peakConnectionCount, u32(m_connections.size()));
 
-			m_logger.Info(TC("Client %u connected (v%u) %s"), connectionInfo.GetId(), clientVersion, hint.c_str());
+			StringBuffer<> logStr;
+			logStr.Appendf(TC("Client %u connected"), connectionInfo.GetId());
+			if (clientVersion != CacheNetworkVersion)
+				logStr.Appendf(TC(" (v%u)"), clientVersion);
+			if (!hint.empty())
+				logStr.Appendf(TC(" %s"), hint.c_str());
+			m_logger.Log(LogEntryType_Info, logStr);
 
 			return true;
 		}
diff --git a/Engine/Source/Programs/UnrealBuildAccelerator/Common/Private/UbaLogger.cpp b/Engine/Source/Programs/UnrealBuildAccelerator/Common/Private/UbaLogger.cpp
--- a/Engine/Source/Programs/UnrealBuildAccelerator/Common/Private/UbaLogger.cpp
+++ b/Engine/Source/Programs/UnrealBuildAccelerator/Common/Private/UbaLogger.cpp
@@ -204,6 +204,11 @@ namespace uba
 		Log(type, buf.data(), u32(len));
 	}
 
+	void Logger::Log(LogEntryType type, const StringView& str)
+	{
+		Log(type, str.data, str.count);
+	}
+
 	LoggerWithWriter::LoggerWithWriter(LogWriter& writer, const tchar* prefix)
 	:	m_writer(writer)
 	,	m_prefix(prefix)
diff --git a/Engine/Source/Programs/UnrealBuildAccelerator/Common/Private/UbaObjectFileLLVMIR.cpp b/Engine/Source/Programs/UnrealBuildAccelerator/Common/Private/UbaObjectFileLLVMIR.cpp
--- a/Engine/Source/Programs/UnrealBuildAccelerator/Common/Private/UbaObjectFileLLVMIR.cpp
+++ b/Engine/Source/Programs/UnrealBuildAccelerator/Common/Private/UbaObjectFileLLVMIR.cpp
@@ -1027,8 +1027,8 @@ namespace uba
 			else
 				reader.SkipRecord(entry.id);
 
-#if UBA_LOG_STREAM
-			//logger.Log(LogEntryType_Info, blob.c_str(), blob.size());
+			#if UBA_LOG_STREAM
+			//logger.Log(LogEntryType_Info, blob);
 			#endif
 		}
 
diff --git a/Engine/Source/Programs/UnrealBuildAccelerator/Common/Private/UbaSession.cpp b/Engine/Source/Programs/UnrealBuildAccelerator/Common/Private/UbaSession.cpp
--- a/Engine/Source/Programs/UnrealBuildAccelerator/Common/Private/UbaSession.cpp
+++ b/Engine/Source/Programs/UnrealBuildAccelerator/Common/Private/UbaSession.cpp
@@ -274,7 +274,7 @@ namespace uba
 		if (isRefresh)
 			str.Append(TC("(Refresh) "));
 		str.Append(dirPath).Append('\n');
-		m_debugLogger->Log(LogEntryType_Info, str.data, str.count);
+		m_debugLogger->Log(LogEntryType_Info, str);
 		#endif
 
 
diff --git a/Engine/Source/Programs/UnrealBuildAccelerator/Common/Public/UbaLogger.h b/Engine/Source/Programs/UnrealBuildAccelerator/Common/Public/UbaLogger.h
--- a/Engine/Source/Programs/UnrealBuildAccelerator/Common/Public/UbaLogger.h
+++ b/Engine/Source/Programs/UnrealBuildAccelerator/Common/Public/UbaLogger.h
@@ -31,6 +31,7 @@ namespace uba
 		void Debug(const tchar* format, ...);
 		void Logf(LogEntryType type, const tchar* format, ...);
 		void LogArg(LogEntryType type, const tchar* format, va_list& args);
+		void Log(LogEntryType type, const StringView& str);
 
 		virtual void BeginScope() = 0;
 		virtual void EndScope() = 0;
@@ -46,6 +47,7 @@ namespace uba
 		virtual void EndScope() { m_writer.EndScope(); }
 		virtual void Log(LogEntryType type, const tchar* str, u32 strLen) { m_writer.Log(type, str, strLen, m_prefix, m_prefixLen); }
 
+		using Logger::Log;
 		LogWriter& m_writer;
 		const tchar* m_prefix;
 		u32 m_prefixLen;
@@ -55,6 +57,7 @@ namespace uba
 	{
 		MutableLogger(LogWriter& writer, const tchar* prefix) : LoggerWithWriter(writer, prefix) {}
 		virtual void Log(LogEntryType type, const tchar* str, u32 strLen) override { if (!isMuted) LoggerWithWriter::Log(type, str, strLen); }
+		using LoggerWithWriter::Log;
 		Atomic<u32> isMuted;
 	};
 
diff --git a/Engine/Source/Programs/UnrealBuildAccelerator/Visualizer/Private/UbaVisualizer.cpp b/Engine/Source/Programs/UnrealBuildAccelerator/Visualizer/Private/UbaVisualizer.cpp
--- a/Engine/Source/Programs/UnrealBuildAccelerator/Visualizer/Private/UbaVisualizer.cpp
+++ b/Engine/Source/Programs/UnrealBuildAccelerator/Visualizer/Private/UbaVisualizer.cpp
@@ -202,6 +202,7 @@ namespace uba
 		HBRUSH backgroundBrush;
 		COLORREF textColor;
 		bool isFirst = true;
+		using Logger::Log;
 	};
 
 	class WriteTextLogger : public Logger
@@ -1914,7 +1915,7 @@ namespace uba
 					logger.Info(L"  ---------------- Log ----------------");
 					logger.AddTextOffset(14);
 					for (auto& line : logLines)
-						logger.Log(LogEntryType_Info, line.c_str(), u32(line.size()));
+						logger.Log(LogEntryType_Info, line);
 				}
 			}
 			logger.AddSpace(3);
@@ -1955,7 +1956,7 @@ namespace uba
 					logger.Info(L" Session summary not available until session is done");
 			}
 			for (auto& line : summary)
-				logger.Log(LogEntryType_Info, line.c_str(), u32(line.size()));
+				logger.Log(LogEntryType_Info, line);
 			logger.AddSpace(3);
 			logger.DrawAtCursor();
 		}
