[UBA]
* Added bucket count to maintenance done log entry

diff --git a/Engine/Source/Programs/UnrealBuildAccelerator/Common/Private/UbaCacheServer.cpp b/Engine/Source/Programs/UnrealBuildAccelerator/Common/Private/UbaCacheServer.cpp
--- a/Engine/Source/Programs/UnrealBuildAccelerator/Common/Private/UbaCacheServer.cpp
+++ b/Engine/Source/Programs/UnrealBuildAccelerator/Common/Private/UbaCacheServer.cpp
@@ -1202,7 +1202,7 @@ namespace uba
 		u64 oldestTime = oldest ? GetFileTimeAsTime(now - (m_creationTime + oldest)) : 0;
 		u64 longestUnused = oldestUsedTime ? GetFileTimeAsTime(now - (m_creationTime + oldestUsedTime)) : 0;
 		u64 duration = GetTime() - startTime;
-		m_logger.Info(TC("Maintenance done! (%s) CasFiles: %s (%s) Entries: %s Oldest: %s LongestUnused: %s MaintenanceMem: %s/%s"), TimeToText(duration).str, CountToText(totalCasCount - deletedCasCount).str, BytesToText(totalCasSize).str, CountToText(totalEntryCount.load()).str, TimeToText(oldestTime, true).str, TimeToText(longestUnused, true).str, BytesToText(maxCommittedMemory).str, BytesToText(m_maintenanceReserveSize).str);
+		m_logger.Info(TC("Maintenance done! (%s) CasFiles: %s (%s) Buckets: %llu Entries: %s Oldest: %s LongestUnused: %s MaintenanceMem: %s/%s"), TimeToText(duration).str, CountToText(totalCasCount - deletedCasCount).str, BytesToText(totalCasSize).str, m_buckets.size(), CountToText(totalEntryCount.load()).str, TimeToText(oldestTime, true).str, TimeToText(longestUnused, true).str, BytesToText(maxCommittedMemory).str, BytesToText(m_maintenanceReserveSize).str);
 		
 		m_longestMaintenance = Max(m_longestMaintenance, duration);
 
