[UBA]
* Fixed bug with compressed obj files and being opened multiple times
* Improved assert info

diff --git a/Engine/Source/Programs/UnrealBuildAccelerator/Detours/Private/Windows/UbaDetoursFileDecompressor.h b/Engine/Source/Programs/UnrealBuildAccelerator/Detours/Private/Windows/UbaDetoursFileDecompressor.h
--- a/Engine/Source/Programs/UnrealBuildAccelerator/Detours/Private/Windows/UbaDetoursFileDecompressor.h
+++ b/Engine/Source/Programs/UnrealBuildAccelerator/Detours/Private/Windows/UbaDetoursFileDecompressor.h
@@ -140,13 +140,13 @@ namespace uba
 			preload.fileInfo = &findIt->second;
 		}
 
-		void Start(const wchar_t* cmdLine)
+		void Start(const wchar_t* cmdLine, bool preprocess)
 		{
 			Oodle_SetUsageWarnings(Oodle_UsageWarnings_Disabled);
 
 			m_isEnabled = true;
 
-			if (true)
+			if (!preprocess)
 				return;
 
 			TimerScope ts(g_stats.preparseObjFiles);
diff --git a/Engine/Source/Programs/UnrealBuildAccelerator/Detours/Private/Windows/UbaDetoursFunctionsKernelBase.inl b/Engine/Source/Programs/UnrealBuildAccelerator/Detours/Private/Windows/UbaDetoursFunctionsKernelBase.inl
--- a/Engine/Source/Programs/UnrealBuildAccelerator/Detours/Private/Windows/UbaDetoursFunctionsKernelBase.inl
+++ b/Engine/Source/Programs/UnrealBuildAccelerator/Detours/Private/Windows/UbaDetoursFunctionsKernelBase.inl
@@ -2212,21 +2212,21 @@ HANDLE Detoured_CreateFileMappingW(HANDLE hFile, LPSECURITY_ATTRIBUTES lpFileMap
 				UBA_ASSERTF((flProtect == PAGE_WRITECOPY ? PAGE_READONLY : flProtect) == fi.fileMapDesiredAccess, L"Code path not implemented (%ls)", HandleToName(hFile));
 			}
 			mdh->fileObject = dh.fileObject;
+
+			u64 maxSize = ToLargeInteger(dwMaximumSizeHigh, dwMaximumSizeLow).QuadPart;
+
 			if (MemoryFile* mf = fi.memoryFile)
 			{
-				LARGE_INTEGER li;
-				li.HighPart = dwMaximumSizeHigh;
-				li.LowPart = dwMaximumSizeLow;
-				if (!(flProtect & MEM_RESERVE) && li.QuadPart)
+				if (!(flProtect & MEM_RESERVE) && maxSize)
 				{
-					mf->EnsureCommitted(*mdh, li.QuadPart);
+					mf->EnsureCommitted(*mdh, maxSize);
 					if (!mf->writtenSize && (flProtect & PAGE_READWRITE)) // TODO: Maybe we should always set writtenSize?
-						mf->writtenSize = li.QuadPart;
+						mf->writtenSize = maxSize;
 				}
 			}
 			InterlockedIncrement(&dh.fileObject->refCount);
 			HANDLE res = makeDetouredHandle(mdh);
-			DEBUG_LOG_DETOURED(L"CreateFileMappingW", L"(%ls) File %llu %ls -> %llu", (fi.memoryFile ? L"MEMORYFILE" : L"FILEMAP"), uintptr_t(hFile), HandleToName(hFile), uintptr_t(res));
+			DEBUG_LOG_DETOURED(L"CreateFileMappingW", L"(%ls) File %llu Protect %u Size %llu (%ls) -> %llu", (fi.memoryFile ? L"MEMORYFILE" : L"FILEMAP"), uintptr_t(hFile), flProtect, maxSize, HandleToName(hFile), uintptr_t(res));
 			SetLastError(ERROR_SUCCESS);
 			return res;
 		}
@@ -2360,7 +2360,7 @@ LPVOID Detoured_MapViewOfFileEx(HANDLE hFileMappingObject, DWORD dwDesiredAccess
 			if (fi.memoryFile && (dwDesiredAccess & FILE_MAP_WRITE)) // We assume changes will happen
 				fi.memoryFile->isReported = false;
 
-			DEBUG_LOG_DETOURED(L"MapViewOfFileEx", L"%llu (%s) -> 0x%llx", uintptr_t(hFileMappingObject), HandleToName(hFileMappingObject), u64(mem));
+			DEBUG_LOG_DETOURED(L"MapViewOfFileEx", L"%llu Size %llu (%s) -> 0x%llx", uintptr_t(hFileMappingObject), dwNumberOfBytesToMap, HandleToName(hFileMappingObject), u64(mem));
 
 			SCOPED_WRITE_LOCK(g_mappedFileTable.m_memLookupLock, lock);
 			++g_mappedFileTable.m_memLookup[mem];
diff --git a/Engine/Source/Programs/UnrealBuildAccelerator/Detours/Private/Windows/UbaDetoursFunctionsWin.cpp b/Engine/Source/Programs/UnrealBuildAccelerator/Detours/Private/Windows/UbaDetoursFunctionsWin.cpp
--- a/Engine/Source/Programs/UnrealBuildAccelerator/Detours/Private/Windows/UbaDetoursFunctionsWin.cpp
+++ b/Engine/Source/Programs/UnrealBuildAccelerator/Detours/Private/Windows/UbaDetoursFunctionsWin.cpp
@@ -415,7 +415,7 @@ bool EnsureMapped(DetouredHandle& handle, DWORD dwFileOffsetHigh = 0, DWORD dwFi
 	}
 	else
 	{
-		UBA_ASSERT(!info.freeFileMapOnClose);
+		UBA_ASSERTF(!info.freeFileMapOnClose, TC("File %s has been freed because of earlier close and is now reopened (%s)"), info.name, info.originalName);
 		TimerScope ts(g_kernelStats.mapViewOfFile);
 		info.fileMapMem = (u8*)True_MapViewOfFileEx(info.trueFileMapHandle, info.fileMapViewDesiredAccess, 0, 0, numberOfBytesToMap, baseAddress);
 	}
@@ -942,6 +942,7 @@ void Init(const DetoursPayload& payload, u64 startTime)
 		DEBUG_LOG(L"WorkingDir: %ls", g_virtualWorkingDir.data);
 		DEBUG_LOG(L"ExeDir: %ls", g_virtualApplicationDir.data);
 		DEBUG_LOG(L"ExeDir (true): %ls", g_exeDir.data);
+		DEBUG_LOG(L"Rules: %u (%u)", g_rules->index, GetApplicationRules()[g_rules->index].hash);
 		LogVfsInfo();
 	}
 	#endif
@@ -981,7 +982,10 @@ void Init(const DetoursPayload& payload, u64 startTime)
 	g_stats.attach.count = 1;
 
 	if (payload.storeObjFilesCompressed && g_rules->ShouldDecompressFiles(StringView()))
-		g_fileDecompressor.Start(cmdLine);
+	{
+		bool preprocess = g_rules->index == 14;
+		g_fileDecompressor.Start(cmdLine, preprocess);
+	}
 }
 
 void Deinit(u64 startTime)
