[UBA]
* Added Join function to StringBufferBase and cleaned up code in cacheserver

diff --git a/Engine/Source/Programs/UnrealBuildAccelerator/Common/Private/UbaCacheServer.cpp b/Engine/Source/Programs/UnrealBuildAccelerator/Common/Private/UbaCacheServer.cpp
--- a/Engine/Source/Programs/UnrealBuildAccelerator/Common/Private/UbaCacheServer.cpp
+++ b/Engine/Source/Programs/UnrealBuildAccelerator/Common/Private/UbaCacheServer.cpp
@@ -131,6 +131,7 @@ namespace uba
 		ReaderWriterLock cacheEntryLookupLock;
 		UnorderedMap<CasKey, CacheEntry> cacheEntryLookup;
 		u64 id;
+		u32 index = ~0u;
 	};
 
 	struct CacheServer::Connection
@@ -138,7 +139,7 @@ namespace uba
 		u32 clientVersion;
 		UnorderedMap<u64, ConnectionBucket> storeBuckets;
 		u64 storeEntryCount = 0;
-		UnorderedSet<u64> fetchBuckets;
+		UnorderedSet<u32> fetchBuckets;
 		u64 fetchEntryCount = 0;
 		u64 fetchEntryHitCount = 0;
 		u64 connectTime = 0;
@@ -1221,36 +1222,24 @@ namespace uba
 
 	void CacheServer::OnDisconnected(u32 clientId)
 	{
-		u64 fetchBucketsCount = 0;
-		u64 fetchEntryCount = 0;
-		u64 storeBucketsCount = 0;
-		u64 storeEntryCount = 0;
-		u64 fetchEntryHitCount = 0;
-		u64 connectTime = 0;
+		StringBuffer<> log;
+		log.Appendf(TC("Client %u disconnected"), clientId);
+
 		SCOPED_WRITE_LOCK(m_connectionsLock, lock);
 		auto it = m_connections.find(clientId);
 		if (it != m_connections.end())
 		{
-			auto& conn = it->second;
-			fetchBucketsCount = conn.fetchBuckets.size();
-			fetchEntryCount = conn.fetchEntryCount;
-			fetchEntryHitCount = conn.fetchEntryHitCount;
-			storeBucketsCount = conn.storeBuckets.size();
-			storeEntryCount = conn.storeEntryCount;
-			connectTime = conn.connectTime;
+			u64 activeCount = m_connections.size() - 1;
+			auto& c = it->second;
+			log.Appendf(TC(" after %s (%llu active)."), TimeToText(GetTime() - c.connectTime).str, activeCount);
+			if (!c.fetchBuckets.empty())
+				log.Appendf(TC(" Fetched %llu (%llu hits) entries from bucket "), c.fetchEntryCount, c.fetchEntryHitCount).Join(c.fetchBuckets, [&](auto bucketIndex) { log.AppendValue(bucketIndex); }).Append('.');
+			if (!c.storeBuckets.empty())
+				log.Appendf(TC(" Stored %llu entries to bucket "), c.storeEntryCount).Join(c.storeBuckets, [&](auto& kv) { log.AppendValue(kv.second.index); });
 			m_connections.erase(it);
 		}
-		u64 activeCount = m_connections.size();
 		lock.Leave();
 
-		StringBuffer<> log;
-		log.Appendf(TC("Client %u disconnected after %s (%llu active). "), clientId, TimeToText(GetTime() - connectTime).str, activeCount);
-		if (fetchBucketsCount)
-			log.Appendf(TC("Fetched %llu (%llu hits) from %llu bucket(s)"), fetchEntryCount, fetchEntryHitCount, fetchBucketsCount);
-		if (fetchBucketsCount && storeBucketsCount)
-			log.Append(TC(", "));
-		if (storeBucketsCount)
-			log.Appendf(TC("Stored %llu to %llu bucket(s)"), storeEntryCount, storeBucketsCount);
 		m_logger.Log(LogEntryType_Info, log);
 	}
 
@@ -1370,6 +1359,7 @@ namespace uba
 
 			u64 id = connectionBucket.id;
 			Bucket& bucket = GetBucket(id);
+			connectionBucket.index = bucket.index;
 
 			SCOPED_WRITE_LOCK(bucket.m_cacheEntryLookupLock, lock3);
 			auto insres = bucket.m_cacheEntryLookup.try_emplace(cmdKey);
diff --git a/Engine/Source/Programs/UnrealBuildAccelerator/Core/Public/UbaStringBuffer.h b/Engine/Source/Programs/UnrealBuildAccelerator/Core/Public/UbaStringBuffer.h
--- a/Engine/Source/Programs/UnrealBuildAccelerator/Core/Public/UbaStringBuffer.h
+++ b/Engine/Source/Programs/UnrealBuildAccelerator/Core/Public/UbaStringBuffer.h
@@ -45,6 +45,9 @@ namespace uba
 		StringBufferBase& Append(const char* str, u32 charCount);
 		#endif
 
+		template<typename Container, typename Func>
+		StringBufferBase& Join(const Container& c, const Func& f, const tchar* separator = TC(","));
+
 		tchar operator[](u64 i) const { return data[i]; }
 		tchar& operator[](u64 i) { return data[i]; }
 
@@ -123,4 +126,18 @@ namespace uba
 		LastErrorToText(u32 lastError);
 		operator const tchar* () const { return data; };
 	};
+
+	template<typename Container, typename Func>
+	StringBufferBase& StringBufferBase::Join(const Container& c, const Func& f, const tchar* separator)
+	{
+		bool isFirst = true;
+		for (auto& e : c)
+		{
+			if (!isFirst)
+				Append(separator);
+			isFirst = false;
+			f(e);
+		}
+		return *this;
+	}
 }
\ No newline at end of file
