[UBA]
* Removed bucket version from log if version is the same as latest

diff --git a/Engine/Source/Programs/UnrealBuildAccelerator/Common/Private/UbaCacheServer.cpp b/Engine/Source/Programs/UnrealBuildAccelerator/Common/Private/UbaCacheServer.cpp
--- a/Engine/Source/Programs/UnrealBuildAccelerator/Common/Private/UbaCacheServer.cpp
+++ b/Engine/Source/Programs/UnrealBuildAccelerator/Common/Private/UbaCacheServer.cpp
@@ -491,7 +491,13 @@ namespace uba
 
 		bucket.lastSavedTime = GetSystemTimeAsFileTime() - m_creationTime;
 
-		m_logger.Detail(TC("    Bucket %u saved (v%u) - %s (%s)"), bucket.index, bucket.m_pathTable.GetVersion(), BytesToText(file.written).str, TimeToText(GetTime() - saveStart).str);
+		StringBuffer<256> log;
+		log.Appendf(TC("    Bucket %u saved"), bucket.index);
+		u32 version = bucket.m_pathTable.GetVersion();
+		if (version != CacheBucketVersion)
+			log.Appendf(TC(" (v%u)"), version);
+		log.Appendf(TC(" - %s (%s)"), BytesToText(file.written).str, TimeToText(GetTime() - saveStart).str);
+		m_logger.Log(LogEntryType_Detail, log);
 		return true;
 	}
 
