[UBA]
* Fixed bug where cache client hint cstring was used after free

diff --git a/Engine/Source/Programs/UnrealBuildAccelerator/Common/Private/UbaCacheClient.cpp b/Engine/Source/Programs/UnrealBuildAccelerator/Common/Private/UbaCacheClient.cpp
--- a/Engine/Source/Programs/UnrealBuildAccelerator/Common/Private/UbaCacheClient.cpp
+++ b/Engine/Source/Programs/UnrealBuildAccelerator/Common/Private/UbaCacheClient.cpp
@@ -80,7 +80,7 @@ namespace uba
 		m_useCacheHit = info.useCacheHit;
 		m_useRoots = info.useRoots;
 
-		m_client.RegisterOnConnected([this, hint = info.hint]()
+		m_client.RegisterOnConnected([this, hint = TString(info.hint)]()
 			{
 				u32 retryCount = 0;
 				while (retryCount < 10)
