[UBA]
* Fixed so dotnet programs can run remotely on wine in combination with vfs (dotnet can't have icu globalization enabled)

diff --git a/Engine/Source/Programs/UnrealBuildAccelerator/Cli/Private/UbaCli.cpp b/Engine/Source/Programs/UnrealBuildAccelerator/Cli/Private/UbaCli.cpp
--- a/Engine/Source/Programs/UnrealBuildAccelerator/Cli/Private/UbaCli.cpp
+++ b/Engine/Source/Programs/UnrealBuildAccelerator/Cli/Private/UbaCli.cpp
@@ -730,7 +730,10 @@ namespace uba
 					return logger.Error(TC("Failed to open file %s"), application.c_str());
 				u64 fileSize = fa.GetSize();
 				u8* data = fa.GetData();
-				bool is64Bit = true;
+				bool is64Bit = false;
+				bool isArm64 = false;
+				bool isX64 = false;
+				bool isDotnet = false;
 
 				CasKey key = CalculateCasKey(data, fileSize, false, nullptr, application.c_str());
 				CasKey uncompressedKey;
@@ -741,15 +744,23 @@ namespace uba
 						uncompressedKey = hdr.casKey;
 				}
 
-				if (data[0] != 'M' || data[1] != 'Z')
-					is64Bit = false;
-				else
+				if (data[0] == 'M' && data[1] == 'Z')
 				{
 					u32 offset = *(u32*)(data + 0x3c);
-					is64Bit = *(u32*)(data + offset) == 0x00004550;
+					u32* signaturePos = (u32*)(data + offset);
+					is64Bit = *signaturePos == 0x00004550;
+					if (is64Bit)
+					{
+						u16 machine = *(u16*)(signaturePos + 1);
+						isX64 = machine == 0x8664;
+						isArm64 = machine == 0xaa64;
+						if (fileSize > offset + 0x18 + 0x70 + 4)
+							isDotnet = *(u32*)(data + offset + 0x18 + 0x70);
+					}
 				}
 				logger.Info(TC("%s"), application.c_str());
 				logger.Info(TC("  Is64Bit: %s"), (is64Bit ? TC("true") : TC("false")));
+				logger.Info(TC("  Arch: %s"), (isX64 ? TC("x64") : (isArm64 ? TC("arm64") : (isDotnet ? TC(".net") : TC("unknown")))));
 				logger.Info(TC("  Size: %llu"), fileSize);
 				logger.Info(TC("  CasKey: %s"), CasKeyString(key).str);
 				if (uncompressedKey != CasKeyZero)
@@ -1134,7 +1145,7 @@ namespace uba
 		{
 			StringBuffer<512> sdkbin;
 			//sdkbin.count = GetEnvironmentVariable(TC("WindowsSdkVerBinPath"), sdkbin.data, sdkbin.capacity);
-			sdkbin.Append(TC("C:\\Program Files (x86)\\Windows Kits\\10\\bin\\10.0.22621.0\\x64;"));
+			sdkbin.Append(TC(";C:\\Program Files (x86)\\Windows Kits\\10\\bin\\10.0.22621.0\\x64"));
 			if (sdkbin.count)
 			{
 				StringBuffer<4096> temp;
diff --git a/Engine/Source/Programs/UnrealBuildAccelerator/Common/Private/UbaProcess.cpp b/Engine/Source/Programs/UnrealBuildAccelerator/Common/Private/UbaProcess.cpp
--- a/Engine/Source/Programs/UnrealBuildAccelerator/Common/Private/UbaProcess.cpp
+++ b/Engine/Source/Programs/UnrealBuildAccelerator/Common/Private/UbaProcess.cpp
@@ -97,6 +97,7 @@ namespace uba
 		bool is64Bit = true;
 		bool isX64 = false;
 		bool isArm64 = false;
+		bool isDotnet = false;
 	};
 
 	bool GetImageInfo(ImageInfo& out, Logger& logger, const tchar* application, bool calculateCas)
@@ -114,9 +115,7 @@ namespace uba
 			out.key = ToCasKey(hasher, false);
 		}
 
-		if (data[0] != 'M' || data[1] != 'Z')
-			out.is64Bit = false;
-		else
+		if (data[0] == 'M' && data[1] == 'Z')
 		{
 			u32 offset = *(u32*)(data + 0x3c);
 			u32* signaturePos = (u32*)(data + offset);
@@ -124,6 +123,8 @@ namespace uba
 			u16 machine = *(u16*)(signaturePos + 1);
 			out.isX64 = machine == 0x8664;
 			out.isArm64 = machine == 0xaa64;
+			if (out.fileSize > offset + 0x18 + 0x70 + 4)
+				out.isDotnet = *(u32*)(data + offset + 0x18 + 0x70);
 		}
 		return fa.Close();
 	}
@@ -233,7 +234,7 @@ namespace uba
 		// If running remote we can't use mspdbsrv (not supported yet).. so instead embed information in .obj file
 		// TODO: This should be placed somewhere else it feels like.
 		#if PLATFORM_WINDOWS
-		if (runningRemote && (m_startInfo.rules->index == 1 || m_startInfo.rules->index == 2))
+		if (runningRemote && (m_startInfo.rules->index == SpecialRulesIndex_ClExe || m_startInfo.rules->index == SpecialRulesIndex_LinkExe))
 		{
 			tchar* pos = nullptr;
 			if (Contains(m_startInfo.argumentsStr.data(), L"/FS ", true, (const tchar**)&pos))
@@ -299,6 +300,8 @@ namespace uba
 					err.Appendf(TC("Failed to load %s"), m_realApplication.c_str());
 				else if (!imageInfo.is64Bit)
 					err.Append(TC("Doesn't seem to be a 64-bit executable"));
+				else if (imageInfo.isDotnet)
+					err.Append(TC("Dotnet binary"));
 				else if (!imageInfo.isArm64 && !imageInfo.isX64)
 					err.Append(TC("Unknown image architecture"));
 				else if (!machineIsArm && imageInfo.isArm64)
diff --git a/Engine/Source/Programs/UnrealBuildAccelerator/Common/Private/UbaSession.cpp b/Engine/Source/Programs/UnrealBuildAccelerator/Common/Private/UbaSession.cpp
--- a/Engine/Source/Programs/UnrealBuildAccelerator/Common/Private/UbaSession.cpp
+++ b/Engine/Source/Programs/UnrealBuildAccelerator/Common/Private/UbaSession.cpp
@@ -1564,6 +1564,8 @@ namespace uba
 		
 		auto rules = GetApplicationRules();
 		
+		bool isDotnet = false;
+
 		while (true)
 		{
 			exeName.MakeLower();
@@ -1580,8 +1582,10 @@ namespace uba
 			}
 
 			if (!exeName.Equals(TC("dotnet.exe")))
-				return rules[0].rules;
+				return rules[isDotnet].rules;
 			
+			isDotnet = true;
+
 			u32 firstArgumentStart = 0;
 			u32 firstArgumentEnd = 0;
 			bool quoted = false;
@@ -1592,6 +1596,7 @@ namespace uba
 				{
 					if (c == '\\')
 						firstArgumentStart = i + 1;
+					firstArgumentEnd = i + 1;
 					if ((quoted && c != '"') || (!quoted && c != ' ' && c != '\t'))
 						continue;
 					firstArgumentEnd = i;
@@ -1604,7 +1609,7 @@ namespace uba
 						++firstArgumentStart;
 						continue;
 					}
-					if (c == '"')
+					else if (c == '"')
 					{
 						++firstArgumentStart;
 						quoted = true;
@@ -1614,7 +1619,7 @@ namespace uba
 			}
 			exeName.Clear().Append(si.arguments + firstArgumentStart, firstArgumentEnd - firstArgumentStart);
 		}
-		return rules[0].rules;
+		return rules[isDotnet].rules;
 	}
 
 	const tchar* Session::GetTempPath()
diff --git a/Engine/Source/Programs/UnrealBuildAccelerator/Common/Private/UbaSessionClient.cpp b/Engine/Source/Programs/UnrealBuildAccelerator/Common/Private/UbaSessionClient.cpp
--- a/Engine/Source/Programs/UnrealBuildAccelerator/Common/Private/UbaSessionClient.cpp
+++ b/Engine/Source/Programs/UnrealBuildAccelerator/Common/Private/UbaSessionClient.cpp
@@ -2223,7 +2223,7 @@ namespace uba
 			return true;
 
 		const tchar* errorPos = nullptr;
-		if (rules->index == 1)
+		if (rules->index == SpecialRulesIndex_ClExe)
 		{
 			if (!Contains(line, TC("C1083"), false, &errorPos))
 				return true;
diff --git a/Engine/Source/Programs/UnrealBuildAccelerator/Core/Private/UbaApplicationRules.cpp b/Engine/Source/Programs/UnrealBuildAccelerator/Core/Private/UbaApplicationRules.cpp
--- a/Engine/Source/Programs/UnrealBuildAccelerator/Core/Private/UbaApplicationRules.cpp
+++ b/Engine/Source/Programs/UnrealBuildAccelerator/Core/Private/UbaApplicationRules.cpp
@@ -5,6 +5,17 @@
 
 namespace uba
 {
+	class ApplicationRulesDotnet : public ApplicationRules
+	{
+		using Super = ApplicationRules;
+
+	public:
+		virtual bool CanDetour(const tchar* file, bool isRunningRemote) const override
+		{
+			return !isRunningRemote || TStrchr(file, ',') == 0;
+		}
+	};
+
 	class ApplicationRulesVC : public ApplicationRules
 	{
 		using Super = ApplicationRules;
@@ -25,7 +36,7 @@ namespace uba
 				return 1024 * 1024 * 1024;
 			return Super::FileTypeMaxSize(file, isSystemOrTempFile);
 		}
-		virtual bool CanDetour(const tchar* file) const override
+		virtual bool CanDetour(const tchar* file, bool isRunningRemote) const override
 		{
 			if (file[0] == '\\' && file[1] == '\\' && !Contains(file, TC("vctip_"))) // This might be too aggressive but will cover pipes etc.. might need revisit
 				return false;
@@ -593,10 +604,11 @@ namespace uba
 
 		static RulesRec rules[]
 		{
-			{ ApplicationHash(TC("")),							new ApplicationRules() },		 // Must be index 0
+			{ ApplicationHash(TC("")),							new ApplicationRules() },
+			{ ApplicationHash(TC("dotnet")),					new ApplicationRulesDotnet() }, // Fallback for non-specificied dotnet applications
 #if PLATFORM_WINDOWS
-			{ ApplicationHash(TC("cl.exe")),					new ApplicationRulesClExe() },	 // Must be index 1
-			{ ApplicationHash(TC("link.exe")),					new ApplicationRulesLinkExe() }, // Must be index 2
+			{ ApplicationHash(TC("cl.exe")),					new ApplicationRulesClExe() },	 // Must be index 2
+			{ ApplicationHash(TC("link.exe")),					new ApplicationRulesLinkExe() }, // Must be index 3
 			{ ApplicationHash(TC("lib.exe")),					new ApplicationRulesVcLink() },
 			{ ApplicationHash(TC("cvtres.exe")),				new ApplicationRulesLinkExe() },
 			{ ApplicationHash(TC("mt.exe")),					new ApplicationRulesVcLink() },
@@ -608,7 +620,7 @@ namespace uba
 			{ ApplicationHash(TC("ispc.exe")),					new ApplicationRulesISPCExe() },
 			{ ApplicationHash(TC("radlink.exe")),				new ApplicationRulesVcLink() },
 			{ 3340509542,										new ApplicationRulesClangPlusPlusExePlatform1() },
-			{ 4113554641,										new ApplicationRulesLdExePlatform1() },
+			{ 4113554641,										new ApplicationRulesLdExePlatform1() }, // Must be index 15
 			{ 1752955744,										new ApplicationRulesLdExePlatform1() },
 			{ 238360161,										new ApplicationRulesClangPlusPlusExePlatform2() },
 			{ 2119756440,										new ApplicationRulesLldExePlatform2() },
diff --git a/Engine/Source/Programs/UnrealBuildAccelerator/Core/Public/UbaApplicationRules.h b/Engine/Source/Programs/UnrealBuildAccelerator/Core/Public/UbaApplicationRules.h
--- a/Engine/Source/Programs/UnrealBuildAccelerator/Core/Public/UbaApplicationRules.h
+++ b/Engine/Source/Programs/UnrealBuildAccelerator/Core/Public/UbaApplicationRules.h
@@ -19,7 +19,7 @@ namespace uba
 
 		// Some kernel functions are not detoured and can't handle detoured handles (pipe system for example)
 		// This function can be used to prevent uba from detouring handles. Defaults to not detour pipes
-		virtual bool CanDetour(const tchar* file) const
+		virtual bool CanDetour(const tchar* file, bool isRunningRemote) const
 		{
 			if constexpr (IsWindows)
 			{
@@ -175,4 +175,11 @@ namespace uba
 
 	const RulesRec* GetApplicationRules();
 	u32 GetApplicationHash(const StringView& str);
+
+	enum SpecialRulesIndex
+	{
+		SpecialRulesIndex_ClExe = 2,
+		SpecialRulesIndex_LinkExe = 3,
+		SpecialRulesIndex_LdExePlatform1 = 15,
+	};
 }
diff --git a/Engine/Source/Programs/UnrealBuildAccelerator/Detours/Private/UbaDetoursShared.h b/Engine/Source/Programs/UnrealBuildAccelerator/Detours/Private/UbaDetoursShared.h
--- a/Engine/Source/Programs/UnrealBuildAccelerator/Detours/Private/UbaDetoursShared.h
+++ b/Engine/Source/Programs/UnrealBuildAccelerator/Detours/Private/UbaDetoursShared.h
@@ -121,7 +121,7 @@ namespace uba
 		u32 lastError;
 	};
 
-	inline bool CanDetour(const tchar* file) { return !t_disallowDetour && g_rules->CanDetour(file); }
+	inline bool CanDetour(const tchar* file) { return !t_disallowDetour && g_rules->CanDetour(file, g_runningRemote); }
 	inline bool KeepInMemory(const StringView& fileName, bool isWrite) { return g_allowKeepFilesInMemory && g_rules->KeepInMemory(fileName, g_systemTemp.data, g_runningRemote, isWrite); }
 
 	void Shared_WriteConsole(const char* chars, u32 charCount, bool isError);
diff --git a/Engine/Source/Programs/UnrealBuildAccelerator/Detours/Private/Windows/UbaDetoursFunctionsKernelBase.inl b/Engine/Source/Programs/UnrealBuildAccelerator/Detours/Private/Windows/UbaDetoursFunctionsKernelBase.inl
--- a/Engine/Source/Programs/UnrealBuildAccelerator/Detours/Private/Windows/UbaDetoursFunctionsKernelBase.inl
+++ b/Engine/Source/Programs/UnrealBuildAccelerator/Detours/Private/Windows/UbaDetoursFunctionsKernelBase.inl
@@ -1047,7 +1047,7 @@ BOOL Detoured_GetFileAttributesExW(LPCWSTR lpFileName, GET_FILEEX_INFO_LEVELS fI
 	DETOURED_CALL(GetFileAttributesExW);
 	if (t_disallowDetour != 0 || !lpFileName || wcschr(lpFileName, ',') != nullptr || Contains(lpFileName, L"::")) // Some weird .net path used by dotnet.exe ... ignore for now!
 	{
-		UBA_ASSERT(!g_runningRemote);
+		//UBA_ASSERTF(!g_runningRemote, TC("GetFileAttributesExW can't handle this path when running remote: %s"), lpFileName);
 		DEBUG_LOG_TRUE(L"GetFileAttributesExW", L"(%ls)", lpFileName);
 		TimerScope ts(g_kernelStats.getFileInfo);
 		return True_GetFileAttributesExW(lpFileName, fInfoLevelId, lpFileInformation);
@@ -1277,6 +1277,18 @@ DWORD Detoured_GetModuleFileNameExW(HANDLE hProcess, HMODULE hModule, LPWSTR lpF
 	return Detoured_GetModuleFileNameW(hModule, lpFilename, nSize);
 }
 
+BOOL Detoured_GetModuleHandleExW(DWORD dwFlags, LPCWSTR lpModuleName, HMODULE *phModule)
+{
+	StringBuffer<> path;
+	if (lpModuleName && IsAbsolutePath(lpModuleName))
+	{
+		FixPath(path, lpModuleName);
+		DevirtualizePath(path);
+		lpModuleName = path.data;
+	}
+	return True_GetModuleHandleExW(dwFlags, lpModuleName, phModule);
+}
+
 BOOL Detoured_CopyFileExW(LPCWSTR lpExistingFileName, LPCWSTR lpNewFileName, LPPROGRESS_ROUTINE lpProgressRoutine, LPVOID lpData, LPBOOL pbCancel, DWORD dwCopyFlags)
 {
 	DETOURED_CALL(CopyFileExW);
@@ -2641,6 +2653,16 @@ HMODULE Detoured_LoadLibraryExW(LPCWSTR lpLibFileName, HANDLE hFile, DWORD dwFla
 	if (!g_rules->AllowLoadLibrary(lpLibFileName))
 		return 0;
 
+	if (IsKnownSystemFile(lpLibFileName))
+	{
+		SuppressCreateFileDetourScope cfs;
+		const tchar* fileName = lpLibFileName;
+		if (auto lastSeparator = TStrrchr(fileName, PathSeparator))
+			fileName = lastSeparator + 1;
+		DEBUG_LOG_TRUE(L"LoadLibraryExW", L"%ls (%s)", lpLibFileName, fileName);
+		return True_LoadLibraryExW(fileName, hFile, dwFlags);
+	}
+
 	if (HMODULE res = Shared_LoadLibrary(lpLibFileName, hFile, dwFlags))
 		return res;
 
diff --git a/Engine/Source/Programs/UnrealBuildAccelerator/Detours/Private/Windows/UbaDetoursFunctionsWin.cpp b/Engine/Source/Programs/UnrealBuildAccelerator/Detours/Private/Windows/UbaDetoursFunctionsWin.cpp
--- a/Engine/Source/Programs/UnrealBuildAccelerator/Detours/Private/Windows/UbaDetoursFunctionsWin.cpp
+++ b/Engine/Source/Programs/UnrealBuildAccelerator/Detours/Private/Windows/UbaDetoursFunctionsWin.cpp
@@ -793,7 +793,7 @@ void PreInit(const DetoursPayload& payload)
 	}
 
 	// Special cl.exe handling..  this is needed for compiles using pch files where this address _must_ be available.
-	if (payload.rulesIndex == 1) // cl.exe
+	if (payload.rulesIndex == SpecialRulesIndex_ClExe)
 	{
 		g_clExeBaseReservedMemory = VirtualAlloc((void*)g_clExeBaseAddress, g_clExeBaseAddressSize, MEM_RESERVE, PAGE_READWRITE);
 		DEBUG_LOG(L"Reserving %llu bytes at 0x%llx for cl.exe", g_clExeBaseAddressSize, g_clExeBaseAddress);
@@ -943,6 +943,12 @@ void Init(const DetoursPayload& payload, u64 startTime)
 		DEBUG_LOG(L"ExeDir: %ls", g_virtualApplicationDir.data);
 		DEBUG_LOG(L"ExeDir (true): %ls", g_exeDir.data);
 		DEBUG_LOG(L"Rules: %u (%u)", g_rules->index, GetApplicationRules()[g_rules->index].hash);
+		if (g_runningRemote)
+		{
+			StringBuffer<256> computerName;
+			GetComputerNameW(computerName);
+			DEBUG_LOG(L"Remote: %s", computerName.data);
+		}
 		LogVfsInfo();
 	}
 	#endif
@@ -983,7 +989,7 @@ void Init(const DetoursPayload& payload, u64 startTime)
 
 	if (payload.storeObjFilesCompressed && g_rules->ShouldDecompressFiles(StringView()))
 	{
-		bool preprocess = g_rules->index == 14;
+		bool preprocess = g_rules->index == SpecialRulesIndex_LdExePlatform1;
 		g_fileDecompressor.Start(cmdLine, preprocess);
 	}
 }
diff --git a/Engine/Source/Programs/UnrealBuildAccelerator/Detours/Private/Windows/UbaDetoursFunctionsWin.h b/Engine/Source/Programs/UnrealBuildAccelerator/Detours/Private/Windows/UbaDetoursFunctionsWin.h
--- a/Engine/Source/Programs/UnrealBuildAccelerator/Detours/Private/Windows/UbaDetoursFunctionsWin.h
+++ b/Engine/Source/Programs/UnrealBuildAccelerator/Detours/Private/Windows/UbaDetoursFunctionsWin.h
@@ -68,6 +68,7 @@
 	DETOURED_FUNCTION(GetVolumePathNameW) \
 	DETOURED_FUNCTION(GetModuleFileNameW) \
 	DETOURED_FUNCTION(GetModuleFileNameExW) \
+	DETOURED_FUNCTION(GetModuleHandleExW) \
 	DETOURED_FUNCTION(GetFileAttributesW) \
 	DETOURED_FUNCTION(SetFileAttributesW) \
 	DETOURED_FUNCTION(GetFileAttributesExW) \
