[UBA]
* CacheServer - Added additional table per cache entry that contains inputs that are outputs from other cached steps. This should speed up maintenance a lot because when cas files are deleted we have to traverse all inputs to see if the deleted cas files are inputs to cache entries.. but 99.9% of all cache entry inputs are not outputs of other cached steps.

diff --git a/Engine/Source/Programs/UnrealBuildAccelerator/Common/Private/UbaCacheEntry.cpp b/Engine/Source/Programs/UnrealBuildAccelerator/Common/Private/UbaCacheEntry.cpp
--- a/Engine/Source/Programs/UnrealBuildAccelerator/Common/Private/UbaCacheEntry.cpp
+++ b/Engine/Source/Programs/UnrealBuildAccelerator/Common/Private/UbaCacheEntry.cpp
@@ -1,8 +1,10 @@
 // Copyright Epic Games, Inc. All Rights Reserved.
 
 #include "UbaCacheEntry.h"
+#include "UbaCompactTables.h"
 #include "UbaFile.h"
 #include "UbaHashMap.h"
+#include "UbaStorage.h"
 #include <algorithm>
 
 namespace uba
@@ -49,7 +51,7 @@ namespace uba
 		for (auto& entry : entries)
 			size += GetEntrySize(entry, clientVersion, toDisk);
 		if (toDisk)
-			size += sizeof(hasInputsThatAreOutputs);
+			size += sizeof(u32) + inputsThatAreOutputs.size()*sizeof(u32);
 		return size;
 	}
 
@@ -154,12 +156,15 @@ namespace uba
 		}
 
 		if (toDisk)
-			writer.WriteBool(hasInputsThatAreOutputs);
-
+		{
+			writer.WriteU32(u32(inputsThatAreOutputs.size()));
+			for (u32 offset : inputsThatAreOutputs)
+				writer.WriteU32(offset);
+		}
 		return true;
 	}
 
-	bool CacheEntries::Read(Logger& logger, BinaryReader& reader, u32 databaseVersion)
+	bool CacheEntries::ReadFromDisk(Logger& logger, BinaryReader& reader, u32 databaseVersion, StorageImpl& storage, CompactCasKeyTable& table)
 	{
 		if (databaseVersion == 3)
 		{
@@ -233,8 +238,22 @@ namespace uba
 
 		}
 
-		if (databaseVersion >= 7)
-			hasInputsThatAreOutputs = reader.ReadBool();
+		if (databaseVersion < 8)
+		{
+			if (databaseVersion == 7)
+				reader.ReadBool();
+			PopulateInputsThatAreOutputs(sharedInputCasKeyOffsets, storage, table);
+			for (auto& entry : entries)
+				PopulateInputsThatAreOutputs(entry.extraInputCasKeyOffsets, storage, table);
+		}
+
+		if (databaseVersion >= 8)
+		{
+			u32 count = reader.ReadU32();
+			inputsThatAreOutputs.reserve(count);
+			while (count--)
+				inputsThatAreOutputs.insert(reader.ReadU32());
+		}
 
 		return true;
 	}
@@ -449,6 +468,13 @@ namespace uba
 		if (entries.empty())
 			return;
 
+		UnorderedSet<u32> itao(inputsThatAreOutputs.size());
+		for (u32 offset : inputsThatAreOutputs)
+			if (auto o = oldToNewCasKeyOffset.Find(offset))
+				itao.insert(*o);
+		inputsThatAreOutputs.swap(itao);
+
+
 		auto convertOffsets = [&](Vector<u8>& offsets)
 			{
 				temp.clear();
@@ -620,6 +646,32 @@ namespace uba
 		UBA_ASSERT(false);
 	}
 
+	void CacheEntries::PopulateInputsThatAreOutputs(const Vector<u8>& inputData, StorageImpl& storage, CompactCasKeyTable& table)
+	{
+		// one entry that is ~0u is equal to test all inputs
+		if (inputsThatAreOutputs.size() == 1 && *inputsThatAreOutputs.begin() == ~0u)
+			return;
+
+		u32 insertCount = 0;
+		BinaryReader inputReader(inputData.data(), 0, inputData.size());
+		while (inputReader.GetLeft())
+		{
+			u32 casKeyOffset = u32(inputReader.Read7BitEncoded());
+			CasKey casKey;
+			table.GetKey(casKey, casKeyOffset);
+			if (storage.HasBeenSeen(casKey))
+				if (inputsThatAreOutputs.insert(casKeyOffset).second)
+					++insertCount;
+		}
+
+		// If there are lots of inputs that are outputs this is most likely a link step and we switch to always check inputs
+		if (insertCount > 5)
+		{
+			inputsThatAreOutputs.clear();
+			inputsThatAreOutputs.insert(~0u);
+		}
+	}
+
 	void CacheEntries::Flatten(Vector<u8>& out, const CacheEntry& entry)
 	{
 		u64 size = entry.extraInputCasKeyOffsets.size();
diff --git a/Engine/Source/Programs/UnrealBuildAccelerator/Common/Private/UbaCacheServer.cpp b/Engine/Source/Programs/UnrealBuildAccelerator/Common/Private/UbaCacheServer.cpp
--- a/Engine/Source/Programs/UnrealBuildAccelerator/Common/Private/UbaCacheServer.cpp
+++ b/Engine/Source/Programs/UnrealBuildAccelerator/Common/Private/UbaCacheServer.cpp
@@ -27,7 +27,7 @@
 
 namespace uba
 {
-	static constexpr u32 CacheFileVersion = 7;
+	static constexpr u32 CacheFileVersion = 8;
 	static constexpr u32 CacheFileCompatibilityVersion = 3;
 
 	bool IsCaseInsensitive(u64 id) { return (id & (1ull << 32)) == 0; }
@@ -342,15 +342,8 @@ namespace uba
 			auto insres = bucket.m_cacheEntryLookup.try_emplace(reader.ReadCasKey());
 			UBA_ASSERT(insres.second);
 			auto& cacheEntries = insres.first->second;
-			cacheEntries.Read(m_logger, reader, databaseVersion);
+			cacheEntries.ReadFromDisk(m_logger, reader, databaseVersion, m_storage, bucket.m_casKeyTable);
 			outStats.totalCacheEntryCount += cacheEntries.entries.size();
-
-			if (databaseVersion < 7)
-			{
-				cacheEntries.hasInputsThatAreOutputs |= CheckIfInputsAreOutputs(bucket.m_casKeyTable, cacheEntries.sharedInputCasKeyOffsets);
-				for (auto& entry : cacheEntries.entries)
-					cacheEntries.hasInputsThatAreOutputs |= CheckIfInputsAreOutputs(bucket.m_casKeyTable, entry.extraInputCasKeyOffsets);
-			}
 		}
 		return true;
 	}
@@ -662,12 +655,7 @@ namespace uba
 
 		struct CasFileInfo { CasFileInfo(u32 s = 0) : size(s) {} u32 size; Atomic<bool> isUsed; }; // These are compressed cas, should never be over 4gb
 
-		// Existing cas entries can be more than 2 million entries.. which uses a lot of memory
-		constexpr u64 existingCasMemoryReserveSize = 512*1024*1024;
 		MemoryBlock existingCasMemoryBlock;
-		if (!existingCasMemoryBlock.Init(existingCasMemoryReserveSize, nullptr, true))
-			existingCasMemoryBlock.Init(existingCasMemoryReserveSize);
-
 		HashMap<CasKey, CasFileInfo> existingCas;
 
 
@@ -683,6 +671,12 @@ namespace uba
 			SCOPED_WRITE_LOCK(m_storage.m_casLookupLock, lookupLock);
 			
 			totalCasCount = m_storage.m_casLookup.size();
+
+			// Existing cas entries can be more than 2 million entries.. which uses a lot of memory
+			u64 existingCasMemoryReserveSize = existingCas.GetMemoryNeeded(totalCasCount);
+			if (!existingCasMemoryBlock.Init(existingCasMemoryReserveSize, nullptr, true))
+				existingCasMemoryBlock.Init(existingCasMemoryReserveSize);
+
 			existingCas.Init(existingCasMemoryBlock, totalCasCount, TC("ExistingCas"));
 
 			for (auto i=m_storage.m_casLookup.begin(), e=m_storage.m_casLookup.end(); i!=e;)
@@ -826,7 +820,27 @@ namespace uba
 				{
 					CacheEntries& entries = li->second;
 
-					bool checkInputsForDeletes = foundDeletedCasKey && m_checkInputsForDeletedCas && entries.hasInputsThatAreOutputs;
+					bool checkInputsForDeletes = false;
+
+					// Check the inputsThatAreOutputs if any of those are deleted, if they are, then we need to checkInputsForDeletes
+					// This is an optimization since most build steps only have 1-2 inputs that are outputs
+					if (foundDeletedCasKey && m_checkInputsForDeletedCas && !entries.inputsThatAreOutputs.empty())
+					{
+						checkInputsForDeletes = *entries.inputsThatAreOutputs.begin() == ~0u;
+						if (!checkInputsForDeletes)
+						{
+							for (auto i=entries.inputsThatAreOutputs.begin(), e=entries.inputsThatAreOutputs.end(); i!=e;)
+							{
+								if (!deletedOffsets.IsSet(*i))
+								{
+									++i;
+									continue;
+								}
+								i = entries.inputsThatAreOutputs.erase(i);
+								checkInputsForDeletes = true;
+							}
+						}
+					}
 
 					// There is currently no idea keeping more than 256kb worth of entries per lookup key (because that is what fetch max returns).. so let's wipe out
 					// all the entries that overflow that number
@@ -851,19 +865,26 @@ namespace uba
 					{
 						auto& entry = *i;
 						bool deleteEntry = false;
+						u64 neededSize = 0;
 
-						u64 neededSize = entries.GetEntrySize(entry, CacheNetworkVersion, false);
-						if (neededSize > capacityLeft)
+						// Only need to check these things on first iteration.
+						if (deleteIteration == 0)
 						{
-							deleteEntry = true;
-							capacityLeft = 0;
-							++overflowedEntryCount;
-						}
+							// Check if this entry is outside capacity of network message
+							neededSize = entries.GetEntrySize(entry, CacheNetworkVersion, false);
+							if (neededSize > capacityLeft)
+							{
+								deleteEntry = true;
+								capacityLeft = 0;
+								++overflowedEntryCount;
+							}
 
-						if (!deleteEntry && entry.creationTime < lastUseTimeLimit && entry.lastUsedTime < lastUseTimeLimit)
-						{
-							deleteEntry = true;
-							++expiredEntryCount;
+							// Check if entry has expired
+							if (!deleteEntry && entry.creationTime < lastUseTimeLimit && entry.lastUsedTime < lastUseTimeLimit)
+							{
+								deleteEntry = true;
+								++expiredEntryCount;
+							}
 						}
 
 						// This is an attempt at removing entries that has inputs that depends on other entries outputs.
@@ -1013,15 +1034,15 @@ namespace uba
 		existingCasMemoryBlock.Deinit();
 
 		if (overflowedEntryCount)
-			m_logger.Detail(TC("  Found %llu overflowed cache entries"), overflowedEntryCount.load());
+			m_logger.Detail(TC("  Found %s overflowed cache entries"), CountToText(overflowedEntryCount.load()).str);
 		if (expiredEntryCount)
-			m_logger.Detail(TC("  Found %llu expired cache entries"), expiredEntryCount.load());
+			m_logger.Detail(TC("  Found %s expired cache entries"), CountToText(expiredEntryCount.load()).str);
 		if (missingOutputEntryCount)
-			m_logger.Detail(TC("  Found %llu cache entries with missing output cas"), missingOutputEntryCount.load());
+			m_logger.Detail(TC("  Found %s cache entries with missing output cas"), CountToText(missingOutputEntryCount.load()).str);
 		if (missingInputEntryCount)
-			m_logger.Detail(TC("  Found %llu cache entries with missing input cas"), missingInputEntryCount.load());
+			m_logger.Detail(TC("  Found %s cache entries with missing input cas"), CountToText(missingInputEntryCount.load()).str);
 
-		m_logger.Detail(TC("  Deleted %llu cas files and %llu cache entries over %u buckets (%s)"), deletedCasCount, deleteEntryCount.load(), u32(m_buckets.size()), TimeToText(GetTime() - deleteCacheEntriesStartTime).str);
+		m_logger.Detail(TC("  Deleted %s cas files and %s cache entries over %u buckets (%s)"), CountToText(deletedCasCount).str, CountToText(deleteEntryCount.load()).str, u32(m_buckets.size()), TimeToText(GetTime() - deleteCacheEntriesStartTime).str);
 
 		m_bucketIsOverflowing = false;
 
@@ -1209,7 +1230,7 @@ namespace uba
 		u64 oldestTime = oldest ? GetFileTimeAsTime(now - (m_creationTime + oldest)) : 0;
 		u64 longestUnused = oldestUsedTime ? GetFileTimeAsTime(now - (m_creationTime + oldestUsedTime)) : 0;
 		u64 duration = GetTime() - startTime;
-		m_logger.Info(TC("Maintenance done! (%s) CasFiles: %s (%s) Buckets: %llu Entries: %s Oldest: %s LongestUnused: %s MaintenanceMem: %s/%s"), TimeToText(duration).str, CountToText(totalCasCount - deletedCasCount).str, BytesToText(totalCasSize).str, m_buckets.size(), CountToText(totalEntryCount.load()).str, TimeToText(oldestTime, true).str, TimeToText(longestUnused, true).str, BytesToText(maxCommittedMemory).str, BytesToText(m_maintenanceReserveSize).str);
+		m_logger.Info(TC("Maintenance done! (%s) CasFiles: %s (%s) Buckets: %llu Entries: %s Oldest: %s LongestUnused: %s MaxMaintenanceMem: %s/%s"), TimeToText(duration).str, CountToText(totalCasCount - deletedCasCount).str, BytesToText(totalCasSize).str, m_buckets.size(), CountToText(totalEntryCount.load()).str, TimeToText(oldestTime, true).str, TimeToText(longestUnused, true).str, BytesToText(maxCommittedMemory).str, BytesToText(m_maintenanceReserveSize).str);
 		
 		m_longestMaintenance = Max(m_longestMaintenance, duration);
 
@@ -1614,8 +1635,8 @@ namespace uba
 		// We never have to check inputs against deleted cas files if they never existed
 		{
 			if (cacheEntries.entries.empty())
-				cacheEntries.hasInputsThatAreOutputs |= CheckIfInputsAreOutputs(bucket.m_casKeyTable, cacheEntries.sharedInputCasKeyOffsets);
-			cacheEntries.hasInputsThatAreOutputs |= CheckIfInputsAreOutputs(bucket.m_casKeyTable, newEntry.extraInputCasKeyOffsets);
+				cacheEntries.PopulateInputsThatAreOutputs(cacheEntries.sharedInputCasKeyOffsets, m_storage, bucket.m_casKeyTable);
+			cacheEntries.PopulateInputsThatAreOutputs(newEntry.extraInputCasKeyOffsets, m_storage, bucket.m_casKeyTable);
 		}
 
 		// If cache server has all content we can put the new cache entry directly in the lookup.. otherwise we'll have to wait until client has uploaded content
@@ -1881,7 +1902,7 @@ namespace uba
 						for (auto& entry : entries.entries)
 							lastUsed = Max(lastUsed, entry.lastUsedTime);
 						totalEntryCount += entries.entries.size();
-						if (entries.hasInputsThatAreOutputs)
+						if (!entries.inputsThatAreOutputs.empty())
 							++totalCountHasOutputAsInput;
 					}
 					lock2.Leave();
@@ -1961,18 +1982,4 @@ namespace uba
 		DeleteFileW(tempFile.data);
 		return true;
 	}
-
-	bool CacheServer::CheckIfInputsAreOutputs(CompactCasKeyTable& table, const Vector<u8>& inputData)
-	{
-		bool inputsAreOutputs = false;
-		BinaryReader inputReader(inputData.data(), 0, inputData.size());
-		while (inputReader.GetLeft())
-		{
-			u32 casKeyOffset = u32(inputReader.Read7BitEncoded());
-			CasKey casKey;
-			table.GetKey(casKey, casKeyOffset);
-			inputsAreOutputs |= m_storage.HasBeenSeen(casKey);
-		}
-		return inputsAreOutputs;
-	}
 }
diff --git a/Engine/Source/Programs/UnrealBuildAccelerator/Common/Public/UbaCacheEntry.h b/Engine/Source/Programs/UnrealBuildAccelerator/Common/Public/UbaCacheEntry.h
--- a/Engine/Source/Programs/UnrealBuildAccelerator/Common/Public/UbaCacheEntry.h
+++ b/Engine/Source/Programs/UnrealBuildAccelerator/Common/Public/UbaCacheEntry.h
@@ -8,6 +8,8 @@
 
 namespace uba
 {
+	class CompactCasKeyTable;
+	class StorageImpl;
 	template<typename Key, typename Value> struct HashMap2;
 
 	enum LogLinesType : u8
@@ -40,13 +42,13 @@ namespace uba
 		Vector<u8> sharedLogLines;
 		u32 idCounter = 0;
 		u32 primaryId = ~0u; // Id of entry that shared offsets was made from
-		bool hasInputsThatAreOutputs = false;
+		UnorderedSet<u32> inputsThatAreOutputs;
 
 		u64 GetSharedSize();
 		u64 GetEntrySize(CacheEntry& entry, u32 clientVersion, bool toDisk);
 		u64 GetTotalSize(u32 clientVersion, bool toDisk);
 		bool Write(BinaryWriter& writer, u32 clientVersion, bool toDisk);
-		bool Read(Logger& logger, BinaryReader& reader, u32 databaseVersion);
+		bool ReadFromDisk(Logger& logger, BinaryReader& reader, u32 databaseVersion, StorageImpl& storage, CompactCasKeyTable& table);
 		void BuildInputs(CacheEntry& entry, const Set<u32>& inputs);
 		void UpdateEntries(Logger& logger, const HashMap2<u32, u32>& oldToNewCasKeyOffset, Vector<u32>& temp, Vector<u8>& temp2);
 
@@ -60,6 +62,8 @@ namespace uba
 		void BuildRangesFromExcludedT(CacheEntry& entry, const Container& sortedExcludedInputs);
 
 		void ValidateEntry(Logger& logger, CacheEntry& entry, Vector<u8>& inputCasKeyOffsets);
+
+		void PopulateInputsThatAreOutputs(const Vector<u8>& inputData, StorageImpl& storage, CompactCasKeyTable& table);
 	};
 
 
diff --git a/Engine/Source/Programs/UnrealBuildAccelerator/Common/Public/UbaCacheServer.h b/Engine/Source/Programs/UnrealBuildAccelerator/Common/Public/UbaCacheServer.h
--- a/Engine/Source/Programs/UnrealBuildAccelerator/Common/Public/UbaCacheServer.h
+++ b/Engine/Source/Programs/UnrealBuildAccelerator/Common/Public/UbaCacheServer.h
@@ -37,7 +37,7 @@ namespace uba
 		// Set to zero to never expire
 		u64 expirationTimeSeconds = 2*24*60*60;
 
-		// The amount of reserved memory used per core when doing maintenance
+		// The amount of reserved memory used per bucket when doing maintenance
 		u64 maintenanceReserveSize = 128ull * 1024 * 1024;
 
 		// Max size of cas bucket. When within 2mb it will start decreasing expiry time by one hour
@@ -85,8 +85,6 @@ namespace uba
 		bool HandleReportUsedEntry(BinaryReader& reader, BinaryWriter& writer, u32 clientId);
 		bool HandleExecuteCommand(BinaryReader& reader, BinaryWriter& writer);
 
-		bool CheckIfInputsAreOutputs(CompactCasKeyTable& table, const Vector<u8>& inputData);
-
 		MutableLogger m_logger;
 		NetworkServer& m_server;
 		StorageServer& m_storage;
diff --git a/Engine/Source/Programs/UnrealBuildAccelerator/Common/Public/UbaHashMap.h b/Engine/Source/Programs/UnrealBuildAccelerator/Common/Public/UbaHashMap.h
--- a/Engine/Source/Programs/UnrealBuildAccelerator/Common/Public/UbaHashMap.h
+++ b/Engine/Source/Programs/UnrealBuildAccelerator/Common/Public/UbaHashMap.h
@@ -201,6 +201,13 @@ namespace uba
 			memcpy(&o, temp, sizeof(HashMap));
 		}
 
+		u64 GetMemoryNeeded(u64 reserveCount)
+		{
+			reserveCount = Max(reserveCount + 1, 4ull);
+			u64 lookupCount = RoundUpPow2(reserveCount);
+			return lookupCount*sizeof(u32) + reserveCount*sizeof(Entry) + reserveCount*sizeof(Value) + 256; // Add some padding for alignment between tables
+		}
+
 private:
 		void InternalInit(u64 reserveCount, u64 lookupCount, const tchar* hint)
 		{
