[UBA]
* Added cache server command "largestentry" that can be used to query cache server for its largest cache entry. (it is quite crazy that some actions have 5000 inputs (all headers)
* Added UbaCacheService -nomaintenance that can be used to disable all maintenance work

diff --git a/Engine/Source/Programs/UnrealBuildAccelerator/CacheService/Private/UbaCacheService.cpp b/Engine/Source/Programs/UnrealBuildAccelerator/CacheService/Private/UbaCacheService.cpp
--- a/Engine/Source/Programs/UnrealBuildAccelerator/CacheService/Private/UbaCacheService.cpp
+++ b/Engine/Source/Programs/UnrealBuildAccelerator/CacheService/Private/UbaCacheService.cpp
@@ -50,6 +50,7 @@ namespace uba
 		logger.Info(TC("  -expiration=<seconds>   Time until unused cache entries get deleted. Defaults to %s (%u seconds)"), TimeToText(MsToTime(DefaultExpiration*1000)).str, DefaultExpiration);
 		logger.Info(TC("  -http=<port>            If set, a http server will be started and listen on <port>"));
 		logger.Info(TC("  -fullmaintenance        Force a full maintenance"));
+		logger.Info(TC("  -nomaintenance          Skip all maintenance"));
 
 		logger.Info(TC(" Example of how to register crypto key to cache server (when -http=80 is provided)"));
 		logger.Info(TC("   curl http://localhost/addcrypto?3f58aa57466db9999213456789123445"));
@@ -105,6 +106,7 @@ namespace uba
 		bool quiet = false;
 		bool storeCompressed = true;
 		bool fullMaintenance = false;
+		bool maintenanceEnabled = true;
 		u32 expirationTimeSeconds = DefaultExpiration;
 
 		for (int i=1; i!=argc; ++i)
@@ -163,6 +165,10 @@ namespace uba
 			{
 				fullMaintenance = true;
 			}
+			else if (name.Equals(TC("-nomaintenance")))
+			{
+				maintenanceEnabled = false;
+			}
 			else if (name.Equals(TC("-?")))
 			{
 				return PrintHelp(TC(""));
@@ -250,8 +256,9 @@ namespace uba
 		if (fullMaintenance)
 			cacheServer.SetForceFullMaintenance();
 
-		if (!cacheServer.RunMaintenance(true, ShouldExit))
-			return -1;
+		if (maintenanceEnabled)
+			if (!cacheServer.RunMaintenance(true, ShouldExit))
+				return -1;
 
 		HttpServer httpServer(logWriter, networkBackend);
 
@@ -289,8 +296,9 @@ namespace uba
 			while (!ShouldExit() && !cacheServer.ShouldShutdown())
 			{
 				Sleep(1000);
-				if (!cacheServer.RunMaintenance(false, ShouldExit))
-					break;
+				if (maintenanceEnabled)
+					if (!cacheServer.RunMaintenance(false, ShouldExit))
+						break;
 			}
 		}
 
diff --git a/Engine/Source/Programs/UnrealBuildAccelerator/Common/Private/UbaCacheServer.cpp b/Engine/Source/Programs/UnrealBuildAccelerator/Common/Private/UbaCacheServer.cpp
--- a/Engine/Source/Programs/UnrealBuildAccelerator/Common/Private/UbaCacheServer.cpp
+++ b/Engine/Source/Programs/UnrealBuildAccelerator/Common/Private/UbaCacheServer.cpp
@@ -1785,6 +1785,19 @@ namespace uba
 
 		StringBuffer<> line;
 
+		auto writePathsFromOffsets = [&](Bucket& bucket, const Vector<u8>& offsets)
+			{
+				BinaryReader reader2(offsets.data(), 0, offsets.size());
+				while (reader2.GetLeft())
+				{
+					u64 offset = reader2.Read7BitEncoded();
+					CasKey casKey;
+					StringBuffer<> path;
+					bucket.m_casKeyTable.GetPathAndKey(path, casKey, bucket.m_pathTable, offset);
+					writeLine(line.Clear().Appendf(TC("    %s - %s"), path.data, CasKeyString(casKey).str));
+				}
+			};
+
 		if (command.Equals(TC("content")))
 		{
 			writeLine(AsView(TC("UbaCache server summary")));
@@ -1852,24 +1865,11 @@ namespace uba
 						u64 age = GetFileTimeAsTime(now - (m_creationTime + entry.creationTime));
 						writeLine(line.Clear().Appendf(TC("  #%u (%s ago)"), index, TimeToText(age, true).str));
 
-						auto writeOffsets = [&](const Vector<u8>& offsets)
-							{
-								BinaryReader reader2(offsets.data(), 0, offsets.size());
-								while (reader2.GetLeft())
-								{
-									u64 offset = reader2.Read7BitEncoded();
-									CasKey casKey;
-									StringBuffer<> path;
-									bucket.m_casKeyTable.GetPathAndKey(path, casKey, bucket.m_pathTable, offset);
-									writeLine(line.Clear().Appendf(TC("    %s - %s"), path.data, CasKeyString(casKey).str));
-								}
-							};
-
 						writeLine(line.Clear().Append(TC("   Inputs:")));
 						entries.Flatten(temp, entry);
-						writeOffsets(temp);
+						writePathsFromOffsets(bucket, temp);
 						writeLine(line.Clear().Append(TC("   Outputs:")));
-						writeOffsets(entry.outputCasKeyOffsets);
+						writePathsFromOffsets(bucket, entry.outputCasKeyOffsets);
 						++index;
 					}
 				}
@@ -1929,25 +1929,52 @@ namespace uba
 		}
 		else if (command.Equals(TC("pathtable")))
 		{
-				SCOPED_READ_LOCK(m_bucketsLock, bucketsLock);
-				for (auto& kv : m_buckets)
-				{
-					Bucket& bucket = kv.second;
-					SCOPED_READ_LOCK(bucket.m_cacheEntryLookupLock, lock2);
+			SCOPED_READ_LOCK(m_bucketsLock, bucketsLock);
+			for (auto& kv : m_buckets)
+			{
+				Bucket& bucket = kv.second;
+				SCOPED_READ_LOCK(bucket.m_cacheEntryLookupLock, lock2);
 
-					writeLine(line.Clear().Appendf(TC("Bucket #%u"), bucket.index));
+				writeLine(line.Clear().Appendf(TC("Bucket #%u"), bucket.index));
 
-					Set<TString> sortedPaths;
+				Set<TString> sortedPaths;
 
-					bucket.m_pathTable.TraversePaths([&](const StringView& path)
-						{
-							if (!sortedPaths.emplace(path.data).second)
-								writeLine(AsView(TC("EEEEERRRRRROOOORRRR!!!!!!")));
-						});
+				bucket.m_pathTable.TraversePaths([&](const StringView& path)
+					{
+						if (!sortedPaths.emplace(path.data).second)
+							writeLine(AsView(TC("EEEEERRRRRROOOORRRR!!!!!!")));
+					});
 
-					for (auto& path : sortedPaths)
-						writeLine(line.Clear().Append(TC("   ")).Append(path));
+				for (auto& path : sortedPaths)
+					writeLine(line.Clear().Append(TC("   ")).Append(path));
+			}
+		}
+		else if (command.Equals(TC("largestentry")))
+		{
+			CacheEntries* largestEntries = nullptr;
+			Bucket* largestEntriesBucket = nullptr;
+			u64 largestSize = 0;
+			SCOPED_READ_LOCK(m_bucketsLock, bucketsLock);
+			for (auto& kv : m_buckets)
+			{
+				Bucket& bucket = kv.second;
+				for (auto& kv2 : bucket.m_cacheEntryLookup)
+				{
+					CacheEntries& entries = kv2.second;
+					SCOPED_READ_LOCK(entries.lock, lock3);
+					if (entries.sharedInputCasKeyOffsets.size() <= largestSize)
+						continue;
+					largestSize = entries.sharedInputCasKeyOffsets.size();
+					largestEntries = &entries;
+					largestEntriesBucket = &bucket;
 				}
+			}
+			if (!largestEntries)
+				return true;
+			CacheEntries& entries = *largestEntries;
+			SCOPED_READ_LOCK(entries.lock, lock3);
+			writeLine(line.Clear().Append(TC("   SharedInputs:")));
+			writePathsFromOffsets(*largestEntriesBucket, entries.sharedInputCasKeyOffsets);
 		}
 		else if (command.Equals(TC("obliterate")))
 		{
