[UBA]
* CacheServer - reduced contention in maintenance step
* Improved some logging

diff --git a/Engine/Source/Programs/UnrealBuildAccelerator/Common/Private/UbaCacheServer.cpp b/Engine/Source/Programs/UnrealBuildAccelerator/Common/Private/UbaCacheServer.cpp
--- a/Engine/Source/Programs/UnrealBuildAccelerator/Common/Private/UbaCacheServer.cpp
+++ b/Engine/Source/Programs/UnrealBuildAccelerator/Common/Private/UbaCacheServer.cpp
@@ -154,8 +154,8 @@ namespace uba
 		CompactPathTable m_pathTable;
 		CompactCasKeyTable m_casKeyTable;
 
-		Atomic<u64> totalEntryCount;
-		Atomic<u64> totalEntrySize;
+		u64 totalEntryCount;
+		u64 totalEntrySize;
 		Atomic<bool> hasDeletedEntries;
 		Atomic<bool> needsSave;
 
@@ -779,7 +779,7 @@ namespace uba
 
 				auto& deletedOffsets = context->deletedOffsets;
 
-				ReaderWriterLock keysToEraseLock;
+				ReaderWriterLock bucketLock;
 				Vector<CasKey> keysToErase;
 
 
@@ -811,6 +811,7 @@ namespace uba
 						lastUseTimeLimit = GetSecondsAsFileTime(secondsSinceCreation - bucketExpirationTimeSeconds);
 				}
 
+				u64 bucketOldest = 0;
 				bucket.oldestUsedTime = 0;
 				bucket.totalEntryCount = 0;
 				bucket.totalEntrySize = 0;
@@ -837,6 +838,7 @@ namespace uba
 									continue;
 								}
 								i = entries.inputsThatAreOutputs.erase(i);
+								e = entries.inputsThatAreOutputs.end();
 								checkInputsForDeletes = true;
 							}
 						}
@@ -844,7 +846,8 @@ namespace uba
 
 					// There is currently no idea keeping more than 256kb worth of entries per lookup key (because that is what fetch max returns).. so let's wipe out
 					// all the entries that overflow that number
-					u64 capacityLeft = SendMaxSize - 32 - entries.GetSharedSize();
+					u64 entriesSize = entries.GetSharedSize();
+					u64 capacityLeft = SendMaxSize - 32 - entriesSize;
 
 					// Check if any offset has been deleted in shared offsets..
 					bool offsetDeletedInShared = false;
@@ -861,30 +864,28 @@ namespace uba
 						}
 					}
 
+					u64 entriesOldestUsed = 0;
+					u64 entriesOldest = 0;
+
 					for (auto i=entries.entries.begin(), e=entries.entries.end(); i!=e;)
 					{
 						auto& entry = *i;
 						bool deleteEntry = false;
-						u64 neededSize = 0;
+						u64 neededSize = entries.GetEntrySize(entry, CacheNetworkVersion, false);
 
-						// Only need to check these things on first iteration.
-						if (deleteIteration == 0)
+						// Check if this entry is outside capacity of network message
+						if (neededSize > capacityLeft)
 						{
-							// Check if this entry is outside capacity of network message
-							neededSize = entries.GetEntrySize(entry, CacheNetworkVersion, false);
-							if (neededSize > capacityLeft)
-							{
-								deleteEntry = true;
-								capacityLeft = 0;
-								++overflowedEntryCount;
-							}
+							deleteEntry = true;
+							capacityLeft = 0;
+							++overflowedEntryCount;
+						}
 
-							// Check if entry has expired
-							if (!deleteEntry && entry.creationTime < lastUseTimeLimit && entry.lastUsedTime < lastUseTimeLimit)
-							{
-								deleteEntry = true;
-								++expiredEntryCount;
-							}
+						// Check if entry has expired
+						if (!deleteEntry && entry.creationTime < lastUseTimeLimit && entry.lastUsedTime < lastUseTimeLimit)
+						{
+							deleteEntry = true;
+							++expiredEntryCount;
 						}
 
 						// This is an attempt at removing entries that has inputs that depends on other entries outputs.
@@ -963,22 +964,16 @@ namespace uba
 							continue;
 						}
 
-						++bucket.totalEntryCount;
-
+						entriesSize += neededSize;
 						capacityLeft -= neededSize;
 
-						{
-							u64 lastUsedTime = entry.lastUsedTime;
-							if (!lastUsedTime)
-								lastUsedTime = entry.creationTime;
-							SCOPED_WRITE_LOCK(globalStatsLock, l);
-							if (!oldest || entry.creationTime < oldest)
-								oldest = entry.creationTime;
-							if (!bucket.oldestUsedTime || lastUsedTime < bucket.oldestUsedTime)
-								bucket.oldestUsedTime = lastUsedTime;
-							if (!oldestUsedTime || lastUsedTime < oldestUsedTime)
-								oldestUsedTime = lastUsedTime;
-						}
+						u64 lastUsedTime = entry.lastUsedTime;
+						if (!lastUsedTime)
+							lastUsedTime = entry.creationTime;
+						if (!entriesOldestUsed || lastUsedTime < entriesOldestUsed)
+							entriesOldestUsed = lastUsedTime;
+						if (!entriesOldest || entry.creationTime < entriesOldest)
+							entriesOldest = entry.creationTime;
 
 						for (auto v : touchedCas)
 							*v = true;
@@ -987,15 +982,31 @@ namespace uba
 					}
 
 
-					if (entries.entries.empty())
 					{
-						SCOPED_WRITE_LOCK(keysToEraseLock, lock2);
-						keysToErase.push_back(li->first);
+						SCOPED_WRITE_LOCK(bucketLock, lock2);
+						if (!bucket.oldestUsedTime || entriesOldestUsed < bucket.oldestUsedTime)
+							bucket.oldestUsedTime = entriesOldestUsed;
+						if (!bucketOldest || entriesOldest < bucketOldest)
+							bucketOldest = entriesOldest;
+
+						if (entries.entries.empty())
+							keysToErase.push_back(li->first);
+						else
+						{
+							bucket.totalEntrySize += entriesSize;
+							bucket.totalEntryCount += entries.entries.size();
+						}
 					}
-					else
-						bucket.totalEntrySize += entries.GetTotalSize(CacheNetworkVersion, false);
 				});
 
+				{
+					SCOPED_WRITE_LOCK(globalStatsLock, l);
+					if (!oldestUsedTime || bucketOldest < oldestUsedTime)
+						oldestUsedTime = bucket.oldestUsedTime;
+					if (!oldest || bucketOldest < oldest)
+						oldest = bucketOldest;
+				}
+
 				for (auto& key : keysToErase)
 					bucket.m_cacheEntryLookup.erase(key);
 
@@ -1042,7 +1053,7 @@ namespace uba
 		if (missingInputEntryCount)
 			m_logger.Detail(TC("  Found %s cache entries with missing input cas"), CountToText(missingInputEntryCount.load()).str);
 
-		m_logger.Detail(TC("  Deleted %s cas files and %s cache entries over %u buckets (%s)"), CountToText(deletedCasCount).str, CountToText(deleteEntryCount.load()).str, u32(m_buckets.size()), TimeToText(GetTime() - deleteCacheEntriesStartTime).str);
+		m_logger.Detail(TC("  Deleted %s cas files and %s cache entries over %u buckets (%u iterations in %s)"), CountToText(deletedCasCount).str, CountToText(deleteEntryCount.load()).str, u32(m_buckets.size()), deleteIteration, TimeToText(GetTime() - deleteCacheEntriesStartTime).str);
 
 		m_bucketIsOverflowing = false;
 
@@ -1202,7 +1213,7 @@ namespace uba
 			u64 longestUnusedSeconds = 0;
 			if (bucket.oldestUsedTime)
 				longestUnusedSeconds = GetFileTimeAsSeconds(now - m_creationTime - bucket.oldestUsedTime);
-			m_logger.Info(TC("    Bucket %u Done (%s). CacheEntries: %s (%s) PathTable: %s CasTable: %s LongestUnused: %s"), bucket.index, TimeToText(GetTime() - bucketStartTime).str, CountToText(bucket.totalEntryCount.load()).str, BytesToText(bucket.totalEntrySize.load()).str, BytesToText(bucket.m_pathTable.GetSize()).str, BytesToText(bucket.m_casKeyTable.GetSize()).str, TimeToText(MsToTime(longestUnusedSeconds*1000), true).str);
+			m_logger.Info(TC("    Bucket %u Done (%s). CacheEntries: %s (%s) PathTable: %s CasTable: %s LongestUnused: %s"), bucket.index, TimeToText(GetTime() - bucketStartTime).str, CountToText(bucket.totalEntryCount).str, BytesToText(bucket.totalEntrySize).str, BytesToText(bucket.m_pathTable.GetSize()).str, BytesToText(bucket.m_casKeyTable.GetSize()).str, TimeToText(MsToTime(longestUnusedSeconds*1000), true).str);
 
 			SCOPED_WRITE_LOCK(globalStatsLock, l);
 			maxCommittedMemory = Max(maxCommittedMemory, memoryBlock.writtenSize);
@@ -1219,7 +1230,7 @@ namespace uba
 		if (entriesAdded || deletedCasCount || deleteEntryCount || forceAllSteps)
 		{
 			u64 saveStart = GetTime();
-			m_logger.Detail(TC("  Saving to disk"));
+			m_logger.Detail(TC("  Saving to disk (v%u)"), CacheFileVersion);
 			Event saveCasEvent(true);
 			m_server.AddWork([&]() { m_storage.SaveCasTable(false, false); saveCasEvent.Set(); }, 1, TC("SaveCas"));
 			SaveNoLock();
@@ -1549,7 +1560,7 @@ namespace uba
 				auto findIt = existing.find(path.data);
 				if (findIt == existing.end())
 				{
-					m_logger.Warning(TC("Existing cache entry matches input but does not match output (output file %s did not exist in existing cache entry)"), path.data);
+					m_logger.Warning(TC("Client %s sent cache entry that already exists but does not match output (output file %s did not exist in existing cache entry)"), path.data);
 					cacheEntries.entries.erase(matchingEntry);
 					shouldOverwrite = true;
 					break;
@@ -1785,17 +1796,19 @@ namespace uba
 
 		StringBuffer<> line;
 
+		auto writePathFromOffset = [&](Bucket& bucket, u32 offset)
+			{
+				CasKey casKey;
+				StringBuffer<> path;
+				bucket.m_casKeyTable.GetPathAndKey(path, casKey, bucket.m_pathTable, offset);
+				writeLine(line.Clear().Appendf(TC("    %s - %s"), path.data, CasKeyString(casKey).str));
+			};
+
 		auto writePathsFromOffsets = [&](Bucket& bucket, const Vector<u8>& offsets)
 			{
 				BinaryReader reader2(offsets.data(), 0, offsets.size());
 				while (reader2.GetLeft())
-				{
-					u64 offset = reader2.Read7BitEncoded();
-					CasKey casKey;
-					StringBuffer<> path;
-					bucket.m_casKeyTable.GetPathAndKey(path, casKey, bucket.m_pathTable, offset);
-					writeLine(line.Clear().Appendf(TC("    %s - %s"), path.data, CasKeyString(casKey).str));
-				}
+					writePathFromOffset(bucket, u32(reader2.Read7BitEncoded()));
 			};
 
 		if (command.Equals(TC("content")))
@@ -1851,8 +1864,17 @@ namespace uba
 							continue;
 					}
 
+					writeLine(line.Clear().Appendf(TC("Key: %s"), CasKeyString(kv2.first).str));
+					if (!entries.inputsThatAreOutputs.empty())
+					{
+						writeLine(line.Clear().Appendf(TC("  InputsThatAreOutputs:")));
+						if (*entries.inputsThatAreOutputs.begin() == ~0u)
+							writeLine(line.Clear().Appendf(TC("   WillCheckAll (~0u)")));
+						else
+							for (u32 offset : entries.inputsThatAreOutputs)
+								writePathFromOffset(bucket, offset);
+					}
 
-					writeLine(ToView(CasKeyString(kv2.first).str));
 					u32 index = 0;
 					for (auto& entry : entries.entries)
 					{
@@ -1971,10 +1993,22 @@ namespace uba
 			}
 			if (!largestEntries)
 				return true;
+			writeLine(line.Clear().Append(TC("Largest entry (based on shared inputs)")));
 			CacheEntries& entries = *largestEntries;
 			SCOPED_READ_LOCK(entries.lock, lock3);
 			writeLine(line.Clear().Append(TC("   SharedInputs:")));
 			writePathsFromOffsets(*largestEntriesBucket, entries.sharedInputCasKeyOffsets);
+			u32 index = 0;
+			for (auto& entry : entries.entries)
+			{
+				writeLine(line.Clear().Appendf(TC("  #%u"), index));
+				writeLine(line.Clear().Append(TC("   ExtraInputs:")));
+				writePathsFromOffsets(*largestEntriesBucket, entry.extraInputCasKeyOffsets);
+				writeLine(line.Clear().Append(TC("   Outputs:")));
+				writePathsFromOffsets(*largestEntriesBucket, entry.outputCasKeyOffsets);
+				++index;
+			}
+
 		}
 		else if (command.Equals(TC("obliterate")))
 		{
