[UBA]
* Added so path table is populated with common string segments at creation (cpp, h, rsp, etc).. this makes the total size of the table smaller

diff --git a/Engine/Source/Programs/UnrealBuildAccelerator/Common/Private/UbaCacheClient.cpp b/Engine/Source/Programs/UnrealBuildAccelerator/Common/Private/UbaCacheClient.cpp
--- a/Engine/Source/Programs/UnrealBuildAccelerator/Common/Private/UbaCacheClient.cpp
+++ b/Engine/Source/Programs/UnrealBuildAccelerator/Common/Private/UbaCacheClient.cpp
@@ -169,6 +169,8 @@ namespace uba
 
 		SCOPED_WRITE_LOCK(m_bucketsLock, bucketsLock);
 		Bucket& bucket = m_buckets.try_emplace(bucketId, bucketId).first->second;
+		if (!bucket.sendPathTable.GetMemory())
+			sendPathTable.AddCommonStringSegments();
 		bucketsLock.Leave();
 
 		TString qualifiedPath;
diff --git a/Engine/Source/Programs/UnrealBuildAccelerator/Common/Private/UbaCacheServer.cpp b/Engine/Source/Programs/UnrealBuildAccelerator/Common/Private/UbaCacheServer.cpp
--- a/Engine/Source/Programs/UnrealBuildAccelerator/Common/Private/UbaCacheServer.cpp
+++ b/Engine/Source/Programs/UnrealBuildAccelerator/Common/Private/UbaCacheServer.cpp
@@ -1310,6 +1310,7 @@ namespace uba
 		auto& bucket = insres.first->second;
 		if (!insres.second)
 			return bucket;
+		bucket.m_pathTable.AddCommonStringSegments();
 		bucket.index = u32(m_buckets.size() - 1);
 		return bucket;
 	}
diff --git a/Engine/Source/Programs/UnrealBuildAccelerator/Common/Private/UbaCompactTables.cpp b/Engine/Source/Programs/UnrealBuildAccelerator/Common/Private/UbaCompactTables.cpp
--- a/Engine/Source/Programs/UnrealBuildAccelerator/Common/Private/UbaCompactTables.cpp
+++ b/Engine/Source/Programs/UnrealBuildAccelerator/Common/Private/UbaCompactTables.cpp
@@ -49,6 +49,42 @@ namespace uba
 		return InternalAdd(str, stringKeyString, strLen);
 	}
 
+	void CompactPathTable::AddCommonStringSegments()
+	{
+		const tchar* commonSegments[] =
+		{
+			TC("h"),
+			TC("cpp"),
+			TC("inl"),
+			TC("obj"),
+			TC("o"),
+			TC("c"),
+			TC("lib"),
+			TC("rsp"),
+			TC("dep"),
+			TC("json"),
+			TC("sarif"),
+			TC("d"),
+			TC("gen"),
+			TC("generated"),
+			TC("init"),
+			TC("ispc"),
+		};
+
+		if (!m_pathTableMem.memory)
+			m_pathTableMem.Init(PathTableMaxSize);
+
+		for (const tchar* seg : commonSegments)
+		{
+			u32 segLen = TStrlen(seg);
+			BinaryWriter writer((u8*)m_pathTableMem.AllocateNoLock(segLen+1, 1, TC("CompactPathTable")), 0, segLen+1);
+			StringKey segmentKey = ToStringKeyNoCheck(seg, segLen);
+			u32& strOffset = m_segmentOffsets.Insert(segmentKey);
+			strOffset = u32(m_pathTableMem.writtenSize);
+			writer.WriteString(seg, segLen);
+		}
+	}
+
 	u32 CompactPathTable::InternalAdd(const tchar* str, const tchar* stringKeyString, u64 strLen)
 	{
 		StringKey key = ToStringKeyNoCheck(stringKeyString, strLen);
diff --git a/Engine/Source/Programs/UnrealBuildAccelerator/Common/Public/UbaCompactTables.h b/Engine/Source/Programs/UnrealBuildAccelerator/Common/Public/UbaCompactTables.h
--- a/Engine/Source/Programs/UnrealBuildAccelerator/Common/Public/UbaCompactTables.h
+++ b/Engine/Source/Programs/UnrealBuildAccelerator/Common/Public/UbaCompactTables.h
@@ -21,6 +21,8 @@ namespace uba
 		u32 Add(const tchar* str, u64 strLen, u32& outRequiredPathTableSize);
 		u32 AddNoLock(const tchar* str, u64 strLen);
 
+		void AddCommonStringSegments();
+
 		void GetString(StringBufferBase& out, u64 offset) const;
 
 		u8* GetMemory();
