[UBA]
* Fixed compile error

diff --git a/Engine/Source/Programs/UnrealBuildAccelerator/Common/Private/UbaCacheClient.cpp b/Engine/Source/Programs/UnrealBuildAccelerator/Common/Private/UbaCacheClient.cpp
--- a/Engine/Source/Programs/UnrealBuildAccelerator/Common/Private/UbaCacheClient.cpp
+++ b/Engine/Source/Programs/UnrealBuildAccelerator/Common/Private/UbaCacheClient.cpp
@@ -170,7 +170,7 @@ namespace uba
 		SCOPED_WRITE_LOCK(m_bucketsLock, bucketsLock);
 		Bucket& bucket = m_buckets.try_emplace(bucketId, bucketId).first->second;
 		if (!bucket.sendPathTable.GetMemory())
-			sendPathTable.AddCommonStringSegments();
+			bucket.sendPathTable.AddCommonStringSegments();
 		bucketsLock.Leave();
 
 		TString qualifiedPath;
