[UBA]
* CompactPathTable version 3 - Changed so offset to parent is relative instead of absolute. It will always be shorter and thus store fewer bits

diff --git a/Engine/Source/Programs/UnrealBuildAccelerator/Common/Private/UbaCompactTables.cpp b/Engine/Source/Programs/UnrealBuildAccelerator/Common/Private/UbaCompactTables.cpp
--- a/Engine/Source/Programs/UnrealBuildAccelerator/Common/Private/UbaCompactTables.cpp
+++ b/Engine/Source/Programs/UnrealBuildAccelerator/Common/Private/UbaCompactTables.cpp
@@ -96,6 +96,54 @@ namespace uba
 		const tchar* seg = str;
 		u32 parentOffset = 0;
 		
+		if (m_version >= 3)
+		{
+			tchar parentSeparator = 0;
+			for (const tchar* it = str + strLen - 1; it > str; --it)
+			{
+				if (*it != PathSeparator && *it != '.')
+					continue;
+				parentOffset = InternalAdd(str, stringKeyString, it - str);
+				parentSeparator = *it;
+				seg = it + 1;
+				break;
+			}
+
+			u64 segLen = strLen - (seg - str);
+			u32& offset = m_offsets.GetValueFromIndex(offsetIndex);
+
+			StringKey segmentKey = ToStringKeyNoCheck(seg, segLen);
+			added = false;
+			u32& strOffset = m_segmentOffsets.Insert(segmentKey, added);
+			if (added)
+			{
+				offset = u32(m_pathTableMem.writtenSize);
+				u32 parentRelativeOffset = ((offset - parentOffset) << 1) + (parentSeparator == '.');
+
+				// Put string directly after current element and set segment offset to 0
+				u64 bytesForString = GetStringWriteSize(seg, segLen);
+				u8 bytesForParent = Get7BitEncodedCount(parentRelativeOffset);
+				u64 memSize = bytesForParent + 1 + bytesForString;
+			
+				BinaryWriter writer((u8*)m_pathTableMem.AllocateNoLock(memSize, 1, TC("CompactPathTable")), 0, memSize);
+				writer.Write7BitEncoded(parentRelativeOffset);
+				writer.Write7BitEncoded(0);
+				writer.WriteString(seg, segLen);
+				strOffset = offset + bytesForParent + 1;
+				return offset;
+			}
+
+			offset = u32(m_pathTableMem.writtenSize);
+			u32 parentRelativeOffset = ((offset - parentOffset) << 1) + (parentSeparator == '.');
+
+			u8 bytesForParent = Get7BitEncodedCount(parentRelativeOffset);
+			u64 memSize = bytesForParent + Get7BitEncodedCount(strOffset);
+			BinaryWriter writer((u8*)m_pathTableMem.AllocateNoLock(memSize, 1, TC("CompactPathTable")), 0, memSize);
+			writer.Write7BitEncoded(parentRelativeOffset);
+			writer.Write7BitEncoded(strOffset);
+			return offset;
+		}
+
 		if (m_version >= 2)
 		{
 			for (const tchar* it = str + strLen - 1; it > str; --it)
@@ -181,19 +229,39 @@ namespace uba
 
 		BinaryReader reader(m_pathTableMem.memory, offset, m_pathTableMem.writtenSize);
 
-		while (offset)
+		if (m_version >= 3)
 		{
-			++offsetCount;
-			UBA_ASSERT(offsetCount < sizeof_array(offsets));
-			reader.SetPosition(offset);
-			offset = (u32)reader.Read7BitEncoded();
-			separators[offsetCount] = 0;
-			if (m_version >= 2)
+			offsetCount = 1;
+			while (true)
 			{
-				separators[offsetCount] = offset & 1;
-				offset >>= 1;
+				UBA_ASSERT(offsetCount < sizeof_array(offsets));
+				reader.SetPosition(offset);
+				u64 parentRelativeOffset = reader.Read7BitEncoded();
+				separators[offsetCount] = parentRelativeOffset & 1;
+				parentRelativeOffset >>= 1;
+				if (parentRelativeOffset == offset)
+					break;
+				offset -= parentRelativeOffset;
+				offsets[offsetCount] = u32(offset);
+				++offsetCount;
+			}
+		}
+		else
+		{
+			while (offset)
+			{
+				++offsetCount;
+				UBA_ASSERT(offsetCount < sizeof_array(offsets));
+				reader.SetPosition(offset);
+				offset = (u32)reader.Read7BitEncoded();
+				separators[offsetCount] = 0;
+				if (m_version >= 2)
+				{
+					separators[offsetCount] = offset & 1;
+					offset >>= 1;
+				}
+				offsets[offsetCount] = u32(offset);
 			}
-			offsets[offsetCount] = u32(offset);
 		}
 
 		{
diff --git a/Engine/Source/Programs/UnrealBuildAccelerator/Common/Public/UbaNetwork.h b/Engine/Source/Programs/UnrealBuildAccelerator/Common/Public/UbaNetwork.h
--- a/Engine/Source/Programs/UnrealBuildAccelerator/Common/Public/UbaNetwork.h
+++ b/Engine/Source/Programs/UnrealBuildAccelerator/Common/Public/UbaNetwork.h
@@ -18,7 +18,7 @@ namespace uba
 	static constexpr u32 StorageNetworkVersion = 4;
 	static constexpr u32 SessionNetworkVersion = 42;
 	static constexpr u32 CacheNetworkVersion = 5;
-	static constexpr u32 CacheBucketVersion = 2;
+	static constexpr u32 CacheBucketVersion = 3;
 
 	// Messages used over network between client/server (system, storage and session)
 
diff --git a/Engine/Source/Programs/UnrealBuildAccelerator/Test/Private/UbaTestCache.h b/Engine/Source/Programs/UnrealBuildAccelerator/Test/Private/UbaTestCache.h
--- a/Engine/Source/Programs/UnrealBuildAccelerator/Test/Private/UbaTestCache.h
+++ b/Engine/Source/Programs/UnrealBuildAccelerator/Test/Private/UbaTestCache.h
@@ -212,23 +212,26 @@ namespace uba
 
 	bool TestCompactPathTable(Logger& logger, const StringBufferBase& rootDir)
 	{
-		for (u32 version=0; version<CacheBucketVersion; ++version)
+		for (u32 version=0; version<=CacheBucketVersion; ++version)
 		{
 			CompactPathTable table(CaseInsensitiveFs, 0, 0, version);
 
 			StringBuffer<> str;
-			str.Append("foo").EnsureEndsWithSlash().Append("bar.h");
-			u32 offset = table.Add(str.data, str.count);
+			u32 offset;
 
-			StringBuffer<> str2;
-			table.GetString(str2, offset);
-			if (!str.Equals(str2.data))
-				return false;
+			auto testStr = [&](const StringView& str)
+				{
+					StringBuffer<> str2;
+					offset = table.Add(str.data, str.count);
+					table.GetString(str2, offset);
+					return str.Equals(str2.data);
+				};
 
-			str.Clear().Append(PathSeparator).Append("foo").Append(PathSeparator).Append("bar.h");
-			offset = table.Add(str.data, str.count);
-			table.GetString(str2.Clear(), offset);
-			if (!str.Equals(str2.data))
+			if (!testStr(str.Append("foo")))
+				return false;
+			if (!testStr(str.Clear().Append("foo").EnsureEndsWithSlash().Append("bar.h")))
+				return false;
+			if (!testStr(str.Clear().Append(PathSeparator).Append("foo").Append(PathSeparator).Append("bar.h")))
 				return false;
 
 			CompactPathTable table2(CaseInsensitiveFs, 0, 0, version);
