[uba]
* Disabled send-pathtable common string segments for now since it seems to cause some issues. will revisit

diff --git a/Engine/Source/Programs/UnrealBuildAccelerator/Common/Private/UbaCacheClient.cpp b/Engine/Source/Programs/UnrealBuildAccelerator/Common/Private/UbaCacheClient.cpp
--- a/Engine/Source/Programs/UnrealBuildAccelerator/Common/Private/UbaCacheClient.cpp
+++ b/Engine/Source/Programs/UnrealBuildAccelerator/Common/Private/UbaCacheClient.cpp
@@ -169,8 +169,8 @@ namespace uba
 
 		SCOPED_WRITE_LOCK(m_bucketsLock, bucketsLock);
 		Bucket& bucket = m_buckets.try_emplace(bucketId, bucketId).first->second;
-		if (!bucket.sendPathTable.GetMemory())
-			bucket.sendPathTable.AddCommonStringSegments();
+		//if (!bucket.sendPathTable.GetMemory())
+		//	bucket.sendPathTable.AddCommonStringSegments();
 		bucketsLock.Leave();
 
 		TString qualifiedPath;
