[UBA]
* Fixed so GetSymbols work for child processes on remote helpers.

#skipci

diff --git a/Engine/Source/Programs/UnrealBuildAccelerator/Common/Private/UbaProcess.cpp b/Engine/Source/Programs/UnrealBuildAccelerator/Common/Private/UbaProcess.cpp
--- a/Engine/Source/Programs/UnrealBuildAccelerator/Common/Private/UbaProcess.cpp
+++ b/Engine/Source/Programs/UnrealBuildAccelerator/Common/Private/UbaProcess.cpp
@@ -1322,7 +1322,7 @@ namespace uba
 
 	bool ProcessImpl::HandleResolveCallstack(BinaryReader& reader, BinaryWriter& writer)
 	{
-		m_session.GetSymbols(*this, reader, writer);
+		m_session.GetSymbols(m_startInfo.application, m_isArmBinary, reader, writer);
 		return true;
 	}
 
diff --git a/Engine/Source/Programs/UnrealBuildAccelerator/Common/Private/UbaSession.cpp b/Engine/Source/Programs/UnrealBuildAccelerator/Common/Private/UbaSession.cpp
--- a/Engine/Source/Programs/UnrealBuildAccelerator/Common/Private/UbaSession.cpp
+++ b/Engine/Source/Programs/UnrealBuildAccelerator/Common/Private/UbaSession.cpp
@@ -2802,7 +2802,7 @@ namespace uba
 		return true;
 	}
 
-	bool Session::GetSymbols(Process& process, BinaryReader& reader, BinaryWriter& writer)
+	bool Session::GetSymbols(const tchar* application, bool isArm, BinaryReader& reader, BinaryWriter& writer)
 	{
 		StringBuffer<256> detoursLibPath;
 		StringBuffer<256> alternativeLibPath;
@@ -2813,7 +2813,7 @@ namespace uba
 		StringView searchPaths[3] = { detoursLibPath, alternativeLibPath, {} };
 
 		StringBuffer<16*1024> sb;
-		ParseCallstackInfo(sb, reader, process.GetStartInfo().application, searchPaths);
+		ParseCallstackInfo(sb, reader, application, searchPaths);
 		writer.WriteString(sb);
 		return true;
 	}
diff --git a/Engine/Source/Programs/UnrealBuildAccelerator/Common/Private/UbaSessionClient.cpp b/Engine/Source/Programs/UnrealBuildAccelerator/Common/Private/UbaSessionClient.cpp
--- a/Engine/Source/Programs/UnrealBuildAccelerator/Common/Private/UbaSessionClient.cpp
+++ b/Engine/Source/Programs/UnrealBuildAccelerator/Common/Private/UbaSessionClient.cpp
@@ -2138,12 +2138,12 @@ namespace uba
 		return true;
 	}
 
-	bool SessionClient::GetSymbols(Process& process, BinaryReader& reader, BinaryWriter& writer)
+	bool SessionClient::GetSymbols(const tchar* application, bool isArm, BinaryReader& reader, BinaryWriter& writer)
 	{
 		StackBinaryWriter<SendMaxSize> msgWriter;
 		NetworkMessage msg(m_client, ServiceId, SessionMessageType_GetSymbols, msgWriter);
-		msgWriter.WriteU32(process.GetId());
-		msgWriter.WriteBool(process.IsArm());
+		msgWriter.WriteString(application);
+		msgWriter.WriteBool(isArm);
 		msgWriter.WriteBytes(reader.GetPositionData(), reader.GetLeft());
 
 		BinaryReader responseReader(writer.GetData(), 0, writer.GetCapacityLeft());
diff --git a/Engine/Source/Programs/UnrealBuildAccelerator/Common/Private/UbaSessionServer.cpp b/Engine/Source/Programs/UnrealBuildAccelerator/Common/Private/UbaSessionServer.cpp
--- a/Engine/Source/Programs/UnrealBuildAccelerator/Common/Private/UbaSessionServer.cpp
+++ b/Engine/Source/Programs/UnrealBuildAccelerator/Common/Private/UbaSessionServer.cpp
@@ -2292,11 +2292,8 @@ namespace uba
 
 	bool SessionServer::HandleGetSymbols(const ConnectionInfo& connectionInfo, BinaryReader& reader, BinaryWriter& writer)
 	{
-		u32 processId = reader.ReadU32();
+		TString application = reader.ReadString();
 		bool isClientArm = reader.ReadBool();
-		ProcessHandle h(GetProcess(processId));
-		if (!h.IsValid())
-			return m_logger.Error(TC("Failed to find process for id %u when receiving update environment message"), processId);
 
 		if (IsArmBinary != isClientArm)
 		{
@@ -2304,7 +2301,7 @@ namespace uba
 			return true;
 		}
 
-		GetSymbols(*h.m_process, reader, writer);
+		GetSymbols(application.c_str(), isClientArm, reader, writer);
 		return true;
 
 		#if 0
diff --git a/Engine/Source/Programs/UnrealBuildAccelerator/Common/Public/UbaSession.h b/Engine/Source/Programs/UnrealBuildAccelerator/Common/Public/UbaSession.h
--- a/Engine/Source/Programs/UnrealBuildAccelerator/Common/Public/UbaSession.h
+++ b/Engine/Source/Programs/UnrealBuildAccelerator/Common/Public/UbaSession.h
@@ -156,7 +156,7 @@ namespace uba
 		virtual bool CustomMessage(Process& process, BinaryReader& reader, BinaryWriter& writer);
 		virtual bool SHGetKnownFolderPath(Process& process, BinaryReader& reader, BinaryWriter& writer);
 		virtual bool HostRun(BinaryReader& reader, BinaryWriter& writer);
-		virtual bool GetSymbols(Process& process, BinaryReader& reader, BinaryWriter& writer);
+		virtual bool GetSymbols(const tchar* application, bool isArm, BinaryReader& reader, BinaryWriter& writer);
 		virtual void FileEntryAdded(StringKey fileNameKey, u64 lastWritten, u64 size);
 		virtual bool FlushWrittenFiles(ProcessImpl& process);
 		virtual bool UpdateEnvironment(ProcessImpl& process, const StringView& reason, bool resetStats);
diff --git a/Engine/Source/Programs/UnrealBuildAccelerator/Common/Public/UbaSessionClient.h b/Engine/Source/Programs/UnrealBuildAccelerator/Common/Public/UbaSessionClient.h
--- a/Engine/Source/Programs/UnrealBuildAccelerator/Common/Public/UbaSessionClient.h
+++ b/Engine/Source/Programs/UnrealBuildAccelerator/Common/Public/UbaSessionClient.h
@@ -68,7 +68,7 @@ namespace uba
 		virtual bool CustomMessage(Process& process, BinaryReader& reader, BinaryWriter& writer) override;
 		virtual bool SHGetKnownFolderPath(Process& process, BinaryReader& reader, BinaryWriter& writer) override;
 		virtual bool HostRun(BinaryReader& reader, BinaryWriter& writer) override;
-		virtual bool GetSymbols(Process& process, BinaryReader& reader, BinaryWriter& writer);
+		virtual bool GetSymbols(const tchar* application, bool isArm, BinaryReader& reader, BinaryWriter& writer);
 		virtual bool FlushWrittenFiles(ProcessImpl& process) override;
 		virtual bool UpdateEnvironment(ProcessImpl& process, const StringView& reason, bool resetStats) override;
 		virtual bool LogLine(ProcessImpl& process, const tchar* line, LogEntryType logType) override;
