[UBA]
* Take write lock in storage save to prevent potential deadlock when cache client is running while shutting down

#skipci

diff --git a/Engine/Source/Programs/UnrealBuildAccelerator/Common/Private/UbaStorage.cpp b/Engine/Source/Programs/UnrealBuildAccelerator/Common/Private/UbaStorage.cpp
--- a/Engine/Source/Programs/UnrealBuildAccelerator/Common/Private/UbaStorage.cpp
+++ b/Engine/Source/Programs/UnrealBuildAccelerator/Common/Private/UbaStorage.cpp
@@ -1592,9 +1592,9 @@ namespace uba
 			if (!tempFile.CreateWrite(false, DefaultAttributes(), 0, m_tempPath.data))
 				return false;
 			
-			SCOPED_READ_LOCK(m_fileTableLookupLock, fileTableLock);
-			SCOPED_READ_LOCK(m_casLookupLock, casLookupLock);
-			SCOPED_READ_LOCK(m_accessLock, accessLock);
+			SCOPED_WRITE_LOCK(m_fileTableLookupLock, fileTableLock);
+			SCOPED_WRITE_LOCK(m_casLookupLock, casLookupLock);
+			SCOPED_WRITE_LOCK(m_accessLock, accessLock);
 
 			Vector<u8> buffer;
 
