[UBA]
* Fixed cache server error message
* Fixed so cache server unregisters itself from network server

#skipci

diff --git a/Engine/Source/Programs/UnrealBuildAccelerator/Common/Private/UbaCacheServer.cpp b/Engine/Source/Programs/UnrealBuildAccelerator/Common/Private/UbaCacheServer.cpp
--- a/Engine/Source/Programs/UnrealBuildAccelerator/Common/Private/UbaCacheServer.cpp
+++ b/Engine/Source/Programs/UnrealBuildAccelerator/Common/Private/UbaCacheServer.cpp
@@ -221,6 +221,8 @@ namespace uba
 
 	CacheServer::~CacheServer()
 	{
+		m_server.UnregisterOnClientDisconnected(CacheServiceId);
+		m_server.UnregisterService(CacheServiceId);
 	}
 
 	struct CacheServer::LoadStats
@@ -1377,7 +1379,7 @@ namespace uba
 		{
 			u32 clientVersion;
 			auto& bucket = GetConnectionBucket(connectionInfo, reader, &clientVersion);
-			return HandleStoreEntry(bucket, reader, writer, clientVersion);
+			return HandleStoreEntry(bucket, reader, writer, clientVersion, connectionInfo.GetId());
 		}
 		case CacheMessageType_StoreEntryDone:
 		{
@@ -1437,7 +1439,7 @@ namespace uba
 		}
 	}
 
-	bool CacheServer::HandleStoreEntry(ConnectionBucket& connectionBucket, BinaryReader& reader, BinaryWriter& writer, u32 clientVersion)
+	bool CacheServer::HandleStoreEntry(ConnectionBucket& connectionBucket, BinaryReader& reader, BinaryWriter& writer, u32 clientVersion, u32 clientId)
 	{
 		CasKey cmdKey = reader.ReadCasKey();
 
@@ -1561,7 +1563,7 @@ namespace uba
 				auto findIt = existing.find(path.data);
 				if (findIt == existing.end())
 				{
-					m_logger.Warning(TC("Client %s sent cache entry that already exists but does not match output (output file %s did not exist in existing cache entry)"), path.data);
+					m_logger.Warning(TC("Client %s sent cache entry that already exists but does not match output (output file %s did not exist in existing cache entry)"), clientId, path.data);
 					cacheEntries.entries.erase(matchingEntry);
 					shouldOverwrite = true;
 					break;
diff --git a/Engine/Source/Programs/UnrealBuildAccelerator/Common/Public/UbaCacheServer.h b/Engine/Source/Programs/UnrealBuildAccelerator/Common/Public/UbaCacheServer.h
--- a/Engine/Source/Programs/UnrealBuildAccelerator/Common/Public/UbaCacheServer.h
+++ b/Engine/Source/Programs/UnrealBuildAccelerator/Common/Public/UbaCacheServer.h
@@ -78,7 +78,7 @@ namespace uba
 		u32 GetBucketWorkerCount();
 
 		bool HandleMessage(const ConnectionInfo& connectionInfo, u8 messageType, BinaryReader& reader, BinaryWriter& writer);
-		bool HandleStoreEntry(ConnectionBucket& bucket, BinaryReader& reader, BinaryWriter& writer, u32 clientVersion);
+		bool HandleStoreEntry(ConnectionBucket& bucket, BinaryReader& reader, BinaryWriter& writer, u32 clientVersion, u32 clientId);
 		bool HandleFetchPathTable(BinaryReader& reader, BinaryWriter& writer);
 		bool HandleFetchCasTable(BinaryReader& reader, BinaryWriter& writer);
 		bool HandleFetchEntries(BinaryReader& reader, BinaryWriter& writer, u32 clientId);
