[UBA]
* Bumped cache database version to 9 (added common string segments) and fixed so code can handle older versions

#skipci

diff --git a/Engine/Source/Programs/UnrealBuildAccelerator/Common/Private/UbaCacheServer.cpp b/Engine/Source/Programs/UnrealBuildAccelerator/Common/Private/UbaCacheServer.cpp
--- a/Engine/Source/Programs/UnrealBuildAccelerator/Common/Private/UbaCacheServer.cpp
+++ b/Engine/Source/Programs/UnrealBuildAccelerator/Common/Private/UbaCacheServer.cpp
@@ -27,7 +27,7 @@
 
 namespace uba
 {
-	static constexpr u32 CacheFileVersion = 8;
+	static constexpr u32 CacheFileVersion = 9;
 	static constexpr u32 CacheFileCompatibilityVersion = 3;
 
 	bool IsCaseInsensitive(u64 id) { return (id & (1ull << 32)) == 0; }
@@ -285,12 +285,13 @@ namespace uba
 					u64 id;
 					if (!keyName.Parse(id))
 						return;
-					GetBucket(id);
+					GetBucket(id, false);
 				});
 
 			m_server.ParallelFor(GetBucketWorkerCount(), m_buckets, [&](auto& it)
 				{
 					u64 key = it->first;
+					Bucket& bucket = it->second;
 
 					StringBuffer<MaxPath> bucketFilename(bucketsDir);
 					bucketFilename.EnsureEndsWithSlash().AppendValue(key);
@@ -302,7 +303,7 @@ namespace uba
 					}
 					BinaryReader reader(bucketFile.GetData(), 0, bucketFile.GetSize());
 					u32 bucketVersion = reader.ReadU32();
-					LoadBucket(it->second, reader, bucketVersion, stats);
+					LoadBucket(bucket, reader, bucketVersion, stats);
 				});
 		}
 
@@ -316,10 +317,14 @@ namespace uba
 		if (databaseVersion != CacheFileVersion)
 			bucket.needsSave = true;
 
+		if (databaseVersion > 8)
+			bucket.m_pathTable.AddCommonStringSegments();
+
 		u32 pathTableSize = reader.ReadU32();
 		if (pathTableSize)
 		{
-			BinaryReader pathTableReader(reader.GetPositionData(), 0, pathTableSize);
+			u64 skipCommonSegments = bucket.m_pathTable.GetSize();
+			BinaryReader pathTableReader(reader.GetPositionData(), skipCommonSegments, pathTableSize);
 			bucket.m_pathTable.ReadMem(pathTableReader, true);
 			reader.Skip(pathTableSize);
 		}
@@ -1304,7 +1309,7 @@ namespace uba
 		
 	}
 
-	CacheServer::Bucket& CacheServer::GetBucket(u64 id)
+	CacheServer::Bucket& CacheServer::GetBucket(u64 id, bool addCommon)
 	{
 		u32 bucketVersion = u32(id >> 35u);
 		SCOPED_WRITE_LOCK(m_bucketsLock, bucketsLock);
@@ -1312,7 +1317,8 @@ namespace uba
 		auto& bucket = insres.first->second;
 		if (!insres.second)
 			return bucket;
-		bucket.m_pathTable.AddCommonStringSegments();
+		if (addCommon)
+			bucket.m_pathTable.AddCommonStringSegments();
 		bucket.index = u32(m_buckets.size() - 1);
 		return bucket;
 	}
diff --git a/Engine/Source/Programs/UnrealBuildAccelerator/Common/Public/UbaCacheServer.h b/Engine/Source/Programs/UnrealBuildAccelerator/Common/Public/UbaCacheServer.h
--- a/Engine/Source/Programs/UnrealBuildAccelerator/Common/Public/UbaCacheServer.h
+++ b/Engine/Source/Programs/UnrealBuildAccelerator/Common/Public/UbaCacheServer.h
@@ -74,7 +74,7 @@ namespace uba
 
 		ConnectionBucket& GetConnectionBucket(const ConnectionInfo& connectionInfo, BinaryReader& reader, u32* outClientVersion = nullptr);
 		Bucket& GetBucket(BinaryReader& reader);
-		Bucket& GetBucket(u64 id);
+		Bucket& GetBucket(u64 id, bool addCommon = true);
 		u32 GetBucketWorkerCount();
 
 		bool HandleMessage(const ConnectionInfo& connectionInfo, u8 messageType, BinaryReader& reader, BinaryWriter& writer);
diff --git a/Engine/Source/Programs/UnrealBuildAccelerator/Test/Private/UbaTestCache.h b/Engine/Source/Programs/UnrealBuildAccelerator/Test/Private/UbaTestCache.h
--- a/Engine/Source/Programs/UnrealBuildAccelerator/Test/Private/UbaTestCache.h
+++ b/Engine/Source/Programs/UnrealBuildAccelerator/Test/Private/UbaTestCache.h
@@ -59,157 +59,6 @@
 		storage.InvalidateCachedFileInfo(GetKeyAndFixedName(fixedFilePath, fileName.data));
 	}
 
-	bool TestCacheClientAndServer(LoggerWithWriter& logger, const StringBufferBase& testRootDir)
-	{
-		LogWriter& logWriter = logger.m_writer;
-		NetworkBackendTcp tcpBackend(logWriter);
-
-		bool ctorSuccess = true;
-		NetworkServer server(ctorSuccess, { logWriter });
-
-		StringBuffer<MaxPath> rootDir;
-		rootDir.Append(testRootDir).Append(TC("Uba"));
-		if (!DeleteAllFiles(logger, rootDir.data))
-			return false;
-
-		StorageServerCreateInfo storageServerInfo(server, rootDir.data, logWriter);
-		storageServerInfo.casCapacityBytes = 1024ull * 1024 * 1024;
-		auto& storageServer = *new StorageServer(storageServerInfo);
-		auto ssg = MakeGuard([&]() { delete &storageServer; });
-
-		CacheServerCreateInfo csci(storageServer, rootDir.data, logWriter);
-		CacheServer cacheServer(csci);
-
-
-		SessionServerCreateInfo sessionInfo(storageServer, server, logWriter);
-		sessionInfo.rootDir = rootDir.data;
-		auto& session = *new SessionServer(sessionInfo);
-		auto sg = MakeGuard([&]() { delete &session; });
-
-		u16 port = 1356;
-		if (!server.StartListen(tcpBackend, port))
-			return logger.Error(TC("Failed to listen"));
-		auto disconnectServer = MakeGuard([&]() { server.DisconnectClients(); });
-
-		StringBuffer<MaxPath> workingDir;
-		workingDir.Append(testRootDir).Append(TC("WorkingDir"));
-		if (!DeleteAllFiles(logger, workingDir.data))
-			return false;
-		if (!storageServer.CreateDirectory(workingDir.data))
-			return false;
-		if (!DeleteAllFiles(logger, workingDir.data, false))
-			return false;
-
-		StringBuffer<> testApp;
-		GetTestAppPath(logger, testApp);
-
-		StringBuffer<MaxPath> inputFile;
-		if (!CreateTextFile(inputFile, logger, workingDir.data, TC("Input.txt"), "Foo"))
-			return false;
-		StringBuffer<MaxPath> outputFile;
-		if (!CreateTextFile(outputFile, logger, workingDir.data, TC("Output.txt"), "Foo"))
-			return false;
-
-		StackBinaryWriter<256> inputs;
-		inputs.WriteString(inputFile);
-
-		StackBinaryWriter<256> outputs;
-		outputs.WriteString(outputFile);
-
-		StackBinaryWriter<256> logLines;
-		logLines.WriteString(TC("Hello"));
-		logLines.WriteByte(1);
-
-		ProcessStartInfo psi;
-		psi.application = testApp.data;
-
-		{
-			NetworkClient client(ctorSuccess, { logWriter });
-			CacheClientCreateInfo ccci(logWriter, storageServer, client, session);
-			ccci.useRoots = false;
-			CacheClient cacheClient(ccci);
-
-			if (!client.Connect(tcpBackend, TC("127.0.0.1"), port))
-				return logger.Error(TC("Failed to connect"));
-			auto disconnectClient = MakeGuard([&]() { client.Disconnect(); });
-
-			{
-				CacheResult result;
-				if (cacheClient.FetchFromCache(result, RootPaths(), 0, psi) || result.hit)
-					return false;
-
-				if (!cacheClient.WriteToCache(RootPaths(), 0, psi, inputs.GetData(), inputs.GetPosition(), outputs.GetData(), outputs.GetPosition(), logLines.GetData(), logLines.GetPosition()))
-					return false;
-
-				if (!DeleteFileW(outputFile.data))
-					return false;
-				if (FileExists(logger, outputFile.data))
-					return false;
-				if (!cacheClient.FetchFromCache(result, RootPaths(), 0, psi))
-					return false;
-				if (!FileExists(logger, outputFile.data))
-					return false;
-				if (result.logLines.size() != 1)
-					return false;
-				if (result.logLines[0].text != TC("Hello"))
-					return false;
-			}
-
-			{
-				if (!DeleteFileW(inputFile.data))
-					return false;
-				if (!CreateTextFile(inputFile, logger, workingDir.data, TC("Input.txt"), "Bar"))
-					return false;
-				InvalidateCachedInfo(storageServer, inputFile);
-
-				CacheResult result;
-				if (cacheClient.FetchFromCache(result, RootPaths(), 0, psi) || result.hit)
-					return false;
-
-				if (!cacheClient.WriteToCache(RootPaths(), 0, psi, inputs.GetData(), inputs.GetPosition(), outputs.GetData(), outputs.GetPosition(), logLines.GetData(), logLines.GetPosition()))
-					return false;
-
-				if (!DeleteFileW(outputFile.data))
-					return false;
-				if (FileExists(logger, outputFile.data))
-					return false;
-				if (!cacheClient.FetchFromCache(result, RootPaths(), 0, psi))
-					return false;
-				if (!FileExists(logger, outputFile.data))
-					return false;
-				if (result.logLines.size() != 1)
-					return false;
-				if (result.logLines[0].text != TC("Hello"))
-					return false;
-			}
-		}
-
-		if (!cacheServer.RunMaintenance(true, []() { return false; }))
-			return false;
-
-		{
-			NetworkClient client(ctorSuccess, { logWriter });
-			CacheClientCreateInfo ccci(logWriter, storageServer, client, session);
-			ccci.useRoots = false;
-			CacheClient cacheClient(ccci);
-
-			if (!client.Connect(tcpBackend, TC("127.0.0.1"), port))
-				return logger.Error(TC("Failed to connect"));
-			auto disconnectClient = MakeGuard([&]() { client.Disconnect(); });
-
-			{
-				CacheResult result;
-				if (!cacheClient.FetchFromCache(result, RootPaths(), 0, psi))
-					return false;
-				if (result.logLines.size() != 1)
-					return false;
-				if (result.logLines[0].text != TC("Hello"))
-					return false;
-			}
-		}
-		return true;
-	}
-
 	bool TestCompactPathTable(Logger& logger, const StringBufferBase& rootDir)
 	{
 		for (u32 version=0; version<=CacheBucketVersion; ++version)
@@ -374,4 +223,169 @@
 #endif
 		return true;
 	}
+
+	bool TestCacheClientAndServer(LoggerWithWriter& logger, const StringBufferBase& testRootDir)
+	{
+		LogWriter& logWriter = logger.m_writer;
+		NetworkBackendTcp tcpBackend(logWriter);
+
+		bool ctorSuccess = true;
+		NetworkServer server(ctorSuccess, { logWriter });
+
+		StringBuffer<MaxPath> rootDir;
+		rootDir.Append(testRootDir).Append(TC("Uba"));
+		if (!DeleteAllFiles(logger, rootDir.data))
+			return false;
+
+		StorageServerCreateInfo storageServerInfo(server, rootDir.data, logWriter);
+		storageServerInfo.casCapacityBytes = 1024ull * 1024 * 1024;
+		auto& storageServer = *new StorageServer(storageServerInfo);
+		auto ssg = MakeGuard([&]() { delete &storageServer; });
+
+		{
+		CacheServerCreateInfo csci(storageServer, rootDir.data, logWriter);
+		CacheServer cacheServer(csci);
+		if (!cacheServer.Load())
+			return false;
+
+
+		SessionServerCreateInfo sessionInfo(storageServer, server, logWriter);
+		sessionInfo.rootDir = rootDir.data;
+		auto& session = *new SessionServer(sessionInfo);
+		auto sg = MakeGuard([&]() { delete &session; });
+
+		u16 port = 1356;
+		if (!server.StartListen(tcpBackend, port))
+			return logger.Error(TC("Failed to listen"));
+		auto disconnectServer = MakeGuard([&]() { server.DisconnectClients(); });
+
+		StringBuffer<MaxPath> workingDir;
+		workingDir.Append(testRootDir).Append(TC("WorkingDir"));
+		if (!DeleteAllFiles(logger, workingDir.data))
+			return false;
+		if (!storageServer.CreateDirectory(workingDir.data))
+			return false;
+		if (!DeleteAllFiles(logger, workingDir.data, false))
+			return false;
+
+		StringBuffer<> testApp;
+		GetTestAppPath(logger, testApp);
+
+		StringBuffer<MaxPath> inputFile;
+		if (!CreateTextFile(inputFile, logger, workingDir.data, TC("Input.txt"), "Foo"))
+			return false;
+		StringBuffer<MaxPath> outputFile;
+		if (!CreateTextFile(outputFile, logger, workingDir.data, TC("Output.txt"), "Foo"))
+			return false;
+
+		StackBinaryWriter<256> inputs;
+		inputs.WriteString(inputFile);
+
+		StackBinaryWriter<256> outputs;
+		outputs.WriteString(outputFile);
+
+		StackBinaryWriter<256> logLines;
+		logLines.WriteString(TC("Hello"));
+		logLines.WriteByte(1);
+
+		ProcessStartInfo psi;
+		psi.application = testApp.data;
+
+		{
+			NetworkClient client(ctorSuccess, { logWriter });
+			CacheClientCreateInfo ccci(logWriter, storageServer, client, session);
+			ccci.useRoots = false;
+			CacheClient cacheClient(ccci);
+
+			if (!client.Connect(tcpBackend, TC("127.0.0.1"), port))
+				return logger.Error(TC("Failed to connect"));
+			auto disconnectClient = MakeGuard([&]() { client.Disconnect(); });
+
+			{
+				CacheResult result;
+				if (cacheClient.FetchFromCache(result, RootPaths(), 0, psi) || result.hit)
+					return false;
+
+				if (!cacheClient.WriteToCache(RootPaths(), 0, psi, inputs.GetData(), inputs.GetPosition(), outputs.GetData(), outputs.GetPosition(), logLines.GetData(), logLines.GetPosition()))
+					return false;
+
+				if (!DeleteFileW(outputFile.data))
+					return false;
+				if (FileExists(logger, outputFile.data))
+					return false;
+				if (!cacheClient.FetchFromCache(result, RootPaths(), 0, psi))
+					return false;
+				if (!FileExists(logger, outputFile.data))
+					return false;
+				if (result.logLines.size() != 1)
+					return false;
+				if (result.logLines[0].text != TC("Hello"))
+					return false;
+			}
+
+			{
+				if (!DeleteFileW(inputFile.data))
+					return false;
+				if (!CreateTextFile(inputFile, logger, workingDir.data, TC("Input.txt"), "Bar"))
+					return false;
+				InvalidateCachedInfo(storageServer, inputFile);
+
+				CacheResult result;
+				if (cacheClient.FetchFromCache(result, RootPaths(), 0, psi) || result.hit)
+					return false;
+
+				if (!cacheClient.WriteToCache(RootPaths(), 0, psi, inputs.GetData(), inputs.GetPosition(), outputs.GetData(), outputs.GetPosition(), logLines.GetData(), logLines.GetPosition()))
+					return false;
+
+				if (!DeleteFileW(outputFile.data))
+					return false;
+				if (FileExists(logger, outputFile.data))
+					return false;
+				if (!cacheClient.FetchFromCache(result, RootPaths(), 0, psi))
+					return false;
+				if (!FileExists(logger, outputFile.data))
+					return false;
+				if (result.logLines.size() != 1)
+					return false;
+				if (result.logLines[0].text != TC("Hello"))
+					return false;
+			}
+		}
+
+		if (!cacheServer.RunMaintenance(true, []() { return false; }))
+			return false;
+
+		{
+			NetworkClient client(ctorSuccess, { logWriter });
+			CacheClientCreateInfo ccci(logWriter, storageServer, client, session);
+			ccci.useRoots = false;
+			CacheClient cacheClient(ccci);
+
+			if (!client.Connect(tcpBackend, TC("127.0.0.1"), port))
+				return logger.Error(TC("Failed to connect"));
+			auto disconnectClient = MakeGuard([&]() { client.Disconnect(); });
+
+			{
+				CacheResult result;
+				if (!cacheClient.FetchFromCache(result, RootPaths(), 0, psi))
+					return false;
+				if (result.logLines.size() != 1)
+					return false;
+				if (result.logLines[0].text != TC("Hello"))
+					return false;
+			}
+		}
+
+		if (!cacheServer.Save())
+			return false;
+
+		}
+
+		CacheServerCreateInfo csci(storageServer, rootDir.data, logWriter);
+		CacheServer cacheServer(csci);
+		if (!cacheServer.Load())
+			return false;
+
+		return true;
+	}
 }
