[UBA]
* Fixed bug in common segment logic where segment offset could be 0
* Enabled common segment logic when writing to cache

#skipci

diff --git a/Engine/Source/Programs/UnrealBuildAccelerator/Common/Private/UbaCacheClient.cpp b/Engine/Source/Programs/UnrealBuildAccelerator/Common/Private/UbaCacheClient.cpp
--- a/Engine/Source/Programs/UnrealBuildAccelerator/Common/Private/UbaCacheClient.cpp
+++ b/Engine/Source/Programs/UnrealBuildAccelerator/Common/Private/UbaCacheClient.cpp
@@ -169,8 +169,8 @@ namespace uba
 
 		SCOPED_WRITE_LOCK(m_bucketsLock, bucketsLock);
 		Bucket& bucket = m_buckets.try_emplace(bucketId, bucketId).first->second;
-		//if (!bucket.sendPathTable.GetMemory())
-		//	bucket.sendPathTable.AddCommonStringSegments();
+		if (!bucket.sendPathTable.GetMemory())
+			bucket.sendPathTable.AddCommonStringSegments();
 		bucketsLock.Leave();
 
 		TString qualifiedPath;
diff --git a/Engine/Source/Programs/UnrealBuildAccelerator/Common/Private/UbaCacheServer.cpp b/Engine/Source/Programs/UnrealBuildAccelerator/Common/Private/UbaCacheServer.cpp
--- a/Engine/Source/Programs/UnrealBuildAccelerator/Common/Private/UbaCacheServer.cpp
+++ b/Engine/Source/Programs/UnrealBuildAccelerator/Common/Private/UbaCacheServer.cpp
@@ -1569,7 +1569,7 @@ namespace uba
 				auto findIt = existing.find(path.data);
 				if (findIt == existing.end())
 				{
-					m_logger.Warning(TC("Client %s sent cache entry that already exists but does not match output (output file %s did not exist in existing cache entry)"), clientId, path.data);
+					m_logger.Warning(TC("Client %u sent cache entry that already exists but does not match output (output file %s did not exist in existing cache entry)"), clientId, path.data);
 					cacheEntries.entries.erase(matchingEntry);
 					shouldOverwrite = true;
 					break;
diff --git a/Engine/Source/Programs/UnrealBuildAccelerator/Common/Private/UbaCompactTables.cpp b/Engine/Source/Programs/UnrealBuildAccelerator/Common/Private/UbaCompactTables.cpp
--- a/Engine/Source/Programs/UnrealBuildAccelerator/Common/Private/UbaCompactTables.cpp
+++ b/Engine/Source/Programs/UnrealBuildAccelerator/Common/Private/UbaCompactTables.cpp
@@ -73,14 +73,17 @@ namespace uba
 
 		if (!m_pathTableMem.memory)
 			m_pathTableMem.Init(PathTableMaxSize);
+		if (!m_pathTableMem.writtenSize)
+			m_pathTableMem.AllocateNoLock(1, 1, TC("CompactPathTable"));
 
 		for (const tchar* seg : commonSegments)
 		{
 			u32 segLen = TStrlen(seg);
+			u64 writtenSize = m_pathTableMem.writtenSize;
 			BinaryWriter writer((u8*)m_pathTableMem.AllocateNoLock(segLen+1, 1, TC("CompactPathTable")), 0, segLen+1);
 			StringKey segmentKey = ToStringKeyNoCheck(seg, segLen);
 			u32& strOffset = m_segmentOffsets.Insert(segmentKey);
-			strOffset = u32(m_pathTableMem.writtenSize);
+			strOffset = u32(writtenSize);
 			writer.WriteString(seg, segLen);
 		}
 	}
