[UBA]
* Fixed bug where send/recv graphs get a totalSend/totalRecv of zero causing negative per-second values

#skipci

diff --git a/Engine/Source/Programs/UnrealBuildAccelerator/Visualizer/Private/UbaTraceReader.cpp b/Engine/Source/Programs/UnrealBuildAccelerator/Visualizer/Private/UbaTraceReader.cpp
--- a/Engine/Source/Programs/UnrealBuildAccelerator/Visualizer/Private/UbaTraceReader.cpp
+++ b/Engine/Source/Programs/UnrealBuildAccelerator/Visualizer/Private/UbaTraceReader.cpp
@@ -479,15 +479,19 @@ namespace uba
 				session.prevUpdateTime = session.updates.back().time;
 			}
 
+			if (totalSend < session.prevSend)
+				totalSend = session.prevSend;
+			if (totalRecv < session.prevRecv)
+				totalRecv = session.prevRecv;
+
 			//session.lastPing = lastPing;
 			//session.memAvail = memAvail;
 			session.memTotal = memTotal;
 			session.updates.push_back({time, totalSend, totalRecv, lastPing, memAvail, cpuLoad, connectionCount});
 			
-			u64 send = totalSend - session.prevSend;
-			u64 recv = totalRecv - session.prevRecv;
-			session.highestSendPerS = Max(session.highestSendPerS, float(send)/TimeToS(time - session.prevUpdateTime));
-			session.highestRecvPerS = Max(session.highestRecvPerS, float(recv)/TimeToS(time - session.prevUpdateTime));
+			session.highestSendPerS = Max(session.highestSendPerS, float(totalSend - session.prevSend)/TimeToS(time - session.prevUpdateTime));
+			session.highestRecvPerS = Max(session.highestRecvPerS, float(totalRecv - session.prevRecv)/TimeToS(time - session.prevUpdateTime));
+
 			break;
 		}
 		case TraceType_SessionDisconnect:
