[UBA]
* Fixed so unused memory is used as oodle scratch mem when compressing/decompressing

#skipci

diff --git a/Engine/Source/Programs/UnrealBuildAccelerator/Common/Private/UbaStorage.cpp b/Engine/Source/Programs/UnrealBuildAccelerator/Common/Private/UbaStorage.cpp
--- a/Engine/Source/Programs/UnrealBuildAccelerator/Common/Private/UbaStorage.cpp
+++ b/Engine/Source/Programs/UnrealBuildAccelerator/Common/Private/UbaStorage.cpp
@@ -403,20 +403,28 @@ namespace uba
 			while (left)
 			{
 				u64 uncompressedBlockSize = Min(left, maxUncompressedBlock);
+				
+				void* scratchMem = nullptr;
+				u64 scratchSize = 0;
+
 				if (readMem)
 				{
+					scratchMem = uncompressedData;
+					scratchSize = BufferSlotHalfSize;
 					uncompressedData = readMem + fileSize - left;
 				}
 				else
 				{
 					if (!ReadFile(m_logger, from, readHandle, uncompressedData, uncompressedBlockSize))
 						return false;
+					scratchMem = uncompressedData + uncompressedBlockSize;
+					scratchSize = BufferSlotHalfSize - uncompressedBlockSize;
 				}
 				u8* destBuf = compressBuffer;
 				OO_SINTa compressedBlockSize;
 				{
 					TimerScope cts(stats.compressWrite);
-					compressedBlockSize = OodleLZ_Compress((OodleLZ_Compressor)m_casCompressor, uncompressedData, (OO_SINTa)uncompressedBlockSize, destBuf + 8, (OodleLZ_CompressionLevel)m_casCompressionLevel);
+					compressedBlockSize = OodleLZ_Compress((OodleLZ_Compressor)m_casCompressor, uncompressedData, (OO_SINTa)uncompressedBlockSize, destBuf + 8, (OodleLZ_CompressionLevel)m_casCompressionLevel, NULL, NULL, NULL, scratchMem, scratchSize);
 					if (compressedBlockSize == OODLELZ_FAILED)
 						return m_logger.Error(TC("Failed to compress %llu bytes for %s"), uncompressedBlockSize, from);
 				}
@@ -537,8 +545,10 @@ namespace uba
 				OO_SINTa uncompressedBlockSize = (OO_SINTa)Min(rec->maxUncompressedBlock, rec->fileSize - startOffset);
 				OO_SINTa compressedBlockSize;
 				{
+					void* scratchMem = slot;
+					u64 scratchSize = BufferSlotHalfSize;
 					TimerScope cts(stats.compressWrite);
-					compressedBlockSize = OodleLZ_Compress((OodleLZ_Compressor)m_casCompressor, uncompressedDataSlot, uncompressedBlockSize, compressSlotBuffer + 8, (OodleLZ_CompressionLevel)m_casCompressionLevel);
+					compressedBlockSize = OodleLZ_Compress((OodleLZ_Compressor)m_casCompressor, uncompressedDataSlot, uncompressedBlockSize, compressSlotBuffer + 8, (OodleLZ_CompressionLevel)m_casCompressionLevel, NULL, NULL, NULL, scratchMem, scratchSize);
 					if (compressedBlockSize == OODLELZ_FAILED)
 					{
 						m_logger.Error(TC("Failed to compress %llu bytes for %s"), u64(uncompressedBlockSize), rec->destination->GetFileName());
@@ -756,7 +766,9 @@ namespace uba
 						rec->writePos += decompressedBlockSize;
 						lock.Leave();
 
-						OO_SINTa decompLen = OodleLZ_Decompress(readPos, (OO_SINTa)compressedBlockSize, writePos, (OO_SINTa)decompressedBlockSize); (void)decompLen;
+						void* decoderMem = nullptr;
+						u64 decoderMemSize = 0;
+						OO_SINTa decompLen = OodleLZ_Decompress(readPos, (OO_SINTa)compressedBlockSize, writePos, (OO_SINTa)decompressedBlockSize, OodleLZ_FuzzSafe_Yes, OodleLZ_CheckCRC_No, OodleLZ_Verbosity_None, NULL, 0, NULL, NULL, decoderMem, decoderMemSize);
 						if (decompLen != decompressedBlockSize)
 						{
 							bool f = false;
@@ -808,8 +820,10 @@ namespace uba
 					return m_logger.Error(TC("Decompressed block size %u is invalid. Decompressed file is %u (%s -> %s)"), decompressedBlockSize, decompressedSize, readHint, writeHint);
 				readPos += sizeof(u32) * 2;
 
+				void* decoderMem = nullptr;
+				u64 decoderMemSize = 0;
 				TimerScope ts(stats.decompressToMem);
-				OO_SINTa decompLen = OodleLZ_Decompress(readPos, (OO_SINTa)compressedBlockSize, writePos, (OO_SINTa)decompressedBlockSize); (void)decompLen;
+				OO_SINTa decompLen = OodleLZ_Decompress(readPos, (OO_SINTa)compressedBlockSize, writePos, (OO_SINTa)decompressedBlockSize, OodleLZ_FuzzSafe_Yes, OodleLZ_CheckCRC_No, OodleLZ_Verbosity_None, NULL, 0, NULL, NULL, decoderMem, decoderMemSize);
 				if (decompLen != decompressedBlockSize)
 					return m_logger.Error(TC("Expecting to be able to decompress %u to %u bytes at pos %llu but got %llu. File compressed size: %llu (%s -> %s)"), compressedBlockSize, decompressedBlockSize, decompressedSize - left, decompLen, compressedSize, readHint, writeHint);
 				writePos += decompressedBlockSize;
@@ -2907,6 +2921,9 @@ namespace uba
 			u8* slot = m_bufferSlots.Pop();
 			auto _ = MakeGuard([&]() { m_bufferSlots.Push(slot); });
 
+			void* decoderMem = slot + BufferSlotHalfSize;
+			u64 decoderMemSize = BufferSlotHalfSize;
+
 			u64 bytesRead = 8; // We know size has already been read
 
 			u8* readBuffer = slot;
@@ -2941,7 +2958,7 @@ namespace uba
 				bytesRead += compressedBlockSize;
 
 				TimerScope ts(stats.decompressToMem);
-				OO_SINTa decompLen = OodleLZ_Decompress(readBuffer, (OO_SINTa)compressedBlockSize, writePos, (OO_SINTa)decompressedBlockSize);
+				OO_SINTa decompLen = OodleLZ_Decompress(readBuffer, (OO_SINTa)compressedBlockSize, writePos, (OO_SINTa)decompressedBlockSize, OodleLZ_FuzzSafe_Yes, OodleLZ_CheckCRC_No, OodleLZ_Verbosity_None, NULL, 0, NULL, NULL, decoderMem, decoderMemSize);
 				if (decompLen != decompressedBlockSize)
 					return m_logger.Error(TC("Failed to decompress data from file %s at pos %llu"), fileName, decompressedSize - left);
 				writePos += decompressedBlockSize;
