[UBA]
* Reduced number of kernel calls by embedding compressed cas header in first write

#skipci

diff --git a/Engine/Source/Programs/UnrealBuildAccelerator/Common/Private/UbaStorageUtils.cpp b/Engine/Source/Programs/UnrealBuildAccelerator/Common/Private/UbaStorageUtils.cpp
--- a/Engine/Source/Programs/UnrealBuildAccelerator/Common/Private/UbaStorageUtils.cpp
+++ b/Engine/Source/Programs/UnrealBuildAccelerator/Common/Private/UbaStorageUtils.cpp
@@ -437,7 +437,7 @@ namespace uba
 			writer.WriteBool(false);
 			writer.WriteCasKey(casKey);
 			writer.WriteString(destination);
-			BinaryReader reader(slot, 0, SendMaxSize);
+			BinaryReader reader(slot + (destinationIsCompressed ? sizeof(CompressedObjFileHeader) : 0), 0, SendMaxSize); // Create some space in front of reader to write obj file header there if destination is compressed
 			if (!msg.Send(reader))
 				return logger.Error(TC("Failed to send fetch begin message for cas %s (%s)"), CasKeyString(casKey).str, destination);
 			sizeOfFirstMessage = u32(reader.GetLeft());
@@ -516,13 +516,14 @@ namespace uba
 
 		if (destinationIsCompressed)
 		{
-			CompressedObjFileHeader header { casKey };
-			if (!WriteDestination(&header, sizeof(header)))
-				return false;
-
 			u8* source = slot + BufferSlotHalfSize;
 			u8* lastSource = readBuffer;
 			u64 lastResponseSize = responseSize;
+
+			lastSource -= sizeof(CompressedObjFileHeader);
+			lastResponseSize += sizeof(CompressedObjFileHeader);
+			new (lastSource) CompressedObjFileHeader{ casKey };
+
 			auto writePrev = [&]() { return WriteDestination(lastSource, lastResponseSize); };
 
 			u64 leftCompressed = fileSize - responseSize;
