[UBA]
* Added file UbaApplication.h and moved ExpandEnvironmentVariables to that file. Plan is to add many of the shared application logic in this file (UbaCli, UbaAgent and UbaCacheService has a bunch similar code)

#skipci

diff --git a/Engine/Source/Programs/UnrealBuildAccelerator/Agent/Private/UbaAgent.cpp b/Engine/Source/Programs/UnrealBuildAccelerator/Agent/Private/UbaAgent.cpp
--- a/Engine/Source/Programs/UnrealBuildAccelerator/Agent/Private/UbaAgent.cpp
+++ b/Engine/Source/Programs/UnrealBuildAccelerator/Agent/Private/UbaAgent.cpp
@@ -1,5 +1,6 @@
 // Copyright Epic Games, Inc. All Rights Reserved.
 
+#include "UbaApplication.h"
 #include "UbaAWS.h"
 #include "UbaConfig.h"
 #include "UbaDirectoryIterator.h"
@@ -254,41 +255,6 @@ namespace uba
 	}
 #endif // UBA_AUTO_UPDATE
 
-	bool ExpandEnvironmentVariables(StringBufferBase& str)
-	{
-		StringBuffer<> expandedDir;
-		u64 offset = 0;
-		while (true)
-		{
-			const tchar* begin = str.First('%', offset);
-			if (!begin)
-			{
-				expandedDir.Append(str.data + offset);
-				str.Clear().Append(expandedDir);
-				break;
-			}
-			u64 beginOffset = begin - str.data;
-			const tchar* end = str.First('%', beginOffset + 1);
-			if (!end)
-				return PrintHelp(TC("Missing closing % for environment variable in dir path"));
-			u64 endOffset = end - str.data;
-			StringBuffer<256> var;
-			var.Append(begin + 1, endOffset - beginOffset - 1);
-			StringBuffer<> value;
-			value.count = GetEnvironmentVariableW(var.data, value.data, value.capacity);
-			if (!value.count)
-			{
-				StringBuffer<256> err;
-				err.Appendf(TC("Can't find environment variable %s used in dir path"), var.data);
-				return PrintHelp(err.data);
-			}
-
-			expandedDir.Append(str.data + offset, beginOffset - offset).Append(value);
-			offset = endOffset + 1;
-		}
-		return true;
-	}
-
 	bool IsTerminating(Logger& logger, const tchar* eventFile, StringBufferBase& outReason, u64& outTerminationTimeMs)
 	{
 		if (!*eventFile)
@@ -444,7 +410,7 @@ namespace uba
 			}
 			else if (name.Equals(TC("-maxcpu")))
 			{
-				if (!ExpandEnvironmentVariables(value))
+				if (!ExpandEnvironmentVariables(value, PrintHelp))
 					return false;
 				u32 defaultValue = maxProcessCount;
 				if (!value.Parse(maxProcessCount))
@@ -455,7 +421,7 @@ namespace uba
 			}
 			else if (name.Equals(TC("-mulcpu")))
 			{
-				if (!ExpandEnvironmentVariables(value))
+				if (!ExpandEnvironmentVariables(value, PrintHelp))
 					return false;
 				float defaultValue = mulProcessValue;
 				if (!value.Parse(mulProcessValue))
@@ -482,8 +448,8 @@ namespace uba
 			else if (name.Equals(TC("-config")))
 			{
 				if (value.IsEmpty())
-					return PrintHelp(TC("-dir needs a value"));
-				if (!ExpandEnvironmentVariables(value))
+					return PrintHelp(TC("-config needs a value"));
+				if (!ExpandEnvironmentVariables(value, PrintHelp))
 					return false;
 				configFile.Append(value);
 			}
@@ -549,7 +515,7 @@ namespace uba
 			{
 				if (value.IsEmpty())
 					return PrintHelp(TC("-dir needs a value"));
-				if (!ExpandEnvironmentVariables(value))
+				if (!ExpandEnvironmentVariables(value, PrintHelp))
 					return false;
 				if ((g_rootDir.count = GetFullPathNameW(value.Replace('\\', PathSeparator).data, g_rootDir.capacity, g_rootDir.data, nullptr)) == 0)
 					return PrintHelp(StringBuffer<>().Appendf(TC("-dir has invalid path %s"), value.data).data);
@@ -609,7 +575,7 @@ namespace uba
 			{
 				if (value.IsEmpty())
 					return PrintHelp(TC("-eventfile needs a value"));
-				if (!ExpandEnvironmentVariables(value))
+				if (!ExpandEnvironmentVariables(value, PrintHelp))
 					return false;
 				if ((eventFile.count = GetFullPathNameW(value.Replace('\\', PathSeparator).data, eventFile.capacity, eventFile.data, nullptr)) == 0)
 					return PrintHelp(StringBuffer<>().Appendf(TC("-eventfile has invalid path %s"), value.data).data);
diff --git a/Engine/Source/Programs/UnrealBuildAccelerator/Common/Public/UbaApplication.h b/Engine/Source/Programs/UnrealBuildAccelerator/Common/Public/UbaApplication.h
new file mode 100644
--- /dev/null
+++ b/Engine/Source/Programs/UnrealBuildAccelerator/Common/Public/UbaApplication.h
@@ -0,0 +1,44 @@
+// Copyright Epic Games, Inc. All Rights Reserved.
+
+#pragma once
+
+#include "UbaPlatform.h"
+#include "UbaStringBuffer.h"
+
+namespace uba
+{
+	inline bool ExpandEnvironmentVariables(StringBufferBase& str, const Function<bool(const tchar*)>& errorFunc)
+	{
+		StringBuffer<> expandedDir;
+		u64 offset = 0;
+		while (true)
+		{
+			const tchar* begin = str.First('%', offset);
+			if (!begin)
+			{
+				expandedDir.Append(str.data + offset);
+				str.Clear().Append(expandedDir);
+				break;
+			}
+			u64 beginOffset = begin - str.data;
+			const tchar* end = str.First('%', beginOffset + 1);
+			if (!end)
+				return errorFunc(TC("Missing closing % for environment variable in dir path"));
+			u64 endOffset = end - str.data;
+			StringBuffer<256> var;
+			var.Append(begin + 1, endOffset - beginOffset - 1);
+			StringBuffer<> value;
+			value.count = GetEnvironmentVariableW(var.data, value.data, value.capacity);
+			if (!value.count)
+			{
+				StringBuffer<256> err;
+				err.Appendf(TC("Can't find environment variable %s used in dir path"), var.data);
+				return errorFunc(err.data);
+			}
+
+			expandedDir.Append(str.data + offset, beginOffset - offset).Append(value);
+			offset = endOffset + 1;
+		}
+		return true;
+	}
+}
\ No newline at end of file
