[UBA]
* Moved code that use 4kb stack to separate non-inlined function. This function is only called when there are errors

#skipci

diff --git a/Engine/Source/Programs/UnrealBuildAccelerator/Common/Private/UbaFile.cpp b/Engine/Source/Programs/UnrealBuildAccelerator/Common/Private/UbaFile.cpp
--- a/Engine/Source/Programs/UnrealBuildAccelerator/Common/Private/UbaFile.cpp
+++ b/Engine/Source/Programs/UnrealBuildAccelerator/Common/Private/UbaFile.cpp
@@ -133,6 +133,16 @@ namespace uba
 		return true;
 	}
 
+	UBA_NOINLINE bool ReportOpenFileError(Logger& logger, const tchar* fileName, u32 lastError) // To prevent stack usage
+	{
+		StringBuffer<4096> additionalInfo;
+		#if PLATFORM_WINDOWS
+		if (lastError == ERROR_SHARING_VIOLATION)
+			GetProcessHoldingFile(additionalInfo, fileName);
+		#endif
+		return logger.Error(TC("ERROR opening file %s for read (%s%s)"), fileName, LastErrorToText(lastError).data, additionalInfo.data);
+	}
+
 	bool OpenFileSequentialRead(Logger& logger, const tchar* fileName, FileHandle& outHandle, bool fileNotFoundIsError, bool overlapped)
 	{
 		u32 dwFlagsAndAttributes = DefaultAttributes();
@@ -155,12 +165,7 @@ namespace uba
 				if (IsDirectory(attr))
 					return false;
 		}
-		StringBuffer<4096> additionalInfo;
-		#if PLATFORM_WINDOWS
-		if (lastError == ERROR_SHARING_VIOLATION)
-			GetProcessHoldingFile(additionalInfo, fileName);
-		#endif
-		return logger.Error(TC("ERROR opening file %s for read (%s%s)"), fileName, LastErrorToText(lastError).data, additionalInfo.data);
+		return ReportOpenFileError(logger, fileName, lastError);
 	}
 
 	bool GetFileInformationByHandle(FileInformation& out, Logger& logger, const tchar* fileName, FileHandle hFile)
