[UBA]
* Revived nobuffer file writing. Unfortunately callsite needs improvement since source buffer needs to be aligned to 4k (or at least 16 bytes)... path is disabled behind define

#skipci

diff --git a/Engine/Source/Programs/UnrealBuildAccelerator/Common/Private/UbaFileAccessor.cpp b/Engine/Source/Programs/UnrealBuildAccelerator/Common/Private/UbaFileAccessor.cpp
--- a/Engine/Source/Programs/UnrealBuildAccelerator/Common/Private/UbaFileAccessor.cpp
+++ b/Engine/Source/Programs/UnrealBuildAccelerator/Common/Private/UbaFileAccessor.cpp
@@ -15,6 +15,8 @@
 
 namespace uba
 {
+	constexpr u64 WriteUnit = 4096;
+
 	#if PLATFORM_WINDOWS
 	void GetProcessHoldingFile(StringBufferBase& out, const tchar* fileName);
 	HANDLE asHANDLE(FileHandle fh);
@@ -33,6 +35,15 @@ namespace uba
 			return logger.Error(TC("SetFileInformationByHandle (FileDispositionInfo) failed on %llu %s (%s)"), uintptr_t(handle), fileName, LastErrorToText().data);
 		return true;
 	}
+
+	u64 GetFilePointer(Logger& logger, FileHandle& handle)
+	{
+		LARGE_INTEGER pos;
+		if (SetFilePointerEx(asHANDLE(handle), ToLargeInteger(0), &pos, FILE_CURRENT))
+			return pos.QuadPart;
+		logger.Error(TC("SetFilePointerEx failed (%s)"), LastErrorToText().data);
+		return ~0ull;
+	}
 #endif
 
 	FileAccessor::FileAccessor(Logger& logger, const tchar* fileName)
@@ -170,8 +181,27 @@ namespace uba
 
 		#if UBA_USE_WRITE_THROUGH
 		bool useWriteThrough = true;
-		u8 writeThroughBuffer[4 * 1024];
-		bool setFileSize = false;
+		#endif
+
+
+		#if UBA_USE_WRITE_THROUGH
+		if (useWriteThrough)
+		{
+			if (m_writeThroughBufferPos)
+			{
+				UBA_ASSERT(!((u64)m_fileHandle & OverlappedIoFlag) || !offset);
+				u64 toWrite = Min(WriteUnit - m_writeThroughBufferPos, dataLen);
+				memcpy(m_writeThroughBuffer + m_writeThroughBufferPos, data, toWrite);
+				m_writeThroughBufferPos += u32(toWrite);
+				dataLen -= toWrite;
+				data = (const u8*)data + toWrite;
+				if (m_writeThroughBufferPos < WriteUnit)
+					return true;
+				m_writeThroughBufferPos = 0;
+				if (!Write(m_writeThroughBuffer, WriteUnit, 0))
+					return false;
+			}
+		}
 		#endif
 
 #if PLATFORM_WINDOWS
@@ -224,16 +254,19 @@ namespace uba
 				u64 toActuallyWrite = toWrite;
 
 				#if UBA_USE_WRITE_THROUGH
+				UBA_ASSERT(!m_writeThroughBufferPos);
 				if (useWriteThrough)
 				{
 					if (toWrite < BlockSize)
 					{
-						toActuallyWrite = (toWrite / 4096) * 4096;
+						toActuallyWrite = (toWrite / WriteUnit) * WriteUnit;
 						if (!toActuallyWrite)
 						{
-							memcpy(writeThroughBuffer, pos, toWrite);
-							toActuallyWrite = 4096;
-							setFileSize = true;
+							if (!m_writeThroughBuffer)
+								m_writeThroughBuffer = new u8[WriteUnit];
+							memcpy(m_writeThroughBuffer, pos, toWrite);
+							m_writeThroughBufferPos = u32(toWrite);
+							break;
 						}
 						else
 							toWrite = toActuallyWrite;
@@ -259,14 +292,12 @@ namespace uba
 
 			if (!eg.Execute())
 				return false;
-
-			#if UBA_USE_WRITE_THROUGH
-			if (setFileSize)
-				SetEndOfFile(m_logger, m_fileName, m_fileHandle, dataLen);
-			#endif
-
 			return true;
 		}
+
+		#if UBA_USE_WRITE_THROUGH
+		UBA_ASSERT(!useWriteThrough || (GetFilePointer(m_logger, m_fileHandle) % WriteUnit) == 0);
+		#endif
 #endif
 
 		u64 writeLeft = dataLen;
@@ -277,14 +308,17 @@ namespace uba
 			u32 toActuallyWrite = toWrite;
 
 			#if UBA_USE_WRITE_THROUGH
+			UBA_ASSERT(!m_writeThroughBufferPos);
 			if (useWriteThrough)
 			{
-				toActuallyWrite = (toWrite / 4096) * 4096;
+				toActuallyWrite = (toWrite / WriteUnit) * WriteUnit;
 				if (!toActuallyWrite)
 				{
-					memcpy(writeThroughBuffer, pos, toWrite);
-					toActuallyWrite = 4096;
-					setFileSize = true;
+					if (!m_writeThroughBuffer)
+						m_writeThroughBuffer = new u8[WriteUnit];
+					memcpy(m_writeThroughBuffer, pos, toWrite);
+					m_writeThroughBufferPos = toWrite;
+					return true;
 				}
 			}
 			#endif
@@ -293,8 +327,12 @@ namespace uba
 			DWORD written;
 			if (!::WriteFile(asHANDLE(m_fileHandle), pos, toActuallyWrite, &written, NULL))
 			{
+				#if UBA_USE_WRITE_THROUGH
+				UBA_ASSERT(!useWriteThrough || (u64(pos) % WriteUnit) == 0);
+				#endif
+
 				DWORD lastError = GetLastError();
-				m_logger.Error(TC("ERROR writing file %s writing %u bytes (%llu bytes written out of %llu) (%s)"), m_fileName, toWrite, (dataLen - writeLeft), dataLen, LastErrorToText(lastError).data);
+				m_logger.Error(TC("ERROR writing file %s writing %u bytes (%llu bytes written out of %llu. FilePos: %llu) (%s)"), m_fileName, toActuallyWrite, (dataLen - writeLeft), dataLen, GetFilePointer(m_logger, m_fileHandle), LastErrorToText(lastError).data);
 
 				if (lastError == ERROR_DISK_FULL)
 					ExitProcess(ERROR_DISK_FULL);
@@ -315,11 +353,6 @@ namespace uba
 			pos += written;
 		}
 
-		#if UBA_USE_WRITE_THROUGH
-		if (setFileSize)
-			SetEndOfFile(m_logger, m_fileName, m_fileHandle, dataLen);
-		#endif
-
 		return true;
 	}
 
@@ -366,6 +399,27 @@ namespace uba
 
 	bool FileAccessor::InternalClose(bool success, u64* lastWriteTime)
 	{
+#if UBA_USE_WRITE_THROUGH
+		if (m_writeThroughBufferPos)
+		{
+			u32 bufferPos = m_writeThroughBufferPos;
+			m_writeThroughBufferPos = 0;
+
+			u64 pos = GetFilePointer(m_logger, m_fileHandle);
+			if (pos == ~0ull)
+				return false;
+			if (!Write(m_writeThroughBuffer, WriteUnit))
+				return false;
+			if (!SetEndOfFile(m_logger, m_fileName, m_fileHandle, pos + bufferPos))
+				return false;
+		}
+		if (m_writeThroughBuffer)
+		{
+			delete[] m_writeThroughBuffer;
+			m_writeThroughBuffer = nullptr;
+		}
+#endif
+
 		if (m_data)
 		{
 			if (!UnmapViewOfFile(m_logger, m_data, m_size, m_fileName))
diff --git a/Engine/Source/Programs/UnrealBuildAccelerator/Common/Public/UbaFileAccessor.h b/Engine/Source/Programs/UnrealBuildAccelerator/Common/Public/UbaFileAccessor.h
--- a/Engine/Source/Programs/UnrealBuildAccelerator/Common/Public/UbaFileAccessor.h
+++ b/Engine/Source/Programs/UnrealBuildAccelerator/Common/Public/UbaFileAccessor.h
@@ -47,5 +47,8 @@ namespace uba
 		const tchar* m_tempPath = nullptr;
 		u32 m_tempFileIndex = 0;
 		#endif
+
+		u32 m_writeThroughBufferPos = 0;
+		u8* m_writeThroughBuffer = nullptr;
 	};
 }
\ No newline at end of file
