[UBA]
* Moved logic that decides if obj file should be compressed or not to separate function to reduce duplication and fixed one place that was not properly implemented

#skipci

diff --git a/Engine/Source/Programs/UnrealBuildAccelerator/Core/Private/UbaApplicationRules.cpp b/Engine/Source/Programs/UnrealBuildAccelerator/Core/Private/UbaApplicationRules.cpp
--- a/Engine/Source/Programs/UnrealBuildAccelerator/Core/Private/UbaApplicationRules.cpp
+++ b/Engine/Source/Programs/UnrealBuildAccelerator/Core/Private/UbaApplicationRules.cpp
@@ -61,6 +61,11 @@ namespace uba
 		}
 	};
 
+	bool IsCompressedObjFile(const StringView& fileName)
+	{
+		return fileName.EndsWith(TC(".obj")) && !fileName.Contains(TC("pch."));
+	}
+
 	class ApplicationRulesClExe : public ApplicationRulesVC
 	{
 	public:
@@ -102,7 +107,7 @@ namespace uba
 		{
 			// We don't compress pch obj files because they are used by live-coding.
 			// TODO: Add support in live coding module to load compressed obj files instead
-			return fileName.EndsWith(TC(".obj")) && !fileName.Contains(TC("pch."));
+			return IsCompressedObjFile(fileName);
 		}
 
 		virtual bool ShouldExtractSymbols(const StringView& fileName) const override
@@ -177,7 +182,7 @@ namespace uba
 
 		virtual bool ShouldDecompressFiles(const StringView& fileName) const override
 		{
-			return fileName.IsEmpty() || fileName.EndsWith(TC(".obj"));
+			return fileName.IsEmpty() || IsCompressedObjFile(fileName);
 		}
 	};
 
@@ -276,7 +281,7 @@ namespace uba
 
 		virtual bool StoreFileCompressed(const StringView& fileName) const
 		{
-			return (fileName.EndsWith(TC(".obj")) || fileName.EndsWith(TC(".o"))) && !fileName.Contains(TC("pch."));;
+			return IsCompressedObjFile(fileName) || fileName.EndsWith(TC(".o"));
 		}
 
 		virtual bool ShouldExtractSymbols(const StringView& fileName) const
@@ -575,7 +580,7 @@ namespace uba
 
 		virtual bool ShouldDecompressFiles(const StringView& fileName) const override
 		{
-			return fileName.IsEmpty() || fileName.EndsWith(TC(".o")) || fileName.EndsWith(TC(".obj"));
+			return fileName.IsEmpty() || fileName.EndsWith(TC(".o")) || IsCompressedObjFile(fileName);
 		}
 	};
 
