[UBA]
* Fixed so NetworkBackendTcp options can be controlled from config files

#skipci

diff --git a/Engine/Source/Programs/UnrealBuildAccelerator/Common/Private/UbaNetworkBackendTcp.cpp b/Engine/Source/Programs/UnrealBuildAccelerator/Common/Private/UbaNetworkBackendTcp.cpp
--- a/Engine/Source/Programs/UnrealBuildAccelerator/Common/Private/UbaNetworkBackendTcp.cpp
+++ b/Engine/Source/Programs/UnrealBuildAccelerator/Common/Private/UbaNetworkBackendTcp.cpp
@@ -1,6 +1,7 @@
 // Copyright Epic Games, Inc. All Rights Reserved.
 
 #include "UbaNetworkBackendTcp.h"
+#include "UbaConfig.h"
 #include "UbaEvent.h"
 #include "UbaHash.h"
 #include "UbaPlatform.h"
@@ -97,8 +98,9 @@ namespace uba
 	bool SetKeepAlive(Logger& logger, SOCKET socket);
 	bool SetBlocking(Logger& logger, SOCKET socket, bool blocking);
 	bool SetTimeout(Logger& logger, SOCKET socket, u32 timeoutMs);
+	bool SetRecvBuf(Logger& logger, SOCKET socket, u32 windowSize);
+	bool SetSendBuf(Logger& logger, SOCKET socket, u32 windowSize);
 	bool DisableNagle(Logger& logger, SOCKET socket);
-	bool EnableFastLoopback(Logger& logger, SOCKET socket);
 	bool SendSocket(Logger& logger, SOCKET socket, const void* b, u64 bufferLen);
 	bool RecvSocket(Logger& logger, SOCKET socket, void* b, u32& bufferLen, u32 timeoutMs, const Guid& connection, const tchar* hint1, const tchar* hint2, bool isFirstCall, bool allowLess);
 
@@ -144,11 +146,19 @@ namespace uba
 		return false;
 	}
 
+	void NetworkBackendTcpCreateInfo::Apply(Config& config, const tchar* tableName)
+	{
+		const ConfigTable* tablePtr = config.GetTable(tableName);
+		if (!tablePtr)
+			return;
+		const ConfigTable& table = *tablePtr;
+		table.GetValueAsBool(disableNagle, TC("DisableNagle"));
+	}
 
-
-	NetworkBackendTcp::NetworkBackendTcp(LogWriter& writer, const tchar* prefix)
-		: m_logger(writer, prefix)
+	NetworkBackendTcp::NetworkBackendTcp(const NetworkBackendTcpCreateInfo& info, const tchar* prefix)
+		: m_logger(info.logWriter, prefix)
 	{
+		m_disableNagle = info.disableNagle;
 	}
 
 	NetworkBackendTcp::~NetworkBackendTcp()
@@ -428,13 +438,21 @@ namespace uba
 				break;
 			}
 
-			if (!DisableNagle(logger, clientSocket) || !SetKeepAlive(logger, clientSocket))
+			if (m_disableNagle && !DisableNagle(logger, clientSocket))
 			{
 				CloseSocket(logger, clientSocket, TC("disable nagle"));
 				continue;
 			}
 
-			EnableFastLoopback(logger, clientSocket);
+			if (!SetKeepAlive(logger, clientSocket))
+			{
+				CloseSocket(logger, clientSocket, TC("keep alive"));
+				continue;
+			}
+
+			// Hard coding this values disable auto-tuning. Seems like creating multiple sockets is more efficient than playing around with these
+			//SetRecvBuf(logger, clientSocket, 50*1024*1024);
+			//SetSendBuf(logger, clientSocket, 50*1024*1024);
 
 			SCOPED_WRITE_LOCK(m_connectionsLock, lock);
 			auto it = m_connections.emplace(m_connections.end(), logger, clientSocket);
@@ -674,16 +692,17 @@ namespace uba
 		}
 #endif
 
-		// Socket is good, cancel the socket close scope and break out of the loop.
-		if (!DisableNagle(logger, socketFd))
+		if (m_disableNagle && !DisableNagle(logger, socketFd))
 			return false;
 
 		if (!SetKeepAlive(logger, socketFd))
 			return false;
 
-		EnableFastLoopback(logger, socketFd);
-
+		// Hard coding this values disable auto-tuning. Seems like creating multiple sockets is more efficient than playing around with these
+		//SetRecvBuf(logger, socketFd, 50*1024*1024);
+		//SetSendBuf(logger, socketFd, 50*1024*1024);
 
+		// Socket is good, cancel the socket close scope and break out of the loop.
 		socketClose.Cancel();
 
 		SCOPED_WRITE_LOCK(m_connectionsLock, lock);
@@ -770,30 +789,18 @@ namespace uba
 		return true;
 	}
 
-	bool EnableFastLoopback(Logger& logger, SOCKET socket)
+	bool SetRecvBuf(Logger& logger, SOCKET socket, u32 windowSize)
 	{
-#if 0 // PLATFORM_WINDOWS // Disabled for now because it seems like it is not supported on windows 11
-		static bool mightBeSupported = true;
-		if (!mightBeSupported)
-			return true;
-		int optionValue = 1;
-		DWORD ret = 0;
-		int status = WSAIoctl(socket, SIO_LOOPBACK_FAST_PATH, &optionValue, sizeof(optionValue), NULL, 0, &ret, 0, 0);
-		if (status != SOCKET_ERROR)
-			return true;
-		u32 lastError = GetLastError();
-		if (lastError == WSAEOPNOTSUPP)
-		{
-			mightBeSupported = false;
-			return true;
-		}
-		#if UBA_DEBUG
-		logger.Warning(TC("WSAIoctl SIO_LOOPBACK_FAST_PATH failed (error: %s)"), LastErrorToText(lastError).data);
-		#endif
-		return false;
-#else
+		if (setsockopt(socket, SOL_SOCKET, SO_RCVBUF, (const char*)&windowSize, sizeof(windowSize)) == SOCKET_ERROR)
+			return logger.Error(TC("setsockopt TCP_NODELAY error: (error: %s)"), LastErrorToText(WSAGetLastError()).data);
+		return true;
+	}
+
+	bool SetSendBuf(Logger& logger, SOCKET socket, u32 windowSize)
+	{
+		if (setsockopt(socket, SOL_SOCKET, SO_SNDBUF, (const char*)&windowSize, sizeof(windowSize)) == SOCKET_ERROR)
+			return logger.Error(TC("setsockopt TCP_NODELAY error: (error: %s)"), LastErrorToText(WSAGetLastError()).data);
 		return true;
-#endif
 	}
 
 	bool SetKeepAlive(Logger& logger, SOCKET socket)
diff --git a/Engine/Source/Programs/UnrealBuildAccelerator/Common/Public/UbaNetworkBackendTcp.h b/Engine/Source/Programs/UnrealBuildAccelerator/Common/Public/UbaNetworkBackendTcp.h
--- a/Engine/Source/Programs/UnrealBuildAccelerator/Common/Public/UbaNetworkBackendTcp.h
+++ b/Engine/Source/Programs/UnrealBuildAccelerator/Common/Public/UbaNetworkBackendTcp.h
@@ -8,10 +8,21 @@
 
 namespace uba
 {
+	class Config;
+
+	struct NetworkBackendTcpCreateInfo
+	{
+		NetworkBackendTcpCreateInfo(LogWriter& w = g_consoleLogWriter) : logWriter(w) {}
+		LogWriter& logWriter;
+		bool disableNagle = true;
+
+		void Apply(Config& config, const tchar* tableName = TC("NetworkBackendTcp"));
+	};
+
 	class NetworkBackendTcp : public NetworkBackend
 	{
 	public:
-		NetworkBackendTcp(LogWriter& writer, const tchar* prefix = TC("NetworkBackendTcp"));
+		NetworkBackendTcp(const NetworkBackendTcpCreateInfo& info, const tchar* prefix = TC("NetworkBackendTcp"));
 		virtual ~NetworkBackendTcp();
 		virtual void Shutdown(void* connection) override;
 		virtual bool Send(Logger& logger, void* connection, const void* data, u32 dataSize, SendContext& sendContext) override;
@@ -47,6 +58,8 @@ namespace uba
 		Atomic<u64> m_totalSend;
 		Atomic<u64> m_totalRecv;
 
+		bool m_disableNagle;
+
 		#if PLATFORM_WINDOWS
 		bool m_wsaInitDone = false;
 		#endif
