[UBA]
* Added InitMemory to UbaPlatform.. this sets mimalloc options if mimalloc is enabled. Increased purge delay from 10 seconds to 100 seconds

#skipci

diff --git a/Engine/Source/Programs/UnrealBuildAccelerator/Core/Private/UbaPlatform.cpp b/Engine/Source/Programs/UnrealBuildAccelerator/Core/Private/UbaPlatform.cpp
--- a/Engine/Source/Programs/UnrealBuildAccelerator/Core/Private/UbaPlatform.cpp
+++ b/Engine/Source/Programs/UnrealBuildAccelerator/Core/Private/UbaPlatform.cpp
@@ -614,6 +614,15 @@ namespace uba
 #endif
 	}
 
+	void InitMemory()
+	{
+		#if UBA_USE_MIMALLOC
+		mi_option_set(mi_option_purge_delay, 100);
+		//mi_option_set(mi_option_eager_commit, 1);
+		//mi_option_set(mi_option_purge_decommits, 0);
+		#endif
+	}
+
 
 #if !PLATFORM_WINDOWS
 	void strcpy_s(tchar* dest, u64 destCapacity, const tchar* source)
diff --git a/Engine/Source/Programs/UnrealBuildAccelerator/Core/Public/UbaPlatform.h b/Engine/Source/Programs/UnrealBuildAccelerator/Core/Public/UbaPlatform.h
--- a/Engine/Source/Programs/UnrealBuildAccelerator/Core/Public/UbaPlatform.h
+++ b/Engine/Source/Programs/UnrealBuildAccelerator/Core/Public/UbaPlatform.h
@@ -118,6 +118,7 @@ namespace uba
 	void ElevateCurrentThreadPriority();
 	void PrefetchVirtualMemory(const void* mem, u64 size);
 	void AddExceptionHandler();
+	void InitMemory();
 
 #if PLATFORM_WINDOWS
 	inline constexpr bool CaseInsensitiveFs = true;
