[UBA]
* Fixed so test helper client used in ubacli sets workmanager for its storage
* Improved error message when ubacli fails to decompress file

#skipci

diff --git a/Engine/Source/Programs/UnrealBuildAccelerator/Cli/Private/UbaCli.cpp b/Engine/Source/Programs/UnrealBuildAccelerator/Cli/Private/UbaCli.cpp
--- a/Engine/Source/Programs/UnrealBuildAccelerator/Cli/Private/UbaCli.cpp
+++ b/Engine/Source/Programs/UnrealBuildAccelerator/Cli/Private/UbaCli.cpp
@@ -141,6 +141,7 @@ namespace uba
 		using namespace uba;
 
 		AddExceptionHandler();
+		InitMemory();
 
 		u32 storageCapacityGb = DefaultCapacityGb;
 		StringBuffer<256> workDir;
@@ -687,7 +688,7 @@ namespace uba
 
 				OO_SINTa decompLen = OodleLZ_Decompress(reader.GetPositionData(), (OO_SINTa)compressedBlockSize, destMem, (OO_SINTa)decompressedBlockSize, OodleLZ_FuzzSafe_Yes, OodleLZ_CheckCRC_No, OodleLZ_Verbosity_None, NULL, 0, NULL, NULL, decoderMem, decoredMemSize);
 				if (decompLen != decompressedBlockSize)
-					return logger.Error(TC("Failed to decompress %s"), testDecompress.c_str());
+					return logger.Error(TC("Failed to decompress %s (CompressedSize: %llu DecompressedSize: %llu ReadPos: %llu CompressedBlock: %u DecompressedBlock: %u)"), testDecompress.c_str(), fileSize, decompressedSize, reader.GetPosition(), compressedBlockSize, decompressedBlockSize);
 				destMem += decompressedBlockSize;
 				reader.Skip(compressedBlockSize);
 			}
diff --git a/Engine/Source/Programs/UnrealBuildAccelerator/Cli/Private/UbaClient.h b/Engine/Source/Programs/UnrealBuildAccelerator/Cli/Private/UbaClient.h
--- a/Engine/Source/Programs/UnrealBuildAccelerator/Cli/Private/UbaClient.h
+++ b/Engine/Source/Programs/UnrealBuildAccelerator/Cli/Private/UbaClient.h
@@ -42,6 +42,7 @@ namespace uba
 				clientRootDir.Append("Agent").AppendValue(info.index);
 
 			StorageClientCreateInfo storageClientInfo(*networkClient, clientRootDir.data);
+			storageClientInfo.workManager = networkClient;
 			storageClientInfo.zone = info.zone;
 			storageClientInfo.getProxyBackendCallback = [](void* ud, const tchar* h) -> NetworkBackend& { return ((Client*)ud)->GetProxyBackend(h); };
 			storageClientInfo.getProxyBackendUserData = this;
