[UBA]
* Exposed NetworkClient/Server to config

#skipci

diff --git a/Engine/Source/Programs/UnrealBuildAccelerator/Common/Private/UbaNetworkClient.cpp b/Engine/Source/Programs/UnrealBuildAccelerator/Common/Private/UbaNetworkClient.cpp
--- a/Engine/Source/Programs/UnrealBuildAccelerator/Common/Private/UbaNetworkClient.cpp
+++ b/Engine/Source/Programs/UnrealBuildAccelerator/Common/Private/UbaNetworkClient.cpp
@@ -1,6 +1,7 @@
 // Copyright Epic Games, Inc. All Rights Reserved.
 
 #include "UbaNetworkClient.h"
+#include "UbaConfig.h"
 #include "UbaCrypto.h"
 #include "UbaBinaryReaderWriter.h"
 #include "UbaNetworkBackendTcp.h"
@@ -10,6 +11,15 @@
 
 namespace uba
 {
+	void NetworkClientCreateInfo::Apply(Config& config, const tchar* tableName)
+	{
+		const ConfigTable* tablePtr = config.GetTable(tableName);
+		if (!tablePtr)
+			return;
+		const ConfigTable& table = *tablePtr;
+		table.GetValueAsU32(desiredConnectionCount, TC("DesiredConnectionCount"));
+	}
+
 	NetworkClient::NetworkClient(bool& outCtorSuccess, const NetworkClientCreateInfo& info, const tchar* name)
 	: WorkManagerImpl(info.workerCount == 0 ? GetLogicalProcessorCount() : info.workerCount, TC("UbaWorker(NetworkClient)"))
 	,	m_logWriter(info.logWriter)
@@ -25,6 +35,7 @@ namespace uba
 			m_logger.Detail(TC("Adjusted msg size to %u to stay inside limits"), fixedSendSize);
 		m_sendSize = fixedSendSize;
 		m_receiveTimeoutSeconds = info.receiveTimeoutSeconds;
+		m_desiredConnectionCount = info.desiredConnectionCount;
 
 		m_connectionsIt = m_connections.end();
 
@@ -236,7 +247,8 @@ namespace uba
 		NetworkMessage* msg = client.m_activeMessages[messageId];
 		lock.Leave();
 
-		UBA_ASSERT(msg);
+		if (!msg)
+			return false;
 
 		constexpr u32 ErrorSize = 0xffffff - ReceiveHeaderSize; // ReceiveHeaderSize is removed from size in server send
 
diff --git a/Engine/Source/Programs/UnrealBuildAccelerator/Common/Private/UbaNetworkServer.cpp b/Engine/Source/Programs/UnrealBuildAccelerator/Common/Private/UbaNetworkServer.cpp
--- a/Engine/Source/Programs/UnrealBuildAccelerator/Common/Private/UbaNetworkServer.cpp
+++ b/Engine/Source/Programs/UnrealBuildAccelerator/Common/Private/UbaNetworkServer.cpp
@@ -7,7 +7,7 @@
 
 namespace uba
 {
-	void NetworkServerCreateInfo::Apply(Config& config)
+	void NetworkServerCreateInfo::Apply(Config& config, const tchar* tableName)
 	{
 	}
 
diff --git a/Engine/Source/Programs/UnrealBuildAccelerator/Common/Public/UbaNetworkClient.h b/Engine/Source/Programs/UnrealBuildAccelerator/Common/Public/UbaNetworkClient.h
--- a/Engine/Source/Programs/UnrealBuildAccelerator/Common/Public/UbaNetworkClient.h
+++ b/Engine/Source/Programs/UnrealBuildAccelerator/Common/Public/UbaNetworkClient.h
@@ -11,6 +11,7 @@
 
 namespace uba
 {
+	class Config;
 	class NetworkBackend;
 	class NetworkBackendTcp;
 	struct BinaryReader;
@@ -25,7 +26,9 @@ namespace uba
 		u32 sendSize = SendDefaultSize;
 		u32 receiveTimeoutSeconds = 0;
 		u32 workerCount = 0; // Zero means logical processor count will be used
+		u32 desiredConnectionCount = 1; // Desired connection count to server when connected
 		const u8* cryptoKey128 = nullptr;
+		void Apply(Config& config, const tchar* tableName = TC("NetworkClient"));
 	};
 
 	class NetworkClient : public WorkManagerImpl
@@ -64,6 +67,7 @@ namespace uba
 		u32 GetConnectionCount() { return m_connectionCount; }
 		u64 GetTotalSentBytes() { return m_sendBytes; }
 		u64 GetTotalRecvBytes() { return m_recvBytes; }
+		u32 GetDesiredConnectionCount() { return m_desiredConnectionCount; } 
 
 		NetworkBackend* GetFirstConnectionBackend();
 
@@ -94,6 +98,7 @@ namespace uba
 		LoggerWithWriter m_logger;
 		u32 m_sendSize;
 		u32 m_receiveTimeoutSeconds;
+		u32 m_desiredConnectionCount;
 		Atomic<u64> m_sendBytes;
 		Atomic<u64> m_recvBytes;
 		Atomic<u32> m_recvCount;
diff --git a/Engine/Source/Programs/UnrealBuildAccelerator/Common/Public/UbaNetworkServer.h b/Engine/Source/Programs/UnrealBuildAccelerator/Common/Public/UbaNetworkServer.h
--- a/Engine/Source/Programs/UnrealBuildAccelerator/Common/Public/UbaNetworkServer.h
+++ b/Engine/Source/Programs/UnrealBuildAccelerator/Common/Public/UbaNetworkServer.h
@@ -30,7 +30,7 @@ namespace uba
 	{
 		NetworkServerCreateInfo(LogWriter& w = g_consoleLogWriter) : logWriter(w) {}
 
-		void Apply(Config& config);
+		void Apply(Config& config, const tchar* tableName = TC("NetworkServer"));
 
 		LogWriter& logWriter;
 		u32 workerCount = 0; // Zero means it will use the number of logical cores as worker count
