[UBA]
* Fixed so UbaCacheService can be tweaked with config file. Set -config=<file> on command line to apply config

#skipci

diff --git a/Engine/Source/Programs/UnrealBuildAccelerator/CacheService/Private/UbaCacheService.cpp b/Engine/Source/Programs/UnrealBuildAccelerator/CacheService/Private/UbaCacheService.cpp
--- a/Engine/Source/Programs/UnrealBuildAccelerator/CacheService/Private/UbaCacheService.cpp
+++ b/Engine/Source/Programs/UnrealBuildAccelerator/CacheService/Private/UbaCacheService.cpp
@@ -1,6 +1,8 @@
 // Copyright Epic Games, Inc. All Rights Reserved.
 
 #include "UbaCacheServer.h"
+#include "UbaApplication.h"
+#include "UbaConfig.h"
 #include "UbaFile.h"
 #include "UbaHttpServer.h"
 #include "UbaNetworkBackendTcp.h"
@@ -48,6 +50,7 @@ namespace uba
 		logger.Info(TC("  -port=[<host>:]<port>   The ip/name and port (default: %u) to listen for clients on"), DefaultCachePort);
 		logger.Info(TC("  -capacity=<gigaby>      Capacity of local store. Defaults to %u gigabytes"), DefaultCapacityGb);
 		logger.Info(TC("  -expiration=<seconds>   Time until unused cache entries get deleted. Defaults to %s (%u seconds)"), TimeToText(MsToTime(DefaultExpiration*1000)).str, DefaultExpiration);
+		logger.Info(TC("  -config=<file>          Config file that contains options for various systems"));
 		logger.Info(TC("  -http=<port>            If set, a http server will be started and listen on <port>"));
 		logger.Info(TC("  -fullmaintenance        Force a full maintenance"));
 		logger.Info(TC("  -nomaintenance          Skip all maintenance"));
@@ -108,6 +111,7 @@ namespace uba
 		bool fullMaintenance = false;
 		bool maintenanceEnabled = true;
 		u32 expirationTimeSeconds = DefaultExpiration;
+		TString configFile;
 
 		for (int i=1; i!=argc; ++i)
 		{
@@ -169,6 +173,14 @@ namespace uba
 			{
 				maintenanceEnabled = false;
 			}
+			else if (name.Equals(TC("-config")))
+			{
+				if (value.IsEmpty())
+					return PrintHelp(TC("-config needs a value"));
+				if (!ExpandEnvironmentVariables(value, PrintHelp))
+					return false;
+				configFile = value.data;
+			}
 			else if (name.Equals(TC("-?")))
 			{
 				return PrintHelp(TC(""));
@@ -189,6 +201,11 @@ namespace uba
 		g_exitLock->LeaveWrite();
 		auto glg = MakeGuard([]() { g_exitLock->EnterWrite(); g_logger = nullptr; g_exitLock->LeaveWrite(); });
 
+		Config config;
+		if (!configFile.empty())
+			config.LoadFromFile(logger, configFile.data());
+
+
 		const tchar* dbgStr = TC("");
 		#if UBA_DEBUG
 		dbgStr = TC(" (DEBUG)");
@@ -224,8 +241,12 @@ namespace uba
 		signal(SIGTERM, ConsoleHandler);
 		#endif
 
-		NetworkBackendTcp networkBackend(logWriter);
+		NetworkBackendTcpCreateInfo nbtci(logWriter);
+		nbtci.Apply(config);
+		NetworkBackendTcp networkBackend(nbtci);
+
 		NetworkServerCreateInfo nsci(logWriter);
+		nsci.Apply(config);
 		//nsci.workerCount = 4;
 		nsci.logConnections = false; // Let the cache server report instead
 		nsci.receiveTimeoutSeconds = 2*60*60; // Two hours timeout
@@ -235,6 +256,7 @@ namespace uba
 			return -1;
 
 		StorageServerCreateInfo storageInfo(networkServer, g_rootDir.data, logWriter);
+		storageInfo.Apply(config);
 		storageInfo.casCapacityBytes = storageCapacity;
 		storageInfo.storeCompressed = storeCompressed;
 		storageInfo.allowHintAsFallback = false;
@@ -246,6 +268,7 @@ namespace uba
 			return -1;
 
 		CacheServerCreateInfo cacheInfo(storageServer, g_rootDir.data, logWriter);
+		cacheInfo.Apply(config);
 		cacheInfo.expirationTimeSeconds = expirationTimeSeconds;
 		cacheInfo.maintenanceReserveSize = maintenanceReserveSizeMb * 1024 * 1024;
 		CacheServer cacheServer(cacheInfo);
