[UBA]
* Fixed spelling error

#skipci

diff --git a/Engine/Source/Programs/UnrealBuildAccelerator/Detours/Private/Windows/UbaDetoursFileDecompressor.h b/Engine/Source/Programs/UnrealBuildAccelerator/Detours/Private/Windows/UbaDetoursFileDecompressor.h
--- a/Engine/Source/Programs/UnrealBuildAccelerator/Detours/Private/Windows/UbaDetoursFileDecompressor.h
+++ b/Engine/Source/Programs/UnrealBuildAccelerator/Detours/Private/Windows/UbaDetoursFileDecompressor.h
@@ -327,8 +327,8 @@ namespace uba
 
 		void DecompressWork(Preload*& it, u8* totalMem, u64& totalMemOffset, bool allowFreeOnClose)
 		{
-			OO_SINTa decoredMemSize = OodleLZDecoder_MemorySizeNeeded(OodleLZ_Compressor_Kraken);
-			void* decoderMem = malloc(decoredMemSize);
+			OO_SINTa decoderMemSize = OodleLZDecoder_MemorySizeNeeded(OodleLZ_Compressor_Kraken);
+			void* decoderMem = malloc(decoderMemSize);
 			auto mg = MakeGuard([decoderMem]() { free(decoderMem); });
 
 			while (true)
@@ -365,7 +365,7 @@ namespace uba
 				lock.Leave();
 
 				OO_SINTa decompLen = OodleLZ_Decompress(reader.GetPositionData(), (OO_SINTa)compressedBlockSize, destMem, (OO_SINTa)decompressedBlockSize,
-					OodleLZ_FuzzSafe_Yes, OodleLZ_CheckCRC_No, OodleLZ_Verbosity_None, NULL, 0, NULL, NULL, decoderMem, decoredMemSize);
+					OodleLZ_FuzzSafe_Yes, OodleLZ_CheckCRC_No, OodleLZ_Verbosity_None, NULL, 0, NULL, NULL, decoderMem, decoderMemSize);
 				if (decompLen != decompressedBlockSize)
 				{
 					FatalError(1356, TC("Failed to decompress .obj file %s (Expected to get %llu but got %llu)"), info.originalName, u64(decompressedBlockSize), u64(decompLen));
