[UBA]
* Added overlapped I/O test

#skipci

diff --git a/Engine/Source/Programs/UnrealBuildAccelerator/Test/Private/UbaTestAll.h b/Engine/Source/Programs/UnrealBuildAccelerator/Test/Private/UbaTestAll.h
--- a/Engine/Source/Programs/UnrealBuildAccelerator/Test/Private/UbaTestAll.h
+++ b/Engine/Source/Programs/UnrealBuildAccelerator/Test/Private/UbaTestAll.h
@@ -33,6 +33,7 @@ namespace uba
 		UBA_TEST(TestEvents) \
 		UBA_TEST(TestPaths) \
 		UBA_TEST(TestFiles) \
+		UBA_TEST(TestOverlappedIO) \
 		UBA_TEST(TestMemoryBlock) \
 		UBA_TEST(TestParseArguments) \
 		UBA_TEST(TestBinaryWriter) \
diff --git a/Engine/Source/Programs/UnrealBuildAccelerator/Test/Private/UbaTestBasics.h b/Engine/Source/Programs/UnrealBuildAccelerator/Test/Private/UbaTestBasics.h
--- a/Engine/Source/Programs/UnrealBuildAccelerator/Test/Private/UbaTestBasics.h
+++ b/Engine/Source/Programs/UnrealBuildAccelerator/Test/Private/UbaTestBasics.h
@@ -257,6 +257,30 @@ namespace uba
 		return true;
 	}
 
+	bool TestOverlappedIO(Logger& logger, const StringBufferBase& rootDir)
+	{
+		StringBuffer<> testFileName(rootDir);
+		testFileName.Append(TC("UbaTestFile"));
+
+		u64 left = 2*1024*1024;
+
+		FileAccessor fileHandle(logger, testFileName.data);
+		if (!fileHandle.CreateWrite(false, DefaultAttributes() | FILE_FLAG_OVERLAPPED, left))
+			return logger.Error(TC("Failed to create file for write"));
+		u8 buffer[512*1024];
+
+		while (left)
+		{
+			u64 toWrite = Min(left, 277872ull);
+			if (!fileHandle.Write(buffer, toWrite))
+				return logger.Error(TC("Failed to create file for write"));
+			left -= toWrite;
+		}
+		if (!fileHandle.Close())
+			return false;
+		return true;
+	}
+
 	bool TestMemoryBlock(Logger& logger, const StringBufferBase& rootDir)
 	{
 		{
