[UBA]
* Fixed shutdown hang/crash that could happen if ctrl-c is pressed while UBT is in the middle of downloading cache entries

#skipci

diff --git a/Engine/Source/Programs/UnrealBuildAccelerator/Host/Private/UbaExports.cpp b/Engine/Source/Programs/UnrealBuildAccelerator/Host/Private/UbaExports.cpp
--- a/Engine/Source/Programs/UnrealBuildAccelerator/Host/Private/UbaExports.cpp
+++ b/Engine/Source/Programs/UnrealBuildAccelerator/Host/Private/UbaExports.cpp
@@ -76,6 +76,20 @@ namespace uba
 		LoggerWithWriter logger;
 	};
 
+	class CacheClientWithCounter : public CacheClient
+	{
+	public:
+		CacheClientWithCounter(const CacheClientCreateInfo& info) : CacheClient(info) {}
+		Atomic<u32> active;
+	};
+
+	struct CacheClientActiveScope
+	{
+		CacheClientActiveScope(CacheClient* c) : client(*(CacheClientWithCounter*)c) { ++client.active; }
+		~CacheClientActiveScope() { --client.active; }
+		CacheClientWithCounter& client;
+	};
+
 	Config& GetConfig(const tchar* fileName = nullptr)
 	{
 		static Config config;
@@ -105,6 +119,7 @@ extern "C"
 	uba::LogWriter* CreateCallbackLogWriter(uba::CallbackLogWriter::BeginScopeCallback begin, uba::CallbackLogWriter::EndScopeCallback end, uba::CallbackLogWriter::LogCallback log)
 	{
 		uba::AddExceptionHandler();
+		uba::InitMemory();
 
 		return new uba::CallbackLogWriter(begin, end, log);
 	}
@@ -136,7 +151,11 @@ extern "C"
 			networkBackend = new NetworkBackendQuic(writer);
 		else
 		#endif
-			networkBackend = new NetworkBackendTcp(writer);
+		{
+			NetworkBackendTcpCreateInfo info(writer);
+			info.Apply(GetConfig());
+			networkBackend = new NetworkBackendTcp(info);
+		}
 
 		NetworkServerCreateInfo info(writer);
 
@@ -601,6 +620,7 @@ extern "C"
 		NetworkClientCreateInfo ncci(writer);
 		ncci.receiveTimeoutSeconds = 60;
 		ncci.cryptoKey128 = crypto128;
+		ncci.Apply(GetConfig(), TC("CacheNetworkClient"));
 		bool ctorSuccess = false;
 		auto networkClient = new NetworkClientWithBackend(ctorSuccess, ncci, server.backend, TC("UbaCache"));
 		if (!ctorSuccess)
@@ -621,16 +641,23 @@ extern "C"
 			info.hint = hint;
 
 		info.reportMissReason |= reportMissReason;
-		return new CacheClient(info);
+		return new CacheClientWithCounter(info);
 	}
 
 	bool CacheClient_Connect(uba::CacheClient* cacheClient, const uba::tchar* host, int port)
 	{
 		using namespace uba;
 		auto& networkClient = (NetworkClientWithBackend&)cacheClient->GetClient();
+		CacheClientActiveScope ccas(cacheClient);
 
+		u32 desired = networkClient.GetDesiredConnectionCount();
+		if (!desired)
+			return false;
 		if (!networkClient.Connect(*networkClient.backend, host, u16(port)))
 			return false;
+		while (--desired)
+			networkClient.Connect(*networkClient.backend, host, u16(port));
+
 		cacheClient->GetStorage().LoadCasTable();
 		return true;
 	}
@@ -649,17 +676,21 @@ extern "C"
 			logLinesWriter.WriteByte(line.type);
 		}
 
+		CacheClientActiveScope ccas(cacheClient);
 		return cacheClient->WriteToCache(*rootPaths, bucket, process->GetStartInfo(), inputs, inputsSize, outputs, outputsSize, logLinesWriter.GetData(), logLinesWriter.GetPosition(), process->GetId());
 	}
 
 	bool CacheClient_WriteToCache2(uba::CacheClient* cacheClient, uba::u32 bucket, const uba::ProcessHandle* process, const uba::u8* inputs, uba::u32 inputsSize, const uba::u8* outputs, uba::u32 outputsSize, const uba::u8* logLines, uba::u32 logLinesSize)
 	{
+		using namespace uba;
+		CacheClientActiveScope ccas(cacheClient);
 		return cacheClient->WriteToCache(bucket, process->GetStartInfo(), inputs, inputsSize, outputs, outputsSize, logLines, logLinesSize, process->GetId());
 	}
 
 	uba::u32 CacheClient_FetchFromCache(uba::CacheClient* cacheClient, uba::RootPaths* rootPaths, uba::u32 bucket, const uba::ProcessStartInfo& info)
 	{
 		using namespace uba;
+		CacheClientActiveScope ccas(cacheClient);
 		CacheResult cacheResult;
 		bool res = cacheClient->FetchFromCache(cacheResult, *rootPaths, bucket, info);
 		return (res && cacheResult.hit) ? 1 : 0;
@@ -669,6 +700,7 @@ extern "C"
 	{
 		using namespace uba;
 		auto cacheResult = new CacheResult();
+		CacheClientActiveScope ccas(cacheClient);
 		if (cacheClient->FetchFromCache(*cacheResult, *rootPaths, bucket, info))
 			return cacheResult;
 		delete cacheResult;
@@ -679,6 +711,7 @@ extern "C"
 	{
 		using namespace uba;
 		auto cacheResult = new CacheResult();
+		CacheClientActiveScope ccas(cacheClient);
 		if (cacheClient->FetchFromCache(*cacheResult, rootsHandle, bucket, info))
 			return cacheResult;
 		delete cacheResult;
@@ -700,6 +733,10 @@ extern "C"
 
 		auto& networkClient = (NetworkClientWithBackend&)cacheClient->GetClient();
 		networkClient.Disconnect();
+
+		while (((CacheClientWithCounter*)cacheClient)->active)
+			uba::Sleep(10);
+
 		delete cacheClient;
 		delete &networkClient;
 	}
