[UBA]
* Added FetchPathTable2 and FetchCasTable2 to cache server. These messages does not write size in order to be able to write result directly into table on client

#skipci

diff --git a/Engine/Source/Programs/UnrealBuildAccelerator/Common/Private/UbaCacheServer.cpp b/Engine/Source/Programs/UnrealBuildAccelerator/Common/Private/UbaCacheServer.cpp
--- a/Engine/Source/Programs/UnrealBuildAccelerator/Common/Private/UbaCacheServer.cpp
+++ b/Engine/Source/Programs/UnrealBuildAccelerator/Common/Private/UbaCacheServer.cpp
@@ -187,6 +187,10 @@ namespace uba
 		}
 	}
 
+	void CacheServerCreateInfo::Apply(Config& config)
+	{
+	}
+
 	CacheServer::CacheServer(const CacheServerCreateInfo& info)
 	:	m_logger(info.logWriter, TC("UbaCacheServer"))
 	,	m_server(info.storage.GetServer())
@@ -1419,9 +1423,15 @@ namespace uba
 		case CacheMessageType_FetchPathTable:
 			return HandleFetchPathTable(reader, writer);
 
+		case CacheMessageType_FetchPathTable2:
+			return HandleFetchPathTable2(reader, writer);
+
 		case CacheMessageType_FetchCasTable:
 			return HandleFetchCasTable(reader, writer);
 
+		case CacheMessageType_FetchCasTable2:
+			return HandleFetchCasTable2(reader, writer);
+
 		case CacheMessageType_FetchEntries:
 			return HandleFetchEntries(reader, writer, connectionInfo.GetId());
 
@@ -1690,6 +1700,16 @@ namespace uba
 		return true;
 	}
 
+	bool CacheServer::HandleFetchPathTable2(BinaryReader& reader, BinaryWriter& writer)
+	{
+		Bucket& bucket = GetBucket(reader);
+		u32 haveSize = reader.ReadU32();
+		u32 size = bucket.m_pathTable.GetSize();
+		u32 toSend = Min(u32(writer.GetCapacityLeft()), size - haveSize);
+		writer.WriteBytes(bucket.m_casKeyTable.GetMemory() + haveSize, toSend);
+		return true;
+	}
+
 	bool CacheServer::HandleFetchCasTable(BinaryReader& reader, BinaryWriter& writer)
 	{
 		Bucket& bucket = GetBucket(reader);
@@ -1701,6 +1721,16 @@ namespace uba
 		return true;
 	}
 
+	bool CacheServer::HandleFetchCasTable2(BinaryReader& reader, BinaryWriter& writer)
+	{
+		Bucket& bucket = GetBucket(reader);
+		u32 haveSize = reader.ReadU32();
+		u32 size = bucket.m_casKeyTable.GetSize();
+		u32 toSend = Min(u32(writer.GetCapacityLeft()), size - haveSize);
+		writer.WriteBytes(bucket.m_casKeyTable.GetMemory() + haveSize, toSend);
+		return true;
+	}
+
 	bool CacheServer::HandleFetchEntries(BinaryReader& reader, BinaryWriter& writer, u32 clientId)
 	{
 		Bucket& bucket = GetBucket(reader);
diff --git a/Engine/Source/Programs/UnrealBuildAccelerator/Common/Public/UbaCacheServer.h b/Engine/Source/Programs/UnrealBuildAccelerator/Common/Public/UbaCacheServer.h
--- a/Engine/Source/Programs/UnrealBuildAccelerator/Common/Public/UbaCacheServer.h
+++ b/Engine/Source/Programs/UnrealBuildAccelerator/Common/Public/UbaCacheServer.h
@@ -9,6 +9,7 @@
 namespace uba
 {
 	class CompactCasKeyTable;
+	class Config;
 	class NetworkServer;
 	class StorageServer;
 	struct BinaryReader;
@@ -42,6 +43,8 @@ namespace uba
 
 		// Max size of cas bucket. When within 2mb it will start decreasing expiry time by one hour
 		u64 bucketCasTableMaxSize = 32ull * 1024 * 1024;
+
+		void Apply(Config& config);
 	};
 
 	class CacheServer
@@ -81,6 +84,8 @@ namespace uba
 		bool HandleStoreEntry(ConnectionBucket& bucket, BinaryReader& reader, BinaryWriter& writer, u32 clientVersion, u32 clientId);
 		bool HandleFetchPathTable(BinaryReader& reader, BinaryWriter& writer);
 		bool HandleFetchCasTable(BinaryReader& reader, BinaryWriter& writer);
+		bool HandleFetchPathTable2(BinaryReader& reader, BinaryWriter& writer);
+		bool HandleFetchCasTable2(BinaryReader& reader, BinaryWriter& writer);
 		bool HandleFetchEntries(BinaryReader& reader, BinaryWriter& writer, u32 clientId);
 		bool HandleReportUsedEntry(BinaryReader& reader, BinaryWriter& writer, u32 clientId);
 		bool HandleExecuteCommand(BinaryReader& reader, BinaryWriter& writer);
diff --git a/Engine/Source/Programs/UnrealBuildAccelerator/Common/Public/UbaNetwork.h b/Engine/Source/Programs/UnrealBuildAccelerator/Common/Public/UbaNetwork.h
--- a/Engine/Source/Programs/UnrealBuildAccelerator/Common/Public/UbaNetwork.h
+++ b/Engine/Source/Programs/UnrealBuildAccelerator/Common/Public/UbaNetwork.h
@@ -120,6 +120,8 @@ namespace uba
 		UBA_CACHE_MESSAGE(ExecuteCommand) \
 		UBA_CACHE_MESSAGE(RequestShutdown) \
 		UBA_CACHE_MESSAGE(ReportUsedEntry) \
+		UBA_CACHE_MESSAGE(FetchPathTable2) \
+		UBA_CACHE_MESSAGE(FetchCasTable2) \
 
 	enum CacheMessageType : u8
 	{
