[UBA]
* Added so CacheClient config table name can be set from outside
* Fixed so cache client is not logging out lots of spam if failing to download because of disconnect

#skipci

diff --git a/Engine/Source/Programs/UnrealBuildAccelerator/Common/Private/UbaCacheClient.cpp b/Engine/Source/Programs/UnrealBuildAccelerator/Common/Private/UbaCacheClient.cpp
--- a/Engine/Source/Programs/UnrealBuildAccelerator/Common/Private/UbaCacheClient.cpp
+++ b/Engine/Source/Programs/UnrealBuildAccelerator/Common/Private/UbaCacheClient.cpp
@@ -20,9 +20,9 @@
 
 namespace uba
 {
-	void CacheClientCreateInfo::Apply(Config& config)
+	void CacheClientCreateInfo::Apply(Config& config, const tchar* tableName)
 	{
-		const ConfigTable* tablePtr = config.GetTable(TC("CacheClient"));
+		const ConfigTable* tablePtr = config.GetTable(tableName);
 		if (!tablePtr)
 			return;
 		const ConfigTable& table = *tablePtr;
@@ -679,8 +679,9 @@ namespace uba
 
 				struct DowngradedLogger : public LoggerWithWriter
 				{
-					DowngradedLogger(LogWriter& writer, const tchar* prefix) : LoggerWithWriter(writer, prefix) {}
-					virtual void Log(LogEntryType type, const tchar* str, u32 strLen) override { LoggerWithWriter::Log(Max(type, LogEntryType_Info), str, strLen); }
+					DowngradedLogger(Atomic<bool>& c, LogWriter& writer, const tchar* prefix) : LoggerWithWriter(writer, prefix), connected(c) {}
+					virtual void Log(LogEntryType type, const tchar* str, u32 strLen) override { if (connected) LoggerWithWriter::Log(Max(type, LogEntryType_Info), str, strLen); }
+					Atomic<bool>& connected;
 				};
 
 				while (outputsReader.GetLeft())
@@ -702,7 +703,7 @@ namespace uba
 
 					if (IsNormalized(casKey))
 					{
-						DowngradedLogger logger(m_logger.m_writer, TC("UbaCacheClientNormalizedDownload"));
+						DowngradedLogger logger(m_connected, m_logger.m_writer, TC("UbaCacheClientNormalizedDownload"));
 						// Fetch into memory, file is in special format without absolute paths
 						MemoryBlock normalizedBlock(4*1024*1024);
 						bool destinationIsCompressed = false;
@@ -773,7 +774,7 @@ namespace uba
 					}
 					else
 					{
-						DowngradedLogger logger(m_logger.m_writer, TC("UbaCacheClientDownload"));
+						DowngradedLogger logger(m_connected, m_logger.m_writer, TC("UbaCacheClientDownload"));
 						bool destinationIsCompressed = IsFileCompressed(info, path, true);
 						if (!fetcher.RetrieveFile(logger, m_client, casKey, path.data, destinationIsCompressed))
 							return logger.Error(TC("Failed to download cache output %s for %s"), path.data, info.GetDescription());
diff --git a/Engine/Source/Programs/UnrealBuildAccelerator/Common/Public/UbaCacheClient.h b/Engine/Source/Programs/UnrealBuildAccelerator/Common/Public/UbaCacheClient.h
--- a/Engine/Source/Programs/UnrealBuildAccelerator/Common/Public/UbaCacheClient.h
+++ b/Engine/Source/Programs/UnrealBuildAccelerator/Common/Public/UbaCacheClient.h
@@ -29,7 +29,7 @@ namespace uba
 		NetworkClient& client;
 		Session& session;
 
-		void Apply(Config& config);
+		void Apply(Config& config, const tchar* tableName = TC("CacheClient"));
 
 		bool reportCacheKey = false;
 		bool reportMissReason = false; // Report the reason no matching cache entry was found
