[UBA]
* Moved client storage file creation when fetching cas file from server to happen after messages has been sent but before they are received.. reduces latency a bit since file creation can take quite some time.

#skipci

diff --git a/Engine/Source/Programs/UnrealBuildAccelerator/Common/Private/UbaStorageClient.cpp b/Engine/Source/Programs/UnrealBuildAccelerator/Common/Private/UbaStorageClient.cpp
--- a/Engine/Source/Programs/UnrealBuildAccelerator/Common/Private/UbaStorageClient.cpp
+++ b/Engine/Source/Programs/UnrealBuildAccelerator/Common/Private/UbaStorageClient.cpp
@@ -467,29 +467,43 @@ namespace uba
 
 #if !UBA_USE_SPARSEFILE
 			FileAccessor destinationFile(m_logger, casFile.data);
-			if (!mappingBuffer)
-			{
-				u32 extraFlags = DefaultAttributes();
-				bool useOverlap = !IsRunningWine() && isCompressed == IsCompressed(casKey) && sizeOnDisk > 1024 * 1024;
-				if (useOverlap)
-					extraFlags |= FILE_FLAG_OVERLAPPED;
-				if (!destinationFile.CreateWrite(false, extraFlags, sizeOnDisk, m_tempPath.data))
-					return false;
-			}
 #endif
 
-			if (mappingBuffer)
-			{
-				UBA_ASSERT(!writeMem || mappedView.size == sizeOnDisk);
-				if (!writeMem)
+			u8* writePos = nullptr;
+			bool isInitialized = false;
+
+			// This is put in a function that is called in between sending and waiting for segment messages.
+			// It is called in a few more places in case there are no segments
+			auto initForWrite = [&]()
 				{
-					mappedView = mappingBuffer->AllocAndMapView(mappingType, sizeOnDisk, memoryMapAlignment, hint);
-					writeMem = mappedView.memory;
-					if (!writeMem)
-						return false;
-				}
-			}
-			u8* writePos = writeMem;
+					if (isInitialized)
+						return true;
+					isInitialized = true;
+
+					if (!mappingBuffer)
+					{
+						u32 extraFlags = DefaultAttributes();
+						bool useOverlap = !IsRunningWine() && isCompressed == IsCompressed(casKey) && sizeOnDisk > 1024 * 1024;
+						if (useOverlap)
+							extraFlags |= FILE_FLAG_OVERLAPPED;
+						if (!destinationFile.CreateWrite(false, extraFlags, sizeOnDisk, m_tempPath.data))
+							return false;
+					}
+
+					if (mappingBuffer)
+					{
+						UBA_ASSERT(!writeMem || mappedView.size == sizeOnDisk);
+						if (!writeMem)
+						{
+							mappedView = mappingBuffer->AllocAndMapView(mappingType, sizeOnDisk, memoryMapAlignment, hint);
+							writeMem = mappedView.memory;
+							if (!writeMem)
+								return false;
+						}
+					}
+					writePos = writeMem;
+					return true;
+				};
 
 			// This is here just to prevent server from getting a million messages at the same time.
 			// In theory we could have 10 clients with 48 processes each where each one of the processes asks for a large file (64 messages in flight)
@@ -512,7 +526,7 @@ namespace uba
 						if (fetchId == u16(~0))
 							return m_logger.Error(TC("Cas content error. Server believes %s was only one segment but client sees more. Size: %llu Left to read: %llu ResponseSize: %u. (%s)"), hint, fileSize, left, responseSize, casFile.data);
 						readBuffer = slot;
-						if (!SendBatchMessages(m_logger, *client, fetchId, readBuffer, BufferSlotSize, left, sizeOfFirstMessage, readIndex, responseSize))
+						if (!SendBatchMessages(m_logger, *client, fetchId, readBuffer, BufferSlotSize, left, sizeOfFirstMessage, readIndex, responseSize, initForWrite, hint))
 						{
 							if (proxy)
 							{
@@ -526,6 +540,10 @@ namespace uba
 					{
 						sendSegmentMessage = true;
 					}
+
+					if (!initForWrite()) // In case no segments are fetched
+						return false;
+
 					if (!mappingBuffer)
 					{
 						if (!destinationFile.Write(readBuffer, responseSize, writePos - writeMem))
@@ -588,7 +606,7 @@ namespace uba
 									writeCapacity = sizeOfFirstMessage;
 								}
 
-								if (!SendBatchMessages(m_logger, *client, fetchId, writeDest, writeCapacity, leftCompressed, sizeOfFirstMessage, readIndex, responseSize))
+								if (!SendBatchMessages(m_logger, *client, fetchId, writeDest, writeCapacity, leftCompressed, sizeOfFirstMessage, readIndex, responseSize, initForWrite, hint))
 								{
 									if (proxy)
 									{
@@ -663,6 +681,9 @@ namespace uba
 						if (tryAgain)
 							break;
 
+						if (!initForWrite()) // In case no segments are fetched
+							return false;
+
 						// Then decompress
 						{
 							u8* decompressBuffer = slot + BufferSlotHalfSize;
@@ -718,6 +739,9 @@ namespace uba
 				}
 			}
 
+			if (!initForWrite()) // In case no segments are fetched
+				return false;
+
 			if (sendEnd)
 			{
 				StackBinaryWriter<128> writer;
