[UBA]
* Fixed unit test bug when checking for binary imports

#skipci

diff --git a/Engine/Source/Programs/UnrealBuildAccelerator/Test/Private/UbaTestBasics.h b/Engine/Source/Programs/UnrealBuildAccelerator/Test/Private/UbaTestBasics.h
--- a/Engine/Source/Programs/UnrealBuildAccelerator/Test/Private/UbaTestBasics.h
+++ b/Engine/Source/Programs/UnrealBuildAccelerator/Test/Private/UbaTestBasics.h
@@ -501,7 +501,7 @@ namespace uba
 		StringBuffer<> error;
 		FindImports(path.data, [&](const tchar* import, bool isKnown, const char* const* importLoaderPaths)
 		{
-			importKernel = isKnown && Contains(import, TC("KERNEL32.dll"));
+			importKernel |= isKnown && Contains(import, TC("KERNEL32.dll"));
 		}, error);
 		if (!importKernel)
 			return logger.Error(TC("Failed to find Kernel32 as import"));
