[UBA]
* Added WorkContext struct and added so it is sent in everywhere. Not used yet but will be used to create more info in detailed trace streams

#skipci

diff --git a/Engine/Source/Programs/UnrealBuildAccelerator/Common/Private/UbaCacheServer.cpp b/Engine/Source/Programs/UnrealBuildAccelerator/Common/Private/UbaCacheServer.cpp
--- a/Engine/Source/Programs/UnrealBuildAccelerator/Common/Private/UbaCacheServer.cpp
+++ b/Engine/Source/Programs/UnrealBuildAccelerator/Common/Private/UbaCacheServer.cpp
@@ -207,7 +207,7 @@ namespace uba
 		m_rootDir.Replace('/', PathSeparator).EnsureEndsWithSlash();
 
 		m_server.RegisterService(CacheServiceId,
-			[this](const ConnectionInfo& connectionInfo, MessageInfo& messageInfo, BinaryReader& reader, BinaryWriter& writer)
+			[this](const ConnectionInfo& connectionInfo, const WorkContext& workContext, MessageInfo& messageInfo, BinaryReader& reader, BinaryWriter& writer)
 			{
 				return HandleMessage(connectionInfo, messageInfo.type, reader, writer);
 			},
@@ -1047,7 +1047,7 @@ namespace uba
 			for (auto& casKey : deletedCasFiles)
 			{
 				++activeDropCount;
-				m_server.AddWork([&, key = casKey]() { m_storage.DropCasFile(key, true, TC("")); --activeDropCount; }, 1, TC(""));
+				m_server.AddWork([&, key = casKey](const WorkContext&) { m_storage.DropCasFile(key, true, TC("")); --activeDropCount; }, 1, TC(""));
 			}
 			++deleteIteration;
 		}
@@ -1243,7 +1243,7 @@ namespace uba
 			u64 saveStart = GetTime();
 			m_logger.Detail(TC("  Saving to disk (v%u)"), CacheFileVersion);
 			Event saveCasEvent(true);
-			m_server.AddWork([&]() { m_storage.SaveCasTable(false, false); saveCasEvent.Set(); }, 1, TC("SaveCas"));
+			m_server.AddWork([&](const WorkContext&) { m_storage.SaveCasTable(false, false); saveCasEvent.Set(); }, 1, TC("SaveCas"));
 			SaveNoLock();
 			saveCasEvent.IsSet();
 			m_logger.Detail(TC("  Save Done (%s)"), TimeToText(GetTime() - saveStart).str);
diff --git a/Engine/Source/Programs/UnrealBuildAccelerator/Common/Private/UbaFileMapping.cpp b/Engine/Source/Programs/UnrealBuildAccelerator/Common/Private/UbaFileMapping.cpp
--- a/Engine/Source/Programs/UnrealBuildAccelerator/Common/Private/UbaFileMapping.cpp
+++ b/Engine/Source/Programs/UnrealBuildAccelerator/Common/Private/UbaFileMapping.cpp
@@ -650,7 +650,7 @@ namespace uba
 				--file.activeMapCount;
 			};
 		if (m_workManager && newSize == InvalidValue)
-			m_workManager->AddWork([=, h = TString(hint_)]() { unmap(h.c_str()); }, 1, TC("UnmapView"));
+			m_workManager->AddWork([=, h = TString(hint_)](const WorkContext&) { unmap(h.c_str()); }, 1, TC("UnmapView"));
 		else
 			unmap(hint_);
 	}
diff --git a/Engine/Source/Programs/UnrealBuildAccelerator/Common/Private/UbaNetworkServer.cpp b/Engine/Source/Programs/UnrealBuildAccelerator/Common/Private/UbaNetworkServer.cpp
--- a/Engine/Source/Programs/UnrealBuildAccelerator/Common/Private/UbaNetworkServer.cpp
+++ b/Engine/Source/Programs/UnrealBuildAccelerator/Common/Private/UbaNetworkServer.cpp
@@ -75,7 +75,7 @@ namespace uba
 		}
 
 		void ThreadWorker(NetworkServer& server);
-		void Update(WorkerContext& context, bool signalAvailable);
+		void Update(WorkerContext& context);
 		void DoAdditionalWorkAndSignalAvailable(NetworkServer& server);
 
 		Worker* m_nextWorker = nullptr;
@@ -460,22 +460,20 @@ namespace uba
 		return conn.m_shouldDisconnect;
 	}
 
-	void NetworkServer::Worker::Update(WorkerContext& context, bool signalAvailable)
+	void NetworkServer::Worker::Update(WorkerContext& context)
 	{
 		auto& server = context.server;
 
-		auto sg = MakeGuard([&]()
-			{
-				if (signalAvailable)
-					DoAdditionalWorkAndSignalAvailable(server);
-			});
-
 		// This is only additional work
 		if (!context.connection)
 			return;
+
 		auto& connection = *context.connection;
 		context.connection = nullptr;
 
+		WorkerRec& rec = server.m_workerFunctions[context.serviceId];
+		TrackWorkScope tws(server, rec.toString(context.messageType));
+
 		CryptoKey cryptoKey = connection.m_cryptoKey;
 		if (cryptoKey)
 		{
@@ -490,7 +488,6 @@ namespace uba
 
 		BinaryReader reader(context.buffer.data(), 0, context.dataSize);
 
-
 		constexpr u32 HeaderSize = 5; // 2 byte id, 3 bytes size
 		constexpr u32 ErrorSize = 0xffffff;
 
@@ -498,9 +495,6 @@ namespace uba
 		u8* idAndSizePtr = writer.AllocWrite(HeaderSize);
 			
 		u32 size;
-		WorkerRec& rec = server.m_workerFunctions[context.serviceId];
-
-		u32 workId = server.TrackWorkStart(rec.toString(context.messageType));
 
 		MessageInfo mi;
 		mi.type = context.messageType;
@@ -513,7 +507,7 @@ namespace uba
 			connection.SetShouldDisconnect();
 			size = ErrorSize;
 		}
-		else if (!rec.func({&connection}, mi, reader, writer))
+		else if (!rec.func({&connection}, {tws.id}, mi, reader, writer))
 		{
 			if (connection.SetShouldDisconnect())
 				server.m_logger.Error(TC("WORKER FUNCTION FAILED. id: %u, serviceid: %u type: %s, client: %s"), context.id, context.serviceId, rec.toString(context.messageType).data, GuidToString(connection.m_client->uid).str);
@@ -524,8 +518,6 @@ namespace uba
 			size = u32(writer.GetPosition());
 		}
 
-		server.TrackWorkEnd(workId);
-
 		if (mi.messageId)
 		{
 			UBA_ASSERT(size < (1 << 24));
@@ -562,7 +554,11 @@ namespace uba
 
 		t_worker = this;
 		while (m_context->workAvailable.IsSet(~0u) && m_loop)
-			Update(*m_context, true);
+		{
+			Update(*m_context);
+			DoAdditionalWorkAndSignalAvailable(m_context->server);
+		}
+
 		t_worker = nullptr;
 
 		if (m_inUse) // I have no idea how this can happen.. should not be possible. There is a path somewhere where it can leave while still being in use
@@ -583,11 +579,8 @@ namespace uba
 				server.m_additionalWork.pop_front();
 				lock.Leave();
 
-				u32 workId = server.TrackWorkStart(work.desc);
-
-				work.func();
-
-				server.TrackWorkEnd(workId);
+				TrackWorkScope tws(server, work.desc);
+				work.func({tws.id});
 			}
 
 			// Both locks needs to be taken to verify if additional work
@@ -640,7 +633,7 @@ namespace uba
 		m_logConnections = info.logConnections;
 
 		m_workerFunctions[SystemServiceId].toString = GetMessageTypeToName;
-		m_workerFunctions[SystemServiceId].func = [this](const ConnectionInfo& connectionInfo, MessageInfo& messageInfo, BinaryReader& reader, BinaryWriter& writer)
+		m_workerFunctions[SystemServiceId].func = [this](const ConnectionInfo& connectionInfo, const WorkContext& workContext, MessageInfo& messageInfo, BinaryReader& reader, BinaryWriter& writer)
 			{
 				return HandleSystemMessage(connectionInfo, messageInfo.type, reader, writer);
 			};
@@ -837,7 +830,7 @@ namespace uba
 		}
 	}
 
-	void NetworkServer::AddWork(const Function<void()>& work, u32 count, const tchar* desc, bool highPriority)
+	void NetworkServer::AddWork(const WorkFunction& work, u32 count, const tchar* desc, bool highPriority)
 	{
 		SCOPED_WRITE_LOCK(m_additionalWorkLock, lock);
 		for (u32 i = 0; i != count; ++i)
@@ -934,7 +927,6 @@ namespace uba
 
 	bool NetworkServer::DoAdditionalWork()
 	{
-		AdditionalWork work;
 		SCOPED_WRITE_LOCK(m_additionalWorkLock, lock);
 		if (m_additionalWork.empty())
 		{
@@ -961,7 +953,7 @@ namespace uba
 				lock2.Leave();
 				if (!workAvail)
 					context.workAvailable.IsSet(~0u);
-				worker->Update(context, false);
+				worker->Update(context);
 				UBA_ASSERT(worker->m_inUse);
 			}
 			else
@@ -984,13 +976,13 @@ namespace uba
 			worker->m_context = oldContext;
 			return true;
 		}
-		work = std::move(m_additionalWork.front());
+		AdditionalWork work = std::move(m_additionalWork.front());
 		m_additionalWork.pop_front();
 		lock.Leave();
 
 		TrackWorkScope tws(*this, work.desc);
 
-		work.func();
+		work.func({tws.id});
 
 		return true;
 	}
diff --git a/Engine/Source/Programs/UnrealBuildAccelerator/Common/Private/UbaScheduler.cpp b/Engine/Source/Programs/UnrealBuildAccelerator/Common/Private/UbaScheduler.cpp
--- a/Engine/Source/Programs/UnrealBuildAccelerator/Common/Private/UbaScheduler.cpp
+++ b/Engine/Source/Programs/UnrealBuildAccelerator/Common/Private/UbaScheduler.cpp
@@ -450,7 +450,7 @@ namespace uba
 			else if (prevStatus == ProcessStatus_QueuedForCache)
 			{
 				// TODO: This should not use work manager since it is mostly waiting on network
-				m_session.GetServer().AddWork([this, exitInfo]()
+				m_session.GetServer().AddWork([this, exitInfo](const WorkContext&)
 					{
 						ProcessStartInfo& si = *exitInfo->startInfo;
 						u64 startTime = GetTime();
diff --git a/Engine/Source/Programs/UnrealBuildAccelerator/Common/Private/UbaSession.cpp b/Engine/Source/Programs/UnrealBuildAccelerator/Common/Private/UbaSession.cpp
--- a/Engine/Source/Programs/UnrealBuildAccelerator/Common/Private/UbaSession.cpp
+++ b/Engine/Source/Programs/UnrealBuildAccelerator/Common/Private/UbaSession.cpp
@@ -2657,10 +2657,10 @@ namespace uba
 
 			if (shouldEvictFromMemory)
 			{
-				m_workManager->AddWork([mh = file.mappingHandle, n = file.name, this]()
+				m_workManager->AddWork([mh = file.mappingHandle, n = file.name, this](const WorkContext&)
 					{
 						CloseFileMapping(m_logger, mh, n.c_str());
-					}, 1, TC("CFM"));
+					}, 1, TC("ClosFileMapping"));
 			}
 			else
 			{
diff --git a/Engine/Source/Programs/UnrealBuildAccelerator/Common/Private/UbaSessionClient.cpp b/Engine/Source/Programs/UnrealBuildAccelerator/Common/Private/UbaSessionClient.cpp
--- a/Engine/Source/Programs/UnrealBuildAccelerator/Common/Private/UbaSessionClient.cpp
+++ b/Engine/Source/Programs/UnrealBuildAccelerator/Common/Private/UbaSessionClient.cpp
@@ -290,7 +290,7 @@ namespace uba
 
 		for (auto& m : modules)
 		{
-			m_client.AddWork([&handledCount, &m, this, &success, &keyStr]()
+			m_client.AddWork([&handledCount, &m, this, &success, &keyStr](const WorkContext&)
 				{
 					++handledCount;
 					auto g = MakeGuard([&]() { m.done.Set(); });
@@ -1374,7 +1374,7 @@ namespace uba
 				if (storeUncompressed)
 					knownInputKey = AsCompressed(knownInputKey, false);
 
-				m_client.AddWork([knownInputKey, allowProxy, this]()
+				m_client.AddWork([knownInputKey, allowProxy, this](const WorkContext&)
 					{
 						Storage::RetrieveResult result;
 						bool res = m_storage.RetrieveCasFile(result, knownInputKey, TC("KnownInput"), nullptr, 1, allowProxy);
diff --git a/Engine/Source/Programs/UnrealBuildAccelerator/Common/Private/UbaSessionServer.cpp b/Engine/Source/Programs/UnrealBuildAccelerator/Common/Private/UbaSessionServer.cpp
--- a/Engine/Source/Programs/UnrealBuildAccelerator/Common/Private/UbaSessionServer.cpp
+++ b/Engine/Source/Programs/UnrealBuildAccelerator/Common/Private/UbaSessionServer.cpp
@@ -126,9 +126,17 @@ namespace uba
 		m_server.RegisterOnClientDisconnected(ServiceId, [this](const Guid& clientUid, u32 clientId) { OnDisconnected(clientUid, clientId); });
 
 		m_server.RegisterService(ServiceId,
-			[this](const ConnectionInfo& connectionInfo, MessageInfo& messageInfo, BinaryReader& reader, BinaryWriter& writer)
+			[this](const ConnectionInfo& connectionInfo, const WorkContext& workContext, MessageInfo& messageInfo, BinaryReader& reader, BinaryWriter& writer)
 			{
-				return HandleMessage(connectionInfo, messageInfo.type, reader, writer);
+				switch (messageInfo.type)
+				{
+					#define UBA_SESSION_MESSAGE(x) case SessionMessageType_##x: return Handle##x(connectionInfo, workContext, reader, writer);
+					UBA_SESSION_MESSAGES
+					#undef UBA_SESSION_MESSAGE
+				}
+
+				UBA_ASSERT(false);
+				return false;
 			},
 			[](u8 type)
 			{
@@ -626,20 +634,7 @@ namespace uba
 		}
 	}
 
-	bool SessionServer::HandleMessage(const ConnectionInfo& connectionInfo, u8 messageType, BinaryReader& reader, BinaryWriter& writer)
-	{
-		switch (messageType)
-		{
-			#define UBA_SESSION_MESSAGE(x) case SessionMessageType_##x: return Handle##x(connectionInfo, reader, writer);
-			UBA_SESSION_MESSAGES
-			#undef UBA_SESSION_MESSAGE
-		}
-
-		UBA_ASSERT(false);
-		return false;
-	}
-
-	bool SessionServer::HandleConnect(const ConnectionInfo& connectionInfo, BinaryReader& reader, BinaryWriter& writer)
+	bool SessionServer::HandleConnect(const ConnectionInfo& connectionInfo, const WorkContext& workContext, BinaryReader& reader, BinaryWriter& writer)
 	{
 		StringBuffer<128> name;
 		reader.ReadString(name);
@@ -767,7 +762,7 @@ namespace uba
 		return true;
 	}
 
-	bool SessionServer::HandleEnsureBinaryFile(const ConnectionInfo& connectionInfo, BinaryReader& reader, BinaryWriter& writer)
+	bool SessionServer::HandleEnsureBinaryFile(const ConnectionInfo& connectionInfo, const WorkContext& workContext, BinaryReader& reader, BinaryWriter& writer)
 	{
 		bool clientIsArm = reader.ReadBool();
 		StringBuffer<> fileName;
@@ -854,7 +849,7 @@ namespace uba
 		return true;
 	}
 
-	bool SessionServer::HandleGetApplication(const ConnectionInfo& connectionInfo, BinaryReader& reader, BinaryWriter& writer)
+	bool SessionServer::HandleGetApplication(const ConnectionInfo& connectionInfo, const WorkContext& workContext, BinaryReader& reader, BinaryWriter& writer)
 	{
 		u32 processId = reader.ReadU32(); (void)processId;
 		StringBuffer<> applicationName;
@@ -898,7 +893,7 @@ namespace uba
 		return true;
 	}
 
-	bool SessionServer::HandleGetFileFromServer(const ConnectionInfo& connectionInfo, BinaryReader& reader, BinaryWriter& writer)
+	bool SessionServer::HandleGetFileFromServer(const ConnectionInfo& connectionInfo, const WorkContext& workContext, BinaryReader& reader, BinaryWriter& writer)
 	{
 		u32 processId = reader.ReadU32(); (void)processId;
 		StringBuffer<> fileName;
@@ -946,7 +941,7 @@ namespace uba
 		return true;
 	}
 
-	bool SessionServer::HandleGetLongPathName(const ConnectionInfo& connectionInfo, BinaryReader& reader, BinaryWriter& writer)
+	bool SessionServer::HandleGetLongPathName(const ConnectionInfo& connectionInfo, const WorkContext& workContext, BinaryReader& reader, BinaryWriter& writer)
 	{
 		#if PLATFORM_WINDOWS
 		StringBuffer<> shortPath;
@@ -961,7 +956,7 @@ namespace uba
 		#endif
 	}
 
-	bool SessionServer::HandleSendFileToServer(const ConnectionInfo& connectionInfo, BinaryReader& reader, BinaryWriter& writer)
+	bool SessionServer::HandleSendFileToServer(const ConnectionInfo& connectionInfo, const WorkContext& workContext, BinaryReader& reader, BinaryWriter& writer)
 	{
 		u32 processId = reader.ReadU32();
 		StringBuffer<> destination;
@@ -1071,7 +1066,7 @@ namespace uba
 		return true;
 	}
 
-	bool SessionServer::HandleDeleteFile(const ConnectionInfo& connectionInfo, BinaryReader& reader, BinaryWriter& writer)
+	bool SessionServer::HandleDeleteFile(const ConnectionInfo& connectionInfo, const WorkContext& workContext, BinaryReader& reader, BinaryWriter& writer)
 	{
 		StringKey fileNameKey = reader.ReadStringKey();
 		StringBuffer<> fileName;
@@ -1085,7 +1080,7 @@ namespace uba
 		return true;
 	}
 
-	bool SessionServer::HandleCopyFile(const ConnectionInfo& connectionInfo, BinaryReader& reader, BinaryWriter& writer)
+	bool SessionServer::HandleCopyFile(const ConnectionInfo& connectionInfo, const WorkContext& workContext, BinaryReader& reader, BinaryWriter& writer)
 	{
 		StringKey fromNameKey = reader.ReadStringKey(); (void)fromNameKey;
 		StringBuffer<> fromName;
@@ -1101,7 +1096,7 @@ namespace uba
 		return true;
 	}
 
-	bool SessionServer::HandleCreateDirectory(const ConnectionInfo& connectionInfo, BinaryReader& reader, BinaryWriter& writer)
+	bool SessionServer::HandleCreateDirectory(const ConnectionInfo& connectionInfo, const WorkContext& workContext, BinaryReader& reader, BinaryWriter& writer)
 	{
 		CreateDirectoryMessage msg;
 		reader.ReadString(msg.name);
@@ -1113,7 +1108,7 @@ namespace uba
 		return true;
 	}
 
-	bool SessionServer::HandleRemoveDirectory(const ConnectionInfo& connectionInfo, BinaryReader& reader, BinaryWriter& writer)
+	bool SessionServer::HandleRemoveDirectory(const ConnectionInfo& connectionInfo, const WorkContext& workContext, BinaryReader& reader, BinaryWriter& writer)
 	{
 		RemoveDirectoryMessage msg;
 		reader.ReadString(msg.name);
@@ -1125,7 +1120,7 @@ namespace uba
 		return true;
 	}
 
-	bool SessionServer::HandleListDirectory(const ConnectionInfo& connectionInfo, BinaryReader& reader, BinaryWriter& writer)
+	bool SessionServer::HandleListDirectory(const ConnectionInfo& connectionInfo, const WorkContext& workContext, BinaryReader& reader, BinaryWriter& writer)
 	{
 		u32 sessionId = reader.ReadU32();
 		u32 sessionIndex = sessionId - 1;
@@ -1145,7 +1140,7 @@ namespace uba
 		return true;
 	}
 
-	bool SessionServer::HandleGetDirectoriesFromServer(const ConnectionInfo& connectionInfo, BinaryReader& reader, BinaryWriter& writer)
+	bool SessionServer::HandleGetDirectoriesFromServer(const ConnectionInfo& connectionInfo, const WorkContext& workContext, BinaryReader& reader, BinaryWriter& writer)
 	{
 		u32 sessionId = reader.ReadU32();
 		u32 sessionIndex = sessionId - 1;
@@ -1158,7 +1153,7 @@ namespace uba
 		return true;
 	}
 
-	bool SessionServer::HandleGetNameToHashFromServer(const ConnectionInfo& connectionInfo, BinaryReader& reader, BinaryWriter& writer)
+	bool SessionServer::HandleGetNameToHashFromServer(const ConnectionInfo& connectionInfo, const WorkContext& workContext, BinaryReader& reader, BinaryWriter& writer)
 	{
 		u32 requestedSize = reader.ReadU32();
 
@@ -1175,7 +1170,7 @@ namespace uba
 		return true;
 	}
 
-	bool SessionServer::HandleProcessAvailable(const ConnectionInfo& connectionInfo, BinaryReader& reader, BinaryWriter& writer)
+	bool SessionServer::HandleProcessAvailable(const ConnectionInfo& connectionInfo, const WorkContext& workContext, BinaryReader& reader, BinaryWriter& writer)
 	{
 		u32 sessionId = reader.ReadU32();
 		u32 sessionIndex = sessionId - 1;
@@ -1319,7 +1314,7 @@ namespace uba
 		return true;
 	}
 
-	bool SessionServer::HandleProcessInputs(const ConnectionInfo& connectionInfo, BinaryReader& reader, BinaryWriter& writer)
+	bool SessionServer::HandleProcessInputs(const ConnectionInfo& connectionInfo, const WorkContext& workContext, BinaryReader& reader, BinaryWriter& writer)
 	{
 		u32 processId = u32(reader.Read7BitEncoded());
 		ProcessHandle h(GetProcess(processId));
@@ -1336,7 +1331,7 @@ namespace uba
 		return true;
 	}
 
-	bool SessionServer::HandleProcessFinished(const ConnectionInfo& connectionInfo, BinaryReader& reader, BinaryWriter& writer)
+	bool SessionServer::HandleProcessFinished(const ConnectionInfo& connectionInfo, const WorkContext& workContext, BinaryReader& reader, BinaryWriter& writer)
 	{
 		u32 processId = reader.ReadU32();
 
@@ -1403,7 +1398,7 @@ namespace uba
 		return true;
 	}
 
-	bool SessionServer::HandleProcessReturned(const ConnectionInfo& connectionInfo, BinaryReader& reader, BinaryWriter& writer)
+	bool SessionServer::HandleProcessReturned(const ConnectionInfo& connectionInfo, const WorkContext& workContext, BinaryReader& reader, BinaryWriter& writer)
 	{
 		u32 processId = reader.ReadU32();
 		StringBuffer<> reason;
@@ -1447,7 +1442,7 @@ namespace uba
 		return true;
 	}
 
-	bool SessionServer::HandleGetRoots(const ConnectionInfo& connectionInfo, BinaryReader& reader, BinaryWriter& writer)
+	bool SessionServer::HandleGetRoots(const ConnectionInfo& connectionInfo, const WorkContext& workContext, BinaryReader& reader, BinaryWriter& writer)
 	{
 		RootsHandle rootsHandle = reader.ReadU64();
 		auto rootsEntry = GetRootsEntry(rootsHandle);
@@ -1457,13 +1452,13 @@ namespace uba
 		return true;
 	}
 
-	bool SessionServer::HandleVirtualAllocFailed(const ConnectionInfo& connectionInfo, BinaryReader& reader, BinaryWriter& writer)
+	bool SessionServer::HandleVirtualAllocFailed(const ConnectionInfo& connectionInfo, const WorkContext& workContext, BinaryReader& reader, BinaryWriter& writer)
 	{
 		m_logger.Error(TC("VIRTUAL ALLOC FAILING ON REMOTE MACHINE %s !"), GuidToString(connectionInfo.GetUid()).str);
 		return true;
 	}
 
-	bool SessionServer::HandleGetTraceInformation(const ConnectionInfo& connectionInfo, BinaryReader& reader, BinaryWriter& writer)
+	bool SessionServer::HandleGetTraceInformation(const ConnectionInfo& connectionInfo, const WorkContext& workContext, BinaryReader& reader, BinaryWriter& writer)
 	{
 		u32 remotePos = reader.ReadU32();
 		u32 localPos;
@@ -1478,7 +1473,7 @@ namespace uba
 		return true;
 	}
 
-	bool SessionServer::HandlePing(const ConnectionInfo& connectionInfo, BinaryReader& reader, BinaryWriter& writer)
+	bool SessionServer::HandlePing(const ConnectionInfo& connectionInfo, const WorkContext& workContext, BinaryReader& reader, BinaryWriter& writer)
 	{
 		u32 sessionId = reader.ReadU32();
 		u64 lastPing = reader.ReadU64();
@@ -1502,7 +1497,7 @@ namespace uba
 		return true;
 	}
 
-	bool SessionServer::HandleNotification(const ConnectionInfo& connectionInfo, BinaryReader& reader, BinaryWriter& writer)
+	bool SessionServer::HandleNotification(const ConnectionInfo& connectionInfo, const WorkContext& workContext, BinaryReader& reader, BinaryWriter& writer)
 	{
 		u32 sessionId = reader.ReadU32();
 
@@ -1519,7 +1514,7 @@ namespace uba
 		return true;
 	}
 
-	bool SessionServer::HandleGetNextProcess(const ConnectionInfo& connectionInfo, BinaryReader& reader, BinaryWriter& writer)
+	bool SessionServer::HandleGetNextProcess(const ConnectionInfo& connectionInfo, const WorkContext& workContext, BinaryReader& reader, BinaryWriter& writer)
 	{
 		u32 processId = reader.ReadU32();
 		u32 prevExitCode = reader.ReadU32();
@@ -1550,7 +1545,7 @@ namespace uba
 		return true;
 	}
 
-	bool SessionServer::HandleCustom(const ConnectionInfo& connectionInfo, BinaryReader& reader, BinaryWriter& writer)
+	bool SessionServer::HandleCustom(const ConnectionInfo& connectionInfo, const WorkContext& workContext, BinaryReader& reader, BinaryWriter& writer)
 	{
 		u32 processId = reader.ReadU32();
 		ProcessHandle h(GetProcess(processId));
@@ -1562,7 +1557,7 @@ namespace uba
 		return true;
 	}
 
-	bool SessionServer::HandleUpdateEnvironment(const ConnectionInfo& connectionInfo, BinaryReader& reader, BinaryWriter& writer)
+	bool SessionServer::HandleUpdateEnvironment(const ConnectionInfo& connectionInfo, const WorkContext& workContext, BinaryReader& reader, BinaryWriter& writer)
 	{
 		u32 processId = reader.ReadU32();
 		ProcessHandle h(GetProcess(processId));
@@ -1574,14 +1569,14 @@ namespace uba
 		return true;
 	}
 
-	bool SessionServer::HandleSummary(const ConnectionInfo& connectionInfo, BinaryReader& reader, BinaryWriter& writer)
+	bool SessionServer::HandleSummary(const ConnectionInfo& connectionInfo, const WorkContext& workContext, BinaryReader& reader, BinaryWriter& writer)
 	{
 		u32 sessionId = reader.ReadU32();
 		m_trace.SessionSummary(sessionId, reader.GetPositionData(), reader.GetLeft());
 		return true;
 	}
 
-	bool SessionServer::HandleCommand(const ConnectionInfo& connectionInfo, BinaryReader& reader, BinaryWriter& writer)
+	bool SessionServer::HandleCommand(const ConnectionInfo& connectionInfo, const WorkContext& workContext, BinaryReader& reader, BinaryWriter& writer)
 	{
 		StringBuffer<128> command;
 		reader.ReadString(command);
@@ -1645,7 +1640,7 @@ namespace uba
 		return true;
 	}
 
-	bool SessionServer::HandleSHGetKnownFolderPath(const ConnectionInfo& connectionInfo, BinaryReader& reader, BinaryWriter& writer)
+	bool SessionServer::HandleSHGetKnownFolderPath(const ConnectionInfo& connectionInfo, const WorkContext& workContext, BinaryReader& reader, BinaryWriter& writer)
 	{
 #if PLATFORM_WINDOWS
 		GUID kfid;
@@ -2184,7 +2179,7 @@ namespace uba
 		return true;
 	}
 
-	bool SessionServer::HandleDebugFileNotFoundError(const ConnectionInfo& connectionInfo, BinaryReader& reader, BinaryWriter& writer)
+	bool SessionServer::HandleDebugFileNotFoundError(const ConnectionInfo& connectionInfo, const WorkContext& workContext, BinaryReader& reader, BinaryWriter& writer)
 	{
 #if PLATFORM_WINDOWS
 		StringBuffer<> errorPath;
@@ -2285,12 +2280,12 @@ namespace uba
 		return true;
 	}
 
-	bool SessionServer::HandleHostRun(const ConnectionInfo& connectionInfo, BinaryReader& reader, BinaryWriter& writer)
+	bool SessionServer::HandleHostRun(const ConnectionInfo& connectionInfo, const WorkContext& workContext, BinaryReader& reader, BinaryWriter& writer)
 	{
 		return HostRun(reader, writer);
 	}
 
-	bool SessionServer::HandleGetSymbols(const ConnectionInfo& connectionInfo, BinaryReader& reader, BinaryWriter& writer)
+	bool SessionServer::HandleGetSymbols(const ConnectionInfo& connectionInfo, const WorkContext& workContext, BinaryReader& reader, BinaryWriter& writer)
 	{
 		TString application = reader.ReadString();
 		bool isClientArm = reader.ReadBool();
diff --git a/Engine/Source/Programs/UnrealBuildAccelerator/Common/Private/UbaStorage.cpp b/Engine/Source/Programs/UnrealBuildAccelerator/Common/Private/UbaStorage.cpp
--- a/Engine/Source/Programs/UnrealBuildAccelerator/Common/Private/UbaStorage.cpp
+++ b/Engine/Source/Programs/UnrealBuildAccelerator/Common/Private/UbaStorage.cpp
@@ -378,7 +378,7 @@ namespace uba
 				auto udg = MakeGuard([&]()
 					{
 						#if UBA_EXPERIMENTAL
-						m_workManager->AddWork([=, f = TString(from)]() { UnmapViewOfFile(uncompressedData, fileSize, f.c_str()); }, 1, TC("UnmapFile"));
+						m_workManager->AddWork([this, uncompressedData, fileSize, f = TString(from)](const WorkContext&) { UnmapViewOfFile(m_logger, uncompressedData, fileSize, f.c_str()); }, 1, TC("UnmapFile"));
 						#else
 						UnmapViewOfFile(m_logger, uncompressedData, fileSize, from);
 						#endif
@@ -526,7 +526,7 @@ namespace uba
 
 		StorageStats& stats = Stats();
 
-		auto work = [this, rec, &stats]()
+		auto work = [this, rec, &stats](const WorkContext&)
 		{
 			u8* slot = m_bufferSlots.Pop();
 			auto _ = MakeGuard([&](){ m_bufferSlots.Push(slot); });
@@ -583,7 +583,7 @@ namespace uba
 
 		rec->refCount = workerCount + 1; // We need to keep refcount up 1 to make sure it is not deleted before we read rec->written
 		m_workManager->AddWork(work, workerCount-1, TC("Compress")); // We are a worker ourselves
-		work();
+		work({});
 		rec->events[rec->workCount - 1].IsSet();
 
 		totalWritten += rec->written;
@@ -728,7 +728,7 @@ namespace uba
 			rec->done.Create(true);
 			rec->refCount = 2;
 
-			auto work = [rec]()
+			auto work = [rec](const WorkContext&)
 				{
 					u64 lastWritten = 0;
 					while (true)
@@ -793,7 +793,7 @@ namespace uba
 			}
 
 			TimerScope ts(stats.decompressToMem);
-			work();
+			work({});
 			rec->done.IsSet();
 			bool success = !rec->error;
 			if (!success)
@@ -1034,7 +1034,7 @@ namespace uba
 
 				if (allowParallel && m_workManager)
 				{
-					m_workManager->AddWork([&traverseCasFileDir, p = TString(fullPath.data)]() { traverseCasFileDir(p); }, 1, TC("TraverseCasFiles"));
+					m_workManager->AddWork([&traverseCasFileDir, p = TString(fullPath.data)](const WorkContext&) { traverseCasFileDir(p); }, 1, TC("TraverseCasFiles"));
 				}
 				else
 				{
@@ -1744,7 +1744,7 @@ namespace uba
 		TraverseAllCasFiles(casRoot.data, [&](const StringBufferBase& fullPath, const DirectoryEntry& e)
 			{
 				++entryCount;
-				workManager.AddWork([&, filePath = TString(fullPath.data), name = TString(e.name), lastWritten = e.lastWritten]()
+				workManager.AddWork([&, filePath = TString(fullPath.data), name = TString(e.name), lastWritten = e.lastWritten](const WorkContext&)
 					{
 						StringBuffer<> timeStr;
 						writeTimeAgo(timeStr, lastWritten);
@@ -1844,7 +1844,7 @@ namespace uba
 					u64 lastWritten = e.lastWritten;
 					u64 size = e.size;
 
-					workManager.AddWork([&, p = TString(path.data), lastWritten, size]()
+					workManager.AddWork([&, p = TString(path.data), lastWritten, size](const WorkContext&)
 						{
 							++foundFiles;
 							StringKey key = ToStringKey(p.data(), p.size());
@@ -1928,7 +1928,7 @@ namespace uba
 				{
 					if (IsDirectory(e.attributes))
 						return;
-					workManager.AddWork([&, name = TString(e.name)]()
+					workManager.AddWork([&, name = TString(e.name)](const WorkContext&)
 						{
 							StringBuffer<> fullPath;
 							fullPath.Append(m_rootDir).Append(name);
@@ -2870,8 +2870,8 @@ namespace uba
 			}
 			auto udg = MakeGuard([&]() { 
 					#if UBA_EXPERIMENTAL
-					if (m_workManager
-						m_workManager->AddWork([=, fn = TString(fileName)]() { UnmapViewOfFile(fileData, fileSize, fn.c_str()); }, 1, TC("UnmapFile"));
+					if (m_workManager)
+						m_workManager->AddWork([this, fileData, fileSize, fn = TString(fileName)](const WorkContext&) { UnmapViewOfFile(m_logger, fileData, fileSize, fn.c_str()); }, 1, TC("UnmapFile"));
 					else
 					#endif
 						UnmapViewOfFile(m_logger, fileData, fileSize, fileName);
diff --git a/Engine/Source/Programs/UnrealBuildAccelerator/Common/Private/UbaStorageClient.cpp b/Engine/Source/Programs/UnrealBuildAccelerator/Common/Private/UbaStorageClient.cpp
--- a/Engine/Source/Programs/UnrealBuildAccelerator/Common/Private/UbaStorageClient.cpp
+++ b/Engine/Source/Programs/UnrealBuildAccelerator/Common/Private/UbaStorageClient.cpp
@@ -957,7 +957,7 @@ namespace uba
 					if (!seenIds.insert(e.id).second)
 						return;
 					lock.Leave();
-					workManager.AddWork([&, filePath = TString(fullPath.data)]()
+					workManager.AddWork([&, filePath = TString(fullPath.data)](const WorkContext& context)
 						{
 							PopulateCasFromDirsRecursive(filePath.c_str(), workManager, seenIds, seenIdsLock, shouldExit);
 						}, 1, TC(""));
@@ -990,7 +990,7 @@ namespace uba
 					return;
 				}
 
-				workManager.AddWork([&, fe = &fileEntry, lw = e.lastWritten, s = e.size, filePath = TString(fullPath.data)]()
+				workManager.AddWork([&, fe = &fileEntry, lw = e.lastWritten, s = e.size, filePath = TString(fullPath.data)](const WorkContext& context)
 					{
 						auto feLockLeave = MakeGuard([fe]() { fe->lock.LeaveWrite(); });
 
diff --git a/Engine/Source/Programs/UnrealBuildAccelerator/Common/Private/UbaStorageServer.cpp b/Engine/Source/Programs/UnrealBuildAccelerator/Common/Private/UbaStorageServer.cpp
--- a/Engine/Source/Programs/UnrealBuildAccelerator/Common/Private/UbaStorageServer.cpp
+++ b/Engine/Source/Programs/UnrealBuildAccelerator/Common/Private/UbaStorageServer.cpp
@@ -31,9 +31,9 @@ namespace uba
 			UBA_ASSERT(false);
 
 		m_server.RegisterService(ServiceId,
-			[this](const ConnectionInfo& connectionInfo, MessageInfo& messageInfo, BinaryReader& reader, BinaryWriter& writer)
+			[this](const ConnectionInfo& connectionInfo, const WorkContext& workContext, MessageInfo& messageInfo, BinaryReader& reader, BinaryWriter& writer)
 			{
-				return HandleMessage(connectionInfo, messageInfo.type, reader, writer);
+				return HandleMessage(connectionInfo, workContext, messageInfo.type, reader, writer);
 			},
 			[](u8 messageType)
 			{
@@ -132,7 +132,7 @@ namespace uba
 			if (ownsMapping)
 			{
 				++server.m_activeUnmap;
-				server.GetServer().AddWork([&server, mb = memoryBegin, mp = mappedView, rfh = readFileHandle, reason]()
+				server.GetServer().AddWork([&server, mb = memoryBegin, mp = mappedView, rfh = readFileHandle, reason](const WorkContext& context)
 					{
 						if (!UnmapViewOfFile(server.m_logger, mb, mp.size, reason))
 							server.m_logger.Error(TC("%s - Failed to unmap memory at 0x%llx with size %llu (%s)"), reason, mb, mp.size, LastErrorToText().data);
@@ -276,7 +276,7 @@ namespace uba
 		}
 	}
 
-	bool StorageServer::HandleMessage(const ConnectionInfo& connectionInfo, u8 messageType, BinaryReader& reader, BinaryWriter& writer)
+	bool StorageServer::HandleMessage(const ConnectionInfo& connectionInfo, const WorkContext& workContext, u8 messageType, BinaryReader& reader, BinaryWriter& writer)
 	{
 		ActiveStore* firstStore = nullptr;
 		ActiveStore tempStore;
diff --git a/Engine/Source/Programs/UnrealBuildAccelerator/Common/Private/UbaStorageUtils.cpp b/Engine/Source/Programs/UnrealBuildAccelerator/Common/Private/UbaStorageUtils.cpp
--- a/Engine/Source/Programs/UnrealBuildAccelerator/Common/Private/UbaStorageUtils.cpp
+++ b/Engine/Source/Programs/UnrealBuildAccelerator/Common/Private/UbaStorageUtils.cpp
@@ -44,7 +44,7 @@ namespace uba
 			rec->done.Create(true);
 			rec->refCount = 2;
 
-			auto work = [rec]()
+			auto work = [rec](const WorkContext& context)
 			{
 				while (true)
 				{
@@ -78,7 +78,7 @@ namespace uba
 				workManager->AddWork(work, workerCount, TC("CalculateKey"));
 			}
 
-			work();
+			work({});
 			rec->done.IsSet();
 
 			hasher.Update(rec->keys.data(), rec->keys.size()*sizeof(CasKey));
diff --git a/Engine/Source/Programs/UnrealBuildAccelerator/Common/Private/UbaWorkManager.cpp b/Engine/Source/Programs/UnrealBuildAccelerator/Common/Private/UbaWorkManager.cpp
--- a/Engine/Source/Programs/UnrealBuildAccelerator/Common/Private/UbaWorkManager.cpp
+++ b/Engine/Source/Programs/UnrealBuildAccelerator/Common/Private/UbaWorkManager.cpp
@@ -61,7 +61,7 @@ namespace uba
 						lock.Leave();
 
 						TrackWorkScope tws(manager, work.desc);
-						work.func();
+						work.func({tws.id});
 					}
 
 					SCOPED_WRITE_LOCK(manager.m_availableWorkersLock, lock1);
@@ -99,7 +99,7 @@ namespace uba
 	}
 
 
-	void WorkManagerImpl::AddWork(const Function<void()>& work, u32 count, const tchar* desc, bool highPriority)
+	void WorkManagerImpl::AddWork(const WorkFunction& work, u32 count, const tchar* desc, bool highPriority)
 	{
 		SCOPED_WRITE_LOCK(m_workLock, lock);
 		bool trackWork = m_workTracker.load();
@@ -163,7 +163,7 @@ namespace uba
 			lock.Leave();
 
 			TrackWorkScope tws(*this, work.desc);
-			work.func();
+			work.func({tws.id});
 		}
 	}
 
diff --git a/Engine/Source/Programs/UnrealBuildAccelerator/Common/Public/UbaNetworkServer.h b/Engine/Source/Programs/UnrealBuildAccelerator/Common/Public/UbaNetworkServer.h
--- a/Engine/Source/Programs/UnrealBuildAccelerator/Common/Public/UbaNetworkServer.h
+++ b/Engine/Source/Programs/UnrealBuildAccelerator/Common/Public/UbaNetworkServer.h
@@ -63,7 +63,7 @@ namespace uba
 		void PrintSummary(Logger& logger);
 
 		using TypeToNameFunction = StringView(u8 type);
-		using WorkerFunction = Function<bool(const ConnectionInfo& connectionInfo, MessageInfo& messageInfo, BinaryReader& reader, BinaryWriter& writer)>;
+		using WorkerFunction = Function<bool(const ConnectionInfo& connectionInfo, const WorkContext& workContext, MessageInfo& messageInfo, BinaryReader& reader, BinaryWriter& writer)>;
 		void RegisterService(u8 serviceId, const WorkerFunction& function, TypeToNameFunction* typeToNameFunc = nullptr);
 		void UnregisterService(u8 serviceId);
 
@@ -79,7 +79,7 @@ namespace uba
 		//void RegisterOnConnection(u8 id, const OnConnectionFunction& func);
 		//void UnregisterOnConnection(u8 id);
 
-		virtual void AddWork(const Function<void()>& work, u32 count, const tchar* desc, bool highPriority = false) override final;
+		virtual void AddWork(const WorkFunction& work, u32 count, const tchar* desc, bool highPriority = false) override final;
 		virtual void DoWork(u32 count = 1) override final;
 		virtual u32 GetWorkerCount() override final;
 
@@ -143,7 +143,7 @@ namespace uba
 		u32 m_maxWorkerCount = 0;
 
 		ReaderWriterLock m_additionalWorkLock;
-		struct AdditionalWork { Function<void()> func; TString desc; };
+		struct AdditionalWork { WorkFunction func; TString desc; };
 		List<AdditionalWork> m_additionalWork;
 
 		ReaderWriterLock m_availableWorkersLock;
diff --git a/Engine/Source/Programs/UnrealBuildAccelerator/Common/Public/UbaSessionServer.h b/Engine/Source/Programs/UnrealBuildAccelerator/Common/Public/UbaSessionServer.h
--- a/Engine/Source/Programs/UnrealBuildAccelerator/Common/Public/UbaSessionServer.h
+++ b/Engine/Source/Programs/UnrealBuildAccelerator/Common/Public/UbaSessionServer.h
@@ -80,8 +80,7 @@ namespace uba
 		struct ClientSession;
 
 		void OnDisconnected(const Guid& clientUid, u32 clientId);
-		bool HandleMessage(const ConnectionInfo& connectionInfo, u8 messageType, BinaryReader& reader, BinaryWriter& writer);
-		#define UBA_SESSION_MESSAGE(x) bool Handle##x(const ConnectionInfo& connectionInfo, BinaryReader& reader, BinaryWriter& writer);
+		#define UBA_SESSION_MESSAGE(x) bool Handle##x(const ConnectionInfo& connectionInfo, const WorkContext& workContext, BinaryReader& reader, BinaryWriter& writer);
 		UBA_SESSION_MESSAGES
 		#undef UBA_SESSION_MESSAGE
 
diff --git a/Engine/Source/Programs/UnrealBuildAccelerator/Common/Public/UbaStorageServer.h b/Engine/Source/Programs/UnrealBuildAccelerator/Common/Public/UbaStorageServer.h
--- a/Engine/Source/Programs/UnrealBuildAccelerator/Common/Public/UbaStorageServer.h
+++ b/Engine/Source/Programs/UnrealBuildAccelerator/Common/Public/UbaStorageServer.h
@@ -43,7 +43,7 @@ namespace uba
 		virtual bool HasProxy(u32 clientId) override;
 
 		void OnDisconnected(u32 clientId);
-		bool HandleMessage(const ConnectionInfo& connectionInfo, u8 messageType, BinaryReader& reader, BinaryWriter& writer);
+		bool HandleMessage(const ConnectionInfo& connectionInfo, const WorkContext& workContext, u8 messageType, BinaryReader& reader, BinaryWriter& writer);
 		bool WaitForWritten(CasEntry& casEntry, ScopedWriteLock& entryLock, const ConnectionInfo& connectionInfo, const tchar* hint);
 
 		u16 PopId();
diff --git a/Engine/Source/Programs/UnrealBuildAccelerator/Common/Public/UbaWorkManager.h b/Engine/Source/Programs/UnrealBuildAccelerator/Common/Public/UbaWorkManager.h
--- a/Engine/Source/Programs/UnrealBuildAccelerator/Common/Public/UbaWorkManager.h
+++ b/Engine/Source/Programs/UnrealBuildAccelerator/Common/Public/UbaWorkManager.h
@@ -12,14 +12,20 @@ namespace uba
 	class WorkTracker
 	{
 	public:
-		virtual u32 TrackWorkStart(const StringView& desc) { return 0; }
-		virtual void TrackWorkEnd(u32 id) {}
+		virtual u32 TrackWorkStart(const StringView& desc) = 0;
+		virtual void TrackWorkEnd(u32 id) = 0;
+	};
+
+	struct WorkContext
+	{
+		u32 id;
 	};
 
 	class WorkManager
 	{
 	public:
-		virtual void AddWork(const Function<void()>& work, u32 count, const tchar* desc, bool highPriority = false) = 0;
+		using WorkFunction = Function<void(const WorkContext& context)>;
+		virtual void AddWork(const WorkFunction& work, u32 count, const tchar* desc, bool highPriority = false) = 0;
 		virtual u32 GetWorkerCount() = 0;
 		virtual void DoWork(u32 count = 1) = 0;
 
@@ -42,7 +48,7 @@ namespace uba
 	public:
 		WorkManagerImpl(u32 workerCount, const tchar* workerDesc = TC("UbaWorker"));
 		virtual ~WorkManagerImpl();
-		virtual void AddWork(const Function<void()>& work, u32 count, const tchar* desc, bool highPriority = false) override;
+		virtual void AddWork(const WorkFunction& work, u32 count, const tchar* desc, bool highPriority = false) override;
 		virtual u32 GetWorkerCount() override;
 		virtual void DoWork(u32 count = 1) override;
 		void FlushWork();
@@ -54,7 +60,7 @@ namespace uba
 		Worker* PopWorkerNoLock();
 
 		Vector<Worker*> m_workers;
-		struct Work { Function<void()> func; TString desc; };
+		struct Work { WorkFunction func; TString desc; };
 		ReaderWriterLock m_workLock;
 		List<Work> m_work;
 		Atomic<u32> m_activeWorkerCount;
@@ -66,10 +72,10 @@ namespace uba
 
 	struct TrackWorkScope
 	{
-		TrackWorkScope(WorkManager& wm, const StringView& desc) : workManager(wm), workIndex(wm.TrackWorkStart(desc)) {}
-		~TrackWorkScope() { workManager.TrackWorkEnd(workIndex); }
+		TrackWorkScope(WorkManager& wm, const StringView& desc) : workManager(wm), id(wm.TrackWorkStart(desc)) {}
+		~TrackWorkScope() { workManager.TrackWorkEnd(id); }
 		WorkManager& workManager;
-		u32 workIndex;
+		u32 id;
 	};
 
 
@@ -100,7 +106,7 @@ namespace uba
 		context->isDone = false;
 		context->doneEvent = &doneEvent;
 
-		auto work = [context, funcCopy = func]() mutable
+		auto work = [context, funcCopy = func](const WorkContext&) mutable
 			{
 				u32 active = 0;
 				while (true)
@@ -132,7 +138,7 @@ namespace uba
 			};
 		
 		AddWork(work, workCount, description, highPriority);
-		work();
+		work({});
 		doneEvent.IsSet();
 
 		#endif
diff --git a/Engine/Source/Programs/UnrealBuildAccelerator/Test/Private/UbaTestNetwork.h b/Engine/Source/Programs/UnrealBuildAccelerator/Test/Private/UbaTestNetwork.h
--- a/Engine/Source/Programs/UnrealBuildAccelerator/Test/Private/UbaTestNetwork.h
+++ b/Engine/Source/Programs/UnrealBuildAccelerator/Test/Private/UbaTestNetwork.h
@@ -67,7 +67,7 @@ namespace uba
 		NetworkServer server(ctorSuccess, logWriter);
 		NetworkClient client(ctorSuccess, logWriter);
 
-		server.RegisterService(1, [&](const ConnectionInfo& connectionInfo, MessageInfo& messageInfo, BinaryReader& reader, BinaryWriter& writer)
+		server.RegisterService(1, [&](const ConnectionInfo& connectionInfo, const WorkContext& workContext, MessageInfo& messageInfo, BinaryReader& reader, BinaryWriter& writer)
 			{
 				logger.Info(TC("Got ping!"));
 				UBA_ASSERT(messageInfo.type == SessionMessageType_Ping);
@@ -105,7 +105,7 @@ namespace uba
 		NetworkServer server(ctorSuccess, logWriter);
 		NetworkClient client(ctorSuccess, logWriter);
 
-		server.RegisterService(1, [&](const ConnectionInfo& connectionInfo, MessageInfo& messageInfo, BinaryReader& reader, BinaryWriter& writer)
+		server.RegisterService(1, [&](const ConnectionInfo& connectionInfo, const WorkContext& workContext, MessageInfo& messageInfo, BinaryReader& reader, BinaryWriter& writer)
 			{
 				logger.Info(TC("Got ping!"));
 				UBA_ASSERT(messageInfo.type == SessionMessageType_Ping);
@@ -152,7 +152,7 @@ namespace uba
 		NetworkServer server(ctorSuccess, logWriter);
 		NetworkClient client(ctorSuccess, logWriter);
 
-		server.RegisterService(1, [&](const ConnectionInfo& connectionInfo, MessageInfo& messageInfo, BinaryReader& reader, BinaryWriter& writer)
+		server.RegisterService(1, [&](const ConnectionInfo& connectionInfo, const WorkContext& workContext, MessageInfo& messageInfo, BinaryReader& reader, BinaryWriter& writer)
 			{
 				logger.Info(TC("Got ping!"));
 				UBA_ASSERT(messageInfo.type == SessionMessageType_Ping);
