[UBA]
* Added GetFileName to StringView

#skipci

diff --git a/Engine/Source/Programs/UnrealBuildAccelerator/Core/Private/UbaStringBuffer.cpp b/Engine/Source/Programs/UnrealBuildAccelerator/Core/Private/UbaStringBuffer.cpp
--- a/Engine/Source/Programs/UnrealBuildAccelerator/Core/Private/UbaStringBuffer.cpp
+++ b/Engine/Source/Programs/UnrealBuildAccelerator/Core/Private/UbaStringBuffer.cpp
@@ -145,6 +145,13 @@ namespace uba
 		return out != 0 || Equals(str, TC("0"));
 	}
 
+	const tchar* GetFileName(const tchar* path)
+	{
+		if (const tchar* lps = TStrrchr(path, PathSeparator))
+			return lps + 1;
+		return path;
+	}
+
 	StringBufferBase& StringBufferBase::Append(const tchar* str)
 	{
 		return Append(str, u32(TStrlen(str)));
@@ -314,13 +321,6 @@ namespace uba
 		return TStrrchr(data + offset, c);
 	}
 
-	const tchar* StringBufferBase::GetFileName() const
-	{
-		if (const tchar* lps = TStrrchr(data, PathSeparator))
-			return lps + 1;
-		return data;
-	}
-
 	StringBufferBase& StringBufferBase::EnsureEndsWithSlash()
 	{
 		UBA_ASSERT(count);
diff --git a/Engine/Source/Programs/UnrealBuildAccelerator/Core/Public/UbaStringBuffer.h b/Engine/Source/Programs/UnrealBuildAccelerator/Core/Public/UbaStringBuffer.h
--- a/Engine/Source/Programs/UnrealBuildAccelerator/Core/Public/UbaStringBuffer.h
+++ b/Engine/Source/Programs/UnrealBuildAccelerator/Core/Public/UbaStringBuffer.h
@@ -19,6 +19,7 @@ namespace uba
 	inline void ToLower(tchar* str) { while (tchar c = *str) { if (c >= 'A' && c <= 'Z') *str = c - 'A' + 'a'; ++str; } }
 	inline tchar ToLower(tchar c) { return (c >= 'A' && c <= 'Z') ? (c - 'A' + 'a') : c;}
 	inline tchar ToUpper(tchar c) { return (c >= 'a' && c <= 'z') ? (c - 'a' + 'A') : c;}
+	const tchar* GetFileName(const tchar* path);
 
 	class StringBufferBase
 	{
@@ -59,7 +60,7 @@ namespace uba
 		bool Equals(const tchar* str, bool ignoreCase = true) const { return uba::Equals(data, str, ignoreCase); }
 		const tchar* First(tchar c, u64 offset = 0) const;
 		const tchar* Last(tchar c, u64 offset = 0) const;
-		const tchar* GetFileName() const;
+		const tchar* GetFileName() const { return uba::GetFileName(data); }
 		inline StringBufferBase& Replace(tchar from, tchar to) { uba::Replace(data, from, to);  return *this; }
 
 		StringBufferBase& EnsureEndsWithSlash();
@@ -110,6 +111,7 @@ namespace uba
 		bool EndsWith(const tchar* value, bool ignoreCase = true) const { return uba::EndsWith(data, count, value, ignoreCase); }
 		bool Contains(const tchar* str, bool ignoreCase = true) const { return uba::Contains(data, str, ignoreCase); }
 		bool Equals(const tchar* str, bool ignoreCase = true) const { return uba::Equals(data, str, ignoreCase); }
+		StringView GetFileName() const { const tchar* str = uba::GetFileName(data); return StringView(str, count - u32(str - data)); }
 
 		const tchar* data;
 		u32 count;
