[UBA]
* Enabled async unmap of view of file in storage.

#skipci

diff --git a/Engine/Source/Programs/UnrealBuildAccelerator/Common/Private/UbaStorage.cpp b/Engine/Source/Programs/UnrealBuildAccelerator/Common/Private/UbaStorage.cpp
--- a/Engine/Source/Programs/UnrealBuildAccelerator/Common/Private/UbaStorage.cpp
+++ b/Engine/Source/Programs/UnrealBuildAccelerator/Common/Private/UbaStorage.cpp
@@ -81,6 +81,7 @@ namespace uba
 			return;
 		const ConfigTable& table = *tablePtr;
 		table.GetValueAsString(rootDir, TC("RootDir"));
+		table.GetValueAsBool(asyncUnmapViewOfFile, TC("AsyncUnmapViewOfFile"));
 	}
 
 	struct FileMappingScope
@@ -377,11 +378,10 @@ namespace uba
 
 				auto udg = MakeGuard([&]()
 					{
-						#if UBA_EXPERIMENTAL
-						m_workManager->AddWork([this, uncompressedData, fileSize, f = TString(from)](const WorkContext&) { UnmapViewOfFile(m_logger, uncompressedData, fileSize, f.c_str()); }, 1, TC("UnmapFile"));
-						#else
-						UnmapViewOfFile(m_logger, uncompressedData, fileSize, from);
-						#endif
+						if (m_asyncUnmapViewOfFile)
+							m_workManager->AddWork([this, uncompressedData, fileSize, f = TString(from)](const WorkContext&) { UnmapViewOfFile(m_logger, uncompressedData, fileSize, f.c_str()); }, 1, TC("UnmapFile"));
+						else
+							UnmapViewOfFile(m_logger, uncompressedData, fileSize, from);
 					});
 
 				if (!WriteMemToCompressedFile(destinationFile, workCount, uncompressedData, fileSize, maxUncompressedBlock, totalWritten))
@@ -1313,6 +1313,7 @@ namespace uba
 		m_casCapacityBytes = info.casCapacityBytes;
 		m_storeCompressed = info.storeCompressed;
 		m_manuallyHandleOverflow = info.manuallyHandleOverflow;
+		m_asyncUnmapViewOfFile = info.asyncUnmapViewOfFile && m_workManager;
 		m_rootDir.count = GetFullPathNameW(info.rootDir, m_rootDir.capacity, m_rootDir.data, NULL);
 		m_rootDir.Replace('/', PathSeparator).EnsureEndsWithSlash();
 
@@ -2869,11 +2870,9 @@ namespace uba
 				return CasKeyZero;
 			}
 			auto udg = MakeGuard([&]() { 
-					#if UBA_EXPERIMENTAL
-					if (m_workManager)
+					if (m_asyncUnmapViewOfFile)
 						m_workManager->AddWork([this, fileData, fileSize, fn = TString(fileName)](const WorkContext&) { UnmapViewOfFile(m_logger, fileData, fileSize, fn.c_str()); }, 1, TC("UnmapFile"));
 					else
-					#endif
 						UnmapViewOfFile(m_logger, fileData, fileSize, fileName);
 				});
 
diff --git a/Engine/Source/Programs/UnrealBuildAccelerator/Common/Public/UbaStorage.h b/Engine/Source/Programs/UnrealBuildAccelerator/Common/Public/UbaStorage.h
--- a/Engine/Source/Programs/UnrealBuildAccelerator/Common/Public/UbaStorage.h
+++ b/Engine/Source/Programs/UnrealBuildAccelerator/Common/Public/UbaStorage.h
@@ -89,6 +89,7 @@ namespace uba
 		u32 maxParallelCopyOrLink = 1000;
 		bool storeCompressed = true;
 		bool manuallyHandleOverflow = false;
+		bool asyncUnmapViewOfFile = true;
 		WorkManager* workManager = nullptr;
 	};
 
@@ -244,6 +245,8 @@ namespace uba
 		bool m_overflowReported = false;
 		bool m_storeCompressed = false;
 		bool m_manuallyHandleOverflow = false;
+		bool m_asyncUnmapViewOfFile = true;
+
 
 		ExternalFileMappingsProvider m_externalFileMappingsProvider;
 
