[UBA]
* Removed UBA_USE_SPARSEFILE because that code path is very outdated and would require a proper reimplementation if it would ever happen

#skipci

diff --git a/Engine/Source/Programs/UnrealBuildAccelerator/Common/Private/UbaSessionClient.cpp b/Engine/Source/Programs/UnrealBuildAccelerator/Common/Private/UbaSessionClient.cpp
--- a/Engine/Source/Programs/UnrealBuildAccelerator/Common/Private/UbaSessionClient.cpp
+++ b/Engine/Source/Programs/UnrealBuildAccelerator/Common/Private/UbaSessionClient.cpp
@@ -130,10 +130,6 @@ namespace uba
 		TimerScope s(m_stats.storageRetrieve);
 		CasKey tempKey = casKey;
 
-#if UBA_USE_SPARSEFILE
-		storeUncompressed = false;
-#endif
-
 		if (storeUncompressed)
 			tempKey = AsCompressed(casKey, false);
 		//else
@@ -505,18 +501,8 @@ namespace uba
 					if (!RetrieveCasFile(newCasKey, fileSize, casKey, fileName.data, storeUncompressed, allowProxy))
 						return m_logger.Error(TC("Error retrieving cas entry %s (%s)"), CasKeyString(casKey).str, fileName.data);
 
-					#if !UBA_USE_SPARSEFILE
 					if (!m_storage.GetCasFileName(newName, newCasKey))
 						return false;
-					#else
-					if (!memoryMapAlignment)
-						memoryMapAlignment = 4096;
-					MemoryMap map;
-					if (!CreateMemoryMapFromView(map, fileNameKey, fileName.data, newCasKey, memoryMapAlignment))
-						return false;
-					newName.Append(map.name);
-					fileSize = map.size;
-					#endif
 				}
 				else
 				{
diff --git a/Engine/Source/Programs/UnrealBuildAccelerator/Common/Private/UbaStorage.cpp b/Engine/Source/Programs/UnrealBuildAccelerator/Common/Private/UbaStorage.cpp
--- a/Engine/Source/Programs/UnrealBuildAccelerator/Common/Private/UbaStorage.cpp
+++ b/Engine/Source/Programs/UnrealBuildAccelerator/Common/Private/UbaStorage.cpp
@@ -18,10 +18,6 @@ namespace uba
 	constexpr u32 CasTableVersion = IsWindows ? 32 : 34;
 	constexpr u32 MaxWorkItemsPerAction = 128; // Cap this to not starve other things
 
-#if UBA_USE_SPARSEFILE
-	constexpr u32 CasDbDataFileCount = 16;
-#endif
-
     CasKey EmptyFileKey = []() { CasKeyHasher hasher; return ToCasKey(hasher, false); }();
 
 	void Storage::GetMappingString(StringBufferBase& out, FileMappingHandle mappingHandle, u64 offset)
@@ -161,17 +157,11 @@ namespace uba
 
 		casEntry.size = size;
 
-#if !UBA_USE_SPARSEFILE
 		UBA_ASSERT(!casEntry.mappingHandle.IsValid());
-#endif
 
 		if (!m_casCapacityBytes || m_overflowReported || m_casTotalBytes <= m_casCapacityBytes || m_manuallyHandleOverflow)
 			return;
 
-		#if UBA_USE_SPARSEFILE
-		UBA_ASSERT(false);
-		#endif
-
 		TimerScope ts(Stats().handleOverflow);
 		UBA_ASSERT(!m_newestAccessed || !m_newestAccessed->prevAccessed);
 		UBA_ASSERT(!m_oldestAccessed || !m_oldestAccessed->nextAccessed);
@@ -215,7 +205,6 @@ namespace uba
 
 		for (Rec& rec : toDelete)
 		{
-#if !UBA_USE_SPARSEFILE
 			StringBuffer<> casFile;
 			StorageImpl::GetCasFileName(casFile, rec.casEntry.key);
 
@@ -240,10 +229,6 @@ namespace uba
 					continue;
 				}
 			}
-#else
-			// TODO: Zero out space in the file!
-			//UBA_ASSERT(false);
-#endif
 
 			//m_logger.Warning(TC("Evicted %s due to overflow"), casFile.data);
 
@@ -334,7 +319,6 @@ namespace uba
 		u64 maxUncompressedBlock = BufferSlotHalfSize - diff - 8; // 8 bytes for the little header
 		u32 workCount = u32((fileSize + maxUncompressedBlock - 1) / maxUncompressedBlock);
 
-#if !UBA_USE_SPARSEFILE
 		FileAccessor destinationFile(m_logger, to);
 		if (!destinationFile.CreateWrite(false, DefaultAttributes(), 0, m_tempPath.data))
 			return false;
@@ -343,22 +327,7 @@ namespace uba
 				return false;
 		if (!destinationFile.Write(&fileSize, sizeof(u64))) // Store file size first in compressed file
 			return false;
-#else
-		u64 reserveSize = fileSize + (workCount*(diff+8)) + 8; // Create headroom for oodle headers etc. Since we're using a sparse file we can just ignore the unused space
-
-		auto mappedView = m_casDataBuffer.AllocAndMapView(MappedView_Persistent, reserveSize, 1, from, true);
-		auto mvg = MakeGuard([&]() { m_casDataBuffer.UnmapView(mappedView, from, totalWritten); });
 
-		out.mappingHandle = mappedView.handle;
-		out.offset = mappedView.offset;
-		u8* toMem = mappedView.memory;
-		if (!toMem)
-		{
-			UBA_ASSERT(false);
-			return false;
-		}
-		*(u64*)toMem = fileSize;
-#endif
 		totalWritten += sizeof(u64);
 
 		u64 left = fileSize;
@@ -433,13 +402,8 @@ namespace uba
 				*(u32*)(destBuf+4) =  u32(uncompressedBlockSize);
 
 				u32 writeBytes = u32(compressedBlockSize) + 8;
-#if !UBA_USE_SPARSEFILE
 				if (!destinationFile.Write(destBuf, writeBytes))
 					return false;
-#else
-				TimerScope cts(stats.memoryCopy);
-				memcpy(toMem + totalWritten, destBuf, writeBytes);
-#endif
 
 				totalWritten += writeBytes;
 
@@ -447,13 +411,12 @@ namespace uba
 			}
 		}
 
-#if !UBA_USE_SPARSEFILE
 		if (lastWriteTime)
 			if (!SetFileLastWriteTime(destinationFile.GetHandle(), lastWriteTime))
 				return m_logger.Error(TC("Failed to set file time on filehandle for %s"), to);
 		if (!destinationFile.Close())
 			return false;
-#endif
+
 		stats.createCasBytesRaw += fileSize;
 		stats.createCasBytesComp += totalWritten;
 
@@ -497,12 +460,7 @@ namespace uba
 			Atomic<u64> compressCounter;
 			Event* events;
 			const u8* uncompressedData = nullptr;
-#if !UBA_USE_SPARSEFILE
 			FileAccessor* destination = nullptr;
-#else
-			u8* mem = nullptr;
-			u64 memPos = 0;
-#endif
 			u64 written = 0;
 			u64 workCount = 0;
 			u64 maxUncompressedBlock = 0;
@@ -514,12 +472,7 @@ namespace uba
 		rec->uncompressedData = uncompressedData;
 		rec->maxUncompressedBlock = maxUncompressedBlock;
 		rec->fileSize = fileSize;
-#if !UBA_USE_SPARSEFILE
 		rec->destination = &destination;
-#else
-		rec->mem = toMem;
-		rec->memPos = totalWritten;
-#endif
 
 		for (u32 i=0; i!=workCount; ++i)
 			rec->events[i].Create(true);
@@ -564,14 +517,9 @@ namespace uba
 
 				u32 writeBytes = u32(compressedBlockSize) + 8;
 
-#if !UBA_USE_SPARSEFILE
 				if (!rec->destination->Write(compressSlotBuffer, writeBytes))
 					rec->error = true;
-#else
-				TimerScope cts(stats.memoryCopy);
-				memcpy(rec->mem + rec->memPos, compressSlotBuffer, writeBytes);
-				rec->memPos += writeBytes;
-#endif
+
 				rec->written += writeBytes;
 				if (index < rec->workCount)
 					rec->events[index].Set();
@@ -660,14 +608,10 @@ namespace uba
 	{
 		UBA_ASSERT(IsCompressed(casKey) == m_storeCompressed);
 		StringBuffer<> casFile;
-#if !UBA_USE_SPARSEFILE
 		if (!StorageImpl::GetCasFileName(casFile, casKey))
 			return false;
 		if (FileExists(m_logger, casFile.data))
 			return true;
-#else
-		//UBA_ASSERT(false);
-#endif
 
 		StringBuffer<> temp;
 		temp.Append(fileName);
@@ -963,7 +907,6 @@ namespace uba
 
 		casEntry.disallowed = IsDisallowedPath(fileName);
 
-#if !UBA_USE_SPARSEFILE
 		StringBuffer<> casFile;
 		if (!StorageImpl::GetCasFileName(casFile, casKey)) // Force to not use subclassed version
 			return false;
@@ -971,7 +914,6 @@ namespace uba
 		bool verifyReturnValue;
 		if (!casEntry.verified && VerifyExisting(verifyReturnValue, entryLock, casKey, casEntry, casFile))
 			return verifyReturnValue;
-#endif
 
 		if (deferCreation)
 		{
@@ -992,14 +934,9 @@ namespace uba
 		casEntry.exists = false;
 
 		WriteResult res;
-#if !UBA_USE_SPARSEFILE
 		if (!WriteCasFileNoCheck(res, fileNameKey, fileName, fileIsCompressed, casFile.data, IsCompressed(casKey)))
 			return false;
-#else
 
-		if (!WriteCompressed(res, fileName, nullptr))
-			return false;
-#endif
 		casEntry.exists = true;
 		entryLock.Leave();
 		CasEntryWritten(casEntry, res.size);
@@ -1058,7 +995,6 @@ namespace uba
 
 	bool StorageImpl::CheckAllCasFiles(u64 checkContentOfFilesNewerThanTime)
 	{
-#if !UBA_USE_SPARSEFILE
 		u64 startTime = GetTime();
 		u64 before = m_casTotalBytes;
 		m_casTotalBytes = 0;
@@ -1202,9 +1138,6 @@ namespace uba
 		else
 			m_logger.Info(TC("Validated storage (size %s) in %s"), BytesToText(after).str, TimeToText(duration).str);
 		m_casMaxBytes = m_casTotalBytes;
-#else
-		UBA_ASSERT(false); // not implemented
-#endif
 		return true;
 	}
 
@@ -1234,64 +1167,6 @@ namespace uba
 			m_logger.Info(TC("Evicted %s from storage (%s). Estimated new storage is now %s (there might be files db is not aware of)"), BytesToText(before - after).str, TimeToText(GetTime() - startTime).str, BytesToText(after).str);
 	}
 
-	bool StorageImpl::OpenCasDataFile(u32 index, u64 size)
-	{
-#if UBA_USE_SPARSEFILE
-		bool createFile = size == 0;
-		u32 createDisposition = createFile ? CREATE_ALWAYS : OPEN_EXISTING;
-		StringBuffer<> sparseFileName;
-		sparseFileName.Append(m_rootDir).Appendf(TC("casdbdata_%02u"), index);
-		FileHandle sparseFile = uba::CreateFileW(sparseFileName.data, GENERIC_READ | GENERIC_WRITE, FILE_SHARE_READ | FILE_SHARE_WRITE, createDisposition, DefaultAttributes());
-		auto sparseFileGuard = MakeGuard([&](){ CloseFile(sparseFile); });
-		if (sparseFile == InvalidFileHandle)
-			return m_logger.Error(TC("Failed to create database file %s (%s)"), sparseFileName.data, LastErrorToText().data);
-		
-		if (createFile)
-		{
-			DWORD dwTemp;
-			if (!::DeviceIoControl(sparseFile, FSCTL_SET_SPARSE, NULL, 0, NULL, 0, &dwTemp, NULL))
-				return m_logger.Error(TC("Failed to make file %s sparse (%s)"), sparseFileName.data, LastErrorToText().data);
-		}
-		else
-		{
-			u64 fileSize;
-			if (!uba::GetFileSizeEx(fileSize, sparseFile))
-				return m_logger.Error(TC("GetFileSize failed for %s (%s)"), fileName, LastErrorToText().data);
-
-			if (fileSize < size)
-				return m_logger.Error(TC("Sparse file size is smaller than what cas db think it is. Expected %llu, was %llu (%s)"), size, fileSize, sparseFileName.data);
-		}
-		// Mark the range as sparse zero block
-		// FILE_ZERO_DATA_INFORMATION fzdi;
-		// fzdi.FileOffset.QuadPart = start;
-		// fzdi.BeyondFinalZero.QuadPart = start + sparseCapacity;
-		// if (!DeviceIoControl(sparseFile, FSCTL_SET_ZERO_DATA, &fzdi, sizeof(fzdi), NULL, 0, &dwTemp, NULL))
-		// {
-		// 	m_logger.Error(TC("Failed to zero out file"));
-		// 	return false;
-		// }
-
-		u64 sparseCapacity = (1ull << 44) - (64*1024);
-		if (!m_casDataBuffer.AddPersistent(TC("CasDbData"), sparseFile, size, sparseCapacity))
-			return false;
-
-		sparseFileGuard.Cancel();
-		return true;
-#else
-		return false;
-#endif
-	}
-
-	bool StorageImpl::CreateCasDataFiles()
-	{
-#if UBA_USE_SPARSEFILE
-		for (u32 i=0; i!=CasDbDataFileCount; ++i)
-			if (!OpenCasDataFile(i, 0))
-				return false;
-#endif
-		return true;
-	}
-
 	#if UBA_USE_MIMALLOC
 	void* Oodle_MallocAligned(OO_SINTa bytes, OO_S32 alignment)
 	{
@@ -1380,32 +1255,15 @@ namespace uba
 				return m_logger.Error(TC("Failed to create temporary \".isRunning\" file"));
 		}
 
-
-#if UBA_USE_SPARSEFILE
-		if (wasTerminated)
-		{
-			// Do this better!
-			m_logger.Warning(TC("Non-graceful shutdown. Delete all cas data. TODO: This should be improved!"));
-			DeleteAllCas();
-			wasTerminated = false;
-		}
-#endif
-
 		StringBuffer<> fileName;
 		fileName.Append(m_rootDir).Append(TC("casdb"));
-		#if UBA_USE_SPARSEFILE
-		fileName.Append('2');
-		#endif
 
 		FileHandle fileHandle;
 		if (!OpenFileSequentialRead(m_logger, fileName.data, fileHandle, false))
 			return false;
 		if (fileHandle == InvalidFileHandle)
-		{
-			if (!CreateCasDataFiles())
-				return false;
 			return true;
-		}
+
 		auto fileGuard = MakeGuard([&](){ CloseFile(fileName.data, fileHandle); });
 
 		u64 fileSize;
@@ -1433,8 +1291,6 @@ namespace uba
 			loadSaveLock.Leave();
 			if (!SaveCasTable(false, false)) // Save it just to get the version down
 				return false;
-			if (!CreateCasDataFiles())
-				return false;
 			return true;
 		}
 
@@ -1471,15 +1327,6 @@ namespace uba
 		}
 
 
-#if UBA_USE_SPARSEFILE
-		for (u32 i=0; i!=CasDbDataFileCount; ++i)
-		{
-			u64 size = reader.Read7BitEncoded();
-			if (!OpenCasDataFile(i, size))
-				return false;
-		}
-#endif
-
 		m_casLookup.reserve(casLookupSize);
 		CasEntry* prev = nullptr;
 		while (true)
@@ -1501,20 +1348,6 @@ namespace uba
 			casEntry.exists = true;
 			m_casTotalBytes += casEntry.size;
 
-
-#if UBA_USE_SPARSEFILE
-			u32 mappingFileIndex = reader.ReadU32();
-			u64 mappingOffset = reader.ReadU64();
-			u64 mappingSize = reader.ReadU64();
-
-			casEntry.mappingHandle = m_casDataBuffer.GetPersistentHandle(mappingFileIndex);
-			casEntry.mappingOffset = mappingOffset;
-			casEntry.mappingSize = mappingSize;
-			casEntry.verified = true;
-			casEntry.exists = true;
-#endif
-
-
 			if (prev)
 			{
 				prev->nextAccessed = &casEntry;
@@ -1572,32 +1405,8 @@ namespace uba
 		if (!m_casTableLoaded)
 			return true;
 
-		#if UBA_USE_SPARSEFILE
-		UnorderedMap<HANDLE, u32> handleToIndex;
-		handleToIndex.reserve(CasDbDataFileCount);
-		for (u32 i=0; i!=CasDbDataFileCount; ++i)
-			handleToIndex.emplace(m_casDataBuffer.GetPersistentHandle(i), i);
-		m_casDataBuffer.CloseDatabase();
-		/*
-		{
-			FILE_ALLOCATED_RANGE_BUFFER queryRange;
-			queryRange.FileOffset.QuadPart = 0;
-			queryRange.Length = ToLargeInteger(m_casDataBuffer.GetPersistentSize(0));
-			FILE_ALLOCATED_RANGE_BUFFER allocRanges[1024];
-			DWORD nbytes;
-			if (!DeviceIoControl(m_casDataBuffer.GetPersistentFile(0), FSCTL_QUERY_ALLOCATED_RANGES, &queryRange, sizeof(queryRange), allocRanges, sizeof(allocRanges), &nbytes, NULL))
-				return m_logger.Error(TC("Failed to make file %s sparse (%s)"), TC("FOO"), LastErrorToText().data);
-			DWORD dwAllocRangeCount = nbytes / sizeof(FILE_ALLOCATED_RANGE_BUFFER);
-			printf("Range count: %u"), dwAllocRangeCount);
-		}
-		*/
-		#endif
-
 		StringBuffer<256> fileName;
 		fileName.Append(m_rootDir).Append(TC("casdb"));
-		#if UBA_USE_SPARSEFILE
-		fileName.Append('2');
-		#endif
 
 		StringBuffer<256> tempFileName;
 		tempFileName.Append(fileName).Append(TC(".tmp"));
@@ -1654,27 +1463,12 @@ namespace uba
 					continue;
 				if (casEntry.dropped)
 				{
-#if !UBA_USE_SPARSEFILE
 					StringBuffer<512> casFileName;
 					if (!StorageImpl::GetCasFileName(casFileName, casEntry.key))
 						continue;
 					DeleteFileW(casFileName.data);
-#else
-					// TODO!  UBA_ASSERT(false);
-#endif
-					continue;
-				}
-
-#if UBA_USE_SPARSEFILE
-				auto findIt = handleToIndex.find(casEntry.mappingHandle);
-				if (findIt == handleToIndex.end())
-				{
-					if (m_deferredCasCreationLookup.find(casEntry.key) != m_deferredCasCreationLookup.end())
-						continue;
-					m_logger.Error(TC("Can't find cas database file with mappingHandle %llu"), uintptr_t(casEntry.mappingHandle));
 					continue;
 				}
-#endif
 
 				if (writer.GetCapacityLeft() < casLookupEntrySize + sizeof(CasKey))
 					return m_logger.Error(TC("This should not happen, somehow there are more valid entries in access list than lookup. (Lookup has %llu entries)"), m_casLookup.size());
@@ -1683,12 +1477,6 @@ namespace uba
 				writer.WriteCasKey(casEntry.key);
 				writer.WriteU64(casEntry.size);
 				++casTableSize;
-
-#if UBA_USE_SPARSEFILE
-				writer.WriteU32(findIt->second);
-				writer.WriteU64(casEntry.mappingOffset);
-				writer.WriteU64(casEntry.mappingSize);
-#endif
 			}
 			writer.WriteCasKey(CasKeyZero);
 			*casLookupSizePtr = casTableSize;
@@ -1904,17 +1692,7 @@ namespace uba
 		m_casTotalBytes = 0;
 		m_casMaxBytes = 0;
 
-#if UBA_USE_SPARSEFILE
-		if (m_casTableLoaded)
-			m_casDataBuffer.CloseDatabase();
-#endif
-
 		DeleteAllCas();
-
-#if UBA_USE_SPARSEFILE
-		if (m_casTableLoaded)
-			CreateCasDataFiles();
-#endif
 		return true;
 	}
 
@@ -2262,11 +2040,6 @@ namespace uba
 		WriteResult res;
 		if (!WriteCasFile(res, deferredCreation.data, fileIsCompressed, casKey))
 			return false;
-#if UBA_USE_SPARSEFILE
-		casEntry.mappingHandle = res.mappingHandle;
-		casEntry.mappingOffset = res.offset;
-		casEntry.mappingSize = res.size;
-#endif
 
 		casEntry.exists = true;
 		entryLock.Leave();
@@ -2294,14 +2067,13 @@ namespace uba
 		}
 
 		StringBuffer<> casFile;
-#if !UBA_USE_SPARSEFILE
 		if (!StorageImpl::GetCasFileName(casFile, casKey))
 			return false;
 
 		bool verifyReturnValue;
 		if (VerifyExisting(verifyReturnValue, entryLock, casKey, casEntry, casFile))
 			return verifyReturnValue;
-#endif
+
 		casEntry.exists = false;
 		casEntry.verified = true;
 		if (!fileName)
@@ -2316,7 +2088,6 @@ namespace uba
 		return true;
 	}
 
-#if !UBA_USE_SPARSEFILE
 	bool StorageImpl::GetCasFileName(StringBufferBase& out, const CasKey& casKey)
 	{
 		out.Append(m_rootDir.data).AppendHex(((const u8*)&casKey)[0]);
@@ -2325,7 +2096,6 @@ namespace uba
 		out.Append(PathSeparator).Append(CasKeyString(casKey).str);
 		return true;
 	}
-#endif
 
 	MappedView StorageImpl::MapView(const CasKey& casKey, const tchar* hint)
 	{
@@ -2371,12 +2141,8 @@ namespace uba
 			if (forceDelete)
 			{
 				StringBuffer<> casFile;
-				#if !UBA_USE_SPARSEFILE
 				if (!StorageImpl::GetCasFileName(casFile, casKey))
 					return false;
-				#else
-				UBA_ASSERT(false);
-				#endif
 				if (DeleteFileW(casFile.data) == 0)
 				{
 					u32 lastError = GetLastError();
@@ -2394,12 +2160,8 @@ namespace uba
 		if (forceDelete)
 		{
 			StringBuffer<> casFile;
-			#if !UBA_USE_SPARSEFILE
 			if (!StorageImpl::GetCasFileName(casFile, casKey))
 				return false;
-			#else
-			UBA_ASSERT(false);
-			#endif
 
 			u64 sizeDeleted = 0;
 			if (DeleteFileW(casFile.data) == 0)
@@ -2535,12 +2297,9 @@ namespace uba
 				}
 				else
 				{
-#if !UBA_USE_SPARSEFILE
 					if (!StorageImpl::GetCasFileName(casFile, actualKey))
 						return false;
-#else
-					UBA_ASSERT(false);
-#endif
+
 					if (!OpenFileSequentialRead(m_logger, casFile.data, readHandle))
 						return m_logger.Error(TC("[CopyOrLink] Failed to open file %s for read (%s)"), casFile.data, LastErrorToText().data);
 
@@ -2693,12 +2452,8 @@ namespace uba
 			}
 
 			StringBuffer<> casFile;
-			#if !UBA_USE_SPARSEFILE
 			if (!GetCasFileName(casFile, actualKey))
 				return false;
-			#else
-			UBA_ASSERT(false);
-			#endif
 
 			SCOPED_WRITE_LOCK(fileEntry.lock, entryLock);
 			fileEntry.verified = false;
diff --git a/Engine/Source/Programs/UnrealBuildAccelerator/Common/Private/UbaStorageClient.cpp b/Engine/Source/Programs/UnrealBuildAccelerator/Common/Private/UbaStorageClient.cpp
--- a/Engine/Source/Programs/UnrealBuildAccelerator/Common/Private/UbaStorageClient.cpp
+++ b/Engine/Source/Programs/UnrealBuildAccelerator/Common/Private/UbaStorageClient.cpp
@@ -138,7 +138,6 @@ namespace uba
 		return success;
 	}
 
-#if !UBA_USE_SPARSEFILE
 	bool StorageClient::GetCasFileName(StringBufferBase& out, const CasKey& casKey)
 	{
 		SCOPED_READ_LOCK(m_localStorageFilesLock, tempLock);
@@ -155,7 +154,6 @@ namespace uba
 
 		return StorageImpl::GetCasFileName(out, casKey);
 	}
-#endif
 
 	MappedView StorageClient::MapView(const CasKey& casKey, const tchar* hint)
 	{
@@ -187,14 +185,6 @@ namespace uba
 	{
 		FileMappingType mappingType = MappedView_Transient;
 		bool shouldStore = mappingBuffer == nullptr;
-		#if UBA_USE_SPARSEFILE
-		if (shouldStore)
-		{
-			mappingType = MappedView_Persistent;
-			mappingBuffer = &m_casDataBuffer;
-		}
-		#endif
-
 		UBA_ASSERT(AsCompressed(casKey, false) != CasKeyZero);
 
 		out.casKey = casKey;
@@ -241,11 +231,7 @@ namespace uba
 		TimerScope ts2(stats.recvCas);
 
 		StringBuffer<> casFile;
-		#if !UBA_USE_SPARSEFILE
 		GetCasFileName(casFile, casKey);
-		#else
-		casFile.Append(CasKeyString(casKey).str);
-		#endif
 
 		u8* slot = m_bufferSlots.Pop();
 		auto slotGuard = MakeGuard([&](){ m_bufferSlots.Push(slot); });
@@ -465,9 +451,7 @@ namespace uba
 
 			sizeOnDisk = IsCompressed(casKey) ? fileSize : actualSize;
 
-#if !UBA_USE_SPARSEFILE
 			FileAccessor destinationFile(m_logger, casFile.data);
-#endif
 
 			u8* writePos = nullptr;
 			bool isInitialized = false;
@@ -872,13 +856,7 @@ namespace uba
 					else
 					{
 						localFile.casEntry.size = fileSize;
-#if !UBA_USE_SPARSEFILE
 						localFile.fileName = fileName;
-#else
-						localFile.casEntry.mappingHandle = mappingHandle2;
-						lock.Leave();
-						mappingClose.Cancel();
-#endif
 					}
 				}
 			}
diff --git a/Engine/Source/Programs/UnrealBuildAccelerator/Common/Private/UbaStorageServer.cpp b/Engine/Source/Programs/UnrealBuildAccelerator/Common/Private/UbaStorageServer.cpp
--- a/Engine/Source/Programs/UnrealBuildAccelerator/Common/Private/UbaStorageServer.cpp
+++ b/Engine/Source/Programs/UnrealBuildAccelerator/Common/Private/UbaStorageServer.cpp
@@ -519,7 +519,6 @@ namespace uba
 				}
 				else
 				{
-	#if !UBA_USE_SPARSEFILE
 					GetCasFileName(casFile, casKey);
 					if (!OpenFileSequentialRead(m_logger, casFile.data, readFileHandle))
 					{
@@ -529,9 +528,7 @@ namespace uba
 
 					if (!uba::GetFileSizeEx(fileSize, readFileHandle))
 						return m_logger.Error(TC("GetFileSizeEx failed on file %s (%s)"), casFile.data, LastErrorToText().data);
-	#else
-					UBA_ASSERT(false);
-	#endif
+
 					if (fileSize > BufferSlotSize)
 					{
 						ownsMapping = true;
@@ -725,7 +722,6 @@ namespace uba
 
 				if (!exists && casEntry.exists)
 				{
-#if !UBA_USE_SPARSEFILE
 					StringBuffer<> casFile;
 					if (!GetCasFileName(casFile, casKey))
 						return false;
@@ -753,7 +749,6 @@ namespace uba
 						casEntry.exists = false;
 						casEntry.verified = true;
 					}
-#endif
 				}
 				writer.WriteBool(exists);
 				return true;
diff --git a/Engine/Source/Programs/UnrealBuildAccelerator/Common/Public/UbaStorage.h b/Engine/Source/Programs/UnrealBuildAccelerator/Common/Public/UbaStorage.h
--- a/Engine/Source/Programs/UnrealBuildAccelerator/Common/Public/UbaStorage.h
+++ b/Engine/Source/Programs/UnrealBuildAccelerator/Common/Public/UbaStorage.h
@@ -9,8 +9,6 @@
 #include "UbaPathUtils.h"
 #include "UbaStats.h"
 
-#define UBA_USE_SPARSEFILE 0
-
 namespace uba
 {
 	class Config;
@@ -53,9 +51,8 @@ namespace uba
 		using FormattingFunc = Function<bool(MemoryBlock& destData, const void* sourceData, u64 sourceSize, const tchar* hint)>;
 		virtual bool CopyOrLink(const CasKey& casKey, const tchar* destination, u32 fileAttributes, bool writeCompressed = false, const FormattingFunc& formattingFunc = {}) = 0;
 		virtual bool FakeCopy(const CasKey& casKey, const tchar* destination, u64 size = 0, u64 lastWritten = 0, bool deleteExisting = true) = 0;
-#if !UBA_USE_SPARSEFILE
 		virtual bool GetCasFileName(StringBufferBase& out, const CasKey& casKey) = 0;
-#endif
+
 		virtual MappedView MapView(const CasKey& casKey, const tchar* hint) = 0;
 		virtual void UnmapView(const MappedView& view, const tchar* hint) = 0;
 
@@ -140,9 +137,7 @@ namespace uba
 		virtual bool CalculateCasKey(CasKey& out, const tchar* fileName) override;
 		virtual bool CopyOrLink(const CasKey& casKey, const tchar* destination, u32 fileAttributes, bool writeCompressed = false, const FormattingFunc& formattingFunc = {}) override;
 		virtual bool FakeCopy(const CasKey& casKey, const tchar* destination, u64 size = 0, u64 lastWritten = 0, bool deleteExisting = true) override;
-#if !UBA_USE_SPARSEFILE
 		virtual bool GetCasFileName(StringBufferBase& out, const CasKey& casKey) override;
-#endif
 		virtual MappedView MapView(const CasKey& casKey, const tchar* hint) override;
 		virtual void UnmapView(const MappedView& view, const tchar* hint) override;
 
@@ -185,8 +180,6 @@ namespace uba
 		void TraverseAllCasFiles(const Function<void(const CasKey& key, u64 size)>& func, bool allowParallel = false);
 		bool CheckAllCasFiles(u64 checkContentOfFilesNewerThanTime = ~u64(0));
 		void HandleOverflow(UnorderedSet<CasKey>* outDeletedFiles);
-		bool OpenCasDataFile(u32 index, u64 size);
-		bool CreateCasDataFiles();
 
 		struct FileEntry;
 		FileEntry& GetOrCreateFileEntry(StringKey fileNameKey);
diff --git a/Engine/Source/Programs/UnrealBuildAccelerator/Common/Public/UbaStorageClient.h b/Engine/Source/Programs/UnrealBuildAccelerator/Common/Public/UbaStorageClient.h
--- a/Engine/Source/Programs/UnrealBuildAccelerator/Common/Public/UbaStorageClient.h
+++ b/Engine/Source/Programs/UnrealBuildAccelerator/Common/Public/UbaStorageClient.h
@@ -46,9 +46,7 @@ namespace uba
 		using DirVector = Vector<TString>;
 		bool PopulateCasFromDirs(const DirVector& directories, u32 workerCount, const Function<bool()>& shouldExit = {});
 
-		#if !UBA_USE_SPARSEFILE
 		virtual bool GetCasFileName(StringBufferBase& out, const CasKey& casKey) override;
-		#endif
 
 		virtual MappedView MapView(const CasKey& casKey, const tchar* hint) override;
 
