[UBA]
* Fixed compile errors due to missing files in previous submit

#skipci

diff --git a/Engine/Source/Programs/UnrealBuildAccelerator/Cli/Private/UbaCli.cpp b/Engine/Source/Programs/UnrealBuildAccelerator/Cli/Private/UbaCli.cpp
--- a/Engine/Source/Programs/UnrealBuildAccelerator/Cli/Private/UbaCli.cpp
+++ b/Engine/Source/Programs/UnrealBuildAccelerator/Cli/Private/UbaCli.cpp
@@ -576,7 +576,7 @@ namespace uba
 			WorkManagerImpl workManager(DefaultProcessorCount, TC("UbaWorker(CheckCas2)"));
 			storageServer.TraverseAllCasFiles([&](const CasKey& casKey, u64 size)
 				{
-					workManager.AddWork([&, casKey]()
+					workManager.AddWork([&, casKey](const WorkContext&)
 						{
 							Storage::RetrieveResult res;
 							storageServer.EnsureCasFile(casKey, TC("Dummy"));
diff --git a/Engine/Source/Programs/UnrealBuildAccelerator/Common/Private/UbaStorageProxy.cpp b/Engine/Source/Programs/UnrealBuildAccelerator/Common/Private/UbaStorageProxy.cpp
--- a/Engine/Source/Programs/UnrealBuildAccelerator/Common/Private/UbaStorageProxy.cpp
+++ b/Engine/Source/Programs/UnrealBuildAccelerator/Common/Private/UbaStorageProxy.cpp
@@ -59,7 +59,7 @@ namespace uba
 			});
 
 		m_server.RegisterService(StorageServiceId,
-			[this](const ConnectionInfo& connectionInfo, MessageInfo& messageInfo, BinaryReader& reader, BinaryWriter& writer)
+			[this](const ConnectionInfo& connectionInfo, const WorkContext& workContext, MessageInfo& messageInfo, BinaryReader& reader, BinaryWriter& writer)
 			{
 				return HandleMessage(connectionInfo, messageInfo, reader, writer);
 			},
@@ -254,7 +254,7 @@ namespace uba
 						}
 
 						// Move the additional messages to a job to be able to return this one quickly.
-						m_server.AddWork([f = &file, segmentCount, this]()
+						m_server.AddWork([f = &file, segmentCount, this](const WorkContext&)
 							{
 								SCOPED_WRITE_LOCK(m_largeFileLock, lock);
 								//TrackWorkScope tws(m_client, TC("SEGMENTS"));
@@ -267,7 +267,7 @@ namespace uba
 										{
 											auto mif = (MessageInFlight*)userData;
 											mif->error = error;
-											mif->proxy.m_server.AddWork([mif]() { mif->proxy.HandleReceivedData(*mif); }, 1, TC(""));
+											mif->proxy.m_server.AddWork([mif](const WorkContext&) { mif->proxy.HandleReceivedData(*mif); }, 1, TC(""));
 										}, mif);
 									if (!res)
 									{
