[UBA]
* Changed code fetching path table and cas key table to batch messages instead. Tables can be 30+mb which would require 120 messages and sending them sequentially with higher latency would take quite some time

#skipci

diff --git a/Engine/Source/Programs/UnrealBuildAccelerator/Common/Private/UbaCacheClient.cpp b/Engine/Source/Programs/UnrealBuildAccelerator/Common/Private/UbaCacheClient.cpp
--- a/Engine/Source/Programs/UnrealBuildAccelerator/Common/Private/UbaCacheClient.cpp
+++ b/Engine/Source/Programs/UnrealBuildAccelerator/Common/Private/UbaCacheClient.cpp
@@ -47,19 +47,21 @@ namespace uba
 
 		u32 id = 0;
 
+		ReaderWriterLock serverPathTableNetworkLock;
 		CompactPathTable serverPathTable;
+		Atomic<u32> serverPathTableSize;
+
+		ReaderWriterLock serverCasKeyTableNetworkLock;
 		CompactCasKeyTable serverCasKeyTable;
+		Atomic<u32> serverCasKeyTableSize;
 
 		CompactPathTable sendPathTable;
-		CompactCasKeyTable sendCasKeyTable;
-
-		ReaderWriterLock pathTableNetworkLock;
+		ReaderWriterLock sendPathTableNetworkLock;
 		u32 pathTableSizeSent = 0;
 
-		ReaderWriterLock casKeyTableNetworkLock;
+		CompactCasKeyTable sendCasKeyTable;
+		ReaderWriterLock sendCasKeyTableNetworkLock;
 		u32 casKeyTableSizeSent = 0;
-
-		Atomic<u32> availableCasKeyTableSize;
 	};
 
 	CacheClient::CacheClient(const CacheClientCreateInfo& info)
@@ -868,7 +870,7 @@ namespace uba
 
 	bool CacheClient::SendPathTable(Bucket& bucket, u32 requiredPathTableSize)
 	{
-		SCOPED_WRITE_LOCK(bucket.pathTableNetworkLock, lock);
+		SCOPED_WRITE_LOCK(bucket.sendPathTableNetworkLock, lock);
 		if (requiredPathTableSize <= bucket.pathTableSizeSent)
 			return true;
 
@@ -892,7 +894,7 @@ namespace uba
 
 	bool CacheClient::SendCasTable(Bucket& bucket, u32 requiredCasTableSize)
 	{
-		SCOPED_WRITE_LOCK(bucket.casKeyTableNetworkLock, lock);
+		SCOPED_WRITE_LOCK(bucket.sendCasKeyTableNetworkLock, lock);
 		if (requiredCasTableSize <= bucket.casKeyTableSizeSent)
 			return true;
 
@@ -1063,78 +1065,128 @@ namespace uba
 
 	bool CacheClient::FetchCasTable(Bucket& bucket, CacheStats& stats, u32 requiredCasTableOffset)
 	{
-		if (HasEnoughData(bucket, requiredCasTableOffset, bucket.availableCasKeyTableSize))
+		if (HasEnoughData(bucket, requiredCasTableOffset))
 			return true;
 		return FetchCasTable2(bucket, stats, requiredCasTableOffset);
 	}
 
-	UBA_NOINLINE bool CacheClient::FetchCasTable2(Bucket& bucket, CacheStats& stats, u32 requiredCasTableOffset)
+	template<typename TableType>
+	bool CacheClient::FetchCompactTable(u32 bucketId, TableType& table, u32 requiredTableSize, u8 messageType)
 	{
-		TimerScope ts2(stats.fetchCasTable);
+		u32 tableSize = table.GetSize();
+		u32 messageFetchSize = u32(m_client.GetMessageMaxSize() - m_client.GetMessageReceiveHeaderSize());
+		u32 totalFetchSize = requiredTableSize - tableSize;
+		u32 messageCount = (totalFetchSize + messageFetchSize - 1) / messageFetchSize;
 
-		StackBinaryReader<SendMaxSize> reader;
+		u8* data = table.BeginCommit(totalFetchSize);
 
-		SCOPED_WRITE_LOCK(bucket.casKeyTableNetworkLock, lock); // Use one lock over both queries
+		struct Entry
 		{
-			bool messageSent = false;
-			while (true)
-			{
-				u32 tableSize = bucket.serverCasKeyTable.GetSize();
-				if (HasEnoughData(bucket, requiredCasTableOffset, tableSize))
-				{
-					if (!messageSent)
-						return true;
-					break;
-				}
+			Entry(u8* slot, u32 i, u32 messageMaxSize) : reader(slot + i * messageMaxSize, 0, SendMaxSize), done(true) {}
+			NetworkMessage message;
+			BinaryReader reader;
+			Event done;
+		};
 
-				StackBinaryWriter<16> writer;
-				NetworkMessage msg(m_client, CacheServiceId, CacheMessageType_FetchCasTable, writer);
-				writer.Write7BitEncoded(MakeId(bucket.id));
-				writer.WriteU32(tableSize);
+		List<Entry> entries;
+		for (u32 i=0; i!=messageCount; ++i)
+		{
+			Entry& e = entries.emplace_back(data, i, messageFetchSize);
+			StackBinaryWriter<32> writer;
+			e.message.Init(m_client, CacheServiceId, messageType, writer);
+			writer.Write7BitEncoded(MakeId(bucketId));
+			writer.WriteU32(tableSize + i*messageFetchSize);
+			if (!e.message.SendAsync(e.reader, [](bool error, void* userData) { ((Event*)userData)->Set(); }, &e.done))
+				return false;
+		}
+
+		u32 timeOutTimeMs = 5*60*1000;
+
+		u64 written = 0;
+		for (Entry& e : entries)
+		{
+			if (!e.done.IsSet(timeOutTimeMs))
+				return false;
+			if (!e.message.ProcessAsyncResults(e.reader))
+				return false;
+			written += e.reader.GetLeft();
+		}
+		table.EndCommit(data, written);
+		return true;
+	}
 
-				reader.Reset();
-				if (!msg.Send(reader))
+	UBA_NOINLINE bool CacheClient::FetchCasTable2(Bucket& bucket, CacheStats& stats, u32 requiredCasTableOffset)
+	{
+		TimerScope ts2(stats.fetchCasTable);
+		u32 requiredPathTableOffset = 0;
+		{
+			SCOPED_WRITE_LOCK(bucket.serverCasKeyTableNetworkLock, lock);
+			while (!HasEnoughCasData(bucket, requiredCasTableOffset, requiredPathTableOffset))
+			{
+				if (!FetchCompactTable(bucket.id, bucket.serverCasKeyTable, requiredCasTableOffset + sizeof(CasKey) + 8, CacheMessageType_FetchCasTable2))
 					return false;
-				reader.ReadU32();
-				messageSent = true;
-				bucket.serverCasKeyTable.ReadMem(reader, false);
+				bucket.serverCasKeyTableSize = bucket.serverCasKeyTable.GetSize();
 			}
 		}
 		{
-			u32 targetSize = ~0u; // For now, read all because we don't know how much we need (it would require parsing all path offsets in caskey table)
-			while (bucket.serverPathTable.GetSize() < targetSize)
+			SCOPED_WRITE_LOCK(bucket.serverPathTableNetworkLock, lock);
+			while (!HasEnoughPathData(bucket, requiredPathTableOffset))
 			{
-				StackBinaryWriter<16> writer;
-				NetworkMessage msg(m_client, CacheServiceId, CacheMessageType_FetchPathTable, writer);
-				writer.Write7BitEncoded(MakeId(bucket.id));
-				writer.WriteU32(bucket.serverPathTable.GetSize());
-
-				reader.Reset();
-				if (!msg.Send(reader))
+				u32 targetSize = requiredPathTableOffset + 200;
+				if (!FetchCompactTable(bucket.id, bucket.serverPathTable, targetSize, CacheMessageType_FetchPathTable2))
 					return false;
-				u32 size = reader.ReadU32();
-				if (targetSize == ~0u)
-					targetSize = size;
-
-				bucket.serverPathTable.ReadMem(reader, false);
+				bucket.serverPathTableSize = bucket.serverPathTable.GetSize();
 			}
 		}
+		return true;
+	}
 
-		bucket.availableCasKeyTableSize = bucket.serverCasKeyTable.GetSize();
+	bool CacheClient::HasEnoughData(Bucket& bucket, u32 requiredCasTableOffset)
+	{
+		u32 pathOffset;
+		if (!HasEnoughCasData(bucket, requiredCasTableOffset, pathOffset))
+			return false;
+		if (!HasEnoughPathData(bucket, pathOffset))
+			return false;
 		return true;
 	}
 
-	bool CacheClient::HasEnoughData(Bucket& bucket, u32 requiredCasTableOffset, u32 tableSize)
+	bool CacheClient::HasEnoughCasData(Bucket& bucket, u32 requiredCasTableOffset, u32& outPathOffset)
 	{
-		u32 neededSize = requiredCasTableOffset + 4;
+		u32 tableSize = bucket.serverCasKeyTableSize;
+		// CasKeyTable is 7bitEncoded(pathoffset) + CasKey... path table offset is minimum 1 byte
+		u32 neededSize = requiredCasTableOffset + 1 + sizeof(CasKey);
 		if (neededSize > tableSize)
 			return false;
 		BinaryReader r(bucket.serverCasKeyTable.GetMemory(), requiredCasTableOffset, tableSize);
-		u8 bytesNeeded = Get7BitEncodedCount(r.Read7BitEncoded());
-		neededSize = requiredCasTableOffset + bytesNeeded + sizeof(CasKey);
+		outPathOffset = u32(r.Read7BitEncoded());
+		neededSize = requiredCasTableOffset + Get7BitEncodedCount(outPathOffset) + sizeof(CasKey);
 		return neededSize <= tableSize;
 	}
 
+	bool CacheClient::HasEnoughPathData(Bucket& bucket, u32 requiredPathTableOffset)
+	{
+		u32 tableSize = bucket.serverPathTableSize;
+		// PathTable is 7bitEncoded(parentoffset) + 7bitEncoded(stroffset). If stroffset is 0, then string is after stroffset
+		if (requiredPathTableOffset + 200 < tableSize) // Early out. no filename without path + two 7bit encoded values are larger than this
+			return true;
+		if (requiredPathTableOffset + 2 > tableSize) // This means that it must be at least 2 bytes
+			return false;
+		BinaryReader r(bucket.serverPathTable.GetMemory(), requiredPathTableOffset, tableSize);
+		u64 value;
+		if (!r.TryRead7BitEncoded(value)) // Parent offset
+			return false;
+		if (!r.TryRead7BitEncoded(value)) // stroffset
+			return false;
+		if (value) // non-0 means it has the string segment before required offset
+			return true;
+		if (!r.TryRead7BitEncoded(value)) // string length
+			return false;
+		if (r.GetLeft() < value) // Actual string in bytes
+			return false;
+		return true;
+	}
+
 	bool CacheClient::ReportUsedEntry(Vector<ProcessLogLine>& outLogLines, bool ownedLogLines, Bucket& bucket, const CasKey& cmdKey, u32 entryId)
 	{
 		StackBinaryWriter<128> writer;
diff --git a/Engine/Source/Programs/UnrealBuildAccelerator/Common/Private/UbaCacheServer.cpp b/Engine/Source/Programs/UnrealBuildAccelerator/Common/Private/UbaCacheServer.cpp
--- a/Engine/Source/Programs/UnrealBuildAccelerator/Common/Private/UbaCacheServer.cpp
+++ b/Engine/Source/Programs/UnrealBuildAccelerator/Common/Private/UbaCacheServer.cpp
@@ -1706,7 +1706,7 @@ namespace uba
 		u32 haveSize = reader.ReadU32();
 		u32 size = bucket.m_pathTable.GetSize();
 		u32 toSend = Min(u32(writer.GetCapacityLeft()), size - haveSize);
-		writer.WriteBytes(bucket.m_casKeyTable.GetMemory() + haveSize, toSend);
+		writer.WriteBytes(bucket.m_pathTable.GetMemory() + haveSize, toSend);
 		return true;
 	}
 
diff --git a/Engine/Source/Programs/UnrealBuildAccelerator/Common/Private/UbaCompactTables.cpp b/Engine/Source/Programs/UnrealBuildAccelerator/Common/Private/UbaCompactTables.cpp
--- a/Engine/Source/Programs/UnrealBuildAccelerator/Common/Private/UbaCompactTables.cpp
+++ b/Engine/Source/Programs/UnrealBuildAccelerator/Common/Private/UbaCompactTables.cpp
@@ -343,6 +343,18 @@ namespace uba
 		}
 	}
 
+	u8* CompactPathTable::BeginCommit(u64 size)
+	{
+		if (!m_pathTableMem.memory)
+			m_pathTableMem.Init(PathTableMaxSize);
+		return (u8*)m_pathTableMem.AllocateNoLock(size, 1, TC("CompactPathTable"));
+	}
+
+	void CompactPathTable::EndCommit(u8* data, u64 written)
+	{
+		m_pathTableMem.writtenSize = data - m_pathTableMem.memory + written;
+	}
+
 	void CompactPathTable::Swap(CompactPathTable& other)
 	{
 		m_offsets.Swap(other.m_offsets);
@@ -574,6 +586,21 @@ namespace uba
 		}
 	}
 
+	u8* CompactCasKeyTable::BeginCommit(u64 size)
+	{
+		if (!m_casKeyTableMem.memory)
+			m_casKeyTableMem.Init(CasKeyTableMaxSize);
+		if (!m_offsets.IsInitialized())
+			m_offsets.Init(CasKeyLookupInitialCount);
+		void* mem = m_casKeyTableMem.AllocateNoLock(size, 1, TC("CompactCasKeyTable"));
+		return (u8*)mem;
+	}
+
+	void CompactCasKeyTable::EndCommit(u8* data, u64 written)
+	{
+		m_casKeyTableMem.writtenSize = data - m_casKeyTableMem.memory + written;
+	}
+
 	void CompactCasKeyTable::Debug(const CompactPathTable& pathTable)
 	{
 		#if 0
diff --git a/Engine/Source/Programs/UnrealBuildAccelerator/Common/Public/UbaCacheClient.h b/Engine/Source/Programs/UnrealBuildAccelerator/Common/Public/UbaCacheClient.h
--- a/Engine/Source/Programs/UnrealBuildAccelerator/Common/Public/UbaCacheClient.h
+++ b/Engine/Source/Programs/UnrealBuildAccelerator/Common/Public/UbaCacheClient.h
@@ -75,7 +75,10 @@ namespace uba
 		bool SendCacheEntry(Bucket& bucket, const RootPaths& rootPaths, const CasKey& cmdKey, const Map<u32, u32>& inputsStringToCasKey, const Map<u32, u32>& outputsStringToCasKey, const u8* logLines, u64 logLinesSize, u64& outBytesSent);
 		bool FetchCasTable(Bucket& bucket, CacheStats& stats, u32 requiredCasTableOffset);
 		bool FetchCasTable2(Bucket& bucket, CacheStats& stats, u32 requiredCasTableOffset);
-		bool HasEnoughData(Bucket& bucket, u32 requiredCasTableOffset, u32 tableSize);
+		template<typename TableType> bool FetchCompactTable(u32 bucketId, TableType& table, u32 requiredTableOffset, u8 messageType);
+		bool HasEnoughData(Bucket& bucket, u32 requiredCasTableOffset);
+		bool HasEnoughCasData(Bucket& bucket, u32 requiredCasTableOffset, u32& outPathOffset);
+		bool HasEnoughPathData(Bucket& bucket, u32 requiredPathTableOffset);
 		bool ReportUsedEntry(Vector<ProcessLogLine>& outLogLines, bool ownedLogLines, Bucket& bucket, const CasKey& cmdKey, u32 entryId);
 		bool PopulateLogLines(Vector<ProcessLogLine>& outLogLines, const u8* mem, u64 memLen);
 
diff --git a/Engine/Source/Programs/UnrealBuildAccelerator/Common/Public/UbaCompactTables.h b/Engine/Source/Programs/UnrealBuildAccelerator/Common/Public/UbaCompactTables.h
--- a/Engine/Source/Programs/UnrealBuildAccelerator/Common/Public/UbaCompactTables.h
+++ b/Engine/Source/Programs/UnrealBuildAccelerator/Common/Public/UbaCompactTables.h
@@ -31,6 +31,8 @@ namespace uba
 		u32 GetVersion();
 
 		void ReadMem(BinaryReader& reader, bool populateLookup);
+		u8* BeginCommit(u64 size);
+		void EndCommit(u8* data, u64 written);
 		void Swap(CompactPathTable& other);
 
 		u64 GetPathCount() { return m_offsets.Size(); }
@@ -69,6 +71,8 @@ namespace uba
 		ReaderWriterLock& GetLock() { return m_lock; }
 
 		void ReadMem(BinaryReader& reader, bool populateLookup);
+		u8* BeginCommit(u64 size);
+		void EndCommit(u8* data, u64 written);
 		void Swap(CompactCasKeyTable& other);
 
 		void Debug(const CompactPathTable& pathTable);
diff --git a/Engine/Source/Programs/UnrealBuildAccelerator/Core/Public/UbaBinaryReaderWriter.h b/Engine/Source/Programs/UnrealBuildAccelerator/Core/Public/UbaBinaryReaderWriter.h
--- a/Engine/Source/Programs/UnrealBuildAccelerator/Core/Public/UbaBinaryReaderWriter.h
+++ b/Engine/Source/Programs/UnrealBuildAccelerator/Core/Public/UbaBinaryReaderWriter.h
@@ -72,6 +72,7 @@ namespace uba
 		inline u64 Read7BitEncoded();
 		template<typename CharType> inline CharType ReadUtf8Char();
 		TString ReadLongString();
+		inline bool TryRead7BitEncoded(u64& outValue);
 		inline u32 PeekU32();
 		inline u64 PeekU64();
 		inline void Skip(u64 size) { m_pos += size; }
@@ -447,6 +448,24 @@ namespace uba
 		return result;
 	}
 
+	bool BinaryReader::TryRead7BitEncoded(u64& outValue)
+	{
+		u64 result = 0;
+		u64 byteIndex = 0;
+		bool hasMoreBytes;
+		do
+		{
+			if (m_pos == m_end)
+				return false;
+			u8 value = *m_pos++;
+			hasMoreBytes = value & 0x80;
+			result |= u64(value & 0x7f) << (byteIndex * 7);
+			++byteIndex;
+		} while (hasMoreBytes);
+		outValue = result;
+		return true;
+	}
+
 	template<typename CharType>
 	CharType BinaryReader::ReadUtf8Char()
 	{
