[UBA]
* Added RemoteTraceEnabled to config
* Fixed so client session trace files matches name of session

diff --git a/Engine/Source/Programs/UnrealBuildAccelerator/Common/Private/UbaSessionServer.cpp b/Engine/Source/Programs/UnrealBuildAccelerator/Common/Private/UbaSessionServer.cpp
--- a/Engine/Source/Programs/UnrealBuildAccelerator/Common/Private/UbaSessionServer.cpp
+++ b/Engine/Source/Programs/UnrealBuildAccelerator/Common/Private/UbaSessionServer.cpp
@@ -100,6 +100,8 @@ namespace uba
 		if (const ConfigTable* table = config.GetTable(TC("Session")))
 		{
 			table->GetValueAsBool(remoteLogEnabled, TC("RemoteLogEnabled"));
+			table->GetValueAsBool(remoteTraceEnabled, TC("RemoteTraceEnabled"));
+			table->GetValueAsBool(nameToHashTableEnabled, TC("NameToHashTableEnabled"));
 		}
 	}
 
@@ -992,7 +994,18 @@ namespace uba
 			if (destination.StartsWith(TC("<uba>")))
 			{
 				StringBuffer<> ubaPath;
-				ubaPath.Append(m_sessionLogDir).AppendValue(connectionInfo.GetId()).Append(TC(".uba"));
+				ubaPath.Append(m_sessionLogDir);
+
+				ClientSession* session = nullptr;
+				for (auto& s : m_clientSessions)
+					if (s->clientId == connectionInfo.GetId())
+						session = s;
+				if (session)
+					ubaPath.Append(session->name);
+				else
+					ubaPath.Append(TC("Connection")).AppendValue(connectionInfo.GetId());
+
+				ubaPath.Append(TC(".uba"));
 				m_storage.CopyOrLink(casKey, ubaPath.data, attributes);
 				m_storage.DropCasFile(casKey, false, ubaPath.data);
 				writer.WriteBool(true);
