[UBA]
* Fixed so visualizer shows a message on the canvas if trace stream is newer version than visualizer.
* Fixed weird issues with embedded visualizer in UnrealVs extension couldn't receive mouse input on latest version of windows 11

diff --git a/Engine/Source/Programs/UnrealBuildAccelerator/Visualizer/Private/UbaTraceReader.cpp b/Engine/Source/Programs/UnrealBuildAccelerator/Visualizer/Private/UbaTraceReader.cpp
--- a/Engine/Source/Programs/UnrealBuildAccelerator/Visualizer/Private/UbaTraceReader.cpp
+++ b/Engine/Source/Programs/UnrealBuildAccelerator/Visualizer/Private/UbaTraceReader.cpp
@@ -32,6 +32,7 @@ namespace uba
 
 	void TraceView::Clear()
 	{
+		version = 0;
 		sessions.clear();
 		workTracks.clear();
 		strings.clear();
@@ -284,6 +285,9 @@ namespace uba
 	{
 		if (m_memoryEnd == m_memoryPos)
 			return true;
+		
+		if (out.version && (out.version < TraceReadCompatibilityVersion || out.version > TraceVersion))
+			return true;
 
 		u64 toRead = m_memoryEnd - m_memoryPos;
 		BinaryReader traceReader(m_memoryPos);
@@ -294,17 +298,19 @@ namespace uba
 				return true;
 			u32 traceSize = traceReader.ReadU32(); (void)traceSize;
 			u32 version = traceReader.ReadU32();
-			if (version < TraceReadCompatibilityVersion || version > TraceVersion)
-				return m_logger.Error(L"Incompatible trace version (%u). Current executable supports version %u to %u.", version, TraceReadCompatibilityVersion, TraceVersion);
+			out.version = version;
 
 			bool replay = maxTime != ~u64(0);
 
-			out.version = version;
 			u32 hostProcessId = traceReader.ReadU32();
 			m_hostProcess = 0;
 			if (trackHost && !replay)
 				m_hostProcess = OpenProcess(SYNCHRONIZE, FALSE, hostProcessId);
 			u64 traceSystemStartTimeUs = 0;
+
+			if (version < TraceReadCompatibilityVersion || version > TraceVersion)
+				return true;//m_logger.Error(L"Incompatible trace version (%u). Current executable supports version %u to %u.", version, TraceReadCompatibilityVersion, TraceVersion);
+
 			if (version >= 18)
 				traceSystemStartTimeUs = traceReader.Read7BitEncoded();
 			if (version >= 18)
diff --git a/Engine/Source/Programs/UnrealBuildAccelerator/Visualizer/Private/UbaVisualizer.cpp b/Engine/Source/Programs/UnrealBuildAccelerator/Visualizer/Private/UbaVisualizer.cpp
--- a/Engine/Source/Programs/UnrealBuildAccelerator/Visualizer/Private/UbaVisualizer.cpp
+++ b/Engine/Source/Programs/UnrealBuildAccelerator/Visualizer/Private/UbaVisualizer.cpp
@@ -572,94 +572,109 @@ namespace uba
 		wndClassEx.hbrBackground = NULL;
 		wndClassEx.lpszClassName = TEXT("UbaVisualizer");
 		ATOM wndClassAtom = RegisterClassEx(&wndClassEx);
+		const TCHAR* windowClassName = MAKEINTATOM(wndClassAtom);
+
+		auto unreg = MakeGuard([&]() { UnregisterClass(windowClassName, hInstance); });
 
 		UpdateDefaultFont();
 
 		UpdateProcessFont();
 
 		const TCHAR* fontName = TEXT("Consolas");
-		m_popupFont.handle = (HFONT)CreateFontW(-12, 0, 0, 0, FW_NORMAL, false, false, false, ANSI_CHARSET, OUT_DEFAULT_PRECIS, CLIP_DEFAULT_PRECIS, CLEARTYPE_QUALITY, FIXED_PITCH | FF_MODERN, fontName);
+		m_popupFont.handle = CreateFontW(-12, 0, 0, 0, FW_NORMAL, false, false, false, ANSI_CHARSET, OUT_DEFAULT_PRECIS, CLIP_DEFAULT_PRECIS, CLEARTYPE_QUALITY, FIXED_PITCH | FF_MODERN, fontName);
 		m_popupFont.height = 14;
 		//m_popupFont = (HFONT)GetStockObject(SYSTEM_FIXED_FONT);
 
 		DWORD windowStyle = WS_CAPTION | WS_SYSMENU | WS_THICKFRAME | WS_MAXIMIZEBOX | WS_MINIMIZEBOX | WS_CLIPCHILDREN | WS_VSCROLL | WS_HSCROLL;
-		const TCHAR* windowClassName = MAKEINTATOM(wndClassAtom);
 
 		DWORD exStyle = 0;
 		if (m_config.parent)
-			windowStyle = WS_POPUP | WS_VSCROLL | WS_HSCROLL;// | WS_VISIBLE;
+			windowStyle = WS_POPUP | WS_VSCROLL | WS_HSCROLL;
 
 		StringBuffer<> title;
 		GetTitlePrefix(title).Append(L"Initializing...");
 
 		HWND hwnd = CreateWindowEx(exStyle, windowClassName, title.data, windowStyle, winPosX, winPosY, winWidth, winHeight, NULL, NULL, hInstance, this);
 		SetWindowLongPtr(hwnd, GWLP_USERDATA, (LONG_PTR)this);
+		m_hwnd = hwnd;
+
+		auto destroyWin = [&]()
+			{
+				if (m_hwnd)
+				{
+					if (m_config.AutoSaveSettings)
+						SaveSettings();
+					DestroyWindow(m_hwnd);
+					m_hwnd = 0;
+				}
+			};
 
 		BOOL cloak = TRUE;
 		DwmSetWindowAttribute(hwnd, DWMWA_CLOAK, &cloak, sizeof(cloak));
+		auto exitCloak = MakeGuard([&]() { cloak = FALSE; DwmSetWindowAttribute(hwnd, DWMWA_CLOAK, &cloak, sizeof(cloak)); });
 
-		//if (m_config.DarkMode)
-		{
-			SetWindowTheme(hwnd, L"DarkMode_Explorer", NULL);
-			SendMessageW(hwnd, WM_THEMECHANGED, 0, 0);
-			BOOL useDarkMode = true;
-			u32 attribute = 20; // DWMWA_USE_IMMERSIVE_DARK_MODE
-			DwmSetWindowAttribute(hwnd, attribute, &useDarkMode, sizeof(useDarkMode));
-		}
+		if (m_config.DarkMode)
+			UpdateTheme();
 
 		HitTestResult res;
 		HitTest(res, { -1, -1 });
 
 		if (m_config.parent)
 		{
+			exitCloak.Execute();
+
+			// If not child it will not propagate keyboard presses etc to parent
+			LONG style = GetWindowLong(hwnd, GWL_STYLE);
+			style = WS_CHILD;
+			style &= ~WS_POPUP;
+			SetWindowLong(hwnd, GWL_STYLE, WS_CHILD|WS_VSCROLL|WS_HSCROLL);
+
 			m_parentHwnd = (HWND)(uintptr_t)m_config.parent;
 			if (!SetParent(hwnd, m_parentHwnd))
  				m_logger.Error(L"SetParent failed using parentHwnd 0x%llx", m_parentHwnd);
 
+			UpdateWindow(m_hwnd);
+			UpdateScrollbars(true);
 			PostMessage(m_parentHwnd, 0x0444, 0, (LPARAM)hwnd);
 		}
-
-		m_hwnd = hwnd;
-
-		if (!m_parentHwnd)
-			ShowWindow(m_hwnd, SW_SHOW);
-		UpdateWindow(m_hwnd);
-		UpdateScrollbars(true);
-
-		cloak = FALSE;
-		DwmSetWindowAttribute(m_hwnd, DWMWA_CLOAK, &cloak, sizeof(cloak));
+		else
+		{
+			ShowWindow(hwnd, SW_SHOW);
+			UpdateWindow(m_hwnd);
+			UpdateScrollbars(true);
+			exitCloak.Execute();
+		}
 
 		m_startTime = GetTime();
 
 		while (m_looping)
 		{
+			DWORD timeoutMs = 2000;
+			DWORD result = MsgWaitForMultipleObjects(0, NULL, FALSE, timeoutMs, QS_ALLINPUT);
+			if (result == WAIT_TIMEOUT)
+				continue;
+			if (result != WAIT_OBJECT_0)
+				break;
 			MSG msg;
-			while (GetMessage(&msg, NULL, 0, 0))
+			while (PeekMessage(&msg, NULL, 0, 0, PM_REMOVE))
 			{
-				if (m_hwnd && !IsDialogMessage(m_hwnd, &msg))
-				{
-					TranslateMessage(&msg);
-					DispatchMessage(&msg);
-				}
+				TranslateMessage(&msg);
+				DispatchMessage(&msg);
+
 				// It may happen that we receive the WM_DESTOY message from within the DistpachMessage above and handle it directly in WndProc.
 				// So before trying to call GetMessage again, we just need to validate that m_looping is still true otherwise we could end
 				// up waiting forever for this thread to exit.
 				if (!m_looping || msg.message == WM_QUIT || msg.message == WM_DESTROY || msg.message == WM_CLOSE)
 				{
-					if (m_hwnd)
-					{
-						if (m_config.AutoSaveSettings)
-							SaveSettings();
-						DestroyWindow(m_hwnd);
-					}
-					UnregisterClass(windowClassName, hInstance);
-					m_hwnd = 0;
+					destroyWin();
 					m_looping = false;
 					m_listenTimeout.Set();
 					break;
 				}
 			}
 		}
+
+		destroyWin();
 	}
 
 	void Visualizer::Pause(bool pause)
@@ -803,9 +818,9 @@ namespace uba
 		//SystemParametersInfo(SPI_GETNONCLIENTMETRICS, sizeof(nonClientMetrics), &nonClientMetrics, 0);
 		//m_font = (HFONT)CreateFontIndirect(&nonClientMetrics.lfMessageFont);
 
-		font.handle = (HFONT)CreateFontW(4 - fh, 0, 0, 0, FW_NORMAL, 0, 0, 0, ANSI_CHARSET, OUT_DEFAULT_PRECIS, CLIP_DEFAULT_PRECIS, CLEARTYPE_QUALITY, DEFAULT_PITCH, m_config.fontName.c_str());
+		font.handle = CreateFontW(4 - fh, 0, 0, 0, FW_NORMAL, 0, 0, 0, ANSI_CHARSET, OUT_DEFAULT_PRECIS, CLIP_DEFAULT_PRECIS, CLEARTYPE_QUALITY, DEFAULT_PITCH, m_config.fontName.c_str());
 		if (createUnderline)
-			font.handleUnderlined = (HFONT)CreateFontW(4 - fh, 0, 0, 0, FW_NORMAL, 0, 1, 0, ANSI_CHARSET, OUT_DEFAULT_PRECIS, CLIP_DEFAULT_PRECIS, CLEARTYPE_QUALITY, DEFAULT_PITCH, m_config.fontName.c_str());
+			font.handleUnderlined = CreateFontW(4 - fh, 0, 0, 0, FW_NORMAL, 0, 1, 0, ANSI_CHARSET, OUT_DEFAULT_PRECIS, CLIP_DEFAULT_PRECIS, CLEARTYPE_QUALITY, DEFAULT_PITCH, m_config.fontName.c_str());
 	}
 
 	void Visualizer::UpdateDefaultFont()
@@ -986,6 +1001,27 @@ namespace uba
 			drawIndentedText(str, LogEntryType_Info, 6, true);
 		}
 
+		if (m_traceView.version && (m_traceView.version < TraceReadCompatibilityVersion || m_traceView.version > TraceVersion))
+		{
+			if (!m_traceView.finished)
+			{
+				SetTextColor(hdc, m_textWarningColor);
+				StringBuffer<> str;
+				str.Appendf(TC("Unsupported trace version %u (Versions supported are %u to %u)"), m_traceView.version, TraceReadCompatibilityVersion, TraceVersion);
+				HFONT handle = CreateFontW(-20, 0, 0, 0, FW_NORMAL, 0, 0, 0, ANSI_CHARSET, OUT_DEFAULT_PRECIS, CLIP_DEFAULT_PRECIS, CLEARTYPE_QUALITY, DEFAULT_PITCH, m_config.fontName.c_str());
+				SelectObject(hdc, handle);
+				DrawTextW(hdc, str.data, str.count, (RECT*)&clientRect, DT_CENTER|DT_SINGLELINE|DT_VCENTER);
+				DeleteObject(handle);
+			}
+			else
+			{
+				m_traceView.Clear();
+				StringBuffer<> title;
+				SetWindowTextW(m_hwnd, GetTitlePrefix(title).Appendf(L" (Listening for new sessions on channel '%s')", m_listenChannel.data).data);
+			}
+			return;
+		}
+
 		if (m_config.showStatus && !m_traceView.statusMap.empty())
 		{
 			u32 lastRow = ~0u;
@@ -2970,6 +3006,15 @@ namespace uba
 			SelectObject(m_activeHdc, font.handle);
 	}
 
+	void Visualizer::UpdateTheme()
+	{
+		SetWindowTheme(m_hwnd, m_config.DarkMode ? L"DarkMode_Explorer" : L"Explorer", NULL);
+		SendMessageW(m_hwnd, WM_THEMECHANGED, 0, 0);
+		BOOL useDarkMode = m_config.DarkMode;
+		u32 attribute = 20; // DWMWA_USE_IMMERSIVE_DARK_MODE
+		DwmSetWindowAttribute(m_hwnd, attribute, &useDarkMode, sizeof(useDarkMode));
+	}
+
 	void Visualizer::PostNewTrace(u32 replay, bool paused)
 	{
 		KillTimer(m_hwnd, 0);
@@ -3237,6 +3282,10 @@ namespace uba
 			}
 			break;
 		}
+		case WM_NCHITTEST:
+			if (m_parentHwnd)
+				return HTCLIENT;
+			break;
 		case WM_MOUSEMOVE:
 		{
 			//m_logger.Info(TC("Section: %u"), m_activeSection);
@@ -3315,23 +3364,6 @@ namespace uba
 			break;
 		}
 
-		case WM_MOUSEACTIVATE:
-		{
-			if (LOWORD(lParam) != HTCLIENT)
-				break;
-
-			if (m_parentHwnd)
-			{
-				//PostMessage(wParam, WM_MOUSEACTIVATE, wParam, lParam);
-				//SetForegroundWindow(m_parentHwnd);
-				//SetFocus(m_parentHwnd);
-				//SetWindowPos((HWND)wParam, HWND_TOP, 0, 0, 0, 0, SWP_ASYNCWINDOWPOS | SWP_NOMOVE | SWP_NOSIZE);
-				PostMessage(m_parentHwnd, 0x0445, 0, 0);
-				//return MA_NOACTIVATEANDEAT;
-			}
-			return MA_ACTIVATEANDEAT;
-		}
-
 		case WM_LBUTTONDOWN:
 		{
 			if (m_buttonSelected != ~0u)
@@ -3547,11 +3579,7 @@ namespace uba
 				m_config.DarkMode = !m_config.DarkMode;
 				DirtyBitmaps(false);
 				InitBrushes();
-				SetWindowTheme(m_hwnd, m_config.DarkMode ? L"DarkMode_Explorer" : L"Explorer", NULL);
-				SendMessageW(m_hwnd, WM_THEMECHANGED, 0, 0);
-				BOOL useDarkMode = m_config.DarkMode;
-				u32 attribute = 20; // DWMWA_USE_IMMERSIVE_DARK_MODE
-				DwmSetWindowAttribute(m_hwnd, attribute, &useDarkMode, sizeof(useDarkMode));
+				UpdateTheme();
 				Redraw(true);
 				break;
 			}
diff --git a/Engine/Source/Programs/UnrealBuildAccelerator/Visualizer/Private/UbaVisualizer.h b/Engine/Source/Programs/UnrealBuildAccelerator/Visualizer/Private/UbaVisualizer.h
--- a/Engine/Source/Programs/UnrealBuildAccelerator/Visualizer/Private/UbaVisualizer.h
+++ b/Engine/Source/Programs/UnrealBuildAccelerator/Visualizer/Private/UbaVisualizer.h
@@ -151,6 +151,7 @@ namespace uba
 		void ChangeFontSize(int offset);
 		void Redraw(bool now);
 		void SetActiveFont(const Font& font);
+		void UpdateTheme();
 
 		StringBuffer<256> m_namedTrace;
 		StringBuffer<256> m_fileName;
