[UBA]
* Added WorkHint to trace stream that can be used to attach more info to detailed traces

#skipci

diff --git a/Engine/Source/Programs/UnrealBuildAccelerator/Common/Private/UbaSessionServer.cpp b/Engine/Source/Programs/UnrealBuildAccelerator/Common/Private/UbaSessionServer.cpp
--- a/Engine/Source/Programs/UnrealBuildAccelerator/Common/Private/UbaSessionServer.cpp
+++ b/Engine/Source/Programs/UnrealBuildAccelerator/Common/Private/UbaSessionServer.cpp
@@ -902,6 +902,8 @@ namespace uba
 		reader.ReadString(fileName);
 		StringKey fileNameKey = reader.ReadStringKey();
 
+		m_server.TrackWorkHint(workContext.id, StringView(fileName).GetFileName());
+
 		CasKey casKey;
 		if (!StoreCasFile(casKey, fileNameKey, fileName.data))
 			return false;
diff --git a/Engine/Source/Programs/UnrealBuildAccelerator/Common/Private/UbaStorageProxy.cpp b/Engine/Source/Programs/UnrealBuildAccelerator/Common/Private/UbaStorageProxy.cpp
--- a/Engine/Source/Programs/UnrealBuildAccelerator/Common/Private/UbaStorageProxy.cpp
+++ b/Engine/Source/Programs/UnrealBuildAccelerator/Common/Private/UbaStorageProxy.cpp
@@ -61,7 +61,7 @@ namespace uba
 		m_server.RegisterService(StorageServiceId,
 			[this](const ConnectionInfo& connectionInfo, const WorkContext& workContext, MessageInfo& messageInfo, BinaryReader& reader, BinaryWriter& writer)
 			{
-				return HandleMessage(connectionInfo, messageInfo, reader, writer);
+				return HandleMessage(connectionInfo, workContext, messageInfo, reader, writer);
 			},
 			[](u8 messageType)
 			{
@@ -108,7 +108,7 @@ namespace uba
 		m_availableIds.push_back(id);
 	}
 
-	bool StorageProxy::HandleMessage(const ConnectionInfo& connectionInfo, MessageInfo& messageInfo, BinaryReader& reader, BinaryWriter& writer)
+	bool StorageProxy::HandleMessage(const ConnectionInfo& connectionInfo, const WorkContext& workContext, MessageInfo& messageInfo, BinaryReader& reader, BinaryWriter& writer)
 	{
 		LoggerWithWriter logger(g_consoleLogWriter, TC(""));
 
@@ -145,6 +145,8 @@ namespace uba
 				StringBuffer<> hint;
 				reader.ReadString(hint);
 
+				m_server.TrackWorkHint(workContext.id, StringView(hint).GetFileName());
+
 				SCOPED_WRITE_LOCK(m_filesLock, filesLock);
 				auto insres = m_files.try_emplace(casKey);
 				FileEntry& file = insres.first->second;
diff --git a/Engine/Source/Programs/UnrealBuildAccelerator/Common/Private/UbaStorageServer.cpp b/Engine/Source/Programs/UnrealBuildAccelerator/Common/Private/UbaStorageServer.cpp
--- a/Engine/Source/Programs/UnrealBuildAccelerator/Common/Private/UbaStorageServer.cpp
+++ b/Engine/Source/Programs/UnrealBuildAccelerator/Common/Private/UbaStorageServer.cpp
@@ -399,6 +399,8 @@ namespace uba
 				StringBuffer<> hint;
 				reader.ReadString(hint);
 
+				m_server.TrackWorkHint(workContext.id, StringView(hint).GetFileName());
+
 				casKey = AsCompressed(casKey, m_storeCompressed);
 
 				bool detailedTrace = m_traceFetch;
diff --git a/Engine/Source/Programs/UnrealBuildAccelerator/Common/Private/UbaTrace.cpp b/Engine/Source/Programs/UnrealBuildAccelerator/Common/Private/UbaTrace.cpp
--- a/Engine/Source/Programs/UnrealBuildAccelerator/Common/Private/UbaTrace.cpp
+++ b/Engine/Source/Programs/UnrealBuildAccelerator/Common/Private/UbaTrace.cpp
@@ -133,13 +133,18 @@ namespace uba
 	u32 Trace::TrackWorkStart(const StringView& desc)
 	{
 		u32 workId = m_workCounter++;
-		BeginWork(workId, desc);
+		WorkBegin(workId, desc);
 		return workId;
 	}
 
+	void Trace::TrackWorkHint(u32 id, const StringView& hint)
+	{
+		WorkHint(id, hint);
+	}
+
 	void Trace::TrackWorkEnd(u32 id)
 	{
-		EndWork(id);
+		WorkEnd(id);
 	}
 
 	void Trace::FreeMemory()
@@ -372,17 +377,25 @@ namespace uba
 		writer.WriteCasKey(key);
 	}
 
-	void Trace::BeginWork(u32 workIndex, const StringView& desc)
+	void Trace::WorkBegin(u32 workIndex, const StringView& desc)
 	{
 		u32 stringIndex = AddString(desc);
-		BEGIN_TRACE_ENTRY(BeginWork);
+		BEGIN_TRACE_ENTRY(WorkBegin);
+		writer.Write7BitEncoded(workIndex);
+		writer.Write7BitEncoded(stringIndex);
+	}
+
+	void Trace::WorkHint(u32 workIndex, const StringView& hint)
+	{
+		u32 stringIndex = AddString(hint);
+		BEGIN_TRACE_ENTRY(WorkHint);
 		writer.Write7BitEncoded(workIndex);
 		writer.Write7BitEncoded(stringIndex);
 	}
 
-	void Trace::EndWork(u32 workIndex)
+	void Trace::WorkEnd(u32 workIndex)
 	{
-		BEGIN_TRACE_ENTRY(EndWork);
+		BEGIN_TRACE_ENTRY(WorkEnd);
 		writer.Write7BitEncoded(workIndex);
 	}
 
diff --git a/Engine/Source/Programs/UnrealBuildAccelerator/Common/Private/UbaWorkManager.cpp b/Engine/Source/Programs/UnrealBuildAccelerator/Common/Private/UbaWorkManager.cpp
--- a/Engine/Source/Programs/UnrealBuildAccelerator/Common/Private/UbaWorkManager.cpp
+++ b/Engine/Source/Programs/UnrealBuildAccelerator/Common/Private/UbaWorkManager.cpp
@@ -14,6 +14,12 @@ namespace uba
 		return 0;
 	}
 
+	void WorkManager::TrackWorkHint(u32 id, const StringView& hint)
+	{
+		if (auto t = m_workTracker.load())
+			return t->TrackWorkHint(id, hint);
+	}
+
 	void WorkManager::TrackWorkEnd(u32 id)
 	{
 		if (auto t = m_workTracker.load())
diff --git a/Engine/Source/Programs/UnrealBuildAccelerator/Common/Public/UbaStorageProxy.h b/Engine/Source/Programs/UnrealBuildAccelerator/Common/Public/UbaStorageProxy.h
--- a/Engine/Source/Programs/UnrealBuildAccelerator/Common/Public/UbaStorageProxy.h
+++ b/Engine/Source/Programs/UnrealBuildAccelerator/Common/Public/UbaStorageProxy.h
@@ -15,6 +15,7 @@ namespace uba
 	struct BinaryWriter;
 	struct ConnectionInfo;
 	struct MessageInfo;
+	struct WorkContext;
 
 	class StorageProxy
 	{
@@ -28,7 +29,7 @@ namespace uba
 	protected:
 		u16 PopId();
 		void PushId(u16 id);
-		bool HandleMessage(const ConnectionInfo& connectionInfo, MessageInfo& messageInfo, BinaryReader& reader, BinaryWriter& writer);
+		bool HandleMessage(const ConnectionInfo& connectionInfo, const WorkContext& workContext, MessageInfo& messageInfo, BinaryReader& reader, BinaryWriter& writer);
 
 		struct MessageInFlight;
 		void HandleReceivedData(MessageInFlight& mif);
diff --git a/Engine/Source/Programs/UnrealBuildAccelerator/Common/Public/UbaTrace.h b/Engine/Source/Programs/UnrealBuildAccelerator/Common/Public/UbaTrace.h
--- a/Engine/Source/Programs/UnrealBuildAccelerator/Common/Public/UbaTrace.h
+++ b/Engine/Source/Programs/UnrealBuildAccelerator/Common/Public/UbaTrace.h
@@ -39,8 +39,8 @@ namespace uba
 		UBA_TRACE_TYPE(FileStoreBegin) \
 		UBA_TRACE_TYPE(FileStoreEnd) \
 		UBA_TRACE_TYPE(Summary) \
-		UBA_TRACE_TYPE(BeginWork) \
-		UBA_TRACE_TYPE(EndWork) \
+		UBA_TRACE_TYPE(WorkBegin) \
+		UBA_TRACE_TYPE(WorkEnd) \
 		UBA_TRACE_TYPE(String) \
 		UBA_TRACE_TYPE(SessionSummary) \
 		UBA_TRACE_TYPE(ProcessEnvironmentUpdated) \
@@ -59,6 +59,7 @@ namespace uba
 		UBA_TRACE_TYPE(RemoteExecutionDisabled) \
 		UBA_TRACE_TYPE(FileFetchSize) \
 		UBA_TRACE_TYPE(ProcessBreadcrumbs) \
+		UBA_TRACE_TYPE(WorkHint) \
 
 	enum TraceType : u8
 	{
@@ -100,8 +101,9 @@ namespace uba
 		void FileFetchEnd(u32 clientId, const CasKey& key);
 		void FileStoreBegin(u32 clientId, const CasKey& key, u64 size, const StringView& hint, bool detailed);
 		void FileStoreEnd(u32 clientId, const CasKey& key);
-		void BeginWork(u32 workIndex, const StringView& desc);
-		void EndWork(u32 workIndex);
+		void WorkBegin(u32 workIndex, const StringView& desc);
+		void WorkHint(u32 workIndex, const StringView& hint);
+		void WorkEnd(u32 workIndex);
 		void ProgressUpdate(u32 processesTotal, u32 processesDone, u32 errorCount);
 		void StatusUpdate(u32 statusRow, u32 statusColumn, const tchar* statusText, LogEntryType statusType, const tchar* statusLink);
 		void RemoteExecutionDisabled();
@@ -114,6 +116,7 @@ namespace uba
 		bool StopWrite(const tchar* writeFileName);
 
 		virtual u32 TrackWorkStart(const StringView& desc) final override;
+		virtual void TrackWorkHint(u32 id, const StringView& hint) final override;
 		virtual void TrackWorkEnd(u32 id) final override;
 
 	private:
diff --git a/Engine/Source/Programs/UnrealBuildAccelerator/Common/Public/UbaWorkManager.h b/Engine/Source/Programs/UnrealBuildAccelerator/Common/Public/UbaWorkManager.h
--- a/Engine/Source/Programs/UnrealBuildAccelerator/Common/Public/UbaWorkManager.h
+++ b/Engine/Source/Programs/UnrealBuildAccelerator/Common/Public/UbaWorkManager.h
@@ -13,6 +13,7 @@ namespace uba
 	{
 	public:
 		virtual u32 TrackWorkStart(const StringView& desc) = 0;
+		virtual void TrackWorkHint(u32 id, const StringView& hint) = 0;
 		virtual void TrackWorkEnd(u32 id) = 0;
 	};
 
@@ -30,6 +31,7 @@ namespace uba
 		virtual void DoWork(u32 count = 1) = 0;
 
 		u32 TrackWorkStart(const StringView& desc);
+		void TrackWorkHint(u32 id, const StringView& hint);
 		void TrackWorkEnd(u32 id);
 
 		void SetWorkTracker(WorkTracker* workTracker) { m_workTracker = workTracker; }
diff --git a/Engine/Source/Programs/UnrealBuildAccelerator/Visualizer/Private/UbaTraceReader.cpp b/Engine/Source/Programs/UnrealBuildAccelerator/Visualizer/Private/UbaTraceReader.cpp
--- a/Engine/Source/Programs/UnrealBuildAccelerator/Visualizer/Private/UbaTraceReader.cpp
+++ b/Engine/Source/Programs/UnrealBuildAccelerator/Visualizer/Private/UbaTraceReader.cpp
@@ -959,7 +959,7 @@ namespace uba
 			StopAllActive(out, time);
 			return false;
 		}
-		case TraceType_BeginWork:
+		case TraceType_WorkBegin:
 		{
 			u32 workIndex;
 			if (out.version < 14)
@@ -1004,7 +1004,7 @@ namespace uba
 			record.stop = ~u64(0);
 			break;
 		}
-		case TraceType_EndWork:
+		case TraceType_WorkEnd:
 		{
 			u32 workIndex;
 			if (out.version < 14)
@@ -1138,6 +1138,18 @@ namespace uba
 			write.end = time;
 			break;
 		}
+		case TraceType_WorkHint:
+		{
+			u32 workIndex = u32(reader.Read7BitEncoded());
+			u64 stringIndex = reader.Read7BitEncoded();
+			auto findIt = m_activeWorkRecords.find(workIndex);
+			if (findIt == m_activeWorkRecords.end())
+				return false;
+			WorkRecordLocation active = findIt->second;
+			TraceView::WorkRecord& record = out.workTracks[active.track].records[active.index];
+			record.entries.emplace_back(time, out.strings[stringIndex].data()); 
+			break;
+		}
 		default:
 		{
 			return m_logger.Error(L"Unknown trace type found in stream. UbaVisualizer too old?");
diff --git a/Engine/Source/Programs/UnrealBuildAccelerator/Visualizer/Private/UbaTraceReader.h b/Engine/Source/Programs/UnrealBuildAccelerator/Visualizer/Private/UbaTraceReader.h
--- a/Engine/Source/Programs/UnrealBuildAccelerator/Visualizer/Private/UbaTraceReader.h
+++ b/Engine/Source/Programs/UnrealBuildAccelerator/Visualizer/Private/UbaTraceReader.h
@@ -50,11 +50,18 @@ namespace uba
 			u8 connectionCount;
 		};
 
+		struct WorkRecordLogEntry
+		{
+			u64 time = 0;
+			const tchar* text;
+		};
+
 		struct WorkRecord
 		{
 			const tchar* description = nullptr;
 			u64 start = 0;
 			u64 stop = 0;
+			Vector<WorkRecordLogEntry> entries;
 			HBITMAP bitmap = 0;
 			u32 bitmapOffset = 0;
 		};
diff --git a/Engine/Source/Programs/UnrealBuildAccelerator/Visualizer/Private/UbaVisualizer.cpp b/Engine/Source/Programs/UnrealBuildAccelerator/Visualizer/Private/UbaVisualizer.cpp
--- a/Engine/Source/Programs/UnrealBuildAccelerator/Visualizer/Private/UbaVisualizer.cpp
+++ b/Engine/Source/Programs/UnrealBuildAccelerator/Visualizer/Private/UbaVisualizer.cpp
@@ -1972,6 +1972,10 @@ namespace uba
 			logger.Info(L"  %ls", selectedWork.description);
 			logger.Info(L"  Start:     %ls", TimeToText(selectedWork.start, true).str);
 			logger.Info(L"  Duration:  %ls", TimeToText(duration, true).str);
+			if (!selectedWork.entries.empty())
+				logger.AddSpace();
+			for (auto& e : selectedWork.entries)
+				logger.Info(L"   %ls (%ls)", e.text, TimeToText(e.time).str);
 			logger.AddSpace();
 			logger.DrawAtCursor();
 		}
