[UBA]
* Added so agent using local storage for proxy is a command line option
* Moved drop-logic for received cas file to happen inside StorageServer since we know if the file is written to disk or not. This also makes the option of writing received cas files to disk useful
* Made CheckIfExists message optional and disabled it.. since we don't write cas files to disk this will always return false anyway.44
* Added more unit test logic
* Fixed spelling error

diff --git a/Engine/Source/Programs/UnrealBuildAccelerator/Agent/Private/UbaAgent.cpp b/Engine/Source/Programs/UnrealBuildAccelerator/Agent/Private/UbaAgent.cpp
--- a/Engine/Source/Programs/UnrealBuildAccelerator/Agent/Private/UbaAgent.cpp
+++ b/Engine/Source/Programs/UnrealBuildAccelerator/Agent/Private/UbaAgent.cpp
@@ -355,6 +355,7 @@ namespace uba
 		bool useQuic = false;
 		bool poll = true;
 		bool allowProxy = true;
+		bool proxyUseLocalStorage = false;
 		bool couldBeAws = true;
 		bool useStorage = true;
 		bool resetStore = false;
@@ -1043,58 +1044,48 @@ namespace uba
 
 			Event wakeupSessionWait(false);
 			Atomic<u32> targetConnectionCount = 1;
+			StorageClient* storageClient = nullptr;
 
-			struct Proxy
-			{
-				LogWriter& logWriter;
-				NetworkBackend& networkBackend;
-				NetworkBackend& networkBackendMem;
-				NetworkClient* client;
-				Event& wakeupSessionWait;
-				u32& maxConnectionCount;
-				Atomic<u32>& targetConnectionCount;
-				Atomic<NetworkServer*> server;
-				StorageProxy* storage = nullptr;
-				StorageClient* storageClient = nullptr;
-				TString serverPrefix;
-			} proxy { g_consoleLogWriter, *networkBackend, networkBackendMem, client, wakeupSessionWait, maxConnectionCount, targetConnectionCount };
-			auto psg = MakeGuard([&]() { delete proxy.server; });
-			auto pg = MakeGuard([&]() { delete proxy.storage; });
-
-			static auto startProxy = [](void* userData, u16 proxyPort, const Guid& storageServerUid)
-				{
-					auto& proxy = *(Proxy*)userData;
 
-					NetworkServerCreateInfo nsci(proxy.logWriter);
+			StorageProxy* storageProxy = nullptr;
+			Atomic<NetworkServer*> proxyNetworkServer;
+			auto psg = MakeGuard([&]() { delete proxyNetworkServer; });
+			auto pg = MakeGuard([&]() { delete storageProxy; });
+			TString proxyNetworkServerPrefix;
+
+			auto startProxy = [&](u16 proxyPort, const Guid& storageServerUid)
+				{
+					NetworkServerCreateInfo nsci(g_consoleLogWriter);
 					nsci.workerCount = 192;
 					nsci.receiveTimeoutSeconds = 60;
 
-					StringBuffer<256> prefix;
-					prefix.Append(TC("UbaProxyServer (")).Append(GuidToString(proxy.client->GetUid()).str).Append(')');
-					proxy.serverPrefix = prefix.data;
+					proxyNetworkServerPrefix = StringBuffer<256>().Append(TC("UbaProxyServer (")).Append(GuidToString(client->GetUid()).str).Append(')').data;
 					bool ctorSuccess = true;
-					auto proxyServer = new NetworkServer(ctorSuccess, nsci, proxy.serverPrefix.c_str());
+					auto proxyServer = new NetworkServer(ctorSuccess, nsci, proxyNetworkServerPrefix.c_str());
 					if (!ctorSuccess)
 					{
 						delete proxyServer;
 						return false;
 					}
 
-					proxy.storage = new StorageProxy(*proxyServer, *proxy.client, storageServerUid, TC("Wooohoo"), proxy.storageClient);
+					StorageProxyCreateInfo proxyInfo { *proxyServer, *client, storageServerUid, TC("Wooohoo"), storageClient };
+					proxyInfo.useLocalStorage = proxyUseLocalStorage;
 
-					proxyServer->RegisterOnClientConnected(0, [p = &proxy](const Guid& clientUid, u32 clientId) { p->wakeupSessionWait.Set(); });
-					proxyServer->SetWorkTracker(proxy.client->GetWorkTracker());
-					proxyServer->StartListen(proxy.networkBackendMem, proxyPort);
-					proxyServer->StartListen(proxy.networkBackend, proxyPort);
-					proxy.targetConnectionCount = proxy.maxConnectionCount;
+					storageProxy = new StorageProxy(proxyInfo);
 
-					proxy.server = proxyServer;
-					proxy.wakeupSessionWait.Set();
+					proxyServer->RegisterOnClientConnected(0, [&wakeupSessionWait](const Guid& clientUid, u32 clientId) { wakeupSessionWait.Set(); });
+					proxyServer->SetWorkTracker(client->GetWorkTracker());
+					proxyServer->StartListen(networkBackendMem, proxyPort);
+					proxyServer->StartListen(*networkBackend, proxyPort);
+
+					proxyNetworkServer = proxyServer;
+					wakeupSessionWait.Set();
+					targetConnectionCount = maxConnectionCount;
 					return true;
 				};
 
 			Atomic<bool> isDisconnected;
-			client->RegisterOnDisconnected([&]() { isDisconnected = true; networkBackend->StopListen(); if (auto proxyServer = proxy.server.load()) proxyServer->DisconnectClients(); });
+			client->RegisterOnDisconnected([&]() { isDisconnected = true; networkBackend->StopListen(); if (auto proxyServer = proxyNetworkServer.load()) proxyServer->DisconnectClients(); });
 
 			struct NetworkBackends
 			{
@@ -1118,19 +1109,17 @@ namespace uba
 			storageInfo.getProxyBackendCallback = getProxyBackend;
 			storageInfo.getProxyBackendUserData = &backends;
 			storageInfo.allowProxy = allowProxy;
-			storageInfo.startProxyCallback = startProxy;
-			storageInfo.startProxyUserData = &proxy;
+			storageInfo.startProxyCallback = [](void* userData, u16 proxyPort, const Guid& storageServerUid) { return (*(decltype(startProxy)*)userData)(proxyPort, storageServerUid); };;
+			storageInfo.startProxyUserData = &startProxy;
 			storageInfo.zone = zone.data;
 			storageInfo.proxyPort = proxyPort;
 
-			auto storageClient = new StorageClient(storageInfo);
+			storageClient = new StorageClient(storageInfo);
 			auto bscsg = MakeGuard([&]() { delete storageClient; });
 
 			if (!storageClient->LoadCasTable(true))
 				return false;
 
-			proxy.storageClient = storageClient;
-
 			SessionClient* sessionClient = nullptr;
 
 			CasKey keys[2];
@@ -1199,12 +1188,12 @@ namespace uba
 				storageClient->StopProxy();
 				sessionClient->Stop();
 
-				auto proxyServer = proxy.server.load();
+				auto proxyServer = proxyNetworkServer.load();
 				if (proxyServer)
 				{
 					// Let's give the active fetches some time (60 seconds)
 					u32 waitCount = 60*10;
-					while (proxy.storage->GetActiveFetchCount())
+					while (storageProxy->GetActiveFetchCount())
 					{
 						Sleep(100);
 						if (!waitCount--)
@@ -1297,7 +1286,7 @@ namespace uba
 				}
 
 				// If we are the proxy server and have external connections we lower max process count. Note that it will always have one connection which is itself
-				if (auto proxyServer = proxy.server.load())
+				if (auto proxyServer = proxyNetworkServer.load())
 					if (proxyServer->GetConnectionCount() > 1)
 						sessionClient->SetMaxProcessCount(maxProcessCount - 2);
 
diff --git a/Engine/Source/Programs/UnrealBuildAccelerator/CacheService/Private/UbaCacheService.cpp b/Engine/Source/Programs/UnrealBuildAccelerator/CacheService/Private/UbaCacheService.cpp
--- a/Engine/Source/Programs/UnrealBuildAccelerator/CacheService/Private/UbaCacheService.cpp
+++ b/Engine/Source/Programs/UnrealBuildAccelerator/CacheService/Private/UbaCacheService.cpp
@@ -261,7 +261,7 @@ namespace uba
 		storageInfo.storeCompressed = storeCompressed;
 		storageInfo.allowHintAsFallback = false;
 		storageInfo.manuallyHandleOverflow = true;
-		storageInfo.writeRecievedCasFilesToDisk = true;
+		storageInfo.writeReceivedCasFilesToDisk = true;
 		StorageServer storageServer(storageInfo);
 
 		if (!storageServer.LoadCasTable(true, true))
diff --git a/Engine/Source/Programs/UnrealBuildAccelerator/Cli/Private/UbaCli.cpp b/Engine/Source/Programs/UnrealBuildAccelerator/Cli/Private/UbaCli.cpp
--- a/Engine/Source/Programs/UnrealBuildAccelerator/Cli/Private/UbaCli.cpp
+++ b/Engine/Source/Programs/UnrealBuildAccelerator/Cli/Private/UbaCli.cpp
@@ -15,7 +15,6 @@
 #include "UbaSessionClient.h"
 #include "UbaSessionServer.h"
 #include "UbaStorageClient.h"
-#include "UbaStorageProxy.h"
 #include "UbaStorageServer.h"
 #include "UbaStorageUtils.h"
 #include "UbaVersion.h"
@@ -1031,7 +1030,7 @@ namespace uba
 			cacheRootDir.Append("CacheServer");
 			NetworkServer cacheNetworkServer(ctorSuccess);
 			StorageServerCreateInfo storageInfo2(cacheNetworkServer, cacheRootDir.data, logWriter);
-			storageInfo2.writeRecievedCasFilesToDisk = true;
+			storageInfo2.writeReceivedCasFilesToDisk = true;
 			StorageServer cacheStorageServer(storageInfo2);
 			CacheServer cacheServer(logWriter, cacheRootDir.data, cacheNetworkServer, cacheStorageServer);
 			auto csg = MakeGuard([&]() { cacheServer.Save(); });
diff --git a/Engine/Source/Programs/UnrealBuildAccelerator/Cli/Private/UbaClient.h b/Engine/Source/Programs/UnrealBuildAccelerator/Cli/Private/UbaClient.h
--- a/Engine/Source/Programs/UnrealBuildAccelerator/Cli/Private/UbaClient.h
+++ b/Engine/Source/Programs/UnrealBuildAccelerator/Cli/Private/UbaClient.h
@@ -85,7 +85,8 @@ namespace uba
 				delete proxyNetworkServer;
 				return false;
 			}
-			proxyStorage = new StorageProxy(*proxyNetworkServer, *networkClient, storageServerUid, TC("Wooohoo"), storageClient);
+			StorageProxyCreateInfo proxyInfo { *proxyNetworkServer, *networkClient, storageServerUid, TC("Wooohoo"), storageClient };
+			proxyStorage = new StorageProxy(proxyInfo);
 			proxyNetworkServer->StartListen(*networkBackendMem, proxyPort);
 			proxyNetworkServer->StartListen(*networkBackend, proxyPort);
 			return true;
diff --git a/Engine/Source/Programs/UnrealBuildAccelerator/Common/Private/UbaSessionServer.cpp b/Engine/Source/Programs/UnrealBuildAccelerator/Common/Private/UbaSessionServer.cpp
--- a/Engine/Source/Programs/UnrealBuildAccelerator/Common/Private/UbaSessionServer.cpp
+++ b/Engine/Source/Programs/UnrealBuildAccelerator/Common/Private/UbaSessionServer.cpp
@@ -1056,7 +1056,6 @@ namespace uba
 
 		if (success)
 		{
-			m_storage.DropCasFile(casKey, false, destination.data);
 			bool invalidateStorage = false; // No need, already handled in m_storage.CopyOrLink
 			RegisterCreateFileForWrite(StringKeyZero, destination, shouldWriteToDisk, 0, 0, invalidateStorage);
 
diff --git a/Engine/Source/Programs/UnrealBuildAccelerator/Common/Private/UbaStorage.cpp b/Engine/Source/Programs/UnrealBuildAccelerator/Common/Private/UbaStorage.cpp
--- a/Engine/Source/Programs/UnrealBuildAccelerator/Common/Private/UbaStorage.cpp
+++ b/Engine/Source/Programs/UnrealBuildAccelerator/Common/Private/UbaStorage.cpp
@@ -115,13 +115,13 @@ namespace uba
 
 	void StorageImpl::CasEntryAccessed(CasEntry& casEntry)
 	{
-		bool hasMapping;
 		{
 			SCOPED_READ_LOCK(casEntry.lock, l); // Note, this lock is taken again outside CasEntryAccessed.. so if this takes a long time it won't help to remove this lock
-			hasMapping = casEntry.mappingHandle.IsValid();
+			if (casEntry.dropped)
+				return;
+			//if (casEntry.mappingHandle.IsValid())
+			//	return;
 		}
-		if (hasMapping)
-			return;
 
 		SCOPED_WRITE_LOCK(m_accessLock, lock);
 
@@ -159,9 +159,19 @@ namespace uba
 
 		UBA_ASSERT(!casEntry.mappingHandle.IsValid());
 
-		if (!m_casCapacityBytes || m_overflowReported || m_casTotalBytes <= m_casCapacityBytes || m_manuallyHandleOverflow)
+		if (!m_casCapacityBytes || m_overflowReported || m_manuallyHandleOverflow)
 			return;
 
+		if (m_casTotalBytes <= m_casCapacityBytes)
+		{
+			if (!casEntry.prevAccessed && !casEntry.nextAccessed && (!m_newestAccessed || m_newestAccessed != &casEntry))
+			{
+				UBA_ASSERT(!casEntry.dropped);
+				AttachEntry(casEntry);
+			}
+			return;
+		}
+
 		TimerScope ts(Stats().handleOverflow);
 		UBA_ASSERT(!m_newestAccessed || !m_newestAccessed->prevAccessed);
 		UBA_ASSERT(!m_oldestAccessed || !m_oldestAccessed->nextAccessed);
@@ -1405,12 +1415,15 @@ namespace uba
 		if (!m_casTableLoaded)
 			return true;
 
+		u64 startTime = GetTime();
+
 		StringBuffer<256> fileName;
 		fileName.Append(m_rootDir).Append(TC("casdb"));
 
 		StringBuffer<256> tempFileName;
 		tempFileName.Append(fileName).Append(TC(".tmp"));
 
+		u32 deleteCount = 0;
 		{
 			FileAccessor tempFile(m_logger, tempFileName.data);
 			if (!tempFile.CreateWrite(false, DefaultAttributes(), 0, m_tempPath.data))
@@ -1461,12 +1474,13 @@ namespace uba
 				CasEntry& casEntry = *it;
 				if (casEntry.verified && !casEntry.exists)
 					continue;
-				if (casEntry.dropped)
+				if (deleteDropped && casEntry.dropped)
 				{
 					StringBuffer<512> casFileName;
 					if (!StorageImpl::GetCasFileName(casFileName, casEntry.key))
 						continue;
-					DeleteFileW(casFileName.data);
+					if (DeleteFileW(casFileName.data))
+						++deleteCount;
 					continue;
 				}
 
@@ -1501,6 +1515,9 @@ namespace uba
 
 		if (m_overflowReported)
 			m_logger.Info(TC("Session needs at least %s to not overflow."), BytesToText(m_casMaxBytes).str);
+
+		u64 duration = GetTime() - startTime;
+		m_logger.Detail(TC("Database saved in %s (deleted %u dropped cas files)"), TimeToText(duration).str, deleteCount);
 		return true;
 	}
 
diff --git a/Engine/Source/Programs/UnrealBuildAccelerator/Common/Private/UbaStorageClient.cpp b/Engine/Source/Programs/UnrealBuildAccelerator/Common/Private/UbaStorageClient.cpp
--- a/Engine/Source/Programs/UnrealBuildAccelerator/Common/Private/UbaStorageClient.cpp
+++ b/Engine/Source/Programs/UnrealBuildAccelerator/Common/Private/UbaStorageClient.cpp
@@ -35,6 +35,7 @@ namespace uba
 		table.GetValueAsBool(sendCompressed, TC("SendCompressed"));
 		table.GetValueAsBool(allowProxy, TC("AllowProxy"));
 		table.GetValueAsBool(sendOneBigFileAtTheTime, TC("SendOneBigAtTheTime"));
+		table.GetValueAsBool(checkExistsOnServer, TC("CheckExistsOnServer"));
 	}
 
 	StorageClient::StorageClient(const StorageClientCreateInfo& info)
@@ -43,6 +44,7 @@ namespace uba
 	,	m_sendCompressed(info.sendCompressed)
 	,	m_allowProxy(info.allowProxy)
 	,	m_sendOneBigFileAtTheTime(info.sendOneBigFileAtTheTime)
+	,	m_checkExistsOnServer(info.checkExistsOnServer)
 	,	m_zone(info.zone)
 	,	m_getProxyBackendCallback(info.getProxyBackendCallback)
 	,	m_getProxyBackendUserData(info.getProxyBackendUserData)
@@ -866,13 +868,20 @@ namespace uba
 			UBA_ASSERT(false);
 		}
 
-		StackBinaryWriter<128> writer;
-		NetworkMessage msg(client, ServiceId, StorageMessageType_ExistsOnServer, writer);
-		writer.WriteCasKey(casKey);
-		StackBinaryReader<128> reader;
-		if (!msg.Send(reader))
-			return false;
-		if (!reader.ReadBool())
+		bool existsOnServer = false;
+
+		if (m_checkExistsOnServer)
+		{
+			StackBinaryWriter<128> writer;
+			NetworkMessage msg(client, ServiceId, StorageMessageType_ExistsOnServer, writer);
+			writer.WriteCasKey(casKey);
+			StackBinaryReader<128> reader;
+			if (!msg.Send(reader))
+				return false;
+			existsOnServer = reader.ReadBool();
+		}
+
+		if (!existsOnServer)
 		{
 			if (storeCompressed)
 			{
diff --git a/Engine/Source/Programs/UnrealBuildAccelerator/Common/Private/UbaStorageProxy.cpp b/Engine/Source/Programs/UnrealBuildAccelerator/Common/Private/UbaStorageProxy.cpp
--- a/Engine/Source/Programs/UnrealBuildAccelerator/Common/Private/UbaStorageProxy.cpp
+++ b/Engine/Source/Programs/UnrealBuildAccelerator/Common/Private/UbaStorageProxy.cpp
@@ -1,6 +1,7 @@
 // Copyright Epic Games, Inc. All Rights Reserved.
 
 #include "UbaStorageProxy.h"
+#include "UbaConfig.h"
 #include "UbaFileAccessor.h"
 #include "UbaNetworkClient.h"
 #include "UbaNetworkMessage.h"
@@ -34,15 +35,25 @@ namespace uba
 		bool error = false;
 	};
 
+	void StorageProxyCreateInfo::Apply(Config& config, const tchar* tableName)
+	{
+		const ConfigTable* tablePtr = config.GetTable(tableName);
+		if (!tablePtr)
+			return;
+		const ConfigTable& table = *tablePtr;
+		table.GetValueAsBool(useLocalStorage, TC("UseLocalStorage"));
+	}
 
-	StorageProxy::StorageProxy(NetworkServer& server, NetworkClient& client, const Guid& storageServerUid, const tchar* name, StorageImpl* localStorage)
-	:	m_server(server)
-	,	m_client(client)
-	,	m_localStorage(localStorage)
-	,	m_logger(client.GetLogWriter(), TC("StorageProxy"))
-	,	m_storageServerUid(storageServerUid)
-	,	m_name(name)
+	StorageProxy::StorageProxy(const StorageProxyCreateInfo& info)
+	:	m_server(info.server)
+	,	m_client(info.client)
+	,	m_localStorage(info.localStorage)
+	,	m_logger(info.client.GetLogWriter(), TC("StorageProxy"))
+	,	m_storageServerUid(info.storageServerUid)
+	,	m_name(info.name)
 	{
+		m_useLocalStorage = info.useLocalStorage;
+
 		m_server.RegisterOnClientDisconnected(0, [this](const Guid& clientUid, u32 clientId)
 			{
 				SCOPED_WRITE_LOCK(m_activeFetchesLock, lock);
@@ -159,9 +170,8 @@ namespace uba
 					if (file.memory)
 						break;
 
-					bool useLocalStorage = false; // Seems like it might be deadlocks in this code.. need to revisit
 					bool storeCompressed = true;
-					if (useLocalStorage && m_localStorage && IsCompressed(casKey) && m_inProcessClientId && connectionInfo.GetId() != m_inProcessClientId)
+					if (m_useLocalStorage && m_localStorage && IsCompressed(casKey) && m_inProcessClientId && connectionInfo.GetId() != m_inProcessClientId)
 					{
 						// We need to leave this lock here since the in-process storage client might be asking for this file too and then we can end up in a deadlock
 						fileLock.Leave();
diff --git a/Engine/Source/Programs/UnrealBuildAccelerator/Common/Private/UbaStorageServer.cpp b/Engine/Source/Programs/UnrealBuildAccelerator/Common/Private/UbaStorageServer.cpp
--- a/Engine/Source/Programs/UnrealBuildAccelerator/Common/Private/UbaStorageServer.cpp
+++ b/Engine/Source/Programs/UnrealBuildAccelerator/Common/Private/UbaStorageServer.cpp
@@ -15,7 +15,7 @@ namespace uba
 		if (!tablePtr)
 			return;
 		const ConfigTable& table = *tablePtr;
-		table.GetValueAsBool(writeRecievedCasFilesToDisk, TC("WriteRecievedCasFilesToDisk"));
+		table.GetValueAsBool(writeReceivedCasFilesToDisk, TC("WriteReceivedCasFilesToDisk"));
 		table.GetValueAsBool(allowHintAsFallback, TC("AllowHintAsFallback"));
 	}
 
@@ -25,7 +25,7 @@ namespace uba
 	{
 		m_zone = info.zone;
 		m_allowHintAsFallback = info.allowHintAsFallback;
-		m_writeRecievedCasFilesToDisk = info.writeRecievedCasFilesToDisk;
+		m_writeReceivedCasFilesToDisk = info.writeReceivedCasFilesToDisk;
 
 		if (!CreateGuid(m_uid))
 			UBA_ASSERT(false);
@@ -276,6 +276,39 @@ namespace uba
 		}
 	}
 
+	bool StorageServer::VerifyExists(bool& outExists, CasEntry& casEntry, ScopedWriteLock& entryLock, const CasKey& casKey)
+	{
+		outExists = false;
+		StringBuffer<> casFile;
+		if (!GetCasFileName(casFile, casKey))
+			return false;
+		u64 outFileSize = 0;
+		if (FileExists(m_logger, casFile.data, &outFileSize))
+		{
+			if (outFileSize == 0 && casKey != EmptyFileKey)
+			{
+				m_logger.Warning(TC("Found file %s with size 0 which did not have the zero-size-caskey. Deleting"), casFile.data);
+				if (!uba::DeleteFileW(casFile.data))
+					return m_logger.Error(TC("Failed to delete %s. Clean cas folder and restart"), casFile.data);
+				casEntry.exists = false;
+				casEntry.verified = true;
+			}
+			else
+			{
+				casEntry.verified = true;
+				outExists = true;
+				entryLock.Leave();
+				CasEntryWritten(casEntry, outFileSize);
+			}
+		}
+		else
+		{
+			casEntry.exists = false;
+			casEntry.verified = true;
+		}
+		return true;
+	}
+
 	bool StorageServer::HandleMessage(const ConnectionInfo& connectionInfo, const WorkContext& workContext, u8 messageType, BinaryReader& reader, BinaryWriter& writer)
 	{
 		ActiveStore* firstStore = nullptr;
@@ -724,33 +757,8 @@ namespace uba
 
 				if (!exists && casEntry.exists)
 				{
-					StringBuffer<> casFile;
-					if (!GetCasFileName(casFile, casKey))
+					if (!VerifyExists(exists, casEntry, entryLock, casKey))
 						return false;
-					u64 outFileSize = 0;
-					if (FileExists(m_logger, casFile.data, &outFileSize))
-					{
-						if (outFileSize == 0 && casKey != EmptyFileKey)
-						{
-							m_logger.Warning(TC("Found file %s with size 0 which did not have the zero-size-caskey. Deleting"), casFile.data);
-							if (!uba::DeleteFileW(casFile.data))
-								return m_logger.Error(TC("Failed to delete %s. Clean cas folder and restart"), casFile.data);
-							casEntry.exists = false;
-							casEntry.verified = true;
-						}
-						else
-						{
-							casEntry.verified = true;
-							exists = true;
-							entryLock.Leave();
-							CasEntryWritten(casEntry, outFileSize);
-						}
-					}
-					else
-					{
-						casEntry.exists = false;
-						casEntry.verified = true;
-					}
 				}
 				writer.WriteBool(exists);
 				return true;
@@ -774,7 +782,20 @@ namespace uba
 				SCOPED_WRITE_LOCK(casEntry.lock, entryLock);
 				if (!casEntry.verified)
 				{
-					casEntry.verified = true;
+					if (bool exists = casEntry.exists)
+					{
+						if (!VerifyExists(exists, casEntry, entryLock, casKey))
+							return false;
+						if (exists)
+						{
+							writer.WriteU16(u16(~0));
+							writer.WriteBool(m_traceStore);
+							m_logger.Warning(TC("Store request of %s which already exists in casdb (%s)"), CasKeyString(casKey).str, hint.data);
+							return true;
+						}
+					}
+					else
+						casEntry.verified = true;
 				}
 				else
 				{
@@ -787,6 +808,7 @@ namespace uba
 						CasEntryAccessed(casEntry);
 						writer.WriteU16(u16(~0));
 						writer.WriteBool(m_traceStore);
+						m_logger.Warning(TC("Store request of %s which already exists in casdb (%s)"), CasKeyString(casKey).str, hint.data);
 						return true;
 					}
 				}
@@ -800,7 +822,7 @@ namespace uba
 				MappedView mappedView;
 				FileAccessor* fileAccessor = nullptr;
 				
-				if (m_writeRecievedCasFilesToDisk)
+				if (m_writeReceivedCasFilesToDisk)
 				{
 					StringBuffer<> casKeyName;
 					GetCasFileName(casKeyName, casKey);
@@ -812,7 +834,7 @@ namespace uba
 						delete fileAccessor;
 						free((void*)filename);
 
-						m_logger.Error(TC("Failed to create cas file %s"), casKeyName.data);
+						m_logger.Error(TC("Failed to create cas file %s (%s)"), casKeyName.data, hint.data);
 						casEntry.verified = false;
 						return false;
 					}
@@ -831,6 +853,7 @@ namespace uba
 						casEntry.verified = false;
 						return false;
 					}
+					casEntry.dropped = true; // Set to dropped because this will never hit disk
 				}
 
 				casEntry.beingWritten = true;
@@ -903,9 +926,9 @@ namespace uba
 						casEntry.beingWritten = false;
 					}
 
-					bool isPersistentStore = m_writeRecievedCasFilesToDisk;
+					bool isPersistentStore = m_writeReceivedCasFilesToDisk;
 					if (isPersistentStore)
-						CasEntryWritten(*activeStore.casEntry, totalWritten);
+						CasEntryWritten(casEntry, totalWritten);
 
 					activeStore.recvCasTime += GetTime() - time2;
 
diff --git a/Engine/Source/Programs/UnrealBuildAccelerator/Common/Private/UbaStorageUtils.cpp b/Engine/Source/Programs/UnrealBuildAccelerator/Common/Private/UbaStorageUtils.cpp
--- a/Engine/Source/Programs/UnrealBuildAccelerator/Common/Private/UbaStorageUtils.cpp
+++ b/Engine/Source/Programs/UnrealBuildAccelerator/Common/Private/UbaStorageUtils.cpp
@@ -184,6 +184,16 @@ namespace uba
 		u64 sendLeft = fileSize;
 		u64 sendPos = 0;
 
+		auto sendEndMessage = [&]()
+			{
+				if (!sendEnd)
+					return true;
+				StackBinaryWriter<128> writer;
+				NetworkMessage msg(client, StorageServiceId, StorageMessageType_StoreEnd, writer);
+				writer.WriteCasKey(casKey);
+				return msg.Send();
+			};
+
 		while (sendLeft)
 		{
 			StackBinaryWriter<SendMaxSize> writer;
@@ -228,7 +238,7 @@ namespace uba
 				if (storeId == u16(~0)) // File already exists on server
 				{
 					logger.Info(TC("Server already has file %s (%s)"), CasKeyString(casKey).str, hint);
-					return true;
+					return sendEndMessage();
 				}
 
 				isFirst = false;
@@ -242,16 +252,7 @@ namespace uba
 			}
 		}
 
-		if (sendEnd)
-		{
-			StackBinaryWriter<128> writer;
-			NetworkMessage msg(client, StorageServiceId, StorageMessageType_StoreEnd, writer);
-			writer.WriteCasKey(casKey);
-			if (!msg.Send())
-				return false;
-		}
-
-		return true;
+		return sendEndMessage();
 	}
 
 	bool FileSender::SendFileCompressed(const CasKey& casKey, const tchar* fileName, const u8* sourceMem, u64 sourceSize, const tchar* hint)
@@ -329,6 +330,16 @@ namespace uba
 		u64 sendLeft = fileSize;
 		u64 sendPos = 0;
 
+		auto sendEndMessage = [&]()
+			{
+				if (!sendEnd)
+					return true;
+				StackBinaryWriter<128> writer;
+				NetworkMessage msg(client, StorageServiceId, StorageMessageType_StoreEnd, writer);
+				writer.WriteCasKey(casKey);
+				return msg.Send();
+			};
+
 		bool hasSendOneAtTheTimeLock = false;
 		auto lockGuard = MakeGuard([&]() { if (hasSendOneAtTheTimeLock) m_sendOneAtTheTimeLock.LeaveWrite(); });
 
@@ -382,7 +393,7 @@ namespace uba
 				if (storeId == u16(~0)) // File already exists on server
 				{
 					m_logger.Info(TC("Server already has file %s (%s)"), CasKeyString(casKey).str, hint);
-					return true;
+					return sendEndMessage();
 				}
 
 				isFirst = false;
@@ -396,20 +407,12 @@ namespace uba
 			}
 		}
 
-		if (sendEnd)
-		{
-			StackBinaryWriter<128> writer;
-			NetworkMessage msg(client, StorageServiceId, StorageMessageType_StoreEnd, writer);
-			writer.WriteCasKey(casKey);
-			if (!msg.Send())
-				return false;
-		}
 
 		m_stats.sendCasBytesRaw += sourceSize;
 		m_stats.sendCasBytesComp += fileSize;
 		m_bytesSent = fileSize;
 
-		return true;
+		return sendEndMessage();
 	}
 
 	bool FileFetcher::RetrieveFile(Logger& logger, NetworkClient& client, const CasKey& casKey, const tchar* destination, bool destinationIsCompressed, MemoryBlock* destinationMem)
diff --git a/Engine/Source/Programs/UnrealBuildAccelerator/Common/Public/UbaStorageClient.h b/Engine/Source/Programs/UnrealBuildAccelerator/Common/Public/UbaStorageClient.h
--- a/Engine/Source/Programs/UnrealBuildAccelerator/Common/Public/UbaStorageClient.h
+++ b/Engine/Source/Programs/UnrealBuildAccelerator/Common/Public/UbaStorageClient.h
@@ -24,6 +24,7 @@ namespace uba
 		bool sendCompressed = true;
 		bool allowProxy = true;
 		bool sendOneBigFileAtTheTime = true;
+		bool checkExistsOnServer = false; // There is no point having this enabled if server is not received cas files to disk (which is default)
 		
 		GetProxyBackendCallback* getProxyBackendCallback = nullptr;
 		void* getProxyBackendUserData = nullptr;
@@ -65,6 +66,7 @@ namespace uba
 		bool m_sendCompressed;
 		bool m_allowProxy;
 		bool m_sendOneBigFileAtTheTime;
+		bool m_checkExistsOnServer;
 
 		Guid m_storageServerUid;
 
diff --git a/Engine/Source/Programs/UnrealBuildAccelerator/Common/Public/UbaStorageProxy.h b/Engine/Source/Programs/UnrealBuildAccelerator/Common/Public/UbaStorageProxy.h
--- a/Engine/Source/Programs/UnrealBuildAccelerator/Common/Public/UbaStorageProxy.h
+++ b/Engine/Source/Programs/UnrealBuildAccelerator/Common/Public/UbaStorageProxy.h
@@ -8,6 +8,7 @@
 
 namespace uba
 {
+	class Config;
     class NetworkClient;
 	class NetworkServer;
 	class StorageImpl;
@@ -17,10 +18,22 @@ namespace uba
 	struct MessageInfo;
 	struct WorkContext;
 
+	struct StorageProxyCreateInfo
+	{
+		NetworkServer& server;
+		NetworkClient& client;
+		Guid storageServerUid;
+		const tchar* name = TC("");
+		StorageImpl* localStorage = nullptr;
+		bool useLocalStorage = false; // Use local storage to populate proxy
+
+		void Apply(Config& config, const tchar* tableName = TC("StorageProxy"));
+	};
+
 	class StorageProxy
 	{
 	public:
-		StorageProxy(NetworkServer& server, NetworkClient& client, const Guid& storageServerUid, const tchar* name, StorageImpl* localStorage = nullptr);
+		StorageProxy(const StorageProxyCreateInfo& info);
 		~StorageProxy();
 		void PrintSummary();
 
@@ -85,5 +98,7 @@ namespace uba
 
 		Vector<u16> m_availableIds;
 		u16 m_availableIdsHigh = 1;
+
+		bool m_useLocalStorage;
 	};
 }
diff --git a/Engine/Source/Programs/UnrealBuildAccelerator/Common/Public/UbaStorageServer.h b/Engine/Source/Programs/UnrealBuildAccelerator/Common/Public/UbaStorageServer.h
--- a/Engine/Source/Programs/UnrealBuildAccelerator/Common/Public/UbaStorageServer.h
+++ b/Engine/Source/Programs/UnrealBuildAccelerator/Common/Public/UbaStorageServer.h
@@ -15,7 +15,7 @@ namespace uba
 
 		NetworkServer& server;
 		bool allowHintAsFallback = true; // Will fallback to file system to recalculate cas if cas content does not exist anymore
-		bool writeRecievedCasFilesToDisk = false;
+		bool writeReceivedCasFilesToDisk = false; // Write received files to disk. This is only useful if we expect same file to be written multiple times (requires deterministic outputs ofc)
 		const tchar* zone = TC("");
 	};
 
@@ -45,6 +45,7 @@ namespace uba
 		void OnDisconnected(u32 clientId);
 		bool HandleMessage(const ConnectionInfo& connectionInfo, const WorkContext& workContext, u8 messageType, BinaryReader& reader, BinaryWriter& writer);
 		bool WaitForWritten(CasEntry& casEntry, ScopedWriteLock& entryLock, const ConnectionInfo& connectionInfo, const tchar* hint);
+		bool VerifyExists(bool& outExists, CasEntry& casEntry, ScopedWriteLock& entryLock, const CasKey& casKey);
 
 		u16 PopId();
 		void PushId(u16 id);
@@ -138,6 +139,6 @@ namespace uba
 		Vector<TString> m_disallowedPaths;
 
 		bool m_allowHintAsFallback;
-		bool m_writeRecievedCasFilesToDisk;
+		bool m_writeReceivedCasFilesToDisk;
 	};
 }
diff --git a/Engine/Source/Programs/UnrealBuildAccelerator/StorageProxy/Private/UbaStorageProxyApp.cpp b/Engine/Source/Programs/UnrealBuildAccelerator/StorageProxy/Private/UbaStorageProxyApp.cpp
--- a/Engine/Source/Programs/UnrealBuildAccelerator/StorageProxy/Private/UbaStorageProxyApp.cpp
+++ b/Engine/Source/Programs/UnrealBuildAccelerator/StorageProxy/Private/UbaStorageProxyApp.cpp
@@ -149,7 +149,8 @@ int wmain(int argc, wchar_t *argv[])
 		// TODO: We need to send a message to the storage server to get the uid
 		Guid storageServerUid;
 
-		StorageProxy proxy(server, *client, storageServerUid, proxyName.data);
+		StorageProxyCreateInfo proxyInfo { server, *client, storageServerUid, proxyName.data };
+		StorageProxy proxy(proxyInfo);
 
 		logger.Info(L"Waiting to connect to %s:%u", host.data, u16(hostPort));
 		int retryCount = 5;
diff --git a/Engine/Source/Programs/UnrealBuildAccelerator/Test/Private/UbaTestSession.h b/Engine/Source/Programs/UnrealBuildAccelerator/Test/Private/UbaTestSession.h
--- a/Engine/Source/Programs/UnrealBuildAccelerator/Test/Private/UbaTestSession.h
+++ b/Engine/Source/Programs/UnrealBuildAccelerator/Test/Private/UbaTestSession.h
@@ -33,6 +33,11 @@ namespace uba
 		SessionServerCreateInfo sessionServerInfo(storage, server, logWriter);
 		sessionServerInfo.checkMemory = false;
 		sessionServerInfo.rootDir = rootDir.data;
+
+		#if UBA_DEBUG
+		sessionServerInfo.logToFile = true;
+		#endif
+
 		SessionServer session(sessionServerInfo);
 
 		StringBuffer<MaxPath> workingDir;
@@ -48,7 +53,7 @@ namespace uba
 		return testFunc(logger, session, workingDir.data, [&](const ProcessStartInfo& pi) { return session.RunProcess(pi, true, enableDetour); });
 	}
 
-	bool RunRemote(LoggerWithWriter& logger, const StringBufferBase& testRootDir, const TestSessionFunction& testFunc)
+	bool RunRemote(LoggerWithWriter& logger, const StringBufferBase& testRootDir, const TestSessionFunction& testFunc, bool deleteAll = true)
 	{
 		LogWriter& logWriter = logger.m_writer;
 		NetworkBackendTcp tcpBackend(logWriter);
@@ -59,7 +64,7 @@ namespace uba
 
 		StringBuffer<MaxPath> rootDir;
 		rootDir.Append(testRootDir).Append(TC("Uba"));
-		if (!DeleteAllFiles(logger, rootDir.data))
+		if (deleteAll && !DeleteAllFiles(logger, rootDir.data))
 			return false;
 
 		StorageServerCreateInfo storageServerInfo(server, rootDir.data, logWriter);
@@ -70,13 +75,19 @@ namespace uba
 		SessionServerCreateInfo sessionServerInfo(storageServer, server, logWriter);
 		sessionServerInfo.checkMemory = false;
 		sessionServerInfo.rootDir = rootDir.data;
+
+		#if UBA_DEBUG
+		sessionServerInfo.logToFile = true;
+		sessionServerInfo.remoteLogEnabled = true;
+		#endif
+
 		auto& sessionServer = *new SessionServer(sessionServerInfo);
 		auto ssg2 = MakeGuard([&]() { delete &sessionServer; });
 
 		auto sg = MakeGuard([&]() { server.DisconnectClients(); });
 
 		rootDir.Append(TC("Client"));
-		if (!DeleteAllFiles(logger, rootDir.data))
+		if (deleteAll && !DeleteAllFiles(logger, rootDir.data))
 			return false;
 
 		StorageClientCreateInfo storageClientInfo(client, rootDir.data);
@@ -85,6 +96,12 @@ namespace uba
 
 		SessionClientCreateInfo sessionClientInfo(storageClient, client, logWriter);
 		sessionClientInfo.rootDir = rootDir.data;
+		sessionClientInfo.allowKeepFilesInMemory = false;
+
+		#if UBA_DEBUG
+		sessionClientInfo.logToFile = true;
+		#endif
+
 		auto& sessionClient = *new SessionClient(sessionClientInfo);
 		auto scg2 = MakeGuard([&]() { delete &sessionClient; });
 
@@ -92,11 +109,11 @@ namespace uba
 
 		StringBuffer<MaxPath> workingDir;
 		workingDir.Append(testRootDir).Append(TC("WorkingDir"));
-		if (!DeleteAllFiles(logger, workingDir.data))
+		if (deleteAll && !DeleteAllFiles(logger, workingDir.data))
 			return false;
 		if (!storageServer.CreateDirectory(workingDir.data))
 			return false;
-		if (!DeleteAllFiles(logger, workingDir.data, false))
+		if (deleteAll && !DeleteAllFiles(logger, workingDir.data, false))
 			return false;
 
 		storageClient.Start();
@@ -226,7 +243,9 @@ namespace uba
 		if (!codeFile.Close())
 			return false;
 
-#if PLATFORM_MAC
+#if PLATFORM_WINDOWS
+		const tchar* clangPath = TC("c:\\Program Files\\Microsoft Visual Studio\\2022\\Professional\\VC\\Tools\\Llvm\\bin\\clang-cl.exe");
+#elif PLATFORM_MAC
 		StringBuffer<MaxPath> xcodePath;
 		ExecuteCommand(logger, "/usr/bin/xcrun -f clang++", xcodePath);
 		const tchar* clangPath = xcodePath.data;
@@ -239,17 +258,21 @@ namespace uba
 
 		ProcessStartInfo processInfo;
 		processInfo.application = clangPath;
-#if PLATFORM_MAC
-		StringBuffer<MaxPath> xcodeSDKPath;
+
 		StringBuffer<MaxPath> args;
+
+#if PLATFORM_WINDOWS
+		args.Append("/Brepro ");
+#elif PLATFORM_MAC
+		StringBuffer<MaxPath> xcodeSDKPath;
 		ExecuteCommand(logger, "xcrun --show-sdk-path", xcodeSDKPath);
 		args.Append("-isysroot ");
-		args.Append(xcodeSDKPath.data);
-		args.Append(" -o code Code.cpp");
-		processInfo.arguments = args.data;
-#else
-		processInfo.arguments = TC("-o code Code.cpp");
+		args.Append(xcodeSDKPath.data).Append(' ');
 #endif
+		args.Append(TC("-o code Code.cpp"));
+
+		processInfo.arguments = args.data;
+
 		processInfo.workingDir = workingDir;
 		//processInfo.logFile = TC("[REDACTED_USER_HOME]/RunClang.log");
 		ProcessHandle process = runProcess(processInfo);
@@ -329,7 +352,10 @@ namespace uba
 
 	bool TestRemoteDetouredClang(LoggerWithWriter& logger, const StringBufferBase& testRootDir)
 	{
-		return RunRemote(logger, testRootDir, RunClang);
+		// Run twice to test LoadCasTable/SaveCasTable etc
+		if (!RunRemote(logger, testRootDir, RunClang))
+			return false;
+		return RunRemote(logger, testRootDir, RunClang, false);
 	}
 
 	bool TestDetouredTouch(LoggerWithWriter& logger, const StringBufferBase& testRootDir)
diff --git a/Engine/Source/Programs/UnrealBuildAccelerator/Visualizer/Private/UbaTraceReader.cpp b/Engine/Source/Programs/UnrealBuildAccelerator/Visualizer/Private/UbaTraceReader.cpp
--- a/Engine/Source/Programs/UnrealBuildAccelerator/Visualizer/Private/UbaTraceReader.cpp
+++ b/Engine/Source/Programs/UnrealBuildAccelerator/Visualizer/Private/UbaTraceReader.cpp
@@ -1169,6 +1169,13 @@ namespace uba
 			process.bitmapDirty = true;
 		}
 
+		for (auto& pair : m_activeWorkRecords)
+		{
+			const WorkRecordLocation& active = pair.second;
+			TraceView::WorkRecord& record = out.workTracks[active.track].records[active.index];
+			record.stop = stopTime;
+		}
+
 		m_activeProcesses.clear();
 		out.finished = true;
 	}
