[UBA]
* Fixed code path in win64 detouring code where GetFileInformationByHandle on file that is file map might fail

diff --git a/Engine/Source/Programs/UnrealBuildAccelerator/Detours/Private/Windows/UbaDetoursFunctionsKernelBase.inl b/Engine/Source/Programs/UnrealBuildAccelerator/Detours/Private/Windows/UbaDetoursFunctionsKernelBase.inl
--- a/Engine/Source/Programs/UnrealBuildAccelerator/Detours/Private/Windows/UbaDetoursFunctionsKernelBase.inl
+++ b/Engine/Source/Programs/UnrealBuildAccelerator/Detours/Private/Windows/UbaDetoursFunctionsKernelBase.inl
@@ -2081,7 +2081,7 @@ BOOL Detoured_GetFileInformationByHandle(HANDLE hFile, LPBY_HANDLE_FILE_INFORMAT
 			return TRUE;
 		}
 
-		if (g_runningRemote)
+		if (g_runningRemote || fi.isFileMap)
 		{
 			StringBuffer<> fixedName;
 			FixPath(fixedName, fi.originalName);
