[UBA]
* Added -proxyuselocalstorage to UbaAgent that can be used to enable storage proxy to access local storage. Reason this is behind option is that we might have seen agent hangs because of this feature and I want to easily disable it again

diff --git a/Engine/Source/Programs/UnrealBuildAccelerator/Agent/Private/UbaAgent.cpp b/Engine/Source/Programs/UnrealBuildAccelerator/Agent/Private/UbaAgent.cpp
--- a/Engine/Source/Programs/UnrealBuildAccelerator/Agent/Private/UbaAgent.cpp
+++ b/Engine/Source/Programs/UnrealBuildAccelerator/Agent/Private/UbaAgent.cpp
@@ -105,6 +105,7 @@ namespace uba
 		logger.Info(TC("  -zone                   Set the zone this machine exists in. This info is used to figure out if proxies should be created."));
 		logger.Info(TC("  -version                Prints the version for this executable."));
 		logger.Info(TC("  -noproxy                Does not allow this agent to be a storage proxy for other agents"));
+		logger.Info(TC("  -proxyuselocalstorage   Storage proxy will use local storage to see if files exist"));
 		logger.Info(TC("  -noaws                  Will not try to connect to aws meta data server (this can take time during first startup)"));
 		logger.Info(TC("  -killrandom             Kills random process and exit session"));
 		logger.Info(TC("  -memwait=<percent>      The amount of memory needed to spawn a process. Set this to 100 to disable. Defaults to 80%%"));
@@ -644,6 +645,10 @@ namespace uba
 			{
 				allowProxy = false;
 			}
+			else if (name.Equals(TC("-proxyuselocalstorage")))
+			{
+				proxyUseLocalStorage = true;
+			}
 			else if (name.Equals(TC("-noaws")))
 			{
 				couldBeAws = false;
