[UBA]
* Downgraded warning about sending cas file that already exists to detail... this is very rare (should only happen in non-unity builds.. the PerModuleInline.gen.cpp)
* Removed TestRemoteDetouredClang from macos target for now until I have figured out why it fails

diff --git a/Engine/Source/Programs/UnrealBuildAccelerator/Common/Private/UbaStorageServer.cpp b/Engine/Source/Programs/UnrealBuildAccelerator/Common/Private/UbaStorageServer.cpp
--- a/Engine/Source/Programs/UnrealBuildAccelerator/Common/Private/UbaStorageServer.cpp
+++ b/Engine/Source/Programs/UnrealBuildAccelerator/Common/Private/UbaStorageServer.cpp
@@ -788,9 +788,10 @@ namespace uba
 							return false;
 						if (exists)
 						{
+							// Unless writeReceivedCasFilesToDisk is set this is extremely unlikely to happen. "PerModuleInline.gen.cpp" seems to be the only one (they are supposed to all be the same)
 							writer.WriteU16(u16(~0));
 							writer.WriteBool(m_traceStore);
-							m_logger.Warning(TC("Store request of %s which already exists in casdb (%s)"), CasKeyString(casKey).str, hint.data);
+							m_logger.Detail(TC("Store request of %s which already exists in casdb (%s)"), CasKeyString(casKey).str, hint.data);
 							return true;
 						}
 					}
@@ -808,7 +809,7 @@ namespace uba
 						CasEntryAccessed(casEntry);
 						writer.WriteU16(u16(~0));
 						writer.WriteBool(m_traceStore);
-						m_logger.Warning(TC("Store request of %s which already exists in casdb (%s)"), CasKeyString(casKey).str, hint.data);
+						m_logger.Detail(TC("Store request of %s which already exists in casdb (%s)"), CasKeyString(casKey).str, hint.data);
 						return true;
 					}
 				}
