[UBA]
* Added logic to prevent same casfile from being sent from same client multiple times

diff --git a/Engine/Source/Programs/UnrealBuildAccelerator/Agent/Private/UbaAgent.cpp b/Engine/Source/Programs/UnrealBuildAccelerator/Agent/Private/UbaAgent.cpp
--- a/Engine/Source/Programs/UnrealBuildAccelerator/Agent/Private/UbaAgent.cpp
+++ b/Engine/Source/Programs/UnrealBuildAccelerator/Agent/Private/UbaAgent.cpp
@@ -111,6 +111,7 @@ namespace uba
 		logger.Info(TC("  -memwait=<percent>      The amount of memory needed to spawn a process. Set this to 100 to disable. Defaults to 80%%"));
 		logger.Info(TC("  -memkill=<percent>      The amount of memory needed before processes starts to be killed. Set this to 100 to disable. Defaults to 90%%"));
 		logger.Info(TC("  -crypto=<key>           32 character (16 bytes) crypto key used for secure network transfer"));
+		logger.Info(TC("  -resendcas              Will try to send same cas multiple times (set this to true if server is allowed to remove cas files)"));
 		logger.Info(TC("  -populateCas=<dir>      Prepopulate cas database with files in dir. If files needed exists on machine this can be an optimization"));
 		#if PLATFORM_MAC
 		logger.Info(TC("  -populateCasFromXcodeVersion=<version>   Prepopulate cas database with files from local xcode installation that matches the version."));
@@ -366,6 +367,7 @@ namespace uba
 		bool killRandom = false;
 		bool downloadDetoursLib = true;
 		bool useExceptionHandler = true;
+		bool resendCas = false;
 		StringBuffer<512> sentryUrl;
 		StringBuffer<128> zone;
 		u32 maxIdleSeconds = ~0u;
@@ -604,6 +606,10 @@ namespace uba
 				((u64*)crypto)[1] = StringToValue(value.data + 16, 16);
 				hasCrypto = true;
 			}
+			else if (name.Equals(TC("-resendcas")))
+			{
+				resendCas = true;
+			}
 			else if (name.Equals(TC("-populateCas")))
 			{
 				if (value.IsEmpty())
@@ -1110,6 +1116,7 @@ namespace uba
 			storageInfo.casCapacityBytes = storageCapacity;
 			storageInfo.storeCompressed = storeCompressed;
 			storageInfo.sendCompressed = sendCompressed;
+			storageInfo.resendCas = resendCas;
 			storageInfo.workManager = client;
 			storageInfo.getProxyBackendCallback = getProxyBackend;
 			storageInfo.getProxyBackendUserData = &backends;
diff --git a/Engine/Source/Programs/UnrealBuildAccelerator/Common/Private/UbaStorageClient.cpp b/Engine/Source/Programs/UnrealBuildAccelerator/Common/Private/UbaStorageClient.cpp
--- a/Engine/Source/Programs/UnrealBuildAccelerator/Common/Private/UbaStorageClient.cpp
+++ b/Engine/Source/Programs/UnrealBuildAccelerator/Common/Private/UbaStorageClient.cpp
@@ -36,6 +36,7 @@ namespace uba
 		table.GetValueAsBool(allowProxy, TC("AllowProxy"));
 		table.GetValueAsBool(sendOneBigFileAtTheTime, TC("SendOneBigAtTheTime"));
 		table.GetValueAsBool(checkExistsOnServer, TC("CheckExistsOnServer"));
+		table.GetValueAsBool(resendCas, TC("ResendCas"));
 	}
 
 	StorageClient::StorageClient(const StorageClientCreateInfo& info)
@@ -45,6 +46,7 @@ namespace uba
 	,	m_allowProxy(info.allowProxy)
 	,	m_sendOneBigFileAtTheTime(info.sendOneBigFileAtTheTime)
 	,	m_checkExistsOnServer(info.checkExistsOnServer)
+	,	m_resendCas(info.resendCas)
 	,	m_zone(info.zone)
 	,	m_getProxyBackendCallback(info.getProxyBackendCallback)
 	,	m_getProxyBackendUserData(info.getProxyBackendUserData)
@@ -146,11 +148,18 @@ namespace uba
 		auto findIt = m_localStorageFiles.find(AsCompressed(casKey, false));
 		if (findIt != m_localStorageFiles.end())
 		{
-			if (findIt->second.casEntry.mappingHandle.IsValid())
-				GetMappingString(out, findIt->second.casEntry.mappingHandle, 0);
-			else
-				out.Append(findIt->second.fileName);
-			return true;
+			const LocalFile& localFile = findIt->second;
+			if (localFile.casEntry.mappingHandle.IsValid())
+			{
+				GetMappingString(out, localFile.casEntry.mappingHandle, 0);
+				return true;
+			}
+
+			if (!localFile.fileName.empty())
+			{
+				out.Append(localFile.fileName);
+				return true;
+			}
 		}
 		tempLock.Leave();
 
@@ -165,7 +174,10 @@ namespace uba
 		tempLock.Leave();
 		if (!isValid)
 			return StorageImpl::MapView(casKey, hint);
+
 		LocalFile& file = findIt->second;
+		if (!file.casEntry.mappingHandle.IsValid())
+			return StorageImpl::MapView(casKey, hint);
 
 		MappedView view;
 		view.handle = file.casEntry.mappingHandle;
@@ -198,13 +210,16 @@ namespace uba
 		auto findIt = m_localStorageFiles.find(AsCompressed(casKey, false));
 		if (findIt != m_localStorageFiles.end())
 		{
-			out.casKey = findIt->first;
 			LocalFile& mf = findIt->second;
-			out.size = mf.casEntry.size;
-			out.view.handle = mf.casEntry.mappingHandle;
-			out.view.size = mf.casEntry.size;
-			out.view.isCompressed = false;
-			return true;
+			if (mf.casEntry.mappingHandle.IsValid())
+			{
+				out.casKey = findIt->first;
+				out.size = mf.casEntry.size;
+				out.view.handle = mf.casEntry.mappingHandle;
+				out.view.size = mf.casEntry.size;
+				out.view.isCompressed = false;
+				return true;
+			}
 		}
 		tempLock.Leave();
 
@@ -837,31 +852,37 @@ namespace uba
 		//UBA_ASSERT(wcsstr(fileName, TC("\\output\\"))); // Just make sure our assumption is correct :)
 		if (isTemporaryStorage)
 		{
-			if (keepMappingInMemory)
+			SCOPED_WRITE_LOCK(m_localStorageFilesLock, lock);
+			auto insres = m_localStorageFiles.try_emplace(AsCompressed(casKey, false));
+			LocalFile& localFile = insres.first->second;
+			if (keepMappingInMemory && isPersistentMapping && !localFile.casEntry.mappingHandle.IsValid())
 			{
-				SCOPED_WRITE_LOCK(m_localStorageFilesLock, lock);
-				auto insres = m_localStorageFiles.try_emplace(AsCompressed(casKey, false));
-				LocalFile& localFile = insres.first->second;
-				if (insres.second || !localFile.casEntry.mappingHandle.IsValid())
+				FileMappingHandle mappingHandle2;
+				if (DuplicateFileMapping(m_logger, GetCurrentProcessHandle(), mappingHandle, GetCurrentProcessHandle(), &mappingHandle2, FILE_MAP_READ, false, 0, fileName))
 				{
-					if (isPersistentMapping)
-					{
-						FileMappingHandle mappingHandle2;
-						if (DuplicateFileMapping(m_logger, GetCurrentProcessHandle(), mappingHandle, GetCurrentProcessHandle(), &mappingHandle2, FILE_MAP_READ, false, 0, fileName))
-						{
-							localFile.casEntry.mappingHandle = mappingHandle2;
-							localFile.casEntry.size = fileSize;
-						}
-						else
-							m_logger.Warning(TC("Failed to duplicate handle for file mapping %s (%s)"), fileName, LastErrorToText().data);
-					}
-					else
-					{
-						localFile.casEntry.size = fileSize;
-						localFile.fileName = fileName;
-					}
+					localFile.casEntry.mappingHandle = mappingHandle2;
+					localFile.casEntry.size = fileSize;
+					localFile.casEntry.exists = true;
 				}
+				else
+					m_logger.Warning(TC("Failed to duplicate handle for file mapping %s (%s)"), fileName, LastErrorToText().data);
 			}
+
+			if (!localFile.casEntry.mappingHandle.IsValid() && localFile.fileName.empty())
+			{
+				localFile.casEntry.size = fileSize;
+				localFile.casEntry.exists = true;
+				localFile.fileName = fileName;
+			}
+
+			// Prevent same file from being sent multiple times from same client (it could still be sent from other clients)
+			if (!m_resendCas && localFile.hasBeenSent)
+			{
+				out = casKey;
+				return true;
+			}
+
+			localFile.hasBeenSent = true;
 		}
 		else
 		{
diff --git a/Engine/Source/Programs/UnrealBuildAccelerator/Common/Public/UbaStorageClient.h b/Engine/Source/Programs/UnrealBuildAccelerator/Common/Public/UbaStorageClient.h
--- a/Engine/Source/Programs/UnrealBuildAccelerator/Common/Public/UbaStorageClient.h
+++ b/Engine/Source/Programs/UnrealBuildAccelerator/Common/Public/UbaStorageClient.h
@@ -25,6 +25,7 @@ namespace uba
 		bool allowProxy = true;
 		bool sendOneBigFileAtTheTime = true;
 		bool checkExistsOnServer = false; // There is no point having this enabled if server is not received cas files to disk (which is default)
+		bool resendCas = false; // Will try to send created file everytime even if the cas has been seen before
 		
 		GetProxyBackendCallback* getProxyBackendCallback = nullptr;
 		void* getProxyBackendUserData = nullptr;
@@ -67,13 +68,14 @@ namespace uba
 		bool m_allowProxy;
 		bool m_sendOneBigFileAtTheTime;
 		bool m_checkExistsOnServer;
+		bool m_resendCas;
 
 		Guid m_storageServerUid;
 
 		TString m_zone;
 
 		ReaderWriterLock m_localStorageFilesLock;
-		struct LocalFile { CasEntry casEntry; TString fileName; LocalFile() : casEntry(CasKeyZero) {} };
+		struct LocalFile { CasEntry casEntry; TString fileName; bool hasBeenSent; LocalFile() : casEntry(CasKeyZero), hasBeenSent(false) {} };
 		UnorderedMap<CasKey, LocalFile> m_localStorageFiles;
 
 		ReaderWriterLock m_sendOneAtTheTimeLock;
diff --git a/Engine/Source/Programs/UnrealBuildAccelerator/Visualizer/Private/UbaVisualizer.cpp b/Engine/Source/Programs/UnrealBuildAccelerator/Visualizer/Private/UbaVisualizer.cpp
--- a/Engine/Source/Programs/UnrealBuildAccelerator/Visualizer/Private/UbaVisualizer.cpp
+++ b/Engine/Source/Programs/UnrealBuildAccelerator/Visualizer/Private/UbaVisualizer.cpp
@@ -1975,7 +1975,7 @@ namespace uba
 			if (!selectedWork.entries.empty())
 				logger.AddSpace();
 			for (auto& e : selectedWork.entries)
-				logger.Info(L"   %ls (%ls)", e.text, TimeToText(e.time).str);
+				logger.Info(L"   %ls (%ls)", e.text, TimeToText(e.time - selectedWork.start).str);
 			logger.AddSpace();
 			logger.DrawAtCursor();
 		}
