[UBA]
* Fixed bug when same cas file was attempted to send a second time it would early out too early (before the file was actually transferred) causing the cas file to not exist on the server yet when following messages were sent
* Apply(Config&) -> Apply(const Config&)
* Fixed so logic extracting file name from path for application rules handle both / and \. This logic runs before paths are cleaned up for startup info
* Added casdb used bytes to agent information.
* Added Event::IsCreated() that returns true if Create has been called on event (or if ctor took in parameters which automatically created it)

diff --git a/Engine/Source/Programs/UnrealBuildAccelerator/Common/Private/UbaScheduler.cpp b/Engine/Source/Programs/UnrealBuildAccelerator/Common/Private/UbaScheduler.cpp
--- a/Engine/Source/Programs/UnrealBuildAccelerator/Common/Private/UbaScheduler.cpp
+++ b/Engine/Source/Programs/UnrealBuildAccelerator/Common/Private/UbaScheduler.cpp
@@ -77,7 +77,7 @@ namespace uba
 		Vector<ProcessLogLine> logLines;
 	};
 
-	void SchedulerCreateInfo::Apply(Config& config)
+	void SchedulerCreateInfo::Apply(const Config& config)
 	{
 		if (const ConfigTable* table = config.GetTable(TC("Scheduler")))
 		{
diff --git a/Engine/Source/Programs/UnrealBuildAccelerator/Common/Private/UbaSession.cpp b/Engine/Source/Programs/UnrealBuildAccelerator/Common/Private/UbaSession.cpp
--- a/Engine/Source/Programs/UnrealBuildAccelerator/Common/Private/UbaSession.cpp
+++ b/Engine/Source/Programs/UnrealBuildAccelerator/Common/Private/UbaSession.cpp
@@ -203,7 +203,7 @@ namespace uba
 		process->AddRef();
 	}
 
-	void SessionCreateInfo::Apply(Config& config)
+	void SessionCreateInfo::Apply(const Config& config)
 	{
 		const ConfigTable* tablePtr = config.GetTable(TC("Session"));
 		if (!tablePtr)
@@ -1549,18 +1549,20 @@ namespace uba
 
 	const ApplicationRules* Session::GetRules(const ProcessStartInfo& si)
 	{
-		u32 exeNameStart = 0;
-		u32 exeNameEnd = TStrlen(si.application);
-		UBA_ASSERT(exeNameEnd > 1);
-		if (const tchar* lastSeparator = TStrrchr(si.application, PathSeparator))
-			exeNameStart = u32(lastSeparator - si.application + 1);
-		else if (si.application[exeNameStart] == '"')
+		const tchar* exeNameStart = si.application;
+		const tchar* exeNameEnd = exeNameStart + TStrlen(si.application);
+		UBA_ASSERT(exeNameEnd - exeNameStart > 1);
+		if (const tchar* lastSeparator = TStrrchr(exeNameStart, PathSeparator))
+			exeNameStart = lastSeparator + 1;
+		if (const tchar* lastSeparator2 = TStrrchr(exeNameStart, NonPathSeparator))
+			exeNameStart = lastSeparator2 + 1;
+		if (*exeNameStart == '"')
 			++exeNameStart;
-		if (si.application[exeNameEnd - 1] == '"')
+		if (exeNameEnd[-1] == '"')
 			--exeNameEnd;
 		StringBuffer<128> exeName;
 		UBA_ASSERTF(exeNameStart < exeNameEnd, TC("Bad application string: %s"), si.application);
-		exeName.Append(si.application + exeNameStart, exeNameEnd - exeNameStart);
+		exeName.Append(exeNameStart, exeNameEnd - exeNameStart);
 		
 		auto rules = GetApplicationRules();
 		
@@ -3008,7 +3010,7 @@ namespace uba
 		BytesToText temp(capacity);
 		if (capacity)
 			capacityStr = temp.str;
-		out.Appendf(TC("%s Mem:%ugb Cas:%s"), cpuStr.data, u32(totalMemoryInKilobytes/(1024*1024)), capacityStr);
+		out.Appendf(TC("%s Mem:%ugb Cas:%s/%s"), cpuStr.data, u32(totalMemoryInKilobytes/(1024*1024)), BytesToText(m_storage.GetStorageUsed()).str, capacityStr);
 
 		StringBuffer<128> zone;
 		if (m_storage.GetZone(zone))
diff --git a/Engine/Source/Programs/UnrealBuildAccelerator/Common/Private/UbaSessionServer.cpp b/Engine/Source/Programs/UnrealBuildAccelerator/Common/Private/UbaSessionServer.cpp
--- a/Engine/Source/Programs/UnrealBuildAccelerator/Common/Private/UbaSessionServer.cpp
+++ b/Engine/Source/Programs/UnrealBuildAccelerator/Common/Private/UbaSessionServer.cpp
@@ -93,7 +93,7 @@ namespace uba
 		u32 m_knownInputsCount = 0;
 	};
 
-	void SessionServerCreateInfo::Apply(Config& config)
+	void SessionServerCreateInfo::Apply(const Config& config)
 	{
 		SessionCreateInfo::Apply(config);
 
diff --git a/Engine/Source/Programs/UnrealBuildAccelerator/Common/Private/UbaStorage.cpp b/Engine/Source/Programs/UnrealBuildAccelerator/Common/Private/UbaStorage.cpp
--- a/Engine/Source/Programs/UnrealBuildAccelerator/Common/Private/UbaStorage.cpp
+++ b/Engine/Source/Programs/UnrealBuildAccelerator/Common/Private/UbaStorage.cpp
@@ -70,7 +70,7 @@ namespace uba
 #endif
 	}
 
-	void StorageCreateInfo::Apply(Config& config)
+	void StorageCreateInfo::Apply(const Config& config)
 	{
 		const ConfigTable* tablePtr = config.GetTable(TC("Storage"));
 		if (!tablePtr)
@@ -1695,6 +1695,11 @@ namespace uba
 		return m_casCapacityBytes;
 	}
 
+	u64 StorageImpl::GetStorageUsed()
+	{
+		return m_casTotalBytes;
+	}
+
 	bool StorageImpl::GetZone(StringBufferBase& out)
 	{
 		return false;
diff --git a/Engine/Source/Programs/UnrealBuildAccelerator/Common/Private/UbaStorageClient.cpp b/Engine/Source/Programs/UnrealBuildAccelerator/Common/Private/UbaStorageClient.cpp
--- a/Engine/Source/Programs/UnrealBuildAccelerator/Common/Private/UbaStorageClient.cpp
+++ b/Engine/Source/Programs/UnrealBuildAccelerator/Common/Private/UbaStorageClient.cpp
@@ -24,7 +24,7 @@ namespace uba
 	}
 
 
-	void StorageClientCreateInfo::Apply(Config& config)
+	void StorageClientCreateInfo::Apply(const Config& config)
 	{
 		StorageCreateInfo::Apply(config);
 
@@ -848,46 +848,45 @@ namespace uba
 		if (casKey == CasKeyZero)
 			return false;
 
-		bool isTemporaryStorage = true; // This function is currently only used for temporary storage (output from client)
-		//UBA_ASSERT(wcsstr(fileName, TC("\\output\\"))); // Just make sure our assumption is correct :)
-		if (isTemporaryStorage)
+		SCOPED_WRITE_LOCK(m_localStorageFilesLock, lock);
+		auto insres = m_localStorageFiles.try_emplace(AsCompressed(casKey, false));
+		LocalFile& localFile = insres.first->second;
+		if (keepMappingInMemory && isPersistentMapping && !localFile.casEntry.mappingHandle.IsValid())
 		{
-			SCOPED_WRITE_LOCK(m_localStorageFilesLock, lock);
-			auto insres = m_localStorageFiles.try_emplace(AsCompressed(casKey, false));
-			LocalFile& localFile = insres.first->second;
-			if (keepMappingInMemory && isPersistentMapping && !localFile.casEntry.mappingHandle.IsValid())
-			{
-				FileMappingHandle mappingHandle2;
-				if (DuplicateFileMapping(m_logger, GetCurrentProcessHandle(), mappingHandle, GetCurrentProcessHandle(), &mappingHandle2, FILE_MAP_READ, false, 0, fileName))
-				{
-					localFile.casEntry.mappingHandle = mappingHandle2;
-					localFile.casEntry.size = fileSize;
-					localFile.casEntry.exists = true;
-				}
-				else
-					m_logger.Warning(TC("Failed to duplicate handle for file mapping %s (%s)"), fileName, LastErrorToText().data);
-			}
-
-			if (!localFile.casEntry.mappingHandle.IsValid() && localFile.fileName.empty())
+			FileMappingHandle mappingHandle2;
+			if (DuplicateFileMapping(m_logger, GetCurrentProcessHandle(), mappingHandle, GetCurrentProcessHandle(), &mappingHandle2, FILE_MAP_READ, false, 0, fileName))
 			{
+				localFile.casEntry.mappingHandle = mappingHandle2;
 				localFile.casEntry.size = fileSize;
 				localFile.casEntry.exists = true;
-				localFile.fileName = fileName;
-			}
-
-			// Prevent same file from being sent multiple times from same client (it could still be sent from other clients)
-			if (!m_resendCas && localFile.hasBeenSent)
-			{
-				out = casKey;
-				return true;
 			}
+			else
+				m_logger.Warning(TC("Failed to duplicate handle for file mapping %s (%s)"), fileName, LastErrorToText().data);
+		}
 
-			localFile.hasBeenSent = true;
+		if (!localFile.casEntry.mappingHandle.IsValid() && localFile.fileName.empty())
+		{
+			localFile.casEntry.size = fileSize;
+			localFile.casEntry.exists = true;
+			localFile.fileName = fileName;
 		}
-		else
+
+		// Prevent same file from being sent multiple times from same client (it could still be sent from other clients)
+		if (!m_resendCas)
 		{
-			UBA_ASSERT(false);
+			if (localFile.hasBeenSent.IsCreated()) // If it is created it means that the file has either been transferred or is being transferred
+			{
+				lock.Leave();
+				if (localFile.hasBeenSent.IsSet(30*1000)) // 30 seconds timeout, then we try to send it ourselves
+				{
+					out = casKey;
+					return true;
+				}
+			}
+			else
+				localFile.hasBeenSent.Create(true);
 		}
+		lock.Leave();
 
 		bool existsOnServer = false;
 
@@ -921,10 +920,10 @@ namespace uba
 				stats.sendCasBytesComp += fileSize;
 			}
 		}
-		else
-		{
-			//m_logger.Warning(TC("File already exists on server %s"), fileName);
-		}
+
+		if (!m_resendCas)
+			localFile.hasBeenSent.Set();
+
 		out = casKey;
 		return true;
 	}
diff --git a/Engine/Source/Programs/UnrealBuildAccelerator/Common/Private/UbaStorageServer.cpp b/Engine/Source/Programs/UnrealBuildAccelerator/Common/Private/UbaStorageServer.cpp
--- a/Engine/Source/Programs/UnrealBuildAccelerator/Common/Private/UbaStorageServer.cpp
+++ b/Engine/Source/Programs/UnrealBuildAccelerator/Common/Private/UbaStorageServer.cpp
@@ -8,7 +8,7 @@
 
 namespace uba
 {
-	void StorageServerCreateInfo::Apply(Config& config)
+	void StorageServerCreateInfo::Apply(const Config& config)
 	{
 		StorageCreateInfo::Apply(config);
 		const ConfigTable* tablePtr = config.GetTable(TC("Storage"));
diff --git a/Engine/Source/Programs/UnrealBuildAccelerator/Common/Public/UbaScheduler.h b/Engine/Source/Programs/UnrealBuildAccelerator/Common/Public/UbaScheduler.h
--- a/Engine/Source/Programs/UnrealBuildAccelerator/Common/Public/UbaScheduler.h
+++ b/Engine/Source/Programs/UnrealBuildAccelerator/Common/Public/UbaScheduler.h
@@ -21,7 +21,7 @@ namespace uba
 	{
 		SchedulerCreateInfo(SessionServer& s) : session(s) {}
 
-		void Apply(Config& config);
+		void Apply(const Config& config);
 
 		SessionServer& session;
 		CacheClient* cacheClient = nullptr; // Set cache client for scheduler to use when building
diff --git a/Engine/Source/Programs/UnrealBuildAccelerator/Common/Public/UbaSessionCreateInfo.h b/Engine/Source/Programs/UnrealBuildAccelerator/Common/Public/UbaSessionCreateInfo.h
--- a/Engine/Source/Programs/UnrealBuildAccelerator/Common/Public/UbaSessionCreateInfo.h
+++ b/Engine/Source/Programs/UnrealBuildAccelerator/Common/Public/UbaSessionCreateInfo.h
@@ -13,7 +13,7 @@ namespace uba
 	{
 		SessionCreateInfo(Storage& s, LogWriter& w = g_consoleLogWriter) : storage(s), logWriter(w) {}
 
-		void Apply(Config& config);
+		void Apply(const Config& config);
 
 		Storage& storage;
 		LogWriter& logWriter;
diff --git a/Engine/Source/Programs/UnrealBuildAccelerator/Common/Public/UbaSessionServerCreateInfo.h b/Engine/Source/Programs/UnrealBuildAccelerator/Common/Public/UbaSessionServerCreateInfo.h
--- a/Engine/Source/Programs/UnrealBuildAccelerator/Common/Public/UbaSessionServerCreateInfo.h
+++ b/Engine/Source/Programs/UnrealBuildAccelerator/Common/Public/UbaSessionServerCreateInfo.h
@@ -12,7 +12,7 @@ namespace uba
 	{
 		SessionServerCreateInfo(Storage& s, NetworkServer& c, LogWriter& writer = g_consoleLogWriter) : SessionCreateInfo(s, writer), server(c) {}
 
-		void Apply(Config& config);
+		void Apply(const Config& config);
 
 		NetworkServer& server;
 		u8 memWaitLoadPercent = 85; // When memory usage goes above this percent, no new processes will be spawned until back below
diff --git a/Engine/Source/Programs/UnrealBuildAccelerator/Common/Public/UbaStorage.h b/Engine/Source/Programs/UnrealBuildAccelerator/Common/Public/UbaStorage.h
--- a/Engine/Source/Programs/UnrealBuildAccelerator/Common/Public/UbaStorage.h
+++ b/Engine/Source/Programs/UnrealBuildAccelerator/Common/Public/UbaStorage.h
@@ -27,6 +27,7 @@ namespace uba
 		virtual bool Reset() = 0;
 		virtual bool SaveCasTable(bool deleteIsRunningfile, bool deleteDropped = true) = 0;
 		virtual u64 GetStorageCapacity() = 0;
+		virtual u64 GetStorageUsed() = 0;
 		virtual bool GetZone(StringBufferBase& out) = 0;
 		virtual bool HasProxy(u32 clientId) { return false; }
 
@@ -78,7 +79,7 @@ namespace uba
 	{
 		StorageCreateInfo(const tchar* rootDir_, LogWriter& w) : writer(w), rootDir(rootDir_) {}
 
-		void Apply(Config& config);
+		void Apply(const Config& config);
 
 		LogWriter& writer;
 		const tchar* rootDir;
@@ -117,6 +118,7 @@ namespace uba
 
 		virtual bool SaveCasTable(bool deleteIsRunningfile, bool deleteDropped = true) override;
 		virtual u64 GetStorageCapacity() override;
+		virtual u64 GetStorageUsed() override;
 		virtual bool GetZone(StringBufferBase& out) override;
 		virtual bool Reset() override;
 		bool DeleteAllCas();
diff --git a/Engine/Source/Programs/UnrealBuildAccelerator/Common/Public/UbaStorageClient.h b/Engine/Source/Programs/UnrealBuildAccelerator/Common/Public/UbaStorageClient.h
--- a/Engine/Source/Programs/UnrealBuildAccelerator/Common/Public/UbaStorageClient.h
+++ b/Engine/Source/Programs/UnrealBuildAccelerator/Common/Public/UbaStorageClient.h
@@ -16,7 +16,7 @@ namespace uba
 	{
 		StorageClientCreateInfo(NetworkClient& c, const tchar* rootDir_) : StorageCreateInfo(rootDir_, c.GetLogWriter()), client(c) {}
 
-		void Apply(Config& config);
+		void Apply(const Config& config);
 
 		NetworkClient& client;
 		const tchar* zone = TC("");
@@ -75,7 +75,7 @@ namespace uba
 		TString m_zone;
 
 		ReaderWriterLock m_localStorageFilesLock;
-		struct LocalFile { CasEntry casEntry; TString fileName; bool hasBeenSent; LocalFile() : casEntry(CasKeyZero), hasBeenSent(false) {} };
+		struct LocalFile { CasEntry casEntry; TString fileName; Event hasBeenSent; LocalFile() : casEntry(CasKeyZero) {} };
 		UnorderedMap<CasKey, LocalFile> m_localStorageFiles;
 
 		ReaderWriterLock m_sendOneAtTheTimeLock;
diff --git a/Engine/Source/Programs/UnrealBuildAccelerator/Common/Public/UbaStorageServer.h b/Engine/Source/Programs/UnrealBuildAccelerator/Common/Public/UbaStorageServer.h
--- a/Engine/Source/Programs/UnrealBuildAccelerator/Common/Public/UbaStorageServer.h
+++ b/Engine/Source/Programs/UnrealBuildAccelerator/Common/Public/UbaStorageServer.h
@@ -11,7 +11,7 @@ namespace uba
 	{
 		StorageServerCreateInfo(NetworkServer& s, const tchar* rootDir_, LogWriter& writer) : StorageCreateInfo(rootDir_, writer), server(s) { workManager = &server; }
 
-		void Apply(Config& config);
+		void Apply(const Config& config);
 
 		NetworkServer& server;
 		bool allowHintAsFallback = true; // Will fallback to file system to recalculate cas if cas content does not exist anymore
diff --git a/Engine/Source/Programs/UnrealBuildAccelerator/Core/Private/UbaEvent.cpp b/Engine/Source/Programs/UnrealBuildAccelerator/Core/Private/UbaEvent.cpp
--- a/Engine/Source/Programs/UnrealBuildAccelerator/Core/Private/UbaEvent.cpp
+++ b/Engine/Source/Programs/UnrealBuildAccelerator/Core/Private/UbaEvent.cpp
@@ -165,6 +165,11 @@ namespace uba
 			return ts;
 		}
 
+		bool IsCreated()
+		{
+			return m_initialized;
+		}
+
 		bool IsSet(u32 timeoutMs = ~0u)
 		{
 			if (!m_initialized)
@@ -341,6 +346,15 @@ namespace uba
 		#endif
 	}
 
+	bool Event::IsCreated()
+	{
+		#if PLATFORM_WINDOWS
+		return m_ev != nullptr;
+		#else
+		return ((EventImpl&)m_data).IsCreated();
+		#endif
+	}
+
 	bool Event::IsSet(u32 timeOutMs)
 	{
 		#if PLATFORM_WINDOWS
diff --git a/Engine/Source/Programs/UnrealBuildAccelerator/Core/Public/UbaEvent.h b/Engine/Source/Programs/UnrealBuildAccelerator/Core/Public/UbaEvent.h
--- a/Engine/Source/Programs/UnrealBuildAccelerator/Core/Public/UbaEvent.h
+++ b/Engine/Source/Programs/UnrealBuildAccelerator/Core/Public/UbaEvent.h
@@ -16,6 +16,7 @@ namespace uba
 		void Destroy();
 		void Set();
 		void Reset();
+		bool IsCreated();
 		bool IsSet(u32 timeOutMs = ~0u);
 		void* GetHandle();
 	private:
