[UBA]
* Fixed bug in cache server where common string segments were not added when doing maintenance

diff --git a/Engine/Source/Programs/UnrealBuildAccelerator/Common/Private/UbaCacheServer.cpp b/Engine/Source/Programs/UnrealBuildAccelerator/Common/Private/UbaCacheServer.cpp
--- a/Engine/Source/Programs/UnrealBuildAccelerator/Common/Private/UbaCacheServer.cpp
+++ b/Engine/Source/Programs/UnrealBuildAccelerator/Common/Private/UbaCacheServer.cpp
@@ -1141,6 +1141,8 @@ namespace uba
 				u64 pathReserveCount = pathTable.GetPathCount() + 1000; // Add extra reservation space for growth
 				u64 segmentReserveCount = pathTable.GetSegmentCount() + 1000;  // Add extra reservation space for growth
 				CompactPathTable newPathTable(pathTable.GetCaseInsensitive(), pathReserveCount, segmentReserveCount, pathTable.GetVersion());
+				newPathTable.AddCommonStringSegments();
+
 				oldToNewPathOffset.Init(memoryBlock, usedPathOffsets.CountSetBits(), TC("OldToNewPathOffset"));
 
 				StringBuffer<> temp;
