[UBA]
* Added SessionServer_Create2 which does not need SessionServerCreateInfo but instead use config object. This makes it possible to set things but less coupled.
* Added so it is possible to register client side config to network server which is automatically downloaded by agents

diff --git a/Engine/Source/Programs/UnrealBuildAccelerator/Agent/Private/UbaAgent.cpp b/Engine/Source/Programs/UnrealBuildAccelerator/Agent/Private/UbaAgent.cpp
--- a/Engine/Source/Programs/UnrealBuildAccelerator/Agent/Private/UbaAgent.cpp
+++ b/Engine/Source/Programs/UnrealBuildAccelerator/Agent/Private/UbaAgent.cpp
@@ -1052,6 +1052,8 @@ namespace uba
 				return true;
 			}
 
+			if (!client->FetchConfig(config))
+				continue;
 
 			Event wakeupSessionWait(false);
 			Atomic<u32> targetConnectionCount = 1;
diff --git a/Engine/Source/Programs/UnrealBuildAccelerator/Common/Private/UbaConfig.cpp b/Engine/Source/Programs/UnrealBuildAccelerator/Common/Private/UbaConfig.cpp
--- a/Engine/Source/Programs/UnrealBuildAccelerator/Common/Private/UbaConfig.cpp
+++ b/Engine/Source/Programs/UnrealBuildAccelerator/Common/Private/UbaConfig.cpp
@@ -32,6 +32,15 @@ namespace uba
 		return buf.Parse(out);
 	}
 
+	bool ConfigTable::GetValueAsU64(u64& out, const tchar* key) const
+	{
+		auto findIt = m_values.find(key);
+		if (findIt == m_values.end())
+			return m_parent ? m_parent->GetValueAsU64(out, key) : false;
+		StringBuffer<> buf(findIt->second.string);
+		return buf.Parse(out);
+	}
+
 	bool ConfigTable::GetValueAsInt(int& out, const tchar* key) const
 	{
 		auto findIt = m_values.find(key);
@@ -73,6 +82,11 @@ namespace uba
 		return &findIt->second;
 	}
 
+	ConfigTable& ConfigTable::AddTable(const tchar* name)
+	{
+		return m_tables.try_emplace(name).first->second;
+	}
+
 	void ConfigTable::AddValue(const tchar* key, int value)
 	{
 		tchar buf[256];
@@ -87,6 +101,13 @@ namespace uba
 		m_values[key] = Value{ValueType_Value, buf};
 	}
 
+	void ConfigTable::AddValue(const tchar* key, u64 value)
+	{
+		tchar buf[256];
+		TSprintf_s(buf, sizeof_array(buf), TC("%llu"), value);
+		m_values[key] = Value{ValueType_Value, buf};
+	}
+
 	void ConfigTable::AddValue(const tchar* key, bool value)
 	{
 		m_values[key] = Value{ValueType_Value, value ? TC("true") : TC("false")};
@@ -97,6 +118,41 @@ namespace uba
 		m_values[key] = Value{ValueType_String, str};
 	}
 
+	bool ConfigTable::SaveToText(Logger& logger, Vector<char>& outText) const
+	{
+		char line[1024];
+		for (auto& kv : m_values)
+		{
+			const char* quote = kv.second.type == ValueType_String ? "\"" : "";
+			#if PLATFORM_WINDOWS
+			int written = sprintf_s(line, sizeof_array(line), "%S = %s%S%s\r\n", kv.first.c_str(), quote, kv.second.string.c_str(), quote);
+			#else
+			int written = snprintf(line, sizeof_array(line), "%s = %s%s%s\r\n", kv.first.c_str(), quote, kv.second.string.c_str(), quote);
+			#endif
+			
+			u64 size = outText.size();
+			outText.resize(size + written);
+			memcpy(outText.data() + size, line, written);
+		}
+
+		for (auto& kv : m_tables)
+		{
+			#if PLATFORM_WINDOWS
+			int written = sprintf_s(line, sizeof_array(line), "[%S]\r\n", kv.first.c_str());
+			#else
+			int written = snprintf(line, sizeof_array(line), "[%s]\r\n", kv.first.c_str());
+			#endif
+
+			u64 size = outText.size();
+			outText.resize(size + written);
+			memcpy(outText.data() + size, line, written);
+
+			kv.second.SaveToText(logger, outText);
+		}
+
+		return true;
+	}
+
 	bool Config::LoadFromFile(Logger& logger, const tchar* configFile)
 	{
 		m_isLoaded = true;
@@ -125,10 +181,8 @@ namespace uba
 		return LoadFromText(logger, (const char*)fa.GetData(), fa.GetSize());
 	}
 
-	bool Config::LoadFromText(Logger& logger, const char* text, u64 textLen)
+	bool ConfigTable::LoadFromText(Logger& logger, const char* text, u64 textLen)
 	{
-		m_isLoaded = true;
-
 		const char* i = text;
 		const char* e = i + textLen;
 
@@ -244,10 +298,8 @@ namespace uba
 		return m_isLoaded;
 	}
 
-	bool Config::SaveToFile(Logger& logger, const tchar* configFile)
+	bool Config::SaveToFile(Logger& logger, const tchar* configFile) const
 	{
-		UBA_ASSERT(m_tables.empty()); // TODO
-
 		StringBuffer<> dir;
 		dir.AppendDir(configFile);
 		if (!DirectoryCache().CreateDirectory(logger, dir.data))
@@ -257,19 +309,12 @@ namespace uba
 		if (!fa.CreateWrite())
 			return false;
 
-		for (auto& kv : m_values)
-		{
-			char line[1024];
-			const char* quote = kv.second.type == ValueType_String ? "\"" : "";
-			#if PLATFORM_WINDOWS
-			int written = sprintf_s(line, sizeof_array(line), "%S = %s%S%s\r\n", kv.first.c_str(), quote, kv.second.string.c_str(), quote);
-			#else
-			int written = snprintf(line, sizeof_array(line), "%s = %s%s%s\r\n", kv.first.c_str(), quote, kv.second.string.c_str(), quote);
-			#endif
-			if (!fa.Write(line, written))
-				return false;
-		}
-		
+		Vector<char> text;
+		if (!SaveToText(logger, text))
+			return false;
+		if (!fa.Write(text.data(), text.size()))
+			return false;
+
 		return fa.Close();
 	}
 }
diff --git a/Engine/Source/Programs/UnrealBuildAccelerator/Common/Private/UbaNetworkClient.cpp b/Engine/Source/Programs/UnrealBuildAccelerator/Common/Private/UbaNetworkClient.cpp
--- a/Engine/Source/Programs/UnrealBuildAccelerator/Common/Private/UbaNetworkClient.cpp
+++ b/Engine/Source/Programs/UnrealBuildAccelerator/Common/Private/UbaNetworkClient.cpp
@@ -338,6 +338,18 @@ namespace uba
 		return msg.Send();
 	}
 
+	bool NetworkClient::FetchConfig(Config& config)
+	{
+		StackBinaryWriter<64> writer;
+		NetworkMessage msg(*this, SystemServiceId, SystemMessageType_FetchConfig, writer);
+		writer.WriteByte(0); // Need to have a body
+		StackBinaryReader<SendMaxSize> reader;
+		if (!msg.Send(reader))
+			return false;
+		u64 textLen = reader.Read7BitEncoded();
+		return config.LoadFromText(m_logger, (const char*)reader.GetPositionData(), textLen);
+	}
+
 	bool NetworkClient::IsConnected(u32 waitTimeoutMs)
 	{
 		return m_isConnected.IsSet(waitTimeoutMs);
diff --git a/Engine/Source/Programs/UnrealBuildAccelerator/Common/Private/UbaNetworkServer.cpp b/Engine/Source/Programs/UnrealBuildAccelerator/Common/Private/UbaNetworkServer.cpp
--- a/Engine/Source/Programs/UnrealBuildAccelerator/Common/Private/UbaNetworkServer.cpp
+++ b/Engine/Source/Programs/UnrealBuildAccelerator/Common/Private/UbaNetworkServer.cpp
@@ -1,6 +1,7 @@
 // Copyright Epic Games, Inc. All Rights Reserved.
 
 #include "UbaNetworkServer.h"
+#include "UbaConfig.h"
 #include "UbaCrypto.h"
 #include "UbaBinaryReaderWriter.h"
 #include "UbaPlatform.h"
@@ -713,6 +714,11 @@ namespace uba
 		return true;
 	}
 
+	void NetworkServer::SetClientsConfig(const Config& config)
+	{
+		config.SaveToText(m_logger, m_clientsConfig);
+	}
+
 	bool NetworkServer::AddClient(NetworkBackend& backend, const tchar* ip, u16 port, const u8* cryptoKey128)
 	{
 		SCOPED_WRITE_LOCK(m_addConnectionsLock, lock);
@@ -1220,6 +1226,12 @@ namespace uba
 				// No-op
 				return true;
 			}
+			case SystemMessageType_FetchConfig:
+			{
+				writer.Write7BitEncoded(m_clientsConfig.size());
+				writer.WriteBytes(m_clientsConfig.data(), m_clientsConfig.size());
+				return true;
+			}
 		}
 		return false;
 	}
diff --git a/Engine/Source/Programs/UnrealBuildAccelerator/Common/Private/UbaStorage.cpp b/Engine/Source/Programs/UnrealBuildAccelerator/Common/Private/UbaStorage.cpp
--- a/Engine/Source/Programs/UnrealBuildAccelerator/Common/Private/UbaStorage.cpp
+++ b/Engine/Source/Programs/UnrealBuildAccelerator/Common/Private/UbaStorage.cpp
@@ -78,6 +78,7 @@ namespace uba
 		const ConfigTable& table = *tablePtr;
 		table.GetValueAsString(rootDir, TC("RootDir"));
 		table.GetValueAsBool(asyncUnmapViewOfFile, TC("AsyncUnmapViewOfFile"));
+		table.GetValueAsU64(casCapacityBytes, TC("CasCapacityBytes"));
 	}
 
 	struct FileMappingScope
diff --git a/Engine/Source/Programs/UnrealBuildAccelerator/Common/Public/UbaConfig.h b/Engine/Source/Programs/UnrealBuildAccelerator/Common/Public/UbaConfig.h
--- a/Engine/Source/Programs/UnrealBuildAccelerator/Common/Public/UbaConfig.h
+++ b/Engine/Source/Programs/UnrealBuildAccelerator/Common/Public/UbaConfig.h
@@ -16,16 +16,23 @@ namespace uba
 		bool GetValueAsString(const tchar*& out, const tchar* key) const;
 		bool GetValueAsString(TString& out, const tchar* key) const;
 		bool GetValueAsU32(u32& out, const tchar* key) const;
+		bool GetValueAsU64(u64& out, const tchar* key) const;
 		bool GetValueAsInt(int& out, const tchar* key) const;
 		bool GetValueAsBool(bool& out, const tchar* key) const;
 
 		const ConfigTable* GetTable(const tchar* name) const;
 
+		ConfigTable& AddTable(const tchar* name);
+
 		void AddValue(const tchar* key, int value);
 		void AddValue(const tchar* key, u32 value);
+		void AddValue(const tchar* key, u64 value);
 		void AddValue(const tchar* key, bool value);
 		void AddValue(const tchar* key, const tchar* str);
 
+		bool SaveToText(Logger& logger, Vector<char>& outText) const;
+		bool LoadFromText(Logger& logger, const char* text, u64 textLen);
+
 	private:
 		ConfigTable* m_parent = nullptr;
 		enum ValueType { ValueType_Value, ValueType_String };
@@ -40,10 +47,9 @@ namespace uba
 	{
 	public:
 		bool LoadFromFile(Logger& logger, const tchar* configFile);
-		bool LoadFromText(Logger& logger, const char* text, u64 textLen);
 		bool IsLoaded() const;
 
-		bool SaveToFile(Logger& logger, const tchar* configFile);
+		bool SaveToFile(Logger& logger, const tchar* configFile) const;
 
 		bool m_isLoaded = false;
 	};
diff --git a/Engine/Source/Programs/UnrealBuildAccelerator/Common/Public/UbaNetwork.h b/Engine/Source/Programs/UnrealBuildAccelerator/Common/Public/UbaNetwork.h
--- a/Engine/Source/Programs/UnrealBuildAccelerator/Common/Public/UbaNetwork.h
+++ b/Engine/Source/Programs/UnrealBuildAccelerator/Common/Public/UbaNetwork.h
@@ -14,7 +14,7 @@ namespace uba
 	static constexpr u8 SessionServiceId = 2;
 	static constexpr u8 CacheServiceId = 3;
 	
-	static constexpr u32 SystemNetworkVersion = 1339;
+	static constexpr u32 SystemNetworkVersion = 1340;
 	static constexpr u32 StorageNetworkVersion = 4;
 	static constexpr u32 SessionNetworkVersion = 42;
 	static constexpr u32 CacheNetworkVersion = 5;
@@ -26,6 +26,7 @@ namespace uba
 	{
 		SystemMessageType_SetConnectionCount,
 		SystemMessageType_KeepAlive,
+		SystemMessageType_FetchConfig,
 	};
 
 	#define UBA_STORAGE_MESSAGES \
diff --git a/Engine/Source/Programs/UnrealBuildAccelerator/Common/Public/UbaNetworkClient.h b/Engine/Source/Programs/UnrealBuildAccelerator/Common/Public/UbaNetworkClient.h
--- a/Engine/Source/Programs/UnrealBuildAccelerator/Common/Public/UbaNetworkClient.h
+++ b/Engine/Source/Programs/UnrealBuildAccelerator/Common/Public/UbaNetworkClient.h
@@ -44,6 +44,8 @@ namespace uba
 		bool SetConnectionCount(u32 count);
 		bool SendKeepAlive();
 
+		bool FetchConfig(Config& config);
+
 		bool IsConnected(u32 waitTimeoutMs = 0);
 		bool IsOrWasConnected(u32 waitTimeoutMs = 0);
 
diff --git a/Engine/Source/Programs/UnrealBuildAccelerator/Common/Public/UbaNetworkServer.h b/Engine/Source/Programs/UnrealBuildAccelerator/Common/Public/UbaNetworkServer.h
--- a/Engine/Source/Programs/UnrealBuildAccelerator/Common/Public/UbaNetworkServer.h
+++ b/Engine/Source/Programs/UnrealBuildAccelerator/Common/Public/UbaNetworkServer.h
@@ -57,6 +57,7 @@ namespace uba
 		void DisconnectClients();	// Disconnect all active connections
 
 		bool RegisterCryptoKey(const u8* cryptoKey128, u64 expirationTime = ~u64(0));
+		void SetClientsConfig(const Config& config);
 
 		bool AddClient(NetworkBackend& backend, const tchar* ip, u16 port = DefaultPort, const u8* cryptoKey128 = nullptr); // Adds a client that server will create one or more connections to (note this will return before we know if it was a success or not)
 
@@ -165,6 +166,8 @@ namespace uba
 		u32 m_connectionIdCounter = 1;
 		u32 m_maxActiveConnections = 0;
 
+		Vector<char> m_clientsConfig;
+
 		struct Client
 		{
 		public:
diff --git a/Engine/Source/Programs/UnrealBuildAccelerator/Host/Private/UbaExports.cpp b/Engine/Source/Programs/UnrealBuildAccelerator/Host/Private/UbaExports.cpp
--- a/Engine/Source/Programs/UnrealBuildAccelerator/Host/Private/UbaExports.cpp
+++ b/Engine/Source/Programs/UnrealBuildAccelerator/Host/Private/UbaExports.cpp
@@ -136,6 +136,51 @@ extern "C"
 		return true;
 	}
 
+	uba::Config& Config_Create()
+	{
+		return *new uba::Config();
+	}
+
+	void Config_Destroy(uba::Config& config)
+	{
+		delete &config;
+	}
+
+	uba::ConfigTable& Config_RootTable(uba::Config& config)
+	{
+		return config;
+	}
+
+	uba::ConfigTable& Config_AddTable(uba::Config& config, const uba::tchar* name)
+	{
+		return config.AddTable(name);
+	}
+
+	void ConfigTable_AddValueInt(uba::ConfigTable& table, const uba::tchar* key, int value)
+	{
+		table.AddValue(key, value);
+	}
+
+	void ConfigTable_AddValueU32(uba::ConfigTable& table, const uba::tchar* key, uba::u32 value)
+	{
+		table.AddValue(key, value);
+	}
+
+	void ConfigTable_AddValueU64(uba::ConfigTable& table, const uba::tchar* key, uba::u64 value)
+	{
+		table.AddValue(key, value);
+	}
+
+	void ConfigTable_AddValueBool(uba::ConfigTable& table, const uba::tchar* key, bool value)
+	{
+		table.AddValue(key, value);
+	}
+
+	void ConfigTable_AddValueString(uba::ConfigTable& table, const uba::tchar* key, const uba::tchar* str)
+	{
+		table.AddValue(key, str);
+	}
+
 	uba::NetworkServer* NetworkServer_Create(uba::LogWriter& writer, uba::u32 workerCount, uba::u32 sendSize, uba::u32 receiveTimeoutSeconds, bool useQuic)
 	{
 		using namespace uba;
@@ -209,6 +254,10 @@ extern "C"
 		server->DisconnectClients();
 	}
 
+	void NetworkServer_SetClientsConfig(uba::NetworkServer* server, const uba::Config& config)
+	{
+		server->SetClientsConfig(config);
+	}
 
 	bool NetworkServer_AddClient(uba::NetworkServer* server, const uba::tchar* ip, int port, const uba::tchar* crypto)
 	{
@@ -222,36 +271,49 @@ extern "C"
 		return s->AddClient(*s->backend, ip, u16(port), crypto128);
 	}
 
-	uba::StorageServer* StorageServer_Create(uba::NetworkServer& server, const uba::tchar* rootDir, uba::u64 casCapacityBytes, bool storeCompressed, uba::LogWriter& writer, const uba::tchar* zone)
+	uba::StorageServer* StorageServer_InternalCreate(uba::StorageServerCreateInfo& info)
 	{
 		using namespace uba;
-
-		StorageServerCreateInfo info(server, rootDir, writer);
-		info.Apply(GetConfig());
-
 		#if UBA_USE_AWS
 		StringBuffer<> fixedRootDir;
 		fixedRootDir.count = GetFullPathNameW(info.rootDir, fixedRootDir.capacity, fixedRootDir.data, NULL);
 		fixedRootDir.Replace('/', PathSeparator).EnsureEndsWithSlash();
 		info.rootDir = fixedRootDir.data;
 		AWS aws;
-		if (!zone || !*zone)
+		if (!info.zone || !*info.zone)
 		{
-			LoggerWithWriter logger(writer, TC(""));
+			LoggerWithWriter logger(info.writer, TC(""));
 			if (aws.QueryAvailabilityZone(logger, info.rootDir))
-				zone = aws.GetAvailabilityZone();
+				info.zone = aws.GetAvailabilityZone();
 		}
 		#endif
 
 		StringBuffer<256> zoneTemp;
-		if (!zone || !*zone)
+		if (!info.zone || !*info.zone)
 			if (GetZone(zoneTemp))
-				zone = zoneTemp.data;
+				info.zone = zoneTemp.data;
 
+		return new StorageServer(info);
+	}
+
+	uba::StorageServer* StorageServer_Create(uba::NetworkServer& server, const uba::tchar* rootDir, uba::u64 casCapacityBytes, bool storeCompressed, uba::LogWriter& writer, const uba::tchar* zone)
+	{
+		using namespace uba;
+		StorageServerCreateInfo info(server, rootDir, writer);
+		info.Apply(GetConfig());
+		if (zone && !zone)
+			info.zone = zone;
 		info.casCapacityBytes = casCapacityBytes;
 		info.storeCompressed = storeCompressed;
-		info.zone = zone;
-		return new StorageServer(info);
+		return StorageServer_InternalCreate(info);
+	}
+
+	uba::StorageServer* StorageServer_Create2(uba::NetworkServer& server, const uba::Config& config, uba::LogWriter& writer)
+	{
+		using namespace uba;
+		StorageServerCreateInfo info(server, TC(""), writer);
+		info.Apply(config);
+		return StorageServer_InternalCreate(info);
 	}
 
 	void StorageServer_Destroy(uba::StorageServer* storageServer)
@@ -363,6 +425,12 @@ extern "C"
 	{
 		return new uba::SessionServer(info, environment, environmentSize);
 	}
+	uba::SessionServer* SessionServer_Create2(uba::StorageServer& s, uba::NetworkServer& ns, const uba::Config& c, uba::LogWriter& lw, const uba::u8* environment, uba::u32 environmentSize)
+	{
+		uba::SessionServerCreateInfo info(s, ns, lw);
+		info.Apply(c);
+		return new uba::SessionServer(info, environment, environmentSize);
+	}
 	void SessionServer_SetRemoteProcessAvailable(uba::SessionServer* server, SessionServer_RemoteProcessAvailableCallback* available, void* userData)
 	{
 		server->SetRemoteProcessSlotAvailableEvent([available, userData](bool isCrossArchitecture) { available(userData, isCrossArchitecture); });
@@ -563,6 +631,14 @@ extern "C"
 		return new uba::Scheduler(info);
 	}
 
+	uba::Scheduler* Scheduler_Create2(uba::SessionServer& session, const uba::Config& config)
+	{
+		uba::SchedulerCreateInfo info{session};
+		info.Apply(config);
+		info.processConfigs = &uba::GetConfig(); // TODO: Clone provided config?
+		return new uba::Scheduler(info);
+	}
+
 	void Scheduler_Start(uba::Scheduler* scheduler)
 	{
 		scheduler->Start();
diff --git a/Engine/Source/Programs/UnrealBuildAccelerator/Host/Public/UbaExports.h b/Engine/Source/Programs/UnrealBuildAccelerator/Host/Public/UbaExports.h
--- a/Engine/Source/Programs/UnrealBuildAccelerator/Host/Public/UbaExports.h
+++ b/Engine/Source/Programs/UnrealBuildAccelerator/Host/Public/UbaExports.h
@@ -8,11 +8,12 @@
 namespace uba
 {
 	class CacheClient;
+	class Config;
+	class ConfigTable;
 	class NetworkServer;
 	class Process;
 	class ProcessHandle;
 	class RootPaths;
-	class Session;
 	class SessionServer;
 	class Scheduler;
 	class StorageServer;
@@ -48,16 +49,27 @@ extern "C"
 
 	// Config
 	UBA_API bool Config_Load(const uba::tchar* configFile);
+	UBA_API uba::Config& Config_Create();
+	UBA_API void Config_Destroy(uba::Config& config);
+	UBA_API uba::ConfigTable& Config_RootTable(uba::Config& config);
+	UBA_API uba::ConfigTable& Config_AddTable(uba::Config& config, const uba::tchar* name);
+	UBA_API void ConfigTable_AddValueInt(uba::ConfigTable& table, const uba::tchar* key, int value);
+	UBA_API void ConfigTable_AddValueU32(uba::ConfigTable& table, const uba::tchar* key, uba::u32 value);
+	UBA_API void ConfigTable_AddValueU64(uba::ConfigTable& table, const uba::tchar* key, uba::u64 value);
+	UBA_API void ConfigTable_AddValueBool(uba::ConfigTable& table, const uba::tchar* key, bool value);
+	UBA_API void ConfigTable_AddValueString(uba::ConfigTable& table, const uba::tchar* key, const uba::tchar* str);
 
 	// NetworkServer
 	UBA_API uba::NetworkServer* NetworkServer_Create(uba::LogWriter& writer = uba::g_consoleLogWriter, uba::u32 workerCount = 64, uba::u32 sendSize = uba::SendDefaultSize, uba::u32 receiveTimeoutSeconds = 60, bool useQuic = false);
 	UBA_API void NetworkServer_Destroy(uba::NetworkServer* server);
 	UBA_API bool NetworkServer_StartListen(uba::NetworkServer* server, int port = uba::DefaultPort, const uba::tchar* ip = nullptr, const uba::tchar* crypto = nullptr);
 	UBA_API void NetworkServer_Stop(uba::NetworkServer* server);
+	UBA_API void NetworkServer_SetClientsConfig(uba::NetworkServer* server, const uba::Config& config);
 	UBA_API bool NetworkServer_AddClient(uba::NetworkServer* server, const uba::tchar* ip, int port = uba::DefaultPort, const uba::tchar* crypto = nullptr);
 
 	// StorageServer
 	UBA_API uba::StorageServer* StorageServer_Create(uba::NetworkServer& server, const uba::tchar* rootDir, uba::u64 casCapacityBytes, bool storeCompressed, uba::LogWriter& writer = uba::g_consoleLogWriter, const uba::tchar* zone = TC(""));
+	UBA_API uba::StorageServer* StorageServer_Create2(uba::NetworkServer& server, const uba::Config& config, uba::LogWriter& writer = uba::g_consoleLogWriter);
 	UBA_API void StorageServer_Destroy(uba::StorageServer* storageServer);
 	UBA_API void StorageServer_SaveCasTable(uba::StorageServer* storageServer);
 	UBA_API void StorageServer_RegisterDisallowedPath(uba::StorageServer* storageServer, const uba::tchar* path);
@@ -91,6 +103,7 @@ extern "C"
 	UBA_API void SessionServerCreateInfo_Destroy(uba::SessionServerCreateInfo* info);
 
 	UBA_API uba::SessionServer* SessionServer_Create(const uba::SessionServerCreateInfo& info, const uba::u8* environment = nullptr, uba::u32 environmentSize = 0);
+	UBA_API uba::SessionServer* SessionServer_Create2(uba::StorageServer& s, uba::NetworkServer& ns, const uba::Config& c, uba::LogWriter& lw = uba::g_consoleLogWriter, const uba::u8* environment = nullptr, uba::u32 environmentSize = 0);
 	UBA_API void SessionServer_SetRemoteProcessAvailable(uba::SessionServer* server, SessionServer_RemoteProcessAvailableCallback* available, void* userData);
 	UBA_API void SessionServer_SetRemoteProcessReturned(uba::SessionServer* server, SessionServer_RemoteProcessReturnedCallback* returned, void* userData);
 	UBA_API bool SessionServer_RefreshDirectory(uba::SessionServer* server, const uba::tchar* directory);
@@ -119,6 +132,7 @@ extern "C"
 
 	// Scheduler
 	UBA_API uba::Scheduler* Scheduler_Create(uba::SessionServer* session, uba::u32 maxLocalProcessors = ~0u, bool enableProcessReuse = false);
+	UBA_API uba::Scheduler* Scheduler_Create2(uba::SessionServer& session, const uba::Config& config);
 	UBA_API void Scheduler_Start(uba::Scheduler* scheduler);
 	UBA_API uba::u32 Scheduler_EnqueueProcess(uba::Scheduler* scheduler, const uba::ProcessStartInfo& info, float weight = 1.0f, const void* knownInputs = nullptr, uba::u32 knownInputsBytes = 0, uba::u32 knownInputsCount = 0);
 	UBA_API void Scheduler_SetMaxLocalProcessors(uba::Scheduler* scheduler, uba::u32 maxLocalProcessors);
diff --git a/Engine/Source/Programs/UnrealBuildAccelerator/Test/Private/UbaTestAll.h b/Engine/Source/Programs/UnrealBuildAccelerator/Test/Private/UbaTestAll.h
--- a/Engine/Source/Programs/UnrealBuildAccelerator/Test/Private/UbaTestAll.h
+++ b/Engine/Source/Programs/UnrealBuildAccelerator/Test/Private/UbaTestAll.h
@@ -50,7 +50,8 @@ namespace uba
 		UBA_TEST(TestCompactCasKeyTable) \
 		UBA_TEST(TestCacheEntry) \
 		UBA_TEST(TestHashTable) \
-		UBA_TEST(TestConfig) \
+		UBA_TEST(TestLoadConfig) \
+		UBA_TEST(TestSaveConfig) \
 		UBA_TEST(TestCrypto) \
 		UBA_TEST(TestBinDependencies) \
 		UBA_TEST(TestRootPaths) \
diff --git a/Engine/Source/Programs/UnrealBuildAccelerator/Test/Private/UbaTestConfig.h b/Engine/Source/Programs/UnrealBuildAccelerator/Test/Private/UbaTestConfig.h
--- a/Engine/Source/Programs/UnrealBuildAccelerator/Test/Private/UbaTestConfig.h
+++ b/Engine/Source/Programs/UnrealBuildAccelerator/Test/Private/UbaTestConfig.h
@@ -7,7 +7,7 @@
 
 namespace uba
 {
-	bool TestConfig(Logger& logger, const StringBufferBase& rootDir)
+	bool TestLoadConfig(Logger& logger, const StringBufferBase& rootDir)
 	{
 		static const char* configText =
 			"RootDir = \"e:\\foo\"\r\n"
@@ -38,4 +38,30 @@ namespace uba
 			return false;
 		return true;
 	}
+
+	bool TestSaveConfig(Logger& logger, const StringBufferBase& rootDir)
+	{
+		Vector<char> data;
+		{
+			Config config;
+			ConfigTable& table = config.AddTable(TC("TestTable"));
+			table.AddValue(TC("Foo"), 42);
+			config.SaveToText(logger, data);
+		}
+
+		Config config;
+		if (!config.LoadFromText(logger, data.data(), data.size()))
+			return false;
+
+		const ConfigTable* tablePtr = config.GetTable(TC("TestTable"));
+		if (!tablePtr)
+			return false;
+		const ConfigTable& table = *tablePtr;
+		int foo = 0;
+		if (!table.GetValueAsInt(foo, TC("Foo")))
+			return false;
+		if (foo != 42)
+			return false;
+		return true;
+	}
 }
\ No newline at end of file
diff --git a/Engine/Source/Programs/UnrealBuildAccelerator/Test/Private/UbaTestSession.h b/Engine/Source/Programs/UnrealBuildAccelerator/Test/Private/UbaTestSession.h
--- a/Engine/Source/Programs/UnrealBuildAccelerator/Test/Private/UbaTestSession.h
+++ b/Engine/Source/Programs/UnrealBuildAccelerator/Test/Private/UbaTestSession.h
@@ -86,11 +86,26 @@ namespace uba
 
 		auto sg = MakeGuard([&]() { server.DisconnectClients(); });
 
+		Config clientConfig;
+		clientConfig.AddTable(TC("Storage")).AddValue(TC("CheckExistsOnServer"), true);
+		server.SetClientsConfig(clientConfig);
+
+		u16 port = 1356;
+		if (!server.StartListen(tcpBackend, port))
+			return logger.Error(TC("Failed to listen"));
+		if (!client.Connect(tcpBackend, TC("127.0.0.1"), port))
+			return logger.Error(TC("Failed to connect"));
+
+		Config config;
+		if (!client.FetchConfig(config))
+			return false;
+
 		rootDir.Append(TC("Client"));
 		if (deleteAll && !DeleteAllFiles(logger, rootDir.data))
 			return false;
 
 		StorageClientCreateInfo storageClientInfo(client, rootDir.data);
+		storageClientInfo.Apply(config);
 		auto& storageClient = *new StorageClient(storageClientInfo);
 		auto scg = MakeGuard([&]() { delete &storageClient; });
 
@@ -119,11 +134,6 @@ namespace uba
 		storageClient.Start();
 		sessionClient.Start();
 
-		u16 port = 1356;
-		if (!server.StartListen(tcpBackend, port))
-			return logger.Error(TC("Failed to listen"));
-		if (!client.Connect(tcpBackend, TC("127.0.0.1"), port))
-			return logger.Error(TC("Failed to connect"));
 
 		workingDir.EnsureEndsWithSlash();
 		return testFunc(logger, sessionServer, workingDir.data, [&](const ProcessStartInfo& pi) { return sessionServer.RunProcessRemote(pi); });
