[UBA]
* Reverted some of the changes in StorageClient related to local files since it broke prepopulation somehow.

diff --git a/Engine/Source/Programs/UnrealBuildAccelerator/Common/Private/UbaStorageClient.cpp b/Engine/Source/Programs/UnrealBuildAccelerator/Common/Private/UbaStorageClient.cpp
--- a/Engine/Source/Programs/UnrealBuildAccelerator/Common/Private/UbaStorageClient.cpp
+++ b/Engine/Source/Programs/UnrealBuildAccelerator/Common/Private/UbaStorageClient.cpp
@@ -148,18 +148,11 @@ namespace uba
 		auto findIt = m_localStorageFiles.find(AsCompressed(casKey, false));
 		if (findIt != m_localStorageFiles.end())
 		{
-			const LocalFile& localFile = findIt->second;
-			if (localFile.casEntry.mappingHandle.IsValid())
-			{
-				GetMappingString(out, localFile.casEntry.mappingHandle, 0);
-				return true;
-			}
-
-			if (!localFile.fileName.empty())
-			{
-				out.Append(localFile.fileName);
-				return true;
-			}
+			if (findIt->second.casEntry.mappingHandle.IsValid())
+				GetMappingString(out, findIt->second.casEntry.mappingHandle, 0);
+			else
+				out.Append(findIt->second.fileName);
+			return true;
 		}
 		tempLock.Leave();
 
@@ -176,8 +169,6 @@ namespace uba
 			return StorageImpl::MapView(casKey, hint);
 
 		LocalFile& file = findIt->second;
-		if (!file.casEntry.mappingHandle.IsValid())
-			return StorageImpl::MapView(casKey, hint);
 
 		MappedView view;
 		view.handle = file.casEntry.mappingHandle;
@@ -210,16 +201,13 @@ namespace uba
 		auto findIt = m_localStorageFiles.find(AsCompressed(casKey, false));
 		if (findIt != m_localStorageFiles.end())
 		{
+			out.casKey = findIt->first;
 			LocalFile& mf = findIt->second;
-			if (mf.casEntry.mappingHandle.IsValid())
-			{
-				out.casKey = findIt->first;
-				out.size = mf.casEntry.size;
-				out.view.handle = mf.casEntry.mappingHandle;
-				out.view.size = mf.casEntry.size;
-				out.view.isCompressed = false;
-				return true;
-			}
+			out.size = mf.casEntry.size;
+			out.view.handle = mf.casEntry.mappingHandle;
+			out.view.size = mf.casEntry.size;
+			out.view.isCompressed = false;
+			return true;
 		}
 		tempLock.Leave();
 
