[UBA]
* Fixed linux compile errors

diff --git a/Engine/Source/Programs/UnrealBuildAccelerator/Host/Private/UbaExports.cpp b/Engine/Source/Programs/UnrealBuildAccelerator/Host/Private/UbaExports.cpp
--- a/Engine/Source/Programs/UnrealBuildAccelerator/Host/Private/UbaExports.cpp
+++ b/Engine/Source/Programs/UnrealBuildAccelerator/Host/Private/UbaExports.cpp
@@ -136,24 +136,24 @@ extern "C"
 		return true;
 	}
 
-	uba::Config& Config_Create()
+	uba::Config* Config_Create()
 	{
-		return *new uba::Config();
+		return new uba::Config();
 	}
 
-	void Config_Destroy(uba::Config& config)
+	void Config_Destroy(uba::Config* config)
 	{
 		delete &config;
 	}
 
-	uba::ConfigTable& Config_RootTable(uba::Config& config)
+	uba::ConfigTable* Config_RootTable(uba::Config& config)
 	{
-		return config;
+		return &config;
 	}
 
-	uba::ConfigTable& Config_AddTable(uba::Config& config, const uba::tchar* name)
+	uba::ConfigTable* Config_AddTable(uba::Config& config, const uba::tchar* name)
 	{
-		return config.AddTable(name);
+		return &config.AddTable(name);
 	}
 
 	void ConfigTable_AddValueInt(uba::ConfigTable& table, const uba::tchar* key, int value)
diff --git a/Engine/Source/Programs/UnrealBuildAccelerator/Host/Public/UbaExports.h b/Engine/Source/Programs/UnrealBuildAccelerator/Host/Public/UbaExports.h
--- a/Engine/Source/Programs/UnrealBuildAccelerator/Host/Public/UbaExports.h
+++ b/Engine/Source/Programs/UnrealBuildAccelerator/Host/Public/UbaExports.h
@@ -49,10 +49,10 @@ extern "C"
 
 	// Config
 	UBA_API bool Config_Load(const uba::tchar* configFile);
-	UBA_API uba::Config& Config_Create();
-	UBA_API void Config_Destroy(uba::Config& config);
-	UBA_API uba::ConfigTable& Config_RootTable(uba::Config& config);
-	UBA_API uba::ConfigTable& Config_AddTable(uba::Config& config, const uba::tchar* name);
+	UBA_API uba::Config* Config_Create();
+	UBA_API void Config_Destroy(uba::Config* config);
+	UBA_API uba::ConfigTable* Config_RootTable(uba::Config& config);
+	UBA_API uba::ConfigTable* Config_AddTable(uba::Config& config, const uba::tchar* name);
 	UBA_API void ConfigTable_AddValueInt(uba::ConfigTable& table, const uba::tchar* key, int value);
 	UBA_API void ConfigTable_AddValueU32(uba::ConfigTable& table, const uba::tchar* key, uba::u32 value);
 	UBA_API void ConfigTable_AddValueU64(uba::ConfigTable& table, const uba::tchar* key, uba::u64 value);
