[UBA]
* Reverted reverted changelist and made a small modification to hopefully fix prepopulated files (used by osx)

diff --git a/Engine/Source/Programs/UnrealBuildAccelerator/Common/Private/UbaSession.cpp b/Engine/Source/Programs/UnrealBuildAccelerator/Common/Private/UbaSession.cpp
--- a/Engine/Source/Programs/UnrealBuildAccelerator/Common/Private/UbaSession.cpp
+++ b/Engine/Source/Programs/UnrealBuildAccelerator/Common/Private/UbaSession.cpp
@@ -563,7 +563,7 @@ namespace uba
 				isDir = true;
 			}
 			else
-				return m_logger.Error(TC("Failed to open file %s (%u)"), fileName, error);
+				return m_logger.Error(TC("Failed to open file %s (%s)"), fileName, LastErrorToText().data);
 		}
 		auto _ = MakeGuard([&](){ uba::CloseFile(fileName, fileHandle); });
 
diff --git a/Engine/Source/Programs/UnrealBuildAccelerator/Common/Private/UbaStorageClient.cpp b/Engine/Source/Programs/UnrealBuildAccelerator/Common/Private/UbaStorageClient.cpp
--- a/Engine/Source/Programs/UnrealBuildAccelerator/Common/Private/UbaStorageClient.cpp
+++ b/Engine/Source/Programs/UnrealBuildAccelerator/Common/Private/UbaStorageClient.cpp
@@ -148,11 +148,18 @@ namespace uba
 		auto findIt = m_localStorageFiles.find(AsCompressed(casKey, false));
 		if (findIt != m_localStorageFiles.end())
 		{
-			if (findIt->second.casEntry.mappingHandle.IsValid())
-				GetMappingString(out, findIt->second.casEntry.mappingHandle, 0);
-			else
-				out.Append(findIt->second.fileName);
-			return true;
+			const LocalFile& localFile = findIt->second;
+			if (localFile.casEntry.mappingHandle.IsValid())
+			{
+				GetMappingString(out, localFile.casEntry.mappingHandle, 0);
+				return true;
+			}
+
+			if (!localFile.fileName.empty())
+			{
+				out.Append(localFile.fileName);
+				return true;
+			}
 		}
 		tempLock.Leave();
 
@@ -169,6 +176,8 @@ namespace uba
 			return StorageImpl::MapView(casKey, hint);
 
 		LocalFile& file = findIt->second;
+		if (!file.casEntry.mappingHandle.IsValid())
+			return StorageImpl::MapView(casKey, hint);
 
 		MappedView view;
 		view.handle = file.casEntry.mappingHandle;
@@ -201,13 +210,19 @@ namespace uba
 		auto findIt = m_localStorageFiles.find(AsCompressed(casKey, false));
 		if (findIt != m_localStorageFiles.end())
 		{
-			out.casKey = findIt->first;
-			LocalFile& mf = findIt->second;
-			out.size = mf.casEntry.size;
-			out.view.handle = mf.casEntry.mappingHandle;
-			out.view.size = mf.casEntry.size;
-			out.view.isCompressed = false;
-			return true;
+			LocalFile& lf = findIt->second;
+			if (lf.casEntry.exists)
+			{
+				out.casKey = findIt->first;
+				if (lf.casEntry.mappingHandle.IsValid())
+				{
+					out.size = lf.casEntry.size;
+					out.view.handle = lf.casEntry.mappingHandle;
+					out.view.size = lf.casEntry.size;
+					out.view.isCompressed = false;
+				}
+				return true;
+			}
 		}
 		tempLock.Leave();
 
@@ -852,9 +867,10 @@ namespace uba
 				m_logger.Warning(TC("Failed to duplicate handle for file mapping %s (%s)"), fileName, LastErrorToText().data);
 		}
 
-		if (!localFile.casEntry.mappingHandle.IsValid() && localFile.fileName.empty())
+		if (!isPersistentMapping && !localFile.casEntry.mappingHandle.IsValid() && localFile.fileName.empty())
 		{
 			localFile.casEntry.size = fileSize;
+			localFile.casEntry.verified = true;
 			localFile.casEntry.exists = true;
 			localFile.fileName = fileName;
 		}
