[UBA]
* Fixed cpu load calculations for linux

diff --git a/Engine/Source/Programs/UnrealBuildAccelerator/Common/Private/UbaSession.cpp b/Engine/Source/Programs/UnrealBuildAccelerator/Common/Private/UbaSession.cpp
--- a/Engine/Source/Programs/UnrealBuildAccelerator/Common/Private/UbaSession.cpp
+++ b/Engine/Source/Programs/UnrealBuildAccelerator/Common/Private/UbaSession.cpp
@@ -3204,87 +3204,98 @@ namespace uba
 		*lineCount = logWriter.count;
 	}
 
-	float Session::UpdateCpuLoad()
+	bool GetCpuTime(u64& outTotalTime, u64& outIdleTime)
 	{
-		u64 idleTime = 0;
-		u64 totalTime = 0;
-
+		outTotalTime = 0;
+        outIdleTime = 0;
 #if PLATFORM_WINDOWS
 		u64 kernelTime, userTime;
-		if (!GetSystemTimes((FILETIME*)&idleTime, (FILETIME*)&kernelTime, (FILETIME*)&userTime))
-			return m_cpuLoad;
-		totalTime = kernelTime + userTime;
+		if (!GetSystemTimes((FILETIME*)&outIdleTime, (FILETIME*)&kernelTime, (FILETIME*)&userTime))
+			return false;
+		outTotalTime = kernelTime + userTime;
 #elif PLATFORM_LINUX
 		int fd = open("/proc/stat", O_RDONLY | O_CLOEXEC);
-		if (fd != -1)
+		if (fd == -1)
+			return false;
+		char buffer[512];
+		int size = read(fd, buffer, sizeof_array(buffer) - 1);
+		close(fd);
+		if (size == -1)
+			return false;
+		buffer[size] = 0;
+		char* endl = strchr(buffer, '\n');
+		if (!endl)
+			return false;
+		u64 values[16] = { 0 };
+		u32 valueCount = 0;
+		*endl = 0;
+		char* parsePos = buffer;
+
+		// cpu
+		parsePos = strchr(parsePos, ' ');
+		if (!parsePos)
+			return false;
+
+		while (true)
 		{
-			char buffer[512];
-			int size = read(fd, buffer, sizeof_array(buffer) - 1);
-			if (size != -1)
-			{
-				buffer[size] = 0;
-				if (char* endl = strchr(buffer, '\n'))
-				{
-					u64 values[16];
-					u32 valueCount = 0;
-					*endl = 0;
-					char* parsePos = buffer;
-					char* space = nullptr;
-					do
-					{
-						space = strchr(parsePos, ' ');
-						if (space)
-							*space = 0;
-						values[valueCount++] = strtoull(parsePos, nullptr, 10);
-
-						parsePos = space + 1;
-					} while (space);
-
-					// user: normal processes executing in user mode
-					// nice: niced processes executing in user mode
-					// system: processes executing in kernel mode
-					// idle: twiddling thumbs
-					// iowait: waiting for I/O to complete
-					// irq: servicing interrupts
-					// softirq: servicing softirqs
-					if (valueCount > 6)
-					{
-						u64 work = values[0] + values[1] + values[2];
-						idleTime = values[3] + values[4] + values[5] + values[6];
-						totalTime = work + idleTime;
-					}
-				}
-			}
-			close(fd);
+			// remove spaces
+			while (*parsePos && (*parsePos < '0' || *parsePos > '9'))
+				++parsePos;
+			if (!*parsePos)
+				break;
+			char* numberStart = parsePos;
+			while (*parsePos && *parsePos >= '0' && *parsePos <= '9')
+				++parsePos;
+			bool isLast = *parsePos == 0;
+			*parsePos = 0;
+			++parsePos;
+			values[valueCount++] = strtoull(numberStart, nullptr, 10);
+			if (isLast)
+				break;
 		}
-		// TODO: Read "/proc/stat" to get cpu cycles
-#else // PLATFORM_MAC
-        mach_msg_type_number_t  CpuMsgCount = 0;
-        processor_flavor_t CpuInfoType = PROCESSOR_CPU_LOAD_INFO;;
-        natural_t CpuCount = 0;
-        processor_cpu_load_info_t CpuData;
-        host_t host = mach_host_self();
 
-        int res = 0;
-        u64 work = 0;
-        idleTime = 0;
+		// user: normal processes executing in user mode
+		// nice: niced processes executing in user mode
+		// system: processes executing in kernel mode
+		// idle: twiddling thumbs
+		// iowait: waiting for I/O to complete
+		// irq: servicing interrupts
+		// softirq: servicing softirqs
+		// steal
+		if (valueCount <= 6)
+			return false;
+		u64 work = values[0] + values[1] + values[2];
+		outIdleTime = values[3] + values[4] + values[5] + values[6] + values[7];
+		outTotalTime = work + outIdleTime;
 
-        res = host_processor_info(host, CpuInfoType, &CpuCount, (processor_info_array_t *)&CpuData, &CpuMsgCount);
+#else // PLATFORM_MAC
+		mach_msg_type_number_t  CpuMsgCount = 0;
+		processor_flavor_t CpuInfoType = PROCESSOR_CPU_LOAD_INFO;;
+		natural_t CpuCount = 0;
+		processor_cpu_load_info_t CpuData;
+		host_t host = mach_host_self();
+		u64 work = 0;
+		int res = host_processor_info(host, CpuInfoType, &CpuCount, (processor_info_array_t *)&CpuData, &CpuMsgCount);
 		if(res != KERN_SUCCESS)
-        {
-				return m_logger.Error(TC("Kernel error: %s"), mach_error_string(res));
-        }
-
-        for(int i = 0; i < (int)CpuCount; i++)
-        {
-                work += CpuData[i].cpu_ticks[CPU_STATE_SYSTEM];
-                work += CpuData[i].cpu_ticks[CPU_STATE_USER];
-                work += CpuData[i].cpu_ticks[CPU_STATE_NICE];
-                idleTime += CpuData[i].cpu_ticks[CPU_STATE_IDLE];
-        }
-
-		totalTime = work + idleTime;
+			return false;//m_logger.Error(TC("Kernel error: %s"), mach_error_string(res));
+		for(int i = 0; i < (int)CpuCount; i++)
+		{
+			work += CpuData[i].cpu_ticks[CPU_STATE_SYSTEM];
+			work += CpuData[i].cpu_ticks[CPU_STATE_USER];
+			work += CpuData[i].cpu_ticks[CPU_STATE_NICE];
+			outIdleTime += CpuData[i].cpu_ticks[CPU_STATE_IDLE];
+		}
+		outTotalTime = work + outIdleTime;
 #endif
+		return true;
+	}
+
+	float Session::UpdateCpuLoad()
+	{
+		u64 totalTime = 0;
+		u64 idleTime = 0;
+		if (!GetCpuTime(totalTime, idleTime))
+			return m_cpuLoad;
 
 		u64 totalTimeSinceLastTime = totalTime - m_previousTotalCpuTime;
 		u64 idleTimeSinceLastTime = idleTime - m_previousIdleCpuTime;
