[UBA]
* Downdgraded system network version and upgraded session network version instead because I realize this will screw up the cache server

diff --git a/Engine/Source/Programs/UnrealBuildAccelerator/Common/Public/UbaNetwork.h b/Engine/Source/Programs/UnrealBuildAccelerator/Common/Public/UbaNetwork.h
--- a/Engine/Source/Programs/UnrealBuildAccelerator/Common/Public/UbaNetwork.h
+++ b/Engine/Source/Programs/UnrealBuildAccelerator/Common/Public/UbaNetwork.h
@@ -14,9 +14,9 @@ namespace uba
 	static constexpr u8 SessionServiceId = 2;
 	static constexpr u8 CacheServiceId = 3;
 	
-	static constexpr u32 SystemNetworkVersion = 1340;
+	static constexpr u32 SystemNetworkVersion = 1339;
 	static constexpr u32 StorageNetworkVersion = 4;
-	static constexpr u32 SessionNetworkVersion = 42;
+	static constexpr u32 SessionNetworkVersion = 43;
 	static constexpr u32 CacheNetworkVersion = 5;
 	static constexpr u32 CacheBucketVersion = 3;
 
