[UBA]
* Added detouring of FirstFileA and FirstNextFileA since it is used by some upcoming sdk binaries
* Also had to promote GetFullPathNameA, GetModuleFileNameA and GetModuleFileNameExA from debug-only because they were also used and on newest windows they don't seem to fallback to W-version

diff --git a/Engine/Source/Programs/UnrealBuildAccelerator/Detours/Private/Windows/UbaDetoursFunctionsKernelBase.inl b/Engine/Source/Programs/UnrealBuildAccelerator/Detours/Private/Windows/UbaDetoursFunctionsKernelBase.inl
--- a/Engine/Source/Programs/UnrealBuildAccelerator/Detours/Private/Windows/UbaDetoursFunctionsKernelBase.inl
+++ b/Engine/Source/Programs/UnrealBuildAccelerator/Detours/Private/Windows/UbaDetoursFunctionsKernelBase.inl
@@ -1166,16 +1166,38 @@ DWORD Detoured_GetFullPathNameW(LPCWSTR lpFileName, DWORD nBufferLength, LPWSTR
 {
 	DETOURED_CALL(GetFullPathNameW);
 
-	StringBuffer<> temp;
-	FixPath(temp, lpFileName);
-	u64 requiredSize = temp.count + 1;
+	StringBuffer<> fullPath;
+	FixPath(fullPath, lpFileName);
+	u64 requiredSize = fullPath.count + 1;
 	if (nBufferLength < requiredSize)
 		return DWORD(requiredSize);
-	memcpy(lpBuffer, temp.data, requiredSize * 2);
+	memcpy(lpBuffer, fullPath.data, requiredSize * 2);
 	if (lpFilePart)
 		*lpFilePart = wcsrchr(lpBuffer, '\\') + 1;
-	auto res = DWORD(temp.count);
-	DEBUG_LOG_DETOURED(L"GetFullPathNameW", L"%ls TO %ls -> %u", lpFileName, temp.data, res);
+	auto res = DWORD(fullPath.count);
+	DEBUG_LOG_DETOURED(L"GetFullPathNameW", L"%ls TO %ls -> %u", lpFileName, fullPath.data, res);
+	SetLastError(ERROR_SUCCESS);
+	return res;
+}
+
+DWORD Detoured_GetFullPathNameA(LPCSTR lpFileName, DWORD nBufferLength, LPSTR lpBuffer, LPSTR* lpFilePart)
+{
+	// Is verified that this does NOT always call GetFullPathNameW
+	DETOURED_CALL(GetFullPathNameA);
+
+	StringBuffer<> temp;
+	temp.Append(lpFileName);
+
+	StringBuffer<> fullPath;
+	FixPath(fullPath, temp.data);
+	u64 requiredSize = fullPath.count + 1;
+	if (nBufferLength < requiredSize)
+		return DWORD(requiredSize);
+	fullPath.Parse(lpBuffer, requiredSize);
+	if (lpFilePart)
+		*lpFilePart = strrchr(lpBuffer, '\\') + 1;
+	auto res = DWORD(fullPath.count);
+	DEBUG_LOG_DETOURED(L"GetFullPathNameA", L"%ls TO %ls -> %u", temp.data, fullPath.data, res);
 	SetLastError(ERROR_SUCCESS);
 	return res;
 }
@@ -1206,7 +1228,7 @@ BOOL Detoured_GetVolumePathNameW(LPCWSTR lpszFileName, LPWSTR lpszVolumePathName
 	return res;
 }
 
-DWORD Shared_GetModuleFileNameW(HMODULE hModule, const StringView& moduleName, LPWSTR lpFilename, DWORD nSize)
+DWORD Shared_GetModuleFileNameInner(const tchar* func, HMODULE hModule, const StringView& moduleName, LPWSTR lpFilename, DWORD nSize)
 {
 	if (nSize <= moduleName.count)
 	{
@@ -1216,55 +1238,66 @@ DWORD Shared_GetModuleFileNameW(HMODULE hModule, const StringView& moduleName, L
 			lpFilename[nSize - 1] = 0;
 		}
 		SetLastError(ERROR_INSUFFICIENT_BUFFER);
-		DEBUG_LOG_DETOURED(L"GetModuleFileNameW", L"%llu  %u INSUFFICIENT BUFFER (%ls) -> %u", uintptr_t(hModule), nSize, moduleName.data, moduleName.count + 1);
+		DEBUG_LOG_DETOURED(func, L"%llu  %u INSUFFICIENT BUFFER (%ls) -> %u", uintptr_t(hModule), nSize, moduleName.data, moduleName.count + 1);
 		return nSize;
 	}
 
 	memcpy(lpFilename, moduleName.data, moduleName.count * 2);
 	lpFilename[moduleName.count] = 0;
-	DEBUG_LOG_DETOURED(L"GetModuleFileNameW", L"%llu  %u (%ls) -> %u", uintptr_t(hModule), nSize, lpFilename, moduleName.count);
+	DEBUG_LOG_DETOURED(func, L"%llu  %u (%ls) -> %u", uintptr_t(hModule), nSize, lpFilename, moduleName.count);
 	SetLastError(ERROR_SUCCESS);
 	return moduleName.count;
 }
 
-DWORD Detoured_GetModuleFileNameW(HMODULE hModule, LPWSTR lpFilename, DWORD nSize)
+DWORD Shared_GetModuleFileNameW(const tchar* func, HMODULE hModule, LPWSTR lpFilename, DWORD nSize)
 {
-	DETOURED_CALL(GetModuleFileNameW);
-
 	// If null we use the virtual application name
 	if (hModule == NULL)
-		return Shared_GetModuleFileNameW(hModule, g_virtualApplication, lpFilename, nSize);
+		return Shared_GetModuleFileNameInner(func, hModule, g_virtualApplication, lpFilename, nSize);
 
 	{
 		// Check if there are any stored paths from dynamically loaded dlls
 		SCOPED_READ_LOCK(g_loadedModulesLock, lock);
 		auto findIt = g_loadedModules.find(hModule);
 		if (findIt != g_loadedModules.end())
-			return Shared_GetModuleFileNameW(hModule, findIt->second, lpFilename, nSize);
+			return Shared_GetModuleFileNameInner(func, hModule, findIt->second, lpFilename, nSize);
 	}
 
 	if (!g_runningRemote && !IsVfsEnabled())
-		return True_GetModuleFileNameW(hModule, lpFilename, nSize);
+	{
+		DWORD res = True_GetModuleFileNameW(hModule, lpFilename, nSize);
+		DEBUG_LOG_TRUE(L"GetModuleFileNameW", L"%llu  %u (%ls) -> %u", uintptr_t(hModule), nSize, lpFilename, res);
+		return res;
+	}
 
 	StringBuffer<512> moduleName;
-	DWORD res = True_GetModuleFileNameW(hModule, moduleName.data, moduleName.capacity);
-	if (res == 0)
-		return res;
-	UBA_ASSERT(GetLastError() != ERROR_INSUFFICIENT_BUFFER);
-	moduleName.count = res;
+	{
+		DWORD res = True_GetModuleFileNameW(hModule, moduleName.data, moduleName.capacity);
+		DEBUG_LOG_TRUE(L"GetModuleFileNameW", L"%llu  %u (%ls) -> %u", uintptr_t(hModule), moduleName.capacity, moduleName.data, res);
+		if (res == 0)
+			return res;
+		UBA_ASSERT(GetLastError() != ERROR_INSUFFICIENT_BUFFER);
+		moduleName.count = res;
+	}
 
 	// This could be dlls that are loaded early one so might not exist in g_loadedModules
 	// TODO: These could be wrong.. since the files could have been copied from different directories into the remote exedir
 	if (!moduleName.StartsWith(g_exeDir.data))
 	{
 		VirtualizePath(moduleName);
-		return Shared_GetModuleFileNameW(hModule, moduleName, lpFilename, nSize);
+		return Shared_GetModuleFileNameInner(func, hModule, moduleName, lpFilename, nSize);
 	}
 
 	StringBuffer<350> fileName;
 	fileName.Append(g_virtualApplicationDir);
 	fileName.Append(moduleName.data + g_exeDir.count);
-	return Shared_GetModuleFileNameW(hModule, fileName, lpFilename, nSize);
+	return Shared_GetModuleFileNameInner(func, hModule, fileName, lpFilename, nSize);
+}
+
+DWORD Detoured_GetModuleFileNameW(HMODULE hModule, LPWSTR lpFilename, DWORD nSize)
+{
+	DETOURED_CALL(GetModuleFileNameW);
+	return Shared_GetModuleFileNameW(L"GetModuleFileNameW", hModule, lpFilename, nSize);
 }
 
 DWORD Detoured_GetModuleFileNameExW(HANDLE hProcess, HMODULE hModule, LPWSTR lpFilename, DWORD nSize)
@@ -1272,9 +1305,42 @@ DWORD Detoured_GetModuleFileNameExW(HANDLE hProcess, HMODULE hModule, LPWSTR lpF
 	if (hProcess != (HANDLE)-1)
 	{
 		UBA_ASSERT(!g_runningRemote); // Not implemented
-		return True_GetModuleFileNameExW(hProcess, hModule, lpFilename, nSize);
+		DWORD res = True_GetModuleFileNameExW(hProcess, hModule, lpFilename, nSize);
+		DEBUG_LOG_TRUE(L"GetModuleFileNameExW", L"%llu  %u (%ls) -> %u", uintptr_t(hModule), nSize, lpFilename, res);
+		return res;
 	}
-	return Detoured_GetModuleFileNameW(hModule, lpFilename, nSize);
+	return Shared_GetModuleFileNameW(L"GetModuleFileNameExW", hModule, lpFilename, nSize);
+}
+
+DWORD Shared_GetModuleFileNameA(const tchar* func, HMODULE hModule, LPSTR lpFilename, DWORD nSize)
+{
+	// Verified called from used applications.. and does not automatically call W version
+	StringBuffer<> temp;
+	UBA_ASSERT(nSize < temp.capacity);
+	DWORD res = Shared_GetModuleFileNameW(func, hModule, temp.data, nSize);
+	temp.Parse(lpFilename, res);
+	return res;
+}
+
+DWORD Detoured_GetModuleFileNameA(HMODULE hModule, LPSTR lpFilename, DWORD nSize)
+{
+	// Verified called from used applications.. and does not automatically call W version
+	DETOURED_CALL(GetModuleFileNameA);
+	return Shared_GetModuleFileNameA(L"GetModuleFileNameA", hModule, lpFilename, nSize);
+}
+
+DWORD Detoured_GetModuleFileNameExA(HANDLE hProcess, HMODULE hModule, LPSTR lpFilename, DWORD nSize)
+{
+	// Verified called from used applications.. and does not automatically call W version
+	DETOURED_CALL(GetModuleFileNameExA);
+	if (hProcess != (HANDLE)-1)
+	{
+		UBA_ASSERT(!g_runningRemote); // Not implemented
+		DWORD res = True_GetModuleFileNameExA(hProcess, hModule, lpFilename, nSize);
+		DEBUG_LOG_TRUE(L"GetModuleFileNameExA", L"%llu  %u (%hs) -> %u", uintptr_t(hModule), nSize, lpFilename, res);
+		return res;
+	}
+	return Shared_GetModuleFileNameA(L"GetModuleFileNameExA", hModule, lpFilename, nSize);
 }
 
 BOOL Detoured_GetModuleHandleExW(DWORD dwFlags, LPCWSTR lpModuleName, HMODULE *phModule)
@@ -1819,6 +1885,40 @@ BOOL Detoured_FindNextFileW(HANDLE hFindFile, LPWIN32_FIND_DATAW lpFindFileData)
 	return True_FindNextFileW(hFindFile, lpFindFileData);
 }
 
+HANDLE Detoured_FindFirstFileA(LPCSTR lpFileName, LPWIN32_FIND_DATAA lpFindFileData)
+{
+	DETOURED_CALL(FindFirstFileW);
+
+	StringBuffer<> fileName;
+	fileName.Append(lpFileName);
+
+	WIN32_FIND_DATAW findFileData;
+	HANDLE res = Shared_FindFirstFileExW(fileName.data, FindExInfoStandard, &findFileData, FindExSearchNameMatch, NULL, 0, L"FindFirstFileA");
+	if (res == INVALID_HANDLE_VALUE)
+		return res;
+
+	memcpy(lpFindFileData, &findFileData, 48); // 48 is not exact but it is at least down to where the name starts (the types are identical down to the name)
+	size_t destLen;
+	errno_t err = wcstombs_s(&destLen, lpFindFileData->cFileName, MAX_PATH, findFileData.cFileName, MAX_PATH-1);
+	UBA_ASSERT(err == 0);(void)err;
+
+	return res;
+}
+
+BOOL Detoured_FindNextFileA(HANDLE hFindFile, LPWIN32_FIND_DATAA lpFindFileData)
+{
+	WIN32_FIND_DATAW findFileData;
+	if (!Detoured_FindNextFileW(hFindFile, &findFileData))
+		return false;
+
+	memcpy(lpFindFileData, &findFileData, 48); // 48 is not exact but it is at least down to where the name starts (the types are identical down to the name)
+	size_t destLen;
+	errno_t err = wcstombs_s(&destLen, lpFindFileData->cFileName, MAX_PATH, findFileData.cFileName, MAX_PATH-1);
+	UBA_ASSERT(err == 0);(void)err;
+
+	return true;
+}
+
 BOOL Detoured_FindClose(HANDLE handle)
 {
 	DETOURED_CALL(FindClose);
@@ -3296,17 +3396,16 @@ DWORD Detoured_GetLongPathNameA(LPCSTR lpszShortPath, LPSTR lpszLongPath, DWORD
 	return True_GetLongPathNameA(lpszShortPath, lpszLongPath, cchBuffer);
 }
 
-DWORD Detoured_GetFullPathNameA(LPCSTR lpFileName, DWORD nBufferLength, LPSTR lpBuffer, LPSTR* lpFilePart)
+BOOL Detoured_GetVolumePathNameA(LPCSTR lpszFileName, LPSTR lpszVolumePathName, DWORD cchBufferLength)
 {
-	// Is verified that both windows and wine are calling GetFullPathNameW
-	DETOURED_CALL(GetFullPathNameA);
-	DEBUG_LOG_TRUE(L"GetFullPathNameA", L"");
-	return True_GetFullPathNameA(lpFileName, nBufferLength, lpBuffer, lpFilePart);
+	DETOURED_CALL(GetVolumePathNameA);
+	DEBUG_LOG_TRUE(L"GetVolumePathNameA", L"%S", lpszFileName);
+	return True_GetVolumePathNameA(lpszFileName, lpszVolumePathName, cchBufferLength);
 }
 
 DWORD Detoured_GetFileAttributesA(LPCSTR lpFileName)
 {
-	// Is verified that both windows and wine are calling GetFullPathNameW
+	// Is verified that both windows and wine are calling GetFileAttributesW
 	DEBUG_LOG_TRUE(L"GetFileAttributesA", L"");
 	return True_GetFileAttributesA(lpFileName);
 }
@@ -3327,85 +3426,6 @@ HMODULE Detoured_LoadLibraryW(LPCWSTR lpLibFileName)
 }
 #endif
 
-DWORD Shared_GetModuleFileNameA(HMODULE hModule, const wchar_t* moduleName, u32 moduleNameLen, LPSTR lpFilename, DWORD nSize)
-{
-	if (nSize <= moduleNameLen)
-	{
-		if (nSize)
-		{
-			size_t res;
-			errno_t err = wcstombs_s(&res, lpFilename, nSize, moduleName, nSize);
-			if (err)
-				UBA_ASSERTF(false, L"wcstombs_s failed for string '%s' with error code: %u", moduleName, err);
-			lpFilename[nSize - 1] = 0;
-		}
-		SetLastError(ERROR_INSUFFICIENT_BUFFER);
-		DEBUG_LOG_DETOURED(L"GetModuleFileNameExA", L"%llu  %u INSUFFICIENT BUFFER (%s) -> %u", uintptr_t(hModule), nSize, moduleName, moduleNameLen + 1);
-		return nSize;
-	}
-
-	size_t res;
-	errno_t err = wcstombs_s(&res, lpFilename, nSize, moduleName, moduleNameLen);
-	if (err)
-		UBA_ASSERTF(false, L"wcstombs_s failed for string '%s' with error code: %u", moduleName, err);
-	lpFilename[moduleNameLen] = 0;
-	memset(lpFilename + moduleNameLen, 0, nSize - moduleNameLen);
-	DEBUG_LOG_DETOURED(L"GetModuleFileNameExA", L"%llu  %u (%hs) -> %u", uintptr_t(hModule), nSize, lpFilename, moduleNameLen);
-	SetLastError(ERROR_SUCCESS);
-	return moduleNameLen + 1;
-}
-
-DWORD Detoured_GetModuleFileNameExA(HANDLE hProcess, HMODULE hModule, LPSTR lpFilename, DWORD nSize)
-{
-	DETOURED_CALL(GetModuleFileNameExA);
-
-	if (hProcess != (HANDLE)-1)
-		if (GetProcessId(hProcess) != GetCurrentProcessId())
-		{
-			UBA_ASSERTF(!g_runningRemote, L"THIS SHOULD NOT HAPPEN!");
-			return True_GetModuleFileNameExA(hProcess, hModule, lpFilename, nSize);
-		}
-
-	// If null we use the virtual application name
-	if (hModule == NULL)
-		return Shared_GetModuleFileNameA(hModule, g_virtualApplication.data, u32(g_virtualApplication.count), lpFilename, nSize);
-
-	{
-		// Check if there are any stored paths from dynamically loaded dlls
-		SCOPED_READ_LOCK(g_loadedModulesLock, lock);
-		auto findIt = g_loadedModules.find(hModule);
-		if (findIt != g_loadedModules.end())
-			return Shared_GetModuleFileNameA(hModule, findIt->second.c_str(), u32(findIt->second.size()), lpFilename, nSize);
-	}
-
-	if (!g_runningRemote && !IsVfsEnabled())
-		return True_GetModuleFileNameA(hModule, lpFilename, nSize);
-
-	wchar_t moduleName[350];
-	DWORD res = True_GetModuleFileNameW(hModule, moduleName, sizeof_array(moduleName));
-	if (res == 0)
-		return res;
-	UBA_ASSERT(GetLastError() != ERROR_INSUFFICIENT_BUFFER);
-
-	// This could be dlls that are loaded early one so might not exist in g_loadedModules
-	// TODO: These could be wrong.. since the files could have been copied from different directories into the remote exedir
-	if (!StartsWith(moduleName, g_exeDir.data))
-		return Shared_GetModuleFileNameA(hModule, moduleName, res, lpFilename, nSize);
-
-	StringBuffer<350> fileName;
-	fileName.Append(g_virtualApplicationDir);
-	fileName.Append(moduleName + g_exeDir.count);
-	return Shared_GetModuleFileNameA(hModule, fileName.data, u32(fileName.count), lpFilename, nSize);
-}
-
-DWORD Detoured_GetModuleFileNameA(HMODULE hModule, LPSTR lpFilename, DWORD nSize)
-{
-	DETOURED_CALL(GetModuleFileNameA);
-	DEBUG_LOG_TRUE(L"GetModuleFileNameA", L"");
-	//UBA_ASSERT(!g_runningRemote);
-	return Detoured_GetModuleFileNameExA((HANDLE)-1, hModule, lpFilename, nSize);
-}
-
 DWORD Detoured_GetModuleBaseNameA(HANDLE hProcess, HMODULE hModule, LPSTR lpBaseName, DWORD nSize)
 {
 	DETOURED_CALL(GetModuleBaseNameA);
diff --git a/Engine/Source/Programs/UnrealBuildAccelerator/Detours/Private/Windows/UbaDetoursFunctionsWin.h b/Engine/Source/Programs/UnrealBuildAccelerator/Detours/Private/Windows/UbaDetoursFunctionsWin.h
--- a/Engine/Source/Programs/UnrealBuildAccelerator/Detours/Private/Windows/UbaDetoursFunctionsWin.h
+++ b/Engine/Source/Programs/UnrealBuildAccelerator/Detours/Private/Windows/UbaDetoursFunctionsWin.h
@@ -65,9 +65,12 @@
 	DETOURED_FUNCTION(GetFileType) \
 	DETOURED_FUNCTION(GetLongPathNameW) \
 	DETOURED_FUNCTION(GetFullPathNameW) \
+	DETOURED_FUNCTION(GetFullPathNameA) \
 	DETOURED_FUNCTION(GetVolumePathNameW) \
 	DETOURED_FUNCTION(GetModuleFileNameW) \
 	DETOURED_FUNCTION(GetModuleFileNameExW) \
+	DETOURED_FUNCTION(GetModuleFileNameA) \
+	DETOURED_FUNCTION(GetModuleFileNameExA) \
 	DETOURED_FUNCTION(GetModuleHandleExW) \
 	DETOURED_FUNCTION(GetFileAttributesW) \
 	DETOURED_FUNCTION(SetFileAttributesW) \
@@ -81,6 +84,8 @@
 	DETOURED_FUNCTION(FindFirstFileW) \
 	DETOURED_FUNCTION(FindFirstFileExW) \
 	DETOURED_FUNCTION(FindNextFileW) \
+	DETOURED_FUNCTION(FindFirstFileA) \
+	DETOURED_FUNCTION(FindNextFileA) \
 	DETOURED_FUNCTION(FindClose) \
 	DETOURED_FUNCTION(SetFileInformationByHandle) \
 	DETOURED_FUNCTION(CreateFileMappingW) \
@@ -251,12 +256,10 @@
 	DETOURED_FUNCTION(GetUserDefaultLocaleName) \
 	DETOURED_FUNCTION(GetDiskFreeSpaceExA) \
 	DETOURED_FUNCTION(GetLongPathNameA) \
-	DETOURED_FUNCTION(GetFullPathNameA) \
+	DETOURED_FUNCTION(GetVolumePathNameA) \
 	DETOURED_FUNCTION(GetFileAttributesA) \
 	DETOURED_FUNCTION(GetFileAttributesExA) \
 	DETOURED_FUNCTION_X64(LoadLibraryW) \
-	DETOURED_FUNCTION(GetModuleFileNameA) \
-	DETOURED_FUNCTION(GetModuleFileNameExA) \
 	DETOURED_FUNCTION(GetModuleBaseNameA) \
 	DETOURED_FUNCTION(GetModuleBaseNameW) \
 	DETOURED_FUNCTION(SetUnhandledExceptionFilter) \
