[UBA]
* Added option -crash to UbaCacheServer to be able to test crashes

diff --git a/Engine/Source/Programs/UnrealBuildAccelerator/CacheService/Private/UbaCacheService.cpp b/Engine/Source/Programs/UnrealBuildAccelerator/CacheService/Private/UbaCacheService.cpp
--- a/Engine/Source/Programs/UnrealBuildAccelerator/CacheService/Private/UbaCacheService.cpp
+++ b/Engine/Source/Programs/UnrealBuildAccelerator/CacheService/Private/UbaCacheService.cpp
@@ -169,6 +169,13 @@ namespace uba
 			{
 				fullMaintenance = true;
 			}
+			else if (name.Equals(TC("-crash")))
+			{
+				#ifndef __clang_analyzer__
+				int* ptr = nullptr;
+				*ptr = 0;
+				#endif
+			}
 			else if (name.Equals(TC("-nomaintenance")))
 			{
 				maintenanceEnabled = false;
