[UBA]
* Fixes to make static analysis happy (and one was actually a real bug. QueryDosDeviceW got a buffer size in bytes instead of wide character count)

diff --git a/Engine/Source/Programs/UnrealBuildAccelerator/Cli/Private/UbaCli.cpp b/Engine/Source/Programs/UnrealBuildAccelerator/Cli/Private/UbaCli.cpp
--- a/Engine/Source/Programs/UnrealBuildAccelerator/Cli/Private/UbaCli.cpp
+++ b/Engine/Source/Programs/UnrealBuildAccelerator/Cli/Private/UbaCli.cpp
@@ -471,7 +471,7 @@ namespace uba
 			if (!parseArg(argv[i]))
 				return false;
 
-		auto addOption = [&](const tchar* name, const tchar* value) { StringBuffer<512> v(value); return parseOption(ToView(name), v); };
+		auto addOption = [&](const tchar* name, const tchar* value) { StringBuffer<512> v(value); return parseOption(ToView(name), v); };(void)addOption;
 
 		if (false)
 		{
diff --git a/Engine/Source/Programs/UnrealBuildAccelerator/Common/Private/UbaLogger.cpp b/Engine/Source/Programs/UnrealBuildAccelerator/Common/Private/UbaLogger.cpp
--- a/Engine/Source/Programs/UnrealBuildAccelerator/Common/Private/UbaLogger.cpp
+++ b/Engine/Source/Programs/UnrealBuildAccelerator/Common/Private/UbaLogger.cpp
@@ -540,11 +540,8 @@ namespace uba
 			searchPathString.Append(*it);
 		}
 
-		bool symInitialized = false;
-
 		if (SymInitializeW(processHandle, searchPathString.data, FALSE))
 		{
-			symInitialized = true;
 			for (auto& mod : modules)
 			{
 				if (mod.name.empty())
diff --git a/Engine/Source/Programs/UnrealBuildAccelerator/Common/Private/UbaProcessStartInfoHolder.cpp b/Engine/Source/Programs/UnrealBuildAccelerator/Common/Private/UbaProcessStartInfoHolder.cpp
--- a/Engine/Source/Programs/UnrealBuildAccelerator/Common/Private/UbaProcessStartInfoHolder.cpp
+++ b/Engine/Source/Programs/UnrealBuildAccelerator/Common/Private/UbaProcessStartInfoHolder.cpp
@@ -54,10 +54,10 @@ namespace uba
 		while (argsBegin && *argsBegin == ' ')
 			++argsBegin;
 
-		if (*argsBegin == '/') // Unknown switch, don't try to expand cmd
+		if (argsBegin && *argsBegin == '/') // Unknown switch, don't try to expand cmd
 			return false;
 
-		if (*argsBegin == '\"')
+		if (argsBegin && *argsBegin == '\"')
 		{
 			++argsBegin;
 			--argsEnd;
diff --git a/Engine/Source/Programs/UnrealBuildAccelerator/Common/Private/UbaSession.cpp b/Engine/Source/Programs/UnrealBuildAccelerator/Common/Private/UbaSession.cpp
--- a/Engine/Source/Programs/UnrealBuildAccelerator/Common/Private/UbaSession.cpp
+++ b/Engine/Source/Programs/UnrealBuildAccelerator/Common/Private/UbaSession.cpp
@@ -3089,7 +3089,7 @@ namespace uba
 			{
 				DWORD pagefileOnOsVolume = 0;
 				DWORD pagefileOnOsVolumeSize = 4;
-				res = RegGetValueW(HKEY_LOCAL_MACHINE, TC("SYSTEM\\CurrentControlSet\\Control\\Session Manager\\Memory Management"), TC("PagefileOnOsVolume"), RRF_RT_DWORD, NULL, &pagefileOnOsVolume, &pagefileOnOsVolumeSize);
+				RegGetValueW(HKEY_LOCAL_MACHINE, TC("SYSTEM\\CurrentControlSet\\Control\\Session Manager\\Memory Management"), TC("PagefileOnOsVolume"), RRF_RT_DWORD, NULL, &pagefileOnOsVolume, &pagefileOnOsVolumeSize);
 
 				wchar_t* line = str;
 				for (; size_t lineLen = wcslen(line); line += lineLen + 1)
diff --git a/Engine/Source/Programs/UnrealBuildAccelerator/Common/Private/UbaSessionServer.cpp b/Engine/Source/Programs/UnrealBuildAccelerator/Common/Private/UbaSessionServer.cpp
--- a/Engine/Source/Programs/UnrealBuildAccelerator/Common/Private/UbaSessionServer.cpp
+++ b/Engine/Source/Programs/UnrealBuildAccelerator/Common/Private/UbaSessionServer.cpp
@@ -414,7 +414,7 @@ namespace uba
 
 		while (reader.GetPosition() < dataLen)
 		{
-			tchar str[512];
+			tchar str[512] = { 0 };
 			reader.ReadString(str, sizeof_array(str));
 			tchar* path = str;
 
@@ -648,7 +648,7 @@ namespace uba
 		m_logger.Detail(TC("Client session %s connected (Id: %u, Uid: %s%s)"), name.data, connectionInfo.GetId(), GuidToString(connectionInfo.GetUid()).str, (isClientArm ? TC(", IsArm: true") : TC("")));
 
 		CasKey clientAgentKey = reader.ReadCasKey();
-		CasKey clientDetoursKey = reader.ReadCasKey();
+		CasKey clientDetoursKey = reader.ReadCasKey();(void)clientDetoursKey;
 
 		CasKey detoursBinaryKey[2];
 		CasKey& agentBinaryKey = m_agentBinaryKey[isClientArm];
diff --git a/Engine/Source/Programs/UnrealBuildAccelerator/Common/Private/UbaThread.cpp b/Engine/Source/Programs/UnrealBuildAccelerator/Common/Private/UbaThread.cpp
--- a/Engine/Source/Programs/UnrealBuildAccelerator/Common/Private/UbaThread.cpp
+++ b/Engine/Source/Programs/UnrealBuildAccelerator/Common/Private/UbaThread.cpp
@@ -59,6 +59,9 @@ namespace uba
 		m_func = std::move(f);
 #if PLATFORM_WINDOWS
 		m_handle = CreateThread(NULL, 0, [](LPVOID p) -> DWORD { return ((Thread*)p)->m_func(); }, this, 0, NULL);
+		UBA_ASSERT(m_handle);
+		if (!m_handle)
+			return;
 		if (description)
 			SetThreadDescription(m_handle, description);
 		AlternateThreadGroupAffinity(m_handle);
diff --git a/Engine/Source/Programs/UnrealBuildAccelerator/Detours/Private/Windows/UbaDetoursFunctionsNtDll.inl b/Engine/Source/Programs/UnrealBuildAccelerator/Detours/Private/Windows/UbaDetoursFunctionsNtDll.inl
--- a/Engine/Source/Programs/UnrealBuildAccelerator/Detours/Private/Windows/UbaDetoursFunctionsNtDll.inl
+++ b/Engine/Source/Programs/UnrealBuildAccelerator/Detours/Private/Windows/UbaDetoursFunctionsNtDll.inl
@@ -1531,9 +1531,9 @@ NTSTATUS Detoured_NtQueryObject(HANDLE Handle, OBJECT_INFORMATION_CLASS ObjectIn
 				}
 				else
 				{
-					wchar_t drive[3] = { buffer[0], ':', 0};
+					wchar_t drive[3] = { buffer[0], ':', 0 };
 					wchar_t device[256];
-					DWORD deviceLen = QueryDosDeviceW(drive, device, sizeof(device));
+					DWORD deviceLen = QueryDosDeviceW(drive, device, sizeof_array(device));
 					UBA_ASSERT(deviceLen);
 					buffer.Prepend(StringView(device, deviceLen), 2);
 				}
diff --git a/Engine/Source/Programs/UnrealBuildAccelerator/Host/Private/UbaExports.cpp b/Engine/Source/Programs/UnrealBuildAccelerator/Host/Private/UbaExports.cpp
--- a/Engine/Source/Programs/UnrealBuildAccelerator/Host/Private/UbaExports.cpp
+++ b/Engine/Source/Programs/UnrealBuildAccelerator/Host/Private/UbaExports.cpp
@@ -143,7 +143,7 @@ extern "C"
 
 	void Config_Destroy(uba::Config* config)
 	{
-		delete &config;
+		delete config;
 	}
 
 	uba::ConfigTable* Config_RootTable(uba::Config& config)
@@ -274,6 +274,7 @@ extern "C"
 	uba::StorageServer* StorageServer_InternalCreate(uba::StorageServerCreateInfo& info)
 	{
 		using namespace uba;
+		#ifndef __clang_analyzer__
 		#if UBA_USE_AWS
 		StringBuffer<> fixedRootDir;
 		fixedRootDir.count = GetFullPathNameW(info.rootDir, fixedRootDir.capacity, fixedRootDir.data, NULL);
@@ -292,6 +293,7 @@ extern "C"
 		if (!info.zone || !*info.zone)
 			if (GetZone(zoneTemp))
 				info.zone = zoneTemp.data;
+		#endif
 
 		return new StorageServer(info);
 	}
diff --git a/Engine/Source/Programs/UnrealBuildAccelerator/Visualizer/Private/UbaVisualizer.cpp b/Engine/Source/Programs/UnrealBuildAccelerator/Visualizer/Private/UbaVisualizer.cpp
--- a/Engine/Source/Programs/UnrealBuildAccelerator/Visualizer/Private/UbaVisualizer.cpp
+++ b/Engine/Source/Programs/UnrealBuildAccelerator/Visualizer/Private/UbaVisualizer.cpp
@@ -624,9 +624,6 @@ namespace uba
 			exitCloak.Execute();
 
 			// If not child it will not propagate keyboard presses etc to parent
-			LONG style = GetWindowLong(hwnd, GWL_STYLE);
-			style = WS_CHILD;
-			style &= ~WS_POPUP;
 			SetWindowLong(hwnd, GWL_STYLE, WS_CHILD|WS_VSCROLL|WS_HSCROLL);
 
 			m_parentHwnd = (HWND)(uintptr_t)m_config.parent;
