[UBA]
* Fixed so static analysis scripts work when vfs is enabled by default

diff --git a/Engine/Source/Programs/UnrealBuildAccelerator/Build-StaticAnalysis.bat b/Engine/Source/Programs/UnrealBuildAccelerator/Build-StaticAnalysis.bat
--- a/Engine/Source/Programs/UnrealBuildAccelerator/Build-StaticAnalysis.bat
+++ b/Engine/Source/Programs/UnrealBuildAccelerator/Build-StaticAnalysis.bat
@@ -14,7 +14,7 @@ set config=%1
 echo.
 echo === Building Targets ===
 pushd "../../../.."
-call Engine/Build/BatchFiles/RunUBT.bat -NoUba -NoUbaLocal -NoSNDBS -NoXGE -Architectures=x64 ^
+call Engine/Build/BatchFiles/RunUBT.bat -NoUba -NoUbaLocal -NoSNDBS -NoXGE -NoVFS -Architectures=x64 ^
 	-Target="UbaAgent Win64 %config%" ^
 	-Target="UbaCli Win64 %config%" ^
 	-Target="UbaDetours Win64 %config%" ^
diff --git a/Engine/Source/Programs/UnrealBuildAccelerator/Build-StaticAnalysis2.bat b/Engine/Source/Programs/UnrealBuildAccelerator/Build-StaticAnalysis2.bat
new file mode 100755
--- /dev/null
+++ b/Engine/Source/Programs/UnrealBuildAccelerator/Build-StaticAnalysis2.bat
@@ -0,0 +1,36 @@
+@echo off
+setlocal 
+
+pushd "%~dp0"
+
+set config=Development
+
+if "%~1"=="" goto blank
+
+set config=%1
+
+:blank
+
+echo.
+echo === Building Targets ===
+pushd "../../../.."
+call Engine/Build/BatchFiles/RunUBT.bat -NoUba -NoUbaLocal -NoPCH -DisableUnity -AllModules -NoLink -NoVFS -WarningsAsErrors -IncludeHeaders -staticanalyzer=Default ^
+	-Target="UbaAgent Win64 %config%" ^
+	-Target="UbaCli Win64 %config%" ^
+	-Target="UbaDetours Win64 %config%" ^
+	-Target="UbaHost Win64 %config%" ^
+	-Target="UbaStorageProxy Win64 %config%" ^
+	-Target="UbaVisualizer Win64 %config%" ^
+	-Target="UbaCacheService Win64 %config%" ^
+	-Target="UbaObjTool Win64 %config%" ^
+	-Target="UbaTest Win64 %config%"
+
+popd
+
+popd
+
+if %ERRORLEVEL% NEQ 0 exit /b %ERRORLEVEL%
+
+endlocal
+
+exit /b 0
diff --git a/Engine/Source/Programs/UnrealBuildAccelerator/Test/UbaTest.Build.cs b/Engine/Source/Programs/UnrealBuildAccelerator/Test/UbaTest.Build.cs
--- a/Engine/Source/Programs/UnrealBuildAccelerator/Test/UbaTest.Build.cs
+++ b/Engine/Source/Programs/UnrealBuildAccelerator/Test/UbaTest.Build.cs
@@ -8,6 +8,7 @@ public class UbaTest : ModuleRules
 	public UbaTest(ReadOnlyTargetRules Target) : base(Target)
 	{
 		PrivatePCHHeaderFile = "../Core/Public/UbaCorePch.h";
+		IWYUSupport = IWYUSupport.None;
 
 		CppCompileWarningSettings.UnsafeTypeCastWarningLevel = WarningLevel.Error;
 
