[UBA]
* Minor cosmetic cleanup in UbaNetwork.h

diff --git a/Engine/Source/Programs/UnrealBuildAccelerator/Common/Public/UbaNetwork.h b/Engine/Source/Programs/UnrealBuildAccelerator/Common/Public/UbaNetwork.h
--- a/Engine/Source/Programs/UnrealBuildAccelerator/Common/Public/UbaNetwork.h
+++ b/Engine/Source/Programs/UnrealBuildAccelerator/Common/Public/UbaNetwork.h
@@ -22,12 +22,10 @@
 
 	// Messages used over network between client/server (system, storage and session)
 
-	enum SystemMessageType : u8
-	{
-		SystemMessageType_SetConnectionCount,
-		SystemMessageType_KeepAlive,
-		SystemMessageType_FetchConfig,
-	};
+	#define UBA_SYSTEM_MESSAGES \
+		UBA_SYSTEM_MESSAGE(SetConnectionCount) \
+		UBA_SYSTEM_MESSAGE(KeepAlive) \
+		UBA_SYSTEM_MESSAGE(FetchConfig) \
 
 	#define UBA_STORAGE_MESSAGES \
 		UBA_STORAGE_MESSAGE(FetchBegin) \
@@ -41,26 +39,6 @@
 		UBA_STORAGE_MESSAGE(ProxyFetchBegin) \
 		UBA_STORAGE_MESSAGE(ProxyFetchEnd) \
 
-	enum StorageMessageType : u8
-	{
-		#define UBA_STORAGE_MESSAGE(x) StorageMessageType_##x,
-		UBA_STORAGE_MESSAGES
-		#undef UBA_STORAGE_MESSAGE
-	};
-
-	inline StringView ToString(StorageMessageType type)
-	{
-		switch (type)
-		{
-			#define UBA_STORAGE_MESSAGE(x) case StorageMessageType_##x: return AsView(TC("")#x);
-			UBA_STORAGE_MESSAGES
-			#undef UBA_STORAGE_MESSAGE
-		default:
-			return ToView(TC("Unknown")); // Should never happen
-		}
-	}
-
-
 	#define UBA_SESSION_MESSAGES \
 		UBA_SESSION_MESSAGE(Connect) \
 		UBA_SESSION_MESSAGE(EnsureBinaryFile) \
@@ -94,21 +72,6 @@
 		UBA_SESSION_MESSAGE(HostRun) \
 		UBA_SESSION_MESSAGE(GetSymbols) \
 
-	enum SessionMessageType : u8 
-	{
-		#define UBA_SESSION_MESSAGE(x) SessionMessageType_##x,
-		UBA_SESSION_MESSAGES
-		#undef UBA_SESSION_MESSAGE
-	};
-
-	// Response types for SessionMessageType_ProcessAvailable
-	enum SessionProcessAvailableResponse : u32
-	{
-		SessionProcessAvailableResponse_None = 0,
-		SessionProcessAvailableResponse_Disconnect = ~u32(0),
-		SessionProcessAvailableResponse_RemoteExecutionDisabled = ~u32(0) - 1,
-	};
-
 	#define UBA_CACHE_MESSAGES \
 		UBA_CACHE_MESSAGE(Connect) \
 		UBA_CACHE_MESSAGE(StorePathTable) \
@@ -124,6 +87,27 @@
 		UBA_CACHE_MESSAGE(FetchPathTable2) \
 		UBA_CACHE_MESSAGE(FetchCasTable2) \
 
+	enum SystemMessageType : u8
+	{
+		#define UBA_SYSTEM_MESSAGE(x) SystemMessageType_##x,
+		UBA_SYSTEM_MESSAGES
+		#undef UBA_SYSTEM_MESSAGE
+	};
+
+	enum StorageMessageType : u8
+	{
+		#define UBA_STORAGE_MESSAGE(x) StorageMessageType_##x,
+		UBA_STORAGE_MESSAGES
+		#undef UBA_STORAGE_MESSAGE
+	};
+
+	enum SessionMessageType : u8 
+	{
+		#define UBA_SESSION_MESSAGE(x) SessionMessageType_##x,
+		UBA_SESSION_MESSAGES
+		#undef UBA_SESSION_MESSAGE
+	};
+
 	enum CacheMessageType : u8
 	{
 		#define UBA_CACHE_MESSAGE(x) CacheMessageType_##x,
@@ -131,5 +115,25 @@
 		#undef UBA_CACHE_MESSAGE
 	};
 
+	inline StringView ToString(StorageMessageType type)
+	{
+		switch (type)
+		{
+			#define UBA_STORAGE_MESSAGE(x) case StorageMessageType_##x: return AsView(TC("")#x);
+			UBA_STORAGE_MESSAGES
+			#undef UBA_STORAGE_MESSAGE
+		default:
+			return ToView(TC("Unknown")); // Should never happen
+		}
+	}
+
+	// Response types for SessionMessageType_ProcessAvailable
+	enum SessionProcessAvailableResponse : u32
+	{
+		SessionProcessAvailableResponse_None = 0,
+		SessionProcessAvailableResponse_Disconnect = ~u32(0),
+		SessionProcessAvailableResponse_RemoteExecutionDisabled = ~u32(0) - 1,
+	};
+
 	inline constexpr const char EncryptionHandshakeString[] = "This is a test string used to check so encryption keys matches between client and server. This string is 128 characters long...";
 }
