[UBA]
* Added support in export api for creating ProcessStartInfo using config objects
* Changed inside of export api to use ProcessStartInfoHolder instead of ProcessStartInfo

diff --git a/Engine/Source/Programs/UnrealBuildAccelerator/Common/Private/UbaConfig.cpp b/Engine/Source/Programs/UnrealBuildAccelerator/Common/Private/UbaConfig.cpp
--- a/Engine/Source/Programs/UnrealBuildAccelerator/Common/Private/UbaConfig.cpp
+++ b/Engine/Source/Programs/UnrealBuildAccelerator/Common/Private/UbaConfig.cpp
@@ -76,6 +76,8 @@ namespace uba
 
 	const ConfigTable* ConfigTable::GetTable(const tchar* name) const
 	{
+		if (!name || !*name)
+			return this;
 		auto findIt = m_tables.find(name);
 		if (findIt == m_tables.end())
 			return nullptr;
@@ -84,7 +86,11 @@ namespace uba
 
 	ConfigTable& ConfigTable::AddTable(const tchar* name)
 	{
-		return m_tables.try_emplace(name).first->second;
+		if (!name || !*name)
+			return *this;
+		ConfigTable& table = m_tables.try_emplace(name).first->second;
+		table.m_parent = this;
+		return table;
 	}
 
 	void ConfigTable::AddValue(const tchar* key, int value)
diff --git a/Engine/Source/Programs/UnrealBuildAccelerator/Common/Private/UbaProcessStartInfoHolder.cpp b/Engine/Source/Programs/UnrealBuildAccelerator/Common/Private/UbaProcessStartInfoHolder.cpp
--- a/Engine/Source/Programs/UnrealBuildAccelerator/Common/Private/UbaProcessStartInfoHolder.cpp
+++ b/Engine/Source/Programs/UnrealBuildAccelerator/Common/Private/UbaProcessStartInfoHolder.cpp
@@ -1,6 +1,7 @@
 // Copyright Epic Games, Inc. All Rights Reserved.
 
 #include "UbaProcessStartInfoHolder.h"
+#include "UbaConfig.h"
 
 namespace uba
 {
@@ -136,4 +137,31 @@ namespace uba
 	#endif
 		return false;
 	}
+
+	void ProcessStartInfoHolder::Apply(const Config& config, const uba::tchar* configTable)
+	{
+		const ConfigTable* tablePtr = config.GetTable(configTable);
+		if (!tablePtr)
+			return;
+		const ConfigTable& table = *tablePtr;
+		if (table.GetValueAsString(applicationStr, TC("Application")))
+			application = applicationStr.c_str();
+		if (table.GetValueAsString(argumentsStr, TC("Arguments")))
+			arguments = argumentsStr.c_str();
+		if (table.GetValueAsString(workingDirStr, TC("WorkingDir")))
+			workingDir = workingDirStr.c_str();
+		if (table.GetValueAsString(descriptionStr, TC("Description")))
+			description = descriptionStr.c_str();
+		if (table.GetValueAsString(logFileStr, TC("LogFile")))
+			logFile = logFileStr.c_str();
+		if (table.GetValueAsString(breadcrumbsStr, TC("Breadcrumbs")))
+			breadcrumbs = breadcrumbsStr.c_str();
+
+		table.GetValueAsU32(priorityClass, TC("PriorityClass"));
+		table.GetValueAsBool(trackInputs, TC("TrackInputs"));
+		table.GetValueAsBool(useCustomAllocator, TC("UseCustomAllocator"));
+		table.GetValueAsBool(writeOutputFilesOnFail, TC("WriteOutputFilesOnFail"));
+		table.GetValueAsBool(startSuspended, TC("StartSuspended"));
+		table.GetValueAsU64(rootsHandle, TC("RootsHandle"));
+	}
 }
diff --git a/Engine/Source/Programs/UnrealBuildAccelerator/Common/Public/UbaProcessStartInfoHolder.h b/Engine/Source/Programs/UnrealBuildAccelerator/Common/Public/UbaProcessStartInfoHolder.h
--- a/Engine/Source/Programs/UnrealBuildAccelerator/Common/Public/UbaProcessStartInfoHolder.h
+++ b/Engine/Source/Programs/UnrealBuildAccelerator/Common/Public/UbaProcessStartInfoHolder.h
@@ -9,6 +9,8 @@
 
 namespace uba
 {
+	class Config;
+
 	inline void FixFileName(StringBufferBase& out, const tchar* fileName, const tchar* workingDir)
 	{
 		tchar buffer[1024];
@@ -120,6 +122,8 @@ namespace uba
 
 		bool Expand();
 
+		void Apply(const Config& config, const uba::tchar* configTable = TC(""));
+
 		TString descriptionStr;
 		TString applicationStr;
 		TString argumentsStr;
diff --git a/Engine/Source/Programs/UnrealBuildAccelerator/Host/Private/UbaExports.cpp b/Engine/Source/Programs/UnrealBuildAccelerator/Host/Private/UbaExports.cpp
--- a/Engine/Source/Programs/UnrealBuildAccelerator/Host/Private/UbaExports.cpp
+++ b/Engine/Source/Programs/UnrealBuildAccelerator/Host/Private/UbaExports.cpp
@@ -586,41 +586,48 @@ extern "C"
 
 	uba::ProcessStartInfo* ProcessStartInfo_Create(const uba::tchar* application, const uba::tchar* arguments, const uba::tchar* workingDir, const uba::tchar* description, uba::u32 priorityClass, uba::u64 outputStatsThresholdMs, bool trackInputs, const uba::tchar* logFile, ProcessHandle_ExitCallback* exit)
 	{
-		auto info = new uba::ProcessStartInfo();
-		info->application = TStrdup(application);
-		info->arguments = TStrdup(arguments);
-		info->workingDir = TStrdup(workingDir);
-		info->description = TStrdup(description);
+		auto info = new uba::ProcessStartInfoHolder();
+
+		info->applicationStr = application;
+		info->argumentsStr = arguments;
+		info->workingDirStr = workingDir;
+		info->descriptionStr = description;
+		info->logFileStr = logFile;
+
+		info->application = info->applicationStr.c_str();
+		info->arguments = info->argumentsStr.c_str();
+		info->workingDir = info->workingDirStr.c_str();
+		info->description = info->descriptionStr.c_str();
+		info->logFile = info->logFileStr.c_str();
 		info->priorityClass = priorityClass;
 		info->trackInputs = trackInputs;
-		info->logFile = TStrdup(logFile);
 		info->exitedFunc = exit;
 		return info;
 	}
 
 	uba::ProcessStartInfo* ProcessStartInfo_Create2(const uba::tchar* application, const uba::tchar* arguments, const uba::tchar* workingDir, const uba::tchar* description, uba::u32 priorityClass, uba::u64 rootsHandle, bool trackInputs, const uba::tchar* logFile, ProcessHandle_ExitCallback* exit)
 	{
-		auto info = new uba::ProcessStartInfo();
-		info->application = TStrdup(application);
-		info->arguments = TStrdup(arguments);
-		info->workingDir = TStrdup(workingDir);
-		info->description = TStrdup(description);
-		info->priorityClass = priorityClass;
-		info->trackInputs = trackInputs;
-		info->logFile = TStrdup(logFile);
-		info->exitedFunc = exit;
+		auto info = ProcessStartInfo_Create(application, arguments, workingDir, description, priorityClass, 0, trackInputs, logFile, exit);
 		info->rootsHandle = rootsHandle;
 		return info;
 	}
 
+	uba::ProcessStartInfo* ProcessStartInfo_Create3(const uba::Config& config, const uba::tchar* configTable)
+	{
+		auto info = new uba::ProcessStartInfoHolder();
+		info->Apply(config, configTable);
+		return info;
+	}
+
+	void ProcessStartInfo_SetExitedCallback(uba::ProcessStartInfo& info, ProcessHandle_ExitCallback* exitedFunc, void* exitedUserData)
+	{
+		info.exitedFunc = exitedFunc;
+		info.userData = exitedUserData;
+	}
+	
 	void ProcessStartInfo_Destroy(uba::ProcessStartInfo* info)
 	{
-		free((void*)info->application);
-		free((void*)info->arguments);
-		free((void*)info->workingDir);
-		free((void*)info->description);
-		free((void*)info->logFile);
-		delete info;
+		delete (uba::ProcessStartInfoHolder*)info;
 	}
 
 	uba::Scheduler* Scheduler_Create(uba::SessionServer* session, uba::u32 maxLocalProcessors, bool enableProcessReuse)
diff --git a/Engine/Source/Programs/UnrealBuildAccelerator/Host/Public/UbaExports.h b/Engine/Source/Programs/UnrealBuildAccelerator/Host/Public/UbaExports.h
--- a/Engine/Source/Programs/UnrealBuildAccelerator/Host/Public/UbaExports.h
+++ b/Engine/Source/Programs/UnrealBuildAccelerator/Host/Public/UbaExports.h
@@ -79,6 +79,8 @@ extern "C"
 	using ProcessHandle_ExitCallback = void(void* userData, const uba::ProcessHandle&);
 	UBA_API uba::ProcessStartInfo* ProcessStartInfo_Create(const uba::tchar* application, const uba::tchar* arguments, const uba::tchar* workingDir, const uba::tchar* description, uba::u32 priorityClass, uba::u64 outputStatsThresholdMs, bool trackInputs, const uba::tchar* logFile, ProcessHandle_ExitCallback* exit);
 	UBA_API uba::ProcessStartInfo* ProcessStartInfo_Create2(const uba::tchar* application, const uba::tchar* arguments, const uba::tchar* workingDir, const uba::tchar* description, uba::u32 priorityClass, uba::u64 rootsHandle, bool trackInputs, const uba::tchar* logFile, ProcessHandle_ExitCallback* exit);
+	UBA_API uba::ProcessStartInfo* ProcessStartInfo_Create3(const uba::Config& config, const uba::tchar* configTable = TC(""));
+	UBA_API void ProcessStartInfo_SetExitedCallback(uba::ProcessStartInfo& info, ProcessHandle_ExitCallback* exitedFunc, void* exitedUserData);
 	UBA_API void ProcessStartInfo_Destroy(uba::ProcessStartInfo* info);
 
 	// ProcessHandle
