[UBA]
* Added more info to log warning in CacheServer

diff --git a/Engine/Source/Programs/UnrealBuildAccelerator/Common/Private/UbaCacheServer.cpp b/Engine/Source/Programs/UnrealBuildAccelerator/Common/Private/UbaCacheServer.cpp
--- a/Engine/Source/Programs/UnrealBuildAccelerator/Common/Private/UbaCacheServer.cpp
+++ b/Engine/Source/Programs/UnrealBuildAccelerator/Common/Private/UbaCacheServer.cpp
@@ -1505,7 +1505,7 @@ namespace uba
 			auto insres = inputs.insert(casKeyOffset);
 			if (!insres.second)
 			{
-				m_logger.Warning(TC("Input file %s exists more than once in cache entry"), path.data);
+				m_logger.Warning(TC("Client %u Input file %s exists more than once in cache entry (CasKey: %s ClientVersion: %u PathtableVersion: %u)"), clientId, path.data, CasKeyString(casKey).str, clientVersion, connectionBucket.pathTable.GetVersion());
 				continue;
 			}
 			bytesForInput += Get7BitEncodedCount(casKeyOffset);
