[UBA]
* Improved error messages to make it easier for non-uba devs to know the errors are in uba

diff --git a/Engine/Source/Programs/UnrealBuildAccelerator/Common/Private/UbaLogger.cpp b/Engine/Source/Programs/UnrealBuildAccelerator/Common/Private/UbaLogger.cpp
--- a/Engine/Source/Programs/UnrealBuildAccelerator/Common/Private/UbaLogger.cpp
+++ b/Engine/Source/Programs/UnrealBuildAccelerator/Common/Private/UbaLogger.cpp
@@ -103,13 +103,13 @@ namespace uba
 			TStrcpy_s(buffer, 1024, format);
 		va_end(arg);
 #if PLATFORM_WINDOWS
-		wprintf(TC("FATAL ERROR %u: %s\n"), code, buffer);
+		wprintf(TC("UBA FATAL ERROR %u: %s\n"), code, buffer);
 		fflush(stdout);
 		if (IsDebuggerPresent())
 			DebugBreak();
 		ExitProcess(code);
 #else
-		printf(TC("FATAL ERROR %u: %s\n"), code, buffer);
+		printf(TC("UBA FATAL ERROR %u: %s\n"), code, buffer);
 		fflush(stdout);
 		kill(getpid(), SIGKILL);
 #endif
diff --git a/Engine/Source/Programs/UnrealBuildAccelerator/Core/Private/UbaMemory.cpp b/Engine/Source/Programs/UnrealBuildAccelerator/Core/Private/UbaMemory.cpp
--- a/Engine/Source/Programs/UnrealBuildAccelerator/Core/Private/UbaMemory.cpp
+++ b/Engine/Source/Programs/UnrealBuildAccelerator/Core/Private/UbaMemory.cpp
@@ -99,7 +99,7 @@ namespace uba
 		u64 newPos = startPos + bytes;
 		
 		if (newPos > reserveSize)
-			FatalError(9882, TC("Ran out of reserved space . Reserved %llu, Needed %llu (%s)"), reserveSize, newPos, hint);
+			FatalError(9882, TC("Ran out of reserved virtual address space. Reserved %llu, Needed %llu (%s)"), reserveSize, newPos, hint);
 
 		#if PLATFORM_WINDOWS
 		if (newPos > mappedSize)
