[UBA]
* Added more error handling and logging to UbaCacheServer to try to track down crash
* Added a couple more commands to cache server (usedpaths, usedsegments)

diff --git a/Engine/Source/Programs/UnrealBuildAccelerator/CacheService/Private/UbaCacheService.cpp b/Engine/Source/Programs/UnrealBuildAccelerator/CacheService/Private/UbaCacheService.cpp
--- a/Engine/Source/Programs/UnrealBuildAccelerator/CacheService/Private/UbaCacheService.cpp
+++ b/Engine/Source/Programs/UnrealBuildAccelerator/CacheService/Private/UbaCacheService.cpp
@@ -54,6 +54,8 @@ namespace uba
 		logger.Info(TC("  -http=<port>            If set, a http server will be started and listen on <port>"));
 		logger.Info(TC("  -fullmaintenance        Force a full maintenance"));
 		logger.Info(TC("  -nomaintenance          Skip all maintenance"));
+		logger.Info(TC("  -crash                  Force a crash (for testing)"));
+		logger.Info(TC("  -nosignalhandler        Will not hook up signal handler"));
 
 		logger.Info(TC(" Example of how to register crypto key to cache server (when -http=80 is provided)"));
 		logger.Info(TC("   curl http://localhost/addcrypto?3f58aa57466db9999213456789123445"));
@@ -99,8 +101,6 @@ namespace uba
 	{
 		using namespace uba;
 
-		AddExceptionHandler();
-
 		u32 storageCapacityGb = DefaultCapacityGb;
 		StringBuffer<256> workDir;
 		StringBuffer<128> listenIp;
@@ -110,6 +110,8 @@ namespace uba
 		bool storeCompressed = true;
 		bool fullMaintenance = false;
 		bool maintenanceEnabled = true;
+		bool signalHandlerEnabled = true;
+		bool shouldCrash = false;
 		u32 expirationTimeSeconds = DefaultExpiration;
 		TString configFile;
 
@@ -171,10 +173,11 @@ namespace uba
 			}
 			else if (name.Equals(TC("-crash")))
 			{
-				#ifndef __clang_analyzer__
-				int* ptr = nullptr;
-				*ptr = 0;
-				#endif
+				shouldCrash = true;
+			}
+			else if (name.Equals(TC("-nosignalhandler")))
+			{
+				signalHandlerEnabled = false;
 			}
 			else if (name.Equals(TC("-nomaintenance")))
 			{
@@ -200,6 +203,17 @@ namespace uba
 			}
 		}
 
+		if (signalHandlerEnabled)
+			AddExceptionHandler();
+
+		if (shouldCrash)
+		{
+			#ifndef __clang_analyzer__
+			int* ptr = nullptr;
+			*ptr = 0;
+			#endif
+		}
+
 		FilteredLogWriter logWriter(g_consoleLogWriter, quiet ? LogEntryType_Info : LogEntryType_Detail);
 		LoggerWithWriter logger(logWriter, TC(""));
 
diff --git a/Engine/Source/Programs/UnrealBuildAccelerator/Common/Private/UbaCacheServer.cpp b/Engine/Source/Programs/UnrealBuildAccelerator/Common/Private/UbaCacheServer.cpp
--- a/Engine/Source/Programs/UnrealBuildAccelerator/Common/Private/UbaCacheServer.cpp
+++ b/Engine/Source/Programs/UnrealBuildAccelerator/Common/Private/UbaCacheServer.cpp
@@ -274,7 +274,7 @@ namespace uba
 			u32 bucketCount = reader.ReadU32();
 			while (bucketCount--)
 			{
-				Bucket& bucket = GetBucket(reader.ReadU64());
+				Bucket& bucket = GetBucket(reader.ReadU64(), TC("Loading"));
 				LoadBucket(bucket, reader, databaseVersion, stats);
 			}
 		}
@@ -289,7 +289,7 @@ namespace uba
 					u64 id;
 					if (!keyName.Parse(id))
 						return;
-					GetBucket(id, false);
+					GetBucket(id, TC("Loading"), false);
 				});
 
 			m_server.ParallelFor(GetBucketWorkerCount(), m_buckets, [&](auto& it)
@@ -1308,14 +1308,14 @@ namespace uba
 		return connection.storeBuckets.try_emplace(id, id, bucketVersion).first->second;
 	}
 
-	CacheServer::Bucket& CacheServer::GetBucket(BinaryReader& reader)
+	CacheServer::Bucket& CacheServer::GetBucket(BinaryReader& reader, const tchar* reason)
 	{
 		u64 id = reader.Read7BitEncoded();
-		return GetBucket(id);
+		return GetBucket(id, reason);
 		
 	}
 
-	CacheServer::Bucket& CacheServer::GetBucket(u64 id, bool addCommon)
+	CacheServer::Bucket& CacheServer::GetBucket(u64 id, const tchar* reason, bool addCommon)
 	{
 		u32 bucketVersion = u32(id >> 35u);
 		SCOPED_WRITE_LOCK(m_bucketsLock, bucketsLock);
@@ -1326,6 +1326,7 @@ namespace uba
 		if (addCommon)
 			bucket.m_pathTable.AddCommonStringSegments();
 		bucket.index = u32(m_buckets.size() - 1);
+		m_logger.Info(TC("Bucket %u created with id %llu and version %u (%s)"), bucket.index, id, bucketVersion, reason);
 		return bucket;
 	}
 
@@ -1411,7 +1412,7 @@ namespace uba
 				return true;
 
 			u64 id = connectionBucket.id;
-			Bucket& bucket = GetBucket(id);
+			Bucket& bucket = GetBucket(id, TC("StoreEntryDone"));
 
 			SCOPED_WRITE_LOCK(bucket.m_cacheEntryLookupLock, lock3);
 			auto insres = bucket.m_cacheEntryLookup.try_emplace(cmdKey);
@@ -1474,7 +1475,7 @@ namespace uba
 
 		u64 outputStartOffset = reader.GetPosition();
 		u64 id = connectionBucket.id;
-		Bucket& bucket = GetBucket(id);
+		Bucket& bucket = GetBucket(id, TC("StoreEntry"));
 		connectionBucket.index = bucket.index;
 
 		while (reader.GetLeft())
@@ -1489,7 +1490,8 @@ namespace uba
 
 			CasKey casKey;
 			StringBuffer<> path;
-			connectionBucket.casKeyTable.GetPathAndKey(path, casKey, connectionBucket.pathTable, offset);
+			if (!connectionBucket.casKeyTable.GetPathAndKey(path, casKey, connectionBucket.pathTable, offset))
+				return m_logger.Error(TC("Client %u has provided corrupt key or path table (CasKey: %s ClientVersion: %u PathtableVersion: %u)"), clientId, CasKeyString(casKey).str, clientVersion, connectionBucket.pathTable.GetVersion());
 			UBA_ASSERT(path.count);
 
 			u32 pathOffset = bucket.m_pathTable.Add(path.data, path.count);
@@ -1693,7 +1695,7 @@ namespace uba
 
 	bool CacheServer::HandleFetchPathTable(BinaryReader& reader, BinaryWriter& writer)
 	{
-		Bucket& bucket = GetBucket(reader);
+		Bucket& bucket = GetBucket(reader, TC("FetchPathTable"));
 		u32 haveSize = reader.ReadU32();
 		u32 size = bucket.m_pathTable.GetSize();
 		writer.WriteU32(size);
@@ -1704,7 +1706,7 @@ namespace uba
 
 	bool CacheServer::HandleFetchPathTable2(BinaryReader& reader, BinaryWriter& writer)
 	{
-		Bucket& bucket = GetBucket(reader);
+		Bucket& bucket = GetBucket(reader, TC("FetchPathTable2"));
 		u32 haveSize = reader.ReadU32();
 		u32 size = bucket.m_pathTable.GetSize();
 		u32 toSend = Min(u32(writer.GetCapacityLeft()), size - haveSize);
@@ -1714,7 +1716,7 @@ namespace uba
 
 	bool CacheServer::HandleFetchCasTable(BinaryReader& reader, BinaryWriter& writer)
 	{
-		Bucket& bucket = GetBucket(reader);
+		Bucket& bucket = GetBucket(reader, TC("FetchCasTable"));
 		u32 haveSize = reader.ReadU32();
 		u32 size = bucket.m_casKeyTable.GetSize();
 		writer.WriteU32(size);
@@ -1725,7 +1727,7 @@ namespace uba
 
 	bool CacheServer::HandleFetchCasTable2(BinaryReader& reader, BinaryWriter& writer)
 	{
-		Bucket& bucket = GetBucket(reader);
+		Bucket& bucket = GetBucket(reader, TC("FetchCasTable2"));
 		u32 haveSize = reader.ReadU32();
 		u32 size = bucket.m_casKeyTable.GetSize();
 		u32 toSend = Min(u32(writer.GetCapacityLeft()), size - haveSize);
@@ -1735,7 +1737,7 @@ namespace uba
 
 	bool CacheServer::HandleFetchEntries(BinaryReader& reader, BinaryWriter& writer, u32 clientId)
 	{
-		Bucket& bucket = GetBucket(reader);
+		Bucket& bucket = GetBucket(reader, TC("FetchEntries"));
 		CasKey cmdKey = reader.ReadCasKey();
 
 		u32 clientVersion;
@@ -1765,7 +1767,7 @@ namespace uba
 
 	bool CacheServer::HandleReportUsedEntry(BinaryReader& reader, BinaryWriter& writer, u32 clientId)
 	{
-		Bucket& bucket = GetBucket(reader);
+		Bucket& bucket = GetBucket(reader, TC("ReportUsedEntry"));
 		CasKey cmdKey = reader.ReadCasKey();
 		u64 entryId = reader.Read7BitEncoded();
 
@@ -2012,6 +2014,74 @@ namespace uba
 					writeLine(line.Clear().Append(TC("   ")).Append(path));
 			}
 		}
+		else if (command.Equals(TC("usedpaths")))
+		{
+			SCOPED_READ_LOCK(m_bucketsLock, bucketsLock);
+			for (auto& kv : m_buckets)
+			{
+				Bucket& bucket = kv.second;
+				SCOPED_READ_LOCK(bucket.m_cacheEntryLookupLock, lock2);
+
+				writeLine(line.Clear().Appendf(TC("Bucket #%u"), bucket.index));
+
+				Set<TString> sortedPaths;
+
+				BinaryReader tableReader(bucket.m_casKeyTable.GetMemory(), 0, bucket.m_casKeyTable.GetSize());
+				while (tableReader.GetLeft())
+				{
+					u32 stringOffset = (u32)tableReader.Read7BitEncoded();
+					tableReader.ReadCasKey();
+					StringBuffer<> path;
+					bucket.m_pathTable.GetString(path, stringOffset);
+					sortedPaths.emplace(path.data);
+				}
+
+				for (auto& path : sortedPaths)
+					writeLine(line.Clear().Append(TC("   ")).Append(path));
+			}
+		}
+		else if (command.Equals(TC("usedsegments")))
+		{
+			SCOPED_READ_LOCK(m_bucketsLock, bucketsLock);
+			for (auto& kv : m_buckets)
+			{
+				UnorderedMap<u32, u32> segments;
+
+				Bucket& bucket = kv.second;
+				SCOPED_READ_LOCK(bucket.m_cacheEntryLookupLock, lock2);
+
+				writeLine(line.Clear().Appendf(TC("Bucket #%u"), bucket.index));
+
+				BinaryReader tableReader(bucket.m_pathTable.GetMemory(), bucket.m_pathTable.GetCommonSize(), bucket.m_pathTable.GetSize());
+				while (tableReader.GetLeft())
+				{
+					tableReader.Read7BitEncoded(); // parent
+					u32 strOffset = u32(tableReader.Read7BitEncoded());
+					if (strOffset == 0)
+					{
+						strOffset = u32(tableReader.GetPosition());
+						tableReader.SkipString();
+					}
+					++segments[strOffset];
+				}
+
+				Map<u32, Vector<TString>> sortedSegments;
+
+				for (auto& kv2 : segments)
+				{
+					tableReader.SetPosition(kv2.first);
+					StringBuffer<> segment;
+					tableReader.ReadString(segment);
+					sortedSegments[1000000 - kv2.second].push_back(segment.data);
+				}
+				for (auto& kv2 : sortedSegments)
+				{
+					writeLine(line.Clear().Append(TC("  ")).AppendValue(1000000 - kv2.first));
+					for (auto& str : kv2.second)
+						writeLine(line.Clear().Append(TC("    ")).Append(str));
+				}
+			}
+		}
 		else if (command.Equals(TC("largestentry")))
 		{
 			CacheEntries* largestEntries = nullptr;
diff --git a/Engine/Source/Programs/UnrealBuildAccelerator/Common/Private/UbaCompactTables.cpp b/Engine/Source/Programs/UnrealBuildAccelerator/Common/Private/UbaCompactTables.cpp
--- a/Engine/Source/Programs/UnrealBuildAccelerator/Common/Private/UbaCompactTables.cpp
+++ b/Engine/Source/Programs/UnrealBuildAccelerator/Common/Private/UbaCompactTables.cpp
@@ -51,7 +51,7 @@ namespace uba
 
 	void CompactPathTable::AddCommonStringSegments()
 	{
-		const tchar* commonSegments[] =
+		const tchar* commonSegments1[] =
 		{
 			TC("h"),
 			TC("cpp"),
@@ -69,23 +69,52 @@ namespace uba
 			TC("generated"),
 			TC("init"),
 			TC("ispc"),
+			// If adding more we want these
 		};
+		// 68 bytes used
+
+		const tchar* commonSegments2[] =
+		{
+			TC("H"),
+			TC("0"),
+			TC("1"),
+			TC("2"),
+			TC("3"),
+			TC("Definitions"),
+			TC("Private"),
+			TC("Shared"),
+			TC("Public"),
+			TC("Inc"),
+			TC("UHT"),
+			TC("x64"),
+			TC("res"),
+		};
+		// 124 bytes used
 
 		if (!m_pathTableMem.memory)
 			m_pathTableMem.Init(PathTableMaxSize);
 		if (!m_pathTableMem.writtenSize)
 			m_pathTableMem.AllocateNoLock(1, 1, TC("CompactPathTable"));
 
-		for (const tchar* seg : commonSegments)
-		{
-			u32 segLen = TStrlen(seg);
-			u64 writtenSize = m_pathTableMem.writtenSize;
-			BinaryWriter writer((u8*)m_pathTableMem.AllocateNoLock(segLen+1, 1, TC("CompactPathTable")), 0, segLen+1);
-			StringKey segmentKey = ToStringKeyNoCheck(seg, segLen);
-			u32& strOffset = m_segmentOffsets.Insert(segmentKey);
-			strOffset = u32(writtenSize);
-			writer.WriteString(seg, segLen);
-		}
+		auto addSegments = [this](const tchar** segments, u64 segmentCount)
+			{
+				for (const tchar** it=segments, **itEnd=it+segmentCount; it!=itEnd; ++it)
+				{
+					const tchar* seg = *it;
+					u32 segLen = TStrlen(seg);
+					u64 writtenSize = m_pathTableMem.writtenSize;
+					BinaryWriter writer((u8*)m_pathTableMem.AllocateNoLock(segLen+1, 1, TC("CompactPathTable")), 0, segLen+1);
+					StringKey segmentKey = ToStringKeyNoCheck(seg, segLen);
+					u32& strOffset = m_segmentOffsets.Insert(segmentKey);
+					strOffset = u32(writtenSize);
+					writer.WriteString(seg, segLen);
+				}
+			};
+
+		addSegments(commonSegments1, sizeof_array(commonSegments1));
+		if (m_version >= 4)
+			addSegments(commonSegments2, sizeof_array(commonSegments2));
+		m_commonSize = u32(m_pathTableMem.writtenSize);
 	}
 
 	u32 CompactPathTable::InternalAdd(const tchar* str, const tchar* stringKeyString, u64 strLen)
@@ -216,7 +245,7 @@ namespace uba
 		return offset;
 	}
 
-	void CompactPathTable::GetString(StringBufferBase& out, u64 offset) const
+	bool CompactPathTable::GetString(StringBufferBase& out, u64 offset) const
 	{
 		#if UBA_DEBUG
 		{
@@ -225,8 +254,9 @@ namespace uba
 		}
 		#endif
 
-		u32 offsets[256];
-		bool separators[256];
+		enum { MaxSegments = 50 }; // This number is arbitrary.. don't know how many folders/dots deep something can be. We want it to be as low as possible to catch corrupt tables
+		u32 offsets[MaxSegments + 1];
+		bool separators[MaxSegments + 1];
 		offsets[0] = u32(offset);
 		u32 offsetCount = 0;
 
@@ -242,7 +272,7 @@ namespace uba
 				u64 parentRelativeOffset = reader.Read7BitEncoded();
 				separators[offsetCount] = parentRelativeOffset & 1;
 				parentRelativeOffset >>= 1;
-				if (parentRelativeOffset == offset)
+				if (parentRelativeOffset == offset || offsetCount == MaxSegments)
 					break;
 				offset -= parentRelativeOffset;
 				offsets[offsetCount] = u32(offset);
@@ -251,7 +281,7 @@ namespace uba
 		}
 		else
 		{
-			while (offset)
+			while (offset || offsetCount == MaxSegments)
 			{
 				++offsetCount;
 				UBA_ASSERT(offsetCount < sizeof_array(offsets));
@@ -266,6 +296,8 @@ namespace uba
 				offsets[offsetCount] = u32(offset);
 			}
 		}
+		if (offsetCount == MaxSegments)
+			return false;
 
 		{
 			bool isFirst = true;
@@ -283,6 +315,7 @@ namespace uba
 				reader.ReadString(out);
 			}
 		}
+		return true;
 	}
 
 	u8* CompactPathTable::GetMemory()
@@ -296,6 +329,11 @@ namespace uba
 		return u32(m_pathTableMem.writtenSize);
 	}
 
+	u32 CompactPathTable::GetCommonSize()
+	{
+		return m_commonSize;
+	}
+
 	bool CompactPathTable::GetCaseInsensitive()
 	{
 		return m_caseInsensitive;
@@ -306,7 +344,7 @@ namespace uba
 		return m_version;
 	}
 
-	void CompactPathTable::ReadMem(BinaryReader& reader, bool populateLookup)
+	bool CompactPathTable::ReadMem(BinaryReader& reader, bool populateLookup)
 	{
 		if (!m_pathTableMem.memory)
 			m_pathTableMem.Init(PathTableMaxSize);
@@ -317,7 +355,7 @@ namespace uba
 		reader.ReadBytes(mem, left);
 
 		if (!populateLookup)
-			return;
+			return true;
 
 		BinaryReader reader2(m_pathTableMem.memory, writtenSize, m_pathTableMem.writtenSize);
 		if (!writtenSize)
@@ -336,11 +374,13 @@ namespace uba
 				m_segmentOffsets.Insert(ToStringKeyNoCheck(seg.data, seg.count)) = strOffset;
 			}
 			StringBuffer<> str;
-			GetString(str, offset);
+			if (!GetString(str, offset))
+				return false;
 			if (m_caseInsensitive)
 				str.MakeLower();
 			m_offsets.Insert(ToStringKeyNoCheck(str.data, str.count)) = offset;
 		}
+		return true;
 	}
 
 	u8* CompactPathTable::BeginCommit(u64 size)
@@ -531,7 +571,7 @@ namespace uba
 		outKey = reader.ReadCasKey();
 	}
 
-	void CompactCasKeyTable::GetPathAndKey(StringBufferBase& outPath, CasKey& outKey, const CompactPathTable& pathTable, u64 offset) const
+	bool CompactCasKeyTable::GetPathAndKey(StringBufferBase& outPath, CasKey& outKey, const CompactPathTable& pathTable, u64 offset) const
 	{
 		#if UBA_DEBUG
 		{
@@ -543,7 +583,7 @@ namespace uba
 		BinaryReader reader(m_casKeyTableMem.memory, offset, ~0u);
 		u32 stringOffset = (u32)reader.Read7BitEncoded();
 		outKey = reader.ReadCasKey();
-		pathTable.GetString(outPath, stringOffset);
+		return pathTable.GetString(outPath, stringOffset);
 	}
 
 	u8* CompactCasKeyTable::GetMemory()
diff --git a/Engine/Source/Programs/UnrealBuildAccelerator/Common/Public/UbaCacheServer.h b/Engine/Source/Programs/UnrealBuildAccelerator/Common/Public/UbaCacheServer.h
--- a/Engine/Source/Programs/UnrealBuildAccelerator/Common/Public/UbaCacheServer.h
+++ b/Engine/Source/Programs/UnrealBuildAccelerator/Common/Public/UbaCacheServer.h
@@ -76,8 +76,8 @@ namespace uba
 		struct ConnectionBucket;
 
 		ConnectionBucket& GetConnectionBucket(const ConnectionInfo& connectionInfo, BinaryReader& reader, u32* outClientVersion = nullptr);
-		Bucket& GetBucket(BinaryReader& reader);
-		Bucket& GetBucket(u64 id, bool addCommon = true);
+		Bucket& GetBucket(BinaryReader& reader, const tchar* reason);
+		Bucket& GetBucket(u64 id, const tchar* reason, bool addCommon = true);
 		u32 GetBucketWorkerCount();
 
 		bool HandleMessage(const ConnectionInfo& connectionInfo, u8 messageType, BinaryReader& reader, BinaryWriter& writer);
diff --git a/Engine/Source/Programs/UnrealBuildAccelerator/Common/Public/UbaCompactTables.h b/Engine/Source/Programs/UnrealBuildAccelerator/Common/Public/UbaCompactTables.h
--- a/Engine/Source/Programs/UnrealBuildAccelerator/Common/Public/UbaCompactTables.h
+++ b/Engine/Source/Programs/UnrealBuildAccelerator/Common/Public/UbaCompactTables.h
@@ -23,14 +23,15 @@ namespace uba
 
 		void AddCommonStringSegments();
 
-		void GetString(StringBufferBase& out, u64 offset) const;
+		bool GetString(StringBufferBase& out, u64 offset) const;
 
 		u8* GetMemory();
 		u32 GetSize();
+		u32 GetCommonSize();
 		bool GetCaseInsensitive();
 		u32 GetVersion();
 
-		void ReadMem(BinaryReader& reader, bool populateLookup);
+		bool ReadMem(BinaryReader& reader, bool populateLookup);
 		u8* BeginCommit(u64 size);
 		void EndCommit(u8* data, u64 written);
 		void Swap(CompactPathTable& other);
@@ -39,7 +40,7 @@ namespace uba
 		u64 GetSegmentCount() { return m_segmentOffsets.Size(); }
 
 		template<typename Func>
-		void TraversePaths(const Func& func) const;
+		bool TraversePaths(const Func& func) const;
 
 	private:
 		u32 InternalAdd(const tchar* str, const tchar* stringKeyString, u64 strLen);
@@ -47,8 +48,12 @@ namespace uba
 		MemoryBlock m_pathTableMem;
 		HashMap<StringKey, u32, true> m_offsets;
 		HashMap<StringKey, u32, true> m_segmentOffsets;
-		u32 m_version;
-		bool m_caseInsensitive;
+		u32 m_version = 0;
+		u32 m_commonSize = 0;
+		bool m_caseInsensitive = true;
+
+		CompactPathTable(const CompactPathTable&) = delete;
+		void operator=(const CompactPathTable&) = delete;
 	};
 
 	class CompactCasKeyTable
@@ -64,7 +69,7 @@ namespace uba
 		void TraverseOffsets(const CasKey& casKey, const Func& func) const;
 
 		void GetKey(CasKey& outKey, u64 offset) const;
-		void GetPathAndKey(StringBufferBase& outPath, CasKey& outKey, const CompactPathTable& pathTable, u64 offset) const;
+		bool GetPathAndKey(StringBufferBase& outPath, CasKey& outKey, const CompactPathTable& pathTable, u64 offset) const;
 
 		u8* GetMemory();
 		u32 GetSize();
@@ -103,17 +108,22 @@ namespace uba
 		};
 
 		HashMap<CasKey, Value, true> m_offsets;
+
+		CompactCasKeyTable(const CompactCasKeyTable&) = delete;
+		void operator=(const CompactCasKeyTable&) = delete;
 	};
 
 	template<typename Func>
-	void CompactPathTable::TraversePaths(const Func& func) const
+	bool CompactPathTable::TraversePaths(const Func& func) const
 	{
 		for (auto i = m_offsets.ValuesBegin(), e = m_offsets.ValuesEnd(); i!=e; ++i)
 		{
 			StringBuffer path;
-			GetString(path, *i);
+			if (!GetString(path, *i))
+				return false;
 			func(path);
 		}
+		return true;
 	}
 
 	template<typename Func>
