[UBA]
* Added more bounds validation to cache client when fetching cache entries

diff --git a/Engine/Source/Programs/UnrealBuildAccelerator/Common/Private/UbaCacheClient.cpp b/Engine/Source/Programs/UnrealBuildAccelerator/Common/Private/UbaCacheClient.cpp
--- a/Engine/Source/Programs/UnrealBuildAccelerator/Common/Private/UbaCacheClient.cpp
+++ b/Engine/Source/Programs/UnrealBuildAccelerator/Common/Private/UbaCacheClient.cpp
@@ -1067,7 +1067,35 @@ namespace uba
 	{
 		if (HasEnoughData(bucket, requiredCasTableOffset))
 			return true;
-		return FetchCasTable2(bucket, stats, requiredCasTableOffset);
+
+		if (requiredCasTableOffset > 256*1024*1024)
+		{
+			m_logger.Warning(TC("Cas table offset %u too large. Cache entry corrupt (Bucket %llu)"), bucket.id);
+			return false;
+		}
+
+		TimerScope ts2(stats.fetchCasTable);
+		u32 requiredPathTableOffset = 0;
+		{
+			SCOPED_WRITE_LOCK(bucket.serverCasKeyTableNetworkLock, lock);
+			while (!HasEnoughCasData(bucket, requiredCasTableOffset, requiredPathTableOffset))
+			{
+				if (!FetchCompactTable(bucket.id, bucket.serverCasKeyTable, requiredCasTableOffset + sizeof(CasKey) + 8, CacheMessageType_FetchCasTable2))
+					return false;
+				bucket.serverCasKeyTableSize = bucket.serverCasKeyTable.GetSize();
+			}
+		}
+		{
+			SCOPED_WRITE_LOCK(bucket.serverPathTableNetworkLock, lock);
+			while (!HasEnoughPathData(bucket, requiredPathTableOffset))
+			{
+				u32 targetSize = requiredPathTableOffset + 200;
+				if (!FetchCompactTable(bucket.id, bucket.serverPathTable, targetSize, CacheMessageType_FetchPathTable2))
+					return false;
+				bucket.serverPathTableSize = bucket.serverPathTable.GetSize();
+			}
+		}
+		return true;
 	}
 
 	template<typename TableType>
@@ -1115,32 +1143,6 @@ namespace uba
 		return true;
 	}
 
-	UBA_NOINLINE bool CacheClient::FetchCasTable2(Bucket& bucket, CacheStats& stats, u32 requiredCasTableOffset)
-	{
-		TimerScope ts2(stats.fetchCasTable);
-		u32 requiredPathTableOffset = 0;
-		{
-			SCOPED_WRITE_LOCK(bucket.serverCasKeyTableNetworkLock, lock);
-			while (!HasEnoughCasData(bucket, requiredCasTableOffset, requiredPathTableOffset))
-			{
-				if (!FetchCompactTable(bucket.id, bucket.serverCasKeyTable, requiredCasTableOffset + sizeof(CasKey) + 8, CacheMessageType_FetchCasTable2))
-					return false;
-				bucket.serverCasKeyTableSize = bucket.serverCasKeyTable.GetSize();
-			}
-		}
-		{
-			SCOPED_WRITE_LOCK(bucket.serverPathTableNetworkLock, lock);
-			while (!HasEnoughPathData(bucket, requiredPathTableOffset))
-			{
-				u32 targetSize = requiredPathTableOffset + 200;
-				if (!FetchCompactTable(bucket.id, bucket.serverPathTable, targetSize, CacheMessageType_FetchPathTable2))
-					return false;
-				bucket.serverPathTableSize = bucket.serverPathTable.GetSize();
-			}
-		}
-		return true;
-	}
-
 	bool CacheClient::HasEnoughData(Bucket& bucket, u32 requiredCasTableOffset)
 	{
 		u32 pathOffset;
diff --git a/Engine/Source/Programs/UnrealBuildAccelerator/Common/Public/UbaCacheClient.h b/Engine/Source/Programs/UnrealBuildAccelerator/Common/Public/UbaCacheClient.h
--- a/Engine/Source/Programs/UnrealBuildAccelerator/Common/Public/UbaCacheClient.h
+++ b/Engine/Source/Programs/UnrealBuildAccelerator/Common/Public/UbaCacheClient.h
@@ -74,7 +74,6 @@ namespace uba
 		bool SendCasTable(Bucket& bucket, u32 requiredCasTableSize);
 		bool SendCacheEntry(Bucket& bucket, const RootPaths& rootPaths, const CasKey& cmdKey, const Map<u32, u32>& inputsStringToCasKey, const Map<u32, u32>& outputsStringToCasKey, const u8* logLines, u64 logLinesSize, u64& outBytesSent);
 		bool FetchCasTable(Bucket& bucket, CacheStats& stats, u32 requiredCasTableOffset);
-		bool FetchCasTable2(Bucket& bucket, CacheStats& stats, u32 requiredCasTableOffset);
 		template<typename TableType> bool FetchCompactTable(u32 bucketId, TableType& table, u32 requiredTableOffset, u8 messageType);
 		bool HasEnoughData(Bucket& bucket, u32 requiredCasTableOffset);
 		bool HasEnoughCasData(Bucket& bucket, u32 requiredCasTableOffset, u32& outPathOffset);
