[UBA]
* Logger::Warning now returns a bool that is always false.

diff --git a/Engine/Source/Programs/UnrealBuildAccelerator/Common/Private/UbaLogger.cpp b/Engine/Source/Programs/UnrealBuildAccelerator/Common/Private/UbaLogger.cpp
--- a/Engine/Source/Programs/UnrealBuildAccelerator/Common/Private/UbaLogger.cpp
+++ b/Engine/Source/Programs/UnrealBuildAccelerator/Common/Private/UbaLogger.cpp
@@ -139,12 +139,13 @@ namespace uba
 		va_end(arg);
 	}
 
-	void Logger::Warning(const tchar* format, ...)
+	bool Logger::Warning(const tchar* format, ...)
 	{
 		va_list arg;
 		va_start(arg, format);
 		LogArg(LogEntryType_Warning, format, arg);
 		va_end(arg);
+		return false;
 	}
 
 	bool Logger::Error(const tchar* format, ...)
diff --git a/Engine/Source/Programs/UnrealBuildAccelerator/Common/Public/UbaLogger.h b/Engine/Source/Programs/UnrealBuildAccelerator/Common/Public/UbaLogger.h
--- a/Engine/Source/Programs/UnrealBuildAccelerator/Common/Public/UbaLogger.h
+++ b/Engine/Source/Programs/UnrealBuildAccelerator/Common/Public/UbaLogger.h
@@ -25,7 +25,7 @@ namespace uba
 	public:
 		Logger() {}
 		bool Error(const tchar* format, ...);
-		void Warning(const tchar* format, ...);
+		bool Warning(const tchar* format, ...);
 		void Info(const tchar* format, ...);
 		void Detail(const tchar* format, ...);
 		void Debug(const tchar* format, ...);
