[UBA]
* Fixed so GetFullPathName handles filename that is null

diff --git a/Engine/Source/Programs/UnrealBuildAccelerator/Core/Public/UbaPathUtils.h b/Engine/Source/Programs/UnrealBuildAccelerator/Core/Public/UbaPathUtils.h
--- a/Engine/Source/Programs/UnrealBuildAccelerator/Core/Public/UbaPathUtils.h
+++ b/Engine/Source/Programs/UnrealBuildAccelerator/Core/Public/UbaPathUtils.h
@@ -16,6 +16,7 @@ namespace uba
 	inline bool FixPath2(const tchar* fileName, const tchar* workingDir, u64 workingDirCharLen, tchar* buffer, u64 bufferCharCapacity, u32* outBufferCharLen)
 	{
 		buffer[1] = 0;
+		UBA_ASSERT(fileName);
 		UBA_ASSERTF(workingDir == nullptr || !*workingDir || workingDir[workingDirCharLen-1] == PathSeparator, TC("WorkingDir needs to end with path separator"));
 
 #if PLATFORM_WINDOWS
diff --git a/Engine/Source/Programs/UnrealBuildAccelerator/Detours/Private/Windows/UbaDetoursFunctionsKernelBase.inl b/Engine/Source/Programs/UnrealBuildAccelerator/Detours/Private/Windows/UbaDetoursFunctionsKernelBase.inl
--- a/Engine/Source/Programs/UnrealBuildAccelerator/Detours/Private/Windows/UbaDetoursFunctionsKernelBase.inl
+++ b/Engine/Source/Programs/UnrealBuildAccelerator/Detours/Private/Windows/UbaDetoursFunctionsKernelBase.inl
@@ -1165,7 +1165,11 @@ DWORD Detoured_GetLongPathNameW(LPCWSTR lpszShortPath, LPWSTR lpszLongPath, DWOR
 DWORD Detoured_GetFullPathNameW(LPCWSTR lpFileName, DWORD nBufferLength, LPWSTR lpBuffer, LPWSTR* lpFilePart)
 {
 	DETOURED_CALL(GetFullPathNameW);
-
+	if (!lpFileName)
+	{
+		SetLastError(ERROR_INVALID_NAME);
+		return 0;
+	}
 	StringBuffer<> fullPath;
 	FixPath(fullPath, lpFileName);
 	u64 requiredSize = fullPath.count + 1;
@@ -1184,10 +1188,13 @@ DWORD Detoured_GetFullPathNameA(LPCSTR lpFileName, DWORD nBufferLength, LPSTR lp
 {
 	// Is verified that this does NOT always call GetFullPathNameW
 	DETOURED_CALL(GetFullPathNameA);
-
+	if (!lpFileName)
+	{
+		SetLastError(ERROR_INVALID_NAME);
+		return 0;
+	}
 	StringBuffer<> temp;
 	temp.Append(lpFileName);
-
 	StringBuffer<> fullPath;
 	FixPath(fullPath, temp.data);
 	u64 requiredSize = fullPath.count + 1;
