[UBA]
* Minor cosmetic cleanup

diff --git a/Engine/Source/Programs/UnrealBuildAccelerator/Common/Private/UbaFileAccessor.cpp b/Engine/Source/Programs/UnrealBuildAccelerator/Common/Private/UbaFileAccessor.cpp
--- a/Engine/Source/Programs/UnrealBuildAccelerator/Common/Private/UbaFileAccessor.cpp
+++ b/Engine/Source/Programs/UnrealBuildAccelerator/Common/Private/UbaFileAccessor.cpp
@@ -509,9 +509,7 @@ namespace uba
 			{
 				*lastWriteTime = 0;
 				if (!GetFileLastWriteTime(*lastWriteTime, m_fileHandle))
-				{
 					m_logger.Warning(TC("Failed to get file time for %s %llu (%s)"), m_fileName, m_fileHandle, LastErrorToText().data);
-				}
 			}
 		}
 		else
diff --git a/Engine/Source/Programs/UnrealBuildAccelerator/Common/Private/UbaNetworkClient.cpp b/Engine/Source/Programs/UnrealBuildAccelerator/Common/Private/UbaNetworkClient.cpp
--- a/Engine/Source/Programs/UnrealBuildAccelerator/Common/Private/UbaNetworkClient.cpp
+++ b/Engine/Source/Programs/UnrealBuildAccelerator/Common/Private/UbaNetworkClient.cpp
@@ -183,10 +183,7 @@ namespace uba
 		}
 
 		if (rc.error == 5)
-		{
-			m_logger.Warning(TC("A connection from a server with different uid was requested. Ignore"));
-			return false;
-		}
+			return m_logger.Warning(TC("A connection from a server with different uid was requested. Ignore"));
 
 		if (m_connectionCount.fetch_add(1) != 0)
 			return true;
diff --git a/Engine/Source/Programs/UnrealBuildAccelerator/Common/Private/UbaStorage.cpp b/Engine/Source/Programs/UnrealBuildAccelerator/Common/Private/UbaStorage.cpp
--- a/Engine/Source/Programs/UnrealBuildAccelerator/Common/Private/UbaStorage.cpp
+++ b/Engine/Source/Programs/UnrealBuildAccelerator/Common/Private/UbaStorage.cpp
@@ -1282,10 +1282,8 @@ namespace uba
 			return m_logger.Error(TC("GetFileSize failed for %s (%s)"), fileName.data, LastErrorToText().data);
 
 		if (fileSize < sizeof(u32))
-		{
-			m_logger.Warning(TC("CasTable file %s is corrupt (size: %u)"), fileName.data, fileSize);
-			return false;
-		}
+			return m_logger.Warning(TC("CasTable file %s is corrupt (size: %u)"), fileName.data, fileSize);
+
 		u8* buffer = new u8[fileSize];
 		auto memGuard = MakeGuard([&](){ delete[] buffer; });
 
@@ -1318,8 +1316,7 @@ namespace uba
 			if (reader.GetPosition() + 24 > fileSize)
 			{
 				m_fileTableLookup.clear();
-				m_logger.Warning(TC("CasTable file %s is corrupt"), fileName.data);
-				return false;
+				return m_logger.Warning(TC("CasTable file %s is corrupt"), fileName.data);
 			}
 			u64 size = reader.ReadU64();
 			u64 lastWritten = reader.ReadU64();
diff --git a/Engine/Source/Programs/UnrealBuildAccelerator/Common/Private/UbaTrace.cpp b/Engine/Source/Programs/UnrealBuildAccelerator/Common/Private/UbaTrace.cpp
--- a/Engine/Source/Programs/UnrealBuildAccelerator/Common/Private/UbaTrace.cpp
+++ b/Engine/Source/Programs/UnrealBuildAccelerator/Common/Private/UbaTrace.cpp
@@ -177,8 +177,7 @@ namespace uba
 		}
 
 		FreeMemory();
-		m_logger.Warning(TC("Failed to commit memory for trace (Pos: %llu Capacity: %llu, Already Committed: %llu, Needed: %llu): %s"), m_memoryPos, m_memoryCapacity, m_memoryCommitted, committedMemoryNeeded, LastErrorToText().data);
-		return false;
+		return m_logger.Warning(TC("Failed to commit memory for trace (Pos: %llu Capacity: %llu, Already Committed: %llu, Needed: %llu): %s"), m_memoryPos, m_memoryCapacity, m_memoryCommitted, committedMemoryNeeded, LastErrorToText().data);
 	}
 
 
diff --git a/Engine/Source/Programs/UnrealBuildAccelerator/Visualizer/Private/UbaTraceReader.cpp b/Engine/Source/Programs/UnrealBuildAccelerator/Visualizer/Private/UbaTraceReader.cpp
--- a/Engine/Source/Programs/UnrealBuildAccelerator/Visualizer/Private/UbaTraceReader.cpp
+++ b/Engine/Source/Programs/UnrealBuildAccelerator/Visualizer/Private/UbaTraceReader.cpp
@@ -1205,10 +1205,7 @@ namespace uba
 	bool TraceReader::SaveAs(const tchar* fileName)
 	{
 		if (!m_memoryBegin)
-		{
-			m_logger.Warning(L"Can only save traces that are opened using -listen/-name.");
-			return false;
-		}
+			return m_logger.Warning(L"Can only save traces that are opened using -listen/-name.");
 
 		FileAccessor file(m_logger, fileName);
 		if (!file.CreateWrite())
