[UBA]
* Added more validation in cache server around getting entries from clients
* Fixed cache client logic when sending caskey/path table to update send position after message has been sent since message can fail and if cache client has multiple connections then then one message can succeed after the failed one

diff --git a/Engine/Source/Programs/UnrealBuildAccelerator/Common/Private/UbaCacheClient.cpp b/Engine/Source/Programs/UnrealBuildAccelerator/Common/Private/UbaCacheClient.cpp
--- a/Engine/Source/Programs/UnrealBuildAccelerator/Common/Private/UbaCacheClient.cpp
+++ b/Engine/Source/Programs/UnrealBuildAccelerator/Common/Private/UbaCacheClient.cpp
@@ -295,10 +295,7 @@ namespace uba
 				{
 					
 					if (shouldValidate && FileExists(m_logger, path.data))
-					{
-						m_logger.Warning(TC("CasDb claims file %s does not exist but it does! Will not populate cache for %s"), path.data, info.GetDescription()); 
-						return false;
-					}
+						return m_logger.Warning(TC("CasDb claims file %s does not exist but it does! Will not populate cache for %s"), path.data, info.GetDescription()); 
 
 					//m_logger.Warning(TC("Can't find file %s"), path.data); 
 					stringToCasKey.erase(insres.first);
@@ -310,10 +307,7 @@ namespace uba
 				{
 					FileAccessor fa(m_logger, path.data);
 					if (!fa.OpenMemoryRead())
-					{
-						m_logger.Warning(TC("CasDb claims file %s does exist but can't open it. Will not populate cache for %s"), path.data, info.GetDescription()); 
-						return false;
-					}
+						return m_logger.Warning(TC("CasDb claims file %s does exist but can't open it. Will not populate cache for %s"), path.data, info.GetDescription()); 
 
 					CasKey oldKey = AsCompressed(casKey, false);
 					CasKey newKey;
@@ -335,9 +329,8 @@ namespace uba
 						SCOPED_READ_LOCK(fileEntry.lock, lock);
 
 						auto ToString = [](bool b) { return b ? TC("true") : TC("false"); };
-						m_logger.Warning(TC("CasDb claims file %s has caskey %s but recalculating it gives us %s (FileEntry: %llu/%llu/%s, Real: %llu/%llu). Will not populate cache for %s"),
+						return m_logger.Warning(TC("CasDb claims file %s has caskey %s but recalculating it gives us %s (FileEntry: %llu/%llu/%s, Real: %llu/%llu). Will not populate cache for %s"),
 							path.data, CasKeyString(oldKey).str, CasKeyString(newKey).str, fileEntry.size, fileEntry.lastWritten, ToString(fileEntry.verified), fileSize, fileInfo.lastWriteTime, info.GetDescription());
-						return false;
 					}
 				}
 			}
@@ -883,11 +876,12 @@ namespace uba
 			u32 toSend = Min(requiredPathTableSize - bucket.pathTableSizeSent, u32(m_client.GetMessageMaxSize() - 32));
 			left -= toSend;
 			writer.WriteBytes(bucket.sendPathTable.GetMemory() + bucket.pathTableSizeSent, toSend);
-			bucket.pathTableSizeSent += toSend;
 
 			StackBinaryReader<16> reader;
 			if (!msg.Send(reader))
 				return false;
+
+			bucket.pathTableSizeSent += toSend;
 		}
 		return true;
 	}
@@ -907,11 +901,12 @@ namespace uba
 			u32 toSend = Min(requiredCasTableSize - bucket.casKeyTableSizeSent, u32(m_client.GetMessageMaxSize() - 32));
 			left -= toSend;
 			writer.WriteBytes(bucket.sendCasKeyTable.GetMemory() + bucket.casKeyTableSizeSent, toSend);
-			bucket.casKeyTableSizeSent += toSend;
 
 			StackBinaryReader<16> reader;
 			if (!msg.Send(reader))
 				return false;
+
+			bucket.casKeyTableSizeSent += toSend;
 		}
 		return true;
 	}
@@ -1052,7 +1047,7 @@ namespace uba
 				u64 sizeToSend = block.writtenSize;
 
 				if (!sender.SendFileCompressed(casKey, path.data, dataToSend, sizeToSend, TC("SendCacheEntry")))
-					return m_logger.Error(TC("Failed to send cas content for file %s"), path.data);
+					return m_logger.Warning(TC("Failed to send cas content for file %s"), path.data);
 
 				outBytesSent += sender.m_bytesSent;
 			}
@@ -1069,10 +1064,7 @@ namespace uba
 			return true;
 
 		if (requiredCasTableOffset > 256*1024*1024)
-		{
-			m_logger.Warning(TC("Cas table offset %u too large. Cache entry corrupt (Bucket %llu)"), bucket.id);
-			return false;
-		}
+			return m_logger.Warning(TC("Cas table offset %u too large. Cache entry corrupt (Bucket %llu)"), bucket.id);
 
 		TimerScope ts2(stats.fetchCasTable);
 		u32 requiredPathTableOffset = 0;
diff --git a/Engine/Source/Programs/UnrealBuildAccelerator/Common/Private/UbaCacheServer.cpp b/Engine/Source/Programs/UnrealBuildAccelerator/Common/Private/UbaCacheServer.cpp
--- a/Engine/Source/Programs/UnrealBuildAccelerator/Common/Private/UbaCacheServer.cpp
+++ b/Engine/Source/Programs/UnrealBuildAccelerator/Common/Private/UbaCacheServer.cpp
@@ -1494,6 +1494,9 @@ namespace uba
 				return m_logger.Error(TC("Client %u has provided corrupt key or path table (CasKey: %s ClientVersion: %u PathtableVersion: %u)"), clientId, CasKeyString(casKey).str, clientVersion, connectionBucket.pathTable.GetVersion());
 			UBA_ASSERT(path.count);
 
+			if (casKey == CasKeyZero)
+				return m_logger.Error(TC("Client %u has provided cas key that is zero. Corrupt data? (Offset: %u TableSize: %u ClientVersion: %u PathtableVersion: %u)"), clientId, offset, connectionBucket.casKeyTable.GetSize(), clientVersion, connectionBucket.pathTable.GetVersion());
+
 			u32 pathOffset = bucket.m_pathTable.Add(path.data, path.count);
 
 			#if 0
@@ -1506,10 +1509,8 @@ namespace uba
 			u32 casKeyOffset = bucket.m_casKeyTable.Add(casKey, pathOffset, requiredCasTableSize);
 			auto insres = inputs.insert(casKeyOffset);
 			if (!insres.second)
-			{
-				m_logger.Warning(TC("Client %u Input file %s exists more than once in cache entry (CasKey: %s ClientVersion: %u PathtableVersion: %u)"), clientId, path.data, CasKeyString(casKey).str, clientVersion, connectionBucket.pathTable.GetVersion());
-				continue;
-			}
+				return m_logger.Warning(TC("Client %u Input file %s exists more than once in cache entry (CasKey: %s ClientVersion: %u PathtableVersion: %u)"), clientId, path.data, CasKeyString(casKey).str, clientVersion, connectionBucket.pathTable.GetVersion());
+
 			bytesForInput += Get7BitEncodedCount(casKeyOffset);
 
 			//m_logger.Info(TC("%s - %s"), path.data, CasKeyString(casKey).str);
