[UBA]
* Changed so cache client logic for sending files to cache server uses a downgraded logger in order to not cause build errors in case it fails to upload files to cache

diff --git a/Engine/Source/Programs/UnrealBuildAccelerator/Common/Private/UbaCacheClient.cpp b/Engine/Source/Programs/UnrealBuildAccelerator/Common/Private/UbaCacheClient.cpp
--- a/Engine/Source/Programs/UnrealBuildAccelerator/Common/Private/UbaCacheClient.cpp
+++ b/Engine/Source/Programs/UnrealBuildAccelerator/Common/Private/UbaCacheClient.cpp
@@ -20,6 +20,13 @@
 
 namespace uba
 {
+	struct DowngradedLogger : public LoggerWithWriter
+	{
+		DowngradedLogger(Atomic<bool>& c, LogWriter& writer, const tchar* prefix) : LoggerWithWriter(writer, prefix), connected(c) {}
+		virtual void Log(LogEntryType type, const tchar* str, u32 strLen) override { if (connected) LoggerWithWriter::Log(Max(type, LogEntryType_Info), str, strLen); }
+		Atomic<bool>& connected;
+	};
+
 	void CacheClientCreateInfo::Apply(Config& config, const tchar* tableName)
 	{
 		const ConfigTable* tablePtr = config.GetTable(tableName);
@@ -672,13 +679,6 @@ namespace uba
 
 				// Fetch output files from cache (and some files need to be "denormalized" before written to disk
 
-				struct DowngradedLogger : public LoggerWithWriter
-				{
-					DowngradedLogger(Atomic<bool>& c, LogWriter& writer, const tchar* prefix) : LoggerWithWriter(writer, prefix), connected(c) {}
-					virtual void Log(LogEntryType type, const tchar* str, u32 strLen) override { if (connected) LoggerWithWriter::Log(Max(type, LogEntryType_Info), str, strLen); }
-					Atomic<bool>& connected;
-				};
-
 				while (outputsReader.GetLeft())
 				{
 					u32 casKeyOffset = u32(outputsReader.Read7BitEncoded());
@@ -953,6 +953,8 @@ namespace uba
 				return msg.Send(reader);
 			});
 
+		DowngradedLogger logger(m_connected, m_logger.m_writer, TC("UbaCacheClientUpload"));
+
 		// There is content we need to upload to server
 		while (reader.GetLeft())
 		{
@@ -995,7 +997,7 @@ namespace uba
 					fileSize = file.GetSize();
 				}
 
-				if (!SendFile(m_logger, m_client, casKey, fileData, fileSize, casKeyFileName.data))
+				if (!SendFile(logger, m_client, casKey, fileData, fileSize, casKeyFileName.data))
 					return false;
 
 				outBytesSent += fileSize;
@@ -1041,7 +1043,7 @@ namespace uba
 
 
 				auto& s = m_storage;
-				FileSender sender { m_logger, m_client, s.m_bufferSlots, s.Stats(), m_sendOneAtTheTimeLock, s.m_casCompressor, s.m_casCompressionLevel };
+				FileSender sender { logger, m_client, s.m_bufferSlots, s.Stats(), m_sendOneAtTheTimeLock, s.m_casCompressor, s.m_casCompressionLevel };
 
 				u8* dataToSend = block.memory;
 				u64 sizeToSend = block.writtenSize;
