[UBA]
* Forcing cache server function to not be inlined in hopes of crash callstacks will contain function and line numbers

diff --git a/Engine/Source/Programs/UnrealBuildAccelerator/Common/Private/UbaCacheServer.cpp b/Engine/Source/Programs/UnrealBuildAccelerator/Common/Private/UbaCacheServer.cpp
--- a/Engine/Source/Programs/UnrealBuildAccelerator/Common/Private/UbaCacheServer.cpp
+++ b/Engine/Source/Programs/UnrealBuildAccelerator/Common/Private/UbaCacheServer.cpp
@@ -1337,7 +1337,7 @@ namespace uba
 		return Min(workerCountToUse, u32(m_buckets.size()));
 	}
 
-	bool CacheServer::HandleMessage(const ConnectionInfo& connectionInfo, u8 messageType, BinaryReader& reader, BinaryWriter& writer)
+	UBA_NOINLINE bool CacheServer::HandleMessage(const ConnectionInfo& connectionInfo, u8 messageType, BinaryReader& reader, BinaryWriter& writer)
 	{
 		if (messageType != CacheMessageType_Connect && m_isRunningMaintenance)
 			return m_logger.Error(TC("Can't handle network message %s while running maintenance mode"), ToString(CacheMessageType(messageType)));
