[UBA]
* Even more cache server validation of data coming from client

diff --git a/Engine/Source/Programs/UnrealBuildAccelerator/Common/Private/UbaCacheServer.cpp b/Engine/Source/Programs/UnrealBuildAccelerator/Common/Private/UbaCacheServer.cpp
--- a/Engine/Source/Programs/UnrealBuildAccelerator/Common/Private/UbaCacheServer.cpp
+++ b/Engine/Source/Programs/UnrealBuildAccelerator/Common/Private/UbaCacheServer.cpp
@@ -1488,6 +1488,9 @@ namespace uba
 			if (!isInput)
 				continue;
 
+			if (offset + sizeof(CasKey) >= connectionBucket.casKeyTable.GetSize())
+				return m_logger.Error(TC("Client %u has provided offset that is outside caskey table size. Corrupt data? (Offset: %u TableSize: %u ClientVersion: %u PathtableVersion: %u)"), clientId, offset, connectionBucket.casKeyTable.GetSize(), clientVersion, connectionBucket.pathTable.GetVersion());
+
 			CasKey casKey;
 			StringBuffer<> path;
 			if (!connectionBucket.casKeyTable.GetPathAndKey(path, casKey, connectionBucket.pathTable, offset))
