[UBA]
* Changed so NetworkClient is trying to find connection that is still connected when sending message.. this will handle the scenario where the network client has two sockets and one disconnects even though there is not a full disconnect (I don't know if this can happen but in case, lol)

diff --git a/Engine/Source/Programs/UnrealBuildAccelerator/Common/Private/UbaCacheClient.cpp b/Engine/Source/Programs/UnrealBuildAccelerator/Common/Private/UbaCacheClient.cpp
--- a/Engine/Source/Programs/UnrealBuildAccelerator/Common/Private/UbaCacheClient.cpp
+++ b/Engine/Source/Programs/UnrealBuildAccelerator/Common/Private/UbaCacheClient.cpp
@@ -1151,12 +1151,15 @@ namespace uba
 	{
 		u32 tableSize = bucket.serverCasKeyTableSize;
 		// CasKeyTable is 7bitEncoded(pathoffset) + CasKey... path table offset is minimum 1 byte
-		u32 neededSize = requiredCasTableOffset + 1 + sizeof(CasKey);
-		if (neededSize > tableSize)
+		u32 neededSizeMin = requiredCasTableOffset + 1 + sizeof(CasKey);
+		u32 neededSizeMax = neededSizeMin + 4;
+		if (neededSizeMax <= tableSize)
+			return true;
+		if (neededSizeMin > tableSize)
 			return false;
 		BinaryReader r(bucket.serverCasKeyTable.GetMemory(), requiredCasTableOffset, tableSize);
 		outPathOffset = u32(r.Read7BitEncoded());
-		neededSize = requiredCasTableOffset + Get7BitEncodedCount(outPathOffset) + sizeof(CasKey);
+		u32 neededSize = requiredCasTableOffset + Get7BitEncodedCount(outPathOffset) + sizeof(CasKey);
 		return neededSize <= tableSize;
 	}
 
diff --git a/Engine/Source/Programs/UnrealBuildAccelerator/Common/Private/UbaCompactTables.cpp b/Engine/Source/Programs/UnrealBuildAccelerator/Common/Private/UbaCompactTables.cpp
--- a/Engine/Source/Programs/UnrealBuildAccelerator/Common/Private/UbaCompactTables.cpp
+++ b/Engine/Source/Programs/UnrealBuildAccelerator/Common/Private/UbaCompactTables.cpp
@@ -387,12 +387,13 @@ namespace uba
 	{
 		if (!m_pathTableMem.memory)
 			m_pathTableMem.Init(PathTableMaxSize);
-		return (u8*)m_pathTableMem.AllocateNoLock(size, 1, TC("CompactPathTable"));
+		m_pathTableMem.ReserveNoLock(size, TC("CompactPathTable::BeginCommit"));
+		return (u8*)m_pathTableMem.memory + m_pathTableMem.writtenSize;
 	}
 
 	void CompactPathTable::EndCommit(u8* data, u64 written)
 	{
-		m_pathTableMem.writtenSize = data - m_pathTableMem.memory + written;
+		m_pathTableMem.AllocateNoLock(written, 1, TC("CompactPathTable::EndCommit"));
 	}
 
 	void CompactPathTable::Swap(CompactPathTable& other)
@@ -632,13 +633,13 @@ namespace uba
 			m_casKeyTableMem.Init(CasKeyTableMaxSize);
 		if (!m_offsets.IsInitialized())
 			m_offsets.Init(CasKeyLookupInitialCount);
-		void* mem = m_casKeyTableMem.AllocateNoLock(size, 1, TC("CompactCasKeyTable"));
-		return (u8*)mem;
+		m_casKeyTableMem.ReserveNoLock(size, TC("CompactCasKeyTable::BeginCommit"));
+		return (u8*)m_casKeyTableMem.memory + m_casKeyTableMem.writtenSize;
 	}
 
 	void CompactCasKeyTable::EndCommit(u8* data, u64 written)
 	{
-		m_casKeyTableMem.writtenSize = data - m_casKeyTableMem.memory + written;
+		m_casKeyTableMem.AllocateNoLock(written, 1, TC("CompactCasKeyTable::EndCommit"));
 	}
 
 	void CompactCasKeyTable::Debug(const CompactPathTable& pathTable)
diff --git a/Engine/Source/Programs/UnrealBuildAccelerator/Common/Private/UbaNetworkClient.cpp b/Engine/Source/Programs/UnrealBuildAccelerator/Common/Private/UbaNetworkClient.cpp
--- a/Engine/Source/Programs/UnrealBuildAccelerator/Common/Private/UbaNetworkClient.cpp
+++ b/Engine/Source/Programs/UnrealBuildAccelerator/Common/Private/UbaNetworkClient.cpp
@@ -472,13 +472,26 @@ namespace uba
 			return false;
 		}
 
-		Connection& connection = *m_connectionsIt;
-		++m_connectionsIt;
-		if (m_connectionsIt == m_connections.end())
-			m_connectionsIt = m_connections.begin();
+		// Skip connections that has disconnected
+		Connection* connectionPtr = &*m_connectionsIt;
+		Connection* connectionPtrStart = connectionPtr;
+		while (true)
+		{
+			++m_connectionsIt;
+			if (m_connectionsIt == m_connections.end())
+				m_connectionsIt = m_connections.begin();
+			if (connectionPtr->connected)
+				break;
+			connectionPtr = &*m_connectionsIt;
+			if (connectionPtr == connectionPtrStart)
+				break;
+		}
+
 		connectionItLock.Leave();
 		connectionLock.Leave();
 
+		Connection& connection = *connectionPtr;
+
 		message.m_response = response;
 		message.m_responseCapacity = responseCapacity;
 		message.m_connection = &connection;
