[UBA]
* Fixed bug in previous change where it earlied out before setting path table offset

diff --git a/Engine/Source/Programs/UnrealBuildAccelerator/Common/Private/UbaCacheClient.cpp b/Engine/Source/Programs/UnrealBuildAccelerator/Common/Private/UbaCacheClient.cpp
--- a/Engine/Source/Programs/UnrealBuildAccelerator/Common/Private/UbaCacheClient.cpp
+++ b/Engine/Source/Programs/UnrealBuildAccelerator/Common/Private/UbaCacheClient.cpp
@@ -1152,14 +1152,11 @@ namespace uba
 		u32 tableSize = bucket.serverCasKeyTableSize;
 		// CasKeyTable is 7bitEncoded(pathoffset) + CasKey... path table offset is minimum 1 byte
 		u32 neededSizeMin = requiredCasTableOffset + 1 + sizeof(CasKey);
-		u32 neededSizeMax = neededSizeMin + 4;
-		if (neededSizeMax <= tableSize)
-			return true;
 		if (neededSizeMin > tableSize)
 			return false;
 		BinaryReader r(bucket.serverCasKeyTable.GetMemory(), requiredCasTableOffset, tableSize);
 		outPathOffset = u32(r.Read7BitEncoded());
-		u32 neededSize = requiredCasTableOffset + Get7BitEncodedCount(outPathOffset) + sizeof(CasKey);
+		u32 neededSize = u32(r.GetPosition()) + sizeof(CasKey);
 		return neededSize <= tableSize;
 	}
 
