[UBA]
* Changed cache network client to have 1 worker. No need having lots of workers since they are not being used
* Changed so fetching cache hit files are done in parallel
* Changed so NetworkClient workercount=0 means 0.. ~0u means "same as number of cores"
* Set workercount to 0 for cache client and proxy client since none of them uses work

diff --git a/Engine/Source/Programs/UnrealBuildAccelerator/Common/Private/UbaCacheClient.cpp b/Engine/Source/Programs/UnrealBuildAccelerator/Common/Private/UbaCacheClient.cpp
--- a/Engine/Source/Programs/UnrealBuildAccelerator/Common/Private/UbaCacheClient.cpp
+++ b/Engine/Source/Programs/UnrealBuildAccelerator/Common/Private/UbaCacheClient.cpp
@@ -679,120 +679,141 @@ namespace uba
 
 				// Fetch output files from cache (and some files need to be "denormalized" before written to disk
 
+				Vector<u32> casKeyOffsets;
+
 				while (outputsReader.GetLeft())
+					casKeyOffsets.push_back(u32(outputsReader.Read7BitEncoded()));
+
+				if (!casKeyOffsets.empty())
 				{
-					u32 casKeyOffset = u32(outputsReader.Read7BitEncoded());
-					if (!FetchCasTable(bucket, cacheStats, casKeyOffset))
+					Atomic<bool> fetchSuccess = true;
+					m_storage.m_workManager->ParallelFor(u32(casKeyOffsets.size() - 1), casKeyOffsets, [&](auto& it)
+						{
+							u32 casKeyOffset = *it;
+							if (!FetchCasTable(bucket, cacheStats, casKeyOffset))
+							{
+								fetchSuccess = false;
+								return;
+							}
+							TimerScope fts(cacheStats.fetchOutput);
+							if (!FetchFile(bucket, rootPaths, info, cacheStats, storageStats, casKeyOffset))
+								fetchSuccess = false;
+						}, TC(""));
+					if (!fetchSuccess)
 						return false;
+				}
 
-					TimerScope fts(cacheStats.fetchOutput);
-
-					StringBuffer<MaxPath> path;
-					CasKey casKey;
-					if (!GetLocalPathAndCasKey(bucket, rootPaths, path, casKey, bucket.serverCasKeyTable, bucket.serverPathTable, casKeyOffset))
-						return false;
-					UBA_ASSERT(IsCompressed(casKey));
+				outResult.hit = true;
+				success = true;
+				return true;
+			}
+		}
 
-					FileFetcher fetcher { m_storage.m_bufferSlots, storageStats };
-					fetcher.m_errorOnFail = false;
+		for (auto& miss : misses)
+			m_logger.Info(TC("Cache miss on %s because of mismatch of %s (entry: %u, local: %s cache: %s bucket: %u)"), info.GetDescription(), miss.path.data(), miss.entryIndex, CasKeyString(miss.local).str, CasKeyString(miss.cache).str, bucketId);
 
-					if (IsNormalized(casKey))
-					{
-						DowngradedLogger logger(m_connected, m_logger.m_writer, TC("UbaCacheClientNormalizedDownload"));
-						// Fetch into memory, file is in special format without absolute paths
-						MemoryBlock normalizedBlock(4*1024*1024);
-						bool destinationIsCompressed = false;
-						if (!fetcher.RetrieveFile(logger, m_client, casKey, path.data, destinationIsCompressed, &normalizedBlock))
-							return logger.Error(TC("Failed to download cache output for %s"), info.GetDescription());
-
-						MemoryBlock localBlock(4*1024*1024);
-
-						u32 rootOffsets = *(u32*)(normalizedBlock.memory);
-						char* fileStart = (char*)(normalizedBlock.memory + sizeof(u32));
-						UBA_ASSERT(rootOffsets <= normalizedBlock.writtenSize);
-
-						// "denormalize" fetched file into another memory block that will be written to disk
-						u64 lastWritten = 0;
-						BinaryReader reader2(normalizedBlock.memory, rootOffsets, normalizedBlock.writtenSize);
-						while (reader2.GetLeft())
-						{
-							u64 rootOffset = reader2.Read7BitEncoded();
-							if (u64 toWrite = rootOffset - lastWritten)
-								memcpy(localBlock.Allocate(toWrite, 1, TC("")), fileStart + lastWritten, toWrite);
-							u8 rootIndex = fileStart[rootOffset] - RootPaths::RootStartByte;
-							const TString& root = rootPaths.GetRoot(rootIndex);
-							if (root.empty())
-								return logger.Error(TC("Cache entry uses root path index %u which is not set for this startupinfo (%s)"), rootIndex, info.GetDescription());
-
-							#if PLATFORM_WINDOWS
-							StringBuffer<> pathTemp;
-							pathTemp.Append(root);
-							char rootPath[512];
-							u32 rootPathLen = pathTemp.Parse(rootPath, sizeof_array(rootPath));
-							UBA_ASSERT(rootPathLen);
-							--rootPathLen;
-							#else
-							const char* rootPath = root.data();
-							u32 rootPathLen = root.size();
-							#endif
-
-							if (u32 toWrite = rootPathLen)
-								memcpy(localBlock.Allocate(toWrite, 1, TC("")), rootPath, toWrite);
-							lastWritten = rootOffset + 1;
-						}
+		return false;
+	}
 
-						u64 fileSize = rootOffsets - sizeof(u32);
-						if (u64 toWrite = fileSize - lastWritten)
-							memcpy(localBlock.Allocate(toWrite, 1, TC("")), fileStart + lastWritten, toWrite);
+	bool CacheClient::FetchFile(Bucket& bucket, const RootPaths& rootPaths, const ProcessStartInfo& info, CacheStats& cacheStats, StorageStats& storageStats, u32 casKeyOffset)
+	{
+		StringBuffer<MaxPath> path;
+		CasKey casKey;
+		if (!GetLocalPathAndCasKey(bucket, rootPaths, path, casKey, bucket.serverCasKeyTable, bucket.serverPathTable, casKeyOffset))
+			return false;
+		UBA_ASSERT(IsCompressed(casKey));
 
-						FileAccessor destFile(logger, path.data);
+		FileFetcher fetcher { m_storage.m_bufferSlots, storageStats };
+		fetcher.m_errorOnFail = false;
 
-						bool useFileMapping = true;
-						if (useFileMapping)
-						{
-							if (!destFile.CreateMemoryWrite(false, DefaultAttributes(), localBlock.writtenSize))
-								return logger.Error(TC("Failed to create file for cache output %s for %s"), path.data, info.GetDescription());
-							MapMemoryCopy(destFile.GetData(), localBlock.memory, localBlock.writtenSize);
-						}
-						else
-						{
-							if (!destFile.CreateWrite())
-								return logger.Error(TC("Failed to create file for cache output %s for %s"), path.data, info.GetDescription());
-							if (!destFile.Write(localBlock.memory, localBlock.writtenSize))
-								return false;
-						}
-						if (!destFile.Close(&fetcher.lastWritten))
-							return false;
+		if (IsNormalized(casKey))
+		{
+			DowngradedLogger logger(m_connected, m_logger.m_writer, TC("UbaCacheClientNormalizedDownload"));
+			// Fetch into memory, file is in special format without absolute paths
+			MemoryBlock normalizedBlock(4*1024*1024);
+			bool destinationIsCompressed = false;
+			if (!fetcher.RetrieveFile(logger, m_client, casKey, path.data, destinationIsCompressed, &normalizedBlock))
+				return logger.Error(TC("Failed to download cache output for %s"), info.GetDescription());
+
+			MemoryBlock localBlock(4*1024*1024);
+
+			u32 rootOffsets = *(u32*)(normalizedBlock.memory);
+			char* fileStart = (char*)(normalizedBlock.memory + sizeof(u32));
+			UBA_ASSERT(rootOffsets <= normalizedBlock.writtenSize);
+
+			// "denormalize" fetched file into another memory block that will be written to disk
+			u64 lastWritten = 0;
+			BinaryReader reader2(normalizedBlock.memory, rootOffsets, normalizedBlock.writtenSize);
+			while (reader2.GetLeft())
+			{
+				u64 rootOffset = reader2.Read7BitEncoded();
+				if (u64 toWrite = rootOffset - lastWritten)
+					memcpy(localBlock.Allocate(toWrite, 1, TC("")), fileStart + lastWritten, toWrite);
+				u8 rootIndex = fileStart[rootOffset] - RootPaths::RootStartByte;
+				const TString& root = rootPaths.GetRoot(rootIndex);
+				if (root.empty())
+					return logger.Error(TC("Cache entry uses root path index %u which is not set for this startupinfo (%s)"), rootIndex, info.GetDescription());
+
+				#if PLATFORM_WINDOWS
+				StringBuffer<> pathTemp;
+				pathTemp.Append(root);
+				char rootPath[512];
+				u32 rootPathLen = pathTemp.Parse(rootPath, sizeof_array(rootPath));
+				UBA_ASSERT(rootPathLen);
+				--rootPathLen;
+				#else
+				const char* rootPath = root.data();
+				u32 rootPathLen = root.size();
+				#endif
+
+				if (u32 toWrite = rootPathLen)
+					memcpy(localBlock.Allocate(toWrite, 1, TC("")), rootPath, toWrite);
+				lastWritten = rootOffset + 1;
+			}
 
-						fetcher.sizeOnDisk = localBlock.writtenSize;
-						casKey = CalculateCasKey(localBlock.memory, localBlock.writtenSize, false, nullptr, path.data);
-					}
-					else
-					{
-						DowngradedLogger logger(m_connected, m_logger.m_writer, TC("UbaCacheClientDownload"));
-						bool destinationIsCompressed = IsFileCompressed(info, path, true);
-						if (!fetcher.RetrieveFile(logger, m_client, casKey, path.data, destinationIsCompressed))
-							return logger.Error(TC("Failed to download cache output %s for %s"), path.data, info.GetDescription());
-					}
+			u64 fileSize = rootOffsets - sizeof(u32);
+			if (u64 toWrite = fileSize - lastWritten)
+				memcpy(localBlock.Allocate(toWrite, 1, TC("")), fileStart + lastWritten, toWrite);
 
-					cacheStats.fetchBytesRaw += fetcher.sizeOnDisk;
-					cacheStats.fetchBytesComp += fetcher.bytesReceived;
+			FileAccessor destFile(logger, path.data);
 
-					if (!m_storage.FakeCopy(casKey, path.data, fetcher.sizeOnDisk, fetcher.lastWritten, false))
-						return false;
-					if (!m_session.RegisterNewFile(path.data))
-						return false;
-				}
-				outResult.hit = true;
-				success = true;
-				return true;
+			bool useFileMapping = true;
+			if (useFileMapping)
+			{
+				if (!destFile.CreateMemoryWrite(false, DefaultAttributes(), localBlock.writtenSize))
+					return logger.Error(TC("Failed to create file for cache output %s for %s"), path.data, info.GetDescription());
+				MapMemoryCopy(destFile.GetData(), localBlock.memory, localBlock.writtenSize);
+			}
+			else
+			{
+				if (!destFile.CreateWrite())
+					return logger.Error(TC("Failed to create file for cache output %s for %s"), path.data, info.GetDescription());
+				if (!destFile.Write(localBlock.memory, localBlock.writtenSize))
+					return false;
 			}
+			if (!destFile.Close(&fetcher.lastWritten))
+				return false;
+
+			fetcher.sizeOnDisk = localBlock.writtenSize;
+			casKey = CalculateCasKey(localBlock.memory, localBlock.writtenSize, false, nullptr, path.data);
+		}
+		else
+		{
+			DowngradedLogger logger(m_connected, m_logger.m_writer, TC("UbaCacheClientDownload"));
+			bool destinationIsCompressed = IsFileCompressed(info, path, true);
+			if (!fetcher.RetrieveFile(logger, m_client, casKey, path.data, destinationIsCompressed))
+				return logger.Error(TC("Failed to download cache output %s for %s"), path.data, info.GetDescription());
 		}
 
-		for (auto& miss : misses)
-			m_logger.Info(TC("Cache miss on %s because of mismatch of %s (entry: %u, local: %s cache: %s bucket: %u)"), info.GetDescription(), miss.path.data(), miss.entryIndex, CasKeyString(miss.local).str, CasKeyString(miss.cache).str, bucketId);
+		cacheStats.fetchBytesRaw += fetcher.sizeOnDisk;
+		cacheStats.fetchBytesComp += fetcher.bytesReceived;
 
-		return false;
+		if (!m_storage.FakeCopy(casKey, path.data, fetcher.sizeOnDisk, fetcher.lastWritten, false))
+			return false;
+		if (!m_session.RegisterNewFile(path.data))
+			return false;
+
+		return true;
 	}
 
 	bool CacheClient::RequestServerShutdown(const tchar* reason)
diff --git a/Engine/Source/Programs/UnrealBuildAccelerator/Common/Private/UbaNetworkClient.cpp b/Engine/Source/Programs/UnrealBuildAccelerator/Common/Private/UbaNetworkClient.cpp
--- a/Engine/Source/Programs/UnrealBuildAccelerator/Common/Private/UbaNetworkClient.cpp
+++ b/Engine/Source/Programs/UnrealBuildAccelerator/Common/Private/UbaNetworkClient.cpp
@@ -21,7 +21,7 @@ namespace uba
 	}
 
 	NetworkClient::NetworkClient(bool& outCtorSuccess, const NetworkClientCreateInfo& info, const tchar* name)
-	: WorkManagerImpl(info.workerCount == 0 ? GetLogicalProcessorCount() : info.workerCount, TC("UbaWorker(NetworkClient)"))
+	: WorkManagerImpl(info.workerCount == ~0u ? GetLogicalProcessorCount() : info.workerCount, TC("UbaWorker(NetworkClient)"))
 	,	m_logWriter(info.logWriter)
 	,	m_logger(info.logWriter, SetGetPrefix(name))
 	,	m_isConnected(true)
diff --git a/Engine/Source/Programs/UnrealBuildAccelerator/Common/Private/UbaNetworkServer.cpp b/Engine/Source/Programs/UnrealBuildAccelerator/Common/Private/UbaNetworkServer.cpp
--- a/Engine/Source/Programs/UnrealBuildAccelerator/Common/Private/UbaNetworkServer.cpp
+++ b/Engine/Source/Programs/UnrealBuildAccelerator/Common/Private/UbaNetworkServer.cpp
@@ -868,6 +868,7 @@ namespace uba
 			worker->m_context->connection = nullptr;
 			worker->m_context->workAvailable.Set();
 		}
+		UBA_ASSERT(m_createdWorkerCount);
 	}
 
 	void NetworkServer::DoWork(u32 count)
diff --git a/Engine/Source/Programs/UnrealBuildAccelerator/Common/Private/UbaStorageClient.cpp b/Engine/Source/Programs/UnrealBuildAccelerator/Common/Private/UbaStorageClient.cpp
--- a/Engine/Source/Programs/UnrealBuildAccelerator/Common/Private/UbaStorageClient.cpp
+++ b/Engine/Source/Programs/UnrealBuildAccelerator/Common/Private/UbaStorageClient.cpp
@@ -399,6 +399,7 @@ namespace uba
 					m_lastTestedProxyIp = proxyHost.data;
 
 					NetworkClientCreateInfo ncci(m_logger.m_writer);
+					ncci.workerCount = 0;
 					bool ctorSuccess = true;
 					proxy = new ProxyClient(ctorSuccess, ncci);
 
diff --git a/Engine/Source/Programs/UnrealBuildAccelerator/Common/Private/UbaWorkManager.cpp b/Engine/Source/Programs/UnrealBuildAccelerator/Common/Private/UbaWorkManager.cpp
--- a/Engine/Source/Programs/UnrealBuildAccelerator/Common/Private/UbaWorkManager.cpp
+++ b/Engine/Source/Programs/UnrealBuildAccelerator/Common/Private/UbaWorkManager.cpp
@@ -107,6 +107,7 @@ namespace uba
 
 	void WorkManagerImpl::AddWork(const WorkFunction& work, u32 count, const tchar* desc, bool highPriority)
 	{
+		UBA_ASSERT(!m_workers.empty());
 		SCOPED_WRITE_LOCK(m_workLock, lock);
 		bool trackWork = m_workTracker.load();
 		for (u32 i = 0; i != count; ++i)
diff --git a/Engine/Source/Programs/UnrealBuildAccelerator/Common/Public/UbaCacheClient.h b/Engine/Source/Programs/UnrealBuildAccelerator/Common/Public/UbaCacheClient.h
--- a/Engine/Source/Programs/UnrealBuildAccelerator/Common/Public/UbaCacheClient.h
+++ b/Engine/Source/Programs/UnrealBuildAccelerator/Common/Public/UbaCacheClient.h
@@ -18,6 +18,7 @@ namespace uba
 	struct CacheStats;
 	struct CasKey;
 	struct ProcessStartInfo;
+	struct StorageStats;
 	struct StringView;
 	using RootsHandle = u64;
 
@@ -74,6 +75,7 @@ namespace uba
 		bool SendCasTable(Bucket& bucket, u32 requiredCasTableSize);
 		bool SendCacheEntry(Bucket& bucket, const RootPaths& rootPaths, const CasKey& cmdKey, const Map<u32, u32>& inputsStringToCasKey, const Map<u32, u32>& outputsStringToCasKey, const u8* logLines, u64 logLinesSize, u64& outBytesSent);
 		bool FetchCasTable(Bucket& bucket, CacheStats& stats, u32 requiredCasTableOffset);
+		bool FetchFile(Bucket& bucket, const RootPaths& rootPaths, const ProcessStartInfo& info, CacheStats& cacheStats, StorageStats& storageStats, u32 casKeyOffset);
 		template<typename TableType> bool FetchCompactTable(u32 bucketId, TableType& table, u32 requiredTableOffset, u8 messageType);
 		bool HasEnoughData(Bucket& bucket, u32 requiredCasTableOffset);
 		bool HasEnoughCasData(Bucket& bucket, u32 requiredCasTableOffset, u32& outPathOffset);
diff --git a/Engine/Source/Programs/UnrealBuildAccelerator/Common/Public/UbaNetworkClient.h b/Engine/Source/Programs/UnrealBuildAccelerator/Common/Public/UbaNetworkClient.h
--- a/Engine/Source/Programs/UnrealBuildAccelerator/Common/Public/UbaNetworkClient.h
+++ b/Engine/Source/Programs/UnrealBuildAccelerator/Common/Public/UbaNetworkClient.h
@@ -25,7 +25,7 @@ namespace uba
 		LogWriter& logWriter;
 		u32 sendSize = SendDefaultSize;
 		u32 receiveTimeoutSeconds = 0;
-		u32 workerCount = 0; // Zero means logical processor count will be used
+		u32 workerCount = ~0u; // ~0u means logical processor count will be used
 		u32 desiredConnectionCount = 1; // Desired connection count to server when connected
 		const u8* cryptoKey128 = nullptr;
 		void Apply(Config& config, const tchar* tableName = TC("NetworkClient"));
diff --git a/Engine/Source/Programs/UnrealBuildAccelerator/Host/Private/UbaExports.cpp b/Engine/Source/Programs/UnrealBuildAccelerator/Host/Private/UbaExports.cpp
--- a/Engine/Source/Programs/UnrealBuildAccelerator/Host/Private/UbaExports.cpp
+++ b/Engine/Source/Programs/UnrealBuildAccelerator/Host/Private/UbaExports.cpp
@@ -705,6 +705,7 @@ extern "C"
 		NetworkClientCreateInfo ncci(writer);
 		ncci.receiveTimeoutSeconds = 60;
 		ncci.cryptoKey128 = crypto128;
+		ncci.workerCount = 0;
 		ncci.Apply(GetConfig(), TC("CacheNetworkClient"));
 		bool ctorSuccess = false;
 		auto networkClient = new NetworkClientWithBackend(ctorSuccess, ncci, server.backend, TC("UbaCache"));
diff --git a/Engine/Source/Programs/UnrealBuildAccelerator/Visualizer/Private/UbaVisualizer.cpp b/Engine/Source/Programs/UnrealBuildAccelerator/Visualizer/Private/UbaVisualizer.cpp
--- a/Engine/Source/Programs/UnrealBuildAccelerator/Visualizer/Private/UbaVisualizer.cpp
+++ b/Engine/Source/Programs/UnrealBuildAccelerator/Visualizer/Private/UbaVisualizer.cpp
@@ -331,7 +331,7 @@ namespace uba
 			{
 				bool ctorSuccess = true;
 				NetworkClientCreateInfo ncci;
-				ncci.workerCount = 1;
+				ncci.workerCount = 0;
 				m_client = new NetworkClient(ctorSuccess, ncci);
 				if (!ctorSuccess)
 					return false;
