[uba]
* Added more log information to help find rare cache server crash

diff --git a/Engine/Source/Programs/UnrealBuildAccelerator/Common/Private/UbaCacheClient.cpp b/Engine/Source/Programs/UnrealBuildAccelerator/Common/Private/UbaCacheClient.cpp
--- a/Engine/Source/Programs/UnrealBuildAccelerator/Common/Private/UbaCacheClient.cpp
+++ b/Engine/Source/Programs/UnrealBuildAccelerator/Common/Private/UbaCacheClient.cpp
@@ -955,7 +955,10 @@ namespace uba
 					writer.WriteBytes(logLines, logLinesSize);
 
 			if (!msg.Send(reader))
+			{
+				m_logger.Info(TC("Failed to send cache entry. CasTable: %u/%u PathTable: %u/%u"), bucket.casKeyTableSizeSent, bucket.sendCasKeyTable.GetSize(), bucket.pathTableSizeSent, bucket.sendPathTable.GetSize());
 				return false;
+			}
 		}
 
 		// Server has all content for caskeys.. upload is done
diff --git a/Engine/Source/Programs/UnrealBuildAccelerator/Common/Private/UbaCacheServer.cpp b/Engine/Source/Programs/UnrealBuildAccelerator/Common/Private/UbaCacheServer.cpp
--- a/Engine/Source/Programs/UnrealBuildAccelerator/Common/Private/UbaCacheServer.cpp
+++ b/Engine/Source/Programs/UnrealBuildAccelerator/Common/Private/UbaCacheServer.cpp
@@ -1494,7 +1494,11 @@ namespace uba
 			CasKey casKey;
 			StringBuffer<> path;
 			if (!connectionBucket.casKeyTable.GetPathAndKey(path, casKey, connectionBucket.pathTable, offset))
-				return m_logger.Error(TC("Client %u has provided corrupt key or path table (CasKey: %s ClientVersion: %u PathtableVersion: %u)"), clientId, CasKeyString(casKey).str, clientVersion, connectionBucket.pathTable.GetVersion());
+			{
+				BinaryReader reader2(connectionBucket.casKeyTable.GetMemory(), offset, connectionBucket.casKeyTable.GetSize());
+				u32 pathOffset = u32(reader2.Read7BitEncoded());
+				return m_logger.Error(TC("Client %u has provided corrupt key or path table (CasKey: %s Offset: %u TableSize: %u PathOffset: %u PathTableSize: %u ClientVersion: %u PathtableVersion: %u)"), clientId, CasKeyString(casKey).str, offset, connectionBucket.casKeyTable.GetSize(), pathOffset, connectionBucket.pathTable.GetSize(), clientVersion, connectionBucket.pathTable.GetVersion());
+			}
 			UBA_ASSERT(path.count);
 
 			if (casKey == CasKeyZero)
