[UBA]
* Fixed race condition related to cas table being loaded at the same time it is being read. Added a write lock around the load code

diff --git a/Engine/Source/Programs/UnrealBuildAccelerator/Common/Private/UbaStorage.cpp b/Engine/Source/Programs/UnrealBuildAccelerator/Common/Private/UbaStorage.cpp
--- a/Engine/Source/Programs/UnrealBuildAccelerator/Common/Private/UbaStorage.cpp
+++ b/Engine/Source/Programs/UnrealBuildAccelerator/Common/Private/UbaStorage.cpp
@@ -1306,35 +1306,38 @@ namespace uba
 		u32 fileTableSize = reader.ReadU32();
 		u32 casLookupSize = reader.ReadU32();
 
-		UBA_ASSERT(m_casLookup.empty());
-		UBA_ASSERT(m_fileTableLookup.empty());
-
-		m_fileTableLookup.reserve(fileTableSize);
-		for (u32 i=0; i!=fileTableSize; ++i)
 		{
-			StringKey fileNameKey = reader.ReadStringKey();
-			if (reader.GetPosition() + 24 > fileSize)
+			SCOPED_WRITE_LOCK(m_fileTableLookupLock, fileTableLock); // Note, this can be reached from InvalidateCachedFileInfo which is called directly from session, so use a lock
+
+			UBA_ASSERT(m_casLookup.empty());
+			UBA_ASSERT(m_fileTableLookup.empty());
+
+			m_fileTableLookup.reserve(fileTableSize);
+			for (u32 i=0; i!=fileTableSize; ++i)
 			{
-				m_fileTableLookup.clear();
-				return m_logger.Warning(TC("CasTable file %s is corrupt"), fileName.data);
+				StringKey fileNameKey = reader.ReadStringKey();
+				if (reader.GetPosition() + 24 > fileSize)
+				{
+					m_fileTableLookup.clear();
+					return m_logger.Warning(TC("CasTable file %s is corrupt"), fileName.data);
+				}
+				u64 size = reader.ReadU64();
+				u64 lastWritten = reader.ReadU64();
+				CasKey casKey = reader.ReadCasKey();
+				if (casKey != CasKeyZero)
+					casKey = AsCompressed(casKey, m_storeCompressed);
+				else if (size || lastWritten)
+					continue; // This should not happen
+
+				auto insres = m_fileTableLookup.try_emplace(fileNameKey);
+				FileEntry& fileEntry = insres.first->second;
+				fileEntry.verified = false;
+				fileEntry.size = size;
+				fileEntry.lastWritten = lastWritten;
+				fileEntry.casKey = casKey;
 			}
-			u64 size = reader.ReadU64();
-			u64 lastWritten = reader.ReadU64();
-			CasKey casKey = reader.ReadCasKey();
-			if (casKey != CasKeyZero)
-				casKey = AsCompressed(casKey, m_storeCompressed);
-			else if (size || lastWritten)
-				continue; // This should not happen
-
-			auto insres = m_fileTableLookup.try_emplace(fileNameKey);
-			FileEntry& fileEntry = insres.first->second;
-			fileEntry.verified = false;
-			fileEntry.size = size;
-			fileEntry.lastWritten = lastWritten;
-			fileEntry.casKey = casKey;
 		}
 
-
 		m_casLookup.reserve(casLookupSize);
 		CasEntry* prev = nullptr;
 		while (true)
