[UBA]
* Moved unit tests to cpp files instead of h files

diff --git a/Engine/Source/Programs/UnrealBuildAccelerator/Test/Private/UbaTestAll.h b/Engine/Source/Programs/UnrealBuildAccelerator/Test/Private/UbaTestAll.h
--- a/Engine/Source/Programs/UnrealBuildAccelerator/Test/Private/UbaTestAll.h
+++ b/Engine/Source/Programs/UnrealBuildAccelerator/Test/Private/UbaTestAll.h
@@ -2,20 +2,9 @@
 
 #pragma once
 
-#define CHECK_TRUE(x) \
-	if (!(x)) \
-		return logger.Error(TC("Failed %s (%s:%u)"), TC(#x), __FILE__, __LINE__);
-
-#include "UbaTestBasics.h"
-#include "UbaTestCache.h"
-#include "UbaTestConfig.h"
-#include "UbaTestCrypto.h"
-#include "UbaTestNetwork.h"
-#include "UbaTestScheduler.h"
-#include "UbaTestSession.h"
-#include "UbaTestStorage.h"
-#include "UbaTestStress.h"
-#include "UbaTestStdOut.h"
+#include "UbaFile.h"
+#include "UbaStringBuffer.h"
+#include "UbaTest.h"
 
 namespace uba
 {
@@ -114,6 +103,10 @@ namespace uba
 		UBA_LINUX_TESTS \
 		UBA_MAC_TESTS \
 
+	#define UBA_TEST(x) bool x(LoggerWithWriter& logger, const StringBufferBase& testRootDir);
+	UBA_TESTS
+	#undef UBA_TEST
+
 
 	#define UBA_TEST(x) \
 		if (!filter || Contains(TC(#x), filter)) \
