[UBA]
* Moved unit tests to cpp files instead of h files

diff --git a/Engine/Source/Programs/UnrealBuildAccelerator/Test/Private/UbaTest.h b/Engine/Source/Programs/UnrealBuildAccelerator/Test/Private/UbaTest.h
new file mode 100644
--- /dev/null
+++ b/Engine/Source/Programs/UnrealBuildAccelerator/Test/Private/UbaTest.h
@@ -0,0 +1,8 @@
+// Copyright Epic Games, Inc. All Rights Reserved.
+
+#pragma once
+
+#define CHECK_TRUE(x) \
+	if (!(x)) \
+		return logger.Error(TC("Failed %s (%s:%u)"), TC(#x), __FILE__, __LINE__);
+
diff --git a/Engine/Source/Programs/UnrealBuildAccelerator/Test/Private/UbaTestBasics.cpp b/Engine/Source/Programs/UnrealBuildAccelerator/Test/Private/UbaTestBasics.cpp
new file mode 100644
--- /dev/null
+++ b/Engine/Source/Programs/UnrealBuildAccelerator/Test/Private/UbaTestBasics.cpp
@@ -0,0 +1,513 @@
+// Copyright Epic Games, Inc. All Rights Reserved.
+
+#include "UbaBinaryReaderWriter.h"
+#include "UbaFileAccessor.h"
+#include "UbaPathUtils.h"
+#include "UbaPlatform.h"
+#include "UbaProcessUtils.h"
+#include "UbaEvent.h"
+#include "UbaFile.h"
+#include "UbaLogger.h"
+#include "UbaRootPaths.h"
+#include "UbaTest.h"
+#include "UbaThread.h"
+#include "UbaTimer.h"
+#include "UbaDirectoryIterator.h"
+
+#if PLATFORM_WINDOWS
+#include "UbaWinBinDependencyParser.h"
+#elif PLATFORM_MAC
+#include "UbaMacBinDependencyParser.h"
+#endif
+
+#define VA_ARGS(...) , ##__VA_ARGS__
+#define UBA_TEST_CHECK(expr, fmt, ...) if (!(expr)) return logger.Error(TC(fmt) VA_ARGS(__VA_ARGS__));
+
+
+namespace uba
+{
+	bool TestTime(LoggerWithWriter& logger, const StringBufferBase& rootDir)
+	{
+#if 0
+		LoggerWithWriter consoleLogger(g_consoleLogWriter); (void)consoleLogger;
+		u64 time1 = GetSystemTimeUs();
+		Sleep(1000);
+		u64 time2 = GetSystemTimeUs();
+		u64 ms = (time2 - time1) / 1000;
+		consoleLogger.Info(TC("Slept ms: %llu"), ms);
+
+		time1 = GetTime();
+		Sleep(1000);
+		time2 = GetTime();
+		ms = (time2 * 1000 / GetFrequency()) - (time1 * 1000 / GetFrequency());
+		consoleLogger.Info(TC("Slept ms: %llu"), ms);
+#endif
+
+		u64 seconds = 15;
+		u64 fileTime = GetSecondsAsFileTime(seconds);
+		u64 seconds2 = GetFileTimeAsSeconds(fileTime);
+		UBA_TEST_CHECK(seconds == seconds2, "GetSecondsAsFileTime does not match GetFileTimeAsSeconds");
+		return true;
+	}
+
+	bool TestEvents(LoggerWithWriter& logger, const StringBufferBase& rootDir)
+	{
+		for (u32 i=0; i!=2; ++i)
+		{
+			Event ev;
+			if (!ev.Create(true, i == 1))
+				return logger.Error(TC("Failed to create event"));
+
+			Thread t([&]()
+				{
+					Sleep(500);
+					//logger.Info(TC("Setting event"));
+					ev.Set();
+					Sleep(500);
+					return true;
+				});
+
+			if (ev.IsSet(1))
+				return logger.Error(TC("Event was set after 1ms timeout where it should take 500ms"));
+
+			if (ev.IsSet(0))
+				return logger.Error(TC("Event was set after no timeout where it should take 500ms"));
+
+			//logger.Info(TC("Waiting for event"));
+			if (!ev.IsSet(2000))
+				return logger.Error(TC("Event was not set after 2000ms where it should take 500ms"));
+			//logger.Info(TC("Event was set"));
+
+			if (t.Wait(0))
+				return logger.Error(TC("Thread wait timed out. Should already be done after 2000ms"));
+
+			if (!t.Wait(2000))
+				return logger.Error(TC("Thread wait did not timed out should be done after 2000ms"));
+
+			#if 0 // Long time test... disabled by default
+			u64 time = GetTime();
+			Event ev2(true);
+			ev2.IsSet(10 * 60 * 1000);
+			if (TimeToMs((GetTime() - time) < 9 * 60 * 1000))
+				return logger.Error(TC("Event timeout was way too fast"));
+			#endif
+		}
+
+		return true;
+	}
+
+	bool TestPaths(LoggerWithWriter& logger, const StringBufferBase& rootDir)
+	{
+		const tchar* workingDir = IsWindows ? TC("e:\\dev\\") : TC("/dev/bar/");
+		tchar buffer[1024];
+		u32 lengthResult;
+
+		auto TestPath = [&](const tchar* path) { return FixPath2(path, workingDir, TStrlen(workingDir), buffer, sizeof_array(buffer), &lengthResult); };
+
+#if PLATFORM_WINDOWS
+		if (!FixPath2(TC("\"e:\\temp\""), workingDir, TStrlen(workingDir), buffer, sizeof_array(buffer), &lengthResult))
+			return logger.Error(TC("FixPath2 (1) failed"));
+#else
+
+		if (!TestPath(TC("/..")))
+			return logger.Error(TC("FixPath2 should have failed"));
+		UBA_TEST_CHECK(Equals(buffer, TC("/")), "Should not contain ..");
+
+		if (!FixPath2(TC("/../Foo"), workingDir, TStrlen(workingDir), buffer, sizeof_array(buffer), &lengthResult))
+			return logger.Error(TC("FixPath2 should have failed"));
+		UBA_TEST_CHECK(Equals(buffer, TC("/Foo")), "Should not contain ..");
+
+		if (!FixPath2(TC("/usr/bin//clang++"), workingDir, TStrlen(workingDir), buffer, sizeof_array(buffer), &lengthResult))
+			return logger.Error(TC("FixPath2 should have failed"));
+		UBA_TEST_CHECK(!Contains(buffer, TC("//")), "Should not contain //");
+#endif
+
+		if (!FixPath2(TC("../Foo"), workingDir, TStrlen(workingDir), buffer, sizeof_array(buffer), &lengthResult))
+			return logger.Error(TC("FixPath2 (1) failed"));
+		UBA_TEST_CHECK(!Contains(buffer, TC("..")), "Should not contain ..");
+
+		if (!FixPath2(TC("@../Foo"), workingDir, TStrlen(workingDir), buffer, sizeof_array(buffer), &lengthResult))
+			return logger.Error(TC("FixPath2 (1) failed"));
+		UBA_TEST_CHECK(Contains(buffer, TC("..")), "Should contain ..");
+
+		if (!FixPath2(TC("..@/Foo"), workingDir, TStrlen(workingDir), buffer, sizeof_array(buffer), &lengthResult))
+			return logger.Error(TC("FixPath2 (1) failed"));
+		UBA_TEST_CHECK(Contains(buffer, TC("..")), "Should contain ..");
+
+		return true;
+	}
+
+	bool TestFiles(LoggerWithWriter& logger, const StringBufferBase& rootDir)
+	{
+		StringBuffer<> testFileName(rootDir);
+		testFileName.Append(TC("UbaTestFile"));
+
+		FileAccessor fileHandle(logger, testFileName.data);
+		if (!fileHandle.CreateWrite())
+			return logger.Error(TC("Failed to create file for write"));
+
+		u8 byte = 'H';
+		if (!fileHandle.Write(&byte, 1))
+			return false;
+
+		if (!fileHandle.Close())
+			return false;
+
+		FileHandle fileHandle2;
+		if (!OpenFileSequentialRead(logger, testFileName.data, fileHandle2))
+			return logger.Error(TC("Failed to create file for read"));
+
+		u64 writeTime = 0;
+		if (!GetFileLastWriteTime(writeTime, fileHandle2))
+			return logger.Error(TC("Failed to get last written time"));
+
+		u64 writeTime2 = 0;
+		TraverseDir(logger, rootDir, [&](const DirectoryEntry& de)
+			{
+				if (Equals(de.name, TC("UbaTestFile")))
+					writeTime2 = de.lastWritten;
+			});
+
+		if (writeTime != writeTime2)
+			return logger.Error(TC("GetFileLastWriteTime and TraverseDir are returning different last write time for same file"));
+
+		u64 systemTime = GetSystemTimeAsFileTime();
+		if (systemTime < writeTime)
+			return logger.Error(TC("System time is lower than last written time"));
+		if (GetFileTimeAsSeconds(systemTime) - GetFileTimeAsSeconds(writeTime) > 3)
+			return logger.Error(TC("System time or last written time is wrong (system: %llu, write: %llu, diffInSec: %llu)"), systemTime, writeTime, GetFileTimeAsSeconds(systemTime) - GetFileTimeAsSeconds(writeTime));
+
+
+		u8 byte2 = 0;
+		if (!ReadFile(logger, testFileName.data, fileHandle2, &byte2, 1))
+			return false;
+
+		if (!CloseFile(testFileName.data, fileHandle2))
+			return false;
+
+		FileHandle fileHandle3;
+		if (!OpenFileSequentialRead(logger, TC("NonExistingFile"), fileHandle3, false))
+			return logger.Error(TC("OpenFileSequentialRead failed with non existing file"));
+		if (fileHandle3 != InvalidFileHandle)
+			return logger.Error(TC("OpenFileSequentialRead found file that doesn't exist"));
+
+		if (RemoveDirectoryW(TC("TestDir")))
+			return logger.Error(TC("Did not fail to remove non-existing TestDir (or were things not cleaned before test)"));
+		else if (GetLastError() != ERROR_FILE_NOT_FOUND)
+			return logger.Error(TC("GetLastError did not return correct error failing to remove non-existing directory TestDir"));
+
+		if (!CreateDirectoryW(TC("TestDir")))
+			return logger.Error(TC("Failed to create dir"));
+
+		FileHandle fileHandle4;
+		if (OpenFileSequentialRead(logger, TC("TestDir"), fileHandle4))
+			return logger.Error(TC("This should return fail"));
+
+		if (!RemoveDirectoryW(TC("TestDir")))
+			return logger.Error(TC("Fail to remove TestDir"));
+
+		u64 size = 0;
+		if (!FileExists(logger, testFileName.data, &size) || size != 1)
+			return logger.Error(TC("UbaTestFile not found"));
+
+		StringBuffer<> testFileName2(rootDir);
+		testFileName2.Append(TC("UbaTestFile2"));
+
+		DeleteFileW(testFileName2.data);
+
+		if (DeleteFileW(testFileName2.data))
+			return logger.Error(TC("Did not fail to delete non-existing UbaTestFile2 (or were things not cleaned before test)"));
+		else if (GetLastError() != ERROR_FILE_NOT_FOUND)
+			return logger.Error(TC("GetLastError did not return correct error failing to delete non-existing file UbaTestFile2"));
+
+		if (!CreateHardLinkW(testFileName2.data, testFileName.data))
+			return logger.Error(TC("Failed to create hardlink from UbaTestFile to UbaTestFile2"));
+
+		if (!DeleteFileW(testFileName.data))
+			return logger.Error(TC("Failed to delete UbaTestFile"));
+
+		if (FileExists(logger, testFileName.data))
+			return logger.Error(TC("Found non-existing file UbaTestFile"));
+
+		// CreateHardLinkW is a symbolic link on non-windows.. need to revisit
+		#if PLATFORM_WINDOWS
+		if (!FileExists(logger, testFileName2.data))
+			return logger.Error(TC("Failed to find file UbaTestFile2"));
+
+		StringBuffer<> currentDir;
+		if (!GetCurrentDirectoryW(currentDir))
+			return logger.Error(TC("GetCurrentDirectoryW failed"));
+
+		bool foundFile = false;
+		if (!TraverseDir(logger, rootDir, [&](const DirectoryEntry& de) { foundFile |= TStrcmp(de.name, TC("UbaTestFile2")) == 0; }, true))
+			return logger.Error(TC("Failed to TraverseDir '.'"));
+
+		if (!foundFile)
+			return logger.Error(TC("Did not find UbaTestFile2 with TraverseDir"));
+
+		if (!DeleteFileW(testFileName2.data))
+			return false;
+		#endif
+
+		LoggerWithWriter nullLogger(g_nullLogWriter);
+		if (TraverseDir(nullLogger, AsView(TC("TestDir2")), [&](const DirectoryEntry&) {}, true))
+			return logger.Error(TC("TraverseDir failed to report fail on non existing dir"));
+
+		return true;
+	}
+
+	bool TestOverlappedIO(LoggerWithWriter& logger, const StringBufferBase& rootDir)
+	{
+		StringBuffer<> testFileName(rootDir);
+		testFileName.Append(TC("UbaTestFile"));
+
+		u64 left = 2*1024*1024;
+
+		FileAccessor fileHandle(logger, testFileName.data);
+		if (!fileHandle.CreateWrite(false, DefaultAttributes() | FILE_FLAG_OVERLAPPED, left))
+			return logger.Error(TC("Failed to create file for write"));
+		u8 buffer[512*1024];
+
+		while (left)
+		{
+			u64 toWrite = Min(left, 277872ull);
+			if (!fileHandle.Write(buffer, toWrite))
+				return logger.Error(TC("Failed to create file for write"));
+			left -= toWrite;
+		}
+		if (!fileHandle.Close())
+			return false;
+		return true;
+	}
+
+	bool TestMemoryBlock(LoggerWithWriter& logger, const StringBufferBase& rootDir)
+	{
+		{
+			MemoryBlock block(1024 * 1024);
+			u64* mem = (u64*)block.Allocate(8, 1, TC("Foo"));
+			*mem = 0x1234;
+			block.Free(mem);
+		}
+
+		if (GetHugePageCount())
+		{
+			MemoryBlock block;
+			if (!block.Init(1024 * 1024, nullptr, true))
+				return logger.Error(TC("Failed to allocate huge pages even though system says they exists"));
+			u64* mem = (u64*)block.Allocate(8, 1, TC("Foo"));
+			*mem = 0x1234;
+			block.Free(mem);
+		}
+
+		return true;
+	}
+
+	bool TestParseArguments(LoggerWithWriter& logger, const StringBufferBase& rootDir)
+	{
+		auto ParseArguments = [](Vector<TString>& a, const tchar* args) { uba::ParseArguments(args, [&](const tchar* arg, u32 argLen) { a.push_back({arg, argLen}); }); };
+
+		Vector<TString> arguments;
+		ParseArguments(arguments, TC("foo bar"));
+		UBA_TEST_CHECK(arguments.size() == 2, "ParseArguments 1 failed (%llu)", arguments.size());
+
+		Vector<TString> arguments2;
+		ParseArguments(arguments2, TC("\"foo\" bar"));
+		UBA_TEST_CHECK(arguments2.size() == 2, "ParseArguments 2 failed");
+
+		Vector<TString> arguments3;
+		ParseArguments(arguments3, TC("\"foo meh\" bar"));
+		UBA_TEST_CHECK(arguments3.size() == 2, "ParseArguments 3 failed");
+		UBA_TEST_CHECK(Contains(arguments3[0].data(), TC(" ")), "ParseArguments 3 failed");
+
+		Vector<TString> arguments4;
+		ParseArguments(arguments4, TC("\"app\" @\"rsp\""));
+		UBA_TEST_CHECK(arguments4.size() == 2, "ParseArguments 4 failed");
+		UBA_TEST_CHECK(!Contains(arguments4[1].data(), TC("\"")), "ParseArguments 4 failed");
+
+		Vector<TString> arguments5;
+		ParseArguments(arguments5, TC("\"app\" @\"rsp foo\""));
+		UBA_TEST_CHECK(arguments5.size() == 2, "ParseArguments 4 failed");
+		UBA_TEST_CHECK(!Contains(arguments5[1].data(), TC("\"")), "ParseArguments 5 failed");
+		UBA_TEST_CHECK(Contains(arguments5[1].data(), TC(" ")), "ParseArguments 5 failed");
+
+		Vector<TString> arguments6;
+		ParseArguments(arguments6, TC("\"app\"\"1\" @\"rsp foo\""));
+		UBA_TEST_CHECK(arguments6.size() == 2, "ParseArguments 6 failed");
+		UBA_TEST_CHECK(Equals(arguments6[0].data(), TC("app1")), "ParseArguments 6 failed");
+
+		Vector<TString> arguments7;
+		ParseArguments(arguments7, TC("app \" \\\"foo\\\" bar\""));
+		UBA_TEST_CHECK(arguments7.size() == 2, "ParseArguments 7 failed");
+		UBA_TEST_CHECK(Contains(arguments7[1].data(), TC("\"")), "ParseArguments 7 failed");
+
+		Vector<TString> arguments8;
+		ParseArguments(arguments8, TC("\nline1\r\nline2\r\nline3\n\r\n"));
+		UBA_TEST_CHECK(arguments8.size() == 3, "ParseArguments 8 failed");
+		UBA_TEST_CHECK(Equals(arguments8[0].data(), TC("line1")), "ParseArguments 8 failed");
+		UBA_TEST_CHECK(Equals(arguments8[1].data(), TC("line2")), "ParseArguments 8 failed");
+		UBA_TEST_CHECK(Equals(arguments8[2].data(), TC("line3")), "ParseArguments 8 failed");
+
+		Vector<TString> arguments9;
+		ParseArguments(arguments9, TC("\"foo\\\\\" \"bar\\\\\""));
+		UBA_TEST_CHECK(arguments9.size() == 2, "ParseArguments 9 failed");
+		UBA_TEST_CHECK(Equals(arguments9[0].data(), TC("foo\\\\")), "ParseArguments 9 failed");
+		UBA_TEST_CHECK(Equals(arguments9[1].data(), TC("bar\\\\")), "ParseArguments 9 failed");
+		return true;
+	}
+
+	bool TestBinaryWriter(LoggerWithWriter& logger, const StringBufferBase& rootDir)
+	{
+		auto testString = [&](const tchar* str)
+		{
+			u8 mem[1024];
+			BinaryWriter writer(mem);
+			writer.WriteString(str);
+			BinaryReader reader(mem);
+			TString s = reader.ReadString();
+			if (s.size() != TStrlen(str))
+				return logger.Error(TC("Serialized string '%s' has wrong strlen"), str);
+			if (s != str)
+				return logger.Error(TC("Serialized string '%s' is different from source"), str);
+			return true;
+		};
+
+		if (!testString(TC("Foo")))
+			return false;
+
+		#if PLATFORM_WINDOWS
+		tchar str1[] = { 54620, 44544, 0 };
+		if (!testString(str1))
+			return false;
+		tchar str2[] = { 'f', 54620, 'o', 44544, 0 };
+		if (!testString(str2))
+			return false;
+		#endif
+
+		return true;
+	}
+
+	#if PLATFORM_WINDOWS
+	bool TestKnownSystemFiles(LoggerWithWriter& logger, const StringBufferBase& rootDir)
+	{
+		for (auto systemFile : g_knownSystemFiles)
+			if (!IsKnownSystemFile(systemFile))
+				return logger.Error(TC("IsKnownSystemFile returned false for %s which is a system file"), systemFile);
+		if (IsKnownSystemFile(TC("Fooo.dll")))
+			return logger.Error(TC("IsKnownSystemFile returned true for Fooo.dll which is not a system file"));
+		return true;
+	}
+	#endif
+
+	bool TestRootPaths(LoggerWithWriter& logger, const StringBufferBase& rootDir)
+	{
+		#if PLATFORM_WINDOWS
+		const tchar root1[] = TC("c:\\temp\\");
+		const tchar root2[] = TC("e:\\temp\\");
+		const tchar str[] = TC("e:\\temp\\foo");
+		#else
+		const tchar root1[] = TC("/mnt/c/");
+		const tchar root2[] = TC("/mnt/e/");
+		const tchar str[] = TC("/mnt/e/foo");
+		#endif
+
+		RootPaths paths;
+		if (!paths.RegisterRoot(logger, root1))
+			return false;
+		if (!paths.RegisterRoot(logger, root2))
+			return false;
+
+		bool success = true;
+		StringBuffer<> temp;
+		u32 rootPos = ~0u;
+		bool res = paths.NormalizeString(logger, str, sizeof_array(str), [&](const tchar* str, u64 strLen, u32 rp)
+			{
+				if (rp != ~0u)
+				{
+					if (strLen != 1)
+						success = false;
+					if (str[0] != RootPaths::RootStartByte + PathsPerRoot + (IsWindows ? 1 : 0)) // Add one for windows because second entry is backslash
+						success = false;
+					rootPos = str[0];
+				}
+				else
+				{
+					temp.Append(str, strLen);
+					if (!temp.Equals(TC("foo")))
+						success = false;
+				}
+			}, false, TC(""));
+
+		if (!res)
+			return false;
+		if (!success)
+			return false;
+
+		StringBuffer<> newStr;
+		auto& root = paths.GetRoot(rootPos - RootPaths::RootStartByte);
+		newStr.Append(root.c_str()).Append(temp);
+		if (!newStr.Equals(str))
+			return false;
+
+		#if PLATFORM_WINDOWS
+		const tchar str2[] = TC("file://e:/temp/");
+		bool foundPath = false;
+		res = paths.NormalizeString(logger, str2, sizeof_array(str2), [&](const tchar* str, u64 strLen, u32 rp)
+			{
+				if (rp != ~0)
+				{
+					if (str[0] != RootPaths::RootStartByte + PathsPerRoot) // Add one for windows because second entry is backslash
+						success = false;
+					foundPath = true;
+				}
+				else
+				{
+					if (!((strLen == 1 && (!str[0] || str[0] == '/')) || (strLen == 6 && Equals(str, TC("file:/"), 6, false))))
+						success = false;
+				}
+
+			}, false, TC(""));
+		if (!res || !foundPath || !success)
+			return false;
+		#endif
+
+		return true;
+	}
+
+
+#if 0//PLATFORM_MAC
+	Set<TString> g_visited;
+	void LogImports(const tchar* import, bool isKnown)
+	{
+		if (!g_visited.insert(import).second)
+			return;
+		LoggerWithWriter logger(g_consoleLogWriter);
+		logger.Info(TC("IMPORT: %s"), import);
+		StringBuffer<> path(TC("[REDACTED_USER_HOME]/p4/fn/Engine/Binaries/Mac"));
+		path.EnsureEndsWithSlash().Append(TC(import));
+		StringBuffer<> error;
+		FindImportsMac(path.data, LogImports, error);
+	}
+#endif
+
+	bool TestBinDependencies(LoggerWithWriter& logger, const StringBufferBase& rootDir)
+	{
+
+#if PLATFORM_WINDOWS
+		StringBuffer<> path;
+		GetDirectoryOfCurrentModule(logger, path);
+		path.EnsureEndsWithSlash().Append(TC("UbaTestApp.exe"));
+		bool importKernel = false;
+		StringBuffer<> error;
+		FindImports(path.data, [&](const tchar* import, bool isKnown, const char* const* importLoaderPaths)
+		{
+			importKernel |= isKnown && Contains(import, TC("KERNEL32.dll"));
+		}, error);
+		if (!importKernel)
+			return logger.Error(TC("Failed to find Kernel32 as import"));
+#elif PLATFORM_MAC
+		//StringBuffer<> error;
+		//FindImportsMac(TC("[REDACTED_USER_HOME]/p4/fn/Engine/Binaries/Mac/ShaderCompileWorker"), LogImports, error);
+#endif
+		return true;
+	}
+}
diff --git a/Engine/Source/Programs/UnrealBuildAccelerator/Test/Private/UbaTestBasics.h b/Engine/Source/Programs/UnrealBuildAccelerator/Test/Private/UbaTestBasics.h
deleted file mode 100644
diff --git a/Engine/Source/Programs/UnrealBuildAccelerator/Test/Private/UbaTestCache.cpp b/Engine/Source/Programs/UnrealBuildAccelerator/Test/Private/UbaTestCache.cpp
new file mode 100644
--- /dev/null
+++ b/Engine/Source/Programs/UnrealBuildAccelerator/Test/Private/UbaTestCache.cpp
@@ -0,0 +1,391 @@
+// Copyright Epic Games, Inc. All Rights Reserved.
+
+#include "UbaCacheClient.h"
+#include "UbaCacheEntry.h"
+#include "UbaCacheServer.h"
+#include "UbaConfig.h"
+#include "UbaCompactTables.h"
+#include "UbaHashMap.h"
+#include "UbaNetworkBackendTcp.h"
+#include "UbaNetworkClient.h"
+#include "UbaNetworkServer.h"
+#include "UbaRootPaths.h"
+#include "UbaSessionServer.h"
+#include "UbaStorageServer.h"
+#include "UbaTest.h"
+
+namespace uba
+{
+	bool TestCacheEntry(LoggerWithWriter& logger, const StringBufferBase& rootDir)
+	{
+		CacheEntries entries;
+
+		auto AddEntry = [&](const Set<u32>& inputs)
+			{
+				Vector<u8> inputOffsets;
+				u64 bytes = 0;
+				for (u32 i : inputs)
+					bytes += Get7BitEncodedCount(i);
+				inputOffsets.resize(bytes);
+				BinaryWriter writer(inputOffsets.data(), 0, inputOffsets.size());
+				for (u32 i : inputs)
+					writer.Write7BitEncoded(i);
+
+				CacheEntry entry;
+				entries.BuildInputs(entry, inputs);
+				entries.entries.push_back(entry);
+				entries.ValidateEntry(logger, entry, inputOffsets);
+			};
+
+		AddEntry(Set<u32>{ 1, 4, 6 });
+		AddEntry(Set<u32>{ 0, 4, 6 });
+		AddEntry(Set<u32>{ 2, 4, 6 });
+		AddEntry(Set<u32>{ 1, 4, 5 });
+		AddEntry(Set<u32>{ 1, 4, 7 });
+		AddEntry(Set<u32>{ 1, 3, 6 });
+		AddEntry(Set<u32>{ 1, 5, 6 });
+		AddEntry(Set<u32>{ 1, 4, 6, 7 });
+		AddEntry(Set<u32>{ 0, 1, 4, 6 });
+
+		return true;
+	}
+
+	void GetTestAppPath(LoggerWithWriter& logger, StringBufferBase& out);
+	bool CreateTextFile(StringBufferBase& outPath, LoggerWithWriter& logger, const tchar* workingDir, const tchar* fileName, const char* text);
+	StringKey GetKeyAndFixedName(StringBuffer<>& fixedFilePath, const tchar* filePath);
+	void InvalidateCachedInfo(StorageImpl& storage, StringBufferBase& fileName)
+	{
+		StringBuffer<> fixedFilePath;
+		storage.InvalidateCachedFileInfo(GetKeyAndFixedName(fixedFilePath, fileName.data));
+	}
+
+	bool TestCompactPathTable(LoggerWithWriter& logger, const StringBufferBase& rootDir)
+	{
+		for (u32 version=0; version<=CacheBucketVersion; ++version)
+		{
+			CompactPathTable table(CaseInsensitiveFs, 0, 0, version);
+
+			StringBuffer<> str;
+			u32 offset;
+
+			auto testStr = [&](const StringView& str)
+				{
+					StringBuffer<> str2;
+					offset = table.Add(str.data, str.count);
+					table.GetString(str2, offset);
+					return str.Equals(str2.data);
+				};
+
+			if (!testStr(str.Append("foo")))
+				return false;
+			if (!testStr(str.Clear().Append("foo").EnsureEndsWithSlash().Append("bar.h")))
+				return false;
+			if (!testStr(str.Clear().Append(PathSeparator).Append("foo").Append(PathSeparator).Append("bar.h")))
+				return false;
+
+			CompactPathTable table2(CaseInsensitiveFs, 0, 0, version);
+			BinaryReader reader(table.GetMemory(), 0, table.GetSize());
+			table2.ReadMem(reader, true);
+			u32 offset2 = table2.Add(str.data, str.count);
+			if (offset != offset2)
+				return false;
+		}
+		return true;
+	}
+
+	bool TestCompactCasKeyTable(LoggerWithWriter& logger, const StringBufferBase& rootDir)
+	{
+		CompactCasKeyTable table;
+		u32 offset0 = table.AddNoLock(CasKeyZero, 0);
+		u32 offset1 = table.AddNoLock(CasKeyZero, 1);
+		u32 offset2 = table.AddNoLock(CasKeyZero, 2);
+
+		CompactCasKeyTable table2;
+		BinaryReader reader(table.GetMemory(), 0, table.GetSize());
+		table2.ReadMem(reader, true);
+		if (table2.AddNoLock(CasKeyZero, 0) != offset0)
+			return false;
+		if (table2.AddNoLock(CasKeyZero, 1) != offset1)
+			return false;
+		if (table2.AddNoLock(CasKeyZero, 2) != offset2)
+			return false;
+
+		for (u32 i=0;i!=32;++i)
+		{
+			CompactCasKeyTable table3(table.GetKeyCount());
+			for (u32 j=0;j!=i;++j)
+				table3.AddNoLock(CasKeyZero, j);
+		}
+		return true;
+	}
+
+	bool TestHashTable(LoggerWithWriter& logger, const StringBufferBase& rootDir)
+	{
+		{
+			MemoryBlock memoryBlock(1024*1024);
+			HashMap<u32, u32> casMap;
+			casMap.Init(memoryBlock, 3);
+			if (casMap.Find(1))
+				return false;
+			casMap.Insert(1) = 2;
+			if (*casMap.Find(1) != 2)
+				return false;
+			casMap.Insert(1) = 3;
+			if (*casMap.Find(1) != 3)
+				return false;
+
+			HashMap<u32, u32, true> casMap2;
+			casMap2.Init(4);
+			for (u32 i=0;i!=4; ++i)
+			{
+				casMap2.Insert(i) = i;
+				if (*casMap2.Find(i) != i)
+					return false;
+			}
+			for (u32 i=0;i!=4; ++i)
+				if (*casMap2.Find(i) != i)
+					return false;
+
+			casMap2.Insert(4) = 4;
+			for (u32 i=0;i!=5; ++i)
+				if (*casMap2.Find(i) != i)
+					return false;
+
+			for (u32 i=5;i!=1000; ++i)
+			{
+				if (casMap2.Find(i))
+					return false;
+				casMap2.Insert(i) = i;
+				if (*casMap2.Find(i) != i)
+					return false;
+			}
+		}
+
+#if 0
+		struct CasFileInfo { CasFileInfo(u32 s = 0) : size(s) {} u32 size; bool isUsed; }; // These are compressed cas, should never be over 4gb
+		constexpr u64 memoryReserveSize = 192*1024*1024;
+		struct ProfileScope
+		{
+			ProfileScope() : startTime(GetTime()) {}
+			~ProfileScope() { u64 duration = GetTime() - startTime; LoggerWithWriter(g_consoleLogWriter, TC("")).Info(TC("Time: %s"), TimeToText(duration).str); }
+			u64 startTime;
+		};
+
+		u32 totalCasCount = 1'800'000;
+
+		{
+			ProfileScope _;
+			MemoryBlock memoryBlock;
+			if (!memoryBlock.Init(memoryReserveSize, nullptr, true))
+				memoryBlock.Init(memoryReserveSize);
+
+			HashMap<CasKey, CasFileInfo> casMap;
+			casMap.Init(memoryBlock, totalCasCount);
+
+			CasKey key;
+
+			for (u32 i=0; i!=totalCasCount; ++i)
+			{
+				key.a = i;
+				casMap.Insert(key);
+			}
+
+			for (u32 i=0; i!=totalCasCount; ++i)
+			{
+				key.a = i;
+				casMap.Find(key);
+			}
+		}
+
+		{
+			ProfileScope _;
+			MemoryBlock memoryBlock;
+			if (!memoryBlock.Init(memoryReserveSize, nullptr, true))
+				memoryBlock.Init(memoryReserveSize);
+
+			GrowingNoLockUnorderedMap<CasKey, CasFileInfo> casMap(&memoryBlock);
+			casMap.reserve(totalCasCount);
+
+			CasKey key;
+
+			for (u32 i=0; i!=totalCasCount; ++i)
+			{
+				key.a = i;
+				casMap.try_emplace(key);
+			}
+
+			for (u32 i=0; i!=totalCasCount; ++i)
+			{
+				key.a = i;
+				auto it = casMap.find(key);
+			}
+		}
+#endif
+		return true;
+	}
+
+	bool TestCacheClientAndServer(LoggerWithWriter& logger, const StringBufferBase& testRootDir)
+	{
+		LogWriter& logWriter = logger.m_writer;
+		NetworkBackendTcp tcpBackend(logWriter);
+
+		bool ctorSuccess = true;
+		NetworkServer server(ctorSuccess, { logWriter });
+
+		StringBuffer<MaxPath> rootDir;
+		rootDir.Append(testRootDir).Append(TC("Uba"));
+		if (!DeleteAllFiles(logger, rootDir.data))
+			return false;
+
+		StorageServerCreateInfo storageServerInfo(server, rootDir.data, logWriter);
+		storageServerInfo.casCapacityBytes = 1024ull * 1024 * 1024;
+		auto& storageServer = *new StorageServer(storageServerInfo);
+		auto ssg = MakeGuard([&]() { delete &storageServer; });
+
+		{
+		CacheServerCreateInfo csci(storageServer, rootDir.data, logWriter);
+		CacheServer cacheServer(csci);
+		if (!cacheServer.Load())
+			return false;
+
+
+		SessionServerCreateInfo sessionInfo(storageServer, server, logWriter);
+		sessionInfo.rootDir = rootDir.data;
+		auto& session = *new SessionServer(sessionInfo);
+		auto sg = MakeGuard([&]() { delete &session; });
+
+		u16 port = 1356;
+		if (!server.StartListen(tcpBackend, port))
+			return logger.Error(TC("Failed to listen"));
+		auto disconnectServer = MakeGuard([&]() { server.DisconnectClients(); });
+
+		StringBuffer<MaxPath> workingDir;
+		workingDir.Append(testRootDir).Append(TC("WorkingDir"));
+		if (!DeleteAllFiles(logger, workingDir.data))
+			return false;
+		if (!storageServer.CreateDirectory(workingDir.data))
+			return false;
+		if (!DeleteAllFiles(logger, workingDir.data, false))
+			return false;
+
+		StringBuffer<> testApp;
+		GetTestAppPath(logger, testApp);
+
+		StringBuffer<MaxPath> inputFile;
+		if (!CreateTextFile(inputFile, logger, workingDir.data, TC("Input.txt"), "Foo"))
+			return false;
+		StringBuffer<MaxPath> outputFile;
+		if (!CreateTextFile(outputFile, logger, workingDir.data, TC("Output.txt"), "Foo"))
+			return false;
+
+		StackBinaryWriter<256> inputs;
+		inputs.WriteString(inputFile);
+
+		StackBinaryWriter<256> outputs;
+		outputs.WriteString(outputFile);
+
+		StackBinaryWriter<256> logLines;
+		logLines.WriteString(TC("Hello"));
+		logLines.WriteByte(1);
+
+		ProcessStartInfo psi;
+		psi.application = testApp.data;
+
+		{
+			NetworkClient client(ctorSuccess, { logWriter });
+			CacheClientCreateInfo ccci(logWriter, storageServer, client, session);
+			ccci.useRoots = false;
+			CacheClient cacheClient(ccci);
+
+			if (!client.Connect(tcpBackend, TC("127.0.0.1"), port))
+				return logger.Error(TC("Failed to connect"));
+			auto disconnectClient = MakeGuard([&]() { client.Disconnect(); });
+
+			{
+				CacheResult result;
+				if (cacheClient.FetchFromCache(result, RootPaths(), 0, psi) || result.hit)
+					return false;
+
+				if (!cacheClient.WriteToCache(RootPaths(), 0, psi, inputs.GetData(), inputs.GetPosition(), outputs.GetData(), outputs.GetPosition(), logLines.GetData(), logLines.GetPosition()))
+					return false;
+
+				if (!DeleteFileW(outputFile.data))
+					return false;
+				if (FileExists(logger, outputFile.data))
+					return false;
+				if (!cacheClient.FetchFromCache(result, RootPaths(), 0, psi))
+					return false;
+				if (!FileExists(logger, outputFile.data))
+					return false;
+				if (result.logLines.size() != 1)
+					return false;
+				if (result.logLines[0].text != TC("Hello"))
+					return false;
+			}
+
+			{
+				if (!DeleteFileW(inputFile.data))
+					return false;
+				if (!CreateTextFile(inputFile, logger, workingDir.data, TC("Input.txt"), "Bar"))
+					return false;
+				InvalidateCachedInfo(storageServer, inputFile);
+
+				CacheResult result;
+				if (cacheClient.FetchFromCache(result, RootPaths(), 0, psi) || result.hit)
+					return false;
+
+				if (!cacheClient.WriteToCache(RootPaths(), 0, psi, inputs.GetData(), inputs.GetPosition(), outputs.GetData(), outputs.GetPosition(), logLines.GetData(), logLines.GetPosition()))
+					return false;
+
+				if (!DeleteFileW(outputFile.data))
+					return false;
+				if (FileExists(logger, outputFile.data))
+					return false;
+				if (!cacheClient.FetchFromCache(result, RootPaths(), 0, psi))
+					return false;
+				if (!FileExists(logger, outputFile.data))
+					return false;
+				if (result.logLines.size() != 1)
+					return false;
+				if (result.logLines[0].text != TC("Hello"))
+					return false;
+			}
+		}
+
+		if (!cacheServer.RunMaintenance(true, []() { return false; }))
+			return false;
+
+		{
+			NetworkClient client(ctorSuccess, { logWriter });
+			CacheClientCreateInfo ccci(logWriter, storageServer, client, session);
+			ccci.useRoots = false;
+			CacheClient cacheClient(ccci);
+
+			if (!client.Connect(tcpBackend, TC("127.0.0.1"), port))
+				return logger.Error(TC("Failed to connect"));
+			auto disconnectClient = MakeGuard([&]() { client.Disconnect(); });
+
+			{
+				CacheResult result;
+				if (!cacheClient.FetchFromCache(result, RootPaths(), 0, psi))
+					return false;
+				if (result.logLines.size() != 1)
+					return false;
+				if (result.logLines[0].text != TC("Hello"))
+					return false;
+			}
+		}
+
+		if (!cacheServer.Save())
+			return false;
+
+		}
+
+		CacheServerCreateInfo csci(storageServer, rootDir.data, logWriter);
+		CacheServer cacheServer(csci);
+		if (!cacheServer.Load())
+			return false;
+
+		return true;
+	}
+}
\ No newline at end of file
diff --git a/Engine/Source/Programs/UnrealBuildAccelerator/Test/Private/UbaTestCache.h b/Engine/Source/Programs/UnrealBuildAccelerator/Test/Private/UbaTestCache.h
deleted file mode 100644
diff --git a/Engine/Source/Programs/UnrealBuildAccelerator/Test/Private/UbaTestConfig.cpp b/Engine/Source/Programs/UnrealBuildAccelerator/Test/Private/UbaTestConfig.cpp
new file mode 100644
--- /dev/null
+++ b/Engine/Source/Programs/UnrealBuildAccelerator/Test/Private/UbaTestConfig.cpp
@@ -0,0 +1,66 @@
+// Copyright Epic Games, Inc. All Rights Reserved.
+
+#include "UbaConfig.h"
+#include "UbaLogger.h"
+#include "UbaPlatform.h"
+
+namespace uba
+{
+	bool TestLoadConfig(LoggerWithWriter& logger, const StringBufferBase& rootDir)
+	{
+		static const char* configText =
+			"RootDir = \"e:\\foo\"\r\n"
+			"[CacheClient]\r\n"
+			"UseDirectoryPreparsing = true\r\n"
+			"# Comment = true\r\n"
+			"";
+
+		Config config;
+		if (!config.LoadFromText(logger, configText, strlen(configText)))
+			return false;
+
+		const ConfigTable* tablePtr = config.GetTable(TC("CacheClient"));
+		if (!tablePtr)
+			return false;
+		const ConfigTable& table = *tablePtr;
+		bool test = false;
+		if (!table.GetValueAsBool(test, TC("UseDirectoryPreparsing")))
+			return false;
+		if (test != true)
+			return false;
+		const tchar* str = nullptr;
+		if (!table.GetValueAsString(str, TC("RootDir")))
+			return false;
+		if (TStrcmp(str, TC("e:\\foo")) != 0)
+			return false;
+		if (table.GetValueAsBool(test, TC("Comment")))
+			return false;
+		return true;
+	}
+
+	bool TestSaveConfig(LoggerWithWriter& logger, const StringBufferBase& rootDir)
+	{
+		Vector<char> data;
+		{
+			Config config;
+			ConfigTable& table = config.AddTable(TC("TestTable"));
+			table.AddValue(TC("Foo"), 42);
+			config.SaveToText(logger, data);
+		}
+
+		Config config;
+		if (!config.LoadFromText(logger, data.data(), data.size()))
+			return false;
+
+		const ConfigTable* tablePtr = config.GetTable(TC("TestTable"));
+		if (!tablePtr)
+			return false;
+		const ConfigTable& table = *tablePtr;
+		int foo = 0;
+		if (!table.GetValueAsInt(foo, TC("Foo")))
+			return false;
+		if (foo != 42)
+			return false;
+		return true;
+	}
+}
\ No newline at end of file
diff --git a/Engine/Source/Programs/UnrealBuildAccelerator/Test/Private/UbaTestConfig.h b/Engine/Source/Programs/UnrealBuildAccelerator/Test/Private/UbaTestConfig.h
deleted file mode 100644
diff --git a/Engine/Source/Programs/UnrealBuildAccelerator/Test/Private/UbaTestCrypto.cpp b/Engine/Source/Programs/UnrealBuildAccelerator/Test/Private/UbaTestCrypto.cpp
new file mode 100644
--- /dev/null
+++ b/Engine/Source/Programs/UnrealBuildAccelerator/Test/Private/UbaTestCrypto.cpp
@@ -0,0 +1,51 @@
+// Copyright Epic Games, Inc. All Rights Reserved.
+
+#include "UbaCrypto.h"
+#include "UbaLogger.h"
+#include "UbaStringBuffer.h"
+
+namespace uba
+{
+	bool TestCrypto(LoggerWithWriter& logger, const StringBufferBase& rootDir)
+	{
+		u64 key128[] = { 0x1234567812345678llu, 0x1234567812345678llu };
+
+		CryptoKey encryptKey = Crypto::CreateKey(logger, (const u8*)key128);
+
+		for (u32 dataSize=1;dataSize!=135; ++dataSize)
+		{
+			CryptoKey decryptKey = encryptKey;
+
+			u8* originalData = new u8[dataSize];
+			for (u32 i=0;i!=dataSize; ++i)
+				originalData[i] = u8(rand() % 256);
+
+			u8* encryptedData = new u8[dataSize];
+			memcpy(encryptedData, originalData, dataSize);
+
+			for (u32 i=0; i!=3; ++i)
+			{
+				if (!Crypto::Encrypt(logger, encryptKey, encryptedData, dataSize))
+					return false;
+
+				if (dataSize > 16 && memcmp(encryptedData, originalData, dataSize) == 0)
+					return false;
+
+				if (!Crypto::Decrypt(logger, decryptKey, encryptedData, dataSize))
+					return false;
+
+				if (memcmp(encryptedData, originalData, dataSize) != 0)
+					return false;
+
+				if (i == 1)
+					decryptKey = Crypto::DuplicateKey(logger, encryptKey);
+			}
+			delete[] encryptedData;
+			delete[] originalData;
+			Crypto::DestroyKey(decryptKey);
+		}
+
+		Crypto::DestroyKey(encryptKey);
+		return true;
+	}
+}
\ No newline at end of file
diff --git a/Engine/Source/Programs/UnrealBuildAccelerator/Test/Private/UbaTestCrypto.h b/Engine/Source/Programs/UnrealBuildAccelerator/Test/Private/UbaTestCrypto.h
deleted file mode 100644
diff --git a/Engine/Source/Programs/UnrealBuildAccelerator/Test/Private/UbaTestNetwork.cpp b/Engine/Source/Programs/UnrealBuildAccelerator/Test/Private/UbaTestNetwork.cpp
new file mode 100644
--- /dev/null
+++ b/Engine/Source/Programs/UnrealBuildAccelerator/Test/Private/UbaTestNetwork.cpp
@@ -0,0 +1,181 @@
+// Copyright Epic Games, Inc. All Rights Reserved.
+
+#include "UbaBinaryReaderWriter.h"
+#include "UbaPlatform.h"
+#include "UbaNetworkClient.h"
+#include "UbaNetworkServer.h"
+#include "UbaNetworkMessage.h"
+#include "UbaNetworkBackendTcp.h"
+#include "UbaNetworkBackendMemory.h"
+
+namespace uba
+{
+	bool TestSockets(LoggerWithWriter& logger, const StringBufferBase& testRootDir)
+	{
+		static constexpr u16 port = 1346;
+		Atomic<u8> result;
+		NetworkBackendTcp tcp(logger.m_writer);
+		Thread t([&]()
+			{
+				logger.Info(TC("Starting to Listen"));
+				tcp.StartListen(logger, port, TC("127.0.0.1"), [&](void* connection, const sockaddr& remoteSocketAddr)
+					{
+						logger.Info(TC("Listen got connection"));
+						tcp.SetDisconnectCallback(connection, nullptr, [](void*, const Guid&, void*) {});
+						tcp.SetRecvCallbacks(connection, &result, 1,
+							[](void* context, const Guid& connectionUid, u8* headerData, void*& outBodyContext, u8*& outBodyData, u32& outBodySize)
+							{
+								*(Atomic<u8>*)context = *headerData;
+								//wprintf(TC("Listen got data from peer: %u\n"), *headerData);
+								return true;
+							},
+							nullptr, TC(""));
+						return true;
+					});
+				return 0;
+			});
+
+		Sleep(100);
+		logger.Info(TC("Starting to Connect"));
+		tcp.Connect(logger, TC("127.0.0.1"), [&](void* connection, const sockaddr& remoteSocketAddr, bool* timedOut)
+			{
+				tcp.SetDisconnectCallback(connection, nullptr, [](void*, const Guid&, void*) {});
+				tcp.SetRecvCallbacks(connection, &tcp, 1, [](void* context, const Guid& connectionUid, u8* headerData, void*& outBodyContext, u8*& outBodyData, u32& outBodySize) { return true; }, nullptr, TC(""));
+				u8 b = 42;
+				NetworkBackend::SendContext sc;
+				tcp.Send(logger, connection, &b, 1, sc);
+				return true;
+			}, port);
+
+		Sleep(200);
+		t.Wait();
+
+		if (result != 42)
+			return logger.Error(TC("Failed to receive data"));
+
+		return true;
+	}
+
+	bool TestClientServer(LoggerWithWriter& logger, const StringBufferBase& testRootDir)
+	{
+		auto& logWriter = logger.m_writer;
+		NetworkBackendTcp networkBackend(logWriter);
+
+		bool ctorSuccess = true;
+		NetworkServer server(ctorSuccess, logWriter);
+		NetworkClient client(ctorSuccess, logWriter);
+
+		server.RegisterService(1, [&](const ConnectionInfo& connectionInfo, const WorkContext& workContext, MessageInfo& messageInfo, BinaryReader& reader, BinaryWriter& writer)
+			{
+				logger.Info(TC("Got ping!"));
+				UBA_ASSERT(messageInfo.type == SessionMessageType_Ping);
+				writer.WriteByte(42);
+				return true;
+			});
+
+		auto ds = MakeGuard([&]() { server.DisconnectClients(); });
+		if (!server.StartListen(networkBackend, 1234))
+			return logger.Error(TC("Failed to listen"));
+		Sleep(100);
+		auto dg = MakeGuard([&]() { client.Disconnect(); });
+		if (!client.Connect(networkBackend, TC("127.0.0.1"), 1234))
+			return logger.Error(TC("Failed to connect"));
+
+
+		StackBinaryWriter<128> writer;
+		NetworkMessage msg(client, 1, SessionMessageType_Ping, writer);
+		writer.WriteU32(32);
+		StackBinaryReader<32> reader;
+		if (!msg.Send(reader))
+			return logger.Error(TC("Failed to get message"));
+
+		u8 value = reader.ReadByte();
+		logger.Info(TC("Got value %u"), value);
+		return true;
+	}
+
+	bool TestClientServer2(LoggerWithWriter& logger, const StringBufferBase& testRootDir)
+	{
+		auto& logWriter = logger.m_writer;
+		NetworkBackendTcp networkBackend(logWriter);
+
+		bool ctorSuccess = true;
+		NetworkServer server(ctorSuccess, logWriter);
+		NetworkClient client(ctorSuccess, logWriter);
+
+		server.RegisterService(1, [&](const ConnectionInfo& connectionInfo, const WorkContext& workContext, MessageInfo& messageInfo, BinaryReader& reader, BinaryWriter& writer)
+			{
+				logger.Info(TC("Got ping!"));
+				UBA_ASSERT(messageInfo.type == SessionMessageType_Ping);
+				writer.WriteByte(42);
+				return true;
+			});
+
+		if (!client.StartListen(networkBackend, 1239))
+			return logger.Error(TC("Client failed to listen"));
+		Sleep(100);
+		auto ds = MakeGuard([&]() { server.DisconnectClients(); });
+		auto dg = MakeGuard([&]() { client.Disconnect(); });
+		if (!server.AddClient(networkBackend, TC("127.0.0.1"), 1239))
+			return logger.Error(TC("Server failed to connect"));
+		if (!server.AddClient(networkBackend, TC("127.0.0.1"), 1239))
+			return logger.Error(TC("Server failed to connect second"));
+
+		u64 time = GetTime();
+		while (!client.GetConnectionCount())
+		{
+			if (TimeToMs(GetTime() - time) > 4000)
+				return logger.Error(TC("Client failed to establish connection"));
+			Sleep(100);
+		}
+
+		StackBinaryWriter<128> writer;
+		NetworkMessage msg(client, 1, SessionMessageType_Ping, writer);
+		writer.WriteU32(32);
+		StackBinaryReader<32> reader;
+		if (!msg.Send(reader))
+			return logger.Error(TC("Failed to get message"));
+
+		u8 value = reader.ReadByte();
+		logger.Info(TC("Got value %u"), value);
+		return true;
+	}
+
+	bool TestClientServerMem(LoggerWithWriter& logger, const StringBufferBase& testRootDir)
+	{
+		auto& logWriter = logger.m_writer;
+		NetworkBackendMemory networkBackend(logWriter);
+
+		bool ctorSuccess = true;
+		NetworkServer server(ctorSuccess, logWriter);
+		NetworkClient client(ctorSuccess, logWriter);
+
+		server.RegisterService(1, [&](const ConnectionInfo& connectionInfo, const WorkContext& workContext, MessageInfo& messageInfo, BinaryReader& reader, BinaryWriter& writer)
+			{
+				logger.Info(TC("Got ping!"));
+				UBA_ASSERT(messageInfo.type == SessionMessageType_Ping);
+				writer.WriteByte(42);
+				return true;
+			});
+
+		auto ds = MakeGuard([&]() { server.DisconnectClients(); });
+		if (!server.StartListen(networkBackend, 1234))
+			return logger.Error(TC("Failed to listen"));
+		Sleep(100);
+		auto dg = MakeGuard([&]() { client.Disconnect(); });
+		if (!client.Connect(networkBackend, TC("127.0.0.1"), 1234))
+			return logger.Error(TC("Failed to connect"));
+
+
+		StackBinaryWriter<128> writer;
+		NetworkMessage msg(client, 1, SessionMessageType_Ping, writer);
+		writer.WriteU32(32);
+		StackBinaryReader<32> reader;
+		if (!msg.Send(reader))
+			return logger.Error(TC("Failed to get message"));
+
+		u8 value = reader.ReadByte();
+		logger.Info(TC("Got value %u"), value);
+		return true;
+	}
+}
diff --git a/Engine/Source/Programs/UnrealBuildAccelerator/Test/Private/UbaTestNetwork.h b/Engine/Source/Programs/UnrealBuildAccelerator/Test/Private/UbaTestNetwork.h
deleted file mode 100644
diff --git a/Engine/Source/Programs/UnrealBuildAccelerator/Test/Private/UbaTestScheduler.cpp b/Engine/Source/Programs/UnrealBuildAccelerator/Test/Private/UbaTestScheduler.cpp
new file mode 100644
--- /dev/null
+++ b/Engine/Source/Programs/UnrealBuildAccelerator/Test/Private/UbaTestScheduler.cpp
@@ -0,0 +1,100 @@
+// Copyright Epic Games, Inc. All Rights Reserved.
+
+#include "UbaFile.h"
+#include "UbaNetworkServer.h"
+#include "UbaScheduler.h"
+#include "UbaSessionServer.h"
+
+namespace uba
+{
+	using RunProcessFunction = Function<ProcessHandle(const ProcessStartInfo&)>;
+	using TestSessionFunction = Function<bool(LoggerWithWriter& logger, SessionServer& session, const tchar* workingDir, const RunProcessFunction& runProcess)>;
+	bool RunLocal(LoggerWithWriter& logger, const StringBufferBase& testRootDir, const TestSessionFunction& testFunc, bool enableDetour = true);
+	bool RunRemote(LoggerWithWriter& logger, const StringBufferBase& testRootDir, const TestSessionFunction& testFunc, bool deleteAll = true);
+	void GetTestAppPath(LoggerWithWriter& logger, StringBufferBase& out);
+	const tchar* GetSystemApplication();
+	const tchar* GetSystemArguments();
+
+	bool TestLocalSchedule(LoggerWithWriter& logger, const StringBufferBase& testRootDir)
+	{
+		return RunLocal(logger, testRootDir, [](LoggerWithWriter& logger, SessionServer& session, const tchar* workingDir, const RunProcessFunction& runProcess)
+			{
+				SchedulerCreateInfo info(session);
+				Scheduler scheduler(info);
+
+				ProcessStartInfo processInfo;
+				processInfo.application = GetSystemApplication();
+				processInfo.workingDir = workingDir;
+				processInfo.arguments = GetSystemArguments();
+
+				EnqueueProcessInfo epi(processInfo);
+				scheduler.EnqueueProcess(epi);
+				scheduler.Start();
+
+				u32 queued, activeLocal, activeRemote, finished;
+				do { scheduler.GetStats(queued, activeLocal, activeRemote, finished); } while (finished != 1);
+
+				scheduler.Stop();
+				return true;
+			});
+	}
+
+	bool TestLocalScheduleReuse(LoggerWithWriter& logger, const StringBufferBase& testRootDir)
+	{
+		return RunLocal(logger, testRootDir, [](LoggerWithWriter& logger, SessionServer& session, const tchar* workingDir, const RunProcessFunction& runProcess)
+			{
+				SchedulerCreateInfo info(session);
+				info.enableProcessReuse = true;
+				Scheduler scheduler(info);
+
+				StringBuffer<> testApp;
+				GetTestAppPath(logger, testApp);
+
+				ProcessStartInfo processInfo;
+				processInfo.application = testApp.data;
+				processInfo.workingDir = workingDir;
+				processInfo.arguments = TC("-reuse");
+
+				EnqueueProcessInfo epi(processInfo);
+				scheduler.EnqueueProcess(epi);
+				scheduler.Start();
+
+				u32 queued, activeLocal, activeRemote, finished;
+				do { scheduler.GetStats(queued, activeLocal, activeRemote, finished); } while (finished != 1);
+
+				scheduler.Stop();
+				return true;
+			});
+	}
+
+	bool TestRemoteScheduleReuse(LoggerWithWriter& logger, const StringBufferBase& testRootDir)
+	{
+		return RunRemote(logger, testRootDir, [](LoggerWithWriter& logger, SessionServer& session, const tchar* workingDir, const RunProcessFunction& runProcess)
+			{
+				SchedulerCreateInfo info(session);
+				info.enableProcessReuse = true;
+				info.maxLocalProcessors = 0;
+				Scheduler scheduler(info);
+
+				StringBuffer<> testApp;
+				GetTestAppPath(logger, testApp);
+
+				ProcessStartInfo processInfo;
+				processInfo.application = testApp.data;
+				processInfo.workingDir = workingDir;
+				processInfo.arguments = TC("-reuse");
+
+				EnqueueProcessInfo epi(processInfo);
+				scheduler.EnqueueProcess(epi);
+				scheduler.Start();
+
+				u32 queued, activeLocal, activeRemote, finished;
+				do { scheduler.GetStats(queued, activeLocal, activeRemote, finished); } while (finished != 1);
+
+				session.GetServer().DisconnectClients(); // Must make sure all is disconnected since scheduler goes out of scope
+
+				scheduler.Stop();
+				return true;
+			});
+	}
+}
diff --git a/Engine/Source/Programs/UnrealBuildAccelerator/Test/Private/UbaTestScheduler.h b/Engine/Source/Programs/UnrealBuildAccelerator/Test/Private/UbaTestScheduler.h
deleted file mode 100644
diff --git a/Engine/Source/Programs/UnrealBuildAccelerator/Test/Private/UbaTestSession.cpp b/Engine/Source/Programs/UnrealBuildAccelerator/Test/Private/UbaTestSession.cpp
new file mode 100644
--- /dev/null
+++ b/Engine/Source/Programs/UnrealBuildAccelerator/Test/Private/UbaTestSession.cpp
@@ -0,0 +1,612 @@
+// Copyright Epic Games, Inc. All Rights Reserved.
+
+#include "UbaConfig.h"
+#include "UbaFileAccessor.h"
+#include "UbaNetworkBackendTcp.h"
+#include "UbaSessionServer.h"
+#include "UbaSessionClient.h"
+#include "UbaStorageServer.h"
+#include "UbaStorageClient.h"
+#include "UbaTest.h"
+
+namespace uba
+{
+	using RunProcessFunction = Function<ProcessHandle(const ProcessStartInfo&)>;
+	using TestSessionFunction = Function<bool(LoggerWithWriter& logger, SessionServer& session, const tchar* workingDir, const RunProcessFunction& runProcess)>;
+
+	bool RunLocal(LoggerWithWriter& logger, const StringBufferBase& testRootDir, const TestSessionFunction& testFunc, bool enableDetour = true)
+	{
+		LogWriter& logWriter = logger.m_writer;
+
+		StringBuffer<MaxPath> rootDir;
+		rootDir.Append(testRootDir).Append(TC("Uba"));
+		if (!DeleteAllFiles(logger, rootDir.data))
+			return false;
+
+		StorageCreateInfo storageInfo(rootDir.data, logWriter);
+		storageInfo.casCapacityBytes = 1024ull * 1024 * 1024;
+		StorageImpl storage(storageInfo);
+
+		bool ctorSuccess = true;
+		NetworkServer server(ctorSuccess, { logWriter });
+
+		SessionServerCreateInfo sessionServerInfo(storage, server, logWriter);
+		sessionServerInfo.checkMemory = false;
+		sessionServerInfo.rootDir = rootDir.data;
+
+		#if UBA_DEBUG
+		sessionServerInfo.logToFile = true;
+		#endif
+
+		SessionServer session(sessionServerInfo);
+
+		StringBuffer<MaxPath> workingDir;
+		workingDir.Append(testRootDir).Append(TC("WorkingDir"));
+		if (!DeleteAllFiles(logger, workingDir.data))
+			return false;
+
+		if (!storage.CreateDirectory(workingDir.data))
+			return false;
+		if (!DeleteAllFiles(logger, workingDir.data, false))
+			return false;
+		workingDir.EnsureEndsWithSlash();
+		return testFunc(logger, session, workingDir.data, [&](const ProcessStartInfo& pi) { return session.RunProcess(pi, true, enableDetour); });
+	}
+
+	bool RunRemote(LoggerWithWriter& logger, const StringBufferBase& testRootDir, const TestSessionFunction& testFunc, bool deleteAll = true)
+	{
+		LogWriter& logWriter = logger.m_writer;
+		NetworkBackendTcp tcpBackend(logWriter);
+
+		bool ctorSuccess = true;
+		NetworkServer server(ctorSuccess, { logWriter });
+		NetworkClient client(ctorSuccess, { logWriter });
+
+		StringBuffer<MaxPath> rootDir;
+		rootDir.Append(testRootDir).Append(TC("Uba"));
+		if (deleteAll && !DeleteAllFiles(logger, rootDir.data))
+			return false;
+
+		StorageServerCreateInfo storageServerInfo(server, rootDir.data, logWriter);
+		storageServerInfo.casCapacityBytes = 1024ull * 1024 * 1024;
+		auto& storageServer = *new StorageServer(storageServerInfo);
+		auto ssg = MakeGuard([&]() { delete &storageServer; });
+
+		SessionServerCreateInfo sessionServerInfo(storageServer, server, logWriter);
+		sessionServerInfo.checkMemory = false;
+		sessionServerInfo.rootDir = rootDir.data;
+
+		#if UBA_DEBUG
+		sessionServerInfo.logToFile = true;
+		sessionServerInfo.remoteLogEnabled = true;
+		#endif
+
+		auto& sessionServer = *new SessionServer(sessionServerInfo);
+		auto ssg2 = MakeGuard([&]() { delete &sessionServer; });
+
+		auto sg = MakeGuard([&]() { server.DisconnectClients(); });
+
+		Config clientConfig;
+		clientConfig.AddTable(TC("Storage")).AddValue(TC("CheckExistsOnServer"), true);
+		server.SetClientsConfig(clientConfig);
+
+		u16 port = 1356;
+		if (!server.StartListen(tcpBackend, port))
+			return logger.Error(TC("Failed to listen"));
+		if (!client.Connect(tcpBackend, TC("127.0.0.1"), port))
+			return logger.Error(TC("Failed to connect"));
+
+		Config config;
+		if (!client.FetchConfig(config))
+			return false;
+
+		rootDir.Append(TC("Client"));
+		if (deleteAll && !DeleteAllFiles(logger, rootDir.data))
+			return false;
+
+		StorageClientCreateInfo storageClientInfo(client, rootDir.data);
+		storageClientInfo.Apply(config);
+		auto& storageClient = *new StorageClient(storageClientInfo);
+		auto scg = MakeGuard([&]() { delete &storageClient; });
+
+		SessionClientCreateInfo sessionClientInfo(storageClient, client, logWriter);
+		sessionClientInfo.rootDir = rootDir.data;
+		sessionClientInfo.allowKeepFilesInMemory = false;
+
+		#if UBA_DEBUG
+		sessionClientInfo.logToFile = true;
+		#endif
+
+		auto& sessionClient = *new SessionClient(sessionClientInfo);
+		auto scg2 = MakeGuard([&]() { delete &sessionClient; });
+
+		auto cg = MakeGuard([&]() { sessionClient.Stop(); client.Disconnect(); });
+
+		StringBuffer<MaxPath> workingDir;
+		workingDir.Append(testRootDir).Append(TC("WorkingDir"));
+		if (deleteAll && !DeleteAllFiles(logger, workingDir.data))
+			return false;
+		if (!storageServer.CreateDirectory(workingDir.data))
+			return false;
+		if (deleteAll && !DeleteAllFiles(logger, workingDir.data, false))
+			return false;
+
+		storageClient.Start();
+		sessionClient.Start();
+
+
+		workingDir.EnsureEndsWithSlash();
+		return testFunc(logger, sessionServer, workingDir.data, [&](const ProcessStartInfo& pi) { return sessionServer.RunProcessRemote(pi); });
+	}
+
+	void GetTestAppPath(LoggerWithWriter& logger, StringBufferBase& out)
+	{
+		GetDirectoryOfCurrentModule(logger, out);
+		out.EnsureEndsWithSlash();
+		out.Append(IsWindows ? TC("UbaTestApp.exe") : TC("UbaTestApp"));
+	}
+
+	bool CreateTextFile(StringBufferBase& outPath, LoggerWithWriter& logger, const tchar* workingDir, const tchar* fileName, const char* text)
+	{
+		outPath.Clear().Append(workingDir).EnsureEndsWithSlash().Append(fileName);
+		FileAccessor fr(logger, outPath.data);
+		if (!fr.CreateWrite())
+			return false;
+		fr.Write(text, strlen(text) + 1);
+		return fr.Close();
+	}
+
+	bool RunTestApp(LoggerWithWriter& logger, SessionServer& session, const tchar* workingDir, const RunProcessFunction& runProcess)
+	{
+		StringBuffer<MaxPath> testApp;
+		GetTestAppPath(logger, testApp);
+
+		StringBuffer<MaxPath> fileR;
+		if (!CreateTextFile(fileR, logger, workingDir, TC("FileR.h"), "Foo"))
+			return false;
+
+		{
+			StringBuffer<MaxPath> dir;
+			dir.Append(workingDir).Append(TC("Dir1"));
+			if (!CreateDirectoryW(dir.data))
+				return logger.Error(TC("Failed to create dir %s"), dir.data);
+
+			dir.Clear().Append(workingDir).Append(TC("Dir2"));
+			if (!CreateDirectoryW(dir.data))
+				return logger.Error(TC("Failed to create dir %s"), dir.data);
+			dir.EnsureEndsWithSlash().Append(TC("Dir3"));
+			if (!CreateDirectoryW(dir.data))
+				return logger.Error(TC("Failed to create dir %s"), dir.data);
+			dir.EnsureEndsWithSlash().Append(TC("Dir4"));
+			if (!CreateDirectoryW(dir.data))
+				return logger.Error(TC("Failed to create dir %s"), dir.data);
+			dir.EnsureEndsWithSlash().Append(TC("Dir5"));
+			if (!CreateDirectoryW(dir.data))
+				return logger.Error(TC("Failed to create dir %s"), dir.data);
+		}
+
+		ProcessStartInfo processInfo;
+		processInfo.application = testApp.data;
+		processInfo.workingDir = workingDir;
+		processInfo.logLineFunc = [](void* userData, const tchar* line, u32 length, LogEntryType type)
+			{
+				LoggerWithWriter(g_consoleLogWriter, TC("")).Info(line);
+			};
+
+		ProcessHandle process = runProcess(processInfo);
+		if (!process.WaitForExit(100000))
+			return logger.Error(TC("UbaTestApp did not exit in 10 seconds"));
+		u32 exitCode = process.GetExitCode();
+
+		if (exitCode != 0)
+		{
+			for (auto& logLine : process.GetLogLines())
+				logger.Error(logLine.text.c_str());
+			return logger.Error(TC("UbaTestApp returned exit code %u"), exitCode);
+		}
+
+		{
+			StringBuffer<MaxPath> fileW2;
+			fileW2.Append(workingDir).Append(TC("FileW2"));
+			if (!FileExists(logger, fileW2.data))
+				return logger.Error(TC("Can't find file %s"), fileW2.data);
+		}
+		{
+			StringBuffer<MaxPath> fileWF;
+			fileWF.Append(workingDir).Append(TC("FileWF"));
+			if (!FileExists(logger, fileWF.data))
+				return logger.Error(TC("Can't find file %s"), fileWF.data);
+		}
+		return true;
+	}
+
+#if PLATFORM_MAC
+	bool ExecuteCommand(LoggerWithWriter& logger, const tchar* command, StringBufferBase& commandOutput)
+	{
+		FILE* fpCommand = popen(command, "r");
+		if (fpCommand == nullptr || fgets(commandOutput.data, commandOutput.capacity, fpCommand) == nullptr || pclose(fpCommand) != 0)
+		{
+			logger.Error("Failed to get an Xcode from xcode-select");
+			return false;
+		}
+
+		commandOutput.count = strlen(commandOutput.data);
+		while (isspace(commandOutput.data[commandOutput.count-1]))
+		{
+			commandOutput.data[commandOutput.count-1] = 0;
+			commandOutput.count--;
+		}
+		return true;
+	}
+#endif
+
+	bool RunClang(LoggerWithWriter& logger, SessionServer& session, const tchar* workingDir, const RunProcessFunction& runProcess)
+	{
+		StringBuffer<MaxPath> sourceFile;
+		sourceFile.Append(workingDir).Append(TC("Code.cpp"));
+		FileAccessor codeFile(logger, sourceFile.data);
+		if (!codeFile.CreateWrite())
+			return false;
+		char code[] = "#include <stdio.h>\n int main() { printf(\"Hello world\\n\"); return 0; }";
+		if (!codeFile.Write(code, sizeof(code) - 1))
+			return false;
+		if (!codeFile.Close())
+			return false;
+
+#if PLATFORM_WINDOWS
+		const tchar* clangPath = TC("c:\\Program Files\\Microsoft Visual Studio\\2022\\Professional\\VC\\Tools\\Llvm\\bin\\clang-cl.exe");
+#elif PLATFORM_MAC
+		StringBuffer<MaxPath> xcodePath;
+		ExecuteCommand(logger, "/usr/bin/xcrun -r clang++", xcodePath);
+		const tchar* clangPath = xcodePath.data;
+#else
+		const tchar* clangPath = TC("/usr/bin/clang++");
+#endif
+
+		if (!FileExists(logger, clangPath)) // Skipping if clang is not installed.
+			return true;
+
+		ProcessStartInfo processInfo;
+		processInfo.application = clangPath;
+
+		StringBuffer<MaxPath> args;
+
+#if PLATFORM_WINDOWS
+		args.Append("/Brepro ");
+#elif PLATFORM_MAC
+		StringBuffer<MaxPath> xcodeSDKPath;
+		ExecuteCommand(logger, "xcrun --show-sdk-path", xcodeSDKPath);
+		args.Append("-isysroot ");
+		args.Append(xcodeSDKPath.data).Append(' ');
+#endif
+		args.Append(TC("-o code Code.cpp"));
+
+		processInfo.arguments = args.data;
+
+		processInfo.workingDir = workingDir;
+		//processInfo.logFile = TC("[REDACTED_USER_HOME]/RunClang.log");
+		ProcessHandle process = runProcess(processInfo);
+		if (!process.WaitForExit(40000))
+			return logger.Error(TC("clang++ timed out"));
+		u32 exitCode = process.GetExitCode();
+		if (exitCode != 0)
+			return logger.Error(TC("clang++ returned exit code %u"), exitCode);
+		return true;
+	}
+
+	bool RunCustomService(LoggerWithWriter& logger, SessionServer& session, const tchar* workingDir, const RunProcessFunction& runProcess)
+	{
+		bool gotMessage = false;
+		session.RegisterCustomService([&](uba::Process& process, const void* recv, u32 recvSize, void* send, u32 sendCapacity)
+			{
+				gotMessage = true;
+				//wprintf(L"GOT MESSAGE: %.*s\n", recvSize / 2, (const wchar_t*)recv);
+				const wchar_t* hello = L"Hello response from server";
+				u64 helloBytes = wcslen(hello) * 2;
+				memcpy(send, hello, helloBytes);
+				return u32(helloBytes);
+			});
+
+		StringBuffer<> testApp;
+		GetTestAppPath(logger, testApp);
+
+		ProcessStartInfo processInfo;
+		processInfo.application = testApp.data;
+		processInfo.workingDir = workingDir;
+		processInfo.arguments = TC("Whatever");
+		ProcessHandle process = runProcess(processInfo);
+		if (!process.WaitForExit(10000))
+			return logger.Error(TC("UbaTestApp did not exit in 10 seconds"));
+		u32 exitCode = process.GetExitCode();
+
+		if (exitCode != 0)
+			return logger.Error(TC("UbaTestApp returned exit code %u"), exitCode);
+
+		if (!gotMessage)
+			return logger.Error(TC("Never got message from UbaTestApp"));
+
+		return true;
+	}
+
+	// NOTE: This test is dependent on the UbaTestApp<Platform>
+	// The purpose of this test is to validate that the platform specific detours are
+	// working as expected.
+	// Before running the actual UbaTestApp, RunLocal calls through a variety of functions
+	// that sets up the various UbaSession Servers, Clients, etc. It creates some temporary
+	// directories, e.g. Dir1 and eventually call ProcessImpl::InternalCreateProcess.
+	// InternalCreateProcess will setup the shared memory, inject the Detour library
+	// and setup any other necessary environment variables, and spawn the actual process
+	// (in this case the UbaTestApp)
+	// Once UbaTestApp has started, it will first check and validate that the detour library
+	// is in the processes address space. With the detour in place, the test app will
+	// exercise various file functions which will actually go through our detour library.
+	bool TestDetouredTestApp(LoggerWithWriter& logger, const StringBufferBase& testRootDir)
+	{
+		return RunLocal(logger, testRootDir, RunTestApp);
+	}
+
+	bool TestRemoteDetouredTestApp(LoggerWithWriter& logger, const StringBufferBase& testRootDir)
+	{
+		return RunRemote(logger, testRootDir, RunTestApp);
+	}
+
+	bool TestCustomService(LoggerWithWriter& logger, const StringBufferBase& testRootDir)
+	{
+		return RunRemote(logger, testRootDir, RunCustomService);
+	}
+
+	bool TestDetouredClang(LoggerWithWriter& logger, const StringBufferBase& testRootDir)
+	{
+		return RunLocal(logger, testRootDir, RunClang);
+	}
+
+	bool TestRemoteDetouredClang(LoggerWithWriter& logger, const StringBufferBase& testRootDir)
+	{
+		// Run twice to test LoadCasTable/SaveCasTable etc
+		if (!RunRemote(logger, testRootDir, RunClang))
+			return false;
+		return RunRemote(logger, testRootDir, RunClang, false);
+	}
+
+	bool TestDetouredTouch(LoggerWithWriter& logger, const StringBufferBase& testRootDir)
+	{
+		return RunLocal(logger, testRootDir, [](LoggerWithWriter& logger, SessionServer& session, const tchar* workingDir, const RunProcessFunction& runProcess)
+			{
+				StringBuffer<> file;
+				file.Append(workingDir).Append(TC("TouchFile.h"));
+				FileAccessor fr(logger, file.data);
+
+				CHECK_TRUE(fr.CreateWrite());
+				CHECK_TRUE(fr.Write("Foo", 4));
+				CHECK_TRUE(fr.Close());
+				FileInformation oldInfo;
+				CHECK_TRUE(GetFileInformation(oldInfo, logger, file.data));
+
+				Sleep(100);
+
+				ProcessStartInfo processInfo;
+				processInfo.application = TC("/usr/bin/touch");
+				processInfo.workingDir = workingDir;
+				processInfo.arguments = file.data;
+				processInfo.logFile = TC("[REDACTED_USER_HOME]/Touch.log");
+				ProcessHandle process = runProcess(processInfo);
+				if (!process.WaitForExit(10000))
+					return logger.Error(TC("UbaTestApp did not exit in 10 seconds"));
+				u32 exitCode = process.GetExitCode();
+				if (exitCode != 0)
+					return false;
+
+				FileInformation newInfo;
+				CHECK_TRUE(GetFileInformation(newInfo, logger, file.data));
+				if (newInfo.lastWriteTime == oldInfo.lastWriteTime)
+					return logger.Error(TC("File time not changed after touch"));
+				return true;
+			});
+	}
+
+	const tchar* GetSystemApplication()
+	{
+		#if PLATFORM_WINDOWS
+		return TC("c:\\windows\\system32\\ping.exe");
+		#elif PLATFORM_LINUX
+		return TC("/usr/bin/zip");
+		#else
+		return TC("/sbin/zip");
+		#endif
+	}
+
+	const tchar* GetSystemArguments()
+	{
+		#if PLATFORM_WINDOWS
+		return TC("-n 1 localhost");
+		#else
+		return TC("-help");
+		#endif
+	}
+
+	const tchar* GetSystemExpectedLogLine()
+	{
+		#if PLATFORM_WINDOWS
+		return TC("Pinging ");
+		#else
+		return TC("zip [-options]");
+		#endif
+	}
+
+	bool TestMultipleDetouredProcesses(LoggerWithWriter& logger, const StringBufferBase& testRootDir)
+	{
+		return RunLocal(logger, testRootDir, [](LoggerWithWriter& logger, SessionServer& session, const tchar* workingDir, const RunProcessFunction& runProcess)
+			{
+				ProcessStartInfo processInfo;
+				processInfo.application = GetSystemApplication();
+				processInfo.workingDir = workingDir;
+				processInfo.arguments = GetSystemArguments();
+				//processInfo.logFile = "[REDACTED_USER_HOME]/LogFile.log";
+				Vector<ProcessHandle> processes;
+
+				for (u32 i=0; i!=50; ++i)
+					processes.push_back(runProcess(processInfo));
+
+				for (auto& process : processes)
+				{
+					if (!process.WaitForExit(10000))
+						return logger.Error(TC("UbaTestApp did not exit in 10 seconds"));
+					u32 exitCode = process.GetExitCode();
+					if (exitCode != 0)
+						return false;
+				}
+
+				return true;
+			});
+	}
+
+	bool RunSystemApplicationAndLookForLog(LoggerWithWriter& logger, SessionServer& session, const tchar* workingDir, const RunProcessFunction& runProcess)
+	{
+		ProcessStartInfo processInfo;
+		processInfo.application = GetSystemApplication();
+		processInfo.workingDir = workingDir;
+		processInfo.arguments = GetSystemArguments();
+
+		bool foundPingString = false;
+		processInfo.logLineUserData = &foundPingString;
+		processInfo.logLineFunc = [](void* userData, const tchar* line, u32 length, LogEntryType type)
+			{
+				*(bool*)userData |= Contains(line, GetSystemExpectedLogLine());
+			};
+
+		ProcessHandle process = runProcess(processInfo);
+
+		if (!process.WaitForExit(10000))
+			return logger.Error(TC("UbaTestApp did not exit in 10 seconds"));
+		u32 exitCode = process.GetExitCode();
+		if (exitCode != 0)
+			return logger.Error(TC("Got exit code %u"), exitCode);
+		if (!foundPingString)
+			return logger.Error(TC("Did not log string containing \"%s\""), GetSystemExpectedLogLine());
+		return true;
+	}
+
+	bool TestLogLines(LoggerWithWriter& logger, const StringBufferBase& testRootDir)
+	{
+		return RunLocal(logger, testRootDir, RunSystemApplicationAndLookForLog);
+	}
+
+	bool TestLogLinesNoDetour(LoggerWithWriter& logger, const StringBufferBase& testRootDir)
+	{
+		return RunLocal(logger, testRootDir, RunSystemApplicationAndLookForLog, false);
+	}
+
+	bool CheckAttributes(LoggerWithWriter& logger, SessionServer& session, const tchar* workingDir, const RunProcessFunction& runProcess)
+	{
+		StringBuffer<MaxPath> testApp;
+		GetTestAppPath(logger, testApp);
+		ProcessStartInfo processInfo;
+		processInfo.application = testApp.data;
+		processInfo.workingDir = workingDir;
+		processInfo.logLineFunc = [](void* userData, const tchar* line, u32 length, LogEntryType type)
+			{
+				LoggerWithWriter(g_consoleLogWriter, TC("")).Info(line);
+			};
+
+		auto GetAttributes = [&](const StringView& file) -> u32
+			{
+				StringBuffer<> arg(TC("-GetFileAttributes="));
+				arg.Append(file);
+				processInfo.arguments = arg.data;
+				ProcessHandle process = runProcess(processInfo);
+				if (!process.WaitForExit(100000))
+					return logger.Error(TC("UbaTestApp did not exit in 10 seconds"));
+				u32 exitCode =  process.GetExitCode();
+				return exitCode == 255 ? INVALID_FILE_ATTRIBUTES : exitCode;
+			};
+
+		MemoryBlock temp;
+		DirectoryTable dirTable(temp);
+		dirTable.Init(session.GetDirectoryTableMemory(), 0, 0);
+
+		CHECK_TRUE(session.RefreshDirectory(workingDir, true));
+		CHECK_TRUE(session.RefreshDirectory(workingDir));
+		CHECK_TRUE(dirTable.EntryExists(ToView(workingDir)) == DirectoryTable::Exists_Maybe);
+		dirTable.ParseDirectoryTable(session.GetDirectoryTableSize());
+		CHECK_TRUE(dirTable.EntryExists(ToView(workingDir), true) == DirectoryTable::Exists_Yes);
+
+		StringBuffer<MaxPath> sourceFile;
+		sourceFile.Append(workingDir).Append(TC("Code.cpp"));
+
+		CHECK_TRUE(GetAttributes(sourceFile) == INVALID_FILE_ATTRIBUTES);
+		FileAccessor codeFile(logger, sourceFile.data);
+		CHECK_TRUE(codeFile.CreateWrite());
+		CHECK_TRUE(codeFile.Close());
+		CHECK_TRUE(session.RegisterNewFile(sourceFile.data));
+		CHECK_TRUE(GetAttributes(sourceFile) != INVALID_FILE_ATTRIBUTES);
+
+		CHECK_TRUE(dirTable.EntryExists(sourceFile) == DirectoryTable::Exists_No);
+		dirTable.ParseDirectoryTable(session.GetDirectoryTableSize());
+		CHECK_TRUE(dirTable.EntryExists(sourceFile) == DirectoryTable::Exists_Yes);
+
+		StringBuffer<MaxPath> newDir;
+		newDir.Append(workingDir).Append(TC("NewDir"));
+		StringBuffer<MaxPath> newDirAndSlash(newDir);
+		newDirAndSlash.Append('/');
+
+		CHECK_TRUE(GetAttributes(newDir) == INVALID_FILE_ATTRIBUTES);
+		CHECK_TRUE(CreateDirectoryW(newDir.data));
+		CHECK_TRUE(session.RegisterNewFile(newDir.data));
+		CHECK_TRUE(dirTable.EntryExists(newDir) == DirectoryTable::Exists_No);
+		dirTable.ParseDirectoryTable(session.GetDirectoryTableSize());
+		CHECK_TRUE(dirTable.EntryExists(newDir) == DirectoryTable::Exists_Yes);
+		CHECK_TRUE(GetAttributes(newDir) != INVALID_FILE_ATTRIBUTES);
+		CHECK_TRUE(GetAttributes(newDirAndSlash) != INVALID_FILE_ATTRIBUTES);
+
+		StringBuffer<MaxPath> newDir2;
+		newDir2.Append(workingDir).Append(TC("NewDir2"));
+		CHECK_TRUE(CreateDirectoryW(newDir2.data));
+		CHECK_TRUE(GetAttributes(newDir2) == INVALID_FILE_ATTRIBUTES);
+		CHECK_TRUE(session.RefreshDirectory(workingDir))
+		CHECK_TRUE(GetAttributes(newDir2) != INVALID_FILE_ATTRIBUTES);
+		dirTable.ParseDirectoryTable(session.GetDirectoryTableSize());
+		CHECK_TRUE(dirTable.EntryExists(newDir2) == DirectoryTable::Exists_Yes);
+
+		return true;
+	}
+
+	bool TestRegisterChanges(LoggerWithWriter& logger, const StringBufferBase& testRootDir)
+	{
+		return RunLocal(logger, testRootDir, CheckAttributes);
+	}
+
+	bool TestRegisterChangesRemote(LoggerWithWriter& logger, const StringBufferBase& testRootDir)
+	{
+		return RunRemote(logger, testRootDir, CheckAttributes);
+	}
+
+	bool TestSharedReservedMemory(LoggerWithWriter& logger, const StringBufferBase& testRootDir)
+	{
+		return RunLocal(logger, testRootDir, [](LoggerWithWriter& logger, SessionServer& session, const tchar* workingDir, const RunProcessFunction& runProcess)
+			{
+				StringBuffer<MaxPath> testApp;
+				GetTestAppPath(logger, testApp);
+
+				ProcessStartInfo processInfo;
+				processInfo.application = testApp.data;
+				processInfo.workingDir = workingDir;
+				processInfo.arguments = TC("-sleep=100000");
+				Vector<ProcessHandle> processes;
+
+				for (u32 i=0; i!=128; ++i)
+					processes.push_back(runProcess(processInfo));
+
+				for (auto& process : processes)
+				{
+					if (!process.WaitForExit(100000))
+						return logger.Error(TC("UbaTestApp did not exit in 10 seconds"));
+					u32 exitCode = process.GetExitCode();
+					if (exitCode != 0)
+						return false;
+				}
+
+				return true;
+			});
+	}
+
+}
diff --git a/Engine/Source/Programs/UnrealBuildAccelerator/Test/Private/UbaTestSession.h b/Engine/Source/Programs/UnrealBuildAccelerator/Test/Private/UbaTestSession.h
deleted file mode 100644
diff --git a/Engine/Source/Programs/UnrealBuildAccelerator/Test/Private/UbaTestStdOut.cpp b/Engine/Source/Programs/UnrealBuildAccelerator/Test/Private/UbaTestStdOut.cpp
new file mode 100644
--- /dev/null
+++ b/Engine/Source/Programs/UnrealBuildAccelerator/Test/Private/UbaTestStdOut.cpp
@@ -0,0 +1,110 @@
+// Copyright Epic Games, Inc. All Rights Reserved.
+
+#include "UbaNetworkBackendTcp.h"
+#include "UbaSessionServer.h"
+#include "UbaStorageServer.h"
+
+namespace uba
+{
+	void GetTestAppPath(LoggerWithWriter& logger, StringBufferBase& out);
+
+	bool TestStdOut(LoggerWithWriter& logger, const StringBufferBase& testRootDir, bool remote, const tchar* app = nullptr, const tchar* arg = nullptr, const tchar* expectedOut = nullptr)
+	{
+		LogWriter& logWriter = logger.m_writer;
+		NetworkBackendTcp networkBackend(logWriter);
+
+		bool ctorSuccess = true;
+		NetworkServer networkServer(ctorSuccess, { logWriter });
+
+		StringBuffer<> rootDir;
+		rootDir.Append(testRootDir).Append(TC("Uba"));
+		if (!DeleteAllFiles(logger, rootDir.data))
+			return false;
+
+		StorageServerCreateInfo storageServerInfo(networkServer, rootDir.data, logWriter);
+		storageServerInfo.casCapacityBytes = 1024ull * 1024 * 1024;
+		StorageServer storageServer(storageServerInfo);
+
+		SessionServerCreateInfo sessionServerInfo(storageServer, networkServer, logWriter);
+		sessionServerInfo.checkMemory = false;
+		sessionServerInfo.rootDir = rootDir.data;
+		sessionServerInfo.traceEnabled = true;
+		//sessionServerInfo.remoteLogEnabled = true;
+		SessionServer sessionServer(sessionServerInfo);
+
+		auto sg = MakeGuard([&]() { networkServer.DisconnectClients(); });
+
+		StringBuffer<> workingDir;
+		workingDir.Append(testRootDir).Append(TC("WorkingDir"));
+		if (!DeleteAllFiles(logger, workingDir.data))
+			return false;
+		if (!storageServer.CreateDirectory(workingDir.data))
+			return false;
+		if (!DeleteAllFiles(logger, workingDir.data, false))
+			return false;
+		workingDir.EnsureEndsWithSlash();
+
+		if (!networkServer.StartListen(networkBackend))
+			return logger.Error(TC("Failed to listen"));
+
+		StringBuffer<> testApp;
+		if (!app)
+		{
+			GetTestAppPath(logger, testApp);
+			app = testApp.data;
+		}
+		if (!arg)
+		{
+			arg = TC("-stdout=rootprocess");
+			expectedOut = TC("rootprocess");
+		}
+
+		ProcessStartInfo pi;
+		pi.application = app;
+		pi.arguments = arg;
+		pi.workingDir = workingDir.data;
+		pi.description = TC("StdOutDesc");
+		pi.logFile = TC("Log");
+
+		auto ph = remote ? sessionServer.RunProcessRemote(pi) : sessionServer.RunProcess(pi);
+		if (!ph.WaitForExit(5000))
+			return logger.Error(TC("Timed out waiting for process"));
+		if (auto ec = ph.GetExitCode())
+			return logger.Error(TC("Process exited with error code %u"), ec);
+
+		networkBackend.StopListen();
+		networkServer.DisconnectClients();
+		sessionServer.WaitOnAllTasks();
+
+		auto& logLines = ph.GetLogLines();
+		if (logLines.size() != 1)
+			return logger.Error(TC("Application %s produced %u log line(s) but expected 1"), app, u32(logLines.size()));
+		if (!Equals(logLines[0].text.c_str(), expectedOut))
+			return logger.Error(TC("Application %s produced non-matching log line: %s"), app, logLines[0].text.c_str());
+
+		return true;
+	}
+
+	bool TestStdOutLocal(LoggerWithWriter& logger, const StringBufferBase& testRootDir)
+	{
+		if (!IsWindows)
+			return true;
+		return TestStdOut(logger, testRootDir, false);
+	}
+
+	bool TestStdOutRemote(LoggerWithWriter& logger, const StringBufferBase& testRootDir)
+	{
+		return TestStdOut(logger, testRootDir, true);
+	}
+
+	bool TestStdOutViaCmd(LoggerWithWriter& logger, const StringBufferBase& testRootDir)
+	{
+		if (!IsWindows)
+			return true;
+		StringBuffer<> args;
+		args.Append(TC("/c \""));
+		GetTestAppPath(logger, args);
+		args.Append(TC(" -stdout=foo\""));
+		return TestStdOut(logger, testRootDir, false, TC("c:\\windows\\system32\\cmd.exe"), args.data, TC("foo"));
+	}
+}
diff --git a/Engine/Source/Programs/UnrealBuildAccelerator/Test/Private/UbaTestStdOut.h b/Engine/Source/Programs/UnrealBuildAccelerator/Test/Private/UbaTestStdOut.h
deleted file mode 100644
diff --git a/Engine/Source/Programs/UnrealBuildAccelerator/Test/Private/UbaTestStorage.cpp b/Engine/Source/Programs/UnrealBuildAccelerator/Test/Private/UbaTestStorage.cpp
new file mode 100644
--- /dev/null
+++ b/Engine/Source/Programs/UnrealBuildAccelerator/Test/Private/UbaTestStorage.cpp
@@ -0,0 +1,230 @@
+// Copyright Epic Games, Inc. All Rights Reserved.
+
+#include "UbaFileAccessor.h"
+#include "UbaNetworkBackendTcp.h"
+#include "UbaStorageClient.h"
+#include "UbaStorageServer.h"
+
+namespace uba
+{
+	bool TestStorage(LoggerWithWriter& logger, const StringBufferBase& testRootDir)
+	{
+		#if PLATFORM_LINUX
+		if (true) // TODO: Revisit this... fails on farm but works locally
+			return true;
+		#endif
+
+		StringBuffer<> rootDir;
+		rootDir.Append(testRootDir).Append(TC("Uba"));
+		StorageCreateInfo storageInfo(rootDir.data, logger.m_writer);
+		storageInfo.casCapacityBytes = 1024ull * 1024 * 1024;
+		StorageImpl storage(storageInfo);
+
+		StringBuffer<> detoursLib;
+		GetDirectoryOfCurrentModule(logger, detoursLib);
+		detoursLib.EnsureEndsWithSlash().Append(UBA_DETOURS_LIBRARY);
+
+		storage.LoadCasTable();
+
+		CasKey key;
+		bool deferCreation = false;
+		bool fileIsCompressed = false;
+		if (!storage.StoreCasFile(key, detoursLib.data, CasKeyZero, deferCreation, fileIsCompressed))
+			return logger.Error(TC("Failed to store file %s"), detoursLib.data);
+		if (key == CasKeyZero)
+			return logger.Error(TC("Failed to find file %s"), detoursLib.data);
+
+		StringBuffer<> detoursLibCopy(detoursLib);
+		detoursLibCopy.Append(TC(".tmp"));
+
+		auto deleteFile = MakeGuard([&]() { return DeleteFileW(detoursLibCopy.data); });
+
+		if (!storage.CopyOrLink(key, detoursLibCopy.data, DefaultAttributes()))
+			return logger.Error(TC("Failed to copy cas to file %s"), detoursLibCopy.data);
+
+		FileHandle original;
+		if (!OpenFileSequentialRead(logger, detoursLib.data, original))
+			return logger.Error(TC("Failed to open %s for read"), detoursLib.data);
+		auto closeOriginal = MakeGuard([&]() { return CloseFile(detoursLib.data, original); });
+
+		FileHandle copy;
+		if (!OpenFileSequentialRead(logger, detoursLibCopy.data, copy))
+			return logger.Error(TC("Failed to open %s for read"), detoursLibCopy.data);
+		auto closeCopy = MakeGuard([&]() { return CloseFile(detoursLibCopy.data, copy); });
+
+		u64 originalSize;
+		if (!GetFileSizeEx(originalSize, original))
+			return logger.Error(TC("Failed to get size of %s"), detoursLib.data);
+		u64 copySize;
+		if (!GetFileSizeEx(copySize, copy))
+			return logger.Error(TC("Failed to get size of %s"), detoursLibCopy.data);
+		if (originalSize != copySize)
+			return logger.Error(TC("Size mismatch between %s and %s (%llu vs %llu)"), detoursLib.data, detoursLibCopy.data, originalSize, copySize);
+
+		u8 originalBuffer[64 * 1024];
+		u8 copyBuffer[64 * 1024];
+		u64 left = originalSize;
+		while (left)
+		{
+			u64 toRead = Min(left, (u64)sizeof(originalBuffer));
+			if (!ReadFile(logger, detoursLib.data, original, originalBuffer, toRead))
+				return logger.Error(TC("Failed to read %u from %s"), detoursLib.data);
+			if (!ReadFile(logger, detoursLibCopy.data, copy, copyBuffer, toRead))
+				return logger.Error(TC("Failed to read %u from %s"), detoursLibCopy.data);
+			if (memcmp(originalBuffer, copyBuffer, toRead) != 0)
+				return logger.Error(TC("Data mismatch between %s and %s"), detoursLib.data, detoursLibCopy.data);
+
+			left -= toRead;
+		}
+
+		if (!closeOriginal.Execute())
+			return logger.Error(TC("Failed to close %s"), detoursLib.data);
+
+		if (!closeCopy.Execute())
+			return logger.Error(TC("Failed to close %s"), detoursLibCopy.data);
+
+		if (!deleteFile.Execute())
+			return logger.Error(TC("Failed to delete %s"), detoursLibCopy.data);
+
+		return true;
+	}
+
+	bool TestRemoteStorageStore(LoggerWithWriter& logger, const StringBufferBase& testRootDir)
+	{
+		auto& logWriter = logger.m_writer;
+		NetworkBackendTcp serverTcp(logWriter, TC("ServerTcp"));
+		NetworkBackendTcp clientTcp(logWriter, TC("ClientTcp"));
+
+		bool ctorSuccess = true;
+		NetworkServer server(ctorSuccess, logWriter);
+		NetworkClient client(ctorSuccess, logWriter);
+		auto dg = MakeGuard([&]() { client.Disconnect(); });
+
+		StringBuffer<> rootDir;
+		rootDir.Append(testRootDir).Append(TC("Uba"));
+		if (!DeleteAllFiles(logger, rootDir.data))
+			return false;
+
+		StorageServerCreateInfo storageServerInfo(server, rootDir.data, logger.m_writer);
+		storageServerInfo.casCapacityBytes = 1024ull * 1024 * 1024;
+		StorageServer storageServer(storageServerInfo);
+
+		auto g = MakeGuard([&]() { server.DisconnectClients(); });
+
+		rootDir.Append(TC("Client"));
+		if (!DeleteAllFiles(logger, rootDir.data))
+			return false;
+
+		StorageClientCreateInfo storageClientInfo(client, rootDir.data);
+		StorageClient storageClient(storageClientInfo);
+		storageClient.Start();
+
+		if (!storageClient.LoadCasTable(true))
+			return false;
+
+		rootDir.EnsureEndsWithSlash();
+
+		if (!server.StartListen(serverTcp, 1234))
+			return logger.Error(TC("Failed to listen"));
+		Sleep(100);
+		if (!client.Connect(clientTcp, TC("127.0.0.1"), 1234))
+			return logger.Error(TC("Failed to connect"));
+		auto cg = MakeGuard([&]() { client.Disconnect(); });
+
+		StringBuffer<> fileName;
+		{
+			fileName.Append(rootDir).Append(TC("UbaTestFile"));
+			FileAccessor fileHandle(logger, fileName.data);
+			if (!fileHandle.CreateWrite())
+				return logger.Error(TC("Failed to create file for write"));
+			u8 byte = 'H';
+			if (!fileHandle.Write(&byte, 1))
+				return false;
+			if (!fileHandle.Close())
+				return false;
+		}
+
+		CasKey key;
+		bool storeCompressed = false;
+		if (!storageClient.StoreCasFileClient(key, ToStringKeyLower(fileName), fileName.data, FileMappingHandle{}, 0, 0, TC("UbaTestFile"), false, storeCompressed))
+			return logger.Error(TC("Failed to store file %s"), fileName.data);
+
+		fileName.Clear().Append(testRootDir).Append(TC("Uba")).Append(PathSeparator).Append(TC("UbaTestFile"));
+		if (!storageServer.CopyOrLink(key, fileName.data, DefaultAttributes()))
+			return logger.Error(TC("Failed to copy cas to file %s"), fileName.data);
+		return true;
+	}
+
+	bool TestRemoteStorageFetch(LoggerWithWriter& logger, const StringBufferBase& testRootDir)
+	{
+		auto& logWriter = logger.m_writer;
+		NetworkBackendTcp serverTcp(logWriter, TC("ServerTcp"));
+		NetworkBackendTcp clientTcp(logWriter, TC("ClientTcp"));
+
+		bool ctorSuccess = true;
+		NetworkServer server(ctorSuccess, logWriter);
+		NetworkClient client(ctorSuccess, logWriter);
+		auto dg = MakeGuard([&]() { client.Disconnect(); });
+
+		StringBuffer<> rootDir;
+		rootDir.Append(testRootDir).Append(TC("Uba"));
+		if (!DeleteAllFiles(logger, rootDir.data))
+			return false;
+
+		StorageServerCreateInfo storageServerInfo(server, rootDir.data, logger.m_writer);
+		storageServerInfo.casCapacityBytes = 1024ull * 1024 * 1024;
+		storageServerInfo.storeCompressed = false;
+		StorageServer storageServer(storageServerInfo);
+
+		auto g = MakeGuard([&]() { server.DisconnectClients(); });
+
+		rootDir.Append(TC("Client"));
+		if (!DeleteAllFiles(logger, rootDir.data))
+			return false;
+
+		StorageClientCreateInfo storageClientInfo(client, rootDir.data);
+		StorageClient storageClient(storageClientInfo);
+		storageClient.Start();
+
+		if (!storageClient.LoadCasTable(true))
+			return false;
+
+		rootDir.EnsureEndsWithSlash();
+
+		if (!server.StartListen(serverTcp, 1234))
+			return logger.Error(TC("Failed to listen"));
+		Sleep(100);
+		if (!client.Connect(clientTcp, TC("127.0.0.1"), 1234))
+			return logger.Error(TC("Failed to connect"));
+		auto cg = MakeGuard([&]() { client.Disconnect(); });
+
+		StringBuffer<> fileName;
+#if 1
+		{
+			fileName.Append(rootDir).Append(TC("UbaTestFile"));
+			FileAccessor fileHandle(logger, fileName.data);
+			if (!fileHandle.CreateWrite())
+				return logger.Error(TC("Failed to create file for write"));
+			u8 byte = 'H';
+			if (!fileHandle.Write(&byte, 1))
+				return false;
+			if (!fileHandle.Close())
+				return false;
+		}
+#else
+		fileName.Append(TC("e:\\dev\\fn\\QAGame\\Saved\\StagedBuilds\\PS5_Temp\\Split\\qagame\\content\\paks\\qagame-ps5.ucas\\.copied"));
+#endif
+
+		CasKey casKey;
+		if (!storageServer.CalculateCasKey(casKey, fileName.data))
+			return false;
+
+		casKey = AsCompressed(casKey, false);
+
+		Storage::RetrieveResult result;
+		if (!storageClient.RetrieveCasFile(result, casKey, fileName.data))
+			return logger.Error(TC("Failed to store file %s"), fileName.data);
+
+		return true;
+	}
+}
diff --git a/Engine/Source/Programs/UnrealBuildAccelerator/Test/Private/UbaTestStorage.h b/Engine/Source/Programs/UnrealBuildAccelerator/Test/Private/UbaTestStorage.h
deleted file mode 100644
diff --git a/Engine/Source/Programs/UnrealBuildAccelerator/Test/Private/UbaTestStress.cpp b/Engine/Source/Programs/UnrealBuildAccelerator/Test/Private/UbaTestStress.cpp
new file mode 100644
--- /dev/null
+++ b/Engine/Source/Programs/UnrealBuildAccelerator/Test/Private/UbaTestStress.cpp
@@ -0,0 +1,189 @@
+// Copyright Epic Games, Inc. All Rights Reserved.
+
+#include "UbaCoordinatorWrapper.h"
+#include "UbaFileAccessor.h"
+#include "UbaNetworkBackendTcp.h"
+#include "UbaScheduler.h"
+#include "UbaSessionServer.h"
+#include "UbaStorageServer.h"
+
+namespace uba
+{
+	void GetTestAppPath(LoggerWithWriter& logger, StringBufferBase& out);
+
+	bool TestStress(LoggerWithWriter& logger, const StringBufferBase& testRootDir)
+	{
+		LogWriter& logWriter = logger.m_writer;
+		NetworkBackendTcp networkBackend(logWriter);
+
+		bool ctorSuccess = true;
+		NetworkServer networkServer(ctorSuccess, { logWriter });
+
+		StringBuffer<> rootDir;
+		rootDir.Append(testRootDir).Append(TC("Uba"));
+		if (!DeleteAllFiles(logger, rootDir.data))
+			return false;
+
+		StorageServerCreateInfo storageServerInfo(networkServer, rootDir.data, logWriter);
+		storageServerInfo.casCapacityBytes = 1024ull * 1024 * 1024;
+		StorageServer storageServer(storageServerInfo);
+
+		SessionServerCreateInfo sessionServerInfo(storageServer, networkServer, logWriter);
+		sessionServerInfo.checkMemory = false;
+		sessionServerInfo.rootDir = rootDir.data;
+		sessionServerInfo.traceEnabled = true;
+		//sessionServerInfo.remoteLogEnabled = true;
+		SessionServer sessionServer(sessionServerInfo);
+
+		auto sg = MakeGuard([&]() { networkServer.DisconnectClients(); });
+
+		StringBuffer<> workingDir;
+		workingDir.Append(testRootDir).Append(TC("WorkingDir"));
+		if (!DeleteAllFiles(logger, workingDir.data))
+			return false;
+		if (!storageServer.CreateDirectory(workingDir.data))
+			return false;
+		if (!DeleteAllFiles(logger, workingDir.data, false))
+			return false;
+		workingDir.EnsureEndsWithSlash();
+
+		SchedulerCreateInfo info(sessionServer);
+		//info.forceRemote = true;
+		info.enableProcessReuse = true;
+		info.maxLocalProcessors = 1;
+		Scheduler scheduler(info);
+		scheduler.Start();
+
+		if (!networkServer.StartListen(networkBackend))
+			return logger.Error(TC("Failed to listen"));
+
+		CoordinatorWrapper coordinator;
+
+		bool useCoordinator = true;
+		if (useCoordinator)
+		{
+			return logger.Error(TC("uri/pool/oidc must be populated for this test to work"));
+
+			StringBuffer<512> coordinatorWorkDir(testRootDir);
+			coordinatorWorkDir.EnsureEndsWithSlash().Append(TC("Horde"));
+			StringBuffer<512> binariesDir;
+			if (!GetDirectoryOfCurrentModule(logger, binariesDir))
+				return false;
+
+			CoordinatorCreateInfo cinfo;
+			cinfo.workDir = coordinatorWorkDir.data;
+			cinfo.binariesDir = binariesDir.data;
+
+			cinfo.pool = TC("PoolToUse");
+
+
+			cinfo.maxCoreCount = 400;
+			cinfo.logging = true;
+			if (!coordinator.Create(logger, TC("Horde"), cinfo, networkBackend, networkServer))
+				return false;
+		}
+		auto cg = MakeGuard([&]() { coordinator.Destroy(); });
+
+		StringBuffer<> testApp;
+		GetTestAppPath(logger, testApp);
+
+		u32 counter = 0;
+
+		Atomic<u32> activeCount;
+		Atomic<bool> error;
+		while (true)
+		{
+			StringBuffer<512> fileName;
+			fileName.Append(workingDir).Append(TC("File_")).AppendValue(counter++).Append(TC(".in"));
+			{
+				FileAccessor fa(logger, fileName.data);
+				fa.CreateWrite();
+				Guid g;
+				CreateGuid(g);
+				fa.Write(&g, sizeof(g));
+				fa.Close();
+			}
+			sessionServer.RegisterNewFile(fileName.data);
+			StringBuffer<512> args;
+			args.Append(TC("-file=")).Append(fileName);
+			ProcessStartInfo pi;
+			pi.application = testApp.data;
+			pi.arguments = args.data;
+			pi.workingDir = workingDir.data;
+			pi.description = fileName.Last(PathSeparator) + 6;
+
+			struct Rec
+			{
+				Logger& logger;
+				StringBuffer<256> inFile;
+				Atomic<u32>& activeCount;
+				Atomic<bool>& error;
+			};
+			pi.userData = new Rec { logger, fileName, activeCount, error };
+			pi.exitedFunc = [](void* userData, const ProcessHandle& ph)
+				{
+					auto rec = (Rec*)userData;
+					for (auto& line : ph.GetLogLines())
+						rec->logger.Log(line.type, line.text.c_str(), u32(line.text.size()));
+
+					if (u32 exitCode = ph.GetExitCode())
+					{
+						rec->logger.Error(TC("Process %s failed with exit code %u"), ph.GetStartInfo().GetDescription(), exitCode);
+						rec->error = true;
+					}
+					else
+					{
+						auto& logger = rec->logger;
+						if (!DeleteFileW(rec->inFile.data))
+						{
+							logger.Error(TC("Failed to delete file %s"), rec->inFile.data);
+							rec->error = true;
+						}
+						StringBuffer<256> outFile;
+						outFile.Append(rec->inFile.data, rec->inFile.count - 2).Append(TC("out"));
+						if (!FileExists(logger, outFile.data))
+						{
+							logger.Error(TC("Couldn't find output file %s"), outFile.data);
+							rec->error = true;
+						}
+						if (!DeleteFileW(outFile.data))
+						{
+							logger.Error(TC("Failed to delete file %s"), outFile.data);
+							rec->error = true;
+						}
+					}
+					--rec->activeCount;
+					delete rec;
+				};
+			EnqueueProcessInfo epi(pi);
+			scheduler.EnqueueProcess(epi);
+			++activeCount;
+
+			while (activeCount > 800)
+				Sleep(100);
+			if (error)
+				break;
+
+			if (counter != 20000)
+				continue;
+
+			while (true)
+			{
+				u32 queued;
+				u32 activeLocal;
+				u32 activeRemote;
+				u32 finished;
+				scheduler.GetStats(queued, activeLocal, activeRemote, finished);
+				if (finished == 20000)
+					break;
+			}
+			break;
+		}
+
+		networkBackend.StopListen();
+		networkServer.DisconnectClients();
+		scheduler.Stop();
+
+		return !error;
+	}
+}
diff --git a/Engine/Source/Programs/UnrealBuildAccelerator/Test/Private/UbaTestStress.h b/Engine/Source/Programs/UnrealBuildAccelerator/Test/Private/UbaTestStress.h
deleted file mode 100644
