[UBA]
* Fixed uba test linker error

diff --git a/Engine/Source/Programs/UnrealBuildAccelerator/Test/Private/UbaTest.h b/Engine/Source/Programs/UnrealBuildAccelerator/Test/Private/UbaTest.h
--- a/Engine/Source/Programs/UnrealBuildAccelerator/Test/Private/UbaTest.h
+++ b/Engine/Source/Programs/UnrealBuildAccelerator/Test/Private/UbaTest.h
@@ -2,7 +2,25 @@
 
 #pragma once
 
+#include "UbaStringBuffer.h"
+
 #define CHECK_TRUE(x) \
 	if (!(x)) \
 		return logger.Error(TC("Failed %s (%s:%u)"), TC(#x), __FILE__, __LINE__);
 
+
+namespace uba
+{
+	class LoggerWithWriter;
+	class ProcessHandle;
+	class SessionServer;
+	struct ProcessStartInfo;
+
+	using RunProcessFunction = Function<ProcessHandle(const ProcessStartInfo&)>;
+	using TestSessionFunction = Function<bool(LoggerWithWriter& logger, SessionServer& session, const tchar* workingDir, const RunProcessFunction& runProcess)>;
+	bool RunLocal(LoggerWithWriter& logger, const StringBufferBase& testRootDir, const TestSessionFunction& testFunc, bool enableDetour = true);
+	bool RunRemote(LoggerWithWriter& logger, const StringBufferBase& testRootDir, const TestSessionFunction& testFunc, bool deleteAll = true);
+	void GetTestAppPath(LoggerWithWriter& logger, StringBufferBase& out);
+	const tchar* GetSystemApplication();
+	const tchar* GetSystemArguments();
+}
\ No newline at end of file
diff --git a/Engine/Source/Programs/UnrealBuildAccelerator/Test/Private/UbaTestScheduler.cpp b/Engine/Source/Programs/UnrealBuildAccelerator/Test/Private/UbaTestScheduler.cpp
--- a/Engine/Source/Programs/UnrealBuildAccelerator/Test/Private/UbaTestScheduler.cpp
+++ b/Engine/Source/Programs/UnrealBuildAccelerator/Test/Private/UbaTestScheduler.cpp
@@ -4,17 +4,10 @@
 #include "UbaNetworkServer.h"
 #include "UbaScheduler.h"
 #include "UbaSessionServer.h"
+#include "UbaTest.h"
 
 namespace uba
 {
-	using RunProcessFunction = Function<ProcessHandle(const ProcessStartInfo&)>;
-	using TestSessionFunction = Function<bool(LoggerWithWriter& logger, SessionServer& session, const tchar* workingDir, const RunProcessFunction& runProcess)>;
-	bool RunLocal(LoggerWithWriter& logger, const StringBufferBase& testRootDir, const TestSessionFunction& testFunc, bool enableDetour = true);
-	bool RunRemote(LoggerWithWriter& logger, const StringBufferBase& testRootDir, const TestSessionFunction& testFunc, bool deleteAll = true);
-	void GetTestAppPath(LoggerWithWriter& logger, StringBufferBase& out);
-	const tchar* GetSystemApplication();
-	const tchar* GetSystemArguments();
-
 	bool TestLocalSchedule(LoggerWithWriter& logger, const StringBufferBase& testRootDir)
 	{
 		return RunLocal(logger, testRootDir, [](LoggerWithWriter& logger, SessionServer& session, const tchar* workingDir, const RunProcessFunction& runProcess)
diff --git a/Engine/Source/Programs/UnrealBuildAccelerator/Test/Private/UbaTestSession.cpp b/Engine/Source/Programs/UnrealBuildAccelerator/Test/Private/UbaTestSession.cpp
--- a/Engine/Source/Programs/UnrealBuildAccelerator/Test/Private/UbaTestSession.cpp
+++ b/Engine/Source/Programs/UnrealBuildAccelerator/Test/Private/UbaTestSession.cpp
@@ -11,10 +11,7 @@
 
 namespace uba
 {
-	using RunProcessFunction = Function<ProcessHandle(const ProcessStartInfo&)>;
-	using TestSessionFunction = Function<bool(LoggerWithWriter& logger, SessionServer& session, const tchar* workingDir, const RunProcessFunction& runProcess)>;
-
-	bool RunLocal(LoggerWithWriter& logger, const StringBufferBase& testRootDir, const TestSessionFunction& testFunc, bool enableDetour = true)
+	bool RunLocal(LoggerWithWriter& logger, const StringBufferBase& testRootDir, const TestSessionFunction& testFunc, bool enableDetour)
 	{
 		LogWriter& logWriter = logger.m_writer;
 
@@ -53,7 +50,7 @@ namespace uba
 		return testFunc(logger, session, workingDir.data, [&](const ProcessStartInfo& pi) { return session.RunProcess(pi, true, enableDetour); });
 	}
 
-	bool RunRemote(LoggerWithWriter& logger, const StringBufferBase& testRootDir, const TestSessionFunction& testFunc, bool deleteAll = true)
+	bool RunRemote(LoggerWithWriter& logger, const StringBufferBase& testRootDir, const TestSessionFunction& testFunc, bool deleteAll)
 	{
 		LogWriter& logWriter = logger.m_writer;
 		NetworkBackendTcp tcpBackend(logWriter);
