[UBA]
* Optimized maintenance logic for path table by ~30%

diff --git a/Engine/Source/Programs/UnrealBuildAccelerator/CacheService/Private/UbaCacheService.cpp b/Engine/Source/Programs/UnrealBuildAccelerator/CacheService/Private/UbaCacheService.cpp
--- a/Engine/Source/Programs/UnrealBuildAccelerator/CacheService/Private/UbaCacheService.cpp
+++ b/Engine/Source/Programs/UnrealBuildAccelerator/CacheService/Private/UbaCacheService.cpp
@@ -346,8 +346,7 @@ namespace uba
 			}
 		}
 
-		storageServer.SaveCasTable(true);
-		cacheServer.Save();
+		cacheServer.RunMaintenance(false, {});
 		return 0;
 	}
 }
diff --git a/Engine/Source/Programs/UnrealBuildAccelerator/Common/Private/UbaCacheServer.cpp b/Engine/Source/Programs/UnrealBuildAccelerator/Common/Private/UbaCacheServer.cpp
--- a/Engine/Source/Programs/UnrealBuildAccelerator/Common/Private/UbaCacheServer.cpp
+++ b/Engine/Source/Programs/UnrealBuildAccelerator/Common/Private/UbaCacheServer.cpp
@@ -83,13 +83,19 @@ namespace uba
 			}
 		}
 
-		static constexpr UBA_FORCEINLINE u64 CountBits(u64 bits)
+		static UBA_FORCEINLINE u64 CountBits(u64 bits)
 		{
+			#if defined(__clang__)
+			return __builtin_popcountll(bits);
+			#elif PLATFORM_WINDOWS && defined(_M_X64)
+			return __popcnt64(bits);
+			#else
 			// https://en.wikipedia.org/wiki/Hamming_weight
 			bits -= (bits >> 1) & 0x5555555555555555ull;
 			bits = (bits & 0x3333333333333333ull) + ((bits >> 2) & 0x3333333333333333ull);
 			bits = (bits + (bits >> 4)) & 0x0f0f0f0f0f0f0f0full;
 			return (bits * 0x0101010101010101) >> 56;
+			#endif
 		}
 
 		static UBA_FORCEINLINE u64 FindFirstBit(u64 v)
@@ -640,6 +646,8 @@ namespace uba
 		bool entriesAdded = addsSinceMaintenance != 0;
 		m_addsSinceMaintenance = 0;
 
+		bool allowSave = true;
+
 		u64 startTime = GetTime();
 
 		if (m_shouldWipe)
@@ -1044,10 +1052,13 @@ namespace uba
 			}
 
 			// Add drop cas as work so it can run in the background
-			for (auto& casKey : deletedCasFiles)
+			if (allowSave)
 			{
-				++activeDropCount;
-				m_server.AddWork([&, key = casKey](const WorkContext&) { m_storage.DropCasFile(key, true, TC("")); --activeDropCount; }, 1, TC(""));
+				for (auto& casKey : deletedCasFiles)
+				{
+					++activeDropCount;
+					m_server.AddWork([&, key = casKey](const WorkContext&) { m_storage.DropCasFile(key, true, TC("")); --activeDropCount; }, 1, TC(""));
+				}
 			}
 			++deleteIteration;
 		}
@@ -1145,11 +1156,11 @@ namespace uba
 
 				oldToNewPathOffset.Init(memoryBlock, usedPathOffsets.CountSetBits(), TC("OldToNewPathOffset"));
 
+				CompactPathTable::AddContext context { pathTable };
 				StringBuffer<> temp;
 				usedPathOffsets.Traverse([&](u32 pathOffset)
 				{
-					pathTable.GetString(temp.Clear(), pathOffset);
-					u32 newPathOffset = newPathTable.AddNoLock(temp.data, temp.count);
+					u32 newPathOffset = newPathTable.AddNoLock(context, pathOffset);
 
 					#if 0
 					StringBuffer<> test;
@@ -1240,7 +1251,7 @@ namespace uba
 		if (TimeToMs(dropCasDuration) > 10)
 			m_logger.Detail(TC("  Done deleting cas files (%s)"), TimeToText(dropCasDuration).str);
 
-		if (entriesAdded || deletedCasCount || deleteEntryCount || forceAllSteps)
+		if (allowSave && (entriesAdded || deletedCasCount || deleteEntryCount || forceAllSteps))
 		{
 			u64 saveStart = GetTime();
 			m_logger.Detail(TC("  Saving to disk (v%u)"), CacheFileVersion);
diff --git a/Engine/Source/Programs/UnrealBuildAccelerator/Common/Private/UbaCompactTables.cpp b/Engine/Source/Programs/UnrealBuildAccelerator/Common/Private/UbaCompactTables.cpp
--- a/Engine/Source/Programs/UnrealBuildAccelerator/Common/Private/UbaCompactTables.cpp
+++ b/Engine/Source/Programs/UnrealBuildAccelerator/Common/Private/UbaCompactTables.cpp
@@ -19,6 +19,14 @@ namespace uba
 		m_segmentOffsets.Init(reserveSegmentCount);
 	}
 
+	void CompactPathTable::InitMem()
+	{
+		if (!m_pathTableMem.memory)
+			m_pathTableMem.Init(PathTableMaxSize);
+		if (!m_pathTableMem.writtenSize)
+			m_pathTableMem.AllocateNoLock(1, 1, TC("CompactPathTable"));
+	}
+
 	u32 CompactPathTable::Add(const tchar* str, u64 strLen)
 	{
 		SCOPED_WRITE_LOCK(m_lock, lock);
@@ -36,10 +44,7 @@ namespace uba
 
 	u32 CompactPathTable::AddNoLock(const tchar* str, u64 strLen)
 	{
-		if (!m_pathTableMem.memory)
-			m_pathTableMem.Init(PathTableMaxSize);
-		if (!m_pathTableMem.writtenSize)
-			m_pathTableMem.AllocateNoLock(1, 1, TC("CompactPathTable"));
+		InitMem();
 
 		const tchar* stringKeyString = str;
 		StringBuffer<MaxPath> tempStringKeyStr;
@@ -49,6 +54,150 @@ namespace uba
 		return InternalAdd(str, stringKeyString, strLen);
 	}
 
+	u32 CompactPathTable::AddNoLock(AddContext& context, u32 offset)
+	{
+		auto& table = context.fromTable;
+
+		if (m_version < 3)
+		{
+			StringBuffer<> temp;
+			table.GetString(temp, offset);
+			return InternalAdd(temp.data, temp.data, temp.count);
+		}
+
+		u32 offsets[MaxSegments];
+		bool separators[MaxSegments];
+
+		u32 parentOffset = 0;
+		u32 pathOffset = 0;
+
+		BinaryReader reader(table.m_pathTableMem.memory, 0, table.m_pathTableMem.writtenSize);
+		u32 offsetCount = 0;
+
+		u32 isChild = 0;
+
+		u32* fromOffsets = context.fromOffsets;
+		u32* toOffsets = context.toOffsets;
+		u32* pathOffsets = context.pathOffsets;
+		u32 searchIndex = context.offsetsCount - 1;
+
+		// Traverse backwards in old tree. Early out if context already have entry
+		while (true)
+		{
+			UBA_ASSERT(offsetCount < sizeof_array(offsets));
+			reader.SetPosition(offset);
+			u64 parentRelativeOffset = reader.Read7BitEncoded();
+			u32 readPos = u32(reader.GetPosition());
+			offsets[offsetCount++] = readPos;
+			separators[offsetCount] = parentRelativeOffset & 1;
+			parentRelativeOffset >>= 1;
+			if (parentRelativeOffset == offset || offsetCount == (MaxSegments-1))
+				break;
+			offset -= u32(parentRelativeOffset);
+
+			while (readPos < fromOffsets[searchIndex])
+				--searchIndex;
+			if (readPos == fromOffsets[searchIndex])
+			{
+				parentOffset = toOffsets[searchIndex];
+				pathOffset = pathOffsets[searchIndex];
+				isChild = 1;
+				--offsetCount;
+				break;
+			}
+		}
+
+		if (offsetCount == (MaxSegments-1))
+			return false;
+
+		if (!isChild)
+			searchIndex = 0;
+
+		u32 pathPositions[MaxSegments];
+		pathPositions[offsetCount] = pathOffset;
+
+		auto& path = context.path;
+		path.Resize(pathOffset);
+		
+		u32 oldSegOffsets[MaxSegments];
+
+		// Create full string and store path positions
+		for (u32 i=offsetCount;i; --i)
+		{
+			reader.SetPosition(offsets[i - 1]);
+			u32 oldSegOffset = u32(reader.Read7BitEncoded());
+			oldSegOffsets[i] = oldSegOffset;
+			if (oldSegOffset != 0)
+				reader.SetPosition(oldSegOffset);
+
+			if (pathOffset)
+				path.Append(separators[i] ? '.' : PathSeparator);
+
+			reader.ReadString(path);
+
+			pathOffset = 1;
+
+			pathPositions[i - 1] = path.count;
+		}
+
+		// Go backwards to figure out how many new entries to add
+		u32 newOffsetIndices[MaxSegments];
+		u32 newOffsetCount = 0;
+		for (u32 i=offsetCount;i; --i)
+		{
+			StringKey key = ToStringKeyNoCheck(path.data, pathPositions[newOffsetCount]);
+			bool added = false;
+			newOffsetIndices[newOffsetCount] = m_offsets.InsertIndex(key, added); // m_offsets might re-allocate so we need to track index instead
+			if (!added)
+			{
+				parentOffset = m_offsets.GetValueFromIndex(newOffsetIndices[newOffsetCount]);
+				isChild = 1;
+				break;
+			}
+			++newOffsetCount;
+		}
+
+		// Go forward and add new entries
+		for (u32 i=newOffsetCount;i; --i)
+		{
+			u32 iMinusOne = i - 1;
+
+			offset = u32(m_pathTableMem.writtenSize);
+			m_offsets.GetValueFromIndex(newOffsetIndices[iMinusOne]) = offset;
+
+			bool separator = separators[i];
+
+			if (oldSegOffsets[i] && oldSegOffsets[i] < m_commonSize)
+			{
+				InternalAddWithExistingSegment(offset, parentOffset, separator, oldSegOffsets[i]);
+			}
+			else
+			{
+				StringView segment(path.data + pathPositions[i] + isChild, pathPositions[iMinusOne] - pathPositions[i] - isChild);
+				StringKey segmentKey = ToStringKeyNoCheck(segment.data, segment.count);
+				bool added = false;
+				u32& strOffset = m_segmentOffsets.Insert(segmentKey, added);
+
+				if (added)
+					strOffset = InternalAddWithNewSegment(offset, parentOffset, separator, segment);
+				else
+					InternalAddWithExistingSegment(offset, parentOffset, separator, strOffset);
+			}
+
+			searchIndex++;
+			fromOffsets[searchIndex] = offsets[iMinusOne];
+			toOffsets[searchIndex] = offset;
+			pathOffsets[searchIndex] = pathPositions[iMinusOne];
+
+			isChild = 1;
+			parentOffset = offset;
+		}
+
+		context.offsetsCount = searchIndex + 1;
+
+		return offset;
+	}
+
 	void CompactPathTable::AddCommonStringSegments()
 	{
 		const tchar* commonSegments1[] =
@@ -91,10 +240,7 @@ namespace uba
 		};
 		// 124 bytes used
 
-		if (!m_pathTableMem.memory)
-			m_pathTableMem.Init(PathTableMaxSize);
-		if (!m_pathTableMem.writtenSize)
-			m_pathTableMem.AllocateNoLock(1, 1, TC("CompactPathTable"));
+		InitMem();
 
 		auto addSegments = [this](const tchar** segments, u64 segmentCount)
 			{
@@ -147,32 +293,14 @@ namespace uba
 			StringKey segmentKey = ToStringKeyNoCheck(seg, segLen);
 			added = false;
 			u32& strOffset = m_segmentOffsets.Insert(segmentKey, added);
-			if (added)
-			{
-				offset = u32(m_pathTableMem.writtenSize);
-				u32 parentRelativeOffset = ((offset - parentOffset) << 1) + (parentSeparator == '.');
-
-				// Put string directly after current element and set segment offset to 0
-				u64 bytesForString = GetStringWriteSize(seg, segLen);
-				u8 bytesForParent = Get7BitEncodedCount(parentRelativeOffset);
-				u64 memSize = bytesForParent + 1 + bytesForString;
-			
-				BinaryWriter writer((u8*)m_pathTableMem.AllocateNoLock(memSize, 1, TC("CompactPathTable")), 0, memSize);
-				writer.Write7BitEncoded(parentRelativeOffset);
-				writer.Write7BitEncoded(0);
-				writer.WriteString(seg, segLen);
-				strOffset = offset + bytesForParent + 1;
-				return offset;
-			}
 
 			offset = u32(m_pathTableMem.writtenSize);
-			u32 parentRelativeOffset = ((offset - parentOffset) << 1) + (parentSeparator == '.');
 
-			u8 bytesForParent = Get7BitEncodedCount(parentRelativeOffset);
-			u64 memSize = bytesForParent + Get7BitEncodedCount(strOffset);
-			BinaryWriter writer((u8*)m_pathTableMem.AllocateNoLock(memSize, 1, TC("CompactPathTable")), 0, memSize);
-			writer.Write7BitEncoded(parentRelativeOffset);
-			writer.Write7BitEncoded(strOffset);
+			if (added)
+				strOffset = InternalAddWithNewSegment(offset, parentOffset, parentSeparator == '.', StringView(seg, u32(segLen)));
+			else
+				InternalAddWithExistingSegment(offset, parentOffset, parentSeparator == '.', strOffset);
+
 			return offset;
 		}
 
@@ -245,6 +373,34 @@ namespace uba
 		return offset;
 	}
 
+	u32 CompactPathTable::InternalAddWithNewSegment(u32 offset, u32 parentOffset, bool separator, const StringView& segment)
+	{
+		u32 parentRelativeOffset = ((offset - parentOffset) << 1) + separator;
+
+		// Put string directly after current element and set segment offset to 0
+		u64 bytesForString = GetStringWriteSize(segment.data, segment.count);
+		u8 bytesForParent = Get7BitEncodedCount(parentRelativeOffset);
+		u64 memSize = bytesForParent + 1 + bytesForString;
+	
+		BinaryWriter writer((u8*)m_pathTableMem.AllocateNoLock(memSize, 1, TC("CompactPathTable")), 0, memSize);
+		writer.Write7BitEncoded(parentRelativeOffset);
+		writer.Write7BitEncoded(0);
+		writer.WriteString(segment);
+		return offset + bytesForParent + 1;
+	}
+
+	void CompactPathTable::InternalAddWithExistingSegment(u32 offset, u32 parentOffset, bool separator, u32 segmentOffset)
+	{
+		u32 parentRelativeOffset = ((offset - parentOffset) << 1) + separator;
+
+		u8 bytesForParent = Get7BitEncodedCount(parentRelativeOffset);
+		u64 memSize = bytesForParent + Get7BitEncodedCount(segmentOffset);
+		BinaryWriter writer((u8*)m_pathTableMem.AllocateNoLock(memSize, 1, TC("CompactPathTable")), 0, memSize);
+		writer.Write7BitEncoded(parentRelativeOffset);
+		writer.Write7BitEncoded(segmentOffset);
+	}
+
+
 	bool CompactPathTable::GetString(StringBufferBase& out, u64 offset) const
 	{
 		#if UBA_DEBUG
@@ -254,10 +410,8 @@ namespace uba
 		}
 		#endif
 
-		enum { MaxSegments = 50 }; // This number is arbitrary.. don't know how many folders/dots deep something can be. We want it to be as low as possible to catch corrupt tables
-		u32 offsets[MaxSegments + 1];
-		bool separators[MaxSegments + 1];
-		offsets[0] = u32(offset);
+		u32 offsets[MaxSegments];
+		bool separators[MaxSegments];
 		u32 offsetCount = 0;
 
 		BinaryReader reader(m_pathTableMem.memory, offset, m_pathTableMem.writtenSize);
@@ -270,33 +424,33 @@ namespace uba
 				UBA_ASSERT(offsetCount < sizeof_array(offsets));
 				reader.SetPosition(offset);
 				u64 parentRelativeOffset = reader.Read7BitEncoded();
+				offsets[offsetCount-1] = u32(reader.GetPosition());
 				separators[offsetCount] = parentRelativeOffset & 1;
 				parentRelativeOffset >>= 1;
-				if (parentRelativeOffset == offset || offsetCount == MaxSegments)
+				if (parentRelativeOffset == offset || offsetCount == (MaxSegments-1))
 					break;
 				offset -= parentRelativeOffset;
-				offsets[offsetCount] = u32(offset);
 				++offsetCount;
 			}
 		}
 		else
 		{
-			while (offset || offsetCount == MaxSegments)
+			while (offset && offsetCount < (MaxSegments-1))
 			{
 				++offsetCount;
 				UBA_ASSERT(offsetCount < sizeof_array(offsets));
 				reader.SetPosition(offset);
 				offset = (u32)reader.Read7BitEncoded();
+				offsets[offsetCount-1] = u32(reader.GetPosition());
 				separators[offsetCount] = 0;
 				if (m_version >= 2)
 				{
 					separators[offsetCount] = offset & 1;
 					offset >>= 1;
 				}
-				offsets[offsetCount] = u32(offset);
 			}
 		}
-		if (offsetCount == MaxSegments)
+		if (offsetCount == (MaxSegments-1))
 			return false;
 
 		{
@@ -304,7 +458,6 @@ namespace uba
 			for (u32 i=offsetCount;i; --i)
 			{
 				reader.SetPosition(offsets[i-1]);
-				reader.Read7BitEncoded();
 				u32 strOffset = u32(reader.Read7BitEncoded());
 				if (strOffset != 0)
 					reader.SetPosition(strOffset);
diff --git a/Engine/Source/Programs/UnrealBuildAccelerator/Common/Private/UbaNetworkServer.cpp b/Engine/Source/Programs/UnrealBuildAccelerator/Common/Private/UbaNetworkServer.cpp
--- a/Engine/Source/Programs/UnrealBuildAccelerator/Common/Private/UbaNetworkServer.cpp
+++ b/Engine/Source/Programs/UnrealBuildAccelerator/Common/Private/UbaNetworkServer.cpp
@@ -868,7 +868,6 @@ namespace uba
 			worker->m_context->connection = nullptr;
 			worker->m_context->workAvailable.Set();
 		}
-		UBA_ASSERT(m_createdWorkerCount);
 	}
 
 	void NetworkServer::DoWork(u32 count)
diff --git a/Engine/Source/Programs/UnrealBuildAccelerator/Common/Public/UbaCompactTables.h b/Engine/Source/Programs/UnrealBuildAccelerator/Common/Public/UbaCompactTables.h
--- a/Engine/Source/Programs/UnrealBuildAccelerator/Common/Public/UbaCompactTables.h
+++ b/Engine/Source/Programs/UnrealBuildAccelerator/Common/Public/UbaCompactTables.h
@@ -21,6 +21,9 @@ namespace uba
 		u32 Add(const tchar* str, u64 strLen, u32& outRequiredPathTableSize);
 		u32 AddNoLock(const tchar* str, u64 strLen);
 
+		struct AddContext;
+		u32 AddNoLock(AddContext& context, u32 offset);
+
 		void AddCommonStringSegments();
 
 		bool GetString(StringBufferBase& out, u64 offset) const;
@@ -42,8 +45,24 @@ namespace uba
 		template<typename Func>
 		bool TraversePaths(const Func& func) const;
 
+		void InitMem();
+
+		enum { MaxSegments = 48 }; // This number is arbitrary.. don't know how many folders/dots deep something can be. We want it to be as low as possible to catch corrupt tables
+		struct AddContext
+		{
+			const CompactPathTable& fromTable;
+			StringBuffer<> path;
+			u32 fromOffsets[MaxSegments] = { 0 };
+			u32 toOffsets[MaxSegments] = { 0 };
+			u32 pathOffsets[MaxSegments] = { 0 };
+			u32 offsetsCount = 1;
+		};
+
 	private:
 		u32 InternalAdd(const tchar* str, const tchar* stringKeyString, u64 strLen);
+		u32 InternalAddWithNewSegment(u32 offset, u32 parentOffset, bool separator, const StringView& segment);
+		void InternalAddWithExistingSegment(u32 offset, u32 parentOffset, bool separator, u32 segmentOffset);
+
 		ReaderWriterLock m_lock;
 		MemoryBlock m_pathTableMem;
 		HashMap<StringKey, u32, true> m_offsets;
diff --git a/Engine/Source/Programs/UnrealBuildAccelerator/Test/Private/UbaTestCache.cpp b/Engine/Source/Programs/UnrealBuildAccelerator/Test/Private/UbaTestCache.cpp
--- a/Engine/Source/Programs/UnrealBuildAccelerator/Test/Private/UbaTestCache.cpp
+++ b/Engine/Source/Programs/UnrealBuildAccelerator/Test/Private/UbaTestCache.cpp
@@ -61,6 +61,52 @@ namespace uba
 
 	bool TestCompactPathTable(LoggerWithWriter& logger, const StringBufferBase& rootDir)
 	{
+		for (u32 useCommon=0; useCommon<=1; ++useCommon)
+		{
+			const tchar* pathsStr[] =
+			{
+				TC("Foo/Bar/Meh.h"),
+				TC("Foo/Bar/Meh.cpp"),
+				TC("Foo/Bar/Moo.h"),
+				TC("Foo/Boo/Rud.h"),
+				TC("Foo/Boo/Rud.cpp"),
+				TC(")/Boo/Rud.cpp"),
+				TC("%/cl.cpp"),
+				TC(")/Boo/Rud.inl"),
+			};
+			constexpr u32 PathCount = sizeof_array(pathsStr);
+
+			StringBuffer<128> paths[PathCount];
+			for (u32 i=0;i!=PathCount; ++i)
+				paths[i].Append(pathsStr[i]).FixPathSeparators();
+
+			CompactPathTable table(CaseInsensitiveFs, 0, 0, 3);
+			table.InitMem();
+			if (useCommon)
+				table.AddCommonStringSegments();
+
+			u32 offsets[PathCount];
+			for (u32 i=0;i!=PathCount; ++i)
+				offsets[i] = table.AddNoLock(paths[i].data, paths[i].count);
+
+			CompactPathTable table2(CaseInsensitiveFs, 0, 0, 3);
+			table2.InitMem();
+			if (useCommon)
+				table2.AddCommonStringSegments();
+
+			CompactPathTable::AddContext context{table};
+			for (u32 i=0;i!=PathCount; ++i)
+			{
+				u32 offsets2 = table2.AddNoLock(context, offsets[i]);
+
+				StringBuffer<> temp;
+				if (!table2.GetString(temp, offsets2))
+					return logger.Error(TC("Error getting offset %u from table2"), offsets2);
+				if (!temp.Equals(paths[i].data))
+					return logger.Error(TC("Error adding %s to table2. Found %s"), paths[i].data, temp.data);
+			}
+		}
+
 		for (u32 version=0; version<=CacheBucketVersion; ++version)
 		{
 			CompactPathTable table(CaseInsensitiveFs, 0, 0, version);
