[UBA]
* Fixed detouring bug not handling radlink.exe spawning mt.exe
* Improved logging for errors happening on the farm related to cache corrupt path/caskey tables
* Added ref counting of FileInfo to know when decompressed obj files can be deleted
* Added more error logging for file functions on posix failing in unexpected ways

diff --git a/Engine/Source/Programs/UnrealBuildAccelerator/Common/Private/UbaCacheClient.cpp b/Engine/Source/Programs/UnrealBuildAccelerator/Common/Private/UbaCacheClient.cpp
--- a/Engine/Source/Programs/UnrealBuildAccelerator/Common/Private/UbaCacheClient.cpp
+++ b/Engine/Source/Programs/UnrealBuildAccelerator/Common/Private/UbaCacheClient.cpp
@@ -994,6 +994,7 @@ namespace uba
 			StorageImpl::CasEntry* casEntry;
 			if (m_storage.HasCasFile(casKey, &casEntry))
 			{
+				SCOPED_READ_LOCK(casEntry->lock, lock);
 				UBA_ASSERT(!IsNormalized(casKey));
 				StringBuffer<> casKeyFileName;
 				if (!m_storage.GetCasFileName(casKeyFileName, casKey))
@@ -1021,7 +1022,7 @@ namespace uba
 					fileSize = file.GetSize();
 				}
 
-				if (!SendFile(logger, m_client, casKey, fileData, fileSize, casKeyFileName.data))
+				if (!SendFile(logger, m_client, casKey, fileData, fileSize, path.data))
 					return false;
 
 				outBytesSent += fileSize;
@@ -1072,7 +1073,7 @@ namespace uba
 				u8* dataToSend = block.memory;
 				u64 sizeToSend = block.writtenSize;
 
-				if (!sender.SendFileCompressed(casKey, path.data, dataToSend, sizeToSend, TC("SendCacheEntry")))
+				if (!sender.SendFileCompressed(casKey, path.data, dataToSend, sizeToSend, path.data))
 					return m_logger.Warning(TC("Failed to send cas content for file %s"), path.data);
 
 				outBytesSent += sender.m_bytesSent;
diff --git a/Engine/Source/Programs/UnrealBuildAccelerator/Common/Private/UbaCacheServer.cpp b/Engine/Source/Programs/UnrealBuildAccelerator/Common/Private/UbaCacheServer.cpp
--- a/Engine/Source/Programs/UnrealBuildAccelerator/Common/Private/UbaCacheServer.cpp
+++ b/Engine/Source/Programs/UnrealBuildAccelerator/Common/Private/UbaCacheServer.cpp
@@ -1527,7 +1527,7 @@ namespace uba
 			u32 casKeyOffset = bucket.m_casKeyTable.Add(casKey, pathOffset, requiredCasTableSize);
 			auto insres = inputs.insert(casKeyOffset);
 			if (!insres.second)
-				return m_logger.Warning(TC("Client %u Input file %s exists more than once in cache entry (CasKey: %s ClientVersion: %u PathtableVersion: %u)"), clientId, path.data, CasKeyString(casKey).str, clientVersion, connectionBucket.pathTable.GetVersion());
+				return m_logger.Warning(TC("Client %u Input file %s exists more than once in cache entry (PathOffset %u, CasKey: %s CasKeyOffset: %u ClientVersion: %u PathtableVersion: %u)"), clientId, path.data, pathOffset, CasKeyString(casKey).str, casKeyOffset, clientVersion, connectionBucket.pathTable.GetVersion());
 
 			bytesForInput += Get7BitEncodedCount(casKeyOffset);
 
diff --git a/Engine/Source/Programs/UnrealBuildAccelerator/Common/Private/UbaFile.cpp b/Engine/Source/Programs/UnrealBuildAccelerator/Common/Private/UbaFile.cpp
--- a/Engine/Source/Programs/UnrealBuildAccelerator/Common/Private/UbaFile.cpp
+++ b/Engine/Source/Programs/UnrealBuildAccelerator/Common/Private/UbaFile.cpp
@@ -185,10 +185,7 @@ namespace uba
 		struct stat attr;
 		int res = fstat(asFileDescriptor(hFile), &attr);
 		if (res != 0)
-		{
-			UBA_ASSERTF(false, TC("GetFileInformationByHandle (fstat) error handling not implemented"));
-			return false;
-		}
+			return logger.Error(TC("GetFileInformationByHandle (fstat) failed on %s (%s)"), fileName, strerror(errno));
 		out.lastWriteTime = FromTimeSpec(attr.st_mtimespec);
 		out.attributes = attr.st_mode;
 		out.volumeSerialNumber = attr.st_dev;
@@ -214,7 +211,8 @@ namespace uba
 		int res = stat(fileName, &attr);
 		if (res != 0)
 		{
-			UBA_ASSERTF(errno == ENOENT, TC("GetFileInformation: stat failed for file %s and this error handling not implemented (%s)"), fileName, strerror(errno));
+			if (errno != ENOENT)
+				logger.Warning(TC("GetFileInformation: stat failed for file %s and this error handling not implemented (%s)"), fileName, strerror(errno));
 			SetLastError(ERROR_FILE_NOT_FOUND);
 			return false;
 		}
diff --git a/Engine/Source/Programs/UnrealBuildAccelerator/Common/Private/UbaObjectFile.cpp b/Engine/Source/Programs/UnrealBuildAccelerator/Common/Private/UbaObjectFile.cpp
--- a/Engine/Source/Programs/UnrealBuildAccelerator/Common/Private/UbaObjectFile.cpp
+++ b/Engine/Source/Programs/UnrealBuildAccelerator/Common/Private/UbaObjectFile.cpp
@@ -60,7 +60,7 @@ namespace uba
 					OodleLZ_FuzzSafe_Yes, OodleLZ_CheckCRC_No, OodleLZ_Verbosity_None, NULL, 0, NULL, NULL, decoderMem, decoredMemSize);
 				if (decompLen != decompressedBlockSize)
 				{
-					logger.Error(TC("Failed to decompress file %s"), hint);
+					logger.Error(TC("Failed to decompress file %s (Compressed size %llu)"), hint, dataSize);
 					return nullptr;
 				}
 
@@ -86,7 +86,7 @@ namespace uba
 		{
 			if (ownsData)
 				free(data);
-			logger.Error(TC("Unknown object file format. Maybe msvc FE IL? (%s)"), hint);
+			logger.Error(TC("Unknown object file format (Size %llu). Maybe msvc FE IL? (%s)"), dataSize, hint);
 			return nullptr;
 		}
 
diff --git a/Engine/Source/Programs/UnrealBuildAccelerator/Common/Private/UbaStorageServer.cpp b/Engine/Source/Programs/UnrealBuildAccelerator/Common/Private/UbaStorageServer.cpp
--- a/Engine/Source/Programs/UnrealBuildAccelerator/Common/Private/UbaStorageServer.cpp
+++ b/Engine/Source/Programs/UnrealBuildAccelerator/Common/Private/UbaStorageServer.cpp
@@ -791,7 +791,7 @@ namespace uba
 							// Unless writeReceivedCasFilesToDisk is set this is extremely unlikely to happen. "PerModuleInline.gen.cpp" seems to be the only one (they are supposed to all be the same)
 							writer.WriteU16(u16(~0));
 							writer.WriteBool(m_traceStore);
-							m_logger.Detail(TC("Store request of %s which already exists in casdb (%s)"), CasKeyString(casKey).str, hint.data);
+							m_logger.Detail(TC("Client %u Store request of %s which already exists in casdb (%s)"), connectionInfo.GetId(), CasKeyString(casKey).str, hint.data);
 							return true;
 						}
 					}
@@ -809,13 +809,13 @@ namespace uba
 						CasEntryAccessed(casEntry);
 						writer.WriteU16(u16(~0));
 						writer.WriteBool(m_traceStore);
-						m_logger.Detail(TC("Store request of %s which already exists in casdb (%s)"), CasKeyString(casKey).str, hint.data);
+						m_logger.Detail(TC("Client %u Store request of %s which was just added to casdb (%s)"), connectionInfo.GetId(), CasKeyString(casKey).str, hint.data);
 						return true;
 					}
 				}
 				if (!fileSize)
 				{
-					m_logger.Error(TC("Store from client is of 0 size (%s)"), hint.data);
+					m_logger.Error(TC("Client %u Store is of 0 size (%s)"), connectionInfo.GetId(), hint.data);
 					casEntry.verified = false;
 					return false;
 				}
diff --git a/Engine/Source/Programs/UnrealBuildAccelerator/Detours/Private/UbaDetoursFileMappingTable.h b/Engine/Source/Programs/UnrealBuildAccelerator/Detours/Private/UbaDetoursFileMappingTable.h
--- a/Engine/Source/Programs/UnrealBuildAccelerator/Detours/Private/UbaDetoursFileMappingTable.h
+++ b/Engine/Source/Programs/UnrealBuildAccelerator/Detours/Private/UbaDetoursFileMappingTable.h
@@ -7,6 +7,7 @@
 
 namespace uba
 {
+	struct DetouredHandle;
 	struct MemoryFile;
 
 
@@ -57,6 +58,7 @@ namespace uba
 		u64 fileMapMemSize = 0;
 
 		#if PLATFORM_WINDOWS
+		u32 refCount = 0;
 		u32 fileMapDesiredAccess = PAGE_READONLY;
 		u32 fileMapViewDesiredAccess = FILE_MAP_READ; // This is the parameter to MapViewOfFile
 		#endif
@@ -79,7 +81,8 @@ namespace uba
 		ReaderWriterLock m_lookupLock;
 		GrowingUnorderedMap<StringKey, FileInfo> m_lookup;
 		ReaderWriterLock m_memLookupLock;
-		UnorderedMap<const void*, int> m_memLookup;
+		struct Entry { int refCount = 0; DetouredHandle* handle = nullptr; };
+		UnorderedMap<const void*, Entry> m_memLookup;
 	};
 
 	enum : u8 { AccessFlag_Read = 1, AccessFlag_Write = 2 };
diff --git a/Engine/Source/Programs/UnrealBuildAccelerator/Detours/Private/Windows/UbaDetoursFileDecompressor.h b/Engine/Source/Programs/UnrealBuildAccelerator/Detours/Private/Windows/UbaDetoursFileDecompressor.h
--- a/Engine/Source/Programs/UnrealBuildAccelerator/Detours/Private/Windows/UbaDetoursFileDecompressor.h
+++ b/Engine/Source/Programs/UnrealBuildAccelerator/Detours/Private/Windows/UbaDetoursFileDecompressor.h
@@ -387,8 +387,9 @@ namespace uba
 				{
 					SCOPED_WRITE_LOCK(g_mappedFileTable.m_memLookupLock, lock3);
 					auto insres = g_mappedFileTable.m_memLookup.try_emplace(info.fileMapMem);
-					UBA_ASSERTF(insres.second, TC("Failed to add memory 0x%llx for %s"), u64(info.fileMapMem), info.originalName);
-					insres.first->second = 1;
+					UBA_ASSERTF(insres.second, TC("Failed to add memory 0x%llx for %s"), u64(info.fileMapMem), info.originalName);(void)insres;
+					//auto& entry = insres.first->second;
+					//entry.refCount = 1;
 				}
 
 				preload.event.Set();
diff --git a/Engine/Source/Programs/UnrealBuildAccelerator/Detours/Private/Windows/UbaDetoursFunctionsKernelBase.inl b/Engine/Source/Programs/UnrealBuildAccelerator/Detours/Private/Windows/UbaDetoursFunctionsKernelBase.inl
--- a/Engine/Source/Programs/UnrealBuildAccelerator/Detours/Private/Windows/UbaDetoursFunctionsKernelBase.inl
+++ b/Engine/Source/Programs/UnrealBuildAccelerator/Detours/Private/Windows/UbaDetoursFunctionsKernelBase.inl
@@ -193,7 +193,7 @@ HANDLE Detoured_CreateFileW(LPCWSTR lpFileName, DWORD dwDesiredAccess, DWORD dwS
 {
 	DETOURED_CALL(CreateFileW);
 	DEBUG_LOG_DETOURED(L"CreateFileW", L"%ls", lpFileName);
-	u32 disallowDetour = Equals(lpFileName, L"nul");
+	u32 disallowDetour = !lpFileName || Equals(lpFileName, L"nul");
 	t_disallowDetour += disallowDetour;
 	t_createFileFileName = lpFileName;
 	HANDLE h = True_CreateFileW(lpFileName, dwDesiredAccess, dwShareMode, lpSecurityAttributes, dwCreationDisposition, dwFlagsAndAttributes, hTemplateFile);
@@ -2314,14 +2314,15 @@ HANDLE Detoured_CreateFileMappingW(HANDLE hFile, LPSECURITY_ATTRIBUTES lpFileMap
 {
 	DETOURED_CALL(CreateFileMappingW);
 	HANDLE trueHandle = hFile;
+	FileObject* fo = nullptr;
 	if (isDetouredHandle(hFile))
 	{
 		auto& dh = asDetouredHandle(hFile);
 
-		FileObject& fo = *dh.fileObject;
-		fo.wasUsed = true;
+		fo = dh.fileObject;
+		fo->wasUsed = true;
 
-		FileInfo& fi = *fo.fileInfo;
+		FileInfo& fi = *fo->fileInfo;
 		if (fi.memoryFile || fi.isFileMap)
 		{
 			auto mdh = new DetouredHandle(HandleType_FileMapping);
@@ -2367,6 +2368,9 @@ HANDLE Detoured_CreateFileMappingW(HANDLE hFile, LPSECURITY_ATTRIBUTES lpFileMap
 			ToInvestigate(L"Mapping already exists");
 		auto detouredHandle = new DetouredHandle(HandleType_FileMapping);
 		detouredHandle->trueHandle = mappingHandle;
+		//detouredHandle->fileObject = fo;
+		//if (fo)
+		//	InterlockedIncrement(&fo->refCount);
 		mappingHandle = makeDetouredHandle(detouredHandle);
 	}
 	DEBUG_LOG_TRUE(L"CreateFileMappingW", L"File %llu, Size: %llu (%s) -> %llu", uintptr_t(hFile), ToLargeInteger(dwMaximumSizeHigh, dwMaximumSizeLow).QuadPart, HandleToName(hFile), u64(mappingHandle));
@@ -2482,7 +2486,15 @@ LPVOID Detoured_MapViewOfFileEx(HANDLE hFileMappingObject, DWORD dwDesiredAccess
 			DEBUG_LOG_DETOURED(L"MapViewOfFileEx", L"%llu Size %llu (%s) -> 0x%llx", uintptr_t(hFileMappingObject), dwNumberOfBytesToMap, HandleToName(hFileMappingObject), u64(mem));
 
 			SCOPED_WRITE_LOCK(g_mappedFileTable.m_memLookupLock, lock);
-			++g_mappedFileTable.m_memLookup[mem];
+			auto& entry = g_mappedFileTable.m_memLookup[mem];
+			if (!entry.handle)
+			{
+				InterlockedIncrement(&dh.fileObject->refCount);
+				auto newDh = new DetouredHandle(dh.type);
+				newDh->fileObject = dh.fileObject;
+				entry.handle = newDh;
+			}
+			++entry.refCount;
 			return mem;
 		}
 		UBA_ASSERT(dh.trueHandle != INVALID_HANDLE_VALUE);
@@ -2511,8 +2523,13 @@ BOOL Detoured_UnmapViewOfFileEx(PVOID lpBaseAddress, ULONG UnmapFlags)
 		auto it = g_mappedFileTable.m_memLookup.find(lpBaseAddress);
 		if (it != g_mappedFileTable.m_memLookup.end())
 		{
-			if (!--it->second)
+			auto& entry = it->second;
+			if (!--entry.refCount)
+			{
+				if (entry.handle)
+					Detoured_NtClose(makeDetouredHandle(entry.handle));
 				g_mappedFileTable.m_memLookup.erase(it);
+			}
 			SetLastError(ERROR_SUCCESS);
 			return true;
 		}
@@ -2977,7 +2994,10 @@ BOOL Detoured_CreateProcessW(LPCWSTR lpApplicationName, LPWSTR lpCommandLine, LP
 	StringBuffer<512> application;
 	if (lpApplicationName && *lpApplicationName)
 	{
-		FixPath(application, lpApplicationName);
+		if (lpApplicationName[1] == ':') // Only fixup absolute paths (since we want to devirtualize them
+			FixPath(application, lpApplicationName);
+		else
+			application.Append(lpApplicationName);
 		UBA_ASSERTF(application.count, TC("Invalid application name from application field: %s"), lpApplicationName);
 	}
 	else
@@ -2987,7 +3007,10 @@ BOOL Detoured_CreateProcessW(LPCWSTR lpApplicationName, LPWSTR lpCommandLine, LP
 			temp.Append(originalCmd.data+1, commandLineWithoutApplication.data - originalCmd.data - 2);
 		else
 			temp.Append(originalCmd.data, commandLineWithoutApplication.data - originalCmd.data - 1);
-		FixPath(application, temp.data);
+		if (temp[1] == ':')
+			FixPath(application, temp.data);
+		else
+			application.Append(temp.data);
 		UBA_ASSERTF(application.count, TC("Invalid application from command line (name: %s)"), originalCmd.data);
 	}
 	DevirtualizePath(application);
@@ -3880,7 +3903,15 @@ BOOL Detoured_IsProcessorFeaturePresent(DWORD ProcessorFeature)
 {
 	DETOURED_CALL(IsProcessorFeaturePresent);
 	auto res = True_IsProcessorFeaturePresent(ProcessorFeature);
-	DEBUG_LOG_TRUE(L"IsProcessorFeaturePresent", L"%u -> %ls", res, ToString(res));
+	DEBUG_LOG_TRUE(L"IsProcessorFeaturePresent", L"%u -> %ls", ProcessorFeature, ToString(res));
+	return res;
+}
+
+BOOL Detoured_UnmapViewOfFile2(HANDLE Process, PVOID BaseAddress, ULONG UnmapFlags)
+{
+	DETOURED_CALL(UnmapViewOfFile2);
+	auto res = True_UnmapViewOfFile2(Process, BaseAddress, UnmapFlags);
+	DEBUG_LOG_TRUE(L"UnmapViewOfFile2", L"0x%llx -> %ls", uintptr_t(BaseAddress), ToString(res));
 	return res;
 }
 
diff --git a/Engine/Source/Programs/UnrealBuildAccelerator/Detours/Private/Windows/UbaDetoursFunctionsNtDll.inl b/Engine/Source/Programs/UnrealBuildAccelerator/Detours/Private/Windows/UbaDetoursFunctionsNtDll.inl
--- a/Engine/Source/Programs/UnrealBuildAccelerator/Detours/Private/Windows/UbaDetoursFunctionsNtDll.inl
+++ b/Engine/Source/Programs/UnrealBuildAccelerator/Detours/Private/Windows/UbaDetoursFunctionsNtDll.inl
@@ -832,6 +832,7 @@ NTSTATUS NTAPI Shared_NtCreateFile(bool IsCreateFunc, PHANDLE hFileHandle, ACCES
 					dh->fileObject->deleteOnClose = isDeleteOnClose;
 					tempFileInfo->originalName = _wcsdup(fileName.data);
 					tempFileInfo->name = L"GETATTRIBUTES";
+					tempFileInfo->refCount = 1;
 					*hFileHandle = makeDetouredHandle(dh);
 					//SetLastError(ERROR_SUCCESS); // Don't think this is needed
 					DEBUG_LOG_DETOURED(funcName, L"GETATTRIBUTES %llu, (%ls) -> Success", uintptr_t(*hFileHandle), fileName.data);
@@ -1019,6 +1020,7 @@ NTSTATUS NTAPI Shared_NtCreateFile(bool IsCreateFunc, PHANDLE hFileHandle, ACCES
 			fo->desiredAccess = DesiredAccess;
 			fo->closeId = closeId;
 			fo->fileInfo = &info;
+			InterlockedIncrement(&info.refCount);
 			fo->deleteOnClose = isDeleteOnClose;
 			auto dh = new DetouredHandle(HandleType_File, th);
 			dh->dirTableOffset = dirTableOffset;
@@ -1469,8 +1471,10 @@ NTSTATUS NTAPI Detoured_NtClose(HANDLE handle)
 		//UBA_ASSERTF(fo->newName.empty(), L"Got close of file that was renamed but had no closeId. Old: %ls New: %ls", fi.originalName, fo->newName.c_str());
 	}
 
+	auto fiRefCount = InterlockedDecrement(&fi.refCount);
+
 	// TODO: Can't free here..... Might still be open with CreateFileMapping or MapViewOfFile
-	if (fo->wasUsed && fi.freeFileMapOnClose)
+	if (!fiRefCount && fo->wasUsed && fi.freeFileMapOnClose)
 		g_fileDecompressor.ReturnMemory(fi);
 
 	DEBUG_LOG_DETOURED(L"NtClose", L"%llu (%ls) -> %ls", uintptr_t(handle), HandleToName(handle), ToString(res).data);
diff --git a/Engine/Source/Programs/UnrealBuildAccelerator/Detours/Private/Windows/UbaDetoursFunctionsWin.cpp b/Engine/Source/Programs/UnrealBuildAccelerator/Detours/Private/Windows/UbaDetoursFunctionsWin.cpp
--- a/Engine/Source/Programs/UnrealBuildAccelerator/Detours/Private/Windows/UbaDetoursFunctionsWin.cpp
+++ b/Engine/Source/Programs/UnrealBuildAccelerator/Detours/Private/Windows/UbaDetoursFunctionsWin.cpp
@@ -457,7 +457,7 @@ void CloseCaches()
 		FileInfo& info = it.second;
 		if (info.fileMapMem)
 		{
-			DEBUG_LOG_TRUE(L"INTERNAL UnmapViewOfFile", L"%llu (%ls) (%ls)", uintptr_t(info.fileMapMem), info.name, info.originalName);
+			DEBUG_LOG_TRUE(L"INTERNAL UnmapViewOfFile", L"0x%llx (%ls) (%ls)", uintptr_t(info.fileMapMem), info.name, info.originalName);
 			True_UnmapViewOfFile(info.fileMapMem);
 		}
 		if (info.trueFileMapHandle != 0)
diff --git a/Engine/Source/Programs/UnrealBuildAccelerator/Detours/Private/Windows/UbaDetoursFunctionsWin.h b/Engine/Source/Programs/UnrealBuildAccelerator/Detours/Private/Windows/UbaDetoursFunctionsWin.h
--- a/Engine/Source/Programs/UnrealBuildAccelerator/Detours/Private/Windows/UbaDetoursFunctionsWin.h
+++ b/Engine/Source/Programs/UnrealBuildAccelerator/Detours/Private/Windows/UbaDetoursFunctionsWin.h
@@ -312,6 +312,7 @@
 	DETOURED_FUNCTION(OpenFileMappingA) \
 	DETOURED_FUNCTION(GetMappedFileNameW) \
 	DETOURED_FUNCTION(IsProcessorFeaturePresent) \
+	DETOURED_FUNCTION(UnmapViewOfFile2) \
 	//DETOURED_FUNCTION(VirtualFree) \
 	//DETOURED_FUNCTION(BaseThreadInitThunk) \
 	//DETOURED_FUNCTION(VirtualAllocEx) \
diff --git a/Engine/Source/Programs/UnrealBuildAccelerator/ObjTool/Private/UbaObjTool.cpp b/Engine/Source/Programs/UnrealBuildAccelerator/ObjTool/Private/UbaObjTool.cpp
--- a/Engine/Source/Programs/UnrealBuildAccelerator/ObjTool/Private/UbaObjTool.cpp
+++ b/Engine/Source/Programs/UnrealBuildAccelerator/ObjTool/Private/UbaObjTool.cpp
@@ -250,7 +250,7 @@ namespace uba
 		}
 
 		FilteredLogWriter logWriter(g_consoleLogWriter, LogEntryType_Info);
-		LoggerWithWriter logger(logWriter, TC(""));
+		LoggerWithWriter logger(logWriter, TC("UbaObjTool"));
 
 		if (!objFilesToStrip.empty())
 		{
