[UBA]
* Fixed so unit tests are not failing on macs that does not have xcode installed

diff --git a/Engine/Source/Programs/UnrealBuildAccelerator/Test/Private/UbaTestSession.cpp b/Engine/Source/Programs/UnrealBuildAccelerator/Test/Private/UbaTestSession.cpp
--- a/Engine/Source/Programs/UnrealBuildAccelerator/Test/Private/UbaTestSession.cpp
+++ b/Engine/Source/Programs/UnrealBuildAccelerator/Test/Private/UbaTestSession.cpp
@@ -223,7 +223,7 @@ namespace uba
 		FILE* fpCommand = popen(command, "r");
 		if (fpCommand == nullptr || fgets(commandOutput.data, commandOutput.capacity, fpCommand) == nullptr || pclose(fpCommand) != 0)
 		{
-			logger.Error("Failed to get an Xcode from xcode-select");
+			logger.Warning("Failed to get an Xcode from xcode-select");
 			return false;
 		}
 
@@ -254,7 +254,8 @@ namespace uba
 		const tchar* clangPath = TC("c:\\Program Files\\Microsoft Visual Studio\\2022\\Professional\\VC\\Tools\\Llvm\\bin\\clang-cl.exe");
 #elif PLATFORM_MAC
 		StringBuffer<MaxPath> xcodePath;
-		ExecuteCommand(logger, "/usr/bin/xcrun -r clang++", xcodePath);
+		if (!ExecuteCommand(logger, "/usr/bin/xcrun -r clang++", xcodePath))
+			return true;
 		const tchar* clangPath = xcodePath.data;
 #else
 		const tchar* clangPath = TC("/usr/bin/clang++");
@@ -272,7 +273,8 @@ namespace uba
 		args.Append("/Brepro ");
 #elif PLATFORM_MAC
 		StringBuffer<MaxPath> xcodeSDKPath;
-		ExecuteCommand(logger, "xcrun --show-sdk-path", xcodeSDKPath);
+		if (!ExecuteCommand(logger, "xcrun --show-sdk-path", xcodeSDKPath))
+			return true;
 		args.Append("-isysroot ");
 		args.Append(xcodeSDKPath.data).Append(' ');
 #endif
