[UBA]
* Removed lock in cache client added in previous changelist because it caused a deadlock

diff --git a/Engine/Source/Programs/UnrealBuildAccelerator/Common/Private/UbaCacheClient.cpp b/Engine/Source/Programs/UnrealBuildAccelerator/Common/Private/UbaCacheClient.cpp
--- a/Engine/Source/Programs/UnrealBuildAccelerator/Common/Private/UbaCacheClient.cpp
+++ b/Engine/Source/Programs/UnrealBuildAccelerator/Common/Private/UbaCacheClient.cpp
@@ -994,7 +994,6 @@ namespace uba
 			StorageImpl::CasEntry* casEntry;
 			if (m_storage.HasCasFile(casKey, &casEntry))
 			{
-				SCOPED_READ_LOCK(casEntry->lock, lock);
 				UBA_ASSERT(!IsNormalized(casKey));
 				StringBuffer<> casKeyFileName;
 				if (!m_storage.GetCasFileName(casKeyFileName, casKey))
