[UBA]
* Readded lock around sendsocket
* Changed asserts about network protocol error to be real errors

diff --git a/Engine/Source/Programs/UnrealBuildAccelerator/Common/Private/UbaNetworkBackendTcp.cpp b/Engine/Source/Programs/UnrealBuildAccelerator/Common/Private/UbaNetworkBackendTcp.cpp
--- a/Engine/Source/Programs/UnrealBuildAccelerator/Common/Private/UbaNetworkBackendTcp.cpp
+++ b/Engine/Source/Programs/UnrealBuildAccelerator/Common/Private/UbaNetworkBackendTcp.cpp
@@ -45,7 +45,8 @@
 // Hard to say if this lock should be enabled on windows or not.
 // Local tests with cache server and 100% cache hits make no-lock 35s vs lock 40s.
 // But by looking at the farm jobs it looks like more time is spent in send without the lock
-#define UBA_LOCK_AROUND_SEND !PLATFORM_WINDOWS
+// Saw a network error today that did not make sense. Could be caused by not locking around send
+#define UBA_LOCK_AROUND_SEND 1//!PLATFORM_WINDOWS
 
 namespace uba
 {
diff --git a/Engine/Source/Programs/UnrealBuildAccelerator/Common/Private/UbaNetworkServer.cpp b/Engine/Source/Programs/UnrealBuildAccelerator/Common/Private/UbaNetworkServer.cpp
--- a/Engine/Source/Programs/UnrealBuildAccelerator/Common/Private/UbaNetworkServer.cpp
+++ b/Engine/Source/Programs/UnrealBuildAccelerator/Common/Private/UbaNetworkServer.cpp
@@ -324,8 +324,10 @@ namespace uba
 			u8 messageType = serviceIdAndMessageType & 0b111111;
 			u16 messageId = u16(headerData[1] << 8) | u16((*(u32*)(headerData + 2) & 0xff000000) >> 24);
 			u32 messageSize = *(u32*)(headerData + 2) & 0x00ffffff;
-			UBA_ASSERTF(messageSize <= SendMaxSize, TC("Got message size %u which is larger than max %u. Protocol error?"), messageSize, SendMaxSize);
-			UBA_ASSERTF(serviceId < sizeof(NetworkServer::m_workerFunctions), TC("Got message with service id %u which is out of range. Protocol error?"), serviceId);
+			if (messageSize > SendMaxSize)
+				return conn.m_server.m_logger.Error(TC("Client %u Got message size %u which is larger than max %u. Protocol error? (serviceId %u, messageType %u, messageId %u)"), conn.m_client->id, messageSize, SendMaxSize, u32(serviceId), u32(messageType), u32(messageId));
+			if (serviceId >= sizeof(NetworkServer::m_workerFunctions))
+				return conn.m_server.m_logger.Error(TC("Client %u Got message with service id %u which is out of range. Protocol error?"), conn.m_client->id, serviceId);
 
 			//m_logger.Debug(TC("Recv: %u, %u, %u, %u"), serviceId, messageType, id, size);
 			Worker* worker = conn.m_server.PopWorker();
