[UBA]
* Added size of VirtualAlloc when failing to allocate because of oom

diff --git a/Engine/Source/Programs/UnrealBuildAccelerator/Detours/Private/Windows/UbaDetoursFunctionsKernelBase.inl b/Engine/Source/Programs/UnrealBuildAccelerator/Detours/Private/Windows/UbaDetoursFunctionsKernelBase.inl
--- a/Engine/Source/Programs/UnrealBuildAccelerator/Detours/Private/Windows/UbaDetoursFunctionsKernelBase.inl
+++ b/Engine/Source/Programs/UnrealBuildAccelerator/Detours/Private/Windows/UbaDetoursFunctionsKernelBase.inl
@@ -572,7 +572,9 @@ LPVOID Detoured_VirtualAlloc(LPVOID lpAddress, SIZE_T dwSize, DWORD flAllocation
 		DWORD error = GetLastError();
 		if (error != ERROR_NOT_ENOUGH_MEMORY && error != ERROR_COMMITMENT_LIMIT)
 			return res;
-		Rpc_AllocFailed(L"VirtualAlloc", error);
+		StringBuffer<128> reason;
+		reason.Append(TC("VirtualAlloc ")).AppendValue(dwSize);
+		Rpc_AllocFailed(reason.data, error);
 		++counter;
 
 	} while (counter <= 10);
