[UBA]
* Added more validation on cache server side when receiving cache entries
* Added cache server command "validate" that can be used to check if paths are valid

diff --git a/Engine/Source/Programs/UnrealBuildAccelerator/Common/Private/UbaCacheServer.cpp b/Engine/Source/Programs/UnrealBuildAccelerator/Common/Private/UbaCacheServer.cpp
--- a/Engine/Source/Programs/UnrealBuildAccelerator/Common/Private/UbaCacheServer.cpp
+++ b/Engine/Source/Programs/UnrealBuildAccelerator/Common/Private/UbaCacheServer.cpp
@@ -1510,7 +1510,9 @@ namespace uba
 				u32 pathOffset = u32(reader2.Read7BitEncoded());
 				return m_logger.Error(TC("Client %u has provided corrupt key or path table (CasKey: %s Offset: %u TableSize: %u PathOffset: %u PathTableSize: %u ClientVersion: %u PathtableVersion: %u)"), clientId, CasKeyString(casKey).str, offset, connectionBucket.casKeyTable.GetSize(), pathOffset, connectionBucket.pathTable.GetSize(), clientVersion, connectionBucket.pathTable.GetVersion());
 			}
-			UBA_ASSERT(path.count);
+			
+			if (path.count <= 2)
+				return m_logger.Error(TC("Client %u has provided corrupt path %s. Corrupt data? (Offset: %u TableSize: %u ClientVersion: %u PathtableVersion: %u)"), clientId, path.data, offset, connectionBucket.casKeyTable.GetSize(), clientVersion, connectionBucket.pathTable.GetVersion());
 
 			if (casKey == CasKeyZero)
 				return m_logger.Error(TC("Client %u has provided cas key that is zero. Corrupt data? (Offset: %u TableSize: %u ClientVersion: %u PathtableVersion: %u)"), clientId, offset, connectionBucket.casKeyTable.GetSize(), clientVersion, connectionBucket.pathTable.GetVersion());
@@ -1975,14 +1977,23 @@ namespace uba
 					Bucket& bucket = kv.second;
 					SCOPED_READ_LOCK(bucket.m_cacheEntryLookupLock, lock2);
 					u64 mostEntries = 0;
+					CasKey mostEntriesKey;
 					u64 lastUsed = 0;
 					u64 totalEntryCount = 0;
+					u64 totalKeySharedBytes = 0;
 					u64 totalCountHasOutputAsInput = 0;
+					u64 totalEntryBytes = 0;
 					for (auto& kv2 : bucket.m_cacheEntryLookup)
 					{
 						CacheEntries& entries = kv2.second;
 						SCOPED_READ_LOCK(entries.lock, lock3);
-						mostEntries = Max(mostEntries, u64(entries.entries.size()));
+						totalKeySharedBytes += entries.GetSharedSize();
+						totalEntryBytes += entries.GetTotalSize(CacheNetworkVersion, true);
+						if (mostEntries < entries.entries.size())
+						{
+							mostEntries = u64(entries.entries.size());
+							mostEntriesKey = kv2.first;
+						}
 						for (auto& entry : entries.entries)
 							lastUsed = Max(lastUsed, entry.lastUsedTime);
 						totalEntryCount += entries.entries.size();
@@ -1997,10 +2008,11 @@ namespace uba
 					writeLine(line.Clear().Appendf(TC("    #%u - %llu (v%u)"), index++, kv.first, bucket.m_pathTable.GetVersion()));
 					writeLine(line.Clear().Appendf(TC("      PathTable: %llu (%s)"), bucket.m_pathTable.GetPathCount(), BytesToText(bucket.m_pathTable.GetSize()).str));
 					writeLine(line.Clear().Appendf(TC("      CasKeyTable: %llu (%s)"), bucket.m_casKeyTable.GetKeyCount(), BytesToText(bucket.m_casKeyTable.GetSize()).str));
-					writeLine(line.Clear().Appendf(TC("      Keys: %llu"), bucket.m_cacheEntryLookup.size()));
+					writeLine(line.Clear().Appendf(TC("      Keys: %llu (%s)"), bucket.m_cacheEntryLookup.size(), BytesToText(totalEntryBytes).str));
 					writeLine(line.Clear().Appendf(TC("      KeysWithInputsFromOutputs: %llu"), totalCountHasOutputAsInput));
-					writeLine(line.Clear().Appendf(TC("      KeyMostEntries: %llu"), mostEntries));
+					writeLine(line.Clear().Appendf(TC("      KeyMostEntries: %llu (%s)"), mostEntries, CasKeyString(mostEntriesKey).str));
 					writeLine(line.Clear().Appendf(TC("      TotalEntries: %llu"), totalEntryCount));
+					writeLine(line.Clear().Appendf(TC("      TotalKeySharedEntry: %s"), BytesToText(totalKeySharedBytes).str));
 					writeLine(line.Clear().Appendf(TC("      LastEntryUsed: %s ago"), TimeToText(lastUsedTime, true).str));
 				}
 			}
@@ -2011,6 +2023,56 @@ namespace uba
 			writeLine(line.Clear().Appendf(TC("    Count: %llu"), totalCasCount));
 			writeLine(line.Clear().Appendf(TC("    Size: %s"), BytesToText(totalCasSize).str));
 		}
+		else if (command.Equals(TC("validate")))
+		{
+			writeLine(line.Clear().Appendf(TC("  Buckets:")));
+			{
+				auto checkPaths = [&](Bucket& bucket, const Vector<u8>& offsets)
+					{
+						BinaryReader reader2(offsets.data(), 0, offsets.size());
+						while (reader2.GetLeft())
+						{
+							u32 offset = u32(reader2.Read7BitEncoded());
+							CasKey casKey;
+							StringBuffer<> path;
+							bucket.m_casKeyTable.GetPathAndKey(path, casKey, bucket.m_pathTable, offset);
+							if (path.count == 1)
+								return false;
+						}
+						return true;
+					};
+
+				u32 index = 0;
+				SCOPED_READ_LOCK(m_bucketsLock, bucketsLock);
+				for (auto& kv : m_buckets)
+				{
+					u64 keysWithInvalidSharedInputs = 0;
+					u64 invalidExtraInputs = 0;
+					u64 totalExtraInputs = 0;
+					Bucket& bucket = kv.second;
+					SCOPED_READ_LOCK(bucket.m_cacheEntryLookupLock, lock2);
+					for (auto& kv2 : bucket.m_cacheEntryLookup)
+					{
+						CacheEntries& entries = kv2.second;
+						SCOPED_READ_LOCK(entries.lock, lock3);
+						if (!checkPaths(bucket, entries.sharedInputCasKeyOffsets))
+							++keysWithInvalidSharedInputs;
+
+						for (auto& entry : entries.entries)
+						{
+							++totalExtraInputs;
+							if (!checkPaths(bucket, entry.extraInputCasKeyOffsets))
+								++invalidExtraInputs;
+						}
+					}
+					lock2.Leave();
+
+					writeLine(line.Clear().Appendf(TC("    #%u - %llu (v%u)"), index++, kv.first, bucket.m_pathTable.GetVersion()));
+					writeLine(line.Clear().Appendf(TC("      Invalid shared: %llu/%llu"), keysWithInvalidSharedInputs, (bucket.m_cacheEntryLookup.size())));
+					writeLine(line.Clear().Appendf(TC("      Invalid extra: %llu/%llu"), invalidExtraInputs, totalExtraInputs));
+				}
+			}
+		}
 		else if (command.Equals(TC("pathtable")))
 		{
 			SCOPED_READ_LOCK(m_bucketsLock, bucketsLock);
@@ -2121,24 +2183,24 @@ namespace uba
 					largestEntriesBucket = &bucket;
 				}
 			}
-			if (!largestEntries)
-				return true;
 			writeLine(line.Clear().Append(TC("Largest entry (based on shared inputs)")));
-			CacheEntries& entries = *largestEntries;
-			SCOPED_READ_LOCK(entries.lock, lock3);
-			writeLine(line.Clear().Append(TC("   SharedInputs:")));
-			writePathsFromOffsets(*largestEntriesBucket, entries.sharedInputCasKeyOffsets);
-			u32 index = 0;
-			for (auto& entry : entries.entries)
+			if (largestEntries)
 			{
-				writeLine(line.Clear().Appendf(TC("  #%u"), index));
-				writeLine(line.Clear().Append(TC("   ExtraInputs:")));
-				writePathsFromOffsets(*largestEntriesBucket, entry.extraInputCasKeyOffsets);
-				writeLine(line.Clear().Append(TC("   Outputs:")));
-				writePathsFromOffsets(*largestEntriesBucket, entry.outputCasKeyOffsets);
-				++index;
+				CacheEntries& entries = *largestEntries;
+				SCOPED_READ_LOCK(entries.lock, lock3);
+				writeLine(line.Clear().Append(TC("   SharedInputs:")));
+				writePathsFromOffsets(*largestEntriesBucket, entries.sharedInputCasKeyOffsets);
+				u32 index = 0;
+				for (auto& entry : entries.entries)
+				{
+					writeLine(line.Clear().Appendf(TC("  #%u (%s)"), index, TimeToText(GetFileTimeAsTime(GetSystemTimeAsFileTime() - (m_creationTime + entry.creationTime)), true).str));
+					writeLine(line.Clear().Append(TC("   ExtraInputs:")));
+					writePathsFromOffsets(*largestEntriesBucket, entry.extraInputCasKeyOffsets);
+					writeLine(line.Clear().Append(TC("   Outputs:")));
+					writePathsFromOffsets(*largestEntriesBucket, entry.outputCasKeyOffsets);
+					++index;
+				}
 			}
-
 		}
 		else if (command.Equals(TC("obliterate")))
 		{
