[UBA]
* Disabled new CompactPathTable::AddNoLock code. suspect there is a bug in there

diff --git a/Engine/Source/Programs/UnrealBuildAccelerator/Common/Private/UbaCacheServer.cpp b/Engine/Source/Programs/UnrealBuildAccelerator/Common/Private/UbaCacheServer.cpp
--- a/Engine/Source/Programs/UnrealBuildAccelerator/Common/Private/UbaCacheServer.cpp
+++ b/Engine/Source/Programs/UnrealBuildAccelerator/Common/Private/UbaCacheServer.cpp
@@ -2036,7 +2036,7 @@ namespace uba
 							CasKey casKey;
 							StringBuffer<> path;
 							bucket.m_casKeyTable.GetPathAndKey(path, casKey, bucket.m_pathTable, offset);
-							if (path.count == 1)
+							if (path.count <= 2)
 								return false;
 						}
 						return true;
diff --git a/Engine/Source/Programs/UnrealBuildAccelerator/Common/Private/UbaCompactTables.cpp b/Engine/Source/Programs/UnrealBuildAccelerator/Common/Private/UbaCompactTables.cpp
--- a/Engine/Source/Programs/UnrealBuildAccelerator/Common/Private/UbaCompactTables.cpp
+++ b/Engine/Source/Programs/UnrealBuildAccelerator/Common/Private/UbaCompactTables.cpp
@@ -58,13 +58,13 @@ namespace uba
 	{
 		auto& table = context.fromTable;
 
-		if (m_version < 3)
+		if (true)//m_version < 3)
 		{
 			StringBuffer<> temp;
 			table.GetString(temp, offset);
 			return InternalAdd(temp.data, temp.data, temp.count);
 		}
-
+#if 0
 		u32 offsets[MaxSegments];
 		bool separators[MaxSegments];
 
@@ -196,6 +196,7 @@ namespace uba
 		context.offsetsCount = searchIndex + 1;
 
 		return offset;
+#endif
 	}
 
 	void CompactPathTable::AddCommonStringSegments()
