[UBA]
* Fixed shutdown crash in UbaCacheService

diff --git a/Engine/Source/Programs/UnrealBuildAccelerator/Common/Private/UbaCacheServer.cpp b/Engine/Source/Programs/UnrealBuildAccelerator/Common/Private/UbaCacheServer.cpp
--- a/Engine/Source/Programs/UnrealBuildAccelerator/Common/Private/UbaCacheServer.cpp
+++ b/Engine/Source/Programs/UnrealBuildAccelerator/Common/Private/UbaCacheServer.cpp
@@ -721,7 +721,7 @@ namespace uba
 			m_logger.Detail(TC("  Found %s (%s) cas files and %s deleted by overflow (%s)"), CountToText(existingCas.Size()).str, BytesToText(totalCasSize).str, CountToText(deletedCasFiles.size()).str, TimeToText(GetTime() - collectCasStartTime).str);
 		}
 
-		if (shouldExit())
+		if (shouldExit && shouldExit())
 			return true;
 
 		// Take biggest buckets first
