[UBA]
* Fixed bugs around how cache entries are optimized. Added extensive unit tests and made sure all entries does not contain any duplicated entries between shared ranges and extra entries.

diff --git a/Engine/Source/Programs/UnrealBuildAccelerator/Common/Private/UbaCacheEntry.cpp b/Engine/Source/Programs/UnrealBuildAccelerator/Common/Private/UbaCacheEntry.cpp
--- a/Engine/Source/Programs/UnrealBuildAccelerator/Common/Private/UbaCacheEntry.cpp
+++ b/Engine/Source/Programs/UnrealBuildAccelerator/Common/Private/UbaCacheEntry.cpp
@@ -170,6 +170,7 @@ namespace uba
 		{
 			u32 cacheEntryCount = reader.ReadU32();
 			Vector<u32> temp;
+			Vector<u8> temp2;
 			while (cacheEntryCount--)
 			{
 				auto& cacheEntry = entries.emplace_back();
@@ -184,7 +185,7 @@ namespace uba
 				temp.clear();
 				while (reader.GetPosition() < inputEnd)
 					temp.push_back(u32(reader.Read7BitEncoded()));
-				BuildInputsT(cacheEntry, temp, entries.size() == 1);
+				BuildInputsT(cacheEntry, temp, entries.size() == 1, temp2);
 
 				u32 outputSize = reader.ReadU32();
 				cacheEntry.outputCasKeyOffsets.resize(outputSize);
@@ -259,7 +260,7 @@ namespace uba
 	}
 
 	template<typename Container>
-	void CacheEntries::BuildInputsT(CacheEntry& entry, const Container& sortedInputs, bool populateShared)
+	void CacheEntries::BuildInputsT(CacheEntry& entry, const Container& sortedInputs, bool populateShared, Vector<u8>& temp)
 	{
 		StackBinaryWriter<256*1024> rangeWriter;
 
@@ -288,16 +289,17 @@ namespace uba
 		auto inputsIt = sortedInputs.begin();
 		auto inputsEnd = sortedInputs.end();
 
-		BinaryReader sharedReader(sharedInputCasKeyOffsets.data(), 0, sharedInputCasKeyOffsets.size());
+		BinaryReader sharedReader(sharedInputCasKeyOffsets);
 
 		u32 sharedOffset = ~0u;
 		u32 offset = ~0u;
 
 		u32 rangeBegin = 0;
 		bool inRange = false;
+		bool previousWasExtra = false;
 		u32 lastSharedPos = ~0u;
 
-		Vector<u8> extraOffsets;
+		Vector<u8>& extraOffsets = temp;
 		extraOffsets.resize(sortedInputs.size()*4);
 		BinaryWriter extraWriter(extraOffsets.data(), 0, extraOffsets.size());
 
@@ -307,16 +309,28 @@ namespace uba
 
 			if (!sharedReader.GetLeft())
 			{
+				if (previousWasExtra)
+				{
+					while (inputsIt != inputsEnd)
+					{
+						offset = *inputsIt++;
+						if (offset < sharedOffset)
+							extraWriter.Write7BitEncoded(offset);
+					}
+				}
+
 				// Add current range if there is one going
 				if (inRange)
-					writeRange(rangeBegin, sharedPos);
-
-				if (offset > sharedOffset)
+					writeRange(rangeBegin, offset == sharedOffset ? sharedPos : lastSharedPos);
+				else if (offset == sharedOffset) // We want to use a range just to make it easier to debug
+					writeRange(lastSharedPos, sharedPos);
+				else
 					extraWriter.Write7BitEncoded(offset);
 
 				// Populate rest in extraInputCasKeyOffsets
 				for (;inputsIt != inputsEnd ;++inputsIt)
-					extraWriter.Write7BitEncoded(*inputsIt);
+					if (*inputsIt != sharedOffset)
+						extraWriter.Write7BitEncoded(*inputsIt);
 				break;
 			}
 
@@ -324,12 +338,27 @@ namespace uba
 			{
 				// Add current range if there is one going
 				if (inRange)
-					writeRange(rangeBegin, sharedPos);
-				if (offset > sharedOffset)
-					extraWriter.Write7BitEncoded(offset);
+				{
+					writeRange(rangeBegin, previousWasExtra ? lastSharedPos : sharedPos);
+				}
+				else if (offset > sharedOffset)
+				{
+					lastSharedPos = sharedPos;
+					sharedOffset = u32(sharedReader.Read7BitEncoded());
+					continue;
+				}
+				else
+				{
+					if (sharedOffset == offset)
+						writeRange(lastSharedPos, u32(sharedReader.GetPosition()));
+					else if (!previousWasExtra)
+						extraWriter.Write7BitEncoded(offset);
+				}
 				break;
 			}
 
+			previousWasExtra = false;
+
 			if (sharedOffset < offset)
 			{
 				lastSharedPos = sharedPos;
@@ -357,14 +386,15 @@ namespace uba
 			}
 			else
 			{
-				if (inRange)
-				{
-					inRange = false;
-					writeRange(rangeBegin, sharedPos);
-				}
 				if (offset < sharedOffset)
 				{
 					extraWriter.Write7BitEncoded(offset);
+					previousWasExtra = true;
+				}
+				else if (inRange)
+				{
+					inRange = false;
+					writeRange(rangeBegin, sharedPos);
 				}
 			}
 		}
@@ -389,7 +419,7 @@ namespace uba
 		auto excludedInputsIt = sortedExcludedInputs.begin();
 		auto excludedInputsEnd = sortedExcludedInputs.end();
 
-		BinaryReader sharedReader(sharedInputCasKeyOffsets.data(), 0, sharedInputCasKeyOffsets.size());
+		BinaryReader sharedReader(sharedInputCasKeyOffsets);
 
 		u32 sharedOffset = ~0u;
 		u32 offset = ~0u;
@@ -460,10 +490,11 @@ namespace uba
 
 	void CacheEntries::BuildInputs(CacheEntry& entry, const Set<u32>& inputs)
 	{
-		BuildInputsT(entry, inputs, entries.empty());
+		Vector<u8> temp;
+		BuildInputsT(entry, inputs, entries.empty(), temp);
 	}
 
-	void CacheEntries::UpdateEntries(Logger& logger, const HashMap2<u32, u32>& oldToNewCasKeyOffset, Vector<u32>& temp, Vector<u8>& temp2)
+	void CacheEntries::UpdateEntries(Logger& logger, const HashMap2<u32, u32>& oldToNewCasKeyOffset, Vector<u32>& temp, Vector<u8>& temp2, Vector<u8>& temp3)
 	{
 		if (entries.empty())
 			return;
@@ -479,7 +510,7 @@ namespace uba
 			{
 				temp.clear();
 				u32 newOffsetsSize = 0;
-				BinaryReader reader(offsets.data(), 0, offsets.size());
+				BinaryReader reader(offsets);
 				while (reader.GetLeft())
 				{
 					u32 newOffset = u32(reader.Read7BitEncoded());
@@ -562,7 +593,7 @@ namespace uba
 
 					entry.extraInputCasKeyOffsets.clear();
 					entry.sharedInputCasKeyOffsetRanges.clear();
-					BuildInputsT(entry, temp, false);
+					BuildInputsT(entry, temp, false, temp3);
 				}
 			}
 		}
@@ -582,7 +613,7 @@ namespace uba
 
 				auto collectInputs = [&](Vector<u32>& out, u32 rangeBegin, u32 rangeEnd)
 					{
-						BinaryReader excludedReader(oldShared.data() + rangeBegin, 0, rangeEnd - rangeBegin);
+						BinaryReader excludedReader(oldShared.data(), rangeBegin, rangeEnd);
 						while (excludedReader.GetLeft())
 						{
 							u32 offset = u32(excludedReader.Read7BitEncoded());
@@ -595,7 +626,7 @@ namespace uba
 				temp.clear();
 				auto& excludedOffsets = temp;
 
-				BinaryReader rangeReader(entry.sharedInputCasKeyOffsetRanges.data(), 0, entry.sharedInputCasKeyOffsetRanges.size());
+				BinaryReader rangeReader(entry.sharedInputCasKeyOffsetRanges);
 				u32 lastEnd = 0;
 				while (rangeReader.GetLeft())
 				{
@@ -624,28 +655,6 @@ namespace uba
 		}
 	}
 
-	void CacheEntries::ValidateEntry(Logger& logger, CacheEntry& entry, Vector<u8>& inputCasKeyOffsets)
-	{
-		Vector<u8> res;
-		Flatten(res, entry);
-		if (res.size() == inputCasKeyOffsets.size() && memcmp(res.data(), inputCasKeyOffsets.data(), res.size()) == 0)
-			return;
-		BinaryReader reader1(inputCasKeyOffsets.data(), 0, inputCasKeyOffsets.size());
-		BinaryReader reader2(res.data(), 0, res.size());
-		while (reader1.GetLeft() || reader2.GetLeft())
-		{
-			u32 a = ~0u;
-			u32 b = ~0u;
-			if (reader1.GetLeft())
-				a = u32(reader1.Read7BitEncoded());
-			if (reader2.GetLeft())
-				b = u32(reader2.Read7BitEncoded());
-
-			logger.Detail(TC("A: %u B: %u"), a, b);
-		}
-		UBA_ASSERT(false);
-	}
-
 	void CacheEntries::PopulateInputsThatAreOutputs(const Vector<u8>& inputData, StorageImpl& storage, CompactCasKeyTable& table)
 	{
 		// one entry that is ~0u is equal to test all inputs
@@ -653,7 +662,7 @@ namespace uba
 			return;
 
 		u32 insertCount = 0;
-		BinaryReader inputReader(inputData.data(), 0, inputData.size());
+		BinaryReader inputReader(inputData);
 		while (inputReader.GetLeft())
 		{
 			u32 casKeyOffset = u32(inputReader.Read7BitEncoded());
@@ -676,7 +685,7 @@ namespace uba
 	{
 		u64 size = entry.extraInputCasKeyOffsets.size();
 		{
-			BinaryReader rangeReader(entry.sharedInputCasKeyOffsetRanges.data(), 0, entry.sharedInputCasKeyOffsetRanges.size());
+			BinaryReader rangeReader(entry.sharedInputCasKeyOffsetRanges);
 			while (rangeReader.GetLeft())
 			{
 				u64 begin = rangeReader.Read7BitEncoded();
@@ -688,7 +697,7 @@ namespace uba
 		out.resize(size);
 		BinaryWriter writer(out.data(), 0, out.size());
 
-		BinaryReader extraReader(entry.extraInputCasKeyOffsets.data(), 0, entry.extraInputCasKeyOffsets.size());
+		BinaryReader extraReader(entry.extraInputCasKeyOffsets);
 		u32 nextExtra = ~0u;
 		if (extraReader.GetLeft())
 			nextExtra = u32(extraReader.Read7BitEncoded());
@@ -704,12 +713,12 @@ namespace uba
 				};
 			};
 
-		BinaryReader rangeReader(entry.sharedInputCasKeyOffsetRanges.data(), 0, entry.sharedInputCasKeyOffsetRanges.size());
+		BinaryReader rangeReader(entry.sharedInputCasKeyOffsetRanges);
 		while (rangeReader.GetLeft())
 		{
 			u64 begin = rangeReader.Read7BitEncoded();
 			u64 end = rangeReader.Read7BitEncoded();
-			BinaryReader inputReader(sharedInputCasKeyOffsets.data() + begin, 0, end - begin);
+			BinaryReader inputReader(sharedInputCasKeyOffsets.data(), begin, end);
 			while (inputReader.GetLeft())
 			{
 				u32 offset = u32(inputReader.Read7BitEncoded());
@@ -725,7 +734,7 @@ namespace uba
 	{
 		out.clear();
 
-		BinaryReader extraReader(entry.extraInputCasKeyOffsets.data(), 0, entry.extraInputCasKeyOffsets.size());
+		BinaryReader extraReader(entry.extraInputCasKeyOffsets);
 		u32 nextExtra = ~0u;
 		if (extraReader.GetLeft())
 			nextExtra = u32(extraReader.Read7BitEncoded());
@@ -746,7 +755,7 @@ namespace uba
 		{
 			u64 begin = rangeReader.Read7BitEncoded();
 			u64 end = rangeReader.Read7BitEncoded();
-			BinaryReader inputReader(sharedOffsets.data() + begin, 0, end - begin);
+			BinaryReader inputReader(sharedOffsets.data(), begin, end);
 			while (inputReader.GetLeft())
 			{
 				u32 offset = u32(inputReader.Read7BitEncoded());
diff --git a/Engine/Source/Programs/UnrealBuildAccelerator/Common/Private/UbaCacheServer.cpp b/Engine/Source/Programs/UnrealBuildAccelerator/Common/Private/UbaCacheServer.cpp
--- a/Engine/Source/Programs/UnrealBuildAccelerator/Common/Private/UbaCacheServer.cpp
+++ b/Engine/Source/Programs/UnrealBuildAccelerator/Common/Private/UbaCacheServer.cpp
@@ -873,7 +873,7 @@ namespace uba
 					auto& sharedOffsets = entries.sharedInputCasKeyOffsets;
 					if (checkInputsForDeletes)
 					{
-						BinaryReader reader2(sharedOffsets.data(), 0, sharedOffsets.size());
+						BinaryReader reader2(sharedOffsets);
 						while (reader2.GetLeft())
 						{
 							if (!deletedOffsets.IsSet(u32(reader2.Read7BitEncoded())))
@@ -914,7 +914,7 @@ namespace uba
 						{
 							if (!deleteEntry && offsetDeletedInShared)
 							{
-								BinaryReader rangeReader(entry.sharedInputCasKeyOffsetRanges.data(), 0, entry.sharedInputCasKeyOffsetRanges.size());
+								BinaryReader rangeReader(entry.sharedInputCasKeyOffsetRanges);
 								while (!deleteEntry && rangeReader.GetLeft())
 								{
 									u64 begin = rangeReader.Read7BitEncoded();
@@ -934,7 +934,7 @@ namespace uba
 							if (!deleteEntry)
 							{
 								auto& extraInputs = entry.extraInputCasKeyOffsets;
-								BinaryReader extraReader(extraInputs.data(), 0, extraInputs.size());
+								BinaryReader extraReader(extraInputs);
 								while (extraReader.GetLeft())
 								{
 									if (!deletedOffsets.IsSet(u32(extraReader.Read7BitEncoded())))
@@ -953,7 +953,7 @@ namespace uba
 							touchedCas.clear();
 
 							auto& outputs = entry.outputCasKeyOffsets;
-							BinaryReader outputsReader(outputs.data(), 0, outputs.size());
+							BinaryReader outputsReader(outputs);
 							while (outputsReader.GetLeft())
 							{
 								u64 offset = outputsReader.Read7BitEncoded();
@@ -1111,7 +1111,7 @@ namespace uba
 			{
 				auto collectUsedCasKeyOffsets = [&](const Vector<u8>& offsets)
 					{
-						BinaryReader reader2(offsets.data(), 0, offsets.size());
+						BinaryReader reader2(offsets);
 						while (reader2.GetLeft())
 						{
 							u32 offset = u32(reader2.Read7BitEncoded());
@@ -1208,9 +1208,9 @@ namespace uba
 				// Update all casKeyOffsets
 				u64 updateEntriesStart = GetTime();
 
-				m_server.ParallelFor(workerCountToUse, bucket.m_cacheEntryLookup, [&, temp = Vector<u32>(), temp2 = Vector<u8>()](auto& it) mutable
+				m_server.ParallelFor(workerCountToUse, bucket.m_cacheEntryLookup, [&, temp = Vector<u32>(), temp2 = Vector<u8>(), temp3 = Vector<u8>()](auto& it) mutable
 					{
-						it->second.UpdateEntries(m_logger, oldToNewCasKeyOffset, temp, temp2);
+						it->second.UpdateEntries(m_logger, oldToNewCasKeyOffset, temp, temp2, temp3);
 					});
 
 				#if 0
@@ -1489,6 +1489,45 @@ namespace uba
 		Bucket& bucket = GetBucket(id, TC("StoreEntry"));
 		connectionBucket.index = bucket.index;
 
+		auto CreateCasKeyOffset = [&](u32& outCasKeyOffset, CasKey& outCasKey, u32 casKeyOffset, Set<u32>& entries, const tchar* direction)
+			{
+				if (casKeyOffset + sizeof(CasKey) >= connectionBucket.casKeyTable.GetSize())
+					return m_logger.Error(TC("Client %u has provided %s offset that is outside caskey table size. Corrupt data? (Offset: %u TableSize: %u ClientVersion: %u PathtableVersion: %u)"), clientId, direction, casKeyOffset, connectionBucket.casKeyTable.GetSize(), clientVersion, connectionBucket.pathTable.GetVersion());
+
+				CasKey casKey;
+				StringBuffer<> path;
+				if (!connectionBucket.casKeyTable.GetPathAndKey(path, casKey, connectionBucket.pathTable, casKeyOffset))
+				{
+					BinaryReader reader2(connectionBucket.casKeyTable.GetMemory(), casKeyOffset, connectionBucket.casKeyTable.GetSize());
+					u32 pathOffset = u32(reader2.Read7BitEncoded());
+					return m_logger.Error(TC("Client %u has provided corrupt %s key or path table (CasKey: %s Offset: %u TableSize: %u PathOffset: %u PathTableSize: %u ClientVersion: %u PathtableVersion: %u)"), clientId, direction, CasKeyString(casKey).str, casKeyOffset, connectionBucket.casKeyTable.GetSize(), pathOffset, connectionBucket.pathTable.GetSize(), clientVersion, connectionBucket.pathTable.GetVersion());
+				}
+			
+				if (path.count <= 2)
+					return m_logger.Error(TC("Client %u has provided corrupt %s path %s. Corrupt data? (Offset: %u TableSize: %u ClientVersion: %u PathtableVersion: %u)"), clientId, direction, path.data, casKeyOffset, connectionBucket.casKeyTable.GetSize(), clientVersion, connectionBucket.pathTable.GetVersion());
+
+				if (casKey == CasKeyZero)
+					return m_logger.Error(TC("Client %u has provided %s cas key that is zero. Corrupt data? (Offset: %u TableSize: %u ClientVersion: %u PathtableVersion: %u)"), clientId, direction, casKeyOffset, connectionBucket.casKeyTable.GetSize(), clientVersion, connectionBucket.pathTable.GetVersion());
+
+				u32 pathOffset = bucket.m_pathTable.Add(path.data, path.count);
+
+				#if 0
+				StringBuffer<> test;
+				bucket.m_pathTable.GetString(test, pathOffset);
+				UBA_ASSERT(test.Equals(path.data));
+				#endif
+
+				u32 requiredCasTableSize;
+				outCasKeyOffset = bucket.m_casKeyTable.Add(casKey, pathOffset, requiredCasTableSize);
+				auto insres = entries.insert(outCasKeyOffset);
+				if (!insres.second)
+					return m_logger.Warning(TC("Client %u %s file %s exists more than once in cache entry (PathOffset %u, CasKey: %s CasKeyOffset: %u ClientVersion: %u PathtableVersion: %u)"), clientId, direction, path.data, pathOffset, CasKeyString(casKey).str, outCasKeyOffset, clientVersion, connectionBucket.pathTable.GetVersion());
+
+				outCasKey = casKey;
+				return true;
+			};
+
+
 		while (reader.GetLeft())
 		{
 			bool isInput = index++ >= outputCount;
@@ -1499,38 +1538,10 @@ namespace uba
 			if (!isInput)
 				continue;
 
-			if (offset + sizeof(CasKey) >= connectionBucket.casKeyTable.GetSize())
-				return m_logger.Error(TC("Client %u has provided offset that is outside caskey table size. Corrupt data? (Offset: %u TableSize: %u ClientVersion: %u PathtableVersion: %u)"), clientId, offset, connectionBucket.casKeyTable.GetSize(), clientVersion, connectionBucket.pathTable.GetVersion());
-
+			u32 casKeyOffset;
 			CasKey casKey;
-			StringBuffer<> path;
-			if (!connectionBucket.casKeyTable.GetPathAndKey(path, casKey, connectionBucket.pathTable, offset))
-			{
-				BinaryReader reader2(connectionBucket.casKeyTable.GetMemory(), offset, connectionBucket.casKeyTable.GetSize());
-				u32 pathOffset = u32(reader2.Read7BitEncoded());
-				return m_logger.Error(TC("Client %u has provided corrupt key or path table (CasKey: %s Offset: %u TableSize: %u PathOffset: %u PathTableSize: %u ClientVersion: %u PathtableVersion: %u)"), clientId, CasKeyString(casKey).str, offset, connectionBucket.casKeyTable.GetSize(), pathOffset, connectionBucket.pathTable.GetSize(), clientVersion, connectionBucket.pathTable.GetVersion());
-			}
-			
-			if (path.count <= 2)
-				return m_logger.Error(TC("Client %u has provided corrupt path %s. Corrupt data? (Offset: %u TableSize: %u ClientVersion: %u PathtableVersion: %u)"), clientId, path.data, offset, connectionBucket.casKeyTable.GetSize(), clientVersion, connectionBucket.pathTable.GetVersion());
-
-			if (casKey == CasKeyZero)
-				return m_logger.Error(TC("Client %u has provided cas key that is zero. Corrupt data? (Offset: %u TableSize: %u ClientVersion: %u PathtableVersion: %u)"), clientId, offset, connectionBucket.casKeyTable.GetSize(), clientVersion, connectionBucket.pathTable.GetVersion());
-
-			u32 pathOffset = bucket.m_pathTable.Add(path.data, path.count);
-
-			#if 0
-			StringBuffer<> test;
-			bucket.m_pathTable.GetString(test, pathOffset);
-			UBA_ASSERT(test.Equals(path.data));
-			#endif
-
-			u32 requiredCasTableSize;
-			u32 casKeyOffset = bucket.m_casKeyTable.Add(casKey, pathOffset, requiredCasTableSize);
-			auto insres = inputs.insert(casKeyOffset);
-			if (!insres.second)
-				return m_logger.Warning(TC("Client %u Input file %s exists more than once in cache entry (PathOffset %u, CasKey: %s CasKeyOffset: %u ClientVersion: %u PathtableVersion: %u)"), clientId, path.data, pathOffset, CasKeyString(casKey).str, casKeyOffset, clientVersion, connectionBucket.pathTable.GetVersion());
-
+			if (!CreateCasKeyOffset(casKeyOffset, casKey, offset, inputs, TC("input")))
+				return false;
 			bytesForInput += Get7BitEncodedCount(casKeyOffset);
 
 			//m_logger.Info(TC("%s - %s"), path.data, CasKeyString(casKey).str);
@@ -1581,7 +1592,7 @@ namespace uba
 			Map<TString, CasKey> existing;
 
 			// Collect all outputs from entry with matching inputs
-			BinaryReader r2(matchingEntry->outputCasKeyOffsets.data(), 0, matchingEntry->outputCasKeyOffsets.size());
+			BinaryReader r2(matchingEntry->outputCasKeyOffsets);
 			while (r2.GetLeft())
 			{
 				u32 existingOffset = u32(r2.Read7BitEncoded());
@@ -1604,7 +1615,7 @@ namespace uba
 				auto findIt = existing.find(path.data);
 				if (findIt == existing.end())
 				{
-					m_logger.Warning(TC("Client %u sent cache entry that already exists but does not match output (output file %s did not exist in existing cache entry)"), clientId, path.data);
+					m_logger.Warning(TC("Client %u sent cache entry that already exists but does not match output. Output file %s did not exist in existing cache entry. OutputCount Old: %u New: %u"), clientId, path.data, u32(existing.size()), outputCount);
 					cacheEntries.entries.erase(matchingEntry);
 					shouldOverwrite = true;
 					break;
@@ -1631,20 +1642,12 @@ namespace uba
 		while (left--)
 		{
 			u32 outputOffset = u32(reader.Read7BitEncoded());
+
+			u32 casKeyOffset;
 			CasKey casKey;
-			StringBuffer<> path;
-			connectionBucket.casKeyTable.GetPathAndKey(path, casKey, connectionBucket.pathTable, outputOffset);
-			u32 pathOffset = bucket.m_pathTable.Add(path.data, path.count);
-
-			#if 0
-			StringBuffer<> test;
-			bucket.m_pathTable.GetString(test, pathOffset);
-			UBA_ASSERT(test.Equals(path.data));
-			#endif
+			if (!CreateCasKeyOffset(casKeyOffset, casKey, outputOffset, outputs, TC("output")))
+				return false;
 
-			u32 requiredCasTableSize;
-			u32 casKeyOffset = bucket.m_casKeyTable.Add(casKey, pathOffset, requiredCasTableSize);
-			outputs.insert(casKeyOffset);
 			bytesForOutput += Get7BitEncodedCount(casKeyOffset);
 
 			if (!m_storage.EnsureCasFile(casKey, nullptr))
@@ -1860,19 +1863,72 @@ namespace uba
 
 		StringBuffer<> line;
 
-		auto writePathFromOffset = [&](Bucket& bucket, u32 offset)
+		auto writePathFromOffset = [&](Bucket& bucket, u32 offset, u32 index)
 			{
 				CasKey casKey;
 				StringBuffer<> path;
 				bucket.m_casKeyTable.GetPathAndKey(path, casKey, bucket.m_pathTable, offset);
-				writeLine(line.Clear().Appendf(TC("    %s - %s"), path.data, CasKeyString(casKey).str));
+				writeLine(line.Clear().Appendf(TC("    %5u %s - %s (%u)"), index, path.data, CasKeyString(casKey).str, offset));
 			};
 
 		auto writePathsFromOffsets = [&](Bucket& bucket, const Vector<u8>& offsets)
 			{
-				BinaryReader reader2(offsets.data(), 0, offsets.size());
+				u32 index = 0;
+				BinaryReader reader2(offsets);
 				while (reader2.GetLeft())
-					writePathFromOffset(bucket, u32(reader2.Read7BitEncoded()));
+					writePathFromOffset(bucket, u32(reader2.Read7BitEncoded()), index++);
+			};
+
+
+		auto writeEntry = [&](Bucket& bucket, CacheEntries& entries)
+			{
+				#if 0
+				HashMap2<u32, u32> oldToNewCasKeyOffset;
+				MemoryBlock memory(64*1024);
+				oldToNewCasKeyOffset.Init(memory, 1);
+				Vector<u32> temp;
+				Vector<u8> temp2;
+				Vector<u8> temp3;
+				entries.UpdateEntries(m_logger, oldToNewCasKeyOffset, temp, temp2, temp3);
+				#endif
+
+				SCOPED_READ_LOCK(entries.lock, lock3);
+				writeLine(line.Clear().Append(TC("   SharedInputs:")));
+				writePathsFromOffsets(bucket, entries.sharedInputCasKeyOffsets);
+				u32 index = 0;
+				for (auto& entry : entries.entries)
+				{
+					writeLine(line.Clear().Appendf(TC("  #%u (%s)"), index, TimeToText(GetFileTimeAsTime(GetSystemTimeAsFileTime() - (m_creationTime + entry.creationTime)), true).str));
+					writeLine(line.Clear().Append(TC("   InputRanges:")));
+					BinaryReader sharedReader(entries.sharedInputCasKeyOffsets);
+					BinaryReader rangeReader(entry.sharedInputCasKeyOffsetRanges);
+					u32 sharedIndex = 0;
+					while (rangeReader.GetLeft())
+					{
+						u64 begin = rangeReader.Read7BitEncoded();
+						u64 end = rangeReader.Read7BitEncoded();
+
+						while (begin != sharedReader.GetPosition())
+						{
+							sharedReader.Read7BitEncoded();
+							++sharedIndex;
+						}
+						u32 startIndex = sharedIndex;
+						while (end != sharedReader.GetPosition())
+						{
+							sharedReader.Read7BitEncoded();
+							++sharedIndex;
+						}
+
+						writeLine(line.Clear().Appendf(TC("          %llu - %llu   (%llu - %llu)"), startIndex, sharedIndex-1, begin, end));
+					}
+
+					writeLine(line.Clear().Append(TC("   ExtraInputs:")));
+					writePathsFromOffsets(bucket, entry.extraInputCasKeyOffsets);
+					writeLine(line.Clear().Append(TC("   Outputs:")));
+					writePathsFromOffsets(bucket, entry.outputCasKeyOffsets);
+					++index;
+				}
 			};
 
 		if (command.Equals(TC("content")))
@@ -1902,7 +1958,7 @@ namespace uba
 						u32 index = 0;
 						auto findString = [&](const Vector<u8>& offsets)
 							{
-								BinaryReader reader2(offsets.data(), 0, offsets.size());
+								BinaryReader reader2(offsets);
 								while (reader2.GetLeft())
 								{
 									u64 offset = reader2.Read7BitEncoded();
@@ -1935,8 +1991,11 @@ namespace uba
 						if (*entries.inputsThatAreOutputs.begin() == ~0u)
 							writeLine(line.Clear().Appendf(TC("   WillCheckAll (~0u)")));
 						else
+						{
+							u32 index = 0;
 							for (u32 offset : entries.inputsThatAreOutputs)
-								writePathFromOffset(bucket, offset);
+								writePathFromOffset(bucket, offset, index++);
+						}
 					}
 
 					u32 index = 0;
@@ -2027,50 +2086,125 @@ namespace uba
 		{
 			writeLine(line.Clear().Appendf(TC("  Buckets:")));
 			{
-				auto checkPaths = [&](Bucket& bucket, const Vector<u8>& offsets)
-					{
-						BinaryReader reader2(offsets.data(), 0, offsets.size());
-						while (reader2.GetLeft())
-						{
-							u32 offset = u32(reader2.Read7BitEncoded());
-							CasKey casKey;
-							StringBuffer<> path;
-							bucket.m_casKeyTable.GetPathAndKey(path, casKey, bucket.m_pathTable, offset);
-							if (path.count <= 2)
-								return false;
-						}
-						return true;
-					};
-
 				u32 index = 0;
 				SCOPED_READ_LOCK(m_bucketsLock, bucketsLock);
 				for (auto& kv : m_buckets)
 				{
+					bool printInvalid = true;
 					u64 keysWithInvalidSharedInputs = 0;
 					u64 invalidExtraInputs = 0;
 					u64 totalExtraInputs = 0;
+					u64 entriesWithBadExtra = 0;
 					Bucket& bucket = kv.second;
 					SCOPED_READ_LOCK(bucket.m_cacheEntryLookupLock, lock2);
 					for (auto& kv2 : bucket.m_cacheEntryLookup)
 					{
 						CacheEntries& entries = kv2.second;
 						SCOPED_READ_LOCK(entries.lock, lock3);
-						if (!checkPaths(bucket, entries.sharedInputCasKeyOffsets))
+
+						Set<u32> sharedOffsets;
+						BinaryReader sharedInputsReader(entries.sharedInputCasKeyOffsets);
+						bool hasInvalidSharedInputs = false;
+						while (sharedInputsReader.GetLeft())
+						{
+							u32 offset = u32(sharedInputsReader.Read7BitEncoded());
+							sharedOffsets.insert(offset);
+							CasKey casKey;
+							StringBuffer<> path;
+							bucket.m_casKeyTable.GetPathAndKey(path, casKey, bucket.m_pathTable, offset);
+							if (path.count <= 2)
+								hasInvalidSharedInputs = true;
+						}
+						if (hasInvalidSharedInputs)
 							++keysWithInvalidSharedInputs;
 
+						bool hasInvalidExtraInputs = false;
+						bool hasEntriesWithBadExtra = false;
+						u32 entryIndex = 0;
 						for (auto& entry : entries.entries)
 						{
 							++totalExtraInputs;
-							if (!checkPaths(bucket, entry.extraInputCasKeyOffsets))
+
+							BinaryReader extraInputsReader(entry.extraInputCasKeyOffsets);
+							while (extraInputsReader.GetLeft())
+							{
+								u32 offset = u32(extraInputsReader.Read7BitEncoded());
+
+								if (!hasEntriesWithBadExtra && sharedOffsets.find(offset) != sharedOffsets.end())
+								{
+									#if 0
+									u64 fileSize = entries.GetTotalSize(CacheNetworkVersion, true);
+									FileAccessor fa(m_logger, TC("e:\\temp\\CacheEntry.bin"));
+									fa.CreateMemoryWrite(false, DefaultAttributes(), fileSize);
+									BinaryWriter fileWriter(fa.GetData(), 0, fileSize);
+									entries.Write(fileWriter, CacheNetworkVersion, true);
+									fa.Close();
+
+									HashMap2<u32, u32> oldToNewCasKeyOffset;
+									MemoryBlock memory(64*1024);
+									oldToNewCasKeyOffset.Init(memory, 1);
+									Vector<u32> temp;
+									Vector<u8> temp2;
+									Vector<u8> temp3;
+									entries.UpdateEntries(m_logger, oldToNewCasKeyOffset, temp, temp2, temp3);
+									#endif
+									
+									hasEntriesWithBadExtra = true;
+								}
+								
+								CasKey casKey;
+								StringBuffer<> path;
+								bucket.m_casKeyTable.GetPathAndKey(path, casKey, bucket.m_pathTable, offset);
+								if (path.count <= 2)
+									hasInvalidExtraInputs = true;
+							}
+							if (hasEntriesWithBadExtra)
+								++entriesWithBadExtra;
+							if (hasInvalidExtraInputs)
 								++invalidExtraInputs;
+
+							if (printInvalid && (hasInvalidSharedInputs || hasEntriesWithBadExtra || hasInvalidExtraInputs))
+							{
+								writeLine(line.Clear().Appendf(TC("EntryIndex: %u InvalidShared: %u BadExtra: %u InvalidExtra: %u"), entryIndex, hasInvalidSharedInputs, hasEntriesWithBadExtra, hasInvalidExtraInputs));
+								printInvalid = false;
+								writeEntry(bucket, entries);
+								break;
+							}
+							++entryIndex;
 						}
+						if (!printInvalid)
+							break;
+
 					}
 					lock2.Leave();
 
 					writeLine(line.Clear().Appendf(TC("    #%u - %llu (v%u)"), index++, kv.first, bucket.m_pathTable.GetVersion()));
 					writeLine(line.Clear().Appendf(TC("      Invalid shared: %llu/%llu"), keysWithInvalidSharedInputs, (bucket.m_cacheEntryLookup.size())));
 					writeLine(line.Clear().Appendf(TC("      Invalid extra: %llu/%llu"), invalidExtraInputs, totalExtraInputs));
+					writeLine(line.Clear().Appendf(TC("      With bad extra: %llu"), entriesWithBadExtra));
+				}
+			}
+		}
+		else if (command.Equals(TC("updateentries")))
+		{
+			HashMap2<u32, u32> oldToNewCasKeyOffset;
+			MemoryBlock memory(64*1024);
+			oldToNewCasKeyOffset.Init(memory, 1);
+			Vector<u32> temp;
+			Vector<u8> temp2;
+			Vector<u8> temp3;
+
+			SCOPED_READ_LOCK(m_bucketsLock, bucketsLock);
+			for (auto& kv : m_buckets)
+			{
+				Bucket& bucket = kv.second;
+				for (auto& kv2 : bucket.m_cacheEntryLookup)
+				{
+					CacheEntries& entries = kv2.second;
+					SCOPED_READ_LOCK(entries.lock, lock3);
+					entries.UpdateEntries(m_logger, oldToNewCasKeyOffset, temp, temp2, temp3);
 				}
+				bucket.needsSave = true;
 			}
 		}
 		else if (command.Equals(TC("pathtable")))
@@ -2185,22 +2319,7 @@ namespace uba
 			}
 			writeLine(line.Clear().Append(TC("Largest entry (based on shared inputs)")));
 			if (largestEntries)
-			{
-				CacheEntries& entries = *largestEntries;
-				SCOPED_READ_LOCK(entries.lock, lock3);
-				writeLine(line.Clear().Append(TC("   SharedInputs:")));
-				writePathsFromOffsets(*largestEntriesBucket, entries.sharedInputCasKeyOffsets);
-				u32 index = 0;
-				for (auto& entry : entries.entries)
-				{
-					writeLine(line.Clear().Appendf(TC("  #%u (%s)"), index, TimeToText(GetFileTimeAsTime(GetSystemTimeAsFileTime() - (m_creationTime + entry.creationTime)), true).str));
-					writeLine(line.Clear().Append(TC("   ExtraInputs:")));
-					writePathsFromOffsets(*largestEntriesBucket, entry.extraInputCasKeyOffsets);
-					writeLine(line.Clear().Append(TC("   Outputs:")));
-					writePathsFromOffsets(*largestEntriesBucket, entry.outputCasKeyOffsets);
-					++index;
-				}
-			}
+				writeEntry(*largestEntriesBucket, *largestEntries);
 		}
 		else if (command.Equals(TC("obliterate")))
 		{
@@ -2214,6 +2333,10 @@ namespace uba
 			m_addsSinceMaintenance = 1;
 			writeLine(line.Clear().Appendf(TC("Cache server maintenance queued!")));
 		}
+		else if (command.Equals(TC("save")))
+		{
+			Save();
+		}
 		else
 		{
 			writeLine(line.Clear().Appendf(TC("Unknown command: %s"), command.data));
diff --git a/Engine/Source/Programs/UnrealBuildAccelerator/Common/Public/UbaCacheEntry.h b/Engine/Source/Programs/UnrealBuildAccelerator/Common/Public/UbaCacheEntry.h
--- a/Engine/Source/Programs/UnrealBuildAccelerator/Common/Public/UbaCacheEntry.h
+++ b/Engine/Source/Programs/UnrealBuildAccelerator/Common/Public/UbaCacheEntry.h
@@ -50,19 +50,17 @@ namespace uba
 		bool Write(BinaryWriter& writer, u32 clientVersion, bool toDisk);
 		bool ReadFromDisk(Logger& logger, BinaryReader& reader, u32 databaseVersion, StorageImpl& storage, CompactCasKeyTable& table);
 		void BuildInputs(CacheEntry& entry, const Set<u32>& inputs);
-		void UpdateEntries(Logger& logger, const HashMap2<u32, u32>& oldToNewCasKeyOffset, Vector<u32>& temp, Vector<u8>& temp2);
+		void UpdateEntries(Logger& logger, const HashMap2<u32, u32>& oldToNewCasKeyOffset, Vector<u32>& temp, Vector<u8>& temp2, Vector<u8>& temp3);
 
 		void Flatten(Vector<u8>& out, const CacheEntry& entry);
 		void Flatten(Vector<u32>& out, const CacheEntry& entry, const Vector<u8>& sharedOffsets);
 
 		template<typename Container>
-		void BuildInputsT(CacheEntry& entry, const Container& sortedInputs, bool populateShared);
+		void BuildInputsT(CacheEntry& entry, const Container& sortedInputs, bool populateShared, Vector<u8>& temp);
 
 		template<typename Container>
 		void BuildRangesFromExcludedT(CacheEntry& entry, const Container& sortedExcludedInputs);
 
-		void ValidateEntry(Logger& logger, CacheEntry& entry, Vector<u8>& inputCasKeyOffsets);
-
 		void PopulateInputsThatAreOutputs(const Vector<u8>& inputData, StorageImpl& storage, CompactCasKeyTable& table);
 	};
 
diff --git a/Engine/Source/Programs/UnrealBuildAccelerator/Core/Public/UbaBinaryReaderWriter.h b/Engine/Source/Programs/UnrealBuildAccelerator/Core/Public/UbaBinaryReaderWriter.h
--- a/Engine/Source/Programs/UnrealBuildAccelerator/Core/Public/UbaBinaryReaderWriter.h
+++ b/Engine/Source/Programs/UnrealBuildAccelerator/Core/Public/UbaBinaryReaderWriter.h
@@ -82,6 +82,7 @@ namespace uba
 		inline void SetSize(u64 size) { m_end = m_begin + size; }
 		inline const u8* GetPositionData() { return m_pos; }
 		inline BinaryReader(const u8* data, u64 offset = 0, u64 size = InvalidValue) { m_begin = data; m_pos = data + offset; m_end = data + size; }
+		inline BinaryReader(const Vector<u8>& data, u64 offset = 0) { m_begin = data.data(); m_pos = m_begin + offset; m_end = m_begin + data.size(); }
 
 		BinaryReader();
 
diff --git a/Engine/Source/Programs/UnrealBuildAccelerator/Test/Private/UbaTest.h b/Engine/Source/Programs/UnrealBuildAccelerator/Test/Private/UbaTest.h
--- a/Engine/Source/Programs/UnrealBuildAccelerator/Test/Private/UbaTest.h
+++ b/Engine/Source/Programs/UnrealBuildAccelerator/Test/Private/UbaTest.h
@@ -6,7 +6,7 @@
 
 #define CHECK_TRUE(x) \
 	if (!(x)) \
-		return logger.Error(TC("Failed %s (%s:%u)"), TC(#x), __FILE__, __LINE__);
+		return logger.Error(TC("Failed %s (%s:%u)"), TC(#x), TC("") __FILE__, __LINE__);
 
 
 namespace uba
diff --git a/Engine/Source/Programs/UnrealBuildAccelerator/Test/Private/UbaTestCache.cpp b/Engine/Source/Programs/UnrealBuildAccelerator/Test/Private/UbaTestCache.cpp
--- a/Engine/Source/Programs/UnrealBuildAccelerator/Test/Private/UbaTestCache.cpp
+++ b/Engine/Source/Programs/UnrealBuildAccelerator/Test/Private/UbaTestCache.cpp
@@ -5,6 +5,7 @@
 #include "UbaCacheServer.h"
 #include "UbaConfig.h"
 #include "UbaCompactTables.h"
+#include "UbaFileAccessor.h"
 #include "UbaHashMap.h"
 #include "UbaNetworkBackendTcp.h"
 #include "UbaNetworkClient.h"
@@ -20,32 +21,214 @@ namespace uba
 	{
 		CacheEntries entries;
 
-		auto AddEntry = [&](const Set<u32>& inputs)
+		struct Record
+		{
+			Set<u32> offsets;
+			u32 expectedRangeCount;
+			u32 expectedExtraCount;
+		};
+
+		auto CheckEntry = [&](CacheEntry& entry, const Record* rec, bool checkOffsets)
+			{
+				Set<u32> sharedOffsets;
+				{
+					BinaryReader sharedReader(entries.sharedInputCasKeyOffsets);
+					while (sharedReader.GetLeft())
+						if (!sharedOffsets.insert(u32(sharedReader.Read7BitEncoded())).second)
+							return false;
+				}
+
+				Set<u32> entryOffsets;
+
+				u32 rangeCount = 0;
+				BinaryReader rangeReader(entry.sharedInputCasKeyOffsetRanges);
+				while (rangeReader.GetLeft())
+				{
+					++rangeCount;
+					u64 begin = rangeReader.Read7BitEncoded();
+					u64 end = rangeReader.Read7BitEncoded();
+					BinaryReader sharedReader(entries.sharedInputCasKeyOffsets.data(), begin, end);
+					while (sharedReader.GetLeft())
+						if (!entryOffsets.insert(u32(sharedReader.Read7BitEncoded())).second)
+							return false;
+				}
+
+				u32 extraCount = 0;
+				BinaryReader extraReader(entry.extraInputCasKeyOffsets);
+				while (extraReader.GetLeft())
+				{
+					++extraCount;
+					u32 offset = u32(extraReader.Read7BitEncoded());
+					if (sharedOffsets.find(offset) != sharedOffsets.end())
+						return false;
+					if (!entryOffsets.insert(offset).second)
+						return false;
+				}
+
+				if (rec)
+				{
+					if (rangeCount != rec->expectedRangeCount)
+						return false;
+					if (extraCount != rec->expectedExtraCount)
+						return false;
+					if (checkOffsets)
+					{
+						if (rec->offsets != entryOffsets)
+							return false;
+					}
+					else
+					{
+						if (rec->offsets.size() != entryOffsets.size())
+							return false;
+					}
+				}
+				return true;
+			};
+
+		auto AddEntry = [&](const Record& rec)
 			{
 				Vector<u8> inputOffsets;
 				u64 bytes = 0;
-				for (u32 i : inputs)
+				for (u32 i : rec.offsets)
 					bytes += Get7BitEncodedCount(i);
 				inputOffsets.resize(bytes);
 				BinaryWriter writer(inputOffsets.data(), 0, inputOffsets.size());
-				for (u32 i : inputs)
+				for (u32 i : rec.offsets)
 					writer.Write7BitEncoded(i);
 
 				CacheEntry entry;
-				entries.BuildInputs(entry, inputs);
+				entries.BuildInputs(entry, rec.offsets);
 				entries.entries.push_back(entry);
-				entries.ValidateEntry(logger, entry, inputOffsets);
+				return CheckEntry(entry, &rec, true);
+			};
+
+		auto UpdateEntries = [&](u32 multiplier)
+			{
+				HashMap2<u32, u32> oldToNewCasKeyOffset;
+				MemoryBlock memory(64*1024);
+				oldToNewCasKeyOffset.Init(memory, 100);
+				if (multiplier)
+					for (u32 i=0; i!=200; ++i)
+						oldToNewCasKeyOffset.Insert(i) = i * multiplier;
+				Vector<u32> temp;
+				Vector<u8> temp2;
+				Vector<u8> temp3;
+				entries.UpdateEntries(logger, oldToNewCasKeyOffset, temp, temp2, temp3);
+			};
+
+		auto Clear = [&]()
+			{
+				entries.entries.clear();
+				entries.sharedInputCasKeyOffsets.clear();
+				entries.sharedLogLines.clear();
+				entries.idCounter = 0;
+				entries.primaryId = ~0u; // Id of entry that shared offsets was made from
+				entries.inputsThatAreOutputs.clear();
 			};
 
-		AddEntry(Set<u32>{ 1, 4, 6 });
-		AddEntry(Set<u32>{ 0, 4, 6 });
-		AddEntry(Set<u32>{ 2, 4, 6 });
-		AddEntry(Set<u32>{ 1, 4, 5 });
-		AddEntry(Set<u32>{ 1, 4, 7 });
-		AddEntry(Set<u32>{ 1, 3, 6 });
-		AddEntry(Set<u32>{ 1, 5, 6 });
-		AddEntry(Set<u32>{ 1, 4, 6, 7 });
-		AddEntry(Set<u32>{ 0, 1, 4, 6 });
+		#if 0
+		FileAccessor fa(logger, TC("e:\\temp\\CacheEntry.bin"));
+		fa.OpenMemoryRead();
+		BinaryReader fileReader(fa.GetData(), 0, fa.GetSize());
+		entries.ReadFromDisk(logger, fileReader, 9, *(StorageImpl*)nullptr, *(CompactCasKeyTable*)nullptr);
+		fa.Close();
+		UpdateEntries(0);
+		for (auto& entry : entries.entries)
+			CHECK_TRUE(CheckEntry(entry, nullptr, false));
+		#endif
+
+		Record records0[] = 
+		{
+			{ { 1, 4 }, 1, 0 },
+			{ { 1, 2, 3, 4 }, 1, 2 },
+		};
+
+		Clear();
+		for (auto& rec : records0)
+			CHECK_TRUE(AddEntry(rec));
+		for (u32 i=0; i!=4; ++i)
+		{
+			UpdateEntries(i);
+			u32 index = 0;
+			for (auto& entry : entries.entries)
+				CHECK_TRUE(CheckEntry(entry, &records0[index++], i < 2));
+		}
+
+		Record records1[] = 
+		{
+			{ { 1, 4, 6 }, 1, 0 },
+			{ { 0, 4, 6 }, 1, 1 },
+			{ { 2, 4, 6 }, 1, 1 },
+			{ { 1, 4, 5 }, 1, 1 },
+			{ { 1, 4, 7 }, 1, 1 },
+			{ { 1, 3, 6 }, 2, 1 },
+			{ { 1, 5, 6 }, 2, 1 },
+			{ { 1, 4, 6, 7 }, 1, 1 },
+			{ { 0, 1, 4, 6 }, 1, 1 },
+			{ { 1, 4, 5, 6 }, 1, 1 },
+			{ { 0, 1 }, 1, 1 },
+			{ { 0, 2 }, 0, 2 },
+			{ { 0, 4 }, 1, 1 },
+			{ { 0, 5 }, 0, 2 },
+			{ { 0, 6 }, 1, 1 },
+			{ { 0, 7 }, 0, 2 },
+			{ { 1, 2 }, 1, 1 },
+			{ { 1, 4 }, 1, 0 },
+			{ { 1, 5 }, 1, 1 },
+			{ { 1, 6 }, 2, 0 },
+			{ { 1, 7 }, 1, 1 },
+			{ { 2, 4 }, 1, 1 },
+			{ { 2, 3 }, 0, 2 },
+			{ { 7, 8 }, 0, 2 },
+			{ { 0 }, 0, 1 },
+			{ { 1 }, 1, 0 },
+			{ { 2 }, 0, 1 },
+			{ { 3 }, 0, 1 },
+			{ { 4 }, 1, 0 },
+			{ { 5 }, 0, 1 },
+			{ { 6 }, 1, 0 },
+			{ { 7 }, 0, 1 },
+		};
+
+		Clear();
+		for (auto& rec : records1)
+			CHECK_TRUE(AddEntry(rec));
+		for (u32 i=0; i!=4; ++i)
+		{
+			UpdateEntries(i);
+			u32 index = 0;
+			for (auto& entry : entries.entries)
+				CHECK_TRUE(CheckEntry(entry, &records1[index++], i < 2));
+		}
+
+		Record records2[] = 
+		{
+			{ { 2, 4, 6, 10, 14, 18 }, 1, 0 },
+			{ { 2, 4, 6, 10, 14, 18 }, 1, 0 },
+			{ { 2, 4, 5, 10, 15, 18 }, 3, 2 },
+			{ { 2, 4, 6, 10, 19, 20 }, 1, 2 },
+			{ { 0, 1 }, 0, 2 },
+			{ { 4, 10, 18 }, 3, 0 },
+			{ { 7, 8 }, 0, 2 },
+			{ { 6, 7, 8 }, 1, 2 },
+			{ { 5, 6, 7, 8 }, 1, 3 },
+			{ { 2, 4, 6, 7, 8, 10, 14 }, 1, 2 },
+			{ { 2, 4, 6, 7, 8, 10, 14, 18 }, 1, 2 },
+			{ { 7, 8, 10 }, 1, 2 },
+			{ { 7, 8, 10, 14, 18 }, 1, 2 },
+			{ { 4, 7, 14, 18 }, 2, 1 },
+		};
+
+		Clear();
+		for (auto& rec : records2)
+			CHECK_TRUE(AddEntry(rec));
+		for (u32 i=0; i!=4; ++i)
+		{
+			UpdateEntries(i);
+			u32 index = 0;
+			for (auto& entry : entries.entries)
+				CHECK_TRUE(CheckEntry(entry, &records2[index++], i < 2));
+		}
 
 		return true;
 	}
