[UBA]
* Added better error reporting when failing to send messages to server
* Added validation of outputs for cache stores from client to server

diff --git a/Engine/Source/Programs/UnrealBuildAccelerator/Common/Private/UbaStorageClient.cpp b/Engine/Source/Programs/UnrealBuildAccelerator/Common/Private/UbaStorageClient.cpp
--- a/Engine/Source/Programs/UnrealBuildAccelerator/Common/Private/UbaStorageClient.cpp
+++ b/Engine/Source/Programs/UnrealBuildAccelerator/Common/Private/UbaStorageClient.cpp
@@ -531,14 +531,15 @@ namespace uba
 						if (fetchId == u16(~0))
 							return m_logger.Error(TC("Cas content error. Server believes %s was only one segment but client sees more. Size: %llu Left to read: %llu ResponseSize: %u. (%s)"), hint, fileSize, left, responseSize, casFile.data);
 						readBuffer = slot;
-						if (!SendBatchMessages(m_logger, *client, fetchId, readBuffer, BufferSlotSize, left, sizeOfFirstMessage, readIndex, responseSize, initForWrite, hint))
+						u32 error;
+						if (!SendBatchMessages(m_logger, *client, fetchId, readBuffer, BufferSlotSize, left, sizeOfFirstMessage, readIndex, responseSize, initForWrite, hint, &error))
 						{
 							if (proxy)
 							{
 								tryAgain = true;
 								break;
 							}
-							return m_logger.Error(TC("Failed to send batched messages to server (%s)"), casFile.data);
+							return m_logger.Error(TC("Failed to send batched messages to server while retrieving cas %s (%s). Error: %u"), casFile.data, hint, error);
 						}
 					}
 					else
@@ -611,14 +612,15 @@ namespace uba
 									writeCapacity = sizeOfFirstMessage;
 								}
 
-								if (!SendBatchMessages(m_logger, *client, fetchId, writeDest, writeCapacity, leftCompressed, sizeOfFirstMessage, readIndex, responseSize, initForWrite, hint))
+								u32 error;
+								if (!SendBatchMessages(m_logger, *client, fetchId, writeDest, writeCapacity, leftCompressed, sizeOfFirstMessage, readIndex, responseSize, initForWrite, hint, &error))
 								{
 									if (proxy)
 									{
 										tryAgain = true;
 										break;
 									}
-									return m_logger.Error(TC("Failed to send batched messages to server (%s)"), casFile.data);
+									return m_logger.Error(TC("Failed to send batched messages to server while retrieving and decompressing cas %s. (%s) Error: %u"), casFile.data, hint, error);
 								}
 
 								if (extraBuffer)
diff --git a/Engine/Source/Programs/UnrealBuildAccelerator/Common/Private/UbaStorageUtils.cpp b/Engine/Source/Programs/UnrealBuildAccelerator/Common/Private/UbaStorageUtils.cpp
--- a/Engine/Source/Programs/UnrealBuildAccelerator/Common/Private/UbaStorageUtils.cpp
+++ b/Engine/Source/Programs/UnrealBuildAccelerator/Common/Private/UbaStorageUtils.cpp
@@ -101,10 +101,13 @@ namespace uba
 		return ToCasKey(hasher, storeCompressed);
 	}
 
-	bool SendBatchMessages(Logger& logger, NetworkClient& client, u16 fetchId, u8* slot, u64 capacity, u64 left, u32 messageMaxSize, u32& readIndex, u32& responseSize, const Function<bool()>& runInWaitFunc, const tchar* hint)
+	bool SendBatchMessages(Logger& logger, NetworkClient& client, u16 fetchId, u8* slot, u64 capacity, u64 left, u32 messageMaxSize, u32& readIndex, u32& responseSize, const Function<bool()>& runInWaitFunc, const tchar* hint, u32* outError)
 	{
 		responseSize = 0;
 
+		if (outError)
+			*outError = 0;
+
 		struct Entry
 		{
 			Entry(u8* slot, u32 i, u32 messageMaxSize) : reader(slot + i * messageMaxSize, 0, SendMaxSize), done(true) {}
@@ -128,6 +131,7 @@ namespace uba
 		Entry* entries = (Entry*)entriesMem;
 
 		bool success = true;
+		u32 error = 0;
 		u32 inFlightCount = u32(sendCount);
 		for (u32 i=0; i!=sendCount; ++i)
 		{
@@ -138,6 +142,7 @@ namespace uba
 			writer.WriteU32(readIndex + i + 1);
 			if (entry.message.SendAsync(entry.reader, [](bool error, void* userData) { ((Event*)userData)->Set(); }, &entry.done))
 				continue;
+			error = entry.message.GetError();
 			entry.~Entry();
 			inFlightCount = i;
 			success = false;
@@ -145,8 +150,14 @@ namespace uba
 		}
 
 		if (runInWaitFunc)
+		{
 			if (!runInWaitFunc())
-				return false;
+			{
+				success = false;
+				if (!error)
+					error = 100;
+			}
+		}
 
 		u32 timeOutTimeMs = 5*60*1000;
 
@@ -159,7 +170,11 @@ namespace uba
 				timeOutTimeMs = 10;
 			}
 			if (!entry.message.ProcessAsyncResults(entry.reader))
+			{
+				if (!error)
+					error = entry.message.GetError();
 				success = false;
+			}
 			else
 				responseSize += u32(entry.reader.GetLeft());
 		}
@@ -168,6 +183,10 @@ namespace uba
 			entries[i].~Entry();
 
 		readIndex += u32(sendCount);
+
+		if (outError)
+			*outError = error;
+
 		return success;
 	}
 
@@ -442,7 +461,7 @@ namespace uba
 			writer.WriteString(destination);
 			BinaryReader reader(slot + (destinationIsCompressed ? sizeof(CompressedObjFileHeader) : 0), 0, SendMaxSize); // Create some space in front of reader to write obj file header there if destination is compressed
 			if (!msg.Send(reader))
-				return logger.Error(TC("Failed to send fetch begin message for cas %s (%s)"), CasKeyString(casKey).str, destination);
+				return logger.Error(TC("Failed to send fetch begin message for cas %s (%s). Error: %u"), CasKeyString(casKey).str, destination, msg.GetError());
 			sizeOfFirstMessage = u32(reader.GetLeft());
 			fetchId = reader.ReadU16();
 			if (fetchId == 0)
@@ -535,8 +554,9 @@ namespace uba
 				if (fetchId == u16(~0))
 					return logger.Error(TC("Cas content error (2). Server believes %s was only one segment but client sees more. "));//UncompressedSize: %llu LeftUncompressed: %llu Size: %llu Left to read: %llu ResponseSize: %u. (%s)"), destination, actualSize, leftUncompressed, fileSize, left, responseSize, CasKeyString(casKey).str);
 
-				if (!SendBatchMessages(logger, client, fetchId, source, BufferSlotHalfSize, leftCompressed, sizeOfFirstMessage, readIndex, responseSize, writePrev, destination))
-					return logger.Error(TC("Failed to send batched messages to server (%s)"), CasKeyString(casKey).str);
+				u32 error;
+				if (!SendBatchMessages(logger, client, fetchId, source, BufferSlotHalfSize, leftCompressed, sizeOfFirstMessage, readIndex, responseSize, writePrev, destination, &error))
+					return logger.Error(TC("Failed to send batched messages to server while retrieving cas %s to %s. Error: %u"), CasKeyString(casKey).str, destination, error);
 
 				lastSource = source;
 				lastResponseSize = responseSize;
@@ -580,8 +600,9 @@ namespace uba
 					{
 						if (fetchId == u16(~0))
 							return logger.Error(TC("Cas content error (2). Server believes %s was only one segment but client sees more. UncompressedSize: %llu LeftUncompressed: %llu Size: %llu Left to read: %llu ResponseSize: %u. (%s)"), destination, actualSize, leftUncompressed, fileSize, left, responseSize, CasKeyString(casKey).str);
-						if (!SendBatchMessages(logger, client, fetchId, readPosition, maxReadSize - u32(readPosition - readBuffer), leftCompressed, sizeOfFirstMessage, readIndex, responseSize, tryWriteDecompressed))
-							return logger.Error(TC("Failed to send batched messages to server (%s)"), CasKeyString(casKey).str);
+						u32 error;
+						if (!SendBatchMessages(logger, client, fetchId, readPosition, maxReadSize - u32(readPosition - readBuffer), leftCompressed, sizeOfFirstMessage, readIndex, responseSize, tryWriteDecompressed, destination, &error))
+							return logger.Error(TC("Failed to send batched messages to server while retrieving and decompressing cas %s to %s. Error: %u"), CasKeyString(casKey).str, destination, error);
 						leftCompressed -= responseSize;
 					}
 					else
diff --git a/Engine/Source/Programs/UnrealBuildAccelerator/Common/Public/UbaStorageUtils.h b/Engine/Source/Programs/UnrealBuildAccelerator/Common/Public/UbaStorageUtils.h
--- a/Engine/Source/Programs/UnrealBuildAccelerator/Common/Public/UbaStorageUtils.h
+++ b/Engine/Source/Programs/UnrealBuildAccelerator/Common/Public/UbaStorageUtils.h
@@ -16,7 +16,7 @@ namespace uba
 	static constexpr u8 DefaultCompressionLevel = OodleLZ_CompressionLevel_SuperFast;
 
 	CasKey CalculateCasKey(u8* fileMem, u64 fileSize, bool storeCompressed, WorkManager* workManager, const tchar* hint);
-	bool SendBatchMessages(Logger& logger, NetworkClient& client, u16 fetchId, u8* slot, u64 capacity, u64 left, u32 messageMaxSize, u32& readIndex, u32& responseSize, const Function<bool()>& runInWaitFunc = {}, const tchar* hint = TC(""));
+	bool SendBatchMessages(Logger& logger, NetworkClient& client, u16 fetchId, u8* slot, u64 capacity, u64 left, u32 messageMaxSize, u32& readIndex, u32& responseSize, const Function<bool()>& runInWaitFunc = {}, const tchar* hint = TC(""), u32* outError = nullptr);
 	bool SendFile(Logger& logger, NetworkClient& client, const CasKey& casKey, const u8* sourceMem, u64 sourceSize, const tchar* hint);
 
 	struct FileSender
