[UBA]
* Attempt at fixing macos unit test error when trying to find clang++

diff --git a/Engine/Source/Programs/UnrealBuildAccelerator/Test/Private/UbaTestSession.cpp b/Engine/Source/Programs/UnrealBuildAccelerator/Test/Private/UbaTestSession.cpp
--- a/Engine/Source/Programs/UnrealBuildAccelerator/Test/Private/UbaTestSession.cpp
+++ b/Engine/Source/Programs/UnrealBuildAccelerator/Test/Private/UbaTestSession.cpp
@@ -254,7 +254,7 @@ namespace uba
 		const tchar* clangPath = TC("c:\\Program Files\\Microsoft Visual Studio\\2022\\Professional\\VC\\Tools\\Llvm\\bin\\clang-cl.exe");
 #elif PLATFORM_MAC
 		StringBuffer<MaxPath> xcodePath;
-		if (!ExecuteCommand(logger, "/usr/bin/xcrun -r clang++", xcodePath))
+		if (!ExecuteCommand(logger, "/usr/bin/xcrun --find clang++", xcodePath))
 			return true;
 		const tchar* clangPath = xcodePath.data;
 #else
