[UBA]
* Changed readlocks to writelocks in ubacacheserver commands to prevent deadlocks

diff --git a/Engine/Source/Programs/UnrealBuildAccelerator/Common/Private/UbaCacheServer.cpp b/Engine/Source/Programs/UnrealBuildAccelerator/Common/Private/UbaCacheServer.cpp
--- a/Engine/Source/Programs/UnrealBuildAccelerator/Common/Private/UbaCacheServer.cpp
+++ b/Engine/Source/Programs/UnrealBuildAccelerator/Common/Private/UbaCacheServer.cpp
@@ -1941,7 +1941,7 @@ namespace uba
 
 			Vector<u8> temp;
 
-			SCOPED_READ_LOCK(m_bucketsLock, bucketsLock);
+			SCOPED_WRITE_LOCK(m_bucketsLock, bucketsLock);
 			for (auto& kv : m_buckets)
 			{
 				Bucket& bucket = kv.second;
@@ -2030,7 +2030,7 @@ namespace uba
 			u32 index = 0;
 
 			{
-				SCOPED_READ_LOCK(m_bucketsLock, bucketsLock);
+				SCOPED_WRITE_LOCK(m_bucketsLock, bucketsLock);
 				for (auto& kv : m_buckets)
 				{
 					Bucket& bucket = kv.second;
@@ -2087,7 +2087,7 @@ namespace uba
 			writeLine(line.Clear().Appendf(TC("  Buckets:")));
 			{
 				u32 index = 0;
-				SCOPED_READ_LOCK(m_bucketsLock, bucketsLock);
+				SCOPED_WRITE_LOCK(m_bucketsLock, bucketsLock);
 				for (auto& kv : m_buckets)
 				{
 					bool printInvalid = true;
@@ -2194,7 +2194,7 @@ namespace uba
 			Vector<u8> temp2;
 			Vector<u8> temp3;
 
-			SCOPED_READ_LOCK(m_bucketsLock, bucketsLock);
+			SCOPED_WRITE_LOCK(m_bucketsLock, bucketsLock);
 			for (auto& kv : m_buckets)
 			{
 				Bucket& bucket = kv.second;
@@ -2209,7 +2209,7 @@ namespace uba
 		}
 		else if (command.Equals(TC("pathtable")))
 		{
-			SCOPED_READ_LOCK(m_bucketsLock, bucketsLock);
+			SCOPED_WRITE_LOCK(m_bucketsLock, bucketsLock);
 			for (auto& kv : m_buckets)
 			{
 				Bucket& bucket = kv.second;
@@ -2231,7 +2231,7 @@ namespace uba
 		}
 		else if (command.Equals(TC("usedpaths")))
 		{
-			SCOPED_READ_LOCK(m_bucketsLock, bucketsLock);
+			SCOPED_WRITE_LOCK(m_bucketsLock, bucketsLock);
 			for (auto& kv : m_buckets)
 			{
 				Bucket& bucket = kv.second;
@@ -2257,7 +2257,7 @@ namespace uba
 		}
 		else if (command.Equals(TC("usedsegments")))
 		{
-			SCOPED_READ_LOCK(m_bucketsLock, bucketsLock);
+			SCOPED_WRITE_LOCK(m_bucketsLock, bucketsLock);
 			for (auto& kv : m_buckets)
 			{
 				UnorderedMap<u32, u32> segments;
@@ -2302,7 +2302,7 @@ namespace uba
 			CacheEntries* largestEntries = nullptr;
 			Bucket* largestEntriesBucket = nullptr;
 			u64 largestSize = 0;
-			SCOPED_READ_LOCK(m_bucketsLock, bucketsLock);
+			SCOPED_WRITE_LOCK(m_bucketsLock, bucketsLock);
 			for (auto& kv : m_buckets)
 			{
 				Bucket& bucket = kv.second;
