[UBA]
* Changed some asserts for linux/mac to be errors in order to try to track down issue seen on farm
* Improved some error logging to contain more information

diff --git a/Engine/Source/Programs/UnrealBuildAccelerator/Common/Private/UbaFileAccessor.cpp b/Engine/Source/Programs/UnrealBuildAccelerator/Common/Private/UbaFileAccessor.cpp
--- a/Engine/Source/Programs/UnrealBuildAccelerator/Common/Private/UbaFileAccessor.cpp
+++ b/Engine/Source/Programs/UnrealBuildAccelerator/Common/Private/UbaFileAccessor.cpp
@@ -155,7 +155,7 @@ namespace uba
 
 		m_mappingHandle = uba::CreateFileMappingW(m_logger, m_fileHandle, PAGE_READWRITE, size, realFileName);
 		if (!m_mappingHandle.IsValid())
-			return m_logger.Error(TC("Failed to create memory map %s (%s)"), realFileName, LastErrorToText().data);
+			return m_logger.Error(TC("Failed to create memory map %s with size %llu (%s)"), realFileName, size, LastErrorToText().data);
 
 		m_data = MapViewOfFile(m_logger, m_mappingHandle, FILE_MAP_WRITE, 0, size);
 		if (!m_data)
diff --git a/Engine/Source/Programs/UnrealBuildAccelerator/Common/Private/UbaFileMapping.cpp b/Engine/Source/Programs/UnrealBuildAccelerator/Common/Private/UbaFileMapping.cpp
--- a/Engine/Source/Programs/UnrealBuildAccelerator/Common/Private/UbaFileMapping.cpp
+++ b/Engine/Source/Programs/UnrealBuildAccelerator/Common/Private/UbaFileMapping.cpp
@@ -207,13 +207,13 @@ namespace uba
 #if PLATFORM_MAC // For some reason lseek+write does not work on apple silicon platform
 			if (ftruncate(fd, maxSize) == -1)
 			{
-				UBA_ASSERTF(false, "ftruncate to %llu on fd %i failed for %s: %s\n", maxSize, fd, hint, strerror(errno));
+				logger.Error("ftruncate to %llu on fd %i failed for %s: %s\n", maxSize, fd, hint, strerror(errno));
 				return h;
 			}
 #else
 			if (lseek(fd, maxSize - 1, SEEK_SET) != maxSize - 1)
 			{
-				UBA_ASSERTF(false, "lseek to %llu failed for %s: %s\n", maxSize - 1, hint, strerror(errno));
+				logger.Error("lseek to %llu failed for %s: %s\n", maxSize - 1, hint, strerror(errno));
 				return h;
 			}
 
@@ -221,7 +221,7 @@ namespace uba
 			int res = write(fd, "", 1);
 			if (res != 1)
 			{
-				UBA_ASSERTF(false, "write one byte at %llu on fd %i (%s) failed (res: %i): %s\n", maxSize - 1, fd, hint, res, strerror(errno));
+				logger.Error("write one byte at %llu on fd %i (%s) failed (res: %i): %s\n", maxSize - 1, fd, hint, res, strerror(errno));
 				return h;
 			}
 #endif
