[UBA]
* Fixed a cache server bug where in certain scenarios input dependencies could get lost when rebuilding entries.

diff --git a/Engine/Source/Programs/UnrealBuildAccelerator/Common/Private/UbaCacheEntry.cpp b/Engine/Source/Programs/UnrealBuildAccelerator/Common/Private/UbaCacheEntry.cpp
--- a/Engine/Source/Programs/UnrealBuildAccelerator/Common/Private/UbaCacheEntry.cpp
+++ b/Engine/Source/Programs/UnrealBuildAccelerator/Common/Private/UbaCacheEntry.cpp
@@ -300,7 +300,7 @@ namespace uba
 		u32 lastSharedPos = ~0u;
 
 		Vector<u8>& extraOffsets = temp;
-		extraOffsets.resize(sortedInputs.size()*4);
+		extraOffsets.resize(sortedInputs.size()*5);
 		BinaryWriter extraWriter(extraOffsets.data(), 0, extraOffsets.size());
 
 		while (true)
@@ -314,8 +314,9 @@ namespace uba
 					while (inputsIt != inputsEnd)
 					{
 						offset = *inputsIt++;
-						if (offset < sharedOffset)
-							extraWriter.Write7BitEncoded(offset);
+						if (offset >= sharedOffset)
+							break;
+						extraWriter.Write7BitEncoded(offset);
 					}
 				}
 
diff --git a/Engine/Source/Programs/UnrealBuildAccelerator/Test/Private/UbaTestCache.cpp b/Engine/Source/Programs/UnrealBuildAccelerator/Test/Private/UbaTestCache.cpp
--- a/Engine/Source/Programs/UnrealBuildAccelerator/Test/Private/UbaTestCache.cpp
+++ b/Engine/Source/Programs/UnrealBuildAccelerator/Test/Private/UbaTestCache.cpp
@@ -141,6 +141,8 @@ namespace uba
 		{
 			{ { 1, 4 }, 1, 0 },
 			{ { 1, 2, 3, 4 }, 1, 2 },
+			{ { 1, 2, 3, 4, 5 }, 1, 3 },
+			{ { 1, 2, 3, 4, 5, 6 }, 1, 4 },
 		};
 
 		Clear();
@@ -163,7 +165,11 @@ namespace uba
 			{ { 1, 4, 7 }, 1, 1 },
 			{ { 1, 3, 6 }, 2, 1 },
 			{ { 1, 5, 6 }, 2, 1 },
+			{ { 1, 5, 7 }, 1, 2 },
+			{ { 1, 3, 5, 7 }, 1, 3 },
+			{ { 1, 3, 5, 7, 8 }, 1, 4 },
 			{ { 1, 4, 6, 7 }, 1, 1 },
+			{ { 1, 4, 6, 7, 8 }, 1, 2 },
 			{ { 0, 1, 4, 6 }, 1, 1 },
 			{ { 1, 4, 5, 6 }, 1, 1 },
 			{ { 0, 1 }, 1, 1 },
