[UBA]
* Fixed case insensitive logic to prevent lots of additional data
* Added more validation to cache server "validate" command

diff --git a/Engine/Source/Programs/UnrealBuildAccelerator/Common/Private/UbaCacheServer.cpp b/Engine/Source/Programs/UnrealBuildAccelerator/Common/Private/UbaCacheServer.cpp
--- a/Engine/Source/Programs/UnrealBuildAccelerator/Common/Private/UbaCacheServer.cpp
+++ b/Engine/Source/Programs/UnrealBuildAccelerator/Common/Private/UbaCacheServer.cpp
@@ -1079,7 +1079,7 @@ namespace uba
 
 		m_bucketIsOverflowing = false;
 
-		if (shouldExit())
+		if (shouldExit && shouldExit())
 			return true;
 
 		u64 maxCommittedMemory = 0;
@@ -2091,10 +2091,15 @@ namespace uba
 				for (auto& kv : m_buckets)
 				{
 					bool printInvalid = true;
+					bool printLargestDiff = false;
+					u32 largestDiff = 0;
+					CasKey largestDiffKey;
+					u64 mostEntries = 0;
 					u64 keysWithInvalidSharedInputs = 0;
 					u64 invalidExtraInputs = 0;
 					u64 totalExtraInputs = 0;
 					u64 entriesWithBadExtra = 0;
+					u64 entriesWithDuplicatedExtra = 0;
 					Bucket& bucket = kv.second;
 					SCOPED_READ_LOCK(bucket.m_cacheEntryLookupLock, lock2);
 					for (auto& kv2 : bucket.m_cacheEntryLookup)
@@ -2103,10 +2108,13 @@ namespace uba
 						SCOPED_READ_LOCK(entries.lock, lock3);
 
 						Set<u32> sharedOffsets;
+						UnorderedMap<u32, u32> positionToIndex;
 						BinaryReader sharedInputsReader(entries.sharedInputCasKeyOffsets);
 						bool hasInvalidSharedInputs = false;
+						u32 sharedIndex = 0;
 						while (sharedInputsReader.GetLeft())
 						{
+							positionToIndex[u32(sharedInputsReader.GetPosition())] = sharedIndex++;
 							u32 offset = u32(sharedInputsReader.Read7BitEncoded());
 							sharedOffsets.insert(offset);
 							CasKey casKey;
@@ -2115,20 +2123,49 @@ namespace uba
 							if (path.count <= 2)
 								hasInvalidSharedInputs = true;
 						}
+						positionToIndex[u32(sharedInputsReader.GetPosition())] = sharedIndex++;
 						if (hasInvalidSharedInputs)
 							++keysWithInvalidSharedInputs;
 
+						u32 minCount = ~0u;
+						u32 maxCount = 0;
+
+						if (mostEntries < entries.entries.size())
+							mostEntries = entries.entries.size();
+
 						bool hasInvalidExtraInputs = false;
+						bool hasDuplicatedExtraInputs = false;
 						bool hasEntriesWithBadExtra = false;
 						u32 entryIndex = 0;
 						for (auto& entry : entries.entries)
 						{
 							++totalExtraInputs;
 
+							u32 inputsCount = 0;
+
+							BinaryReader rangesReader(entry.sharedInputCasKeyOffsetRanges);
+							while (rangesReader.GetLeft())
+							{
+								u32 begin = u32(rangesReader.Read7BitEncoded());
+								u32 end = u32(rangesReader.Read7BitEncoded());
+								auto f1 = positionToIndex.find(begin);
+								auto f2 = positionToIndex.find(end);
+								UBA_ASSERT(f1 != positionToIndex.end());
+								UBA_ASSERT(f2 != positionToIndex.end());
+								u32 count = f2->second - f1->second;
+								UBA_ASSERT(count);
+								inputsCount += count;
+							}
+
+							Set<u32> extras;
 							BinaryReader extraInputsReader(entry.extraInputCasKeyOffsets);
 							while (extraInputsReader.GetLeft())
 							{
 								u32 offset = u32(extraInputsReader.Read7BitEncoded());
+								++inputsCount;
+
+								if (!hasDuplicatedExtraInputs && !extras.insert(offset).second)
+									hasDuplicatedExtraInputs = true;
 
 								if (!hasEntriesWithBadExtra && sharedOffsets.find(offset) != sharedOffsets.end())
 								{
@@ -2162,7 +2199,9 @@ namespace uba
 								++entriesWithBadExtra;
 							if (hasInvalidExtraInputs)
 								++invalidExtraInputs;
-
+							if (hasDuplicatedExtraInputs)
+								++entriesWithDuplicatedExtra;
+#if 0
 							if (printInvalid && (hasInvalidSharedInputs || hasEntriesWithBadExtra || hasInvalidExtraInputs))
 							{
 								writeLine(line.Clear().Appendf(TC("EntryIndex: %u InvalidShared: %u BadExtra: %u InvalidExtra: %u"), entryIndex, hasInvalidSharedInputs, hasEntriesWithBadExtra, hasInvalidExtraInputs));
@@ -2170,11 +2209,23 @@ namespace uba
 								writeEntry(bucket, entries);
 								break;
 							}
+#endif
 							++entryIndex;
+
+							if (inputsCount < minCount)
+								minCount = inputsCount;
+							if (inputsCount > maxCount)
+								maxCount = inputsCount;
 						}
 						if (!printInvalid)
 							break;
 
+						u32 diff = maxCount - minCount;
+						if (diff > largestDiff)
+						{
+							largestDiff = diff;
+							largestDiffKey = kv2.first;
+						}
 					}
 					lock2.Leave();
 
@@ -2182,6 +2233,14 @@ namespace uba
 					writeLine(line.Clear().Appendf(TC("      Invalid shared: %llu/%llu"), keysWithInvalidSharedInputs, (bucket.m_cacheEntryLookup.size())));
 					writeLine(line.Clear().Appendf(TC("      Invalid extra: %llu/%llu"), invalidExtraInputs, totalExtraInputs));
 					writeLine(line.Clear().Appendf(TC("      With bad extra: %llu"), entriesWithBadExtra));
+					writeLine(line.Clear().Appendf(TC("      With duplicated extra: %llu"), entriesWithDuplicatedExtra));
+					writeLine(line.Clear().Appendf(TC("      LargestDiff: %u"), largestDiff));
+					writeLine(line.Clear().Appendf(TC("      MostEntries: %llu"), mostEntries));
+
+					if (largestDiff && printInvalid && printLargestDiff)
+					{
+						writeEntry(bucket, bucket.m_cacheEntryLookup[largestDiffKey]);
+					}
 				}
 			}
 		}
diff --git a/Engine/Source/Programs/UnrealBuildAccelerator/Common/Private/UbaCompactTables.cpp b/Engine/Source/Programs/UnrealBuildAccelerator/Common/Private/UbaCompactTables.cpp
--- a/Engine/Source/Programs/UnrealBuildAccelerator/Common/Private/UbaCompactTables.cpp
+++ b/Engine/Source/Programs/UnrealBuildAccelerator/Common/Private/UbaCompactTables.cpp
@@ -62,7 +62,12 @@ namespace uba
 		{
 			StringBuffer<> temp;
 			table.GetString(temp, offset);
-			return InternalAdd(temp.data, temp.data, temp.count);
+
+			const tchar* stringKeyString = temp.data;
+			StringBuffer<MaxPath> tempStringKeyStr;
+			if (m_caseInsensitive)
+				stringKeyString = tempStringKeyStr.Append(temp).MakeLower().data;
+			return InternalAdd(temp.data, stringKeyString, temp.count);
 		}
 #if 0
 		u32 offsets[MaxSegments];
