[UBA]
* Changed cache server logic that compares existing matching entry's output with new output to make path lowercase if bucket is path insensitive

diff --git a/Engine/Source/Programs/UnrealBuildAccelerator/Common/Private/UbaCacheServer.cpp b/Engine/Source/Programs/UnrealBuildAccelerator/Common/Private/UbaCacheServer.cpp
--- a/Engine/Source/Programs/UnrealBuildAccelerator/Common/Private/UbaCacheServer.cpp
+++ b/Engine/Source/Programs/UnrealBuildAccelerator/Common/Private/UbaCacheServer.cpp
@@ -1599,6 +1599,8 @@ namespace uba
 				CasKey casKey;
 				StringBuffer<> path;
 				bucket.m_casKeyTable.GetPathAndKey(path, casKey, bucket.m_pathTable, existingOffset);
+				if (IsCaseInsensitive(id))
+					path.MakeLower();
 				existing.try_emplace(path.data, casKey);
 			}
 
@@ -1611,6 +1613,8 @@ namespace uba
 				CasKey casKey;
 				StringBuffer<> path;
 				connectionBucket.casKeyTable.GetPathAndKey(path, casKey, connectionBucket.pathTable, outputOffset);
+				if (IsCaseInsensitive(id))
+					path.MakeLower();
 
 				auto findIt = existing.find(path.data);
 				if (findIt == existing.end())
