[UBA]
* Added debug code for uba cache client to try to track down why some files look like they are size 0. Either they are size 0 and there is a bug in the code downloading/writing the file... or the directory table population has a race condition or something

diff --git a/Engine/Source/Programs/UnrealBuildAccelerator/Common/Private/UbaCacheClient.cpp b/Engine/Source/Programs/UnrealBuildAccelerator/Common/Private/UbaCacheClient.cpp
--- a/Engine/Source/Programs/UnrealBuildAccelerator/Common/Private/UbaCacheClient.cpp
+++ b/Engine/Source/Programs/UnrealBuildAccelerator/Common/Private/UbaCacheClient.cpp
@@ -808,6 +808,21 @@ namespace uba
 		cacheStats.fetchBytesRaw += fetcher.sizeOnDisk;
 		cacheStats.fetchBytesComp += fetcher.bytesReceived;
 
+		// TODO: Remove when bug is found
+		if (path.EndsWith(TC(".obj")) || path.EndsWith(TC(".o")))
+		{
+			if (fetcher.sizeOnDisk == 0)
+				return m_logger.Warning(TC("Received %s file that has size 0!"), path.data);
+			else
+			{
+				FileAccessor file(m_logger, path.data);
+				if (!file.OpenMemoryRead())
+					return m_logger.Warning(TC("Failed to test open received %s file"), path.data);
+				if (file.GetSize() != fetcher.sizeOnDisk)
+					return m_logger.Warning(TC("Received file %s has size %llu even though it shouldn have %llu"), path.data, file.GetSize(), fetcher.sizeOnDisk);
+			}
+		}
+
 		if (!m_storage.FakeCopy(casKey, path.data, fetcher.sizeOnDisk, fetcher.lastWritten, false))
 			return false;
 		if (!m_session.RegisterNewFile(path.data))
diff --git a/Engine/Source/Programs/UnrealBuildAccelerator/Common/Private/UbaSession.cpp b/Engine/Source/Programs/UnrealBuildAccelerator/Common/Private/UbaSession.cpp
--- a/Engine/Source/Programs/UnrealBuildAccelerator/Common/Private/UbaSession.cpp
+++ b/Engine/Source/Programs/UnrealBuildAccelerator/Common/Private/UbaSession.cpp
@@ -857,12 +857,18 @@ namespace uba
 
 			if (lastWriteTime == oldLastWriteTime || isDirectory) // Ignore updating directories.. they should always be the same regardless
 			{
-				#if !PLATFORM_WINDOWS
 				reader.Skip(sizeof(u32) * 2);
+				#if !PLATFORM_WINDOWS
 				u64 oldFileIndex = reader.ReadU64();
 				UBA_ASSERT(oldFileIndex == fileIndex); // Checking so it is really the same file
+				#else
+				reader.Skip(sizeof(u64));
 				#endif
-				return true;
+				u64 oldSize = reader.ReadU64();
+				if (oldSize == fileSize)
+					return true;
+				else
+					m_logger.Warning(TC("Somehow file %s has same last written time at two points in time but different size (old %llu new %llu)"), fileName.data, oldSize, fileSize);
 			}
 		}
 
