[UBA]
* Removed uba cache client test code and removed comment on when last file time is the same but size is different.. will revisit this but for now this hopefully fix some farm bugs

diff --git a/Engine/Source/Programs/UnrealBuildAccelerator/Common/Private/UbaCacheClient.cpp b/Engine/Source/Programs/UnrealBuildAccelerator/Common/Private/UbaCacheClient.cpp
--- a/Engine/Source/Programs/UnrealBuildAccelerator/Common/Private/UbaCacheClient.cpp
+++ b/Engine/Source/Programs/UnrealBuildAccelerator/Common/Private/UbaCacheClient.cpp
@@ -809,6 +809,7 @@ namespace uba
 		cacheStats.fetchBytesComp += fetcher.bytesReceived;
 
 		// TODO: Remove when bug is found
+#if 0
 		if (path.EndsWith(TC(".obj")) || path.EndsWith(TC(".o")))
 		{
 			if (fetcher.sizeOnDisk == 0)
@@ -822,7 +823,7 @@ namespace uba
 					return m_logger.Warning(TC("Received file %s has size %llu even though it shouldn have %llu"), path.data, file.GetSize(), fetcher.sizeOnDisk);
 			}
 		}
-
+#endif
 		if (!m_storage.FakeCopy(casKey, path.data, fetcher.sizeOnDisk, fetcher.lastWritten, false))
 			return false;
 		if (!m_session.RegisterNewFile(path.data))
diff --git a/Engine/Source/Programs/UnrealBuildAccelerator/Common/Private/UbaSession.cpp b/Engine/Source/Programs/UnrealBuildAccelerator/Common/Private/UbaSession.cpp
--- a/Engine/Source/Programs/UnrealBuildAccelerator/Common/Private/UbaSession.cpp
+++ b/Engine/Source/Programs/UnrealBuildAccelerator/Common/Private/UbaSession.cpp
@@ -867,8 +867,9 @@ namespace uba
 				u64 oldSize = reader.ReadU64();
 				if (oldSize == fileSize)
 					return true;
-				else
-					m_logger.Warning(TC("Somehow file %s has same last written time at two points in time but different size (old %llu new %llu)"), fileName.data, oldSize, fileSize);
+				// TODO: Somehow this can happen and I have no idea how. last written time should be set on close file so it shouldnt be possible.
+				//else
+				//	m_logger.Warning(TC("Somehow file %s has same last written time at two points in time but different size (old %llu new %llu)"), fileName.data, oldSize, fileSize);
 			}
 		}
 
