[UBA]
* Fixed bugs in cache entry shared/extra offsets calculations
* Extended unit tests to test all combinations of shared/extra for the records in the test

diff --git a/Engine/Source/Programs/UnrealBuildAccelerator/Common/Private/UbaCacheEntry.cpp b/Engine/Source/Programs/UnrealBuildAccelerator/Common/Private/UbaCacheEntry.cpp
--- a/Engine/Source/Programs/UnrealBuildAccelerator/Common/Private/UbaCacheEntry.cpp
+++ b/Engine/Source/Programs/UnrealBuildAccelerator/Common/Private/UbaCacheEntry.cpp
@@ -309,23 +309,34 @@ namespace uba
 
 			if (!sharedReader.GetLeft())
 			{
+				bool extraWritten = false;
+
+				// Handle all the extra offsets that are lower than shared offset
 				if (previousWasExtra)
 				{
+					extraWritten = true;
 					while (inputsIt != inputsEnd)
 					{
 						offset = *inputsIt++;
 						if (offset >= sharedOffset)
+						{
+							extraWritten = false;
 							break;
+						}
 						extraWriter.Write7BitEncoded(offset);
 					}
 				}
 
 				// Add current range if there is one going
 				if (inRange)
+				{
 					writeRange(rangeBegin, offset == sharedOffset ? sharedPos : lastSharedPos);
+					if (previousWasExtra && offset > sharedOffset)
+						extraWriter.Write7BitEncoded(offset);
+				}
 				else if (offset == sharedOffset) // We want to use a range just to make it easier to debug
 					writeRange(lastSharedPos, sharedPos);
-				else
+				else if (!extraWritten)
 					extraWriter.Write7BitEncoded(offset);
 
 				// Populate rest in extraInputCasKeyOffsets
diff --git a/Engine/Source/Programs/UnrealBuildAccelerator/Test/Private/UbaTestCache.cpp b/Engine/Source/Programs/UnrealBuildAccelerator/Test/Private/UbaTestCache.cpp
--- a/Engine/Source/Programs/UnrealBuildAccelerator/Test/Private/UbaTestCache.cpp
+++ b/Engine/Source/Programs/UnrealBuildAccelerator/Test/Private/UbaTestCache.cpp
@@ -28,7 +28,7 @@ namespace uba
 			u32 expectedExtraCount;
 		};
 
-		auto CheckEntry = [&](CacheEntry& entry, const Record* rec, bool checkOffsets)
+		auto CheckEntry = [&](CacheEntry& entry, const Record* rec, bool checkOffsets, bool checkExpectedCounts)
 			{
 				Set<u32> sharedOffsets;
 				{
@@ -67,10 +67,14 @@ namespace uba
 
 				if (rec)
 				{
-					if (rangeCount != rec->expectedRangeCount)
-						return false;
-					if (extraCount != rec->expectedExtraCount)
-						return false;
+					if (checkExpectedCounts)
+					{
+						if (rangeCount != rec->expectedRangeCount)
+							return false;
+						if (extraCount != rec->expectedExtraCount)
+							return false;
+					}
+
 					if (checkOffsets)
 					{
 						if (rec->offsets != entryOffsets)
@@ -85,7 +89,7 @@ namespace uba
 				return true;
 			};
 
-		auto AddEntry = [&](const Record& rec)
+		auto AddEntry = [&](const Record& rec, bool checkExpectedCount)
 			{
 				Vector<u8> inputOffsets;
 				u64 bytes = 0;
@@ -99,7 +103,7 @@ namespace uba
 				CacheEntry entry;
 				entries.BuildInputs(entry, rec.offsets);
 				entries.entries.push_back(entry);
-				return CheckEntry(entry, &rec, true);
+				return CheckEntry(entry, &rec, true, checkExpectedCount);
 			};
 
 		auto UpdateEntries = [&](u32 multiplier)
@@ -134,10 +138,16 @@ namespace uba
 		fa.Close();
 		UpdateEntries(0);
 		for (auto& entry : entries.entries)
-			CHECK_TRUE(CheckEntry(entry, nullptr, false));
+			CHECK_TRUE(CheckEntry(entry, nullptr, false, false));
 		#endif
 
-		Record records0[] = 
+		Vector<Record> records0 = 
+		{
+			{ { 1 }, 1, 0 },
+			{ { 0 }, 0, 1 },
+		};
+
+		Vector<Record> records1 = 
 		{
 			{ { 1, 4 }, 1, 0 },
 			{ { 1, 2, 3, 4 }, 1, 2 },
@@ -145,18 +155,7 @@ namespace uba
 			{ { 1, 2, 3, 4, 5, 6 }, 1, 4 },
 		};
 
-		Clear();
-		for (auto& rec : records0)
-			CHECK_TRUE(AddEntry(rec));
-		for (u32 i=0; i!=4; ++i)
-		{
-			UpdateEntries(i);
-			u32 index = 0;
-			for (auto& entry : entries.entries)
-				CHECK_TRUE(CheckEntry(entry, &records0[index++], i < 2));
-		}
-
-		Record records1[] = 
+		Vector<Record> records2 = 
 		{
 			{ { 1, 4, 6 }, 1, 0 },
 			{ { 0, 4, 6 }, 1, 1 },
@@ -196,18 +195,7 @@ namespace uba
 			{ { 7 }, 0, 1 },
 		};
 
-		Clear();
-		for (auto& rec : records1)
-			CHECK_TRUE(AddEntry(rec));
-		for (u32 i=0; i!=4; ++i)
-		{
-			UpdateEntries(i);
-			u32 index = 0;
-			for (auto& entry : entries.entries)
-				CHECK_TRUE(CheckEntry(entry, &records1[index++], i < 2));
-		}
-
-		Record records2[] = 
+		Vector<Record> records3 = 
 		{
 			{ { 2, 4, 6, 10, 14, 18 }, 1, 0 },
 			{ { 2, 4, 6, 10, 14, 18 }, 1, 0 },
@@ -225,15 +213,59 @@ namespace uba
 			{ { 4, 7, 14, 18 }, 2, 1 },
 		};
 
-		Clear();
-		for (auto& rec : records2)
-			CHECK_TRUE(AddEntry(rec));
-		for (u32 i=0; i!=4; ++i)
+		Vector<Record> records4 = 
 		{
-			UpdateEntries(i);
-			u32 index = 0;
-			for (auto& entry : entries.entries)
-				CHECK_TRUE(CheckEntry(entry, &records2[index++], i < 2));
+			{ { 1, 4, 7 }, 1, 0 },
+			{ { 1, 5, 6 }, 1, 2 },
+		};
+
+		Vector<Record> records5 = 
+		{
+			{ { 1, 3, 6 }, 1, 0 },
+			{ { 1, 3, 5, 7 }, 1, 2 },
+		};
+
+		Vector<Record>* recordGroups[] =
+		{
+			&records0,
+			&records1,
+			&records2,
+			&records3,
+			&records4,
+			&records5,
+		};
+
+		for (auto recordsPtr : recordGroups)
+		{
+			auto& records = *recordsPtr;
+			Clear();
+			for (auto& rec : records)
+				CHECK_TRUE(AddEntry(rec, true));
+			for (u32 i=0; i!=4; ++i)
+			{
+				UpdateEntries(i);
+				u32 index = 0;
+				for (auto& entry : entries.entries)
+					CHECK_TRUE(CheckEntry(entry, &records[index++], i < 2, true));
+			}
+		}
+
+		for (auto recordsPtr : recordGroups)
+		{
+			auto& records = *recordsPtr;
+
+			for (u32 i=0;i!=records.size(); ++i)
+			{
+				Clear();
+				CHECK_TRUE(AddEntry(records[i], false));
+
+				for (u32 j=0;j!=records.size(); ++j)
+				{
+					u32 index = (i + j + 1) % records.size();
+					CHECK_TRUE(AddEntry(records[index], false));
+				}
+			}
+
 		}
 
 		return true;
