[UBA]
* Added retry in NtCreateFile when getting STATUS_SHARING_VIOLATION. Noone has reported this issue but I've seen it now and then and the retry seems to always succeed after 1 second... have no idea what could cause this.. antivirus?

diff --git a/Engine/Source/Programs/UnrealBuildAccelerator/Detours/Private/Windows/UbaDetoursFunctionsNtDll.inl b/Engine/Source/Programs/UnrealBuildAccelerator/Detours/Private/Windows/UbaDetoursFunctionsNtDll.inl
--- a/Engine/Source/Programs/UnrealBuildAccelerator/Detours/Private/Windows/UbaDetoursFunctionsNtDll.inl
+++ b/Engine/Source/Programs/UnrealBuildAccelerator/Detours/Private/Windows/UbaDetoursFunctionsNtDll.inl
@@ -34,6 +34,7 @@ u8 GetFileAccessFlags(DWORD desiredAccess, u32 createDisposition)
 	return access;
 }
 
+#if UBA_DEBUG_LOG_ENABLED
 StringBuffer<32> ToString(NTSTATUS s)
 {
 	StringBuffer<32> res;
@@ -45,8 +46,13 @@ StringBuffer<32> ToString(NTSTATUS s)
 		return res.Append(L"STATUS_OBJECT_PATH_NOT_FOUND");
 	if (s == STATUS_INVALID_HANDLE)
 		return res.Append(L"STATUS_INVALID_HANDLE");
+	if (s == STATUS_SHARING_VIOLATION)
+		return res.Append(L"STATUS_SHARING_VIOLATION");
+	if (s == STATUS_ACCESS_DENIED)
+		return res.Append(L"STATUS_ACCESS_DENIED");
 	return res.Appendf(L"Error (0x%x)", u32(s));
 }
+#endif
 
 struct FILE_FS_DEVICE_INFORMATION
 {
@@ -598,11 +604,43 @@ NTSTATUS NTAPI Local_NtCreateFile(bool IsCreateFunc, PHANDLE hFileHandle, ACCESS
 
 	TimerScope ts(g_kernelStats.createFile);
 
-	//if (!Contains(ObjectAttributes->ObjectName->Buffer, L".dll") && !Contains(ObjectAttributes->ObjectName->Buffer, L".mui"))
-	//Rpc_WriteLogf(L"NtCreateFile: %ls", ObjectAttributes->ObjectName->Buffer);
-	auto res = True_NtCreateFile(hFileHandle, DesiredAccess, ObjectAttributes, IoStatusBlock, AllocationSize, FileAttributes, ShareAccess, CreateDisposition, CreateOptions, EaBuffer, EaLength);
-	UBA_ASSERTF(res != STATUS_SUCCESS || u64(*hFileHandle) < DetouredHandleStart - 10000, L"Normal handle range is closing in on detoured. Bump detour range (normal: %llu, detour start: %llu)", u64(*hFileHandle), DetouredHandleStart);
-	return res;
+	constexpr u32 retryCount = 15;
+	u32 retriesLeft = retryCount;
+	while (true)
+	{
+		//if (!Contains(ObjectAttributes->ObjectName->Buffer, L".dll") && !Contains(ObjectAttributes->ObjectName->Buffer, L".mui"))
+		//Rpc_WriteLogf(L"NtCreateFile: %ls", ObjectAttributes->ObjectName->Buffer);
+		auto res = True_NtCreateFile(hFileHandle, DesiredAccess, ObjectAttributes, IoStatusBlock, AllocationSize, FileAttributes, ShareAccess, CreateDisposition, CreateOptions, EaBuffer, EaLength);
+
+		UBA_ASSERTF(res != STATUS_SUCCESS || u64(*hFileHandle) < DetouredHandleStart - 10000, L"Normal handle range is closing in on detoured. Bump detour range (normal: %llu, detour start: %llu)", u64(*hFileHandle), DetouredHandleStart);
+
+		// I have no idea why we get this sometimes when trying to open pch for read after recently being written..
+		// All scenarios I've seen this succeeds after 1 second.
+		// Only theory I have is some antivirus or something. 
+		if (res == STATUS_SHARING_VIOLATION)
+		{
+			if (!--retriesLeft)
+				return res;
+			#if UBA_DEBUG
+			StringBuffer<> b;
+			b.Appendf(L"Got access denied trying to open %.*s. Retrying in one second", ObjectAttributes->ObjectName->Length / 2, ObjectAttributes->ObjectName->Buffer);
+			Rpc_WriteLog(b.data, b.count, true, false);
+			#endif
+			Sleep(1000);
+			continue;
+		}
+
+		#if UBA_DEBUG
+		if (retriesLeft != retryCount)
+		{
+			StringBuffer<> b;
+			b.Appendf(L"SUCCEEDED to open %.*s after %u retries.", ObjectAttributes->ObjectName->Length / 2, ObjectAttributes->ObjectName->Buffer, retryCount - retriesLeft);
+			Rpc_WriteLog(b.data, b.count, true, true);
+		}
+		#endif
+
+		return res;
+	}
 }
 
 NTSTATUS NTAPI Shared_NtCreateFile(bool IsCreateFunc, PHANDLE hFileHandle, ACCESS_MASK DesiredAccess, POBJECT_ATTRIBUTES ObjectAttributes, PIO_STATUS_BLOCK IoStatusBlock, PLARGE_INTEGER AllocationSize, ULONG FileAttributes,
