Fix stack corruption when UBA iterates over command line extracting single arguments. Fixed size 16K buffer was insufficient as some Mac link command lines can reach 32k character. Resolve the issue by stack frame allocating an argument buffer the size of the command line itself.

Thanks Stephane Jacoby for providing the fix

diff --git a/Engine/Source/Programs/UnrealBuildAccelerator/Core/Public/UbaProcessUtils.h b/Engine/Source/Programs/UnrealBuildAccelerator/Core/Public/UbaProcessUtils.h
--- a/Engine/Source/Programs/UnrealBuildAccelerator/Core/Public/UbaProcessUtils.h
+++ b/Engine/Source/Programs/UnrealBuildAccelerator/Core/Public/UbaProcessUtils.h
@@ -20,6 +20,8 @@ namespace uba
 		CharType currentChar = 0;
 		CharType lastChar = 0;
 		bool isBackslashOwned = false;
+		CharType* argWriteBuffer = static_cast<CharType*>( alloca(sizeof(CharType) * (argumentsLen + 1)) );
+
 		for (const CharType* it = arguments; !isEnd; lastChar = currentChar, ++it)
 		{
 			bool pastEnd = u64(it - arguments) == argumentsLen;
@@ -37,8 +39,7 @@ namespace uba
 				if ((pastEnd || *argEnd == '\n') && argStart != argEnd && *(argEnd-1) == '\r')
 					--argEnd;
 
-				CharType arg[16*1024];
-				CharType* argWrite = arg;
+				CharType* argWriteIt = argWriteBuffer;
 				CharType lastChar2 = 0;
 				isBackslashOwned = false;
 				while (argIt != argEnd)
@@ -46,7 +47,7 @@ namespace uba
 					if (*argIt == '\"')
 					{
 						if (lastChar2 == '\\' && !isBackslashOwned)
-							argWrite[-1] = '\"';
+							argWriteIt[-1] = '\"';
 						lastChar2 = 0;
 						++argIt;
 						continue;
@@ -57,15 +58,15 @@ namespace uba
 					else
 						isBackslashOwned = false;
 
-					*argWrite++ = *argIt;
+					*argWriteIt++ = *argIt;
 					lastChar2 = *argIt;
 					++argIt;
 				}
 
-				if (arg != argWrite)
+				if (argWriteBuffer != argWriteIt)
 				{
-					*argWrite = 0;
-					argumentFunc(arg, u32(argWrite - arg));
+					*argWriteIt = 0;
+					argumentFunc(argWriteBuffer, u32(argWriteIt - argWriteBuffer));
 				}
 				isInArg = false;
 				isBackslashOwned = false;
