[UBA]
* Fixed devirtualization logic to create a string that is long enough to fit the devirtualized string

diff --git a/Engine/Source/Programs/UnrealBuildAccelerator/Common/Private/UbaSession.cpp b/Engine/Source/Programs/UnrealBuildAccelerator/Common/Private/UbaSession.cpp
--- a/Engine/Source/Programs/UnrealBuildAccelerator/Common/Private/UbaSession.cpp
+++ b/Engine/Source/Programs/UnrealBuildAccelerator/Common/Private/UbaSession.cpp
@@ -1713,7 +1713,7 @@ namespace uba
 				}
 				hasRoot = true;
 				auto& path = rootsEntry->locals[(*str - RootPaths::RootStartByte)/PathsPerRoot];
-				newSize += path.size() - strLen;
+				newSize += path.size();
 			};
 
 		if (!rootsEntry->roots.NormalizeString<tchar>(m_logger, inOut.data(), inOut.size(), checkString, allowPathsWithoutRoot, hint))
@@ -1723,7 +1723,7 @@ namespace uba
 			return false;
 
 		TString newString;
-		newString.resize(newSize + 1);
+		newString.resize(newSize);
 		tchar* newStringPos = newString.data();
 
 		auto handleString = [&](const tchar* str, u64 strLen, u32 rootPos)
@@ -1753,6 +1753,8 @@ namespace uba
 			};
 
 		rootsEntry->roots.NormalizeString<tchar>(m_logger, inOut.data(), inOut.size(), handleString, allowPathsWithoutRoot, hint);
+
+		UBA_ASSERT(newStringPos == newString.data() + newString.size());
 		inOut = std::move(newString);
 
 		return true;
