[UBA]
* Made directory table ~20% smaller

diff --git a/Engine/Source/Programs/UnrealBuildAccelerator/Common/Private/UbaFile.cpp b/Engine/Source/Programs/UnrealBuildAccelerator/Common/Private/UbaFile.cpp
--- a/Engine/Source/Programs/UnrealBuildAccelerator/Common/Private/UbaFile.cpp
+++ b/Engine/Source/Programs/UnrealBuildAccelerator/Common/Private/UbaFile.cpp
@@ -698,15 +698,6 @@ namespace uba
 #endif
 	}
 
-	bool IsDirectory(u32 attributes)
-	{
-#if PLATFORM_WINDOWS
-		return (attributes & FILE_ATTRIBUTE_DIRECTORY) != 0;
-#else
-		return S_ISDIR(attributes);
-#endif
-	}
-
 	bool IsReadOnly(u32 attributes)
 	{
 #if PLATFORM_WINDOWS
diff --git a/Engine/Source/Programs/UnrealBuildAccelerator/Common/Private/UbaProcess.cpp b/Engine/Source/Programs/UnrealBuildAccelerator/Common/Private/UbaProcess.cpp
--- a/Engine/Source/Programs/UnrealBuildAccelerator/Common/Private/UbaProcess.cpp
+++ b/Engine/Source/Programs/UnrealBuildAccelerator/Common/Private/UbaProcess.cpp
@@ -814,7 +814,7 @@ namespace uba
 		CloseFileMessage msg { *this };
 		reader.ReadString(msg.fileName);
 		msg.closeId = reader.ReadU32();
-		msg.attributes = DefaultAttributes(); // reader.ReadU32(); TODO
+		msg.attributes = DefaultAttributes(); // reader.ReadFileAttributes(); TODO
 		msg.deleteOnClose = reader.ReadBool();
 		msg.success = reader.ReadBool();
 		msg.mappingHandle = reader.ReadU64();
diff --git a/Engine/Source/Programs/UnrealBuildAccelerator/Common/Private/UbaSession.cpp b/Engine/Source/Programs/UnrealBuildAccelerator/Common/Private/UbaSession.cpp
--- a/Engine/Source/Programs/UnrealBuildAccelerator/Common/Private/UbaSession.cpp
+++ b/Engine/Source/Programs/UnrealBuildAccelerator/Common/Private/UbaSession.cpp
@@ -320,11 +320,14 @@ namespace uba
 						id = ++m_fileIndexCounter;
 
 					res.first->second = u32(memoryWriter.GetPosition()); // Temporary offset that will be used further down to calculate the real offset
-					memoryWriter.WriteU64(e.lastWritten);
-					memoryWriter.WriteU32(e.attributes);
-					memoryWriter.WriteU32(e.volumeSerial);
-					memoryWriter.WriteU64(id);
-					memoryWriter.WriteU64(e.size);
+					memoryWriter.WriteFileAttributes(e.attributes);
+					memoryWriter.WriteVolumeSerial(e.volumeSerial);
+					memoryWriter.WriteFileIndex(id);
+					if (!IsDirectory(e.attributes))
+					{
+						memoryWriter.WriteFileTime(e.lastWritten);
+						memoryWriter.WriteFileSize(e.size);
+					}
 
 					FileEntryAdded(res.first->first, e.lastWritten, e.size);
 
@@ -383,15 +386,15 @@ namespace uba
 		}
 		else
 		{
-			storageSize += sizeof(u32)*2 + sizeof(u64);
+			storageSize += Get7BitEncodedCount(dirAttributes) + sizeof(volumeSerial) + sizeof(fileIndex);
 			tableWriter.Write7BitEncoded(storageSize);
 			tableWriter.WriteStringKey(dirKey);
 			tableOffset = writePos + u32(tableWriter.GetPosition());
 			tableWriter.Write7BitEncoded(dir.tableOffset);
 			UBA_ASSERT(dirAttributes != 0);
-			tableWriter.WriteU32(dirAttributes);
-			tableWriter.WriteU32(volumeSerial);
-			tableWriter.WriteU64(fileIndex);
+			tableWriter.WriteFileAttributes(dirAttributes);
+			tableWriter.WriteVolumeSerial(volumeSerial);
+			tableWriter.WriteFileIndex(fileIndex);
 		}
 
 		tableWriter.Write7BitEncoded(itemCount);
@@ -853,18 +856,23 @@ namespace uba
 		if (!insres.second && insres.first->second != ~u32(0))
 		{
 			BinaryReader reader(m_directoryTableMem + insres.first->second);
-			u64 oldLastWriteTime = reader.ReadU64();
+			u32 oldAttr = reader.ReadFileAttributes();(void)oldAttr;
+
+			if (isDirectory) // Ignore updating directories.. they should always be the same regardless
+			{
+				UBA_ASSERT(IsDirectory(oldAttr));
+				return true;
+			}
+			reader.ReadVolumeSerial();
+			u64 oldFileIndex = reader.ReadFileIndex();(void)oldFileIndex;
 
-			if (lastWriteTime == oldLastWriteTime || isDirectory) // Ignore updating directories.. they should always be the same regardless
+			u64 oldLastWriteTime = reader.ReadFileTime();
+			if (lastWriteTime == oldLastWriteTime)
 			{
-				reader.Skip(sizeof(u32) * 2);
 				#if !PLATFORM_WINDOWS
-				u64 oldFileIndex = reader.ReadU64();
 				UBA_ASSERT(oldFileIndex == fileIndex); // Checking so it is really the same file
-				#else
-				reader.Skip(sizeof(u64));
 				#endif
-				u64 oldSize = reader.ReadU64();
+				u64 oldSize = reader.ReadFileSize();
 				if (oldSize == fileSize)
 					return true;
 				// TODO: Somehow this can happen and I have no idea how. last written time should be set on close file so it shouldnt be possible.
@@ -890,11 +898,14 @@ namespace uba
 			writer.Write7BitEncoded(1); // Count
 			writer.WriteString(lastSlash + 1);
 			entryPos = writer.GetPosition();
-			writer.WriteU64(lastWriteTime);
-			writer.WriteU32(attributes);
-			writer.WriteU32(volumeSerial);
-			writer.WriteU64(fileIndex);
-			writer.WriteU64(fileSize);
+			writer.WriteFileAttributes(attributes);
+			writer.WriteVolumeSerial(volumeSerial);
+			writer.WriteFileIndex(fileIndex);
+			if (!isDirectory)
+			{
+				writer.WriteFileTime(lastWriteTime);
+				writer.WriteFileSize(fileSize);
+			}
 			written = writer.GetPosition();
 		}
 
@@ -967,11 +978,14 @@ namespace uba
 			writer.Write7BitEncoded(dir.tableOffset); // Previous entry for same directory
 			writer.Write7BitEncoded(1); // Count
 			writer.WriteString(lastSlash + 1);
-			writer.WriteU32(0);
-			writer.WriteU32(0);
-			writer.WriteU64(0);
-			writer.WriteU64(0);
-			writer.WriteU64(0);
+			writer.WriteFileAttributes(0);
+			writer.WriteVolumeSerial(0);
+			writer.WriteFileIndex(0);
+			if (true) // !IsDirectory()
+			{
+				writer.WriteFileTime(0);
+				writer.WriteFileSize(0);
+			}
 			written = writer.GetPosition();
 		}
 
diff --git a/Engine/Source/Programs/UnrealBuildAccelerator/Common/Public/UbaFile.h b/Engine/Source/Programs/UnrealBuildAccelerator/Common/Public/UbaFile.h
--- a/Engine/Source/Programs/UnrealBuildAccelerator/Common/Public/UbaFile.h
+++ b/Engine/Source/Programs/UnrealBuildAccelerator/Common/Public/UbaFile.h
@@ -67,7 +67,6 @@ namespace uba
 	bool GetFileSizeEx(u64& outFileSize, FileHandle hFile);
 
 	u32 GetFileAttributesW(const tchar* fileName);
-	bool IsDirectory(u32 attributes);
 	bool IsReadOnly(u32 attributes);
 	u32 DefaultAttributes(bool execute = false);
 	bool CreateHardLinkW(const tchar* newFileName, const tchar* existingFileName);
diff --git a/Engine/Source/Programs/UnrealBuildAccelerator/Core/Public/UbaBinaryReaderWriter.h b/Engine/Source/Programs/UnrealBuildAccelerator/Core/Public/UbaBinaryReaderWriter.h
--- a/Engine/Source/Programs/UnrealBuildAccelerator/Core/Public/UbaBinaryReaderWriter.h
+++ b/Engine/Source/Programs/UnrealBuildAccelerator/Core/Public/UbaBinaryReaderWriter.h
@@ -31,6 +31,12 @@ namespace uba
 
 		void WriteLongString(const StringView& str, u32 threshold = 10000);
 
+		inline void WriteFileTime(u64 fileTime) { WriteU64(fileTime); }
+		inline void WriteFileAttributes(u32 attributes) { Write7BitEncoded(attributes); }
+		inline void WriteVolumeSerial(u32 volumeSerial) { WriteU32(volumeSerial); }
+		inline void WriteFileIndex(u64 fileIndex) { WriteU64(fileIndex); }
+		inline void WriteFileSize(u64 fileSize) { Write7BitEncoded(fileSize); }
+
 		inline u8* AllocWrite(u64 bytes);
 		inline u64 GetPosition() const { return u64(m_pos - m_begin); }
 		inline u64 GetCapacityLeft() const { return u64(m_end - m_pos); }
@@ -72,6 +78,13 @@ namespace uba
 		inline u64 Read7BitEncoded();
 		template<typename CharType> inline CharType ReadUtf8Char();
 		TString ReadLongString();
+
+		inline u64 ReadFileTime() { return ReadU64(); }
+		inline u32 ReadFileAttributes() { return u32(Read7BitEncoded()); }
+		inline u32 ReadVolumeSerial() { return ReadU32(); }
+		inline u64 ReadFileIndex() { return ReadU64(); }
+		inline u64 ReadFileSize() { return Read7BitEncoded(); }
+
 		inline bool TryRead7BitEncoded(u64& outValue);
 		inline u32 PeekU32();
 		inline u64 PeekU64();
@@ -283,7 +296,7 @@ namespace uba
 		} while (value > 0);
 	}
 
-	inline u8* BinaryWriter::AllocWrite(u64 bytes)
+	u8* BinaryWriter::AllocWrite(u64 bytes)
 	{
 		UBA_ASSERT_WRITE(bytes);
 		u8* data = m_pos;
diff --git a/Engine/Source/Programs/UnrealBuildAccelerator/Core/Public/UbaDirectoryTable.h b/Engine/Source/Programs/UnrealBuildAccelerator/Core/Public/UbaDirectoryTable.h
--- a/Engine/Source/Programs/UnrealBuildAccelerator/Core/Public/UbaDirectoryTable.h
+++ b/Engine/Source/Programs/UnrealBuildAccelerator/Core/Public/UbaDirectoryTable.h
@@ -96,10 +96,11 @@ namespace uba
 				if (firstIsRoot)
 				{
 					firstIsRoot = false;
-					u32 attr = reader.ReadU32();
+					u32 attr = reader.ReadFileAttributes();
 					if (!attr)
 						continue;
-					reader.Skip(sizeof(u32) + sizeof(u64));
+					reader.ReadVolumeSerial(); // Directory volume serial
+					reader.ReadFileIndex(); // Directory file index
 				}
 
 				PopulateDirectoryWithFiles(reader, hasher, files);
@@ -125,7 +126,14 @@ namespace uba
 				reader.ReadString(filename);
 				if (CaseInsensitiveFs)
 					filename.MakeLower();
-				reader.Skip(sizeof(u32)*2 + sizeof(u64)*3);
+				u32 attr = reader.ReadFileAttributes();
+				reader.ReadVolumeSerial();
+				reader.ReadFileIndex();
+				if (!IsDirectory(attr))
+				{
+					reader.ReadFileTime();
+					reader.ReadFileSize();
+				}
 
 				StringKey filenameKey = ToStringKey(hasher, filename.data, filename.count);
 				files[filenameKey] = offset; // Always write, since same file might have been added with new info
@@ -235,11 +243,11 @@ namespace uba
 					reader.SetPosition(prevTableOffset);
 					prevTableOffset = reader.Read7BitEncoded();
 				}
-				outInfo.attributes = reader.ReadU32();
+				outInfo.attributes = reader.ReadFileAttributes();
 				if (outInfo.attributes)
 				{
-					outInfo.volumeSerial = reader.ReadU32();
-					outInfo.fileIndex = reader.ReadU64();
+					outInfo.volumeSerial = reader.ReadVolumeSerial();
+					outInfo.fileIndex = reader.ReadFileIndex();
 				}
 				outInfo.size = 0;
 				outInfo.lastWrite = 0;
@@ -252,11 +260,19 @@ namespace uba
 				reader.ReadString(outFileName, fileNameCapacity);
 			else
 				reader.SkipString();
-			outInfo.lastWrite = reader.ReadU64();
-			outInfo.attributes = reader.ReadU32();
-			outInfo.volumeSerial = reader.ReadU32();
-			outInfo.fileIndex = reader.ReadU64();
-			outInfo.size = reader.ReadU64();
+			outInfo.attributes = reader.ReadFileAttributes();
+			outInfo.volumeSerial = reader.ReadVolumeSerial();
+			outInfo.fileIndex = reader.ReadFileIndex();
+			if (IsDirectory(outInfo.attributes))
+			{
+				outInfo.size = 0;
+				outInfo.lastWrite = 0;
+			}
+			else
+			{
+				outInfo.lastWrite = reader.ReadFileTime();
+				outInfo.size = reader.ReadFileSize();
+			}
 			return u32(reader.GetPosition());
 		}
 
diff --git a/Engine/Source/Programs/UnrealBuildAccelerator/Core/Public/UbaPlatform.h b/Engine/Source/Programs/UnrealBuildAccelerator/Core/Public/UbaPlatform.h
--- a/Engine/Source/Programs/UnrealBuildAccelerator/Core/Public/UbaPlatform.h
+++ b/Engine/Source/Programs/UnrealBuildAccelerator/Core/Public/UbaPlatform.h
@@ -129,6 +129,7 @@ namespace uba
 	inline LONG ToHigh(u64 v) { LARGE_INTEGER li; li.QuadPart = (LONGLONG)v; return li.HighPart; }
 	inline LARGE_INTEGER ToLargeInteger(u64 v) { LARGE_INTEGER li; li.QuadPart = (LONGLONG)v; return li; }
 	inline LARGE_INTEGER ToLargeInteger(u32 high, u32 low) { LARGE_INTEGER li; li.HighPart = (LONG)high; li.LowPart = low; return li; }
+	inline bool IsDirectory(u32 attributes) { return (attributes & FILE_ATTRIBUTE_DIRECTORY) != 0; }
 	#define TStrlen(s) u32(wcslen(s))
 	#define TStrchr(a, b) wcschr(a, b)
 	#define TStrrchr(a, b) wcsrchr(a, b)
@@ -161,6 +162,7 @@ namespace uba
 	bool GetPhysicallyInstalledSystemMemory(u64& outKb);
 	inline u64 FromTimeSpec(const timespec& ts) { return u64(ts.tv_sec) * 10'000'000ull + u64(ts.tv_nsec/100); }
 	inline timespec ToTimeSpec(u64 time) { timespec ts; ts.tv_sec = time / 10'000'000ull; ts.tv_nsec = (time - (u64(ts.tv_sec) * 10'000'000ull)) * 100; return ts; }
+	inline bool IsDirectory(u32 attributes) { return S_ISDIR(attributes); }
 	#if PLATFORM_LINUX
 	inline constexpr bool CaseInsensitiveFs = false;
 	#define st_mtimespec st_mtim
diff --git a/Engine/Source/Programs/UnrealBuildAccelerator/Core/Public/UbaProcessUtils.h b/Engine/Source/Programs/UnrealBuildAccelerator/Core/Public/UbaProcessUtils.h
--- a/Engine/Source/Programs/UnrealBuildAccelerator/Core/Public/UbaProcessUtils.h
+++ b/Engine/Source/Programs/UnrealBuildAccelerator/Core/Public/UbaProcessUtils.h
@@ -4,6 +4,10 @@
 
 #include "UbaPlatform.h"
 
+#if PLATFORM_WINDOWS
+#include <malloc.h>
+#endif
+
 namespace uba
 {
 	// func is void(const CharType* arg, u32 strArg)
