[UBA]
* Fixed bug where WaitEntry (when waiting for cas files to be transferred) event was auto reset. This would meant that if two actions waited for the same cas, then one would timeout

diff --git a/Engine/Source/Programs/UnrealBuildAccelerator/Common/Public/UbaStorageServer.h b/Engine/Source/Programs/UnrealBuildAccelerator/Common/Public/UbaStorageServer.h
--- a/Engine/Source/Programs/UnrealBuildAccelerator/Common/Public/UbaStorageServer.h
+++ b/Engine/Source/Programs/UnrealBuildAccelerator/Common/Public/UbaStorageServer.h
@@ -63,7 +63,7 @@ namespace uba
 
 		struct WaitEntry
 		{
-			WaitEntry() : done(false) {}
+			WaitEntry() : done(true) {}
 			Event done;
 			bool Success = false;
 			u32 refCount = 0;
