[UBA]
* Changed a bunch of gobal variables to be allocated on heap and leaked for shutdown. This resolves an issue with SentinelOne creating files during global deinitialization

diff --git a/Engine/Source/Programs/UnrealBuildAccelerator/Detours/Private/Windows/UbaDetoursFunctionsKernelBase.inl b/Engine/Source/Programs/UnrealBuildAccelerator/Detours/Private/Windows/UbaDetoursFunctionsKernelBase.inl
--- a/Engine/Source/Programs/UnrealBuildAccelerator/Detours/Private/Windows/UbaDetoursFunctionsKernelBase.inl
+++ b/Engine/Source/Programs/UnrealBuildAccelerator/Detours/Private/Windows/UbaDetoursFunctionsKernelBase.inl
@@ -605,6 +605,9 @@ DWORD Detoured_GetSecurityInfo(HANDLE handle, SE_OBJECT_TYPE ObjectType, SECURIT
 	return True_GetSecurityInfo(handle, ObjectType, SecurityInfo, ppsidOwner, ppsidGroup, ppDacl, ppSacl, ppSecurityDescriptor);
 }
 
+StringBuffer<32 * 1024> g_stdFile;
+ReaderWriterLock& g_stdFileLock = *new ReaderWriterLock();
+
 void WriteStdFile(LPCVOID lpBuffer, DWORD nNumberOfBytesToWrite, bool isError)
 {
 	if (!g_echoOn || g_suppressLogging)
diff --git a/Engine/Source/Programs/UnrealBuildAccelerator/Detours/Private/Windows/UbaDetoursFunctionsWin.cpp b/Engine/Source/Programs/UnrealBuildAccelerator/Detours/Private/Windows/UbaDetoursFunctionsWin.cpp
--- a/Engine/Source/Programs/UnrealBuildAccelerator/Detours/Private/Windows/UbaDetoursFunctionsWin.cpp
+++ b/Engine/Source/Programs/UnrealBuildAccelerator/Detours/Private/Windows/UbaDetoursFunctionsWin.cpp
@@ -216,18 +216,18 @@ constexpr u64 ListDirHandlesRange = 4*1024*1024;
 MemoryBlock& g_listDirHandleMemoryBlock = *new MemoryBlock(ListDirHandlesRange);
 constexpr u64 ListDirHandlesStart = DetouredHandleEnd;
 constexpr u64 ListDirHandlesEnd = ListDirHandlesStart + ListDirHandlesRange/sizeof(ListDirectoryHandle);
-BlockAllocator<ListDirectoryHandle> g_listDirectoryHandleAllocator(g_listDirHandleMemoryBlock);
+BlockAllocator<ListDirectoryHandle>& g_listDirectoryHandleAllocator = *new BlockAllocator<ListDirectoryHandle>(g_listDirHandleMemoryBlock);
 void* ListDirectoryHandle::operator new(size_t size) { return g_listDirectoryHandleAllocator.Allocate(); }
 void ListDirectoryHandle::operator delete(void* p) { g_listDirectoryHandleAllocator.Free(p); }
 inline bool isListDirectoryHandle(HANDLE h) { return (u64)h >= ListDirHandlesStart && (u64)h < ListDirHandlesEnd; }
 inline HANDLE makeListDirectoryHandle(ListDirectoryHandle* p) { return HANDLE(ListDirHandlesStart + (p - (ListDirectoryHandle*)g_listDirHandleMemoryBlock.memory)); }
 inline ListDirectoryHandle& asListDirectoryHandle(HANDLE h) { return *(((ListDirectoryHandle*)g_listDirHandleMemoryBlock.memory) + (u64(h) - ListDirHandlesStart)); }
 
-ReaderWriterLock g_loadedModulesLock;
-UnorderedMap<HMODULE, TString> g_loadedModules;
+ReaderWriterLock& g_loadedModulesLock = *new ReaderWriterLock();
+UnorderedMap<HMODULE, TString>& g_loadedModules = *new UnorderedMap<HMODULE, TString>();
 u64 g_memoryFileIndexCounter = ~u64(0) - 1000000; // I really hope this will not collide with anything
 
-FileDecompressor g_fileDecompressor;
+FileDecompressor& g_fileDecompressor = *new FileDecompressor(); // Leak for shutdown
 
 
 struct SuppressCreateFileDetourScope
@@ -244,7 +244,7 @@ const wchar_t* HandleToName(const DetouredHandle& dh)
 	return L"Unknown";
 }
 
-BlockAllocator<FileObject> g_fileObjectAllocator(g_memoryBlock);
+BlockAllocator<FileObject>& g_fileObjectAllocator = *new BlockAllocator<FileObject>(g_memoryBlock);
 
 void MemoryFile::Write(DetouredHandle& handle, LPCVOID lpBuffer, u64 nNumberOfBytesToWrite)
 {
@@ -432,8 +432,9 @@ bool EnsureMapped(DetouredHandle& handle, DWORD dwFileOffsetHigh = 0, DWORD dwFi
 	return true;
 }
 
-ReaderWriterLock g_longPathNameCacheLock;
-GrowingUnorderedMap<const wchar_t*, const wchar_t*, HashString, EqualString> g_longPathNameCache(g_memoryBlock);
+ReaderWriterLock& g_longPathNameCacheLock = *new ReaderWriterLock();
+using LongPathMap = GrowingUnorderedMap<const wchar_t*, const wchar_t*, HashString, EqualString>;
+LongPathMap& g_longPathNameCache = *new LongPathMap(g_memoryBlock);
 
 void Rpc_AllocFailed(const wchar_t* allocType, u32 error)
 {
@@ -447,9 +448,6 @@ void Rpc_AllocFailed(const wchar_t* allocType, u32 error)
 	Sleep(5*1000);
 }
 
-StringBuffer<32 * 1024> g_stdFile;
-ReaderWriterLock g_stdFileLock;
-
 void CloseCaches()
 {
 	for (auto& it : g_mappedFileTable.m_lookup)
@@ -989,7 +987,7 @@ void Init(const DetoursPayload& payload, u64 startTime)
 
 	if (payload.storeObjFilesCompressed && g_rules->ShouldDecompressFiles(StringView()))
 	{
-		bool preprocess = g_rules->index == SpecialRulesIndex_LdExePlatform1;
+		bool preprocess = true;//g_rules->index == SpecialRulesIndex_LdExePlatform1;
 		g_fileDecompressor.Start(cmdLine, preprocess);
 	}
 }
diff --git a/Engine/Source/Programs/UnrealBuildAccelerator/Detours/Private/Windows/UbaDetoursUtilsWin.h b/Engine/Source/Programs/UnrealBuildAccelerator/Detours/Private/Windows/UbaDetoursUtilsWin.h
--- a/Engine/Source/Programs/UnrealBuildAccelerator/Detours/Private/Windows/UbaDetoursUtilsWin.h
+++ b/Engine/Source/Programs/UnrealBuildAccelerator/Detours/Private/Windows/UbaDetoursUtilsWin.h
@@ -24,7 +24,7 @@ namespace uba
 
 		TString newName;
 	};
-	extern BlockAllocator<FileObject> g_fileObjectAllocator;
+	extern BlockAllocator<FileObject>& g_fileObjectAllocator;
 	inline void* FileObject::operator new(size_t size) { return g_fileObjectAllocator.Allocate(); }
 	inline void FileObject::operator delete(void* p) { g_fileObjectAllocator.Free(p); }
 
