[UBA]
* Fixed so application rules for clang-cl handles .obj file as output file. This is an optimization for when -fno-temp-file option is added

diff --git a/Engine/Source/Programs/UnrealBuildAccelerator/Core/Private/UbaApplicationRules.cpp b/Engine/Source/Programs/UnrealBuildAccelerator/Core/Private/UbaApplicationRules.cpp
--- a/Engine/Source/Programs/UnrealBuildAccelerator/Core/Private/UbaApplicationRules.cpp
+++ b/Engine/Source/Programs/UnrealBuildAccelerator/Core/Private/UbaApplicationRules.cpp
@@ -255,6 +255,7 @@ namespace uba
 				|| fileName.EndsWith(TC(".cpp.d"))
 				|| fileName.EndsWith(TC(".o.tmp")) // Clang writes to tmp file and then move
 				|| fileName.EndsWith(TC(".obj.tmp")) // Clang (verse) writes to tmp file and then move
+				|| fileName.EndsWith(TC(".obj")) // Clang with -fno-temp-file write obj file directly
 				;// || fileName.EndsWith(TC(".gch.tmp")); // Need to fix "has been modified since the precompiled header"
 		}
 
