[UBA]
* Added volume serial cache that can be used to reduce directory table size

diff --git a/Engine/Source/Programs/UnrealBuildAccelerator/Common/Private/UbaFile.cpp b/Engine/Source/Programs/UnrealBuildAccelerator/Common/Private/UbaFile.cpp
--- a/Engine/Source/Programs/UnrealBuildAccelerator/Common/Private/UbaFile.cpp
+++ b/Engine/Source/Programs/UnrealBuildAccelerator/Common/Private/UbaFile.cpp
@@ -921,4 +921,57 @@ namespace uba
 			out.Append(isWindowsArm ? TC("arm64") : TC("x64")).Append(PathSeparator);
 		return true;
 	};
+
+	bool VolumeCache::Init(Logger& logger)
+	{
+		#if PLATFORM_WINDOWS
+		tchar rootPathName[MAX_PATH] = {0};
+		HANDLE hFind = FindFirstVolumeW(rootPathName, ARRAYSIZE(rootPathName));
+
+		if (hFind == INVALID_HANDLE_VALUE)
+			return logger.Error(TC("FindFirstVolume failed (%s)"), LastErrorToText().data);
+
+		do
+		{
+			DWORD serialNumber = 0;
+			if (GetVolumeInformationW(rootPathName, nullptr, 0, &serialNumber, nullptr, nullptr, nullptr, 0))
+				volumeSerials.push_back(serialNumber);
+		}
+		while (FindNextVolumeW(hFind, rootPathName, ARRAYSIZE(rootPathName)));
+
+		FindVolumeClose(hFind);
+
+		volumeSerials.push_back(0);
+		#endif
+		return true;
+	}
+
+	u32 VolumeCache::GetSerialIndex(u32 volumeSerial)
+	{
+		#if PLATFORM_WINDOWS
+		u32 index = 1;
+		for (auto s : volumeSerials)
+		{
+			if (s == volumeSerial)
+				return index;
+			++index;
+		}
+		#endif
+		return volumeSerial;
+	}
+
+	void VolumeCache::Write(BinaryWriter& writer)
+	{
+		writer.WriteU16(u16(volumeSerials.size()));
+		for (u32 serial : volumeSerials)
+			writer.WriteU32(serial);
+	}
+
+	void VolumeCache::Read(BinaryReader& reader)
+	{
+		u32 count = reader.ReadU16();
+		volumeSerials.resize(count);
+		for (u32 i=0; i!=count;++i)
+			volumeSerials[i] = reader.ReadU32();
+	}
 }
diff --git a/Engine/Source/Programs/UnrealBuildAccelerator/Common/Public/UbaFile.h b/Engine/Source/Programs/UnrealBuildAccelerator/Common/Public/UbaFile.h
--- a/Engine/Source/Programs/UnrealBuildAccelerator/Common/Public/UbaFile.h
+++ b/Engine/Source/Programs/UnrealBuildAccelerator/Common/Public/UbaFile.h
@@ -137,4 +137,13 @@ namespace uba
 				return false;
 		return true;
 	}
+
+	struct VolumeCache
+	{
+		bool Init(Logger& logger);
+		u32 GetSerialIndex(u32 volumeSerial);
+		void Write(BinaryWriter& writer);
+		void Read(BinaryReader& reader);
+		Vector<u32> volumeSerials;
+	};
 }
diff --git a/Engine/Source/Programs/UnrealBuildAccelerator/Test/Private/UbaTestAll.h b/Engine/Source/Programs/UnrealBuildAccelerator/Test/Private/UbaTestAll.h
--- a/Engine/Source/Programs/UnrealBuildAccelerator/Test/Private/UbaTestAll.h
+++ b/Engine/Source/Programs/UnrealBuildAccelerator/Test/Private/UbaTestAll.h
@@ -68,6 +68,7 @@ namespace uba
 		UBA_TEST(TestCustomService) \
 		UBA_TEST(TestStdOutLocal) \
 		UBA_TEST(TestStdOutViaCmd) \
+		UBA_TEST(TestVolumeCache) \
 
 
 #define UBA_LINUX_TESTS \
diff --git a/Engine/Source/Programs/UnrealBuildAccelerator/Test/Private/UbaTestBasics.cpp b/Engine/Source/Programs/UnrealBuildAccelerator/Test/Private/UbaTestBasics.cpp
--- a/Engine/Source/Programs/UnrealBuildAccelerator/Test/Private/UbaTestBasics.cpp
+++ b/Engine/Source/Programs/UnrealBuildAccelerator/Test/Private/UbaTestBasics.cpp
@@ -510,4 +510,14 @@ namespace uba
 #endif
 		return true;
 	}
+
+	bool TestVolumeCache(LoggerWithWriter& logger, const StringBufferBase& rootDir)
+	{
+		VolumeCache cache;
+		if (!cache.Init(logger))
+			return false;
+		return true;
+	}
+
+
 }
