[UBA]
* Enabled fake volume serials on windows. Can be disabled with UseFakeVolumeSerial = false in session config

diff --git a/Engine/Source/Programs/UnrealBuildAccelerator/Common/Private/UbaSession.cpp b/Engine/Source/Programs/UnrealBuildAccelerator/Common/Private/UbaSession.cpp
--- a/Engine/Source/Programs/UnrealBuildAccelerator/Common/Private/UbaSession.cpp
+++ b/Engine/Source/Programs/UnrealBuildAccelerator/Common/Private/UbaSession.cpp
@@ -229,6 +229,7 @@ namespace uba
 		table.GetValueAsBool(storeObjFilesCompressed, TC("StoreObjFilesCompressed"));
 		table.GetValueAsBool(allowLocalDetour, TC("AllowLocalDetour"));
 		table.GetValueAsBool(extractObjFilesSymbols, TC("ExtractObjFilesSymbols"));
+		table.GetValueAsBool(useFakeVolumeSerial, TC("UseFakeVolumeSerial"));
 		table.GetValueAsU32(traceReserveSizeMb, TC("TraceReserveSizeMb"));
 	}
 
@@ -250,6 +251,7 @@ namespace uba
 		auto& dirTable = m_directoryTable;
 
 		u32 volumeSerial = 0;
+		u32 volumeSerialIndex = 0;
 		u32 dirAttributes = 0;
 		u64 fileIndex = 0;
 
@@ -321,7 +323,7 @@ namespace uba
 
 					res.first->second = u32(memoryWriter.GetPosition()); // Temporary offset that will be used further down to calculate the real offset
 					memoryWriter.WriteFileAttributes(e.attributes);
-					memoryWriter.WriteVolumeSerial(e.volumeSerial);
+					memoryWriter.WriteVolumeSerial(e.volumeSerial == volumeSerial ? volumeSerialIndex : m_volumeCache.GetSerialIndex(e.volumeSerial));
 					memoryWriter.WriteFileIndex(id);
 					if (!IsDirectory(e.attributes))
 					{
@@ -342,6 +344,7 @@ namespace uba
 				[&](const DirectoryInfo& e)
 				{
 					volumeSerial = e.volumeSerial;
+					volumeSerialIndex = m_volumeCache.GetSerialIndex(volumeSerial);
 					dirAttributes = e.attributes;
 					fileIndex = e.id;
 				});
@@ -386,14 +389,14 @@ namespace uba
 		}
 		else
 		{
-			storageSize += Get7BitEncodedCount(dirAttributes) + sizeof(volumeSerial) + sizeof(fileIndex);
+			storageSize += Get7BitEncodedCount(dirAttributes) + Get7BitEncodedCount(volumeSerialIndex) + sizeof(fileIndex);
 			tableWriter.Write7BitEncoded(storageSize);
 			tableWriter.WriteStringKey(dirKey);
 			tableOffset = writePos + u32(tableWriter.GetPosition());
 			tableWriter.Write7BitEncoded(dir.tableOffset);
 			UBA_ASSERT(dirAttributes != 0);
 			tableWriter.WriteFileAttributes(dirAttributes);
-			tableWriter.WriteVolumeSerial(volumeSerial);
+			tableWriter.WriteVolumeSerial(volumeSerialIndex);
 			tableWriter.WriteFileIndex(fileIndex);
 		}
 
@@ -899,7 +902,7 @@ namespace uba
 			writer.WriteString(lastSlash + 1);
 			entryPos = writer.GetPosition();
 			writer.WriteFileAttributes(attributes);
-			writer.WriteVolumeSerial(volumeSerial);
+			writer.WriteVolumeSerial(m_volumeCache.GetSerialIndex(volumeSerial));
 			writer.WriteFileIndex(fileIndex);
 			if (!isDirectory)
 			{
@@ -1305,6 +1308,10 @@ namespace uba
 				return true;
 			});
 
+		if (info.useFakeVolumeSerial && !m_runningRemote)
+			if (!m_volumeCache.Init(m_logger))
+				return false;
+
 		return true;
 	}
 
diff --git a/Engine/Source/Programs/UnrealBuildAccelerator/Common/Public/UbaNetwork.h b/Engine/Source/Programs/UnrealBuildAccelerator/Common/Public/UbaNetwork.h
--- a/Engine/Source/Programs/UnrealBuildAccelerator/Common/Public/UbaNetwork.h
+++ b/Engine/Source/Programs/UnrealBuildAccelerator/Common/Public/UbaNetwork.h
@@ -16,7 +16,7 @@ namespace uba
 	
 	static constexpr u32 SystemNetworkVersion = 1339;
 	static constexpr u32 StorageNetworkVersion = 4;
-	static constexpr u32 SessionNetworkVersion = 43;
+	static constexpr u32 SessionNetworkVersion = 44;
 	static constexpr u32 CacheNetworkVersion = 5;
 	static constexpr u32 CacheBucketVersion = 3;
 
diff --git a/Engine/Source/Programs/UnrealBuildAccelerator/Common/Public/UbaSession.h b/Engine/Source/Programs/UnrealBuildAccelerator/Common/Public/UbaSession.h
--- a/Engine/Source/Programs/UnrealBuildAccelerator/Common/Public/UbaSession.h
+++ b/Engine/Source/Programs/UnrealBuildAccelerator/Common/Public/UbaSession.h
@@ -304,6 +304,8 @@ namespace uba
 		UnorderedMap<TString, bool> m_isX64Application;
 		#endif
 
+		VolumeCache m_volumeCache;
+
 		friend class ProcessImpl;
 	};
 
diff --git a/Engine/Source/Programs/UnrealBuildAccelerator/Common/Public/UbaSessionCreateInfo.h b/Engine/Source/Programs/UnrealBuildAccelerator/Common/Public/UbaSessionCreateInfo.h
--- a/Engine/Source/Programs/UnrealBuildAccelerator/Common/Public/UbaSessionCreateInfo.h
+++ b/Engine/Source/Programs/UnrealBuildAccelerator/Common/Public/UbaSessionCreateInfo.h
@@ -38,6 +38,7 @@ namespace uba
 		bool allowLocalDetour = true;			// Allow local processes to be detoured. If this is false it is up to outside logic to register all created/deleted files
 		bool extractObjFilesSymbols = false;	// Will extract import/export symbols and clear export table from obj file.
 		bool treatTempDirAsEmpty = true;		// If this is true, directory table will always see temp as empty. Use sub directories of temp instead
+		bool useFakeVolumeSerial = true;		// Fake volume serials reduce the directory table size since serials end up being compressed down to one byte
 		u64 deleteSessionsOlderThanSeconds = 12 * 60 * 60; // Delete session folders older than 12 hours by default . Set to 0 to not delete or 1 to delete all
 		u64 keepOutputFileMemoryMapsThreshold = 256 * 1024; // If allowMemoryMaps is true, output files will be kept in memory if smaller than this size
 		u32 traceReserveSizeMb = 128;			// Memory reserved for trace file (in mb)
diff --git a/Engine/Source/Programs/UnrealBuildAccelerator/Core/Public/UbaBinaryReaderWriter.h b/Engine/Source/Programs/UnrealBuildAccelerator/Core/Public/UbaBinaryReaderWriter.h
--- a/Engine/Source/Programs/UnrealBuildAccelerator/Core/Public/UbaBinaryReaderWriter.h
+++ b/Engine/Source/Programs/UnrealBuildAccelerator/Core/Public/UbaBinaryReaderWriter.h
@@ -33,7 +33,7 @@ namespace uba
 
 		inline void WriteFileTime(u64 fileTime) { WriteU64(fileTime); }
 		inline void WriteFileAttributes(u32 attributes) { Write7BitEncoded(attributes); }
-		inline void WriteVolumeSerial(u32 volumeSerial) { WriteU32(volumeSerial); }
+		inline void WriteVolumeSerial(u32 volumeSerial) { Write7BitEncoded(volumeSerial); }
 		inline void WriteFileIndex(u64 fileIndex) { WriteU64(fileIndex); }
 		inline void WriteFileSize(u64 fileSize) { Write7BitEncoded(fileSize); }
 
@@ -81,7 +81,7 @@ namespace uba
 
 		inline u64 ReadFileTime() { return ReadU64(); }
 		inline u32 ReadFileAttributes() { return u32(Read7BitEncoded()); }
-		inline u32 ReadVolumeSerial() { return ReadU32(); }
+		inline u32 ReadVolumeSerial() { return u32(Read7BitEncoded()); }
 		inline u64 ReadFileIndex() { return ReadU64(); }
 		inline u64 ReadFileSize() { return Read7BitEncoded(); }
 
