[UBA]
* Added support for filtering process boxes in visualizer by typing text
* Fixed crash when looking at workers in visualizer

diff --git a/Engine/Source/Programs/UnrealBuildAccelerator/Visualizer/Private/UbaTraceReader.cpp b/Engine/Source/Programs/UnrealBuildAccelerator/Visualizer/Private/UbaTraceReader.cpp
--- a/Engine/Source/Programs/UnrealBuildAccelerator/Visualizer/Private/UbaTraceReader.cpp
+++ b/Engine/Source/Programs/UnrealBuildAccelerator/Visualizer/Private/UbaTraceReader.cpp
@@ -35,6 +35,8 @@ namespace uba
 		version = 0;
 		sessions.clear();
 		workTracks.clear();
+		for (auto s : strings)
+			free((void*)s);
 		strings.clear();
 		statusMap.clear();
 		cacheWrites.clear();
@@ -766,7 +768,7 @@ namespace uba
 				hint = temp.data;
 			}
 			else
-				hint = out.strings[reader.Read7BitEncoded()].data();
+				hint = out.strings[reader.Read7BitEncoded()];
 				
 			if (out.version < 15)
 			{
@@ -896,7 +898,7 @@ namespace uba
 				hint = temp.data;
 			}
 			else
-				hint = out.strings[reader.Read7BitEncoded()].data();
+				hint = out.strings[reader.Read7BitEncoded()];
 			if (out.version < 15)
 			{
 				time = reader.Read7BitEncoded();
@@ -996,7 +998,7 @@ namespace uba
 			else
 				stringIndex = reader.Read7BitEncoded();
 				
-			record.description = out.strings[stringIndex].data();
+			record.description = out.strings[stringIndex];
 			if (out.version < 15)
 				record.start = reader.Read7BitEncoded();
 			else
@@ -1086,7 +1088,7 @@ namespace uba
 		}
 		case TraceType_String:
 		{
-			out.strings.push_back(reader.ReadString());
+			out.strings.push_back(TStrdup(reader.ReadString().data()));
 			break;
 		}
 		case TraceType_CacheBeginFetch:
@@ -1147,7 +1149,7 @@ namespace uba
 				return false;
 			WorkRecordLocation active = findIt->second;
 			TraceView::WorkRecord& record = out.workTracks[active.track].records[active.index];
-			record.entries.emplace_back(time, out.strings[stringIndex].data()); 
+			record.entries.emplace_back(time, out.strings[stringIndex]); 
 			break;
 		}
 		default:
diff --git a/Engine/Source/Programs/UnrealBuildAccelerator/Visualizer/Private/UbaTraceReader.h b/Engine/Source/Programs/UnrealBuildAccelerator/Visualizer/Private/UbaTraceReader.h
--- a/Engine/Source/Programs/UnrealBuildAccelerator/Visualizer/Private/UbaTraceReader.h
+++ b/Engine/Source/Programs/UnrealBuildAccelerator/Visualizer/Private/UbaTraceReader.h
@@ -142,7 +142,7 @@ namespace uba
 
 		Vector<Session> sessions;
 		Vector<WorkTrack> workTracks;
-		Vector<TString> strings;
+		Vector<tchar*> strings;
 		Map<u64, StatusUpdate> statusMap;
 		Map<u32, CacheWrite> cacheWrites;
 		u64 realStartTime = 0;
diff --git a/Engine/Source/Programs/UnrealBuildAccelerator/Visualizer/Private/UbaVisualizer.cpp b/Engine/Source/Programs/UnrealBuildAccelerator/Visualizer/Private/UbaVisualizer.cpp
--- a/Engine/Source/Programs/UnrealBuildAccelerator/Visualizer/Private/UbaVisualizer.cpp
+++ b/Engine/Source/Programs/UnrealBuildAccelerator/Visualizer/Private/UbaVisualizer.cpp
@@ -1355,6 +1355,9 @@ namespace uba
 							if (rect.right <= progressRect.left)
 								continue;
 
+							if (!m_filterString.empty() && !Contains(process.description.c_str(), m_filterString.c_str()))
+								continue;
+
 							rect.right = Max(int(rect.right), left + 1);
 
 							bool selected = m_processSelected && m_processSelectedLocation == processLocation;
@@ -1669,6 +1672,25 @@ namespace uba
 		if (m_config.showTimeline && !m_traceView.sessions.empty())
 			PaintTimeline(hdc, clientRect);
 
+		if (!m_filterString.empty())
+		{
+			StringBuffer<> str;
+			str.Append(TC("Process Filter: ")).Append(m_filterString);
+			RECT rect = { 0, 0, 0, 0 };
+
+			SetActiveFont(m_popupFont);
+			DrawTextW(hdc, str.data, str.count, &rect, DT_SINGLELINE|DT_NOCLIP|DT_CALCRECT);
+			rect.left = (progressRect.right - rect.right)/2;
+			rect.right = (progressRect.right + rect.right)/2;
+			rect.top = (progressRect.bottom - rect.bottom)/2;
+			rect.bottom = (progressRect.bottom + rect.bottom)/2;
+
+			RECT fillRect = rect;
+			InflateRect(&fillRect, 4, 2);
+			FillRect(hdc, &fillRect, m_tooltipBackgroundBrush);
+			DrawTextW(hdc, str.data, str.count, &rect, DT_SINGLELINE|DT_NOCLIP);
+		}
+
 		if (m_config.showCursorLine && m_mouseOverWindow)
 		{
 			float timeScale = (m_horizontalScaleValue * m_zoomValue)*50.0f;
@@ -2742,6 +2764,9 @@ namespace uba
 							if (rect.right <= progressRect.left)
 								continue;
 
+							if (!m_filterString.empty() && !Contains(process.description.c_str(), m_filterString.c_str()))
+								continue;
+
 							rect.right = Max(int(rect.right), left + 1);
 							rect.top = posY;
 							rect.bottom = posY + int(float(18) * m_zoomValue);
@@ -3687,8 +3712,21 @@ namespace uba
 				++m_replay;
 			if (wParam == VK_SUBTRACT)
 				--m_replay;
+			if (wParam == VK_BACK)
+				if (!m_filterString.empty())
+				{
+					m_filterString.resize(m_filterString.size() -1);
+					Redraw(true);
+				}
 			break;
 		}
+		case WM_CHAR:
+		{
+			tchar c = static_cast<tchar>(wParam);
+			if (c > 32 && c != '\t' && c != '\n')
+				m_filterString += static_cast<tchar>(wParam);
+			Redraw(true);
+		}
 		case WM_VSCROLL:
 			{
 				RECT r;
diff --git a/Engine/Source/Programs/UnrealBuildAccelerator/Visualizer/Private/UbaVisualizer.h b/Engine/Source/Programs/UnrealBuildAccelerator/Visualizer/Private/UbaVisualizer.h
--- a/Engine/Source/Programs/UnrealBuildAccelerator/Visualizer/Private/UbaVisualizer.h
+++ b/Engine/Source/Programs/UnrealBuildAccelerator/Visualizer/Private/UbaVisualizer.h
@@ -270,6 +270,8 @@ namespace uba
 		float m_scrollAtAnchorY = 0;
 		int m_dragToScrollCounter = 0;
 
+		TString m_filterString;
+
 		Thread m_thread;
 
 		void PostNewTrace(u32 replay, bool paused);
