[UBA]
* Fixed errors when compiling uba using clang

diff --git a/Engine/Source/Programs/UnrealBuildAccelerator/Detours/Private/Windows/UbaDetoursFunctionsWin.cpp b/Engine/Source/Programs/UnrealBuildAccelerator/Detours/Private/Windows/UbaDetoursFunctionsWin.cpp
--- a/Engine/Source/Programs/UnrealBuildAccelerator/Detours/Private/Windows/UbaDetoursFunctionsWin.cpp
+++ b/Engine/Source/Programs/UnrealBuildAccelerator/Detours/Private/Windows/UbaDetoursFunctionsWin.cpp
@@ -1,6 +1,5 @@
 // Copyright Epic Games, Inc. All Rights Reserved.
 
-#define _CRT_SECURE_CPP_OVERLOAD_SECURE_NAMES 0
 #define UBA_IS_DETOURED_INCLUDE 1
 
 #include "UbaDetoursFunctionsWin.h"
diff --git a/Engine/Source/Programs/UnrealBuildAccelerator/Detours/Private/Windows/UbaDetoursFunctionsWin.h b/Engine/Source/Programs/UnrealBuildAccelerator/Detours/Private/Windows/UbaDetoursFunctionsWin.h
--- a/Engine/Source/Programs/UnrealBuildAccelerator/Detours/Private/Windows/UbaDetoursFunctionsWin.h
+++ b/Engine/Source/Programs/UnrealBuildAccelerator/Detours/Private/Windows/UbaDetoursFunctionsWin.h
@@ -319,7 +319,7 @@
 	//DETOURED_FUNCTION(CryptCreateHash) \
 	//DETOURED_FUNCTION(CryptHashData) \
 	//DETOURED_FUNCTION(CreateFileMapping2) \
-	//DETOURED_FUNCTION(CreateFileMappingNumaW) \
+	//DETOURED_FUNCTION(CreateFileMappingNumaW)
 
 #define DETOURED_FUNCTIONS_KERNEL32_DEBUG \
 
@@ -341,7 +341,7 @@
 	DETOURED_FUNCTION(ZwQueryDirectoryFile) \
 	//DETOURED_FUNCTION(ZwCreateFile) \
 	//DETOURED_FUNCTION(ZwOpenFile) \
-	//DETOURED_FUNCTION(RtlAllocateHeap) \
+	//DETOURED_FUNCTION(RtlAllocateHeap)
 
 #define DETOURED_FUNCTIONS_SHLWAPI_DEBUG \
 	DETOURED_FUNCTION(PathFindFileNameW) \
@@ -360,12 +360,12 @@
 	DETOURED_FUNCTION(_wmakepath_s) \
 	DETOURED_FUNCTION(_getcwd) \
 	//DETOURED_FUNCTION(_wsopen_s) \
-	//DETOURED_FUNCTION(_fileno) \
+	//DETOURED_FUNCTION(_fileno)
 
 #if UBA_USE_MIMALLOC
 #define DETOURED_FUNCTIONS_MEMORY_DEBUG \
 	DETOURED_FUNCTION(_aligned_msize) \
-	//DETOURED_FUNCTION(_free_dbg) \
+	//DETOURED_FUNCTION(_free_dbg)
 
 #endif
 
diff --git a/Engine/Source/Programs/UnrealBuildAccelerator/UbaDetours.Target.cs b/Engine/Source/Programs/UnrealBuildAccelerator/UbaDetours.Target.cs
--- a/Engine/Source/Programs/UnrealBuildAccelerator/UbaDetours.Target.cs
+++ b/Engine/Source/Programs/UnrealBuildAccelerator/UbaDetours.Target.cs
@@ -10,5 +10,6 @@ public class UbaDetoursTarget : TargetRules
 		LaunchModuleName = "UbaDetours";
 		bShouldCompileAsDLL = true;
 		UbaAgentTarget.CommonUbaSettings(this, Target);
+		GlobalDefinitions.Add("_CRT_SECURE_CPP_OVERLOAD_SECURE_NAMES=0");
 	}
 }
