[UBA]
* Added validation of every byte read for cache entry on cache server to try to prevent crash happening now and then when bad data is sent from client

diff --git a/Engine/Source/Programs/UnrealBuildAccelerator/Common/Private/UbaCacheServer.cpp b/Engine/Source/Programs/UnrealBuildAccelerator/Common/Private/UbaCacheServer.cpp
--- a/Engine/Source/Programs/UnrealBuildAccelerator/Common/Private/UbaCacheServer.cpp
+++ b/Engine/Source/Programs/UnrealBuildAccelerator/Common/Private/UbaCacheServer.cpp
@@ -1491,25 +1491,28 @@ namespace uba
 
 		auto CreateCasKeyOffset = [&](u32& outCasKeyOffset, CasKey& outCasKey, u32 casKeyOffset, Set<u32>& entries, const tchar* direction)
 			{
-				if (casKeyOffset + sizeof(CasKey) >= connectionBucket.casKeyTable.GetSize())
-					return m_logger.Error(TC("Client %u has provided %s offset that is outside caskey table size. Corrupt data? (Offset: %u TableSize: %u ClientVersion: %u PathtableVersion: %u)"), clientId, direction, casKeyOffset, connectionBucket.casKeyTable.GetSize(), clientVersion, connectionBucket.pathTable.GetVersion());
+				u32 casKeyTableSize = connectionBucket.casKeyTable.GetSize();
+				BinaryReader casKeyReader(connectionBucket.casKeyTable.GetMemory(), casKeyOffset, casKeyTableSize);
+
+				u64 pathOffset;
+				if (!casKeyReader.TryRead7BitEncoded(pathOffset))
+					return m_logger.Error(TC("Client %u has provided %s offset that is outside caskey table size. Corrupt data? (Offset: %u TableSize: %u ClientVersion: %u PathtableVersion: %u)"), clientId, direction, casKeyOffset, casKeyTableSize, clientVersion, connectionBucket.pathTable.GetVersion());
+
+				if (casKeyReader.GetLeft() < sizeof(CasKey))
+					return m_logger.Error(TC("Client %u has provided %s offset that is outside caskey table size. Corrupt data? (Offset: %u TableSize: %u ClientVersion: %u PathtableVersion: %u)"), clientId, direction, casKeyOffset, casKeyTableSize, clientVersion, connectionBucket.pathTable.GetVersion());
+				CasKey casKey = casKeyReader.ReadCasKey();
+				if (casKey == CasKeyZero)
+					return m_logger.Error(TC("Client %u has provided %s cas key that is zero. Corrupt data? (Offset: %u TableSize: %u ClientVersion: %u PathtableVersion: %u)"), clientId, direction, casKeyOffset, casKeyTableSize, clientVersion, connectionBucket.pathTable.GetVersion());
 
-				CasKey casKey;
 				StringBuffer<> path;
-				if (!connectionBucket.casKeyTable.GetPathAndKey(path, casKey, connectionBucket.pathTable, casKeyOffset))
-				{
-					BinaryReader reader2(connectionBucket.casKeyTable.GetMemory(), casKeyOffset, connectionBucket.casKeyTable.GetSize());
-					u32 pathOffset = u32(reader2.Read7BitEncoded());
-					return m_logger.Error(TC("Client %u has provided corrupt %s key or path table (CasKey: %s Offset: %u TableSize: %u PathOffset: %u PathTableSize: %u ClientVersion: %u PathtableVersion: %u)"), clientId, direction, CasKeyString(casKey).str, casKeyOffset, connectionBucket.casKeyTable.GetSize(), pathOffset, connectionBucket.pathTable.GetSize(), clientVersion, connectionBucket.pathTable.GetVersion());
-				}
+				if (!connectionBucket.pathTable.TryGetString(m_logger, path, pathOffset))
+					return m_logger.Error(TC("Client %u has provided corrupt %s key or path table (CasKey: %s Offset: %u TableSize: %u PathOffset: %u PathTableSize: %u ClientVersion: %u PathtableVersion: %u)"), clientId, direction, CasKeyString(casKey).str, casKeyOffset, casKeyTableSize, pathOffset, connectionBucket.pathTable.GetSize(), clientVersion, connectionBucket.pathTable.GetVersion());
 			
 				if (path.count <= 2)
-					return m_logger.Error(TC("Client %u has provided corrupt %s path %s. Corrupt data? (Offset: %u TableSize: %u ClientVersion: %u PathtableVersion: %u)"), clientId, direction, path.data, casKeyOffset, connectionBucket.casKeyTable.GetSize(), clientVersion, connectionBucket.pathTable.GetVersion());
-
-				if (casKey == CasKeyZero)
-					return m_logger.Error(TC("Client %u has provided %s cas key that is zero. Corrupt data? (Offset: %u TableSize: %u ClientVersion: %u PathtableVersion: %u)"), clientId, direction, casKeyOffset, connectionBucket.casKeyTable.GetSize(), clientVersion, connectionBucket.pathTable.GetVersion());
+					return m_logger.Error(TC("Client %u has provided corrupt %s path %s. Corrupt data? (Offset: %u TableSize: %u ClientVersion: %u PathtableVersion: %u)"), clientId, direction, path.data, casKeyOffset, casKeyTableSize, clientVersion, connectionBucket.pathTable.GetVersion());
 
-				u32 pathOffset = bucket.m_pathTable.Add(path.data, path.count);
+				
+				pathOffset = bucket.m_pathTable.Add(path.data, path.count);
 
 				#if 0
 				StringBuffer<> test;
diff --git a/Engine/Source/Programs/UnrealBuildAccelerator/Common/Private/UbaCompactTables.cpp b/Engine/Source/Programs/UnrealBuildAccelerator/Common/Private/UbaCompactTables.cpp
--- a/Engine/Source/Programs/UnrealBuildAccelerator/Common/Private/UbaCompactTables.cpp
+++ b/Engine/Source/Programs/UnrealBuildAccelerator/Common/Private/UbaCompactTables.cpp
@@ -2,6 +2,7 @@
 
 #include "UbaCompactTables.h"
 #include "UbaBinaryReaderWriter.h"
+#include "UbaLogger.h"
 
 namespace uba
 {
@@ -477,12 +478,63 @@ namespace uba
 		return true;
 	}
 
+	bool CompactPathTable::TryGetString(Logger& logger, StringBufferBase& out, u64 offset) const
+	{
+		if (m_version < 3)
+			return GetString(out, offset);
+
+		u32 offsets[MaxSegments];
+		bool separators[MaxSegments];
+		u32 offsetCount = 1;
+
+		u64 memorySize = GetSize();
+		BinaryReader reader(m_pathTableMem.memory, offset, memorySize);
+
+		u64 parentRelativeOffset;
+		while (true)
+		{
+			if (offsetCount == sizeof_array(offsets))
+				return logger.Error(TC("Too many sections in compressed string"));
+
+			reader.SetPosition(offset);
+			if (!reader.TryRead7BitEncoded(parentRelativeOffset))
+				return logger.Error(TC("Failed to read parentRelativeOffset from PathTableMemory (%ull/%ull)"), offset, memorySize);
+			offsets[offsetCount-1] = u32(reader.GetPosition());
+			separators[offsetCount] = parentRelativeOffset & 1;
+			parentRelativeOffset >>= 1;
+			if (parentRelativeOffset >= offset)
+				break;
+			offset -= parentRelativeOffset;
+			++offsetCount;
+		}
+		if (parentRelativeOffset > offset)
+			return logger.Error(TC("Failed to read from PathTableMemory. Parent offset %llu larger than offset %ull"), parentRelativeOffset, offset);
+
+		bool isFirst = true;
+		for (u32 i=offsetCount;i; --i)
+		{
+			reader.SetPosition(offsets[i-1]);
+			u64 strOffset;
+			if (!reader.TryRead7BitEncoded(strOffset))
+				return logger.Error(TC("Failed to read strOffset from PathTableMemory (%llu/%ull)"), reader.GetPosition(), memorySize);
+			if (strOffset != 0)
+				reader.SetPosition(strOffset);
+
+			if (!isFirst)
+				out.Append(separators[i] ? '.' : PathSeparator);
+			isFirst = false;
+			if (!reader.TryReadString(out))
+				return logger.Error(TC("Failed to read string from PathTableMemory (%llu/%ull)"), strOffset, memorySize);
+		}
+		return true;
+	}
+
 	u8* CompactPathTable::GetMemory()
 	{
 		return m_pathTableMem.memory;
 	}
 
-	u32 CompactPathTable::GetSize()
+	u32 CompactPathTable::GetSize() const
 	{
 		SCOPED_READ_LOCK(m_lock, lock2)
 		return u32(m_pathTableMem.writtenSize);
diff --git a/Engine/Source/Programs/UnrealBuildAccelerator/Common/Public/UbaCompactTables.h b/Engine/Source/Programs/UnrealBuildAccelerator/Common/Public/UbaCompactTables.h
--- a/Engine/Source/Programs/UnrealBuildAccelerator/Common/Public/UbaCompactTables.h
+++ b/Engine/Source/Programs/UnrealBuildAccelerator/Common/Public/UbaCompactTables.h
@@ -8,6 +8,7 @@
 
 namespace uba
 {
+	class Logger;
 	struct BinaryReader;
 	struct BinaryWriter;
 	struct StringKey;
@@ -27,9 +28,10 @@ namespace uba
 		void AddCommonStringSegments();
 
 		bool GetString(StringBufferBase& out, u64 offset) const;
+		bool TryGetString(Logger& logger, StringBufferBase& out, u64 offset) const;
 
 		u8* GetMemory();
-		u32 GetSize();
+		u32 GetSize() const;
 		u32 GetCommonSize();
 		bool GetCaseInsensitive();
 		u32 GetVersion();
diff --git a/Engine/Source/Programs/UnrealBuildAccelerator/Core/Private/UbaSynchronization.cpp b/Engine/Source/Programs/UnrealBuildAccelerator/Core/Private/UbaSynchronization.cpp
--- a/Engine/Source/Programs/UnrealBuildAccelerator/Core/Private/UbaSynchronization.cpp
+++ b/Engine/Source/Programs/UnrealBuildAccelerator/Core/Private/UbaSynchronization.cpp
@@ -87,7 +87,7 @@ namespace uba
 		#endif
 	}
 
-	void ReaderWriterLock::EnterRead()
+	void ReaderWriterLock::EnterRead() const
 	{
 		#if PLATFORM_WINDOWS
 		AcquireSRWLockShared((SRWLOCK*)&data);
@@ -96,7 +96,7 @@ namespace uba
 		#endif
 	}
 
-	void ReaderWriterLock::LeaveRead()
+	void ReaderWriterLock::LeaveRead() const
 	{
 		#if PLATFORM_WINDOWS
 		ReleaseSRWLockShared((SRWLOCK*)&data);
diff --git a/Engine/Source/Programs/UnrealBuildAccelerator/Core/Public/UbaBinaryReaderWriter.h b/Engine/Source/Programs/UnrealBuildAccelerator/Core/Public/UbaBinaryReaderWriter.h
--- a/Engine/Source/Programs/UnrealBuildAccelerator/Core/Public/UbaBinaryReaderWriter.h
+++ b/Engine/Source/Programs/UnrealBuildAccelerator/Core/Public/UbaBinaryReaderWriter.h
@@ -86,6 +86,7 @@ namespace uba
 		inline u64 ReadFileSize() { return Read7BitEncoded(); }
 
 		inline bool TryRead7BitEncoded(u64& outValue);
+		inline bool TryReadString(StringBufferBase& out);
 		inline u32 PeekU32();
 		inline u64 PeekU64();
 		inline void Skip(u64 size) { m_pos += size; }
@@ -480,6 +481,24 @@ namespace uba
 		return true;
 	}
 
+	bool BinaryReader::TryReadString(StringBufferBase& out)
+	{
+		#if PLATFORM_WINDOWS
+		ReadString(out); // TODO: Implement secure read
+		#else
+		u64 charLen;
+		if (!TryRead7BitEncoded(charLen))
+			return false;
+		u32 strCapacity = out.capacity - out.count;
+		if (charLen >= strCapacity)
+			return false;
+		ReadBytes(out.data + out.count, charLen);
+		out.count += charLen;
+		out.data[out.count] = 0;
+		#endif
+		return true;
+	}
+
 	template<typename CharType>
 	CharType BinaryReader::ReadUtf8Char()
 	{
diff --git a/Engine/Source/Programs/UnrealBuildAccelerator/Core/Public/UbaSynchronization.h b/Engine/Source/Programs/UnrealBuildAccelerator/Core/Public/UbaSynchronization.h
--- a/Engine/Source/Programs/UnrealBuildAccelerator/Core/Public/UbaSynchronization.h
+++ b/Engine/Source/Programs/UnrealBuildAccelerator/Core/Public/UbaSynchronization.h
@@ -66,8 +66,8 @@ namespace uba
 		ReaderWriterLock();
 		~ReaderWriterLock();
 
-		void EnterRead();
-		void LeaveRead();
+		void EnterRead() const;
+		void LeaveRead() const;
 
 		void EnterWrite();
 		void LeaveWrite();
@@ -112,12 +112,12 @@ namespace uba
 	class ScopedReadLock
 	{
 	public:
-		ScopedReadLock(ReaderWriterLock& lock) : m_lock(lock) { lock.EnterRead(); }
+		ScopedReadLock(const ReaderWriterLock& lock) : m_lock(lock) { lock.EnterRead(); }
 		~ScopedReadLock() { Leave(); }
 		inline void Enter() { if (m_active) return; m_active = true; m_lock.EnterRead(); }
 		inline void Leave() { if (!m_active) return; m_active = false; m_lock.LeaveRead(); }
 
-		ReaderWriterLock& m_lock;
+		const ReaderWriterLock& m_lock;
 		bool m_active = true;
 	};
 
