[UBA]
* Added more filter support for work in visualizer
* Added color support to work bars in visualizer

diff --git a/Engine/Source/Programs/UnrealBuildAccelerator/Common/Private/UbaNetworkServer.cpp b/Engine/Source/Programs/UnrealBuildAccelerator/Common/Private/UbaNetworkServer.cpp
--- a/Engine/Source/Programs/UnrealBuildAccelerator/Common/Private/UbaNetworkServer.cpp
+++ b/Engine/Source/Programs/UnrealBuildAccelerator/Common/Private/UbaNetworkServer.cpp
@@ -194,7 +194,7 @@ namespace uba
 			auto& logger = conn.m_server.m_logger;
 
 			if (bodySize != sizeof(EncryptionHandshakeString))
-				return logger.Error(TC("Crypto mismatch..."));
+				return logger.Warning(TC("Connection %u Crypto mismatch... (body size was %u, expected %u)"), conn.m_id, bodySize, sizeof(EncryptionHandshakeString));
 
 			auto TestHandshake = [&](CryptoKey key)
 			{
@@ -208,7 +208,7 @@ namespace uba
 			if (conn.m_cryptoKey != InvalidCryptoKey)
 			{
 				if (!TestHandshake(conn.m_cryptoKey))
-					return logger.Error(TC("Crypto mismatch..."));
+					return logger.Warning(TC("Connection %u Crypto mismatch... (Handshake string is encrypted with different key)"), conn.m_id);
 			}
 			else
 			{
@@ -234,7 +234,7 @@ namespace uba
 					break;
 				}
 				if (conn.m_cryptoKey == InvalidCryptoKey)
-					return logger.Error(TC("Crypto mismatch..."));
+					return logger.Warning(TC("Connection %u Crypto mismatch... (Handshake string is encrypted with different key than any registered keys)"), conn.m_id);
 			}
 
 			conn.m_backend.SetAllowLessThanBodySize(conn.m_backendConnection, false);
@@ -475,7 +475,7 @@ namespace uba
 		context.connection = nullptr;
 
 		WorkerRec& rec = server.m_workerFunctions[context.serviceId];
-		TrackWorkScope tws(server, rec.toString(context.messageType));
+		TrackWorkScope tws(server, rec.toString(context.messageType), ColorWork);
 
 		CryptoKey cryptoKey = connection.m_cryptoKey;
 		if (cryptoKey)
@@ -506,14 +506,14 @@ namespace uba
 
 		if (!rec.func)
 		{
-			server.m_logger.Error(TC("WORKER FUNCTION NOT FOUND. id: %u, serviceid: %u type: %s, client: %s"), context.id, context.serviceId, rec.toString(context.messageType).data, GuidToString(connection.m_client->uid).str);
+			server.m_logger.Error(TC("WORKER FUNCTION NOT FOUND. id: %u, serviceid: %u type: %s, client: %u"), context.id, context.serviceId, rec.toString(context.messageType).data, connection.m_client->id);
 			connection.SetShouldDisconnect();
 			size = ErrorSize;
 		}
 		else if (!rec.func({&connection}, {tws.id}, mi, reader, writer))
 		{
 			if (connection.SetShouldDisconnect())
-				server.m_logger.Error(TC("WORKER FUNCTION FAILED. id: %u, serviceid: %u type: %s, client: %s"), context.id, context.serviceId, rec.toString(context.messageType).data, GuidToString(connection.m_client->uid).str);
+				server.m_logger.Error(TC("WORKER FUNCTION FAILED. id: %u, serviceid: %u type: %s, client: %u"), context.id, context.serviceId, rec.toString(context.messageType).data, connection.m_client->id);
 			size = ErrorSize;
 		}
 		else
@@ -582,7 +582,7 @@ namespace uba
 				server.m_additionalWork.pop_front();
 				lock.Leave();
 
-				TrackWorkScope tws(server, work.desc);
+				TrackWorkScope tws(server, work.desc, ColorWork);
 				work.func({tws.id});
 			}
 
@@ -702,7 +702,7 @@ namespace uba
 
 		SCOPED_WRITE_LOCK(m_connectionsLock, lock);
 		m_connections.clear();
-
+		m_allClientsDisconnected = true;
 		m_workersEnabled = true;
 	}
 
@@ -838,7 +838,7 @@ namespace uba
 		}
 	}
 
-	void NetworkServer::AddWork(const WorkFunction& work, u32 count, const tchar* desc, bool highPriority)
+	void NetworkServer::AddWork(const WorkFunction& work, u32 count, const tchar* desc, const Color& color, bool highPriority)
 	{
 		SCOPED_WRITE_LOCK(m_additionalWorkLock, lock);
 		for (u32 i = 0; i != count; ++i)
@@ -988,7 +988,7 @@ namespace uba
 		m_additionalWork.pop_front();
 		lock.Leave();
 
-		TrackWorkScope tws(*this, work.desc);
+		TrackWorkScope tws(*this, work.desc, ColorWork);
 
 		work.func({tws.id});
 
@@ -1244,7 +1244,7 @@ namespace uba
 
 		RemoveDisconnectedConnections();
 
-		if (!m_workersEnabled)
+		if (!m_workersEnabled || m_allClientsDisconnected)
 		{
 			// Just to prevent errors in log
 			backend.SetDisconnectCallback(backendConnection, nullptr, [](void*, const Guid&, void*) {});
diff --git a/Engine/Source/Programs/UnrealBuildAccelerator/Common/Private/UbaStorageProxy.cpp b/Engine/Source/Programs/UnrealBuildAccelerator/Common/Private/UbaStorageProxy.cpp
--- a/Engine/Source/Programs/UnrealBuildAccelerator/Common/Private/UbaStorageProxy.cpp
+++ b/Engine/Source/Programs/UnrealBuildAccelerator/Common/Private/UbaStorageProxy.cpp
@@ -204,7 +204,8 @@ namespace uba
 
 					if (!file.memory)
 					{
-						file.trackId = m_client.TrackWorkStart(StringBuffer<512>().AppendFileName(hint.data));
+						file.trackId = m_client.TrackWorkStart(AsView(TC("ProxyFetchBegin")), ColorWork);
+						m_client.TrackWorkHint(file.trackId, StringView(hint).GetFileName());
 
 						NetworkMessage msg(m_client, ServiceId, StorageMessageType_FetchBegin, writer2);
 						writer2.WriteByte(2); // Does not want proxy but informs it is proxy
@@ -279,7 +280,7 @@ namespace uba
 										{
 											auto mif = (MessageInFlight*)userData;
 											mif->error = error;
-											mif->proxy.m_server.AddWork([mif](const WorkContext&) { mif->proxy.HandleReceivedData(*mif); }, 1, TC(""));
+											mif->proxy.m_server.AddWork([mif](const WorkContext&) { mif->proxy.HandleReceivedData(*mif); }, 1, TC("ProxyWaitMsg"), ColorWork);
 										}, mif);
 									if (!res)
 									{
@@ -287,7 +288,7 @@ namespace uba
 										mif->error = true;
 									}
 								}
-							}, 1, TC(""));
+							}, 1, TC("ProxySpawnMsg"), ColorWork);
 					}
 					else
 					{
diff --git a/Engine/Source/Programs/UnrealBuildAccelerator/Common/Private/UbaTrace.cpp b/Engine/Source/Programs/UnrealBuildAccelerator/Common/Private/UbaTrace.cpp
--- a/Engine/Source/Programs/UnrealBuildAccelerator/Common/Private/UbaTrace.cpp
+++ b/Engine/Source/Programs/UnrealBuildAccelerator/Common/Private/UbaTrace.cpp
@@ -130,10 +130,10 @@ namespace uba
 		return true;
 	}
 
-	u32 Trace::TrackWorkStart(const StringView& desc)
+	u32 Trace::TrackWorkStart(const StringView& desc, const Color& color)
 	{
 		u32 workId = m_workCounter++;
-		WorkBegin(workId, desc);
+		WorkBegin(workId, desc, color);
 		return workId;
 	}
 
@@ -376,12 +376,13 @@ namespace uba
 		writer.WriteCasKey(key);
 	}
 
-	void Trace::WorkBegin(u32 workIndex, const StringView& desc)
+	void Trace::WorkBegin(u32 workIndex, const StringView& desc, const Color& color)
 	{
 		u32 stringIndex = AddString(desc);
 		BEGIN_TRACE_ENTRY(WorkBegin);
 		writer.Write7BitEncoded(workIndex);
 		writer.Write7BitEncoded(stringIndex);
+		writer.WriteU32(color);
 	}
 
 	void Trace::WorkHint(u32 workIndex, const StringView& hint)
diff --git a/Engine/Source/Programs/UnrealBuildAccelerator/Common/Private/UbaWorkManager.cpp b/Engine/Source/Programs/UnrealBuildAccelerator/Common/Private/UbaWorkManager.cpp
--- a/Engine/Source/Programs/UnrealBuildAccelerator/Common/Private/UbaWorkManager.cpp
+++ b/Engine/Source/Programs/UnrealBuildAccelerator/Common/Private/UbaWorkManager.cpp
@@ -5,12 +5,14 @@
 #include "UbaPlatform.h"
 #include "UbaThread.h"
 
+#define UBA_TRACK_WORK 0
+
 namespace uba
 {
-	u32 WorkManager::TrackWorkStart(const StringView& desc)
+	u32 WorkManager::TrackWorkStart(const StringView& desc, const Color& color)
 	{
 		if (auto t = m_workTracker.load())
-			return t->TrackWorkStart(desc);
+			return t->TrackWorkStart(desc, color);
 		return 0;
 	}
 
@@ -66,8 +68,12 @@ namespace uba
 						manager.m_work.pop_front();
 						lock.Leave();
 
+						#if UBA_TRACK_WORK
 						TrackWorkScope tws(manager, work.desc);
 						work.func({tws.id});
+						#else
+						work.func({});
+						#endif
 					}
 
 					SCOPED_WRITE_LOCK(manager.m_availableWorkersLock, lock1);
@@ -105,7 +111,7 @@ namespace uba
 	}
 
 
-	void WorkManagerImpl::AddWork(const WorkFunction& work, u32 count, const tchar* desc, bool highPriority)
+	void WorkManagerImpl::AddWork(const WorkFunction& work, u32 count, const tchar* desc, const Color& color, bool highPriority)
 	{
 		UBA_ASSERT(!m_workers.empty());
 		SCOPED_WRITE_LOCK(m_workLock, lock);
@@ -169,8 +175,12 @@ namespace uba
 			m_work.pop_front();
 			lock.Leave();
 
+			#if UBA_TRACK_WORK
 			TrackWorkScope tws(*this, work.desc);
 			work.func({tws.id});
+			#else
+			work.func({});
+			#endif
 		}
 	}
 
diff --git a/Engine/Source/Programs/UnrealBuildAccelerator/Common/Public/UbaNetworkServer.h b/Engine/Source/Programs/UnrealBuildAccelerator/Common/Public/UbaNetworkServer.h
--- a/Engine/Source/Programs/UnrealBuildAccelerator/Common/Public/UbaNetworkServer.h
+++ b/Engine/Source/Programs/UnrealBuildAccelerator/Common/Public/UbaNetworkServer.h
@@ -80,7 +80,7 @@ namespace uba
 		//void RegisterOnConnection(u8 id, const OnConnectionFunction& func);
 		//void UnregisterOnConnection(u8 id);
 
-		virtual void AddWork(const WorkFunction& work, u32 count, const tchar* desc, bool highPriority = false) override final;
+		virtual void AddWork(const WorkFunction& work, u32 count, const tchar* desc, const Color& color = ColorWork, bool highPriority = false) override final;
 		virtual void DoWork(u32 count = 1) override final;
 		virtual u32 GetWorkerCount() override final;
 
@@ -157,6 +157,7 @@ namespace uba
 		u32 m_maxCreatedWorkerCount = 0;
 		bool m_workersEnabled = true;
 		bool m_logConnections = true;
+		bool m_allClientsDisconnected = false;
 
 		ReaderWriterLock m_addConnectionsLock;
 		List<Thread> m_addConnections;
diff --git a/Engine/Source/Programs/UnrealBuildAccelerator/Common/Public/UbaTrace.h b/Engine/Source/Programs/UnrealBuildAccelerator/Common/Public/UbaTrace.h
--- a/Engine/Source/Programs/UnrealBuildAccelerator/Common/Public/UbaTrace.h
+++ b/Engine/Source/Programs/UnrealBuildAccelerator/Common/Public/UbaTrace.h
@@ -68,9 +68,6 @@ namespace uba
 		#undef UBA_TRACE_TYPE
 	};
 
-	using Color = u32;
-	inline Color toColor(u8 r, u8 g, u8 b) { return (r << 16) + (g << 8) + b; }
-
 	static constexpr u32 TraceVersion = 38;
 	static constexpr u32 TraceReadCompatibilityVersion = 6;
 
@@ -101,7 +98,7 @@ namespace uba
 		void FileFetchEnd(u32 clientId, const CasKey& key);
 		void FileStoreBegin(u32 clientId, const CasKey& key, u64 size, const StringView& hint, bool detailed);
 		void FileStoreEnd(u32 clientId, const CasKey& key);
-		void WorkBegin(u32 workIndex, const StringView& desc);
+		void WorkBegin(u32 workIndex, const StringView& desc, const Color& color);
 		void WorkHint(u32 workIndex, const StringView& hint);
 		void WorkEnd(u32 workIndex);
 		void ProgressUpdate(u32 processesTotal, u32 processesDone, u32 errorCount);
@@ -115,7 +112,7 @@ namespace uba
 
 		bool StopWrite(const tchar* writeFileName);
 
-		virtual u32 TrackWorkStart(const StringView& desc) final override;
+		virtual u32 TrackWorkStart(const StringView& desc, const Color& color) final override;
 		virtual void TrackWorkHint(u32 id, const StringView& hint) final override;
 		virtual void TrackWorkEnd(u32 id) final override;
 
diff --git a/Engine/Source/Programs/UnrealBuildAccelerator/Common/Public/UbaWorkManager.h b/Engine/Source/Programs/UnrealBuildAccelerator/Common/Public/UbaWorkManager.h
--- a/Engine/Source/Programs/UnrealBuildAccelerator/Common/Public/UbaWorkManager.h
+++ b/Engine/Source/Programs/UnrealBuildAccelerator/Common/Public/UbaWorkManager.h
@@ -8,11 +8,15 @@
 namespace uba
 {
 	struct StringView;
+	using Color = u32;
+	constexpr inline Color ToColor(u8 r, u8 g, u8 b) { return (r << 16) + (g << 8) + b; }
+	constexpr Color ColorWhite = ToColor(255, 255, 255);
+	constexpr Color ColorWork = ToColor(70, 70, 100);
 
 	class WorkTracker
 	{
 	public:
-		virtual u32 TrackWorkStart(const StringView& desc) = 0;
+		virtual u32 TrackWorkStart(const StringView& desc, const Color& color) = 0;
 		virtual void TrackWorkHint(u32 id, const StringView& hint) = 0;
 		virtual void TrackWorkEnd(u32 id) = 0;
 	};
@@ -26,11 +30,11 @@ namespace uba
 	{
 	public:
 		using WorkFunction = Function<void(const WorkContext& context)>;
-		virtual void AddWork(const WorkFunction& work, u32 count, const tchar* desc, bool highPriority = false) = 0;
+		virtual void AddWork(const WorkFunction& work, u32 count, const tchar* desc, const Color& color = ColorWork, bool highPriority = false) = 0;
 		virtual u32 GetWorkerCount() = 0;
 		virtual void DoWork(u32 count = 1) = 0;
 
-		u32 TrackWorkStart(const StringView& desc);
+		u32 TrackWorkStart(const StringView& desc, const Color& color);
 		void TrackWorkHint(u32 id, const StringView& hint);
 		void TrackWorkEnd(u32 id);
 
@@ -50,7 +54,7 @@ namespace uba
 	public:
 		WorkManagerImpl(u32 workerCount, const tchar* workerDesc = TC("UbaWorker"));
 		virtual ~WorkManagerImpl();
-		virtual void AddWork(const WorkFunction& work, u32 count, const tchar* desc, bool highPriority = false) override;
+		virtual void AddWork(const WorkFunction& work, u32 count, const tchar* desc, const Color& color = ColorWork, bool highPriority = false) override;
 		virtual u32 GetWorkerCount() override;
 		virtual void DoWork(u32 count = 1) override;
 		void FlushWork();
@@ -74,7 +78,8 @@ namespace uba
 
 	struct TrackWorkScope
 	{
-		TrackWorkScope(WorkManager& wm, const StringView& desc) : workManager(wm), id(wm.TrackWorkStart(desc)) {}
+		TrackWorkScope(WorkManager& wm, const StringView& desc, const Color& color) : workManager(wm), id(wm.TrackWorkStart(desc, color)) {}
+		void AddHint(const StringView& hint) { workManager.TrackWorkHint(id, hint); }
 		~TrackWorkScope() { workManager.TrackWorkEnd(id); }
 		WorkManager& workManager;
 		u32 id;
@@ -139,7 +144,7 @@ namespace uba
 				}
 			};
 		
-		AddWork(work, workCount, description, highPriority);
+		AddWork(work, workCount, description, ColorWork, highPriority);
 		work({});
 		doneEvent.IsSet();
 
diff --git a/Engine/Source/Programs/UnrealBuildAccelerator/Visualizer/Private/UbaTraceReader.cpp b/Engine/Source/Programs/UnrealBuildAccelerator/Visualizer/Private/UbaTraceReader.cpp
--- a/Engine/Source/Programs/UnrealBuildAccelerator/Visualizer/Private/UbaTraceReader.cpp
+++ b/Engine/Source/Programs/UnrealBuildAccelerator/Visualizer/Private/UbaTraceReader.cpp
@@ -911,6 +911,7 @@ namespace uba
 
 			if (auto session = GetSession(out, clientUid))
 			{
+				session->storedFilesActive.try_emplace(key, u32(session->storedFiles.size()));
 				session->storedFiles.push_back({ key, size, hint, time, ~u64(0) });
 				session->storedFilesBytes += size;
 			}
@@ -944,11 +945,12 @@ namespace uba
 
 			if (auto session = GetSession(out, clientUid))
 			{
-				for (auto rit = session->storedFiles.rbegin(), rend = session->storedFiles.rend(); rit != rend; ++rit)
+				auto findIt = session->storedFilesActive.find(key);
+				if (findIt != session->storedFilesActive.end())
 				{
-					if (rit->key != key)
-						continue;
-					rit->stop = time;
+					auto& file = session->storedFiles[findIt->second];
+					file.stop = time;
+					session->storedFilesActive.erase(findIt);
 					break;
 				}
 			}
@@ -998,6 +1000,9 @@ namespace uba
 			else
 				stringIndex = reader.Read7BitEncoded();
 				
+			if (out.version >= 38)
+				record.color = reader.ReadU32();
+
 			record.description = out.strings[stringIndex];
 			if (out.version < 15)
 				record.start = reader.Read7BitEncoded();
@@ -1014,17 +1019,22 @@ namespace uba
 			else
 				workIndex = u32(reader.Read7BitEncoded());
 
+			u64 stop;
+			if (out.version < 15)
+				stop = reader.Read7BitEncoded();
+			else
+				stop = time;
+
 			auto findIt = m_activeWorkRecords.find(workIndex);
 			if (findIt == m_activeWorkRecords.end())
-				return false;
+				return true;
+
 			WorkRecordLocation active = findIt->second;
 			m_activeWorkRecords.erase(findIt);
+			auto& record = out.workTracks[active.track].records[active.index];
+			record.stop = stop;
+			record.bitmapDirty = true;
 
-			TraceView::WorkRecord& record = out.workTracks[active.track].records[active.index];
-			if (out.version < 15)
-				record.stop = reader.Read7BitEncoded();
-			else
-				record.stop = time;
 			break;
 		}
 		case TraceType_ProgressUpdate:
diff --git a/Engine/Source/Programs/UnrealBuildAccelerator/Visualizer/Private/UbaTraceReader.h b/Engine/Source/Programs/UnrealBuildAccelerator/Visualizer/Private/UbaTraceReader.h
--- a/Engine/Source/Programs/UnrealBuildAccelerator/Visualizer/Private/UbaTraceReader.h
+++ b/Engine/Source/Programs/UnrealBuildAccelerator/Visualizer/Private/UbaTraceReader.h
@@ -64,6 +64,8 @@ namespace uba
 			Vector<WorkRecordLogEntry> entries;
 			HBITMAP bitmap = 0;
 			u32 bitmapOffset = 0;
+			u32 color = ColorWork;
+			bool bitmapDirty = true;
 		};
 
 		struct WorkTrack
@@ -98,6 +100,7 @@ namespace uba
 			Vector<TString> summary;
 			UnorderedMap<CasKey, u32> fetchedFilesActive;
 			Vector<FileTransfer> fetchedFiles;
+			UnorderedMap<CasKey, u32> storedFilesActive;
 			Vector<FileTransfer> storedFiles;
 			TString notification;
 			u64 fetchedFilesBytes = 0;
diff --git a/Engine/Source/Programs/UnrealBuildAccelerator/Visualizer/Private/UbaVisualizer.cpp b/Engine/Source/Programs/UnrealBuildAccelerator/Visualizer/Private/UbaVisualizer.cpp
--- a/Engine/Source/Programs/UnrealBuildAccelerator/Visualizer/Private/UbaVisualizer.cpp
+++ b/Engine/Source/Programs/UnrealBuildAccelerator/Visualizer/Private/UbaVisualizer.cpp
@@ -22,6 +22,7 @@ namespace uba
 		Popup_CopyProcessInfo,
 		Popup_CopyProcessLog,
 		Popup_CopyProcessBreadcrumbs,
+		Popup_CopyWorkInfo,
 		Popup_Replay,
 		Popup_Pause,
 		Popup_Play,
@@ -471,8 +472,6 @@ namespace uba
 			m_processBrushes[0].cacheFetch = CreateSolidBrush(RGB(24, 112, 110));
 			m_processBrushes[1].cacheFetch = CreateSolidBrush(RGB(31, 143, 138));
 
-			m_workBrush = CreateSolidBrush(RGB(70, 70, 100));
-
 			m_backgroundBrush = CreateSolidBrush(0x00252526);
 			m_separatorPen = CreatePen(PS_SOLID, 1, RGB(50, 50, 50));
 			m_tooltipBackgroundBrush = CreateSolidBrush(0x00404040);
@@ -508,8 +507,6 @@ namespace uba
 			m_processBrushes[0].cacheFetch = CreateSolidBrush(RGB(150, 150, 200));
 			m_processBrushes[1].cacheFetch = CreateSolidBrush(RGB(170, 170, 200));
 
-			m_workBrush = CreateSolidBrush(RGB(150, 150, 200));
-
 			m_backgroundBrush = GetSysColorBrush(0);
 			m_separatorPen = CreatePen(PS_SOLID, 1, RGB(180, 180, 180));
 			m_tooltipBackgroundBrush = GetSysColorBrush(COLOR_INFOBK);
@@ -754,7 +751,7 @@ namespace uba
 		for (auto& workTrack : m_traceView.workTracks)
 			for (auto& work : workTrack.records)
 			{
-					//work.bitmapDirty = true;
+					work.bitmapDirty = true;
 					if (full)
 						work.bitmap = 0;
 			}
@@ -940,7 +937,7 @@ namespace uba
 		HBITMAP nullBmp = CreateCompatibleBitmap(hdc, 1, 1);
 		HBITMAP oldBmp = (HBITMAP)SelectObject(textDC, nullBmp);
 		HBITMAP lastSelectedBitmap = 0;
-		HBRUSH lastSelectedBrush = 0;
+		//HBRUSH lastSelectedBrush = 0;
 
 		u64 lastStop = 0;
 
@@ -1521,20 +1518,26 @@ namespace uba
 						int posX = int(m_scrollPosX) + progressRect.left;
 						for (auto& work : workTrack.records)
 						{
+							auto inc = MakeGuard([&](){ ++workIndex; });
+
 							if (work.start == work.stop)
-							{
-								++workIndex;
 								continue;
-							}
 
 							float startTime = TimeToS(work.start);
 
 							int left = posX + int(startTime * scaleX);
 
 							if (left >= progressRect.right)
-							{
-								++workIndex;
 								continue;
+
+							if (!m_filterString.empty())
+							{
+								bool keep = Contains(work.description, m_filterString.c_str());
+								if (!keep)
+									for (auto& en : work.entries)
+										keep |= Contains(en.text, m_filterString.c_str());
+								if (!keep)
+									continue;
 							}
 
 							HBRUSH brush;
@@ -1544,7 +1547,13 @@ namespace uba
 							bool done = stop != ~u64(0);
 							if (done)
 							{
-								brush = m_workBrush;
+								auto insres = m_coloredBrushes.try_emplace(work.color);
+								if (insres.second)
+								{
+									auto c = (u8*)&work.color;
+									insres.first->second = CreateSolidBrush(RGB(c[2], c[1], c[0]));
+								}
+								brush = insres.first->second;
 							}
 							else
 							{
@@ -1554,10 +1563,7 @@ namespace uba
 
 							float stopTime = TimeToS(stop);
 							if ((stopTime - startTime) * scaleX < 0.05f)
-							{
-								++workIndex;
 								continue;
-							}
 
 							RECT rect;
 							rect.left = left;
@@ -1566,21 +1572,15 @@ namespace uba
 							rect.bottom = rectBottom - 1;
 
 							if (rect.right <= progressRect.left)
-							{
-								++workIndex;
 								continue;
-							}
 
 							rect.right = Max(int(rect.right), left + 1);
 
 							bool selected = m_workSelected && m_workTrack == trackIndex && m_workIndex == workIndex;
 							if (selected)
-								selectedWork = work;
-
-							if (lastSelectedBrush != brush)
 							{
-								SelectObject(hdc, brush);
-								lastSelectedBrush = brush;
+								selectedWork = work;
+								work.bitmapDirty = true;
 							}
 
 							--rect.top;
@@ -1594,31 +1594,33 @@ namespace uba
 							int processWidth = rect.right - rect.left;
 							if (shouldDrawText && m_config.ShowProcessText && processWidth > 3)
 							{
-								if (!work.bitmap)
+								if (work.bitmapDirty)
 								{
-									if (m_lastBitmapOffset == BitmapCacheHeight)
+									if (!work.bitmap)
 									{
-										if (m_lastBitmap)
-											m_textBitmaps.push_back(m_lastBitmap);
+										if (m_lastBitmapOffset == BitmapCacheHeight)
+										{
+											if (m_lastBitmap)
+												m_textBitmaps.push_back(m_lastBitmap);
 
-										m_lastBitmapOffset = 0;
-										m_lastBitmap = CreateCompatibleBitmap(hdc, 256, BitmapCacheHeight);
+											m_lastBitmapOffset = 0;
+											m_lastBitmap = CreateCompatibleBitmap(hdc, 256, BitmapCacheHeight);
+										}
+										work.bitmap = m_lastBitmap;
+										work.bitmapOffset = m_lastBitmapOffset;
+										m_lastBitmapOffset += m_processFont.height;
+									}
+									if (lastSelectedBitmap != work.bitmap)
+									{
+										SelectObject(textDC, work.bitmap);
+										lastSelectedBitmap = work.bitmap;
 									}
-									SelectObject(textDC, m_lastBitmap);
 
-									RECT rect2{ 0,m_lastBitmapOffset,256, m_lastBitmapOffset + m_processFont.height };
+									RECT rect2{ 0,int(work.bitmapOffset), 256, int(work.bitmapOffset) + m_processFont.height };
 
-									FillRect(textDC, &rect2, m_workBrush);
+									FillRect(textDC, &rect2, brush);
 									ExtTextOutW(textDC, rect2.left, rect2.top, ETO_CLIPPED, &rect2, work.description, int(wcslen(work.description)), NULL);
-									work.bitmap = m_lastBitmap;
-									work.bitmapOffset = m_lastBitmapOffset;
-									m_lastBitmapOffset += m_processFont.height;
-								}
-
-								if (lastSelectedBitmap != work.bitmap)
-								{
-									SelectObject(textDC, work.bitmap);
-									lastSelectedBitmap = work.bitmap;
+									work.bitmapDirty = false;
 								}
 
 								int width = Min(processWidth, 256);
@@ -1647,7 +1649,6 @@ namespace uba
 									BitBlt(hdc, rect.left, rect.top + bltOffsetY, width, height, textDC, bitmapOffsetX, bitmapOffsetY, SRCCOPY);
 								}
 							}
-							++workIndex;
 						}
 					}
 					++trackIndex;
@@ -2413,11 +2414,11 @@ namespace uba
 
 			int fileWidth = 700;
 
-			auto drawFiles = [&](const wchar_t* fileType, Vector<TraceView::FileTransfer>& files, u64 count, u64 bytes, u32& maxVisibleFiles)
+			auto drawFiles = [&](const wchar_t* fileType, Vector<TraceView::FileTransfer>& files, u64 count, u64 bytes, u64 activeCount, u32& maxVisibleFiles)
 				{
 					textRect.left = posX;
 					textRect.right = posX + fileWidth;
-					drawText(L"%ls Files: %u (%s)", fileType, u32(count), BytesToText(bytes));
+					drawText(L"%ls Files: %u (%u) %s", fileType, u32(count), u32(activeCount), BytesToText(bytes));
 					u32 fileCount = 0;
 					for (auto rit = files.rbegin(), rend = files.rend(); rit != rend; ++rit)
 					{
@@ -2449,11 +2450,11 @@ namespace uba
 
 			posY = startPosY;
 			posX += 150;
-			drawFiles(L"Fetched", fetchedFiles, session.fetchedFiles.size(), session.fetchedFilesBytes, session.maxVisibleFiles);
+			drawFiles(L"Fetched", fetchedFiles, session.fetchedFiles.size(), session.fetchedFilesBytes, session.fetchedFilesActive.size(), session.maxVisibleFiles);
 			int posY2 = posY;
 			posY = startPosY;
 			posX += fileWidth;
-			drawFiles(L"Stored", session.storedFiles, session.storedFiles.size(), session.storedFilesBytes, session.maxVisibleFiles);
+			drawFiles(L"Stored", session.storedFiles, session.storedFiles.size(), session.storedFilesBytes, session.storedFilesActive.size(), session.maxVisibleFiles);
 			posY = Max(posY, Max(posY1, posY2));
 		}
 		else
@@ -2809,13 +2810,23 @@ namespace uba
 						int posX = int(m_scrollPosX) + progressRect.left;
 						for (auto& work : workTrack.records)
 						{
+							auto inc = MakeGuard([&](){ ++workIndex; });
+
 							int left = posX + int(float(TimeToS(work.start)) * scaleX);
 
 							if (left >= progressRect.right)
-							{
-								++workIndex;
 								continue;
+
+							if (!m_filterString.empty())
+							{
+								bool keep = Contains(work.description, m_filterString.c_str());
+								if (!keep)
+									for (auto& en : work.entries)
+										keep |= Contains(en.text, m_filterString.c_str());
+								if (!keep)
+									continue;
 							}
+
 							if (left < progressRect.left)
 								left = progressRect.left;
 
@@ -2829,10 +2840,8 @@ namespace uba
 							rect.right = posX + int(float(TimeToS(stopTime)) * scaleX);
 
 							if (rect.right <= progressRect.left)
-							{
-								++workIndex;
 								continue;
-							}
+
 							rect.right = Max(int(rect.right), left + 1);
 							rect.top = posY;
 							rect.bottom = posY + int(float(18) * m_zoomValue);
@@ -2844,7 +2853,6 @@ namespace uba
 								outResult.workSelected = true;
 								return;
 							}
-							++workIndex;
 						}
 					}
 					++trackIndex;
@@ -2915,6 +2923,16 @@ namespace uba
 		}
 	}
 
+	void Visualizer::WriteWorkStats(Logger& out, const TraceView::WorkRecord& record)
+	{
+		out.Info(L"  %ls", record.description);
+		out.Info(L"  Start:     %ls", TimeToText(record.start, true).str);
+		if (record.stop != ~u64(0))
+			out.Info(L"  Duration:  %ls", TimeToText(record.stop - record.start, true).str);
+		for (auto& e : record.entries)
+			out.Info(L"   %ls (%ls)", e.text, TimeToText(e.time - record.start).str);
+	}
+
 	void Visualizer::CopyTextToClipboard(const TString& str)
 	{
 		if (!OpenClipboard(m_hwnd))
@@ -3522,6 +3540,10 @@ namespace uba
 					AppendMenuW(hMenu, MF_STRING, Popup_CopyProcessBreadcrumbs, L"Copy Process &Breadcrumbs");
 				AppendMenuW(hMenu, MF_SEPARATOR, 0, NULL);
 			}
+			else if (m_workSelected)
+			{
+				AppendMenuW(hMenu, MF_STRING, Popup_CopyWorkInfo, L"&Copy Work Info");
+			}
 
 			if (!m_traceView.sessions.empty())
 			{
@@ -3690,6 +3712,15 @@ namespace uba
 				CopyTextToClipboard(process.breadcrumbs);
 				break;
 			}
+			case Popup_CopyWorkInfo:
+			{
+				TString str;
+				WriteTextLogger logger(str);
+				auto& record = m_traceView.workTracks[m_workTrack].records[m_workIndex];
+				WriteWorkStats(logger, record);
+				CopyTextToClipboard(str);
+				break;
+			}
 			}
 
 			DestroyMenu(hMenu);
diff --git a/Engine/Source/Programs/UnrealBuildAccelerator/Visualizer/Private/UbaVisualizer.h b/Engine/Source/Programs/UnrealBuildAccelerator/Visualizer/Private/UbaVisualizer.h
--- a/Engine/Source/Programs/UnrealBuildAccelerator/Visualizer/Private/UbaVisualizer.h
+++ b/Engine/Source/Programs/UnrealBuildAccelerator/Visualizer/Private/UbaVisualizer.h
@@ -122,6 +122,7 @@ namespace uba
 		void HitTest(HitTestResult& outResult, const POINT& pos);
 
 		void WriteProcessStats(Logger& out, const TraceView::Process& process);
+		void WriteWorkStats(Logger& out, const TraceView::WorkRecord& record);
 		void CopyTextToClipboard(const TString& str);
 		void UnselectAndRedraw();
 		bool UpdateAutoscroll();
@@ -184,7 +185,6 @@ namespace uba
 		COLORREF m_memColor = {};
 		HBRUSH m_backgroundBrush = 0;
 		HBRUSH m_tooltipBackgroundBrush = 0;
-		HBRUSH m_workBrush = 0;
 		HPEN m_textPen = 0;
 		HPEN m_separatorPen = 0;
 		HPEN m_sendPen = 0;
@@ -274,6 +274,8 @@ namespace uba
 
 		Thread m_thread;
 
+		UnorderedMap<Color, HBRUSH> m_coloredBrushes;
+
 		void PostNewTrace(u32 replay, bool paused);
 		void PostNewTitle(const StringView& title);
 		void PostQuit();
