[UBA]
* Added ParseDirectoryTableNoLock that takes in from, to to populate DirectoryTable
* Added so DirectoryTable can take in path that ends with path separator.
* Fixed so DirectoryTable returns Exists_No if finding ancestor in path that is a file

diff --git a/Engine/Source/Programs/UnrealBuildAccelerator/Core/Public/UbaDirectoryTable.h b/Engine/Source/Programs/UnrealBuildAccelerator/Core/Public/UbaDirectoryTable.h
--- a/Engine/Source/Programs/UnrealBuildAccelerator/Core/Public/UbaDirectoryTable.h
+++ b/Engine/Source/Programs/UnrealBuildAccelerator/Core/Public/UbaDirectoryTable.h
@@ -33,20 +33,25 @@ namespace uba
 		{
 			if (size <= m_memorySize)
 				return;
-			BinaryReader reader(m_memory, m_memorySize);
+			ParseDirectoryTableNoLock(m_memorySize, size);
+			m_memorySize = size;
+		}
+
+		void ParseDirectoryTableNoLock(u32 from, u32 to)
+		{
+			BinaryReader reader(m_memory, from, to);
 			while (true)
 			{
 				u64 pos = reader.GetPosition();
-				if (pos == size)
+				if (pos == to)
 					break;
-				UBA_ASSERTF(pos < size, TC("Should never read past size (pos: %llu, size: %u)"), pos, size);
+				UBA_ASSERTF(pos < to, TC("Should never read past size (pos: %llu, size: %u)"), pos, to);
 				u64 storageSize = reader.Read7BitEncoded();
 				StringKey dirKey = reader.ReadStringKey();
 				auto insres = m_lookup.try_emplace(dirKey, m_memoryBlock); // Note that this is allowed to overwrite
 				insres.first->second.tableOffset = u32(reader.GetPosition());
 				reader.Skip(storageSize - sizeof(dirKey));
 			}
-			m_memorySize = size;
 		}
 
 		void PopulateDirectory(const StringKeyHasher& hasher, Directory& dir)
@@ -155,6 +160,7 @@ namespace uba
 
 		Exists EntryExistsNoLock(StringKey entryKey, const StringView& entryName, bool checkIfDir = false, u32* tableOffset = nullptr)
 		{
+			u32 startSkip = 2;
 			if (checkIfDir)
 			{
 				auto findIt = m_lookup.find(entryKey);
@@ -164,10 +170,11 @@ namespace uba
 						*tableOffset = u32(findIt->second.tableOffset) | 0x80000000; // Use significant bit to say that this is a dir
 					return Exists_Yes;
 				}
+				startSkip = 1;
 			}
 
 			// scan backwards first
-			const tchar* rit = entryName.data + entryName.count - 2;
+			const tchar* rit = entryName.data + entryName.count - startSkip;
 			bool inAncestor = false;
 			while (rit > entryName.data)
 			{
@@ -199,7 +206,13 @@ namespace uba
 					if (entryIt == parentDir.files.end())
 						return Exists_No;
 					if (inAncestor)
+					{
+						BinaryReader reader(m_memory, entryIt->second, m_memorySize);
+						reader.SkipString();
+						if (!IsDirectory(reader.ReadFileAttributes()))
+							return Exists_No;
 						return Exists_Maybe;
+					}
 					if (tableOffset)
 						*tableOffset = entryIt->second;
 					return Exists_Yes;
