[UBA]
* Fixed so NetworkBackendMemory ignores subsequent calls to Connect (always only have one connection)

diff --git a/Engine/Source/Programs/UnrealBuildAccelerator/Common/Private/UbaNetworkBackendMemory.cpp b/Engine/Source/Programs/UnrealBuildAccelerator/Common/Private/UbaNetworkBackendMemory.cpp
--- a/Engine/Source/Programs/UnrealBuildAccelerator/Common/Private/UbaNetworkBackendMemory.cpp
+++ b/Engine/Source/Programs/UnrealBuildAccelerator/Common/Private/UbaNetworkBackendMemory.cpp
@@ -138,7 +138,10 @@ namespace uba
 			Sleep(10);
 			l.Enter();
 		}
-		UBA_ASSERT(!m_connection);
+
+		if (m_connection) // Ignore multiple connects
+			return true;
+
 		m_connection = new Connection();
 		m_connection->connected = true;
 		if (!m_connectedFunc((void*)1, {}))
