[UBA]
* Added LargeHash (32 bytes)
* Fixed bug in check path logic

diff --git a/Engine/Source/Programs/UnrealBuildAccelerator/Core/Private/UbaHash.cpp b/Engine/Source/Programs/UnrealBuildAccelerator/Core/Private/UbaHash.cpp
--- a/Engine/Source/Programs/UnrealBuildAccelerator/Core/Private/UbaHash.cpp
+++ b/Engine/Source/Programs/UnrealBuildAccelerator/Core/Private/UbaHash.cpp
@@ -98,4 +98,34 @@ namespace uba
 		output[19] = compressed ? 1 : 0;
 		return (CasKey&)output;
 	}
+
+	void LargeHash::Set(const void* data, u64 dataLen)
+	{
+		static_assert(BLAKE3_OUT_LEN == sizeof(LargeHash));
+		blake3_hasher hasher;
+		blake3_hasher_init(&hasher);
+		blake3_hasher_update(&hasher, data, dataLen);
+		blake3_hasher_finalize(&hasher, (u8*)this, BLAKE3_OUT_LEN);
+	}
+
+	void LargeHash::BloomAdd(const LargeHash& lh)
+	{
+		a |= lh.a;
+		b |= lh.b;
+		c |= lh.c;
+		d |= lh.d;
+	}
+
+	bool LargeHash::BloomMiss(const LargeHash& lh) const
+	{
+		return (a & lh.a) != lh.a
+			|| (b & lh.b) != lh.b
+			|| (c & lh.c) != lh.c
+			|| (d & lh.d) != lh.d;
+	}
+
+	bool LargeHash::IsEmpty() const
+	{
+		return !a && !b && !c && !d;
+	}
 }
diff --git a/Engine/Source/Programs/UnrealBuildAccelerator/Core/Public/UbaHash.h b/Engine/Source/Programs/UnrealBuildAccelerator/Core/Public/UbaHash.h
--- a/Engine/Source/Programs/UnrealBuildAccelerator/Core/Public/UbaHash.h
+++ b/Engine/Source/Programs/UnrealBuildAccelerator/Core/Public/UbaHash.h
@@ -179,10 +179,10 @@ namespace uba
 	#if UBA_DEBUG
 	inline void CheckPath(const StringView& fileName)
 	{
-		auto pos = fileName.data;
 		tchar history[3] = { 0 };
-		while (tchar c = *pos++)
+		for (auto i = fileName.data, e = i + fileName.count; i != e; ++i)
 		{
+			tchar c = *i;
 			UBA_ASSERTF(!CaseInsensitiveFs || c < 'A' || c > 'Z', TC("Path is not valid (%.*s)"), fileName.count, fileName.data); // No upper case characters if case insensitive
 			UBA_ASSERTF(!(c == PathSeparator && history[0] == '.' && history[1] == PathSeparator) || (IsWindows && fileName.Contains(TC("\\\\.\\pipe"))), TC("Path is not valid (%.*s)"), fileName.count, fileName.data);  //  /./ not allowed
 			UBA_ASSERTF(!(c == PathSeparator && history[0] == '.' && history[1] == '.' && history[2] == PathSeparator), TC("Path is not valid (%.*s)"), fileName.count, fileName.data); // /../ not allowed
@@ -200,6 +200,19 @@ namespace uba
 	#else
 	#define CHECK_PATH(path) while (false) {}
 	#endif
+
+	struct LargeHash
+	{
+		u64 a = 0;
+		u64 b = 0;
+		u64 c = 0;
+		u64 d = 0;
+
+		void Set(const void* data, u64 dataLen);
+		void BloomAdd(const LargeHash& lh);
+		bool BloomMiss(const LargeHash& lh) const;
+		bool IsEmpty() const;
+	};
 }
 
 //template<> struct std::hash<uba::Guid> { size_t operator()(const uba::Guid& g) const { std::hash<uba::u64> hash; return hash(((uba::u64*)&g)[0]) ^ hash(((uba::u64*)&g)[1]); } };
