[UBA]
* Fixed so NtCreateFile called for path that is a directory but desired access to content returns STATUS_FILE_IS_A_DIRECTORY before accessing file

diff --git a/Engine/Source/Programs/UnrealBuildAccelerator/Detours/Private/Windows/UbaDetoursFunctionsNtDll.inl b/Engine/Source/Programs/UnrealBuildAccelerator/Detours/Private/Windows/UbaDetoursFunctionsNtDll.inl
--- a/Engine/Source/Programs/UnrealBuildAccelerator/Detours/Private/Windows/UbaDetoursFunctionsNtDll.inl
+++ b/Engine/Source/Programs/UnrealBuildAccelerator/Detours/Private/Windows/UbaDetoursFunctionsNtDll.inl
@@ -853,6 +853,11 @@ NTSTATUS NTAPI Shared_NtCreateFile(bool IsCreateFunc, PHANDLE hFileHandle, ACCES
 						DEBUG_LOG_DETOURED(funcName, L"DELETED %llu, (%ls) -> Success", uintptr_t(*hFileHandle), fileName.data);
 						return STATUS_OBJECT_NAME_NOT_FOUND;
 					}
+					else if (useContent && IsDirectory(entryInfo.attributes))
+					{
+						DEBUG_LOG_DETOURED(funcName, L"%llu, (%ls) -> STATUS_FILE_IS_A_DIRECTORY", uintptr_t(*hFileHandle), fileName.data);
+						return STATUS_FILE_IS_A_DIRECTORY;
+					}
 				}
 
 				bool isWriteAttributes = (DesiredAccess & FILE_WRITE_ATTRIBUTES) != 0;
