[UBA]
* Added support for setting connection count when StorageClient connects to proxy. Defaults to 4

diff --git a/Engine/Source/Programs/UnrealBuildAccelerator/Common/Private/UbaStorageClient.cpp b/Engine/Source/Programs/UnrealBuildAccelerator/Common/Private/UbaStorageClient.cpp
--- a/Engine/Source/Programs/UnrealBuildAccelerator/Common/Private/UbaStorageClient.cpp
+++ b/Engine/Source/Programs/UnrealBuildAccelerator/Common/Private/UbaStorageClient.cpp
@@ -37,6 +37,7 @@ namespace uba
 		table.GetValueAsBool(sendOneBigFileAtTheTime, TC("SendOneBigAtTheTime"));
 		table.GetValueAsBool(checkExistsOnServer, TC("CheckExistsOnServer"));
 		table.GetValueAsBool(resendCas, TC("ResendCas"));
+		table.GetValueAsU32(proxyConnectionCount, TC("ProxyConnectionCount"));
 	}
 
 	StorageClient::StorageClient(const StorageClientCreateInfo& info)
@@ -53,6 +54,7 @@ namespace uba
 	,	m_startProxyCallback(info.startProxyCallback)
 	,	m_startProxyUserData(info.startProxyUserData)
 	,	m_proxyPort(info.proxyPort)
+	,	m_proxyConnectionCount(info.proxyConnectionCount)
 	{
 	}
 
@@ -422,8 +424,8 @@ namespace uba
 						continue;
 					}
 					
-					//for (u32 i=0;i!=4;++i)
-					//	proxy->client.Connect(proxyBackend, proxyHost.data, proxyPort);
+					for (u32 i=1;i<m_proxyConnectionCount;++i)
+						proxy->client.Connect(proxyBackend, proxyHost.data, proxyPort);
 
 					u64 connectTime = GetTime() - startTime;
 					if (connectTime > MsToTime(2000))
@@ -974,7 +976,7 @@ namespace uba
 					workManager.AddWork([&, filePath = TString(fullPath.data)](const WorkContext& context)
 						{
 							PopulateCasFromDirsRecursive(filePath.c_str(), workManager, seenIds, seenIdsLock, shouldExit);
-						}, 1, TC(""));
+						}, 1, TC(""), ColorWork);
 					return;
 				}
 
@@ -1034,7 +1036,7 @@ namespace uba
 							localFile.fileName = filePath;
 						}
 
-					}, 1, TC(""));
+					}, 1, TC(""), ColorWork);
 			});
 		return true;
 	}
diff --git a/Engine/Source/Programs/UnrealBuildAccelerator/Common/Public/UbaStorageClient.h b/Engine/Source/Programs/UnrealBuildAccelerator/Common/Public/UbaStorageClient.h
--- a/Engine/Source/Programs/UnrealBuildAccelerator/Common/Public/UbaStorageClient.h
+++ b/Engine/Source/Programs/UnrealBuildAccelerator/Common/Public/UbaStorageClient.h
@@ -27,6 +27,7 @@ namespace uba
 		bool checkExistsOnServer = false; // There is no point having this enabled if server is not received cas files to disk (which is default)
 		bool resendCas = false; // Will try to send created file everytime even if the cas has been seen before
 		
+		u32 proxyConnectionCount = 4;
 		GetProxyBackendCallback* getProxyBackendCallback = nullptr;
 		void* getProxyBackendUserData = nullptr;
 		StartProxyCallback* startProxyCallback = nullptr;
@@ -93,6 +94,7 @@ namespace uba
 		void* m_getProxyBackendUserData = nullptr;
 		StartProxyCallback* m_startProxyCallback = nullptr;
 		void* m_startProxyUserData = nullptr;
+		u32 m_proxyConnectionCount;
 		u16 m_proxyPort = 0;
 	};
 
