[UBA]
* Added FileMappingBuffer::MapView that takes in string (same string that is transferred to detoured processes) which is resolved to handle and offset

diff --git a/Engine/Source/Programs/UnrealBuildAccelerator/Common/Private/UbaFileMapping.cpp b/Engine/Source/Programs/UnrealBuildAccelerator/Common/Private/UbaFileMapping.cpp
--- a/Engine/Source/Programs/UnrealBuildAccelerator/Common/Private/UbaFileMapping.cpp
+++ b/Engine/Source/Programs/UnrealBuildAccelerator/Common/Private/UbaFileMapping.cpp
@@ -614,6 +614,22 @@ namespace uba
 		return res;
 	}
 
+	MappedView FileMappingBuffer::MapView(const StringView& str, u64 size, const tchar* hint)
+	{
+		const tchar* handleStr = str.data + 1;
+		const tchar* handleStrEnd = wcschr(handleStr, '-');
+		if (!handleStrEnd)
+			return {};
+
+		HANDLE mappingHandle = (HANDLE)StringToValue(handleStr, handleStrEnd - handleStr);
+		const tchar* mappingOffsetStr = handleStrEnd + 1;
+		u64 mappingOffset = StringToValue(mappingOffsetStr, wcslen(mappingOffsetStr));
+		
+		FileMappingHandle h;
+		h.handle = mappingHandle;
+		return MapView(h, mappingOffset, size, hint);
+	}
+
 	void FileMappingBuffer::UnmapView(MappedView view, const tchar* hint_, u64 newSize)
 	{
 		if (!view.handle.IsValid())
diff --git a/Engine/Source/Programs/UnrealBuildAccelerator/Common/Public/UbaFileMapping.h b/Engine/Source/Programs/UnrealBuildAccelerator/Common/Public/UbaFileMapping.h
--- a/Engine/Source/Programs/UnrealBuildAccelerator/Common/Public/UbaFileMapping.h
+++ b/Engine/Source/Programs/UnrealBuildAccelerator/Common/Public/UbaFileMapping.h
@@ -83,6 +83,7 @@ namespace uba
 
 		MappedView AllocAndMapView(FileMappingType type, u64 size, u64 alignment, const tchar* hint, bool allowShrink = false);
 		MappedView MapView(FileMappingHandle handle, u64 offset, u64 size, const tchar* hint);
+		MappedView MapView(const StringView& str, u64 size, const tchar* hint);
 		void UnmapView(MappedView view, const tchar* hint, u64 newSize = InvalidValue);
 		
 		void GetSizeAndCount(FileMappingType type, u64& outSize, u32& outCount);
