[UBA]
* Commented out log entry with information about server already having cas file

diff --git a/Engine/Source/Programs/UnrealBuildAccelerator/Common/Private/UbaStorageUtils.cpp b/Engine/Source/Programs/UnrealBuildAccelerator/Common/Private/UbaStorageUtils.cpp
--- a/Engine/Source/Programs/UnrealBuildAccelerator/Common/Private/UbaStorageUtils.cpp
+++ b/Engine/Source/Programs/UnrealBuildAccelerator/Common/Private/UbaStorageUtils.cpp
@@ -256,7 +256,7 @@ namespace uba
 
 				if (storeId == u16(~0)) // File already exists on server
 				{
-					logger.Info(TC("Server already has file %s (%s)"), CasKeyString(casKey).str, hint);
+					//logger.Info(TC("Server already has file %s (%s)"), CasKeyString(casKey).str, hint);
 					return sendEndMessage();
 				}
 
@@ -411,7 +411,7 @@ namespace uba
 
 				if (storeId == u16(~0)) // File already exists on server
 				{
-					m_logger.Info(TC("Server already has file %s (%s)"), CasKeyString(casKey).str, hint);
+					//m_logger.Info(TC("Server already has file %s (%s)"), CasKeyString(casKey).str, hint);
 					return sendEndMessage();
 				}
 
