[UBA]
* Changed session server knowninput cas calculations to happen in parallel and waited on just before data is sent to client.

diff --git a/Engine/Source/Programs/UnrealBuildAccelerator/Common/Private/UbaSessionServer.cpp b/Engine/Source/Programs/UnrealBuildAccelerator/Common/Private/UbaSessionServer.cpp
--- a/Engine/Source/Programs/UnrealBuildAccelerator/Common/Private/UbaSessionServer.cpp
+++ b/Engine/Source/Programs/UnrealBuildAccelerator/Common/Private/UbaSessionServer.cpp
@@ -24,6 +24,8 @@ namespace uba
 
 		~RemoteProcess()
 		{
+			if (m_knownInputsDone.IsCreated())
+				m_knownInputsDone.IsSet(50*1000);
 			delete[] m_knownInputs;
 		}
 
@@ -91,6 +93,7 @@ namespace uba
 		struct KnownInput { CasKey key; u32 mappingAlignment = 0; bool allowProxy = true; };
 		KnownInput* m_knownInputs = nullptr;
 		u32 m_knownInputsCount = 0;
+		Event m_knownInputsDone;
 	};
 
 	void SessionServerCreateInfo::Apply(const Config& config)
@@ -302,6 +305,8 @@ namespace uba
 
 	ProcessHandle SessionServer::RunProcessRemote(const ProcessStartInfo& startInfo, float weight, const void* knownInputs, u32 knownInputsCount, bool allowCrossArchitecture)
 	{
+		//TrackWorkScope tws(m_trace, AsView(TC("RunProcessRemote")), ColorWork);
+
 		UBA_ASSERT(!startInfo.startSuspended);
 
 		FlushDeadProcesses();
@@ -316,40 +321,80 @@ namespace uba
 
 		if (knownInputsCount)
 		{
-			auto keys = remoteProcess->m_knownInputs = new RemoteProcess::KnownInput[knownInputsCount];
+			remoteProcess->m_knownInputsDone.Create(true);
 
-			u32 keysIndex = 0;
-			const TString& workingDir = remoteProcess->m_startInfo.workingDirStr;
-			for (auto kiIt = (const tchar*)knownInputs; *kiIt; kiIt += TStrlen(kiIt) + 1)
-			{
-				StringBuffer<> fileName;
-				FixPath(kiIt, workingDir.c_str(), u32(workingDir.size()), fileName);
-				StringKey fileNameKey = CaseInsensitiveFs ? ToStringKeyLower(fileName) : ToStringKey(fileName);
 
-				// Make sure cas entry exists and caskey is calculated (cas content creation is deferred in case client already has it)
-				CasKey casKey;
-				if (!StoreCasFile(casKey, fileNameKey, fileName.data) || casKey == CasKeyZero)
-					continue;
+			auto kiIt = (const tchar*)knownInputs;
+			for (; *kiIt; kiIt += TStrlen(kiIt) + 1) {}
+			u32 knownInputsSize = u32((const u8*)kiIt - (const u8*)knownInputs) + sizeof(tchar);
+			void* knownInputsCopy = malloc(knownInputsSize);
+			memcpy(knownInputsCopy, knownInputs, knownInputsSize);
+
+			m_server.AddWork([remoteProcess, this, knownInputsCopy, knownInputsCount, rules](const WorkContext& context)
+				{
+					auto keys = remoteProcess->m_knownInputs = new RemoteProcess::KnownInput[knownInputsCount];
 
-				auto& ki = keys[keysIndex++];
-				ki.key = casKey;
-				ki.mappingAlignment = GetMemoryMapAlignment(fileName, true);
-				ki.allowProxy = rules->AllowStorageProxy(fileName);
+					struct Iterator
+					{
+						Iterator() : it(nullptr) {}
+						Iterator(const tchar* ki) : it(ki) {}
+						Iterator operator++(int) { auto prev = it; it += TStrlen(it) + 1; return Iterator(prev); }
+						const tchar* operator*() { return it; }
+						bool operator==(const Iterator& o) const { return *it == *o.it; }
+						const tchar* it;
+					};
+					struct Container
+					{
+						using iterator = Iterator;
+						Container(const tchar* ki) : it(ki) {}
+						Iterator begin() const { return Iterator(it); }
+						Iterator end() const { return Iterator(TC("")); }
+						const tchar* it;
+					};
 
+					Container container((const tchar*)knownInputsCopy);
 
-				// Update name to hash table
-				SCOPED_WRITE_LOCK(m_nameToHashLookupLock, lock);
-				CasKey& lookupCasKey = m_nameToHashLookup[fileNameKey];
-				if (lookupCasKey != casKey)
-				{
-					lookupCasKey = casKey;
-					BinaryWriter w(m_nameToHashTableMem.memory, m_nameToHashTableMem.writtenSize, NameToHashMemSize);
-					m_nameToHashTableMem.AllocateNoLock(sizeof(StringKey) + sizeof(CasKey), 1, TC("NameToHashTable"));
-					w.WriteStringKey(fileNameKey);
-					w.WriteCasKey(lookupCasKey);
-				}
-			}
-			remoteProcess->m_knownInputsCount = keysIndex;
+					Atomic<u32> keysIndex;
+					const TString& workingDir = remoteProcess->m_startInfo.workingDirStr;
+
+					m_server.ParallelFor(knownInputsCount, container, [&](auto& it)
+						{
+							//TrackWorkScope tws2(m_trace, AsView(TC("KnownInputs")), ColorWork);
+							StringBuffer<> fileName;
+							FixPath(*it, workingDir.c_str(), u32(workingDir.size()), fileName);
+							StringKey fileNameKey = CaseInsensitiveFs ? ToStringKeyLower(fileName) : ToStringKey(fileName);
+
+							//tws2.AddHint(fileName);
+
+							// Make sure cas entry exists and caskey is calculated (cas content creation is deferred in case client already has it)
+							CasKey casKey;
+							if (!StoreCasFile(casKey, fileNameKey, fileName.data) || casKey == CasKeyZero)
+								return;
+
+							auto& ki = keys[keysIndex++];
+							ki.key = casKey;
+							ki.mappingAlignment = GetMemoryMapAlignment(fileName, true);
+							ki.allowProxy = rules->AllowStorageProxy(fileName);
+
+
+							// Update name to hash table
+							SCOPED_WRITE_LOCK(m_nameToHashLookupLock, lock);
+							CasKey& lookupCasKey = m_nameToHashLookup[fileNameKey];
+							if (lookupCasKey != casKey)
+							{
+								lookupCasKey = casKey;
+								BinaryWriter w(m_nameToHashTableMem.memory, m_nameToHashTableMem.writtenSize, NameToHashMemSize);
+								m_nameToHashTableMem.AllocateNoLock(sizeof(StringKey) + sizeof(CasKey), 1, TC("NameToHashTable"));
+								w.WriteStringKey(fileNameKey);
+								w.WriteCasKey(lookupCasKey);
+							}
+						}, TC("KnownInputsLoop"), true);
+
+					remoteProcess->m_knownInputsCount = keysIndex;
+					remoteProcess->m_knownInputsDone.Set();
+					free(knownInputsCopy);
+
+				}, 1, TC("KnownInputs"));
 		}
 
 		ProcessHandle h(remoteProcess); // Keep ref count up even if process is removed by callbacks etc.
@@ -1256,9 +1301,12 @@ namespace uba
 			writer.WriteU32(process->m_processId);
 			startInfo.Write(writer, applicationOverride);
 
-			for (auto kiIt = process->m_knownInputs, kiEnd = kiIt + process->m_knownInputsCount; kiIt!=kiEnd; ++kiIt)
-				if (session.sentKeys.insert(kiIt->key).second)
-					knownInputsToSend.push_back(kiIt);
+
+			if (process->m_knownInputsDone.IsCreated())
+				if (process->m_knownInputsDone.IsSet(50*1000))
+					for (auto kiIt = process->m_knownInputs, kiEnd = kiIt + process->m_knownInputsCount; kiIt!=kiEnd; ++kiIt)
+						if (session.sentKeys.insert(kiIt->key).second)
+							knownInputsToSend.push_back(kiIt);
 
 			++addCount;
 
@@ -1870,6 +1918,7 @@ namespace uba
 
 	SessionServer::RemoteProcess* SessionServer::DequeueProcess(ClientSession& session, u32 sessionId, u32 clientId)
 	{
+		//TrackWorkScope tws(m_trace, AsView(TC("DequeueProcess")), ColorWork);
 		SCOPED_READ_LOCK(m_remoteProcessSlotAvailableEventLock, lock);
 		bool hasCalledCallback = !m_remoteProcessSlotAvailableEvent;
 
