[UBA]
* Improved code around when proxy client gets disconnected. Theory is that if multiple threads are sitting in this function they can spin for a long time before actually come to a resolution

diff --git a/Engine/Source/Programs/UnrealBuildAccelerator/Common/Private/UbaStorageClient.cpp b/Engine/Source/Programs/UnrealBuildAccelerator/Common/Private/UbaStorageClient.cpp
--- a/Engine/Source/Programs/UnrealBuildAccelerator/Common/Private/UbaStorageClient.cpp
+++ b/Engine/Source/Programs/UnrealBuildAccelerator/Common/Private/UbaStorageClient.cpp
@@ -300,11 +300,21 @@ namespace uba
 						proxy = m_proxyClient;
 						client = &proxy->client;
 					}
-					else if (!m_proxyClient->refCount)
+					else
 					{
-						delete m_proxyClient;
-						m_proxyClient = nullptr;
-						m_lastTestedProxyIp.clear();
+						while (m_proxyClient && m_proxyClient->refCount)
+						{
+							proxyLock.Leave();
+							Sleep(200);
+							proxyLock.Enter();
+						}
+
+						if (m_proxyClient && !m_proxyClient->refCount)
+						{
+							delete m_proxyClient;
+							m_proxyClient = nullptr;
+							m_lastTestedProxyIp.clear();
+						}
 					}
 				}
 				wantsProxy = proxy == nullptr && m_startProxyCallback;
@@ -420,7 +430,7 @@ namespace uba
 					u64 startTime = GetTime();
 					if (!proxy->client.Connect(proxyBackend, proxyHost.data, proxyPort))
 					{
-						m_logger.Detail(TC("Redirection to proxy %s:%u for cas %s download failed! (%s)"), proxyHost.data, proxyPort, casFile.data, hint);
+						m_logger.Detail(TC("Connecting to proxy %s:%u for cas %s download failed! (%s) (%s)"), proxyHost.data, proxyPort, casFile.data, hint, TimeToText(GetTime() - startTime).str);
 						continue;
 					}
 					
@@ -441,7 +451,10 @@ namespace uba
 						proxyWriter.WriteBool(isInProcessClient);
 						StackBinaryReader<256> proxyReader;
 						if (!proxyMsg.Send(proxyReader))
+						{
+							m_logger.Info(TC("Failed to send connect message to proxy %s:%u. Will ask storage server for new proxy"), proxyHost.data, proxyPort);
 							continue;
+						}
 						if (proxyReader.ReadGuid() != m_storageServerUid)
 						{
 							m_logger.Info(TC("Proxy %s:%u is not the correct proxy anymore. Will ask storage server for new proxy"), proxyHost.data, proxyPort);
