[UBA]
* Better fix for handling scenario where proxy client is disconnected while having lots of active retrieves.

diff --git a/Engine/Source/Programs/UnrealBuildAccelerator/Common/Private/UbaStorageClient.cpp b/Engine/Source/Programs/UnrealBuildAccelerator/Common/Private/UbaStorageClient.cpp
--- a/Engine/Source/Programs/UnrealBuildAccelerator/Common/Private/UbaStorageClient.cpp
+++ b/Engine/Source/Programs/UnrealBuildAccelerator/Common/Private/UbaStorageClient.cpp
@@ -291,31 +291,31 @@ namespace uba
 			if (allowProxy && m_allowProxy)
 			{
 				SCOPED_WRITE_LOCK(m_proxyClientLock, proxyLock);
-				if (m_proxyClient)
+				while (true)
 				{
+					if (!m_proxyClient)
+						break;
+
 					if (m_proxyClient->client.IsConnected())
 					{
 						m_proxyClientKeepAliveTime = GetTime();
 						++m_proxyClient->refCount;
 						proxy = m_proxyClient;
 						client = &proxy->client;
+						break;
 					}
-					else
-					{
-						while (m_proxyClient && m_proxyClient->refCount)
-						{
-							proxyLock.Leave();
-							Sleep(200);
-							proxyLock.Enter();
-						}
 
-						if (m_proxyClient && !m_proxyClient->refCount)
-						{
-							delete m_proxyClient;
-							m_proxyClient = nullptr;
-							m_lastTestedProxyIp.clear();
-						}
+					if (m_proxyClient && !m_proxyClient->refCount)
+					{
+						delete m_proxyClient;
+						m_proxyClient = nullptr;
+						m_lastTestedProxyIp.clear();
+						break;
 					}
+
+					proxyLock.Leave();
+					Sleep(200);
+					proxyLock.Enter();
 				}
 				wantsProxy = proxy == nullptr && m_startProxyCallback;
 			}
