[UBA]
* Changed so WorkContext have a TrackWorkScope reference stored.

diff --git a/Engine/Source/Programs/UnrealBuildAccelerator/Common/Private/UbaCacheClient.cpp b/Engine/Source/Programs/UnrealBuildAccelerator/Common/Private/UbaCacheClient.cpp
--- a/Engine/Source/Programs/UnrealBuildAccelerator/Common/Private/UbaCacheClient.cpp
+++ b/Engine/Source/Programs/UnrealBuildAccelerator/Common/Private/UbaCacheClient.cpp
@@ -698,7 +698,7 @@ namespace uba
 							TimerScope fts(cacheStats.fetchOutput);
 							if (!FetchFile(bucket, rootPaths, info, cacheStats, storageStats, casKeyOffset))
 								fetchSuccess = false;
-						}, TC(""));
+						}, AsView(TC("CacheFetchOutput")));
 					if (!fetchSuccess)
 						return false;
 				}
diff --git a/Engine/Source/Programs/UnrealBuildAccelerator/Common/Private/UbaCacheServer.cpp b/Engine/Source/Programs/UnrealBuildAccelerator/Common/Private/UbaCacheServer.cpp
--- a/Engine/Source/Programs/UnrealBuildAccelerator/Common/Private/UbaCacheServer.cpp
+++ b/Engine/Source/Programs/UnrealBuildAccelerator/Common/Private/UbaCacheServer.cpp
@@ -1030,7 +1030,7 @@ namespace uba
 					bucket.m_cacheEntryLookup.erase(key);
 
 				totalEntryCount += bucket.totalEntryCount;
-			}, TC(""), true);
+			}, AsView(TC("")), true);
 
 			// Reset deleted cas files and update it again..
 			deletedCasFiles.clear();
@@ -1242,7 +1242,7 @@ namespace uba
 			SCOPED_WRITE_LOCK(globalStatsLock, l);
 			maxCommittedMemory = Max(maxCommittedMemory, memoryBlock.writtenSize);
 
-		}, TC(""), true);
+		}, AsView(TC("")), true);
 
 		// Need to make sure all cas entries are dropped before saving cas table
 		u64 dropStartTime = GetTime();
diff --git a/Engine/Source/Programs/UnrealBuildAccelerator/Common/Private/UbaNetworkServer.cpp b/Engine/Source/Programs/UnrealBuildAccelerator/Common/Private/UbaNetworkServer.cpp
--- a/Engine/Source/Programs/UnrealBuildAccelerator/Common/Private/UbaNetworkServer.cpp
+++ b/Engine/Source/Programs/UnrealBuildAccelerator/Common/Private/UbaNetworkServer.cpp
@@ -510,7 +510,7 @@ namespace uba
 			connection.SetShouldDisconnect();
 			size = ErrorSize;
 		}
-		else if (!rec.func({&connection}, {tws.id}, mi, reader, writer))
+		else if (!rec.func({&connection}, {tws}, mi, reader, writer))
 		{
 			if (connection.SetShouldDisconnect())
 				server.m_logger.Error(TC("WORKER FUNCTION FAILED. id: %u, serviceid: %u type: %s, client: %u"), context.id, context.serviceId, rec.toString(context.messageType).data, connection.m_client->id);
@@ -583,7 +583,7 @@ namespace uba
 				lock.Leave();
 
 				TrackWorkScope tws(server, work.desc, ColorWork);
-				work.func({tws.id});
+				work.func({tws});
 			}
 
 			// Both locks needs to be taken to verify if additional work
@@ -990,7 +990,7 @@ namespace uba
 
 		TrackWorkScope tws(*this, work.desc, ColorWork);
 
-		work.func({tws.id});
+		work.func({tws});
 
 		return true;
 	}
diff --git a/Engine/Source/Programs/UnrealBuildAccelerator/Common/Private/UbaSessionServer.cpp b/Engine/Source/Programs/UnrealBuildAccelerator/Common/Private/UbaSessionServer.cpp
--- a/Engine/Source/Programs/UnrealBuildAccelerator/Common/Private/UbaSessionServer.cpp
+++ b/Engine/Source/Programs/UnrealBuildAccelerator/Common/Private/UbaSessionServer.cpp
@@ -388,7 +388,7 @@ namespace uba
 								w.WriteStringKey(fileNameKey);
 								w.WriteCasKey(lookupCasKey);
 							}
-						}, TC("KnownInputsLoop"), true);
+						}, AsView(TC("KnownInputsLoop")), true);
 
 					remoteProcess->m_knownInputsCount = keysIndex;
 					remoteProcess->m_knownInputsDone.Set();
@@ -947,7 +947,7 @@ namespace uba
 		reader.ReadString(fileName);
 		StringKey fileNameKey = reader.ReadStringKey();
 
-		m_server.TrackWorkHint(workContext.id, StringView(fileName).GetFileName());
+		workContext.tracker.AddHint(StringView(fileName).GetFileName());
 
 		CasKey casKey;
 		if (!StoreCasFile(casKey, fileNameKey, fileName.data))
diff --git a/Engine/Source/Programs/UnrealBuildAccelerator/Common/Private/UbaStorage.cpp b/Engine/Source/Programs/UnrealBuildAccelerator/Common/Private/UbaStorage.cpp
--- a/Engine/Source/Programs/UnrealBuildAccelerator/Common/Private/UbaStorage.cpp
+++ b/Engine/Source/Programs/UnrealBuildAccelerator/Common/Private/UbaStorage.cpp
@@ -542,7 +542,10 @@ namespace uba
 
 		rec->refCount = workerCount + 1; // We need to keep refcount up 1 to make sure it is not deleted before we read rec->written
 		m_workManager->AddWork(work, workerCount-1, TC("Compress")); // We are a worker ourselves
-		work({});
+		{
+			TrackWorkScope tws;
+			work({tws});
+		}
 		rec->events[rec->workCount - 1].IsSet();
 
 		totalWritten += rec->written;
@@ -748,7 +751,10 @@ namespace uba
 			}
 
 			TimerScope ts(stats.decompressToMem);
-			work({});
+			{
+				TrackWorkScope tws;
+				work({tws});
+			}
 			rec->done.IsSet();
 			bool success = !rec->error;
 			if (!success)
diff --git a/Engine/Source/Programs/UnrealBuildAccelerator/Common/Private/UbaStorageProxy.cpp b/Engine/Source/Programs/UnrealBuildAccelerator/Common/Private/UbaStorageProxy.cpp
--- a/Engine/Source/Programs/UnrealBuildAccelerator/Common/Private/UbaStorageProxy.cpp
+++ b/Engine/Source/Programs/UnrealBuildAccelerator/Common/Private/UbaStorageProxy.cpp
@@ -156,7 +156,7 @@ namespace uba
 				StringBuffer<> hint;
 				reader.ReadString(hint);
 
-				m_server.TrackWorkHint(workContext.id, StringView(hint).GetFileName());
+				workContext.tracker.AddHint(StringView(hint).GetFileName());
 
 				SCOPED_WRITE_LOCK(m_filesLock, filesLock);
 				auto insres = m_files.try_emplace(casKey);
diff --git a/Engine/Source/Programs/UnrealBuildAccelerator/Common/Private/UbaStorageServer.cpp b/Engine/Source/Programs/UnrealBuildAccelerator/Common/Private/UbaStorageServer.cpp
--- a/Engine/Source/Programs/UnrealBuildAccelerator/Common/Private/UbaStorageServer.cpp
+++ b/Engine/Source/Programs/UnrealBuildAccelerator/Common/Private/UbaStorageServer.cpp
@@ -432,7 +432,7 @@ namespace uba
 				StringBuffer<> hint;
 				reader.ReadString(hint);
 
-				m_server.TrackWorkHint(workContext.id, StringView(hint).GetFileName());
+				workContext.tracker.AddHint(StringView(hint).GetFileName());
 
 				casKey = AsCompressed(casKey, m_storeCompressed);
 
diff --git a/Engine/Source/Programs/UnrealBuildAccelerator/Common/Private/UbaStorageUtils.cpp b/Engine/Source/Programs/UnrealBuildAccelerator/Common/Private/UbaStorageUtils.cpp
--- a/Engine/Source/Programs/UnrealBuildAccelerator/Common/Private/UbaStorageUtils.cpp
+++ b/Engine/Source/Programs/UnrealBuildAccelerator/Common/Private/UbaStorageUtils.cpp
@@ -78,7 +78,10 @@ namespace uba
 				workManager->AddWork(work, workerCount, TC("CalculateKey"));
 			}
 
-			work({});
+			{
+				TrackWorkScope tws;
+				work({tws});
+			}
 			rec->done.IsSet();
 
 			hasher.Update(rec->keys.data(), rec->keys.size()*sizeof(CasKey));
diff --git a/Engine/Source/Programs/UnrealBuildAccelerator/Common/Private/UbaWorkManager.cpp b/Engine/Source/Programs/UnrealBuildAccelerator/Common/Private/UbaWorkManager.cpp
--- a/Engine/Source/Programs/UnrealBuildAccelerator/Common/Private/UbaWorkManager.cpp
+++ b/Engine/Source/Programs/UnrealBuildAccelerator/Common/Private/UbaWorkManager.cpp
@@ -70,10 +70,10 @@ namespace uba
 
 						#if UBA_TRACK_WORK
 						TrackWorkScope tws(manager, work.desc);
-						work.func({tws.id});
 						#else
-						work.func({});
+						TrackWorkScope tws;
 						#endif
+						work.func({tws});
 					}
 
 					SCOPED_WRITE_LOCK(manager.m_availableWorkersLock, lock1);
@@ -177,10 +177,10 @@ namespace uba
 
 			#if UBA_TRACK_WORK
 			TrackWorkScope tws(*this, work.desc);
-			work.func({tws.id});
 			#else
-			work.func({});
+			TrackWorkScope tws;
 			#endif
+			work.func({tws});
 		}
 	}
 
