[UBA]
* Missed file in previous submit

diff --git a/Engine/Source/Programs/UnrealBuildAccelerator/Common/Public/UbaWorkManager.h b/Engine/Source/Programs/UnrealBuildAccelerator/Common/Public/UbaWorkManager.h
--- a/Engine/Source/Programs/UnrealBuildAccelerator/Common/Public/UbaWorkManager.h
+++ b/Engine/Source/Programs/UnrealBuildAccelerator/Common/Public/UbaWorkManager.h
@@ -3,11 +3,13 @@
 #pragma once
 
 #include "UbaEvent.h"
-#include "UbaMemory.h"
+#include "UbaStringBuffer.h"
 
 namespace uba
 {
 	struct StringView;
+	struct TrackWorkScope;
+
 	using Color = u32;
 	constexpr inline Color ToColor(u8 r, u8 g, u8 b) { return (r << 16) + (g << 8) + b; }
 	constexpr Color ColorWhite = ToColor(255, 255, 255);
@@ -23,10 +25,10 @@ namespace uba
 
 	struct WorkContext
 	{
-		u32 id;
+		TrackWorkScope& tracker;
 	};
 
-	class WorkManager
+	class WorkManager : public WorkTracker
 	{
 	public:
 		using WorkFunction = Function<void(const WorkContext& context)>;
@@ -42,7 +44,7 @@ namespace uba
 		WorkTracker* GetWorkTracker() { return m_workTracker; }
 
 		template<typename TContainer, typename TFunc>
-		void ParallelFor(u32 workCount, TContainer& container, const TFunc& func, const tchar* description = TC(""), bool highPriority = false);
+		void ParallelFor(u32 workCount, TContainer& container, const TFunc& func, const StringView& description = AsView(TC("")), bool highPriority = false);
 
 	protected:
 		Atomic<WorkTracker*> m_workTracker;
@@ -78,17 +80,18 @@ namespace uba
 
 	struct TrackWorkScope
 	{
-		TrackWorkScope(WorkManager& wm, const StringView& desc, const Color& color) : workManager(wm), id(wm.TrackWorkStart(desc, color)) {}
-		void AddHint(const StringView& hint) { workManager.TrackWorkHint(id, hint); }
-		~TrackWorkScope() { workManager.TrackWorkEnd(id); }
-		WorkManager& workManager;
+		TrackWorkScope() : tracker(nullptr), id(0) {}
+		TrackWorkScope(WorkTracker& t, const StringView& desc, const Color& color) : tracker(&t), id(t.TrackWorkStart(desc, color)) {}
+		void AddHint(const StringView& hint) { if (tracker) tracker->TrackWorkHint(id, hint); }
+		~TrackWorkScope() { if (tracker) tracker->TrackWorkEnd(id); }
+		WorkTracker* tracker;
 		u32 id;
 	};
 
 
 
 	template<typename TContainer, typename TFunc>
-	void WorkManager::ParallelFor(u32 workCount, TContainer& container, const TFunc& func, const tchar* description, bool highPriority)
+	void WorkManager::ParallelFor(u32 workCount, TContainer& container, const TFunc& func, const StringView& description, bool highPriority)
 	{
 		#if !defined( __clang_analyzer__ ) // Static analyzer claims context can leak but it can only leak when process terminates (and then it doesn't matter)
 
@@ -144,8 +147,11 @@ namespace uba
 				}
 			};
 		
-		AddWork(work, workCount, description, ColorWork, highPriority);
-		work({});
+		AddWork(work, workCount, description.data, ColorWork, highPriority);
+		{
+			TrackWorkScope tws(*this, description, ColorWork);
+			work({tws});
+		}
 		doneEvent.IsSet();
 
 		#endif
