[UBA]
* Fixed the code that fetches known inputs in parallel.. was a bug in the code

diff --git a/Engine/Source/Programs/UnrealBuildAccelerator/Common/Private/UbaSessionServer.cpp b/Engine/Source/Programs/UnrealBuildAccelerator/Common/Private/UbaSessionServer.cpp
--- a/Engine/Source/Programs/UnrealBuildAccelerator/Common/Private/UbaSessionServer.cpp
+++ b/Engine/Source/Programs/UnrealBuildAccelerator/Common/Private/UbaSessionServer.cpp
@@ -324,35 +324,39 @@ namespace uba
 			remoteProcess->m_knownInputsDone.Create(true);
 
 
-			auto kiIt = (const tchar*)knownInputs;
-			for (; *kiIt; kiIt += TStrlen(kiIt) + 1) {}
-			u32 knownInputsSize = u32((const u8*)kiIt - (const u8*)knownInputs) + sizeof(tchar);
-			void* knownInputsCopy = malloc(knownInputsSize);
-			memcpy(knownInputsCopy, knownInputs, knownInputsSize);
+			auto kiBegin = (const tchar*)knownInputs;
+			auto kiEnd = kiBegin;
+			for (u32 i=0;i!=knownInputsCount; ++i)
+				kiEnd += TStrlen(kiEnd) + 1;
+
+			u32 knownInputsBytes = u32((kiEnd - kiBegin)*sizeof(tchar));
+			void* knownInputsCopy = malloc(knownInputsBytes);
+			memcpy(knownInputsCopy, knownInputs, knownInputsBytes);
 
 			m_server.AddWork([remoteProcess, this, knownInputsCopy, knownInputsCount, rules](const WorkContext& context)
 				{
 					auto keys = remoteProcess->m_knownInputs = new RemoteProcess::KnownInput[knownInputsCount];
 
-					struct Iterator
-					{
-						Iterator() : it(nullptr) {}
-						Iterator(const tchar* ki) : it(ki) {}
-						Iterator operator++(int) { auto prev = it; it += TStrlen(it) + 1; return Iterator(prev); }
-						const tchar* operator*() { return it; }
-						bool operator==(const Iterator& o) const { return *it == *o.it; }
-						const tchar* it;
-					};
 					struct Container
 					{
-						using iterator = Iterator;
-						Container(const tchar* ki) : it(ki) {}
-						Iterator begin() const { return Iterator(it); }
-						Iterator end() const { return Iterator(TC("")); }
-						const tchar* it;
+						struct iterator
+						{
+							iterator() : ptr(nullptr), index(0) {}
+							iterator(const tchar* p, u32 i) : ptr(p), index(i) {}
+							iterator operator++(int) { auto prev = ptr; ptr += TStrlen(ptr) + 1; return iterator(prev, index++); }
+							const tchar* operator*() { return ptr; }
+							bool operator==(const iterator& o) const { return index == o.index; }
+							const tchar* ptr;
+							u32 index;
+						};
+						Container(const tchar* b, u32 c) : ptr(b), count(c) {}
+						iterator begin() const { return iterator(ptr, 0); }
+						iterator end() const { return iterator(ptr, count); }
+						const tchar* ptr;
+						u32 count;
 					};
 
-					Container container((const tchar*)knownInputsCopy);
+					Container container((const tchar*)knownInputsCopy, knownInputsCount);
 
 					Atomic<u32> keysIndex;
 					const TString& workingDir = remoteProcess->m_startInfo.workingDirStr;
@@ -371,6 +375,7 @@ namespace uba
 							if (!StoreCasFile(casKey, fileNameKey, fileName.data) || casKey == CasKeyZero)
 								return;
 
+							UBA_ASSERT(keysIndex < knownInputsCount);
 							auto& ki = keys[keysIndex++];
 							ki.key = casKey;
 							ki.mappingAlignment = GetMemoryMapAlignment(fileName, true);
