[UBA]
* Changed so helper assigned as proxy will free up cpu based on how many helpers connected to proxy

diff --git a/Engine/Source/Programs/UnrealBuildAccelerator/Agent/Private/UbaAgent.cpp b/Engine/Source/Programs/UnrealBuildAccelerator/Agent/Private/UbaAgent.cpp
--- a/Engine/Source/Programs/UnrealBuildAccelerator/Agent/Private/UbaAgent.cpp
+++ b/Engine/Source/Programs/UnrealBuildAccelerator/Agent/Private/UbaAgent.cpp
@@ -1301,9 +1301,20 @@ namespace uba
 
 				// If we are the proxy server and have external connections we lower max process count. Note that it will always have one connection which is itself
 				if (auto proxyServer = proxyNetworkServer.load())
-					if (proxyServer->GetConnectionCount() > 1)
-						sessionClient->SetMaxProcessCount(maxProcessCount - 2);
+				{
+					u32 clientCount = proxyServer->GetClientCount();
+					if (clientCount > 1) // when having a proxy the agent itself is always connected to it
+					{
+						u32 processToFree = (clientCount - 1)/3 + 1; // Always free one, and then one per three helpers.. (so 16 helpers connected will remove 6 process count)
+						u32 newProcessCount;
+						if (processToFree < maxProcessCount)
+							newProcessCount = maxProcessCount - processToFree;
+						else
+							newProcessCount = 1;
 
+						sessionClient->SetMaxProcessCount(newProcessCount);
+					}
+				}
 				// This is an estimation based on tcp limitations (ack and sliding windows).
 				// For every 15ms latency on "best ping") we increase targetConnectionCount up to maxConnectionCount
 				if (!storageClient->IsUsingProxy())
diff --git a/Engine/Source/Programs/UnrealBuildAccelerator/Common/Private/UbaNetworkServer.cpp b/Engine/Source/Programs/UnrealBuildAccelerator/Common/Private/UbaNetworkServer.cpp
--- a/Engine/Source/Programs/UnrealBuildAccelerator/Common/Private/UbaNetworkServer.cpp
+++ b/Engine/Source/Programs/UnrealBuildAccelerator/Common/Private/UbaNetworkServer.cpp
@@ -904,6 +904,12 @@ namespace uba
 		return m_sendTimer;
 	}
 
+	u32 NetworkServer::GetClientCount()
+	{
+		SCOPED_READ_LOCK(m_clientsLock, lock)
+		return u32(m_clients.size());
+	}
+
 	u32 NetworkServer::GetConnectionCount()
 	{
 		SCOPED_READ_LOCK(m_connectionsLock, lock);
diff --git a/Engine/Source/Programs/UnrealBuildAccelerator/Common/Public/UbaNetworkServer.h b/Engine/Source/Programs/UnrealBuildAccelerator/Common/Public/UbaNetworkServer.h
--- a/Engine/Source/Programs/UnrealBuildAccelerator/Common/Public/UbaNetworkServer.h
+++ b/Engine/Source/Programs/UnrealBuildAccelerator/Common/Public/UbaNetworkServer.h
@@ -95,6 +95,7 @@ namespace uba
 		u64 GetTotalSentBytes();
 		u64 GetTotalRecvBytes();
 		Timer& GetTotalSentTimer();
+		u32 GetClientCount();
 		u32 GetConnectionCount();
 		void GetClientStats(ClientStats& out, u32 clientId);
 		void ResetTotalStats();
