[UBA]
* New attempt to fix the storage client logic around connecting to proxies and reconnecting to new proxy when disconnected.

Before this change there was a bug where if the client was too fast to connect to the proxy before it was fully created so it timed out. Then when asking host for new proxy it got the same proxy and earlied out. This made the client hang on downloading cas until the proxy was destroyed (so the host returned another proxy)

diff --git a/Engine/Source/Programs/UnrealBuildAccelerator/Common/Private/UbaStorageClient.cpp b/Engine/Source/Programs/UnrealBuildAccelerator/Common/Private/UbaStorageClient.cpp
--- a/Engine/Source/Programs/UnrealBuildAccelerator/Common/Private/UbaStorageClient.cpp
+++ b/Engine/Source/Programs/UnrealBuildAccelerator/Common/Private/UbaStorageClient.cpp
@@ -309,7 +309,6 @@ namespace uba
 					{
 						delete m_proxyClient;
 						m_proxyClient = nullptr;
-						m_lastTestedProxyIp.clear();
 						break;
 					}
 
@@ -325,18 +324,7 @@ namespace uba
 					if (proxy)
 					{
 						SCOPED_WRITE_LOCK(m_proxyClientLock, proxyLock);
-						if (!--proxy->refCount)
-						{
-							if (!proxy->client.IsConnected())
-							{
-								delete proxy;
-								if (proxy == m_proxyClient)
-								{
-									m_proxyClient = nullptr;
-									m_lastTestedProxyIp.clear();
-								}
-							}
-						}
+						--proxy->refCount;
 					}
 				});
 
@@ -403,24 +391,13 @@ namespace uba
 					if (m_proxyClient)
 						continue;
 
-					// If we have x processes sitting here with an already closed proxy we would do x connect attempts to that bad proxy
-					// Therefore we keep track of last checked proxy to prevent repetition of already failing proxy
-					// There is a risk we have multiple threads asking for two bad proxies.. but that is low enough for us not to care
-					if (m_lastTestedProxyIp == proxyHost.data)
-						continue;
-					m_lastTestedProxyIp = proxyHost.data;
-
 					NetworkClientCreateInfo ncci(m_logger.m_writer);
 					ncci.workerCount = 0;
 					bool ctorSuccess = true;
 					proxy = new ProxyClient(ctorSuccess, ncci);
+					m_proxyClient = proxy;
 
-					auto destroyProxy = MakeGuard([&]()
-						{
-							delete proxy;
-							proxy = nullptr;
-							allowProxy = false;
-						});
+					auto disallowProxy = MakeGuard([&]() { proxy = nullptr; allowProxy = false; });
 
 					if (!ctorSuccess)
 						continue;
@@ -462,11 +439,10 @@ namespace uba
 						}
 					}
 
-					destroyProxy.Cancel();
-
+					disallowProxy.Cancel();
 					++proxy->refCount;
 					proxy->client.SetWorkTracker(m_client.GetWorkTracker());
-					m_proxyClient = proxy;
+
 					continue;
 				}
 
diff --git a/Engine/Source/Programs/UnrealBuildAccelerator/Common/Public/UbaStorageClient.h b/Engine/Source/Programs/UnrealBuildAccelerator/Common/Public/UbaStorageClient.h
--- a/Engine/Source/Programs/UnrealBuildAccelerator/Common/Public/UbaStorageClient.h
+++ b/Engine/Source/Programs/UnrealBuildAccelerator/Common/Public/UbaStorageClient.h
@@ -85,7 +85,6 @@ namespace uba
 		static constexpr u8 ServiceId = StorageServiceId;
 
 		struct ProxyClient;
-		TString m_lastTestedProxyIp;
 		ReaderWriterLock m_proxyClientLock;
 		ProxyClient* m_proxyClient = nullptr;
 		u64 m_proxyClientKeepAliveTime = 0;
