[UBA]
* Fixed so work boxes are highlighted when selected
* Added support for custom color on work boxes
* Fixed bug when clicking outside end time on file/memory based traces
* Added trace start time to windows title when showing from file

diff --git a/Engine/Source/Programs/UnrealBuildAccelerator/Visualizer/Private/UbaTraceReader.cpp b/Engine/Source/Programs/UnrealBuildAccelerator/Visualizer/Private/UbaTraceReader.cpp
--- a/Engine/Source/Programs/UnrealBuildAccelerator/Visualizer/Private/UbaTraceReader.cpp
+++ b/Engine/Source/Programs/UnrealBuildAccelerator/Visualizer/Private/UbaTraceReader.cpp
@@ -120,6 +120,8 @@ namespace uba
 		else if (traceSystemStartTimeUs)
 			out.startTime = GetTime() - UsToTime(GetSystemTimeUs() - traceSystemStartTimeUs);
 
+		out.traceSystemStartTimeUs = traceSystemStartTimeUs;
+
 		m_memoryPos += reader.GetPosition();
 		out.finished = false;
 
@@ -139,13 +141,17 @@ namespace uba
 
 	bool TraceReader::UpdateReadFile(TraceView& out, u64 maxTime, bool& outChanged)
 	{
+		bool res = true;
 		BinaryReader traceReader(m_memoryPos);
 		u64 left = u64(m_memoryEnd - m_memoryPos);
 		while (traceReader.GetPosition() < left)
 		{
 			u64 pos = traceReader.GetPosition();
 			if (!ReadTrace(out, traceReader, maxTime))
-				return false;
+			{
+				res = false;
+				break;
+			}
 			if (pos == traceReader.GetPosition())
 			{
 				if (!traceReader.GetLeft())
@@ -155,7 +161,7 @@ namespace uba
 		}
 		outChanged = traceReader.GetPosition() != 0 || !m_activeProcesses.empty();
 		m_memoryPos += traceReader.GetPosition();
-		return true;
+		return res;
 	}
 
 
@@ -333,7 +339,10 @@ namespace uba
 		while (lastPos != toRead)
 		{
 			if (!ReadTrace(out, traceReader, maxTime))
+			{
+				m_memoryPos += lastPos;
 				return false;
+			}
 			u64 pos = traceReader.GetPosition();
 			if (pos == lastPos)
 				break;
@@ -497,9 +506,11 @@ namespace uba
 			session.memTotal = memTotal;
 			session.updates.push_back({time, totalSend, totalRecv, lastPing, memAvail, cpuLoad, connectionCount});
 			
-			session.highestSendPerS = Max(session.highestSendPerS, float(totalSend - session.prevSend)/TimeToS(time - session.prevUpdateTime));
-			session.highestRecvPerS = Max(session.highestRecvPerS, float(totalRecv - session.prevRecv)/TimeToS(time - session.prevUpdateTime));
-
+			if (session.prevUpdateTime)
+			{
+				session.highestSendPerS = Max(session.highestSendPerS, float(totalSend - session.prevSend)/TimeToS(time - session.prevUpdateTime));
+				session.highestRecvPerS = Max(session.highestRecvPerS, float(totalRecv - session.prevRecv)/TimeToS(time - session.prevUpdateTime));
+			}
 			break;
 		}
 		case TraceType_SessionDisconnect:
diff --git a/Engine/Source/Programs/UnrealBuildAccelerator/Visualizer/Private/UbaTraceReader.h b/Engine/Source/Programs/UnrealBuildAccelerator/Visualizer/Private/UbaTraceReader.h
--- a/Engine/Source/Programs/UnrealBuildAccelerator/Visualizer/Private/UbaTraceReader.h
+++ b/Engine/Source/Programs/UnrealBuildAccelerator/Visualizer/Private/UbaTraceReader.h
@@ -149,6 +149,7 @@ namespace uba
 		Map<u64, StatusUpdate> statusMap;
 		Map<u32, CacheWrite> cacheWrites;
 		u64 realStartTime = 0;
+		u64 traceSystemStartTimeUs = 0;
 		u64 startTime = 0;
 		u64 frequency = 0;
 		u32 totalProcessActiveCount = 0;
diff --git a/Engine/Source/Programs/UnrealBuildAccelerator/Visualizer/Private/UbaVisualizer.cpp b/Engine/Source/Programs/UnrealBuildAccelerator/Visualizer/Private/UbaVisualizer.cpp
--- a/Engine/Source/Programs/UnrealBuildAccelerator/Visualizer/Private/UbaVisualizer.cpp
+++ b/Engine/Source/Programs/UnrealBuildAccelerator/Visualizer/Private/UbaVisualizer.cpp
@@ -1232,7 +1232,9 @@ namespace uba
 						int sendY = graphBaseY;
 						int recvY = graphBaseY;
 						int cpuY = graphBaseY - int(cpuLoad*(GraphHeight-2));
-						int memY = graphBaseY - int(double(session.memTotal - update.memAvail)*double(GraphHeight - 2)/double(session.memTotal));
+						int memY = 0;
+						if (update.memAvail && session.memTotal)
+							memY = graphBaseY - int(double(session.memTotal - update.memAvail)*double(GraphHeight - 2)/double(session.memTotal));
 
 						double duration = TimeToS(update.time - prevTime);
 						if (update.time == 0)
@@ -1241,11 +1243,18 @@ namespace uba
 							isFirstUpdate = true;
 
 						double sendInvScaleY = duration * sendScale;
-						if (sendInvScaleY != 0)
+						if (sendInvScaleY != 0 && prevSend)
 							sendY = graphBaseY - int(double(updateSend - prevSend) / sendInvScaleY);
+						if (prevSendY <= 0)
+							prevSendY = sendY;
 						double recvInvScaleY = duration * recvScale;
-						if (recvInvScaleY != 0)
+						if (recvInvScaleY != 0 && prevRecv)
 							recvY = graphBaseY - int(double(updateRecv - prevRecv) / recvInvScaleY) - 1;
+						if (prevRecvY <= 0)
+							prevRecvY = recvY;
+
+						if (prevMemY <= 0)
+							prevMemY = memY;
 
 						if (!isFirstUpdate && x > clientRect.left && prevX <= clientRect.right)
 						{
@@ -1544,23 +1553,6 @@ namespace uba
 
 							u64 stop = work.stop;
 
-							bool done = stop != ~u64(0);
-							if (done)
-							{
-								auto insres = m_coloredBrushes.try_emplace(work.color);
-								if (insres.second)
-								{
-									auto c = (u8*)&work.color;
-									insres.first->second = CreateSolidBrush(RGB(c[2], c[1], c[0]));
-								}
-								brush = insres.first->second;
-							}
-							else
-							{
-								stop = playTime;
-								brush = m_processBrushes[0].inProgress;
-							}
-
 							float stopTime = TimeToS(stop);
 							if ((stopTime - startTime) * scaleX < 0.05f)
 								continue;
@@ -1576,11 +1568,32 @@ namespace uba
 
 							rect.right = Max(int(rect.right), left + 1);
 
+							Color color = work.color;
 							bool selected = m_workSelected && m_workTrack == trackIndex && m_workIndex == workIndex;
 							if (selected)
 							{
 								selectedWork = work;
 								work.bitmapDirty = true;
+								auto c = (u8*)&color;
+								for (u32 j=0; j!=3; ++j)
+									c[j] = u8(Min(int(c[j]) + 40, 255));
+							}
+
+							bool done = stop != ~u64(0);
+							if (done)
+							{
+								auto insres = m_coloredBrushes.try_emplace(color);
+								if (insres.second)
+								{
+									auto c = (u8*)&color;
+									insres.first->second = CreateSolidBrush(RGB(c[2], c[1], c[0]));
+								}
+								brush = insres.first->second;
+							}
+							else
+							{
+								stop = playTime;
+								brush = m_processBrushes[0].inProgress;
 							}
 
 							--rect.top;
@@ -1594,7 +1607,7 @@ namespace uba
 							int processWidth = rect.right - rect.left;
 							if (shouldDrawText && m_config.ShowProcessText && processWidth > 3)
 							{
-								if (work.bitmapDirty)
+								if (!work.bitmap || work.bitmapDirty)
 								{
 									if (!work.bitmap)
 									{
@@ -1620,7 +1633,13 @@ namespace uba
 
 									FillRect(textDC, &rect2, brush);
 									ExtTextOutW(textDC, rect2.left, rect2.top, ETO_CLIPPED, &rect2, work.description, int(wcslen(work.description)), NULL);
-									work.bitmapDirty = false;
+									if (!selected)
+										work.bitmapDirty = false;
+								}
+								else if (lastSelectedBitmap != work.bitmap)
+								{
+									SelectObject(textDC, work.bitmap);
+									lastSelectedBitmap = work.bitmap;
 								}
 
 								int width = Min(processWidth, 256);
@@ -3120,7 +3139,13 @@ namespace uba
 			{
 				m_trace.ReadFile(m_traceView, m_fileName.data, m_replay != 0);
 				m_traceView.finished = m_replay == 0;
-				PostNewTitle(GetTitlePrefix(title).Appendf(L"%s (v%u)", m_fileName.data, m_traceView.version));
+
+				time_t rawtime = m_traceView.traceSystemStartTimeUs / 1'000'000;
+				struct tm timeinfo;
+				localtime_s(&timeinfo, &rawtime);
+				wchar_t buffer[80];
+				wcsftime(buffer, sizeof_array(buffer), L"%Y-%m-%d %H:%M:%S", &timeinfo);
+				PostNewTitle(GetTitlePrefix(title).Appendf(L"%s (v%u) - %s", m_fileName.data, m_traceView.version, buffer));
 			}
 			else
 			{
