[UBA]
* Increased temporary string buffer in virtualize/devirtualize path to be MaxPath

diff --git a/Engine/Source/Programs/UnrealBuildAccelerator/Detours/Private/UbaDetoursFileMappingTable.cpp b/Engine/Source/Programs/UnrealBuildAccelerator/Detours/Private/UbaDetoursFileMappingTable.cpp
--- a/Engine/Source/Programs/UnrealBuildAccelerator/Detours/Private/UbaDetoursFileMappingTable.cpp
+++ b/Engine/Source/Programs/UnrealBuildAccelerator/Detours/Private/UbaDetoursFileMappingTable.cpp
@@ -224,7 +224,7 @@ namespace uba
 	u32 Rpc_GetEntryOffset(const StringKey& entryNameKey, const tchar* entryName, u64 entryNameLen, bool checkIfDir)
 	{
 		u32 dirTableOffset = ~u32(0);
-		StringBuffer<> entryNameForKey;
+		StringBuffer<MaxPath> entryNameForKey;
 		entryNameForKey.Append(entryName, entryNameLen);
 		if (CaseInsensitiveFs)
 			entryNameForKey.MakeLower();
diff --git a/Engine/Source/Programs/UnrealBuildAccelerator/Detours/Private/UbaDetoursShared.cpp b/Engine/Source/Programs/UnrealBuildAccelerator/Detours/Private/UbaDetoursShared.cpp
--- a/Engine/Source/Programs/UnrealBuildAccelerator/Detours/Private/UbaDetoursShared.cpp
+++ b/Engine/Source/Programs/UnrealBuildAccelerator/Detours/Private/UbaDetoursShared.cpp
@@ -241,7 +241,7 @@ namespace uba
 			VfsEntry& entry = g_vfsEntries[i];
 			if (!path.StartsWith(entry.vfs.data))
 				continue;
-			StringBuffer<> temp2(path.data + entry.vfs.count);
+			StringBuffer<MaxPath> temp2(path.data + entry.vfs.count);
 			path.Clear().Append(entry.local).Append(temp2);
 			return true;
 		}
@@ -257,7 +257,7 @@ namespace uba
 			VfsEntry& entry = g_vfsEntries[i];
 			if (path.count < entry.local.count || !path.StartsWith(entry.local.data))
 				continue;
-			StringBuffer<> temp2(path.data + entry.local.count);
+			StringBuffer<MaxPath> temp2(path.data + entry.local.count);
 			path.Clear().Append(entry.vfs).Append(temp2);
 			return true;
 		}
@@ -404,7 +404,7 @@ namespace uba
 
 	const tchar* Shared_GetFileAttributes(FileAttributes& outAttr, const tchar* fileName, bool checkIfDir)
 	{
-		StringBuffer<> fileNameForKey;
+		StringBuffer<MaxPath> fileNameForKey;
 		fileNameForKey.Append(fileName);
 		if (CaseInsensitiveFs)
 			fileNameForKey.MakeLower();
