[UBA]
* Minor cosmetic fixes in UbaCli

diff --git a/Engine/Source/Programs/UnrealBuildAccelerator/Cli/Private/UbaCli.cpp b/Engine/Source/Programs/UnrealBuildAccelerator/Cli/Private/UbaCli.cpp
--- a/Engine/Source/Programs/UnrealBuildAccelerator/Cli/Private/UbaCli.cpp
+++ b/Engine/Source/Programs/UnrealBuildAccelerator/Cli/Private/UbaCli.cpp
@@ -481,6 +481,7 @@ namespace uba
 			addOption(TC("-vfs"), TC("Z:/UEVFS/Root;E:\\dev\\fn"));
 			addOption(TC("-vfs"), TC("Z:/UEVFS/MSVC;c:\\sdk\\AutoSDK\\HostWin64\\Win64\\VS2022\\14.38.33130"));
 			addOption(TC("-vfs"), TC("Z:/UEVFS/WinSDK;C:\\Program Files (x86)\\Windows Kits\\10"));
+			addOption(TC("-vfs"), TC("Z:/UEVFS/Clang;c:\\sdk\\AutoSDK\\HostWin64\\Win64\\LLVM\\18.1.8"));
 			//addOption(TC("-vfs"), TC("Z:/UEVFS/Root;E:\\dev\\fn"));
 			//addOption(TC("-vfs"), TC("Z:/UEVFS/Compiler;c:\\sdk\\AutoSDK\\HostWin64\\Win64\\LLVM\\18.1.8"));
 			//addOption(TC("-vfs"), TC("Z:/UEVFS/Toolchain;c:\\sdk\\AutoSDK\\HostWin64\\Win64\\VS2022\\14.38.33130"));
@@ -831,7 +832,7 @@ namespace uba
 
 		SessionServerCreateInfo info(storageServer, networkServer, logWriter);
 		info.useUniqueId = useScheduler;
-		info.traceEnabled = true;
+		info.traceEnabled = cacheCommand.empty() || writeCacheSummary;
 		//info.detailedTrace = true;
 		info.launchVisualizer = launchVisualizer;
 		info.disableCustomAllocator = disableCustomAllocator;
@@ -870,7 +871,8 @@ namespace uba
 			if (!cacheCommand.empty())
 			{
 				LoggerWithWriter consoleLogger(g_consoleLogWriter);
-				return cacheClient->ExecuteCommand(consoleLogger, cacheCommand.data());
+				const tchar* additionalInfo = nullptr;
+				return cacheClient->ExecuteCommand(consoleLogger, cacheCommand.data(), nullptr, additionalInfo);
 			}
 
 			if (writeCacheSummary)
@@ -946,7 +948,7 @@ namespace uba
 				cacheClient->FetchFromCache(cacheResult, RootPaths(), bucketId, pinfo);
 				if (cacheResult.hit)
 				{
-					logger.Info(TC("%s run took %s [cached]"), (enableDetour ? TC("Boxed") : TC("Native")), TimeToText(GetTime() - start).str);
+					logger.Info(TC("Cached run took %s"), TimeToText(GetTime() - start).str);
 					return true;
 				}
 			}
@@ -962,7 +964,7 @@ namespace uba
 			ProcessHandle process = sessionServer.RunProcess(pinfo, false, enableDetour);
 			if (process.GetExitCode() != 0)
 				return logger.Error(TC("Error exit code: %u"), process.GetExitCode());
-			logger.Info(TC("%s run took %s"), (enableDetour ? TC("Boxed") : TC("Native")), TimeToText(GetTime() - start).str);
+			logger.Info(TC("%s run took %s"), (enableDetour ? TC("Detoured") : TC("Native")), TimeToText(GetTime() - start).str);
 
 			if (populateCache)
 			{
