[UBA]
* Fixed so detoured FindFirstFile works properly for drive roots

diff --git a/Engine/Source/Programs/UnrealBuildAccelerator/Detours/Private/Windows/UbaDetoursFunctionsKernelBase.inl b/Engine/Source/Programs/UnrealBuildAccelerator/Detours/Private/Windows/UbaDetoursFunctionsKernelBase.inl
--- a/Engine/Source/Programs/UnrealBuildAccelerator/Detours/Private/Windows/UbaDetoursFunctionsKernelBase.inl
+++ b/Engine/Source/Programs/UnrealBuildAccelerator/Detours/Private/Windows/UbaDetoursFunctionsKernelBase.inl
@@ -1756,7 +1756,7 @@ __forceinline HANDLE Shared_FindFirstFileExW(LPCWSTR lpFileName, FINDEX_INFO_LEV
 	}
 
 	// We must remove a slash at the end so it matches our cache entries
-	if (bufChars > 3)
+	if (bufChars > 2)
 	{
 		wchar_t* temp = (wchar_t*)buf;
 		if (temp[bufChars - 1] == '\\')
