[UBA]
* Added 128 byte bloom filter

diff --git a/Engine/Source/Programs/UnrealBuildAccelerator/Core/Private/UbaHash.cpp b/Engine/Source/Programs/UnrealBuildAccelerator/Core/Private/UbaHash.cpp
--- a/Engine/Source/Programs/UnrealBuildAccelerator/Core/Private/UbaHash.cpp
+++ b/Engine/Source/Programs/UnrealBuildAccelerator/Core/Private/UbaHash.cpp
@@ -99,33 +99,26 @@ namespace uba
 		return (CasKey&)output;
 	}
 
-	void LargeHash::Set(const void* data, u64 dataLen)
+	void BloomFilter::Add(const StringKey& key)
 	{
-		static_assert(BLAKE3_OUT_LEN == sizeof(LargeHash));
-		blake3_hasher hasher;
-		blake3_hasher_init(&hasher);
-		blake3_hasher_update(&hasher, data, dataLen);
-		blake3_hasher_finalize(&hasher, (u8*)this, BLAKE3_OUT_LEN);
+		u32 bit = key.a % (sizeof(bytes)*8);
+		u32 byte = bit / 8;
+		u32 bitInByte = bit - byte*8;
+		bytes[byte] |= (1 << bitInByte);
 	}
 
-	void LargeHash::BloomAdd(const LargeHash& lh)
+	bool BloomFilter::IsGuaranteedMiss(const StringKey& key) const
 	{
-		a |= lh.a;
-		b |= lh.b;
-		c |= lh.c;
-		d |= lh.d;
+		u32 bit = key.a % (sizeof(bytes)*8);
+		u32 byte = bit / 8;
+		u32 bitInByte = bit - byte*8;
+		return (bytes[byte] & (1 << bitInByte)) == 0;
 	}
 
-	bool LargeHash::BloomMiss(const LargeHash& lh) const
-	{
-		return (a & lh.a) != lh.a
-			|| (b & lh.b) != lh.b
-			|| (c & lh.c) != lh.c
-			|| (d & lh.d) != lh.d;
-	}
+	constexpr BloomFilter g_empty;
 
-	bool LargeHash::IsEmpty() const
+	bool BloomFilter::IsEmpty() const
 	{
-		return !a && !b && !c && !d;
+		return memcmp(g_empty.bytes, bytes, sizeof(bytes)) == 0;
 	}
 }
diff --git a/Engine/Source/Programs/UnrealBuildAccelerator/Core/Public/UbaHash.h b/Engine/Source/Programs/UnrealBuildAccelerator/Core/Public/UbaHash.h
--- a/Engine/Source/Programs/UnrealBuildAccelerator/Core/Public/UbaHash.h
+++ b/Engine/Source/Programs/UnrealBuildAccelerator/Core/Public/UbaHash.h
@@ -201,16 +201,12 @@ namespace uba
 	#define CHECK_PATH(path) while (false) {}
 	#endif
 
-	struct LargeHash
+	struct BloomFilter
 	{
-		u64 a = 0;
-		u64 b = 0;
-		u64 c = 0;
-		u64 d = 0;
+		u8 bytes[128] = { 0 };
 
-		void Set(const void* data, u64 dataLen);
-		void BloomAdd(const LargeHash& lh);
-		bool BloomMiss(const LargeHash& lh) const;
+		void Add(const StringKey& key);
+		bool IsGuaranteedMiss(const StringKey& key) const;
 		bool IsEmpty() const;
 	};
 }
