[UBA]
* Fixed linux compile error in file mapping

diff --git a/Engine/Source/Programs/UnrealBuildAccelerator/Common/Private/UbaFileMapping.cpp b/Engine/Source/Programs/UnrealBuildAccelerator/Common/Private/UbaFileMapping.cpp
--- a/Engine/Source/Programs/UnrealBuildAccelerator/Common/Private/UbaFileMapping.cpp
+++ b/Engine/Source/Programs/UnrealBuildAccelerator/Common/Private/UbaFileMapping.cpp
@@ -616,18 +616,21 @@ namespace uba
 
 	MappedView FileMappingBuffer::MapView(const StringView& str, u64 size, const tchar* hint)
 	{
+#if PLATFORM_WINDOWS
 		const tchar* handleStr = str.data + 1;
-		const tchar* handleStrEnd = wcschr(handleStr, '-');
+		const tchar* handleStrEnd = TStrchr(handleStr, '-');
 		if (!handleStrEnd)
 			return {};
 
-		HANDLE mappingHandle = (HANDLE)StringToValue(handleStr, handleStrEnd - handleStr);
+		FileMappingHandle h;
+		h.handle = (HANDLE)StringToValue(handleStr, handleStrEnd - handleStr);
 		const tchar* mappingOffsetStr = handleStrEnd + 1;
 		u64 mappingOffset = StringToValue(mappingOffsetStr, wcslen(mappingOffsetStr));
 		
-		FileMappingHandle h;
-		h.handle = mappingHandle;
 		return MapView(h, mappingOffset, size, hint);
+#else
+		return {};
+#endif
 	}
 
 	void FileMappingBuffer::UnmapView(MappedView view, const tchar* hint_, u64 newSize)
