[UBA]
* Fixed so detoured GetFinalPathNameByHandleW always return drive letter with uppercase

diff --git a/Engine/Source/Programs/UnrealBuildAccelerator/Detours/Private/Windows/UbaDetoursFunctionsKernelBase.inl b/Engine/Source/Programs/UnrealBuildAccelerator/Detours/Private/Windows/UbaDetoursFunctionsKernelBase.inl
--- a/Engine/Source/Programs/UnrealBuildAccelerator/Detours/Private/Windows/UbaDetoursFunctionsKernelBase.inl
+++ b/Engine/Source/Programs/UnrealBuildAccelerator/Detours/Private/Windows/UbaDetoursFunctionsKernelBase.inl
@@ -2592,6 +2592,7 @@ DWORD Detoured_GetFinalPathNameByHandleW(HANDLE hFile, LPTSTR lpszFilePath, DWOR
 			// Note, this really only matters when building linux target from windows.. then there is path validation that errors if this is not properly fixed
 			g_directoryTable.GetFinalPath(buffer.Clear(), fileName);
 			VirtualizePath(buffer);
+			buffer.data[0] = ToUpper(buffer.data[0]);
 			if (dwFlags == 2)
 				buffer.Prepend(AsView(L"\\??\\"));
 			UBA_ASSERT(requiredBufferSize == buffer.count);
