[UBA]
* Renamed ProxyFetchBegin job to only ProxyFetch

diff --git a/Engine/Source/Programs/UnrealBuildAccelerator/Common/Private/UbaStorageProxy.cpp b/Engine/Source/Programs/UnrealBuildAccelerator/Common/Private/UbaStorageProxy.cpp
--- a/Engine/Source/Programs/UnrealBuildAccelerator/Common/Private/UbaStorageProxy.cpp
+++ b/Engine/Source/Programs/UnrealBuildAccelerator/Common/Private/UbaStorageProxy.cpp
@@ -204,7 +204,7 @@ namespace uba
 
 					if (!file.memory)
 					{
-						file.trackId = m_client.TrackWorkStart(AsView(TC("ProxyFetchBegin")), ColorWork);
+						file.trackId = m_client.TrackWorkStart(AsView(TC("ProxyFetch")), ColorWork);
 						m_client.TrackWorkHint(file.trackId, StringView(hint).GetFileName());
 
 						NetworkMessage msg(m_client, ServiceId, StorageMessageType_FetchBegin, writer2);
