[UBA
* Added number of visible boxes text when filtering boxes in visualizer
* Added world time to info box popping up when selecting processes/work

diff --git a/Engine/Source/Programs/UnrealBuildAccelerator/Visualizer/Private/UbaVisualizer.cpp b/Engine/Source/Programs/UnrealBuildAccelerator/Visualizer/Private/UbaVisualizer.cpp
--- a/Engine/Source/Programs/UnrealBuildAccelerator/Visualizer/Private/UbaVisualizer.cpp
+++ b/Engine/Source/Programs/UnrealBuildAccelerator/Visualizer/Private/UbaVisualizer.cpp
@@ -1127,6 +1127,8 @@ namespace uba
 		SessionRec sortedSessions[1024];
 		Populate(sortedSessions, m_traceView, m_config.SortActiveRemoteSessions);
 
+		u32 visibleBoxes = 0;
+
 		TraceView::ProcessLocation processLocation { 0, 0, 0 };
 		for (u64 i = 0, e = m_traceView.sessions.size(); i != e; ++i)
 		{
@@ -1363,6 +1365,7 @@ namespace uba
 
 							if (!m_filterString.empty() && !Contains(process.description.c_str(), m_filterString.c_str()))
 								continue;
+							++visibleBoxes;
 
 							rect.right = Max(int(rect.right), left + 1);
 
@@ -1554,8 +1557,6 @@ namespace uba
 							u64 stop = work.stop;
 
 							float stopTime = TimeToS(stop);
-							if ((stopTime - startTime) * scaleX < 0.05f)
-								continue;
 
 							RECT rect;
 							rect.left = left;
@@ -1604,6 +1605,8 @@ namespace uba
 							FillRect(hdc, &rect, brush);
 							++rect.top;
 
+							++visibleBoxes;
+
 							int processWidth = rect.right - rect.left;
 							if (shouldDrawText && m_config.ShowProcessText && processWidth > 3)
 							{
@@ -1694,21 +1697,28 @@ namespace uba
 
 		if (!m_filterString.empty())
 		{
-			StringBuffer<> str;
-			str.Append(TC("Process Filter: ")).Append(m_filterString);
-			RECT rect = { 0, 0, 0, 0 };
+			StringBuffer<> str[2];
+			str[0].Append(TC("Box Filter: ")).Append(m_filterString);
+			str[1].Appendf(TC("Box Count: %u"), visibleBoxes);
 
-			SetActiveFont(m_popupFont);
-			DrawTextW(hdc, str.data, str.count, &rect, DT_SINGLELINE|DT_NOCLIP|DT_CALCRECT);
-			rect.left = (progressRect.right - rect.right)/2;
-			rect.right = (progressRect.right + rect.right)/2;
-			rect.top = (progressRect.bottom - rect.bottom)/2;
-			rect.bottom = (progressRect.bottom + rect.bottom)/2;
+			for (u32 i=0; i!=2; ++i)
+			{
+				SetActiveFont(m_popupFont);
+				RECT rect = { 0, 0, 0, 0 };
+				DrawTextW(hdc, str[i].data, str[i].count, &rect, DT_SINGLELINE|DT_NOCLIP|DT_CALCRECT);
+				int top = (progressRect.bottom - rect.bottom)/2 + i*20;
+				rect.left = (progressRect.right - rect.right)/2;
+				rect.right = (progressRect.right + rect.right)/2;
+				rect.bottom = top + (rect.bottom - rect.top);
+				rect.top = top;
 
-			RECT fillRect = rect;
-			InflateRect(&fillRect, 4, 2);
-			FillRect(hdc, &fillRect, m_tooltipBackgroundBrush);
-			DrawTextW(hdc, str.data, str.count, &rect, DT_SINGLELINE|DT_NOCLIP);
+				RECT fillRect = rect;
+				InflateRect(&fillRect, 4, 2);
+				FillRect(hdc, &fillRect, m_tooltipBackgroundBrush);
+				DrawTextW(hdc, str[i].data, str[i].count, &rect, DT_SINGLELINE|DT_NOCLIP);
+
+
+			}
 		}
 
 		if (m_config.showCursorLine && m_mouseOverWindow)
@@ -1868,7 +1878,7 @@ namespace uba
 			logger.Info(L"  %ls", process.description.c_str());
 			logger.Info(L"  Host:        %ls", m_processSelectedLocation.sessionIndex == 0 ? TC("local") : m_traceView.GetSession(m_processSelectedLocation).name.c_str());
 			logger.Info(L"  ProcessId:  %6u", process.id);
-			logger.Info(L"  Start:     %7ls", TimeToText(process.start, true).str);
+			logger.Info(L"  Start:     %7ls (%ls)", TimeToText(process.start, true).str, GetWorldTime(selectedWork.start).data);
 			logger.Info(L"  Duration:  %7ls", TimeToText(duration, true).str);
 			if (!process.returnedReason.empty())
 				logger.Info(L"  Returned: %7s", process.returnedReason.data());
@@ -2009,7 +2019,7 @@ namespace uba
 
 			logger.AddSpace();
 			logger.Info(L"  %ls", selectedWork.description);
-			logger.Info(L"  Start:     %ls", TimeToText(selectedWork.start, true).str);
+			logger.Info(L"  Start:     %ls (%ls)", TimeToText(selectedWork.start, true).str, GetWorldTime(selectedWork.start).data);
 			logger.Info(L"  Duration:  %ls", TimeToText(duration, true).str);
 			if (!selectedWork.entries.empty())
 				logger.AddSpace();
@@ -3078,6 +3088,17 @@ namespace uba
 		DwmSetWindowAttribute(m_hwnd, attribute, &useDarkMode, sizeof(useDarkMode));
 	}
 
+	StringBuffer<128> Visualizer::GetWorldTime(u64 time)
+	{
+		time_t rawTime = m_traceView.traceSystemStartTimeUs / 1'000'000;
+		rawTime += TimeToMs(time)/1000;
+		struct tm timeinfo;
+		localtime_s(&timeinfo, &rawTime);
+		StringBuffer<128> buffer;
+		buffer.count = u32(wcsftime(buffer.data, buffer.capacity, L"%Y-%m-%d %H:%M:%S", &timeinfo));
+		return buffer;
+	}
+
 	void Visualizer::PostNewTrace(u32 replay, bool paused)
 	{
 		KillTimer(m_hwnd, 0);
@@ -3139,13 +3160,7 @@ namespace uba
 			{
 				m_trace.ReadFile(m_traceView, m_fileName.data, m_replay != 0);
 				m_traceView.finished = m_replay == 0;
-
-				time_t rawtime = m_traceView.traceSystemStartTimeUs / 1'000'000;
-				struct tm timeinfo;
-				localtime_s(&timeinfo, &rawtime);
-				wchar_t buffer[80];
-				wcsftime(buffer, sizeof_array(buffer), L"%Y-%m-%d %H:%M:%S", &timeinfo);
-				PostNewTitle(GetTitlePrefix(title).Appendf(L"%s (v%u) - %s", m_fileName.data, m_traceView.version, buffer));
+				PostNewTitle(GetTitlePrefix(title).Appendf(L"%s (v%u) - %s", m_fileName.data, m_traceView.version, GetWorldTime(0).data));
 			}
 			else
 			{
diff --git a/Engine/Source/Programs/UnrealBuildAccelerator/Visualizer/Private/UbaVisualizer.h b/Engine/Source/Programs/UnrealBuildAccelerator/Visualizer/Private/UbaVisualizer.h
--- a/Engine/Source/Programs/UnrealBuildAccelerator/Visualizer/Private/UbaVisualizer.h
+++ b/Engine/Source/Programs/UnrealBuildAccelerator/Visualizer/Private/UbaVisualizer.h
@@ -153,6 +153,7 @@ namespace uba
 		void Redraw(bool now);
 		void SetActiveFont(const Font& font);
 		void UpdateTheme();
+		StringBuffer<128> GetWorldTime(u64 time);
 
 		StringBuffer<256> m_namedTrace;
 		StringBuffer<256> m_fileName;
