[UBA]
* Added -description to uba agent that can be used to populate description sent back to host for visualizer

diff --git a/Engine/Source/Programs/UnrealBuildAccelerator/Agent/Private/UbaAgent.cpp b/Engine/Source/Programs/UnrealBuildAccelerator/Agent/Private/UbaAgent.cpp
--- a/Engine/Source/Programs/UnrealBuildAccelerator/Agent/Private/UbaAgent.cpp
+++ b/Engine/Source/Programs/UnrealBuildAccelerator/Agent/Private/UbaAgent.cpp
@@ -113,6 +113,7 @@ namespace uba
 		logger.Info(TC("  -crypto=<key>           32 character (16 bytes) crypto key used for secure network transfer"));
 		logger.Info(TC("  -resendcas              Will try to send same cas multiple times (set this to true if server is allowed to remove cas files)"));
 		logger.Info(TC("  -populateCas=<dir>      Prepopulate cas database with files in dir. If files needed exists on machine this can be an optimization"));
+		logger.Info(TC("  -description            Add more info about the agent that will show in the trace log when hovering over session"));
 		#if PLATFORM_MAC
 		logger.Info(TC("  -populateCasFromXcodeVersion=<version>   Prepopulate cas database with files from local xcode installation that matches the version."));
 		logger.Info(TC("  -populateCasFromAllXcodes   Prepopulate cas database with files from local xcode installation that matches the version."));
@@ -379,6 +380,7 @@ namespace uba
 		u8 crypto[16];
 		bool hasCrypto = false;
 		Vector<TString> populateCasDirs;
+		TString description;
 
 		#if PLATFORM_MAC
 		StringBuffer<32> populateCasFromXcodeVersion;
@@ -655,6 +657,14 @@ namespace uba
 			{
 				proxyUseLocalStorage = true;
 			}
+			else if (name.Equals(TC("-description")))
+			{
+				if (value.IsEmpty())
+					return PrintHelp(TC("-description"));
+				if (value[value.count - 1] == '\"')
+					value.Resize(value.count - 1);
+				description = value.data + (value[0] == '\"');
+			}
 			else if (name.Equals(TC("-noaws")))
 			{
 				couldBeAws = false;
@@ -763,6 +773,9 @@ namespace uba
 		if (hasCrypto)
 			extraInfo.Append(TC(", Encrypted"));
 
+		if (!description.empty())
+			extraInfo.Append(TC(", ")).Append(description);
+
 		//logger.Info(TC("\033[39m\n"));
 		const tchar* dbgStr = TC("");
 		#if UBA_DEBUG
