[UBA]
* Added include crawler which is used to find and prefetch code files needed for compile process. It is extremely naive and might overfetch files that are hidden behind defines but it seems to give a nice boost when having latency to helpers.

diff --git a/Engine/Source/Programs/UnrealBuildAccelerator/Common/Private/UbaSessionClient.cpp b/Engine/Source/Programs/UnrealBuildAccelerator/Common/Private/UbaSessionClient.cpp
--- a/Engine/Source/Programs/UnrealBuildAccelerator/Common/Private/UbaSessionClient.cpp
+++ b/Engine/Source/Programs/UnrealBuildAccelerator/Common/Private/UbaSessionClient.cpp
@@ -2,6 +2,7 @@
 
 #include "UbaSessionClient.h"
 #include "UbaApplicationRules.h"
+#include "UbaConfig.h"
 #include "UbaNetworkClient.h"
 #include "UbaNetworkMessage.h"
 #include "UbaProcess.h"
@@ -15,6 +16,17 @@
 
 namespace uba
 {
+	void SessionClientCreateInfo::Apply(const Config& config)
+	{
+		SessionCreateInfo::Apply(config);
+
+		if (const ConfigTable* table = config.GetTable(TC("Session")))
+		{
+			table->GetValueAsBool(useIncludeCrawler, TC("UseIncludeCrawler"));
+		}
+	}
+
+
 	struct SessionClient::ModuleInfo
 	{
 		ModuleInfo(const tchar* n, const CasKey& c, u32 a) : name(n), casKey(c), attributes(a), done(true) {}
@@ -32,6 +44,7 @@ namespace uba
 	,	m_waitToSendEvent(false)
 	,	m_loop(true)
 	,	m_allowSpawn(true)
+	,	m_includeCrawler(m_logger, m_client)
 	{
 		m_maxProcessCount = info.maxProcessCount;
 		m_dedicated = info.dedicated;
@@ -45,6 +58,8 @@ namespace uba
 		m_memKillLoadPercent = info.memKillLoadPercent;
 		m_processFinished = info.processFinished;
 
+		m_useIncludeCrawler = info.useIncludeCrawler;
+
 		m_processIdCounter = ~0u / 2; // We set this value to a very high value.. because it will be used by child processes and we don't want id from server and child process id to collide
 
 		if (m_name.IsEmpty())
@@ -64,6 +79,65 @@ namespace uba
 		m_nameToHashTableMem.Init(NameToHashMemSize);
 
 		Create(info);
+
+
+
+		if (m_useIncludeCrawler)
+			m_includeCrawler.Init([this](const StringView& fileName, u32& outAttr) // FileExists
+				{
+					u32 tableOffset = 0;
+					if (!EntryExists(fileName, tableOffset))
+						return false;
+					if ((tableOffset & 0x80000000) != 0)
+						return false;
+					BinaryReader reader(m_directoryTableMem, tableOffset, m_directoryTableMemPos);
+					reader.SkipString();
+					outAttr = reader.ReadFileAttributes();
+					return true;
+				},
+			[this](const StringView& path, const IncludeCrawler::FileFunc& fileFunc) // TraverseFilesFunc
+				{
+					u32 tableOffset = 0;
+					if (!EntryExists(path, tableOffset))
+						return;
+					if ((tableOffset & 0x80000000) == 0)
+						return;
+					tableOffset &= ~0x80000000;
+
+					// This is not entirely correct since files could be added while uba is running.. but since this is only used for includes that
+					// we know are created before uba started we're good.. and there should be no new files added to those dirs causing table offset to be chained
+					BinaryReader reader(m_directoryTableMem, tableOffset, m_directoryTableMemPos);
+					while (true)
+					{
+						u32 prevTableOffset = u32(reader.Read7BitEncoded());
+						if (!prevTableOffset)
+							break;
+						reader.SetPosition(prevTableOffset);
+					}
+
+					u32 dirAttr = reader.ReadFileAttributes();
+					if (!dirAttr)
+						return;
+					reader.ReadVolumeSerial();
+					reader.ReadFileIndex();
+					u64 itemCount = reader.Read7BitEncoded();
+					while (itemCount--)
+					{
+						StringBuffer<> fileName;
+						reader.ReadString(fileName);
+						if (CaseInsensitiveFs)
+							fileName.MakeLower();
+						u32 attr = reader.ReadFileAttributes();
+						bool isDir = IsDirectory(attr);
+						fileFunc(fileName, isDir);
+						reader.ReadVolumeSerial();
+						reader.ReadFileIndex();
+						if (isDir)
+							continue;
+						reader.ReadFileTime();
+						reader.ReadFileSize();
+					}
+				});
 	}
 
 	SessionClient::~SessionClient()
@@ -142,7 +216,7 @@ namespace uba
 		return res;
 	}
 
-	bool SessionClient::GetCasKeyForFile(CasKey& out, u32 processId, const StringBufferBase& fileName, const StringKey& fileNameKey)
+	bool SessionClient::GetCasKeyForFile(CasKey& out, u32 processId, const StringView& fileName, const StringKey& fileNameKey)
 	{
 		TimerScope waitTimer(Stats().waitGetFileMsg);
 		SCOPED_WRITE_LOCK(m_nameToHashLookupLock, lock);
@@ -282,39 +356,30 @@ namespace uba
 		KeyToString keyStr(ToStringKeyLower(applicationDir));
 
 		Atomic<bool> success = true;
-		Atomic<u32> handledCount;
 
-		for (auto& m : modules)
-		{
-			m_client.AddWork([&handledCount, &m, this, &success, &keyStr](const WorkContext&)
+		m_client.ParallelFor(u32(modules.size()), modules, [&](auto& it)
+			{
+				auto& m = *it;
+				TrackWorkScope tws(m_client, AsView(TC("FetchModule")), ColorWork);
+				tws.AddHint(m.name);
+
+				auto g = MakeGuard([&]() { m.done.Set(); });
+				CasKey newCasKey;
+				bool storeUncompressed = true;
+				u64 fileSize;
+				const tchar* moduleName = m.name.c_str();
+				if (!RetrieveCasFile(newCasKey, fileSize, m.casKey, moduleName, storeUncompressed))
 				{
-					++handledCount;
-					auto g = MakeGuard([&]() { m.done.Set(); });
-					CasKey newCasKey;
-					bool storeUncompressed = true;
-					u64 fileSize;
-					const tchar* moduleName = m.name.c_str();
-					if (!RetrieveCasFile(newCasKey, fileSize, m.casKey, moduleName, storeUncompressed))
-					{
-						m_logger.Error(TC("Casfile not found for %s (%s)"), moduleName, CasKeyString(m.casKey).str);
-						success = false;
-						return;
-					}
-					if (const tchar* lastSeparator = TStrrchr(moduleName, PathSeparator))
-						moduleName = lastSeparator + 1;
-					StringBuffer<MaxPath> temp;
-					if (!WriteBinFile(temp, moduleName, newCasKey, keyStr, m.attributes))
-						success = false;
-				}, 1, TC("EnsureApp"));
-		}
-
-		while (handledCount < modules.size())
-			m_client.DoWork();
-
-		// Wait for all to be done
-		for (auto& m : modules)
-			if (!m.done.IsSet(10 * 60 * 1000)) // 10 minutes is a very long time
-				return m_logger.Error(TC("Timed out while waiting for application cas files to be downloaded"));
+					m_logger.Error(TC("Casfile not found for %s (%s)"), moduleName, CasKeyString(m.casKey).str);
+					success = false;
+					return;
+				}
+				if (const tchar* lastSeparator = TStrrchr(moduleName, PathSeparator))
+					moduleName = lastSeparator + 1;
+				StringBuffer<MaxPath> temp;
+				if (!WriteBinFile(temp, moduleName, newCasKey, keyStr, m.attributes))
+					success = false;
+			}, AsView(TC("FetchModule")), true);
 
 		if (!success)
 			return false;
@@ -334,6 +399,8 @@ namespace uba
 
 	bool SessionClient::ReadModules(List<ModuleInfo>& outModules, u32 processId, const tchar* application)
 	{
+		TrackWorkScope tws(m_client, AsView(TC("ReadModules")), ColorWork);
+
 		StackBinaryReader<16*1024> reader;
 		{
 			StackBinaryWriter<256> writer;
@@ -438,24 +505,29 @@ namespace uba
 		}
 		#endif
 
-
-
 		if (RootsHandle rootsHandle = process.GetStartInfo().rootsHandle)
 			if (!SendRootsHandle(rootsHandle))
 				return false;
+
+		RunIncludeCrawler(process);
+
 		return true;
 	}
 
 	bool SessionClient::CreateFile(CreateFileResponse& out, const CreateFileMessage& msg)
 	{
-		const StringBufferBase& fileName = msg.fileName;
-		StringKey fileNameKey = msg.fileNameKey;
-
 		if ((msg.access & FileAccess_Write) != 0)
 			return Session::CreateFile(out, msg);
-	
+
+		TrackWorkScope tws(m_client, AsView(TC("CreateFile")), ColorWork);
+		tws.AddHint(msg.fileName);
+		return CreateFileForRead(out, msg.fileName, msg.fileNameKey, msg.process.m_id, *msg.process.m_startInfo.rules);
+	}
+
+	bool SessionClient::CreateFileForRead(CreateFileResponse& out, const StringView& fileName, const StringKey& fileNameKey, u32 processId, const ApplicationRules& rules)
+	{
 		CasKey casKey;
-		if (!GetCasKeyForFile(casKey, msg.process.m_id, fileName, fileNameKey))
+		if (!GetCasKeyForFile(casKey, processId, fileName, fileNameKey))
 			return false;
 
 		// Not finding a file is a valid path. Some applications try with a path and if fails try another path
@@ -497,7 +569,7 @@ namespace uba
 				if (m_useStorage || memoryMapAlignment == 0)
 				{
 					bool storeUncompressed = memoryMapAlignment == 0;
-					bool allowProxy = msg.process.m_startInfo.rules->AllowStorageProxy(fileName);
+					bool allowProxy = rules.AllowStorageProxy(fileName);
 					if (!RetrieveCasFile(newCasKey, fileSize, casKey, fileName.data, storeUncompressed, allowProxy))
 						return m_logger.Error(TC("Error retrieving cas entry %s (%s)"), CasKeyString(casKey).str, fileName.data);
 
@@ -528,7 +600,7 @@ namespace uba
 						casKey = AsCompressed(casKey, false);
 						entry.handled = true;
 						Storage::RetrieveResult result;
-						bool allowProxy = msg.process.m_startInfo.rules->AllowStorageProxy(fileName);
+						bool allowProxy = rules.AllowStorageProxy(fileName);
 						if (!m_storage.RetrieveCasFile(result, casKey, fileName.data, &m_fileMappingBuffer, memoryMapAlignment, allowProxy))
 							return m_logger.Error(TC("Error retrieving cas entry %s (%s)"), CasKeyString(casKey).str, fileName.data);
 						entry.success = true;
@@ -578,7 +650,7 @@ namespace uba
 					fileSize = map.size;
 					newName.Clear().Append(map.name);
 				}
-				else if (!IsRarelyRead(msg.process, fileName))
+				else if (!rules.IsRarelyRead(fileName))
 				{
 					AddFileMapping(fileNameKey, fileName.data, newName.data, fileSize);
 				}
@@ -913,6 +985,7 @@ namespace uba
 
 	bool SessionClient::GetListDirectoryInfo(ListDirectoryResponse& out, const StringView& dirName, const StringKey& dirKey)
 	{
+		TrackWorkScope tws(m_client, AsView(TC("GetListDir")), ColorWork);
 		StackBinaryWriter<1024> writer;
 		NetworkMessage msg(m_client, ServiceId, SessionMessageType_ListDirectory, writer);
 		writer.WriteU32(m_sessionId);
@@ -1312,6 +1385,7 @@ namespace uba
 		StackBinaryReader<SendMaxSize> reader;
 
 		{
+			TrackWorkScope tws(m_client, AsView(TC("RequestProcesses")), ColorWork);
 			StackBinaryWriter<32> writer;
 			NetworkMessage msg(m_client, ServiceId, SessionMessageType_ProcessAvailable, writer);
 			writer.WriteU32(m_sessionId);
@@ -1362,6 +1436,7 @@ namespace uba
 
 				m_client.AddWork([knownInputKey, allowProxy, this](const WorkContext&)
 					{
+						TrackWorkScope tws(m_client, StringBuffer<>().Appendf(TC("KnownInput")), ColorWork);
 						Storage::RetrieveResult result;
 						bool res = m_storage.RetrieveCasFile(result, knownInputKey, TC("KnownInput"), nullptr, 1, allowProxy);
 						(void)res;
@@ -1453,6 +1528,7 @@ namespace uba
 
 	bool SessionClient::SendUpdateDirectoryTable(StackBinaryReader<SendMaxSize>& reader)
 	{
+		TrackWorkScope tws(m_client, AsView(TC("UpdateDir")), ColorWork);
 		UBA_ASSERT(reader.GetPosition() == 0);
 		StackBinaryWriter<32> writer;
 		NetworkMessage msg(m_client, ServiceId, SessionMessageType_GetDirectoriesFromServer, writer);
@@ -1464,6 +1540,7 @@ namespace uba
 
 	bool SessionClient::SendUpdateNameToHashTable(StackBinaryReader<SendMaxSize>& reader)
 	{
+		TrackWorkScope tws(m_client, AsView(TC("UpdateHashTable")), ColorWork);
 		StackBinaryWriter<32> writer;
 		NetworkMessage msg(m_client, ServiceId, SessionMessageType_GetNameToHashFromServer, writer);
 		writer.WriteU32(~u32(0));
@@ -1483,6 +1560,7 @@ namespace uba
 			return;
 
 		float cpuLoad = UpdateCpuLoad();
+		m_cpuUsage = cpuLoad;
 
 		StackBinaryWriter<128> writer;
 		NetworkMessage msg(m_client, ServiceId, SessionMessageType_Ping, writer);
@@ -1590,6 +1668,19 @@ namespace uba
 
 	void SessionClient::ThreadCreateProcessLoop()
 	{
+		// Ask for dir and hash table straight away to minimize latency
+		m_client.AddWork([this](const WorkContext&)
+			{
+				StackBinaryReader<SendMaxSize> reader;
+				SendUpdateDirectoryTable(reader);
+			}, 1, TC("InitGetDirTable"), ColorWork);
+		m_client.AddWork([this](const WorkContext&)
+			{
+				StackBinaryReader<SendMaxSize> reader;
+				SendUpdateNameToHashTable(reader);
+			}, 1, TC("InitGetHashTable"), ColorWork);
+
+
 		struct ProcessRec
 		{
 			ProcessRec(ProcessImpl* impl) : handle(impl) {}
@@ -1609,6 +1700,8 @@ namespace uba
 		u64 memAvail;
 		u64 memTotal;
 		GetMemoryInfo(memAvail, memTotal);
+		m_memAvail = memAvail;
+		m_memTotal = memTotal;
 
 		SendPing(memAvail, memTotal);
 
@@ -1661,6 +1754,7 @@ namespace uba
 			}
 		};
 
+		Vector<InternalProcessStartInfo> startInfos;
 
 		while (m_loop)
 		{
@@ -1670,6 +1764,9 @@ namespace uba
 			FlushDeadProcesses();
 
 			GetMemoryInfo(memAvail, memTotal);
+			m_memAvail = memAvail;
+			m_memTotal = memTotal;
+
 			if (memAvail < memRequiredFree)
 			{
 				for (auto it = activeProcesses.rbegin(); it != activeProcesses.rend(); ++it)
@@ -1707,8 +1804,11 @@ namespace uba
 				}
 
 				if (!firstCall)
+				{
 					GetMemoryInfo(memAvail, memTotal);
-
+					m_memAvail = memAvail;
+					m_memTotal = memTotal;
+				}
 				if (memAvail < memRequiredToSpawn)
 				{
 					if (waitForMemoryPressureStartTime == 0)
@@ -1730,7 +1830,8 @@ namespace uba
 					availableWeight = Min(availableWeight, 1.0f);
 				}
 
-				Vector<InternalProcessStartInfo> startInfos;
+				startInfos.clear();
+
 				if (!SendProcessAvailable(startInfos, availableWeight))
 				{
 					m_loop = false;
@@ -1992,6 +2093,128 @@ namespace uba
 		return logLineCount;
 	}
 
+	void SessionClient::RunIncludeCrawler(ProcessImpl& process)
+	{
+		if (!m_useIncludeCrawler)
+			return;
+
+		auto& startInfo = process.GetStartInfo();
+		if (!startInfo.rules->SupportsIncludeCrawler())
+			return;
+
+		const tchar* at = TStrchr(startInfo.arguments, '@');
+		if (!at)
+			return;
+
+		auto CreateFileFunc = [this, processId = process.GetId(), rules = startInfo.rules](const StringView& fileName, const IncludeCrawler::AccessFileFunc& func)
+			{
+				CreateFileResponse out;
+				{
+					TrackWorkScope tws(m_client, AsView(TC("CrawlFile")), ColorWork);
+					tws.AddHint(fileName);
+
+					if (!CreateFileForRead(out, fileName, ToStringKey(fileName), processId, *rules))
+						return false;
+				}
+
+				if (out.fileName[0] == '^')
+				{
+					MappedView view = m_fileMappingBuffer.MapView(out.fileName, out.size, fileName.data);
+					if (view.memory)
+					{
+						bool res = func(view.memory, out.size);
+						m_fileMappingBuffer.UnmapView(view, fileName.data);
+						return res;
+					}
+					else
+						m_logger.Warning(TC("Failed to open %s"), out.fileName.data);
+				}
+				else
+				{
+					FileAccessor fa(m_logger, out.fileName.data);
+					if (fa.OpenMemoryRead(0, false))
+						return func(fa.GetData(), fa.GetSize());
+					else
+						m_logger.Warning(TC("Failed to open %s"), out.fileName.data);
+				}
+				return true;
+			};
+
+		auto DevirtualizePathFunc = [this, rootsHandle = startInfo.rootsHandle](StringBufferBase& inOut)
+			{
+				return DevirtualizePath(inOut, rootsHandle);
+			};
+
+		m_includeCrawler.Add(at + 1, startInfo.workingDir, CreateFileFunc, DevirtualizePathFunc, startInfo.application, startInfo.rules);
+	}
+
+	bool SessionClient::ParseDirectoryTable()
+	{
+		SCOPED_WRITE_LOCK(m_directoryTable.m_lookupLock, lock);
+		SCOPED_READ_LOCK(m_directoryTable.m_memoryLock, lock2);
+		u32 newMemPosition = m_directoryTable.m_memorySize;
+		lock2.Leave();
+		if (newMemPosition == m_dirtableParsedPosition)
+			return false;
+		m_directoryTable.ParseDirectoryTableNoLock(m_dirtableParsedPosition, newMemPosition);
+		m_dirtableParsedPosition = newMemPosition;
+		return true;
+	}
+
+	bool SessionClient::EntryExists(const StringView& path, u32& outTableOffset)
+	{
+		StringKey key;
+		if (path.data[path.count - 1] == PathSeparator)
+			key = ToStringKey(path.data, path.count - 1);
+		else
+			key = ToStringKey(path);
+
+		u32 tableOffset = 0;
+		DirectoryTable::Exists exists = m_directoryTable.EntryExists(key, path, true, &tableOffset);
+		if (exists == DirectoryTable::Exists_Maybe)
+		{
+			if (ParseDirectoryTable())
+			{
+				exists = m_directoryTable.EntryExists(key, path, true, &tableOffset);
+			}
+
+			if (exists == DirectoryTable::Exists_Maybe)
+			{
+				StringBuffer<> dirName;
+				if (const tchar* lastSeparator = TStrrchr(path.data, PathSeparator))
+					dirName.Append(path.data, lastSeparator - path.data);
+				
+				StringKey dirKey = ToStringKey(dirName);
+
+				{	// Optimization to reduce number of messages
+					SCOPED_WRITE_LOCK(m_dirVisitedLock, l);
+					DirVisitedEntry& entry = m_dirVisited[dirKey];
+					l.Leave();
+				
+					SCOPED_WRITE_LOCK(entry.lock, l2);
+					if (!entry.handled)
+					{
+						ListDirectoryResponse out;
+						if (!GetListDirectoryInfo(out, dirName, dirKey))
+							return false;
+						ParseDirectoryTable();
+						entry.handled = true;
+					}
+				}
+
+				exists = m_directoryTable.EntryExists(key, path, true, &tableOffset);
+			}
+		}
+		if (exists == DirectoryTable::Exists_No)
+			return false;
+
+		UBA_ASSERT(exists != DirectoryTable::Exists_Maybe);
+
+		outTableOffset = tableOffset;
+
+		return true;
+	}
+
 	bool SessionClient::AllocFailed(Process& process, const tchar* allocType, u32 error)
 	{
 		//StackBinaryWriter<32> writer;
@@ -2352,7 +2575,6 @@ namespace uba
 
 	void SessionClient::TraceSessionUpdate()
 	{
-		float cpuLoad = 0.0f;//UpdateCpuLoad();
 		u64 send;
 		u64 recv;
 		if (auto backend = m_client.GetFirstConnectionBackend())
@@ -2364,10 +2586,8 @@ namespace uba
 			recv = m_client.GetTotalRecvBytes();
 			send = m_client.GetTotalSentBytes();
 		}
-		u64 memAvail = 0;//m_memAvail;
-		u64 memTotal = 0;//m_memTotal;
 
 		// send and recv are swapped on purpose because that is how visualizer is visualizing 
-		m_trace.SessionUpdate(0, 0, send, recv, 0, memAvail, memTotal, cpuLoad);
+		m_trace.SessionUpdate(0, 0, send, recv, m_lastPing, m_memAvail, m_memTotal, m_cpuUsage);
 	}
 }
diff --git a/Engine/Source/Programs/UnrealBuildAccelerator/Common/Public/UbaIncludeCrawler.h b/Engine/Source/Programs/UnrealBuildAccelerator/Common/Public/UbaIncludeCrawler.h
new file mode 100644
--- /dev/null
+++ b/Engine/Source/Programs/UnrealBuildAccelerator/Common/Public/UbaIncludeCrawler.h
@@ -0,0 +1,611 @@
+// Copyright Epic Games, Inc. All Rights Reserved.
+
+#pragma once
+
+#include "UbaApplicationRules.h"
+#include "UbaFileAccessor.h"
+#include "UbaPathUtils.h"
+#include "UbaProcessUtils.h"
+#include "UbaWorkManager.h"
+
+// Note
+// IncludeCrawler is quite Unreal specific. It might work in different environments but is based on using response (.rsp) files.
+// It does parse some defines to figure out custom paths but code should also work without this
+// Ideally it should be a full fledged preprocessor parser but that would also not make it possible to run this in parallel
+
+namespace uba
+{
+	class IncludeCrawler
+	{
+	public:
+		inline IncludeCrawler(Logger& logger, WorkManager& workManager);
+
+		using FileExistsFunc = Function<bool(const StringView& fileName, u32& outAttr)>;
+		using FileFunc = Function<void(const StringView& file, bool isDirectory)>;
+		using TraverseFilesFunc = Function<void(const StringView& path, const FileFunc& fileFunc)>;
+
+		using AccessFileFunc = Function<bool(const void* data, u64 dataSize)>;
+		using DevirtualizePathFunc = Function<bool(StringBufferBase& inOut)>;
+		using CreateFileFunc = Function<bool(const StringView& fileName, const AccessFileFunc& func)>;
+
+		inline void Init(FileExistsFunc&& fileExistsFunc, TraverseFilesFunc&& traverseFilesFunc);
+		inline bool Add(const tchar* rsp, const tchar* workDir, CreateFileFunc&& func, DevirtualizePathFunc&& devirtualizePathFunc, const tchar* app = nullptr, const ApplicationRules* rules = nullptr);
+
+		Logger& m_logger;
+		WorkManager& m_workManager;
+
+		ReaderWriterLock m_codeFilesHandledLock;
+		struct HandledFile { ReaderWriterLock lock; bool handled = false; };
+		UnorderedMap<StringKey, HandledFile> m_codeFilesHandled;
+
+		ReaderWriterLock m_pchLookupLock;
+		struct Pch { ReaderWriterLock lock; bool handled = false; UnorderedSet<StringKey> files; };
+		UnorderedMap<StringKey, Pch> m_pchLookup;
+
+		struct IncludeRoot
+		{
+			TString path;
+			BloomFilter bloomFilter;
+		};
+
+		struct Instance
+		{
+			TString application;
+			TString rsp;
+			TString workDir;
+			TString platform;
+			TString compiledPlatform;
+			TString overriddenPlatformName;
+			TString builtinIncludesDir;
+			bool platformIsExtension = false;
+			bool usePch = false;
+			CreateFileFunc createFileFunc;
+			DevirtualizePathFunc devirtualizePathFunc;
+			Vector<IncludeRoot> includeRoots;
+			Pch* pch = nullptr;
+		};
+
+		ReaderWriterLock m_instancesLock;
+		List<Instance> m_instances;
+		ReaderWriterLock m_builtIncludesHandledLock;
+		UnorderedSet<u32> m_builtIncludesHandled;
+		FileExistsFunc m_fileExistsFunc;
+		TraverseFilesFunc m_traverseFilesFunc;
+
+		struct CodeFile { TString path; bool hasPch = false; };
+		using CodeFiles = List<CodeFile>;
+
+		inline bool ParseRsp(Instance& instance, const tchar* rsp, CodeFiles& outCodeFiles);
+		inline bool ParseRsp2(Instance& instance, const void* data, u64 dataSize, CodeFiles& outCodeFiles);
+		inline bool ParsePch(Pch& pch, Instance& instance, const void* data, u64 dataSize);
+		inline bool ParseCodeFile(Instance& instance, const StringView& codeFile, bool parseDefines);
+		inline bool ParseCodeFile2(Instance& instance, const void* data, u64 dataSize, const StringView& codeFile, bool parseDefines);
+		inline bool HandleInclude(Instance& instance, const StringView& rootPath, const StringView& include, bool parseDefines);
+	};
+
+
+	IncludeCrawler::IncludeCrawler(Logger& logger, WorkManager& workManager)
+	:	m_logger(logger)
+	,	m_workManager(workManager)
+	{
+	}
+
+	void IncludeCrawler::Init(FileExistsFunc&& fileExistsFunc, TraverseFilesFunc&& traverseFilesFunc)
+	{
+		m_fileExistsFunc = std::move(fileExistsFunc);
+		m_traverseFilesFunc = std::move(traverseFilesFunc);
+	}
+
+	bool IncludeCrawler::Add(const tchar* rsp, const tchar* workDir, CreateFileFunc&& createFileFunc, DevirtualizePathFunc&& devirtualizePathFunc, const tchar* app, const ApplicationRules* rules)
+	{
+		SCOPED_WRITE_LOCK(m_instancesLock, lock);
+		Instance& instance = m_instances.emplace_back();
+		lock.Leave();
+
+		instance.application = app;
+		instance.rsp = rsp;
+		instance.workDir = workDir;
+		if (CaseInsensitiveFs)
+			ToLower(instance.workDir.data());
+		if (instance.workDir.back() != PathSeparator)
+			instance.workDir += PathSeparator;
+		instance.createFileFunc = std::move(createFileFunc);
+		instance.devirtualizePathFunc = std::move(devirtualizePathFunc);
+
+		m_workManager.AddWork([this, &instance, rules](const WorkContext& context)
+			{
+				CodeFiles codeFiles;
+				if (!ParseRsp(instance, instance.rsp.c_str(), codeFiles))
+					return;
+
+				for (auto& cf : codeFiles)
+				{
+					if (instance.usePch && cf.hasPch)
+					{
+						// Note, clang is always reading all the individual files that has been stored in the pch to validate content..
+						// So this code path is not implemented for clang.. if -fno-validate-ast-input-files-content is added we can implement this path
+						StringBuffer<> fixedFile;
+						FixPath(cf.path.c_str(), instance.workDir.data(), instance.workDir.size(), fixedFile);
+
+						StringKey pchKey = ToStringKeyLower(fixedFile);
+
+						SCOPED_WRITE_LOCK(m_pchLookupLock, lock);
+						Pch& pch = m_pchLookup[pchKey];
+						lock.Leave();
+
+						SCOPED_WRITE_LOCK(pch.lock, lock2);
+						if (!pch.handled)
+						{
+							fixedFile.Append(TC(".dep.json"));
+							instance.createFileFunc(fixedFile, [&](const void* data, u64 dataSize)
+							{
+								return ParsePch(pch, instance, data, dataSize);
+							});
+							pch.handled = true;
+						}
+						instance.pch = &pch;
+					}
+					else
+					{
+						ParseCodeFile(instance, cf.path, instance.platform.empty());
+						if (instance.platform.empty())
+						{
+							if (!instance.overriddenPlatformName.empty())
+								instance.platform = instance.overriddenPlatformName;
+							else
+								instance.platform = instance.compiledPlatform;
+						}
+					}
+				}
+
+				if (rules && !instance.builtinIncludesDir.empty())
+				{
+					SCOPED_WRITE_LOCK(m_builtIncludesHandledLock, lock);
+					bool shouldHandle = m_builtIncludesHandled.insert(rules->index).second;
+					lock.Leave();
+					if (shouldHandle)
+					{
+						StringBuffer<> fixedPath;
+						FixPath(instance.builtinIncludesDir.c_str(), instance.workDir.data(), instance.workDir.size(), fixedPath);
+						if (CaseInsensitiveFs)
+							fixedPath.MakeLower();
+						fixedPath.Append(PathSeparator);
+						StringBuffer<> codeFilePath(fixedPath);
+						m_traverseFilesFunc(fixedPath, [&](const StringView& file, bool isDirectory)
+							{
+								if (isDirectory)
+									return;
+								codeFilePath.Resize(fixedPath.count).Append(file);
+								m_workManager.AddWork([this, &instance, codeFile = TString(codeFilePath.data)](const WorkContext& context)
+									{
+										StringBuffer<> path;
+										FixPath(codeFile.c_str(), instance.workDir.c_str(), instance.workDir.size(), path);
+										instance.createFileFunc(path, [&](const void* data, u64 dataSize) { return true; });
+									}, 1, TC("Crawler"), ColorWork);
+							});
+					}
+				}
+
+			}, 1, TC("Crawler"), ColorWork);
+		return true;
+	}
+
+	bool IncludeCrawler::ParseRsp(Instance& instance, const tchar* rsp, CodeFiles& outCodeFiles)
+	{
+		StringBuffer<> fixedPath;
+		FixPath(rsp, instance.workDir.data(), instance.workDir.size(), fixedPath);
+		instance.devirtualizePathFunc(fixedPath);
+		if (CaseInsensitiveFs)
+			fixedPath.MakeLower();
+		return instance.createFileFunc(fixedPath, [&](const void* data, u64 dataSize)
+			{
+				return ParseRsp2(instance, data, dataSize, outCodeFiles);
+			});
+	}
+
+	bool IncludeCrawler::ParseRsp2(Instance& instance, const void* data, u64 dataSize, CodeFiles& outCodeFiles)
+	{
+		StringBuffer<> prevArg;
+
+		auto AddCodeFile = [&](const tchar* file, bool pushFront, bool hasPch)
+			{
+				StringBuffer<> fixedPath2;
+				FixPath(file, instance.workDir.data(), instance.workDir.size(), fixedPath2);
+				instance.devirtualizePathFunc(fixedPath2);
+				if (CaseInsensitiveFs)
+					fixedPath2.MakeLower();
+				if (pushFront)
+					outCodeFiles.push_front({TString(fixedPath2.data), hasPch});
+				else
+					outCodeFiles.push_back({TString(fixedPath2.data), hasPch});
+			};
+
+		auto AddRoot = [&](const tchar* path)
+			{
+				StringBuffer<> fixedPath2;
+				FixPath(path, instance.workDir.data(), instance.workDir.size(), fixedPath2);
+				instance.devirtualizePathFunc(fixedPath2);
+				if (CaseInsensitiveFs)
+					fixedPath2.MakeLower();
+				fixedPath2.EnsureEndsWithSlash();
+
+				BloomFilter bloomFilter;
+				m_traverseFilesFunc(fixedPath2, [&](const StringView& file, bool isDirectory)
+					{
+						bloomFilter.Add(ToStringKey(file));
+					});
+				if (!bloomFilter.IsEmpty())
+					instance.includeRoots.push_back({TString(fixedPath2.data), bloomFilter});
+			};
+
+
+		ParseArguments((const char*)data, dataSize, [&](const char* arg, u32 argLen)
+			{
+				StringBuffer<> sb;
+				sb.Append(arg, argLen);
+				if (sb[0] == '-')
+				{
+					prevArg.Clear();
+					if (sb.StartsWith(TC("-isystem"), false))
+						AddRoot(sb.data + 8);
+					else if (sb.StartsWith(TC("-I"), false))
+						AddRoot(sb.data + 2);
+					else if (sb.StartsWith(TC("-resource-dir"), false))
+					{
+						instance.builtinIncludesDir = sb.data + 14;
+						instance.builtinIncludesDir += TC("/include");
+					}
+					else
+					{
+						constexpr const tchar* optionsWithArg[] = 
+						{
+							TC("-D"),
+							TC("-x"),
+							TC("-o"),
+							TC("-include"),
+							TC("-include-pch"),
+							TC("-vctoolsdir"),
+							TC("-Xclang"),
+						};
+
+						for (auto option : optionsWithArg)
+						{
+							if (!sb.Equals(option, false))
+								continue;
+							prevArg = sb;
+							break;
+						}
+					}
+				}
+				else if (sb[0] == '/')
+				{
+					if (sb.StartsWith(TC("/FI")))
+					{
+						// Check if there is a precompiled header deps file
+						StringBuffer<> fixedFile;
+						FixPath(sb.data + 3, instance.workDir.data(), instance.workDir.size(), fixedFile);
+						instance.devirtualizePathFunc(fixedFile);
+						if (CaseInsensitiveFs)
+							fixedFile.MakeLower();
+
+						StringKey key = ToStringKey(fixedFile);
+						SCOPED_WRITE_LOCK(m_codeFilesHandledLock, lock);
+						auto insres = m_codeFilesHandled.try_emplace(key);
+						lock.Leave();
+
+						HandledFile& handledFile = insres.first->second;
+						SCOPED_WRITE_LOCK(handledFile.lock, fileLock);
+						if (!handledFile.handled)
+						{
+							handledFile.handled = true;
+
+							fixedFile.Append(TC(".pch"));
+
+							u32 attributes = 0;
+							bool hasPch = m_fileExistsFunc(fixedFile, attributes);
+							AddCodeFile(sb.data + 3, true, hasPch);
+						}
+
+						prevArg.Clear();
+					}
+					else if (sb.StartsWith(TC("/Yu")))
+					{
+						instance.usePch = true;
+					}
+					else
+					{
+						prevArg = sb;
+					}
+				}
+				else if (sb[0] == '@')
+				{
+					ParseRsp(instance, sb.data + 1, outCodeFiles);
+				}
+				else
+				{
+					if (prevArg.Equals(TC("/I")) || prevArg.Equals(TC("/imsvc")) || prevArg.Equals(TC("/external:I")))
+						AddRoot(sb.data);
+					if (prevArg.Equals(TC("-include"), false))
+					{
+						if (!sb.EndsWith(TC(".gch")))
+							AddCodeFile(sb.data, true, false);
+					}
+					else if (prevArg.IsEmpty())
+						AddCodeFile(sb.data, false, false);
+					prevArg.Clear();
+				}
+			});
+		return true;
+	}
+
+	bool IncludeCrawler::ParsePch(Pch& pch, Instance& instance, const void* data, u64 dataSize)
+	{
+		auto AddHandled = [&](const StringView& str)
+			{
+				StringBuffer<> fullPath;
+				FixPath(str.data, nullptr, 0, fullPath);
+				//instance.devirtualizePathFunc(fullPath); // These paths are never virtual
+				if (CaseInsensitiveFs)
+					fullPath.MakeLower();
+				StringKey key = ToStringKey(fullPath);
+				pch.files.insert(key);
+			};
+
+		StringBuffer<> line;
+		char lastChar = 0;
+
+		const char* it = (const char*)data;
+		const char* end = it + dataSize;
+		for (; it != end; lastChar = *it, ++it)
+		{
+			if ((line.count == 0 && (*it == ' ' || *it == '\t')) || *it == '\r')
+				continue;
+			if (*it == '\\' && lastChar == ' ') // Remove the "dir \dir" extra space in clang deps files
+			{
+				line.Resize(line.count - 1);
+				if (line[line.count-1] != ':')
+					AddHandled(line);
+				line.Clear();
+				continue;
+			}
+			
+			if (*it == '\n')
+			{
+				if (line.count > 3 && !line.Contains(TC("\":")))
+				{
+					AddHandled(line);
+				}
+				line.Clear();
+				continue;
+			}
+
+			if (*it == ' ' && lastChar == '\\')
+				--line.count;
+			line.Append(*it);
+		}
+		return true;
+	}
+
+	bool IncludeCrawler::ParseCodeFile(Instance& instance, const StringView& codeFile, bool parseDefines)
+	{
+		return instance.createFileFunc(codeFile, [&](const void* data, u64 dataSize)
+			{
+				return ParseCodeFile2(instance, data, dataSize, codeFile, parseDefines);
+			});
+	}
+
+	bool IncludeCrawler::ParseCodeFile2(Instance& instance, const void* data, u64 dataSize, const StringView& codeFile, bool parseDefines)
+	{
+		const char* it = (const char*)data;
+		const char* end = it + dataSize;
+		bool hasNonSpace = false;
+
+		while (it < end)
+		{
+			if (*it == '#' && !hasNonSpace)
+			{
+				++it;
+				while (*it == ' ' || *it == '\t')
+					++it;
+
+				if (strncmp(it, "include", 7) == 0)
+				{
+					it += 7;
+					while (*it == ' ' || *it == '\t')
+						++it;
+					StringBuffer<> include;
+					bool isQuote = false;
+
+					if (*it == '\"')
+					{
+						++it;
+						const char* includeStart = it;
+						const char* includeEnd = strchr(it, '\"');
+						it = includeEnd + 1;
+						include.Append(includeStart, u32(includeEnd - includeStart));
+						isQuote =true;
+					}
+					else if (*it == '<')
+					{
+						++it;
+						const char* includeStart = it;
+						const char* includeEnd = strchr(it, '>');
+						it = includeEnd + 1;
+						include.Append(includeStart, u32(includeEnd - includeStart));
+					}
+					else
+					{
+						const char* defineBegin = it;
+						const char* defineEnd = nullptr;
+						const char* argBegin = nullptr;
+						const char* argEnd = nullptr;
+
+						while (it < end)
+						{
+							if (*it == '\r' || *it == '\n')
+							{
+								if (!defineEnd)
+									defineEnd = it;
+								break;
+							}
+							else if (*it == ' ' || *it == '\t')
+							{
+								if (!defineEnd)
+									defineEnd = it;
+							}
+							else if (*it == '(')
+							{
+								defineEnd = it;
+								argBegin = it + 1;
+							}
+							else if (*it == ')')
+							{
+								argEnd = it;
+								break;
+							}
+							else if (!argBegin && defineEnd)
+								break;
+							++it;
+						}
+
+						u32 defineLen = u32(defineEnd - defineBegin);
+						if (argBegin)
+						{
+							if (strncmp(defineBegin, "UE_INLINE_GENERATED_CPP_BY_NAME", defineLen) == 0)
+							{
+								include.Append(argBegin, u32(argEnd - argBegin)).Append(TC(".gen.cpp"));
+							}
+							else if (strncmp(defineBegin, "COMPILED_PLATFORM_HEADER", defineLen) == 0)
+							{
+								if (!instance.platform.empty())
+								{
+									if (!instance.platformIsExtension)
+										include.Append(instance.platform).Append(PathSeparator);
+									include.Append(instance.platform).Append(argBegin, u32(argEnd - argBegin));
+								}
+							}
+						}
+						else if (defineBegin)
+						{
+							if (strncmp(defineBegin, "PER_MODULE_INLINE_FILE", defineLen) == 0)
+							{
+								include.Append(TC("HAL/PerModuleInline.inl")); // TODO: This is not correct.. need revisit
+							}
+						}
+					}
+
+					if (!include.IsEmpty())
+					{
+						include.MakeLower().FixPathSeparators();
+						auto work = 
+						[this, &instance, isQuote, include2 = TString(include.data), codeFile2 = TString(codeFile.data), parseDefines](const WorkContext& context)
+							{
+								if (isQuote)
+								{
+									StringBuffer<> localDir;
+									if (const tchar* lastSeparator = TStrrchr(codeFile2.c_str(), PathSeparator))
+										if (HandleInclude(instance, localDir.Append(codeFile2.c_str(), lastSeparator - codeFile2.c_str() + 1), include2, parseDefines))
+											return;
+								}
+
+								StringKey fileKey;
+								if (const tchar* firstSeparator = TStrchr(include2.c_str(), PathSeparator))
+									fileKey = ToStringKey(include2.c_str(), (firstSeparator - include2.c_str()));
+								else
+									fileKey = ToStringKey(include2);
+
+								for (auto& root : instance.includeRoots)
+								{
+									if (root.bloomFilter.IsGuaranteedMiss(fileKey))
+										continue;
+									if (HandleInclude(instance, root.path, include2, parseDefines))
+										return;
+								}
+
+							};
+						if (parseDefines)
+						{
+							TrackWorkScope tws;
+							work({tws});
+						}
+						else
+							m_workManager.AddWork(work, 1, TC("Crawler"), ColorWork);
+					}
+				}
+				else if (parseDefines && strncmp(it, "define", 6) == 0)
+				{
+					it += 6;
+
+					auto parseDefine = [&](const char* define, u32 defineLength, TString& out)
+						{
+							if (strncmp(it, define, defineLength) != 0)
+								return false;
+							it += defineLength;
+							while (it < end && (*it == '\t' || *it == ' '))
+								++it;
+							const char* platformBegin = it;
+							while (it < end && *it != '\t' && *it != ' ' && *it != '\r' && *it != '\n')
+								++it;
+							out.assign(platformBegin, it);
+							return true;
+						};
+
+					while (*it == ' ' || *it == '\t')
+						++it;
+					TString platformIsExtension;
+					if (parseDefine("UBT_COMPILED_PLATFORM", 21, instance.compiledPlatform))
+					{}
+					else if (parseDefine("OVERRIDE_PLATFORM_HEADER_NAME", 29, instance.overriddenPlatformName))
+					{}
+					else if (parseDefine("PLATFORM_IS_EXTENSION", 21, platformIsExtension))
+					{
+						instance.platformIsExtension = platformIsExtension != TC("0");
+					}
+				}
+			}
+			else if (*it == '\n')
+			{
+				hasNonSpace = false;
+			}
+			else if (*it != ' ' && *it != '\t')
+			{
+				hasNonSpace = true;
+			}
+			++it;
+		}
+
+		return true;
+	}
+
+	bool IncludeCrawler::HandleInclude(Instance& instance, const StringView& rootPath, const StringView& include, bool parseDefines)
+	{
+		StringBuffer<> fullPath;
+		FixPath(include.data, rootPath.data, rootPath.count, fullPath);
+		if (IsAbsolutePath(include.data))
+			instance.devirtualizePathFunc(fullPath);
+		if (CaseInsensitiveFs)
+			fullPath.MakeLower();
+
+		u32 attributes = 0;
+		if (!m_fileExistsFunc(fullPath, attributes))
+			return false;
+		if (IsDirectory(attributes))
+			return false;
+
+		{
+			StringKey key = ToStringKey(fullPath);
+			if (instance.pch)
+			{
+				SCOPED_READ_LOCK(instance.pch->lock, lock);
+				if (instance.pch->files.find(key) != instance.pch->files.end())
+					return true;
+			}
+			SCOPED_WRITE_LOCK(m_codeFilesHandledLock, lock);
+			if (!m_codeFilesHandled.try_emplace(key).second)
+				return true;
+		}
+
+		ParseCodeFile(instance, fullPath, parseDefines);
+		return true;
+	}
+}
diff --git a/Engine/Source/Programs/UnrealBuildAccelerator/Common/Public/UbaSessionClient.h b/Engine/Source/Programs/UnrealBuildAccelerator/Common/Public/UbaSessionClient.h
--- a/Engine/Source/Programs/UnrealBuildAccelerator/Common/Public/UbaSessionClient.h
+++ b/Engine/Source/Programs/UnrealBuildAccelerator/Common/Public/UbaSessionClient.h
@@ -2,6 +2,7 @@
 
 #pragma once
 
+#include "UbaIncludeCrawler.h"
 #include "UbaNetwork.h"
 #include "UbaSession.h"
 
@@ -13,6 +14,8 @@ namespace uba
 	{
 		SessionClientCreateInfo(Storage& s, NetworkClient& c, LogWriter& writer = g_consoleLogWriter) : SessionCreateInfo(s, writer), client(c) {}
 
+		void Apply(const Config& config);
+
 		NetworkClient& client;
 		StringBuffer<128> name;
 		u32 maxProcessCount = 1;
@@ -26,6 +29,7 @@ namespace uba
 		bool killRandom = false;
 		bool useStorage = true;
 		bool downloadDetoursLib = true;
+		bool useIncludeCrawler = true;	// IncludeCrawler will go wide and try to prefetch all headers and source files needed by the build
 		Function<void(const ProcessHandle&)> processFinished;
 	};
 
@@ -77,7 +81,8 @@ namespace uba
 		struct InternalProcessStartInfo;
 		struct ModuleInfo;
 
-		bool GetCasKeyForFile(CasKey& out, u32 processId, const StringBufferBase& fileName, const StringKey& fileNameKey);
+		bool CreateFileForRead(CreateFileResponse& out, const StringView& fileName, const StringKey& fileNameKey, u32 processId, const ApplicationRules& rules);
+		bool GetCasKeyForFile(CasKey& out, u32 processId, const StringView& fileName, const StringKey& fileNameKey);
 		bool ReadModules(List<ModuleInfo>& outModules, u32 processId, const tchar* application);
 		bool EnsureBinaryFile(StringBufferBase& out, StringBufferBase& outVirtual, u32 processId, const StringBufferBase& fileName, const StringKey& fileNameKey, const tchar* applicationDir, const u8* loaderPaths, u32 loaderPathsSize);
 		bool WriteBinFile(StringBufferBase& out, const tchar* binaryName, const CasKey& casKey, const KeyToString& applicationDir, u32 fileAttributes);
@@ -99,6 +104,9 @@ namespace uba
 		void SendLogFileToServer(ProcessImpl& pi);
 		void GetLogFileName(StringBufferBase& out, const tchar* logFile, const tchar* arguments, u32 processId);
 		u32 WriteLogLines(BinaryWriter& writer, ProcessImpl& process);
+		void RunIncludeCrawler(ProcessImpl& process);
+		bool ParseDirectoryTable();
+		bool EntryExists(const StringView& path, u32& outTableOffset);
 
 		void ThreadCreateProcessLoop();
 
@@ -124,10 +132,12 @@ namespace uba
 		bool m_shouldSendLogToServer = false;
 		bool m_shouldSendTraceToServer = false;
 		bool m_remoteExecutionEnabled = true;
+		bool m_useIncludeCrawler = false;
 		
 		Atomic<const tchar*> m_terminationReason;
 		Atomic<u64> m_terminationTime;
 		Atomic<u32> m_maxProcessCount;
+		Atomic<float> m_cpuUsage;
 
 		struct ApplicationEnvironment { ReaderWriterLock lock; TString virtualApplication; TString realApplication; };
 		ReaderWriterLock m_handledApplicationEnvironmentsLock;
@@ -165,7 +175,16 @@ namespace uba
 		u64 m_lastPing = 0;
 		u64 m_lastPingSendTime = 0;
 
+		Atomic<u64> m_memAvail;
+		Atomic<u64> m_memTotal;
+
 		UnorderedMap<CasKey, Vector<u8>> m_hostRunCache;
 		ReaderWriterLock m_hostRunCacheLock;
+
+		IncludeCrawler m_includeCrawler;
+		u32 m_dirtableParsedPosition = 0;
+		ReaderWriterLock m_dirVisitedLock;
+		struct DirVisitedEntry { ReaderWriterLock lock; bool handled = false; };
+		UnorderedMap<StringKey, DirVisitedEntry> m_dirVisited;
 	};
 }
diff --git a/Engine/Source/Programs/UnrealBuildAccelerator/Core/Private/UbaApplicationRules.cpp b/Engine/Source/Programs/UnrealBuildAccelerator/Core/Private/UbaApplicationRules.cpp
--- a/Engine/Source/Programs/UnrealBuildAccelerator/Core/Private/UbaApplicationRules.cpp
+++ b/Engine/Source/Programs/UnrealBuildAccelerator/Core/Private/UbaApplicationRules.cpp
@@ -115,11 +115,16 @@ namespace uba
 			return fileName.EndsWith(TC(".obj"));
 		}
 
-		virtual bool ShouldDevirtualizeFile(const StringView& fileName, bool& outEscapeSpaces) const
+		virtual bool ShouldDevirtualizeFile(const StringView& fileName, bool& outEscapeSpaces) const override
 		{
 			outEscapeSpaces = false;
 			return fileName.EndsWith(TC(".dep.json"));
 		}
+
+		virtual bool SupportsIncludeCrawler() const override
+		{
+			return true;
+		}
 	};
 
 	class ApplicationRulesVcLink : public ApplicationRulesVC
@@ -280,21 +285,26 @@ namespace uba
 			return true;
 		}
 
-		virtual bool StoreFileCompressed(const StringView& fileName) const
+		virtual bool StoreFileCompressed(const StringView& fileName) const override
 		{
 			return IsCompressedObjFile(fileName) || fileName.EndsWith(TC(".o"));
 		}
 
-		virtual bool ShouldExtractSymbols(const StringView& fileName) const
+		virtual bool ShouldExtractSymbols(const StringView& fileName) const override
 		{
 			return fileName.EndsWith(TC(".obj")) || fileName.EndsWith(TC(".o"));
 		}
 
-		virtual bool ShouldDevirtualizeFile(const StringView& fileName, bool& outEscapeSpaces) const
+		virtual bool ShouldDevirtualizeFile(const StringView& fileName, bool& outEscapeSpaces) const override
 		{
 			outEscapeSpaces = true;
 			return fileName.EndsWith(TC(".d"));
 		}
+
+		virtual bool SupportsIncludeCrawler() const override
+		{
+			return true;
+		}
 	};
 
 	class ApplicationRulesLdLLdExe : public ApplicationRulesClang
@@ -494,17 +504,17 @@ namespace uba
 			return true;
 		}
 
-		virtual bool StoreFileCompressed(const StringView& fileName) const
+		virtual bool StoreFileCompressed(const StringView& fileName) const override
 		{
 			return fileName.EndsWith(TC(".obj"));
 		}
 
-		virtual bool ShouldExtractSymbols(const StringView& fileName) const
+		virtual bool ShouldExtractSymbols(const StringView& fileName) const override
 		{
 			return fileName.EndsWith(TC(".obj")) || fileName.EndsWith(TC(".o"));
 		}
 
-		virtual bool ShouldDevirtualizeFile(const StringView& fileName, bool& outEscapeSpaces) const
+		virtual bool ShouldDevirtualizeFile(const StringView& fileName, bool& outEscapeSpaces) const override
 		{
 			return fileName.EndsWith(TC(".ispc.txt"));
 		}
@@ -560,7 +570,7 @@ namespace uba
 			return file.Contains(TC(".uba."));
 		}
 
-		virtual bool AllowDetach() const
+		virtual bool AllowDetach() const override
 		{
 			return true;
 		}
@@ -574,7 +584,7 @@ namespace uba
 				|| fileName.EndsWith(TC(".exp"));
 		}
 
-		virtual bool StoreFileCompressed(const StringView& fileName) const
+		virtual bool StoreFileCompressed(const StringView& fileName) const override
 		{
 			return fileName.EndsWith(TC(".obj")); // Extra obj file for msvc
 		}
diff --git a/Engine/Source/Programs/UnrealBuildAccelerator/Core/Public/UbaApplicationRules.h b/Engine/Source/Programs/UnrealBuildAccelerator/Core/Public/UbaApplicationRules.h
--- a/Engine/Source/Programs/UnrealBuildAccelerator/Core/Public/UbaApplicationRules.h
+++ b/Engine/Source/Programs/UnrealBuildAccelerator/Core/Public/UbaApplicationRules.h
@@ -167,6 +167,11 @@ namespace uba
 			return false;
 		}
 
+		virtual bool SupportsIncludeCrawler() const
+		{
+			return false;
+		}
+
 		u32 index = ~0u;
 	};
 
diff --git a/Engine/Source/Programs/UnrealBuildAccelerator/Test/Private/UbaTestAll.h b/Engine/Source/Programs/UnrealBuildAccelerator/Test/Private/UbaTestAll.h
--- a/Engine/Source/Programs/UnrealBuildAccelerator/Test/Private/UbaTestAll.h
+++ b/Engine/Source/Programs/UnrealBuildAccelerator/Test/Private/UbaTestAll.h
@@ -69,6 +69,7 @@ namespace uba
 		UBA_TEST(TestStdOutLocal) \
 		UBA_TEST(TestStdOutViaCmd) \
 		UBA_TEST(TestVolumeCache) \
+		UBA_TEST(TestIncludeCrawler) \
 
 
 #define UBA_LINUX_TESTS \
diff --git a/Engine/Source/Programs/UnrealBuildAccelerator/Test/Private/UbaTestIncludeCrawler.cpp b/Engine/Source/Programs/UnrealBuildAccelerator/Test/Private/UbaTestIncludeCrawler.cpp
new file mode 100644
--- /dev/null
+++ b/Engine/Source/Programs/UnrealBuildAccelerator/Test/Private/UbaTestIncludeCrawler.cpp
@@ -0,0 +1,58 @@
+// Copyright Epic Games, Inc. All Rights Reserved.
+
+#include "UbaIncludeCrawler.h"
+#include "UbaDirectoryIterator.h"
+#include "UbaLogger.h"
+
+namespace uba
+{
+	bool TestIncludeCrawler(LoggerWithWriter& logger, const StringBufferBase& rootDir)
+	{
+#if 0 // Commented out for now because of hard coded paths below
+
+		WorkManagerImpl workManager(128);
+
+		IncludeCrawler crawler(logger, workManager);
+
+		auto CreateFileFunc = [&](const StringView& fileName, const IncludeCrawler::AccessFileFunc& func)
+			{
+				FileAccessor fa(logger, fileName.data);
+				if (!fa.OpenMemoryRead())
+					return false;
+				return func(fa.GetData(), fa.GetSize());
+			};
+
+		auto FileExistsFunc = [&](const StringView& fileName, u32& outAttr)
+			{
+				return FileExists(logger, fileName.data, nullptr, &outAttr);
+			};
+
+		auto TraverseFilesFunc = [&](const StringView& path, const IncludeCrawler::FileFunc& fileFunc)
+			{
+				TraverseDir(logger, path, [&](const DirectoryEntry& entry)
+					{
+						StringBuffer<> name;
+						name.Append(entry.name, entry.nameLen);
+						if (CaseInsensitiveFs)
+							name.MakeLower();
+						fileFunc(name, IsDirectory(entry.attributes));
+					});
+
+			};
+
+		crawler.Init(FileExistsFunc, TraverseFilesFunc);
+
+		const tchar app[] = TC("c:\\sdk\\AutoSDK\\HostWin64\\Win64\\LLVM\\18.1.8\\bin\\clang-cl.exe");
+		const tchar rsp[] = TC("e:\\dev\\fn\\Engine\\Intermediate\\Build\\Win64\\x64\\BlankProgram\\Development\\Core\\Module.Core.1.cpp.obj.rsp");
+		//const tchar rsp[] = TC("e:\\dev\\fn\\engine\\plugins\\experimental\\geometrycollectionplugin\\intermediate\\build\\win64\\x64\\unrealeditor\\development\\geometrycollectionnodes\\module.geometrycollectionnodes.3.cpp.obj.rsp");
+		//const tchar rsp[] = TC("e:\\dev\\fn\\Engine\\Plugins\\Runtime\\MeshModelingToolset\\Intermediate\\Build\\Win64\\x64\\UnrealEditor\\Development\\ModelingComponents\\Module.ModelingComponents.3.cpp.obj.rsp");
+		//const tchar rsp[] = TC("e:\\dev\\fn\\Engine\\Intermediate\\Build\\Win64\\x64\\UnrealEditor\\Development\\AnimGraphRuntime\\Module.AnimGraphRuntime.3.cpp.obj.rsp");
+
+		auto rules = GetApplicationRules()[10].rules;
+		crawler.Add(rsp, TC("e:\\dev\\fn\\Engine\\Source"), CreateFileFunc, [](StringBufferBase&) { return true; }, app, rules);
+
+		workManager.FlushWork();
+#endif
+		return true;
+	}
+}
\ No newline at end of file
