[UBA]
* Fix in include crawler for builtin clang includes

diff --git a/Engine/Source/Programs/UnrealBuildAccelerator/Common/Public/UbaIncludeCrawler.h b/Engine/Source/Programs/UnrealBuildAccelerator/Common/Public/UbaIncludeCrawler.h
--- a/Engine/Source/Programs/UnrealBuildAccelerator/Common/Public/UbaIncludeCrawler.h
+++ b/Engine/Source/Programs/UnrealBuildAccelerator/Common/Public/UbaIncludeCrawler.h
@@ -167,6 +167,7 @@ namespace uba
 					{
 						StringBuffer<> fixedPath;
 						FixPath(instance.builtinIncludesDir.c_str(), instance.workDir.data(), instance.workDir.size(), fixedPath);
+						instance.devirtualizePathFunc(fixedPath);
 						if (CaseInsensitiveFs)
 							fixedPath.MakeLower();
 						fixedPath.Append(PathSeparator);
@@ -178,9 +179,7 @@ namespace uba
 								codeFilePath.Resize(fixedPath.count).Append(file);
 								m_workManager.AddWork([this, &instance, codeFile = TString(codeFilePath.data)](const WorkContext& context)
 									{
-										StringBuffer<> path;
-										FixPath(codeFile.c_str(), instance.workDir.c_str(), instance.workDir.size(), path);
-										instance.createFileFunc(path, [&](const void* data, u64 dataSize) { return true; });
+										instance.createFileFunc(codeFile, [&](const void* data, u64 dataSize) { return true; });
 									}, 1, TC("Crawler"), ColorWork);
 							});
 					}
