[UBA]
* Added support for not writing cas files to disk on client side (only works on windows atm)

diff --git a/Engine/Source/Programs/UnrealBuildAccelerator/Agent/Private/UbaAgent.cpp b/Engine/Source/Programs/UnrealBuildAccelerator/Agent/Private/UbaAgent.cpp
--- a/Engine/Source/Programs/UnrealBuildAccelerator/Agent/Private/UbaAgent.cpp
+++ b/Engine/Source/Programs/UnrealBuildAccelerator/Agent/Private/UbaAgent.cpp
@@ -538,7 +538,8 @@ namespace uba
 			}
 			else if (name.Equals(TC("-nostore")))
 			{
-				useStorage = false;
+				if (IsWindows) // Only supported on windows atm.
+					useStorage = false;
 			}
 			else if (name.Equals(TC("-nohandler")))
 			{
@@ -807,7 +808,7 @@ namespace uba
 
 		u64 storageCapacity = u64(storageCapacityGb)*1000*1000*1000;
 
-		if (command.empty())
+		if (command.empty() && useStorage)
 		{
 			// Create a uba storage quickly just to fix non-graceful shutdowns
 			StorageCreateInfo info(g_rootDir.data, logWriter);
@@ -1140,6 +1141,7 @@ namespace uba
 			storageInfo.startProxyUserData = &startProxy;
 			storageInfo.zone = zone.data;
 			storageInfo.proxyPort = proxyPort;
+			storageInfo.writeToDisk = useStorage;
 
 			storageClient = new StorageClient(storageInfo);
 			auto bscsg = MakeGuard([&]() { delete storageClient; });
diff --git a/Engine/Source/Programs/UnrealBuildAccelerator/Common/Private/UbaSession.cpp b/Engine/Source/Programs/UnrealBuildAccelerator/Common/Private/UbaSession.cpp
--- a/Engine/Source/Programs/UnrealBuildAccelerator/Common/Private/UbaSession.cpp
+++ b/Engine/Source/Programs/UnrealBuildAccelerator/Common/Private/UbaSession.cpp
@@ -1071,6 +1071,12 @@ namespace uba
 			applicationDir = temp2;
 			applicationDirEnd = temp2 + applicationDirLen;
 		}
+		else
+		{
+			#if PLATFORM_WINDOWS
+			attr = DefaultAttributes();
+			#endif
+		}
 
 		FixPath(applicationName, nullptr, 0, temp3.Clear());
 
diff --git a/Engine/Source/Programs/UnrealBuildAccelerator/Common/Private/UbaStorage.cpp b/Engine/Source/Programs/UnrealBuildAccelerator/Common/Private/UbaStorage.cpp
--- a/Engine/Source/Programs/UnrealBuildAccelerator/Common/Private/UbaStorage.cpp
+++ b/Engine/Source/Programs/UnrealBuildAccelerator/Common/Private/UbaStorage.cpp
@@ -158,7 +158,7 @@ namespace uba
 
 		casEntry.size = size;
 
-		UBA_ASSERT(!casEntry.mappingHandle.IsValid());
+		UBA_ASSERT(!m_casCapacityBytes || !casEntry.mappingHandle.IsValid());
 
 		if (!m_casCapacityBytes || m_overflowReported || m_manuallyHandleOverflow)
 			return;
@@ -1206,13 +1206,20 @@ namespace uba
 		m_storeCompressed = info.storeCompressed;
 		m_manuallyHandleOverflow = info.manuallyHandleOverflow;
 		m_asyncUnmapViewOfFile = info.asyncUnmapViewOfFile && m_workManager;
+		m_writeToDisk = info.writeToDisk;
+		if (!m_writeToDisk)
+			m_casCapacityBytes = 0;
+
 		m_rootDir.count = GetFullPathNameW(info.rootDir, m_rootDir.capacity, m_rootDir.data, NULL);
 		m_rootDir.Replace('/', PathSeparator).EnsureEndsWithSlash();
 
-		m_tempPath.Append(m_rootDir).Append(TC("castemp"));
-		CreateDirectory(m_tempPath.data);
-		DeleteAllFiles(m_logger, m_tempPath.data, false);
-		m_tempPath.EnsureEndsWithSlash();
+		if (m_writeToDisk)
+		{
+			m_tempPath.Append(m_rootDir).Append(TC("castemp"));
+			CreateDirectory(m_tempPath.data);
+			DeleteAllFiles(m_logger, m_tempPath.data, false);
+			m_tempPath.EnsureEndsWithSlash();
+		}
 
 		m_rootDir.Append(TC("cas")).EnsureEndsWithSlash();
 
@@ -1250,6 +1257,9 @@ namespace uba
 
 	bool StorageImpl::LoadCasTable(bool logStats, bool alwaysCheckAllFiles)
 	{
+		if (!m_writeToDisk)
+			return true;
+
 		static bool isExclusive = CheckExclusive(m_logger, m_rootDir);
 		if (!isExclusive)
 			return false;
@@ -2120,8 +2130,9 @@ namespace uba
 	bool StorageImpl::GetCasFileName(StringBufferBase& out, const CasKey& casKey)
 	{
 		out.Append(m_rootDir.data).AppendHex(((const u8*)&casKey)[0]);
-		if (!CreateDirectory(out.data))
-			return false;
+		if (m_writeToDisk)
+			if (!CreateDirectory(out.data))
+				return false;
 		out.Append(PathSeparator).Append(CasKeyString(casKey).str);
 		return true;
 	}
diff --git a/Engine/Source/Programs/UnrealBuildAccelerator/Common/Private/UbaStorageClient.cpp b/Engine/Source/Programs/UnrealBuildAccelerator/Common/Private/UbaStorageClient.cpp
--- a/Engine/Source/Programs/UnrealBuildAccelerator/Common/Private/UbaStorageClient.cpp
+++ b/Engine/Source/Programs/UnrealBuildAccelerator/Common/Private/UbaStorageClient.cpp
@@ -197,12 +197,20 @@ namespace uba
 		return true;
 	}
 
-	bool StorageClient::RetrieveCasFile(RetrieveResult& out, const CasKey& casKey, const tchar* hint, FileMappingBuffer* mappingBuffer, u64 memoryMapAlignment, bool allowProxy)
+	bool StorageClient::RetrieveCasFile(RetrieveResult& out, const CasKey& casKeyTmp, const tchar* hint, FileMappingBuffer* mappingBuffer, u64 memoryMapAlignment, bool allowProxy)
 	{
+		CasKey casKey = casKeyTmp;
+
 		FileMappingType mappingType = MappedView_Transient;
 		bool shouldStore = mappingBuffer == nullptr;
 		UBA_ASSERT(AsCompressed(casKey, false) != CasKeyZero);
 
+		if (!m_writeToDisk && !mappingBuffer)
+		{
+			mappingBuffer = &m_casDataBuffer;
+			casKey = AsCompressed(casKey, true);
+		}
+
 		out.casKey = casKey;
 		out.size = InvalidValue;
 
@@ -476,16 +484,6 @@ namespace uba
 						return true;
 					isInitialized = true;
 
-					if (!mappingBuffer)
-					{
-						u32 extraFlags = DefaultAttributes();
-						bool useOverlap = !IsRunningWine() && isCompressed == IsCompressed(casKey) && sizeOnDisk > 1024 * 1024;
-						if (useOverlap)
-							extraFlags |= FILE_FLAG_OVERLAPPED;
-						if (!destinationFile.CreateWrite(false, extraFlags, sizeOnDisk, m_tempPath.data))
-							return false;
-					}
-
 					if (mappingBuffer)
 					{
 						UBA_ASSERT(!writeMem || mappedView.size == sizeOnDisk);
@@ -497,6 +495,15 @@ namespace uba
 								return false;
 						}
 					}
+					else
+					{
+						u32 extraFlags = DefaultAttributes();
+						bool useOverlap = !IsRunningWine() && isCompressed == IsCompressed(casKey) && sizeOnDisk > 1024 * 1024;
+						if (useOverlap)
+							extraFlags |= FILE_FLAG_OVERLAPPED;
+						if (!destinationFile.CreateWrite(false, extraFlags, sizeOnDisk, m_tempPath.data))
+							return false;
+					}
 					writePos = writeMem;
 					return true;
 				};
diff --git a/Engine/Source/Programs/UnrealBuildAccelerator/Common/Public/UbaStorage.h b/Engine/Source/Programs/UnrealBuildAccelerator/Common/Public/UbaStorage.h
--- a/Engine/Source/Programs/UnrealBuildAccelerator/Common/Public/UbaStorage.h
+++ b/Engine/Source/Programs/UnrealBuildAccelerator/Common/Public/UbaStorage.h
@@ -88,6 +88,7 @@ namespace uba
 		bool storeCompressed = true;
 		bool manuallyHandleOverflow = false;
 		bool asyncUnmapViewOfFile = true;
+		bool writeToDisk = true;
 		WorkManager* workManager = nullptr;
 	};
 
@@ -241,6 +242,7 @@ namespace uba
 		bool m_storeCompressed = false;
 		bool m_manuallyHandleOverflow = false;
 		bool m_asyncUnmapViewOfFile = true;
+		bool m_writeToDisk;
 
 
 		ExternalFileMappingsProvider m_externalFileMappingsProvider;
