[UBA]
* Minor visualizer optimization where 1-pixel wide box is not drawn if pixels have already been drawn by previous box

diff --git a/Engine/Source/Programs/UnrealBuildAccelerator/Visualizer/Private/UbaVisualizer.cpp b/Engine/Source/Programs/UnrealBuildAccelerator/Visualizer/Private/UbaVisualizer.cpp
--- a/Engine/Source/Programs/UnrealBuildAccelerator/Visualizer/Private/UbaVisualizer.cpp
+++ b/Engine/Source/Programs/UnrealBuildAccelerator/Visualizer/Private/UbaVisualizer.cpp
@@ -1526,6 +1526,8 @@ namespace uba
 							ExtTextOutW(hdc, 5, posY + offsetY, ETO_CLIPPED, &rect, buf.data, buf.count, NULL);
 						}
 
+						int lastDrawnRight = 0;
+
 						u32 workIndex = 0;
 						int posX = int(m_scrollPosX) + progressRect.left;
 						for (auto& work : workTrack.records)
@@ -1567,6 +1569,8 @@ namespace uba
 							if (rect.right <= progressRect.left)
 								continue;
 
+							++visibleBoxes;
+
 							rect.right = Max(int(rect.right), left + 1);
 
 							Color color = work.color;
@@ -1579,6 +1583,14 @@ namespace uba
 								for (u32 j=0; j!=3; ++j)
 									c[j] = u8(Min(int(c[j]) + 40, 255));
 							}
+							else
+							{
+								if (rect.left + 1 == rect.right)
+									if (lastDrawnRight == rect.right)
+										continue;
+							}
+
+							lastDrawnRight = rect.right;
 
 							bool done = stop != ~u64(0);
 							if (done)
@@ -1605,8 +1617,6 @@ namespace uba
 							FillRect(hdc, &rect, brush);
 							++rect.top;
 
-							++visibleBoxes;
-
 							int processWidth = rect.right - rect.left;
 							if (shouldDrawText && m_config.ShowProcessText && processWidth > 3)
 							{
