[UBA]
* Changed so message sockets are non-blocking and added a 128kb cache to be able to recv more data at the time (and reduce number of kernel calls)

diff --git a/Engine/Source/Programs/UnrealBuildAccelerator/Common/Private/UbaNetworkBackendTcp.cpp b/Engine/Source/Programs/UnrealBuildAccelerator/Common/Private/UbaNetworkBackendTcp.cpp
--- a/Engine/Source/Programs/UnrealBuildAccelerator/Common/Private/UbaNetworkBackendTcp.cpp
+++ b/Engine/Source/Programs/UnrealBuildAccelerator/Common/Private/UbaNetworkBackendTcp.cpp
@@ -31,6 +31,7 @@
 #define SD_BOTH SHUT_RDWR
 #define WSAHOST_NOT_FOUND 0
 #define WSAENOTCONN ENOTCONN
+#define WSAEWOULDBLOCK EWOULDBLOCK
 #define WSAEADDRINUSE EADDRINUSE
 #define closesocket(a) close(a)
 #define addrinfoW addrinfo
@@ -42,12 +43,6 @@
 #define WSAPoll poll
 #endif
 
-// Hard to say if this lock should be enabled on windows or not.
-// Local tests with cache server and 100% cache hits make no-lock 35s vs lock 40s.
-// But by looking at the farm jobs it looks like more time is spent in send without the lock
-// Saw a network error today that did not make sense. Could be caused by not locking around send
-#define UBA_LOCK_AROUND_SEND 1//!PLATFORM_WINDOWS
-
 namespace uba
 {
 	constexpr u32 MaxHeaderSize = 24;
@@ -103,7 +98,14 @@ namespace uba
 	bool SetSendBuf(Logger& logger, SOCKET socket, u32 windowSize);
 	bool DisableNagle(Logger& logger, SOCKET socket);
 	bool SendSocket(Logger& logger, SOCKET socket, const void* b, u64 bufferLen);
-	bool RecvSocket(Logger& logger, SOCKET socket, void* b, u32& bufferLen, u32 timeoutMs, const Guid& connection, const tchar* hint1, const tchar* hint2, bool isFirstCall, bool allowLess);
+
+	struct RecvCache
+	{
+		u8 bytes[128*1024];
+		u32 byteCount = 0;
+	};
+
+	bool RecvSocket(Logger& logger, RecvCache& recvCache, SOCKET socket, void* b, u32& bufferLen, u32 timeoutMs, const Guid& connection, const tchar* hint1, const tchar* hint2, bool isFirstCall, bool allowLess);
 
 	bool NetworkBackendTcp::EnsureInitialized(Logger& logger)
 	{
@@ -201,18 +203,11 @@ namespace uba
 		auto& conn = *(Connection*)connection;
 		sendContext.isUsed = true;
 
-		#if UBA_LOCK_AROUND_SEND
 		SCOPED_WRITE_LOCK(conn.sendLock, lock);
-		#else
-		SCOPED_READ_LOCK(conn.sendLock, lock);
-		#endif
 		if (conn.socket == INVALID_SOCKET)
 			return false;
 		bool res = SendSocket(logger, conn.socket, data, dataSize);
-
-		#if UBA_LOCK_AROUND_SEND
 		lock.Leave();
-		#endif
 
 		sendContext.isFinished = true;
 
@@ -483,6 +478,13 @@ namespace uba
 
 		if (connection.ready.IsSet(60000)) // This should never time out!
 		{
+			SetBlocking(logger, connection.socket, false);
+			u32 recvTimeoutMs = connection.recvTimeoutMs;
+			if (!recvTimeoutMs)
+				recvTimeoutMs = 10 * 60 * 1000; // 10 minutes time out
+
+			RecvCache recvCache;
+
 			bool isFirst = true;
 			while (connection.socket != INVALID_SOCKET)
 			{
@@ -491,7 +493,7 @@ namespace uba
 				u32 bodySize = 0;
 
 				u8 headerData[MaxHeaderSize];
-				if (!RecvSocket(logger, connection.socket, headerData, connection.headerSize, connection.recvTimeoutMs, connection.uid, connection.recvHint, TC(""), isFirst, false))
+				if (!RecvSocket(logger, recvCache, connection.socket, headerData, connection.headerSize, recvTimeoutMs, connection.uid, connection.recvHint, TC(""), isFirst, false))
 					break;
 				isFirst = false;
 
@@ -509,7 +511,7 @@ namespace uba
 				if (!bodySize)
 					continue;
 
-				bool success = RecvSocket(logger, connection.socket, bodyData, bodySize, connection.recvTimeoutMs, connection.uid, connection.recvHint, TC("Body"), false, connection.allowLess);
+				bool success = RecvSocket(logger, recvCache, connection.socket, bodyData, bodySize, recvTimeoutMs, connection.uid, connection.recvHint, TC("Body"), false, connection.allowLess);
 
 				m_totalRecv += bodySize;
 
@@ -814,35 +816,122 @@ namespace uba
 
 	bool SendSocket(Logger& logger, SOCKET socket, const void* b, u64 bufferLen)
 	{
+		char* buffer = (char*)b;
 		u64 left = bufferLen;
 		while (left)
 		{
-			int sent = (int)send(socket, (char*)b, u32(bufferLen), 0);
+			int sent = (int)send(socket, (char*)buffer, u32(left), 0);
 			if (sent == SOCKET_ERROR)
 			{
+				u32 lastError = WSAGetLastError();
+				if (lastError == WSAEWOULDBLOCK)
+				{
+					u64 timeoutMs = 10 * 60 * 1000; // 10 minutes send timeout
+
+					WSAPOLLFD p;
+					p.fd = socket;
+					p.revents = 0;
+					p.events = POLLWRNORM;
+					int res = WSAPoll(&p, 1, int(timeoutMs));
+					if (!res)
+					{
+						logger.Info(TC("WSAPoll returned timeout for connection trying to send data"));
+						return false;
+					}
+					if (res == SOCKET_ERROR)
+					{
+	#if UBA_DEBUG && PLATFORM_WINDOWS
+						// When cancelling all kinds of errors can happen..
+						lastError = WSAGetLastError();
+						if (lastError != WSAEINTR && lastError != WSAESHUTDOWN && lastError != WSAECONNABORTED && lastError != WSAECONNRESET) // Interrupted by cancel
+							logger.Warning(TC("WSAPoll returned an error for connection"));
+	#endif
+						return false;
+					}
+					continue;
+				}
 				//#if UBA_DEBUG
 				//logger.Warning(TC("ERROR sending socket (error: %s)"), LastErrorToText(WSAGetLastError()).data);
 				//#endif
 				return false;
 			}
 
+			buffer += sent;
 			left -= sent;
-
-			#if !UBA_LOCK_AROUND_SEND
-			if (left != 0)
-				return logger.Error(TC("Failed to send all data in one call. Wanted to send %llu, sent %i"), bufferLen, sent);
-			#endif
 		}
 		return true;
 	}
 
-	bool RecvSocket(Logger& logger, SOCKET socket, void* b, u32& bufferLen, u32 timeoutMs, const Guid& connection, const tchar* hint1, const tchar* hint2, bool isFirstCall, bool allowLess)
+	bool RecvSocket(Logger& logger, RecvCache& recvCache, SOCKET socket, void* b, u32& bufferLen, u32 timeoutMs, const Guid& connection, const tchar* hint1, const tchar* hint2, bool isFirstCall, bool allowLess)
 	{
 		u8* buffer = (u8*)b;
 		u32 recvLeft = bufferLen;
 		while (recvLeft)
 		{
-			if (timeoutMs)
+			if (recvCache.byteCount)
+			{
+				u32 toCopy = Min(recvCache.byteCount, recvLeft);
+				memcpy(buffer, recvCache.bytes, toCopy);
+				memmove(recvCache.bytes, recvCache.bytes + toCopy, recvCache.byteCount - toCopy);
+				recvCache.byteCount -= toCopy;
+				recvLeft -= toCopy;
+				buffer += toCopy;
+				if (!recvLeft)
+					return true;
+			}
+
+			u8* writePos;
+			u32 writeCapacity;
+
+			if (sizeof(recvCache.bytes) < recvLeft)
+			{
+				writePos = buffer;
+				writeCapacity = recvLeft;
+			}
+			else
+			{
+				writePos = recvCache.bytes + recvCache.byteCount;
+				writeCapacity = sizeof(recvCache.bytes) - recvCache.byteCount;
+			}
+
+			int read = (int)recv(socket, (char*)writePos, writeCapacity, 0);
+
+			if (read > 0)
+			{
+				if (writePos == buffer)
+				{
+					recvLeft -= read;
+					buffer += read;
+				}
+				else
+				{
+					recvCache.byteCount += read;
+					u32 toCopy = Min(recvCache.byteCount, recvLeft);
+					memcpy(buffer, recvCache.bytes, toCopy);
+					memmove(recvCache.bytes, recvCache.bytes + toCopy, recvCache.byteCount - toCopy);
+					recvCache.byteCount -= toCopy;
+					recvLeft -= toCopy;
+					buffer += toCopy;
+				}
+
+				if (recvLeft == 0)
+					return true;
+			
+				if (allowLess)
+				{
+					bufferLen = read;
+					return true;
+				}
+				continue;
+			}
+
+			if (read == 0)
+				return false;
+
+			UBA_ASSERT(read == SOCKET_ERROR);
+
+			int lastError = WSAGetLastError();
+			if (lastError == WSAEWOULDBLOCK)
 			{
 				WSAPOLLFD p;
 				p.fd = socket;
@@ -856,47 +945,27 @@ namespace uba
 				}
 				if (res == SOCKET_ERROR)
 				{
-#if UBA_DEBUG && PLATFORM_WINDOWS
+					#if UBA_DEBUG && PLATFORM_WINDOWS
 					// When cancelling all kinds of errors can happen..
-					int lastError = WSAGetLastError();
+					lastError = WSAGetLastError();
 					if (lastError != WSAEINTR && lastError != WSAESHUTDOWN && lastError != WSAECONNABORTED && lastError != WSAECONNRESET) // Interrupted by cancel
 						logger.Warning(TC("WSAPoll returned an error for connection %s: %s (%s%s)"), GuidToString(connection).str, LastErrorToText(lastError).data, hint1, hint2);
-#endif
+					#endif
+					return false;
 				}
+				continue;
 			}
-
-			int read = (int)recv(socket, (char*)buffer, recvLeft, 0);
-			if (read == 0)
-			{
-				//#if UBA_DEBUG
-				//logger.Warning(TC("Socket closed while in recv"));
-				//#endif
-				return false;
-			}
-
-			if (read == SOCKET_ERROR)
-			{
 #if PLATFORM_WINDOWS
 #if UBA_DEBUG
-				// When cancelling all kinds of errors can happen..
-				int lastError = WSAGetLastError();
-				if (lastError != WSAEINTR && lastError != WSAESHUTDOWN && lastError != WSAECONNABORTED && lastError != WSAECONNRESET) // Interrupted by cancel
-					logger.Warning(TC("ERROR receiving socket: %s (%s%s)"), LastErrorToText(lastError).data, hint1, hint2);
+			// When cancelling all kinds of errors can happen..
+			if (lastError != WSAEINTR && lastError != WSAESHUTDOWN && lastError != WSAECONNABORTED && lastError != WSAECONNRESET) // Interrupted by cancel
+				logger.Warning(TC("ERROR receiving socket: %s (%s%s)"), LastErrorToText(lastError).data, hint1, hint2);
 #endif
 #else
-				if (!isFirstCall && errno != ECONNRESET)
-					logger.Error(TC("ERROR receiving socket %i after %u bytes (%s%s) (%s)"), socket, bufferLen, hint1, hint2, strerror(errno));
+			if (!isFirstCall && errno != ECONNRESET)
+				logger.Error(TC("ERROR receiving socket %i after %u bytes (%s%s) (%s)"), socket, bufferLen, hint1, hint2, strerror(errno));
 #endif
-				return false;
-			}
-			recvLeft -= (u32)read;
-			buffer += read;
-
-			if (allowLess)
-			{
-				bufferLen = read;
-				break;
-			}
+			return false;
 		}
 		return true;
 	}
