[UBA]
* Fixed shutdown crash in unit test

diff --git a/Engine/Source/Programs/UnrealBuildAccelerator/Test/Private/UbaTestNetwork.cpp b/Engine/Source/Programs/UnrealBuildAccelerator/Test/Private/UbaTestNetwork.cpp
--- a/Engine/Source/Programs/UnrealBuildAccelerator/Test/Private/UbaTestNetwork.cpp
+++ b/Engine/Source/Programs/UnrealBuildAccelerator/Test/Private/UbaTestNetwork.cpp
@@ -73,7 +73,7 @@ namespace uba
 				return true;
 			});
 
-		auto ds = MakeGuard([&]() { server.DisconnectClients(); });
+		auto ds = MakeGuard([&]() { networkBackend.StopListen(); server.DisconnectClients(); });
 		if (!server.StartListen(networkBackend, 1234))
 			return logger.Error(TC("Failed to listen"));
 		Sleep(100);
@@ -114,7 +114,7 @@ namespace uba
 		if (!client.StartListen(networkBackend, 1239))
 			return logger.Error(TC("Client failed to listen"));
 		Sleep(100);
-		auto ds = MakeGuard([&]() { server.DisconnectClients(); });
+		auto ds = MakeGuard([&]() { networkBackend.StopListen(); server.DisconnectClients(); });
 		auto dg = MakeGuard([&]() { client.Disconnect(); });
 		if (!server.AddClient(networkBackend, TC("127.0.0.1"), 1239))
 			return logger.Error(TC("Server failed to connect"));
