[UBA]
* Fixed so CreateFile retries if ERROR_USER_MAPPED_FILE error happens. This can probably happen since we defer unmapviewoffile in a few places to improve performance
* Fixed so cl.exe is not allowing .i files through proxy (these are preprocessed files that will never be shared)

diff --git a/Engine/Source/Programs/UnrealBuildAccelerator/Common/Private/UbaFileAccessor.cpp b/Engine/Source/Programs/UnrealBuildAccelerator/Common/Private/UbaFileAccessor.cpp
--- a/Engine/Source/Programs/UnrealBuildAccelerator/Common/Private/UbaFileAccessor.cpp
+++ b/Engine/Source/Programs/UnrealBuildAccelerator/Common/Private/UbaFileAccessor.cpp
@@ -92,19 +92,20 @@ namespace uba
 			if (m_fileHandle != InvalidFileHandle)
 			{
 				if (retryCount)
-					m_logger.Warning(TC("Had to retry %u times to open file %s for write (because is was being used%s)"), retryCount, realFileName, additionalInfo.data);
+					m_logger.Info(TC("Had to retry %u times to open file %s for write (because is was being used%s)"), retryCount, realFileName, additionalInfo.data);
 				break;
 			}
 			u32 lastError = GetLastError();
 			#if PLATFORM_WINDOWS
-			if (lastError == ERROR_SHARING_VIOLATION)
+			if (lastError == ERROR_SHARING_VIOLATION || lastError == ERROR_USER_MAPPED_FILE)
 			{
-				if (retryCount == 0)
+				// Since we have unmap of files in jobs it might be that the file is not unmapped yet, therefor we should retry
+				if (retryCount == 1)
 					GetProcessHoldingFile(additionalInfo, m_fileName);
 
-				if (retryCount < 5)
+				if (retryCount < 10)
 				{
-					Sleep(1000);
+					Sleep(500);
 					++retryCount;
 					continue;
 				}
diff --git a/Engine/Source/Programs/UnrealBuildAccelerator/Core/Private/UbaApplicationRules.cpp b/Engine/Source/Programs/UnrealBuildAccelerator/Core/Private/UbaApplicationRules.cpp
--- a/Engine/Source/Programs/UnrealBuildAccelerator/Core/Private/UbaApplicationRules.cpp
+++ b/Engine/Source/Programs/UnrealBuildAccelerator/Core/Private/UbaApplicationRules.cpp
@@ -82,6 +82,7 @@ namespace uba
 		{
 			return file.EndsWith(TC(".cpp"))
 				|| (file.EndsWith(TC(".rsp")) && !file.EndsWith(TC("Shared.rsp")))
+				|| file.EndsWith(TC(".i")) // .i is preprocessed source code (for static analysis)
 				;
 		}
 
