[UBA]
* Speculative fix in UbaNetworkBackendTcp based on failures on the farm

diff --git a/Engine/Source/Programs/UnrealBuildAccelerator/Common/Private/UbaNetworkBackendTcp.cpp b/Engine/Source/Programs/UnrealBuildAccelerator/Common/Private/UbaNetworkBackendTcp.cpp
--- a/Engine/Source/Programs/UnrealBuildAccelerator/Common/Private/UbaNetworkBackendTcp.cpp
+++ b/Engine/Source/Programs/UnrealBuildAccelerator/Common/Private/UbaNetworkBackendTcp.cpp
@@ -872,8 +872,9 @@ namespace uba
 			{
 				u32 toCopy = Min(recvCache.byteCount, recvLeft);
 				memcpy(buffer, recvCache.bytes, toCopy);
-				memmove(recvCache.bytes, recvCache.bytes + toCopy, recvCache.byteCount - toCopy);
 				recvCache.byteCount -= toCopy;
+				if (recvCache.byteCount)
+					memmove(recvCache.bytes, recvCache.bytes + toCopy, recvCache.byteCount);
 				recvLeft -= toCopy;
 				buffer += toCopy;
 				if (!recvLeft)
@@ -906,10 +907,12 @@ namespace uba
 				else
 				{
 					recvCache.byteCount += read;
+					UBA_ASSERT(recvCache.byteCount <= sizeof(recvCache.bytes));
 					u32 toCopy = Min(recvCache.byteCount, recvLeft);
 					memcpy(buffer, recvCache.bytes, toCopy);
-					memmove(recvCache.bytes, recvCache.bytes + toCopy, recvCache.byteCount - toCopy);
 					recvCache.byteCount -= toCopy;
+					if (recvCache.byteCount)
+						memmove(recvCache.bytes, recvCache.bytes + toCopy, recvCache.byteCount);
 					recvLeft -= toCopy;
 					buffer += toCopy;
 				}
