[UBA]
* Added custom code for access violation exceptions to log out more info
* Fixed so when creating fake cas for cache download it adheres to storage compression mode
* Improved assert description

diff --git a/Engine/Source/Programs/UnrealBuildAccelerator/Common/Private/UbaCacheClient.cpp b/Engine/Source/Programs/UnrealBuildAccelerator/Common/Private/UbaCacheClient.cpp
--- a/Engine/Source/Programs/UnrealBuildAccelerator/Common/Private/UbaCacheClient.cpp
+++ b/Engine/Source/Programs/UnrealBuildAccelerator/Common/Private/UbaCacheClient.cpp
@@ -796,6 +796,7 @@ namespace uba
 
 			fetcher.sizeOnDisk = localBlock.writtenSize;
 			casKey = CalculateCasKey(localBlock.memory, localBlock.writtenSize, false, nullptr, path.data);
+			casKey = AsCompressed(casKey, m_storage.StoreCompressed());
 		}
 		else
 		{
diff --git a/Engine/Source/Programs/UnrealBuildAccelerator/Common/Private/UbaProcess.cpp b/Engine/Source/Programs/UnrealBuildAccelerator/Common/Private/UbaProcess.cpp
--- a/Engine/Source/Programs/UnrealBuildAccelerator/Common/Private/UbaProcess.cpp
+++ b/Engine/Source/Programs/UnrealBuildAccelerator/Common/Private/UbaProcess.cpp
@@ -286,7 +286,7 @@ namespace uba
 		}
 		if (exitCode == STILL_ACTIVE)
 			return true;
-		if (!m_gotExitMessage && exitCode != 0xC0000005 && exitCode != 0xC0000409)
+		if (!m_gotExitMessage && exitCode != EXCEPTION_ACCESS_VIOLATION && exitCode != STATUS_STACK_BUFFER_OVERRUN)
 		{
 			StringBuffer<> err;
 
@@ -541,9 +541,9 @@ namespace uba
 				break;
 			--retryCount;
 
-			if (exitCode == 0xC0000005)
+			if (exitCode == EXCEPTION_ACCESS_VIOLATION)
 				m_session.m_logger.Warning(TC("Process exited with access violation. Will do one retry."));
-			else if (exitCode == 0xC0000409)
+			else if (exitCode == STATUS_STACK_BUFFER_OVERRUN)
 				m_session.m_logger.Warning(TC("Process exited with stack buffer overflow. Will do one retry."));
 			else
 				break;
diff --git a/Engine/Source/Programs/UnrealBuildAccelerator/Common/Private/UbaStorage.cpp b/Engine/Source/Programs/UnrealBuildAccelerator/Common/Private/UbaStorage.cpp
--- a/Engine/Source/Programs/UnrealBuildAccelerator/Common/Private/UbaStorage.cpp
+++ b/Engine/Source/Programs/UnrealBuildAccelerator/Common/Private/UbaStorage.cpp
@@ -909,7 +909,7 @@ namespace uba
 
 	bool StorageImpl::AddCasFile(StringKey fileNameKey, const tchar* fileName, const CasKey& casKey, bool deferCreation, bool fileIsCompressed)
 	{
-		UBA_ASSERTF(IsCompressed(casKey) == m_storeCompressed, TC("CasKey compress mode must match storage compress mode (%s)"), fileName);
+		UBA_ASSERTF(IsCompressed(casKey) == m_storeCompressed, TC("CasKey compress mode must match storage compress mode. Compress mode Key: %i, Store: %i Cas: %s (%s)"), (IsCompressed(casKey) ? 1 : 0), (m_storeCompressed ? 1 : 0), CasKeyString(casKey).str, fileName);
 		SCOPED_WRITE_LOCK(m_casLookupLock, lookupLock);
 		auto insres = m_casLookup.try_emplace(casKey, casKey);
 		CasEntry& casEntry = insres.first->second;
diff --git a/Engine/Source/Programs/UnrealBuildAccelerator/Core/Private/UbaPlatform.cpp b/Engine/Source/Programs/UnrealBuildAccelerator/Core/Private/UbaPlatform.cpp
--- a/Engine/Source/Programs/UnrealBuildAccelerator/Core/Private/UbaPlatform.cpp
+++ b/Engine/Source/Programs/UnrealBuildAccelerator/Core/Private/UbaPlatform.cpp
@@ -531,8 +531,8 @@ namespace uba
 	void* g_endOfCurrentLibrary;
 	LONG UbaExceptionHandler(_EXCEPTION_POINTERS* ExceptionInfo)
 	{
-		auto record = ExceptionInfo->ExceptionRecord;
-		u32 exceptionCode = record->ExceptionCode;
+		auto& record = *ExceptionInfo->ExceptionRecord;
+		u32 exceptionCode = record.ExceptionCode;
 
 		if (exceptionCode != EXCEPTION_STACK_OVERFLOW && exceptionCode != EXCEPTION_ACCESS_VIOLATION)
 			return EXCEPTION_CONTINUE_SEARCH;
@@ -547,7 +547,7 @@ namespace uba
 			g_endOfCurrentLibrary = (u8*)mi.lpBaseOfDll + mi.SizeOfImage;
 		}
 
-		if (record->ExceptionAddress < g_startOfCurrentLibrary || record->ExceptionAddress >= g_endOfCurrentLibrary)
+		if (record.ExceptionAddress < g_startOfCurrentLibrary || record.ExceptionAddress >= g_endOfCurrentLibrary)
 			return EXCEPTION_CONTINUE_SEARCH;
 
 		//CONTEXT context;
@@ -556,7 +556,11 @@ namespace uba
 		CONTEXT* contextPtr = ExceptionInfo->ContextRecord;
 
 		StringBuffer<> text;
-		text.Appendf(L"ERROR: Unhandled Exception (Code: 0x%x)", exceptionCode);
+		if (exceptionCode == EXCEPTION_ACCESS_VIOLATION)
+			text.Appendf(L"ERROR: Access violation %s at address: 0x%p", (record.ExceptionInformation[0] == 1 ? "writing" : "reading"), (void*)record.ExceptionInformation[1]);
+		else
+			text.Appendf(L"ERROR: Unhandled Exception (Code: 0x%x)", exceptionCode);
+
 		UbaAssert(text.data, nullptr, 0, nullptr, true, exceptionCode, contextPtr);
 		return EXCEPTION_EXECUTE_HANDLER;
 	}
