[UBA]
* Fixed linux/mac compile errors

diff --git a/Engine/Source/Programs/UnrealBuildAccelerator/Common/Private/UbaProcess.cpp b/Engine/Source/Programs/UnrealBuildAccelerator/Common/Private/UbaProcess.cpp
--- a/Engine/Source/Programs/UnrealBuildAccelerator/Common/Private/UbaProcess.cpp
+++ b/Engine/Source/Programs/UnrealBuildAccelerator/Common/Private/UbaProcess.cpp
@@ -35,6 +35,10 @@ extern char **environ;
 
 #define UBA_DEBUG_TRACK_PROCESS UBA_DEBUG_LOGGER
 
+#if !PLATFORM_WINDOWS
+#define EXCEPTION_ACCESS_VIOLATION 128 + SIGSEGV
+#define EXCEPTION_STACK_BUFFER_OVERRUN 128 + SIGSEGV
+#endif
 
 //////////////////////////////////////////////////////////////////////////////
 
@@ -286,7 +290,7 @@ namespace uba
 		}
 		if (exitCode == STILL_ACTIVE)
 			return true;
-		if (!m_gotExitMessage && exitCode != EXCEPTION_ACCESS_VIOLATION && exitCode != STATUS_STACK_BUFFER_OVERRUN)
+		if (!m_gotExitMessage && exitCode != EXCEPTION_ACCESS_VIOLATION && exitCode != EXCEPTION_STACK_BUFFER_OVERRUN)
 		{
 			StringBuffer<> err;
 
@@ -543,7 +547,7 @@ namespace uba
 
 			if (exitCode == EXCEPTION_ACCESS_VIOLATION)
 				m_session.m_logger.Warning(TC("Process exited with access violation. Will do one retry."));
-			else if (exitCode == STATUS_STACK_BUFFER_OVERRUN)
+			else if (exitCode == EXCEPTION_STACK_BUFFER_OVERRUN)
 				m_session.m_logger.Warning(TC("Process exited with stack buffer overflow. Will do one retry."));
 			else
 				break;
