[UBA]
* Improved error messages more

diff --git a/Engine/Source/Programs/UnrealBuildAccelerator/Common/Private/UbaStorage.cpp b/Engine/Source/Programs/UnrealBuildAccelerator/Common/Private/UbaStorage.cpp
--- a/Engine/Source/Programs/UnrealBuildAccelerator/Common/Private/UbaStorage.cpp
+++ b/Engine/Source/Programs/UnrealBuildAccelerator/Common/Private/UbaStorage.cpp
@@ -2078,7 +2078,7 @@ namespace uba
 		deferredLock.Leave();
 		WriteResult res;
 		if (!WriteCasFile(res, deferredCreation.data, fileIsCompressed, casKey))
-			return false;
+			return m_logger.Error(TC("Failed to write deferred cas %s (%s)"), CasKeyString(casKey), deferredCreation.data);
 
 		casEntry.exists = true;
 		entryLock.Leave();
diff --git a/Engine/Source/Programs/UnrealBuildAccelerator/Common/Private/UbaStorageServer.cpp b/Engine/Source/Programs/UnrealBuildAccelerator/Common/Private/UbaStorageServer.cpp
--- a/Engine/Source/Programs/UnrealBuildAccelerator/Common/Private/UbaStorageServer.cpp
+++ b/Engine/Source/Programs/UnrealBuildAccelerator/Common/Private/UbaStorageServer.cpp
@@ -468,7 +468,7 @@ namespace uba
 						}
 						if (checkedCasKey == CasKeyZero)
 						{
-							m_logger.Info(TC("Server did not find cas for %s in file table lookup. Recalculating cas key"), hint.data);
+							m_logger.Info(TC("Server did not find cas %s for %s in file table lookup. Recalculating cas key"), CasKeyString(casKey).str, hint.data);
 							if (!CalculateCasKey(checkedCasKey, hint.data))
 							{
 								m_logger.Error(TC("FetchBegin failed for cas file %s (%s) requested by %s. Can't calculate cas key for file"), CasKeyString(casKey).str, hint.data, GuidToString(connectionInfo.GetUid()).str);
