[UBA]
* Added so store of cas from client checks if cas is in deferred cas list before receiving it (and use the deferred one instead)
* Added more information to some errors

diff --git a/Engine/Source/Programs/UnrealBuildAccelerator/Common/Private/UbaStorage.cpp b/Engine/Source/Programs/UnrealBuildAccelerator/Common/Private/UbaStorage.cpp
--- a/Engine/Source/Programs/UnrealBuildAccelerator/Common/Private/UbaStorage.cpp
+++ b/Engine/Source/Programs/UnrealBuildAccelerator/Common/Private/UbaStorage.cpp
@@ -2069,12 +2069,16 @@ namespace uba
 		auto findIt = m_deferredCasCreationLookup.find(casKey);
 		if (findIt == m_deferredCasCreationLookup.end())
 			return false;
-		casEntry.verified = true;
 		StringKey fileNameKey = findIt->second.fileNameKey;
 		StringBuffer<> deferredCreation(findIt->second.fileName);
 		bool fileIsCompressed = findIt->second.fileIsCompressed;
 		m_deferredCasCreationLookupByName.erase(fileNameKey);
 		m_deferredCasCreationLookup.erase(findIt);
+
+		if (casEntry.beingWritten)
+			m_logger.Warning(TC("Deferred cas %s is being written from network. This should never happen (%s)"), CasKeyString(casKey), deferredCreation.data);
+
+		casEntry.verified = true;
 		deferredLock.Leave();
 		WriteResult res;
 		if (!WriteCasFile(res, deferredCreation.data, fileIsCompressed, casKey))
diff --git a/Engine/Source/Programs/UnrealBuildAccelerator/Common/Private/UbaStorageServer.cpp b/Engine/Source/Programs/UnrealBuildAccelerator/Common/Private/UbaStorageServer.cpp
--- a/Engine/Source/Programs/UnrealBuildAccelerator/Common/Private/UbaStorageServer.cpp
+++ b/Engine/Source/Programs/UnrealBuildAccelerator/Common/Private/UbaStorageServer.cpp
@@ -780,6 +780,27 @@ namespace uba
 				lookupLock.Leave();
 
 				SCOPED_WRITE_LOCK(casEntry.lock, entryLock);
+
+				{
+					// Check if there is a deferred creation of this file, if there is we can just as well just use the local file
+					SCOPED_WRITE_LOCK(m_deferredCasCreationLookupLock, deferredLock);
+					auto findIt = m_deferredCasCreationLookup.find(casKey);
+					if (findIt != m_deferredCasCreationLookup.end())
+					{
+						deferredLock.Leave();
+						entryLock.Leave();
+						if (HasCasFile(casKey))
+						{
+							writer.WriteU16(u16(~0));
+							writer.WriteBool(m_traceStore);
+							m_logger.Detail(TC("Client %u Store request of %s which already exists in casdb (was in deferred list) (%s)"), connectionInfo.GetId(), CasKeyString(casKey).str, hint.data);
+							return true;
+						}
+						entryLock.Enter();
+					}
+				}
+
+
 				if (!casEntry.verified)
 				{
 					if (bool exists = casEntry.exists)
diff --git a/Engine/Source/Programs/UnrealBuildAccelerator/Core/Private/UbaPlatform.cpp b/Engine/Source/Programs/UnrealBuildAccelerator/Core/Private/UbaPlatform.cpp
--- a/Engine/Source/Programs/UnrealBuildAccelerator/Core/Private/UbaPlatform.cpp
+++ b/Engine/Source/Programs/UnrealBuildAccelerator/Core/Private/UbaPlatform.cpp
@@ -557,7 +557,7 @@ namespace uba
 
 		StringBuffer<> text;
 		if (exceptionCode == EXCEPTION_ACCESS_VIOLATION)
-			text.Appendf(L"ERROR: Access violation %s at address: 0x%p", (record.ExceptionInformation[0] == 1 ? "writing" : "reading"), (void*)record.ExceptionInformation[1]);
+			text.Appendf(L"ERROR: Access violation %s at address: 0x%p", (record.ExceptionInformation[0] == 1 ? L"writing" : L"reading"), (void*)record.ExceptionInformation[1]);
 		else
 			text.Appendf(L"ERROR: Unhandled Exception (Code: 0x%x)", exceptionCode);
 
