[UBA]
* Disabled include crawler by default and instead added it as a command line option (-usecrawler)

diff --git a/Engine/Source/Programs/UnrealBuildAccelerator/Agent/Private/UbaAgent.cpp b/Engine/Source/Programs/UnrealBuildAccelerator/Agent/Private/UbaAgent.cpp
--- a/Engine/Source/Programs/UnrealBuildAccelerator/Agent/Private/UbaAgent.cpp
+++ b/Engine/Source/Programs/UnrealBuildAccelerator/Agent/Private/UbaAgent.cpp
@@ -114,6 +114,7 @@ namespace uba
 		logger.Info(TC("  -resendcas              Will try to send same cas multiple times (set this to true if server is allowed to remove cas files)"));
 		logger.Info(TC("  -populateCas=<dir>      Prepopulate cas database with files in dir. If files needed exists on machine this can be an optimization"));
 		logger.Info(TC("  -description            Add more info about the agent that will show in the trace log when hovering over session"));
+		logger.Info(TC("  -usecrawler             Enables include crawler for known process types (clang/msvc)"));
 		#if PLATFORM_MAC
 		logger.Info(TC("  -populateCasFromXcodeVersion=<version>   Prepopulate cas database with files from local xcode installation that matches the version."));
 		logger.Info(TC("  -populateCasFromAllXcodes   Prepopulate cas database with files from local xcode installation that matches the version."));
@@ -366,6 +367,7 @@ namespace uba
 		bool verbose = false;
 		bool printSummary = false;
 		bool killRandom = false;
+		bool useCrawler = false;
 		bool downloadDetoursLib = true;
 		bool useExceptionHandler = true;
 		bool resendCas = false;
@@ -591,6 +593,10 @@ namespace uba
 			{
 				killRandom = true;
 			}
+			else if (name.Equals(TC("-usecrawler")))
+			{
+				useCrawler = true;
+			}
 			else if (name.Equals(TC("-memwait")))
 			{
 				if (!value.Parse(memWaitLoadPercent) || memWaitLoadPercent > 100)
@@ -1177,6 +1183,7 @@ namespace uba
 			info.downloadDetoursLib = downloadDetoursLib;
 			info.memWaitLoadPercent = u8(memWaitLoadPercent);
 			info.memKillLoadPercent = u8(memKillLoadPercent);
+			info.useIncludeCrawler = useCrawler;
 
 			if (!quiet)
 				info.processFinished = processFinished;
diff --git a/Engine/Source/Programs/UnrealBuildAccelerator/Common/Public/UbaSessionClient.h b/Engine/Source/Programs/UnrealBuildAccelerator/Common/Public/UbaSessionClient.h
--- a/Engine/Source/Programs/UnrealBuildAccelerator/Common/Public/UbaSessionClient.h
+++ b/Engine/Source/Programs/UnrealBuildAccelerator/Common/Public/UbaSessionClient.h
@@ -29,7 +29,7 @@ namespace uba
 		bool killRandom = false;
 		bool useStorage = true;
 		bool downloadDetoursLib = true;
-		bool useIncludeCrawler = true;	// IncludeCrawler will go wide and try to prefetch all headers and source files needed by the build
+		bool useIncludeCrawler = false;	// IncludeCrawler will go wide and try to prefetch all headers and source files needed by the build
 		Function<void(const ProcessHandle&)> processFinished;
 	};
 
diff --git a/Engine/Source/Programs/UnrealBuildAccelerator/Core/Public/UbaHash.h b/Engine/Source/Programs/UnrealBuildAccelerator/Core/Public/UbaHash.h
--- a/Engine/Source/Programs/UnrealBuildAccelerator/Core/Public/UbaHash.h
+++ b/Engine/Source/Programs/UnrealBuildAccelerator/Core/Public/UbaHash.h
@@ -172,7 +172,7 @@ namespace uba
 
 	struct CasKeyString : StringView
 	{
-		CasKeyString(const CasKey& g) : StringView(str, 40) { ToString(str, 41, g); }
+		CasKeyString(const CasKey& g) : StringView(str, 40) { uba::ToString(str, 41, g); }
 		tchar str[41];
 	};
 
diff --git a/Engine/Source/Programs/UnrealBuildAccelerator/Core/Public/UbaMemory.h b/Engine/Source/Programs/UnrealBuildAccelerator/Core/Public/UbaMemory.h
--- a/Engine/Source/Programs/UnrealBuildAccelerator/Core/Public/UbaMemory.h
+++ b/Engine/Source/Programs/UnrealBuildAccelerator/Core/Public/UbaMemory.h
@@ -72,7 +72,6 @@ namespace uba
 	using Allocator = std::allocator<Type>;
 	#endif
 
-	//using WString = std::basic_string<tchar, std::char_traits<tchar>, Allocator<tchar>>;
 	using TString = std::basic_string<tchar, std::char_traits<tchar>, Allocator<tchar>>;
 
 	template<class T> using Function = std::function<T>;
diff --git a/Engine/Source/Programs/UnrealBuildAccelerator/Core/Public/UbaStringBuffer.h b/Engine/Source/Programs/UnrealBuildAccelerator/Core/Public/UbaStringBuffer.h
--- a/Engine/Source/Programs/UnrealBuildAccelerator/Core/Public/UbaStringBuffer.h
+++ b/Engine/Source/Programs/UnrealBuildAccelerator/Core/Public/UbaStringBuffer.h
@@ -73,6 +73,8 @@ namespace uba
 		bool Parse(float& out);
 		u32 Parse(char* out, u64 outCapacity); // Note, return value contains length + null termination.. zero if failed
 
+		TString ToString() const { return TString(data, data + count); }
+
 		u32 count;
 		u32 capacity;
 		tchar data[1];
@@ -113,6 +115,8 @@ namespace uba
 		bool Equals(const tchar* str, bool ignoreCase = true) const { return uba::Equals(data, str, ignoreCase); }
 		StringView GetFileName() const { const tchar* str = uba::GetFileName(data); return StringView(str, count - u32(str - data)); }
 
+		TString ToString() const { return TString(data, data + count); }
+
 		const tchar* data;
 		u32 count;
 	};
