[UBA]
* Implemented wait for the scenario where server gets a local add of a cas entry while cas is being transferred from a client.

diff --git a/Engine/Source/Programs/UnrealBuildAccelerator/Common/Private/UbaStorage.cpp b/Engine/Source/Programs/UnrealBuildAccelerator/Common/Private/UbaStorage.cpp
--- a/Engine/Source/Programs/UnrealBuildAccelerator/Common/Private/UbaStorage.cpp
+++ b/Engine/Source/Programs/UnrealBuildAccelerator/Common/Private/UbaStorage.cpp
@@ -918,6 +918,23 @@ namespace uba
 
 		SCOPED_WRITE_LOCK(casEntry.lock, entryLock);
 
+		if (casEntry.beingWritten)
+		{
+			m_logger.Warning(TC("Cas %s is being written from network while being added by host (%s)"), CasKeyString(casKey).str, fileName);
+			int waitCount = 0;
+			while (casEntry.beingWritten)
+			{
+				CasKey key = casEntry.key;
+				entryLock.Leave();
+				Sleep(100);
+				entryLock.Enter();
+
+				if (++waitCount < 12*60*10)
+					continue;
+				return m_logger.Error(TC("Host waited more than 12 minutes for file %s (%s) to be written by network (Written %llu/%llu)"), CasKeyString(key).str, fileName);
+			}
+		}
+
 		if (casEntry.verified)
 			if (casEntry.exists)
 				return true;
