[UBA]
* Fixed bug where storage of multiple files (different paths) with same cas could cause issue in combination with deferred cas file creation.
* Changed some asserts to errors
* Removed lock that tsan complains about.
* Removed unused function
* Fixed race condition of m_lastPing

diff --git a/Engine/Source/Programs/UnrealBuildAccelerator/Common/Private/UbaNetworkClient.cpp b/Engine/Source/Programs/UnrealBuildAccelerator/Common/Private/UbaNetworkClient.cpp
--- a/Engine/Source/Programs/UnrealBuildAccelerator/Common/Private/UbaNetworkClient.cpp
+++ b/Engine/Source/Programs/UnrealBuildAccelerator/Common/Private/UbaNetworkClient.cpp
@@ -240,7 +240,8 @@ namespace uba
 		SCOPED_READ_LOCK(client.m_activeMessagesLock, lock);
 		if (!connection.connected)
 			return false;
-		UBA_ASSERTF(messageId < client.m_activeMessages.size(), TC("Message id %u is higher than max %u"), messageId, u32(client.m_activeMessages.size()));
+		if (messageId >= client.m_activeMessages.size())
+			return client.m_logger.Error(TC("Message id %u is higher than max %u"), messageId, u32(client.m_activeMessages.size()));
 		NetworkMessage* msg = client.m_activeMessages[messageId];
 		lock.Leave();
 
@@ -255,14 +256,24 @@ namespace uba
 			msg->Done();
 			return true;
 		}
-		else if (!messageSize)
+		
+		if (!messageSize)
 		{
 			++client.m_recvCount;
 			msg->Done();
 			return true;
 		}
 
-		UBA_ASSERTF(messageSize <= msg->m_responseCapacity, TC("Message size is %u but reader capacity is only %u"), messageSize, msg->m_responseCapacity);
+		if (messageSize > msg->m_responseCapacity)
+		{
+			u8 serviceIdAndMessageType = msg->m_sendWriter->GetData()[0];
+			u8 serviceId = serviceIdAndMessageType >> 6;
+			u8 messageType = serviceIdAndMessageType & 0b111111;
+			client.m_logger.Error(TC("Message size is %u but reader capacity is only %u (serviceId %u, messageType %u)"), messageSize, msg->m_responseCapacity, u32(serviceId), u32(messageType));
+			msg->m_error = 1;
+			msg->Done();
+			return false;
+		}
 
 		msg->m_responseSize = messageSize;
 
@@ -621,13 +632,6 @@ namespace uba
 		return !message.m_error;
 	}
 
-	void NetworkClient::ReturnMessageId(u16 id)
-	{
-		SCOPED_WRITE_LOCK(m_activeMessagesLock, lock);
-		m_availableMessageIds.push_back(id);
-		m_activeMessages[id] = nullptr;
-	}
-
 	const tchar* NetworkClient::SetGetPrefix(const tchar* originalPrefix)
 	{
 		CreateGuid(m_uid);
diff --git a/Engine/Source/Programs/UnrealBuildAccelerator/Common/Private/UbaSessionClient.cpp b/Engine/Source/Programs/UnrealBuildAccelerator/Common/Private/UbaSessionClient.cpp
--- a/Engine/Source/Programs/UnrealBuildAccelerator/Common/Private/UbaSessionClient.cpp
+++ b/Engine/Source/Programs/UnrealBuildAccelerator/Common/Private/UbaSessionClient.cpp
@@ -1574,11 +1574,12 @@ namespace uba
 		if (!msg.Send(reader, m_stats.pingMsg))
 			m_loop = false;
 		u64 newTime = GetTime();
-		m_lastPing = newTime - time;
+		u64 lastPing = newTime - time;
+		m_lastPing = lastPing;
 		m_lastPingSendTime = newTime;
 
-		if (m_lastPing < m_bestPing || m_bestPing == 0)
-			m_bestPing = m_lastPing;
+		if (lastPing < m_bestPing || m_bestPing == 0)
+			m_bestPing = lastPing;
 
 		m_storage.Ping();
 
diff --git a/Engine/Source/Programs/UnrealBuildAccelerator/Common/Private/UbaStorage.cpp b/Engine/Source/Programs/UnrealBuildAccelerator/Common/Private/UbaStorage.cpp
--- a/Engine/Source/Programs/UnrealBuildAccelerator/Common/Private/UbaStorage.cpp
+++ b/Engine/Source/Programs/UnrealBuildAccelerator/Common/Private/UbaStorage.cpp
@@ -958,9 +958,12 @@ namespace uba
 				DeferedCasCreation& dcc = res.first->second;
 				dcc.fileName = fileName;
 				dcc.fileIsCompressed = fileIsCompressed;
-				auto res2 = m_deferredCasCreationLookupByName.try_emplace(fileNameKey, casKey);
-				UBA_ASSERTF(res2.second, TC("%s - existing: %s added: %s"), dcc.fileName.c_str(), CasKeyString(res2.first->second).str, CasKeyString(casKey).str); (void)res2;
 			}
+			auto res2 = m_deferredCasCreationLookupByName.try_emplace(fileNameKey, casKey);
+			if (res2.second)
+				res.first->second.names.push_back(res2.first);
+			else if (res2.first->second != casKey)
+				return m_logger.Error(TC("Same file %s is registered with different caskey (%s vs %s)"), fileName, CasKeyString(res2.first->second).str, CasKeyString(casKey).str);
 			return true;
 		}
 
@@ -1488,7 +1491,7 @@ namespace uba
 			for (auto& pair : m_fileTableLookup)
 			{
 				FileEntry& fileEntry = pair.second;
-				SCOPED_READ_LOCK(fileEntry.lock, entryLock); // TODO: Remove this? tsan complains about it
+				//SCOPED_READ_LOCK(fileEntry.lock, entryLock); // TODO: Remove this? tsan complains about it
 				if (fileEntry.casKey == CasKeyZero || fileEntry.lastWritten == 0)
 					continue;
 				writer.WriteStringKey(pair.first);
@@ -2089,7 +2092,8 @@ namespace uba
 		StringKey fileNameKey = findIt->second.fileNameKey;
 		StringBuffer<> deferredCreation(findIt->second.fileName);
 		bool fileIsCompressed = findIt->second.fileIsCompressed;
-		m_deferredCasCreationLookupByName.erase(fileNameKey);
+		for (auto& it2 : findIt->second.names)
+			m_deferredCasCreationLookupByName.erase(it2);
 		m_deferredCasCreationLookup.erase(findIt);
 
 		if (casEntry.beingWritten)
@@ -2587,8 +2591,9 @@ namespace uba
 		auto findIt = m_deferredCasCreationLookupByName.find(fileNameKey);
 		if (findIt == m_deferredCasCreationLookupByName.end())
 			return;
+		CasKey key = findIt->second;
 		deferredLock.Leave();
-		HasCasFile(findIt->second);
+		HasCasFile(key);
 	}
 	/*
 	void StorageImpl::PushStats(StorageStats* stats)
diff --git a/Engine/Source/Programs/UnrealBuildAccelerator/Common/Public/UbaNetworkClient.h b/Engine/Source/Programs/UnrealBuildAccelerator/Common/Public/UbaNetworkClient.h
--- a/Engine/Source/Programs/UnrealBuildAccelerator/Common/Public/UbaNetworkClient.h
+++ b/Engine/Source/Programs/UnrealBuildAccelerator/Common/Public/UbaNetworkClient.h
@@ -91,7 +91,6 @@ namespace uba
 		static bool ReceiveResponseBody(void* context, bool recvError, u8* headerData, void* bodyContext, u8* bodyData, u32 bodySize);
 		void OnDisconnected(Connection& connection, u32 reason);
 		bool Send(NetworkMessage& message, void* response, u32 responseCapacity, bool async);
-		void ReturnMessageId(u16 id);
 		const tchar* SetGetPrefix(const tchar* originalPrefix);
 
 		LogWriter& m_logWriter;
diff --git a/Engine/Source/Programs/UnrealBuildAccelerator/Common/Public/UbaSessionClient.h b/Engine/Source/Programs/UnrealBuildAccelerator/Common/Public/UbaSessionClient.h
--- a/Engine/Source/Programs/UnrealBuildAccelerator/Common/Public/UbaSessionClient.h
+++ b/Engine/Source/Programs/UnrealBuildAccelerator/Common/Public/UbaSessionClient.h
@@ -172,7 +172,7 @@ namespace uba
 		SessionSummaryStats m_stats;
 
 		Atomic<u64> m_bestPing;
-		u64 m_lastPing = 0;
+		Atomic<u64> m_lastPing;
 		u64 m_lastPingSendTime = 0;
 
 		Atomic<u64> m_memAvail;
diff --git a/Engine/Source/Programs/UnrealBuildAccelerator/Common/Public/UbaStorage.h b/Engine/Source/Programs/UnrealBuildAccelerator/Common/Public/UbaStorage.h
--- a/Engine/Source/Programs/UnrealBuildAccelerator/Common/Public/UbaStorage.h
+++ b/Engine/Source/Programs/UnrealBuildAccelerator/Common/Public/UbaStorage.h
@@ -258,9 +258,10 @@ namespace uba
 		FileMappingBuffer m_casDataBuffer;
 
 		ReaderWriterLock m_deferredCasCreationLookupLock;
-		struct DeferedCasCreation { StringKey fileNameKey; TString fileName; bool fileIsCompressed; };
+		using NameToCas = UnorderedMap<StringKey, CasKey>;
+		struct DeferedCasCreation { StringKey fileNameKey; TString fileName; bool fileIsCompressed; List<NameToCas::iterator> names; };
 		UnorderedMap<CasKey, DeferedCasCreation> m_deferredCasCreationLookup;
-		UnorderedMap<StringKey, CasKey> m_deferredCasCreationLookupByName;
+		NameToCas m_deferredCasCreationLookupByName;
 
 		DirectoryCache m_dirCache;
 
