[UBA]
* Removed CancelAllProcessesAndWait call in UbaSessionClient since it will be happening in the update thread anyways and we don't want lots of stuff happening when pressing ctrl-c

diff --git a/Engine/Source/Programs/UnrealBuildAccelerator/Common/Private/UbaSessionClient.cpp b/Engine/Source/Programs/UnrealBuildAccelerator/Common/Private/UbaSessionClient.cpp
--- a/Engine/Source/Programs/UnrealBuildAccelerator/Common/Private/UbaSessionClient.cpp
+++ b/Engine/Source/Programs/UnrealBuildAccelerator/Common/Private/UbaSessionClient.cpp
@@ -162,7 +162,6 @@ namespace uba
 	void SessionClient::Stop()
 	{
 		m_loop = false;
-		CancelAllProcessesAndWait();
 		m_waitToSendEvent.Set();
 		m_loopThread.Wait();
 	}
