[UBA]
* Promoted FileMappingHandle to UbaPlatform.h to be used in both common and detoured code
* Added file handle to FileMappingHandle for win64. This handle is not used yet

diff --git a/Engine/Source/Programs/UnrealBuildAccelerator/Common/Private/UbaFileMapping.cpp b/Engine/Source/Programs/UnrealBuildAccelerator/Common/Private/UbaFileMapping.cpp
--- a/Engine/Source/Programs/UnrealBuildAccelerator/Common/Private/UbaFileMapping.cpp
+++ b/Engine/Source/Programs/UnrealBuildAccelerator/Common/Private/UbaFileMapping.cpp
@@ -65,7 +65,7 @@ namespace uba
 	{
 		ExtendedTimerScope ts(KernelStats::GetCurrent().createFileMapping);
 #if PLATFORM_WINDOWS
-		return { InternalCreateFileMappingW(logger, INVALID_HANDLE_VALUE, flProtect, (DWORD)ToHigh(maxSize), ToLow(maxSize), name, hint) };
+		return { 0, InternalCreateFileMappingW(logger, INVALID_HANDLE_VALUE, flProtect, (DWORD)ToHigh(maxSize), ToLow(maxSize), name, hint) };
 #else
 		UBA_ASSERT(!name);
 		UBA_ASSERT((flProtect & (~u32(PAGE_READWRITE | SEC_RESERVE))) == 0);
@@ -198,7 +198,7 @@ namespace uba
 	{
 		ExtendedTimerScope ts(KernelStats::GetCurrent().createFileMapping);
 #if PLATFORM_WINDOWS
-		return { InternalCreateFileMappingW(logger, asHANDLE(file), protect, (DWORD)ToHigh(maxSize), ToLow(maxSize), NULL, hint) };
+		return { 0, InternalCreateFileMappingW(logger, asHANDLE(file), protect, (DWORD)ToHigh(maxSize), ToLow(maxSize), NULL, hint) };
 #else
 		FileMappingHandle h;
 		int fd = asFileDescriptor(file);
@@ -236,7 +236,7 @@ namespace uba
 	{
 		ExtendedTimerScope ts(KernelStats::GetCurrent().mapViewOfFile);
 #if PLATFORM_WINDOWS
-		u8* res = (u8*)::MapViewOfFile(fileMappingObject.handle, desiredAccess, (DWORD)ToHigh(offset), ToLow(offset), bytesToMap);
+		u8* res = (u8*)::MapViewOfFile(fileMappingObject.mh, desiredAccess, (DWORD)ToHigh(offset), ToLow(offset), bytesToMap);
 		if (!res)
 			return nullptr;
 #else
@@ -324,7 +324,9 @@ namespace uba
 		#endif
 
 #if PLATFORM_WINDOWS
-		return CloseHandle(h.handle);
+		if (h.fh)
+			CloseHandle(h.fh);
+		return CloseHandle(h.mh);
 #else
 		if (h.shmFd == -1)
 			return true;
@@ -350,11 +352,14 @@ namespace uba
 #endif
 	}
 
-	bool DuplicateFileMapping(Logger& logger, ProcHandle hSourceProcessHandle, FileMappingHandle hSourceHandle, ProcHandle hTargetProcessHandle, FileMappingHandle* lpTargetHandle, u32 dwDesiredAccess, bool bInheritHandle, u32 dwOptions, const tchar* hint)
+	bool DuplicateFileMapping(Logger& logger, ProcHandle sourceProcessHandle, FileMappingHandle source, ProcHandle targetProcessHandle, FileMappingHandle& targetHandle, u32 dwDesiredAccess, bool bInheritHandle, u32 dwOptions, const tchar* hint)
 	{
 #if PLATFORM_WINDOWS
-		if (!DuplicateHandle((HANDLE)hSourceProcessHandle, hSourceHandle.handle, (HANDLE)hTargetProcessHandle, &lpTargetHandle->handle, dwDesiredAccess, bInheritHandle, dwOptions))
+		if (!DuplicateHandle((HANDLE)sourceProcessHandle, source.mh, (HANDLE)targetProcessHandle, &targetHandle.mh, dwDesiredAccess, bInheritHandle, dwOptions))
 			return false;
+		if (source.fh)
+			if (!DuplicateHandle((HANDLE)sourceProcessHandle, source.fh, (HANDLE)targetProcessHandle, &targetHandle.fh, dwDesiredAccess, bInheritHandle, dwOptions))
+				return false;
 #else
 		UBA_ASSERT(false);
 		if (true)
@@ -364,9 +369,9 @@ namespace uba
 		#if UBA_DEBUG_FILE_MAPPING
 		{
 			SCOPED_WRITE_LOCK(g_fileMappingsLock, lock);
-			auto insres = g_fileMappings.try_emplace(lpTargetHandle->handle, hint);(void)insres;
+			auto insres = g_fileMappings.try_emplace(targetHandle.mh, hint);(void)insres;
 			UBA_ASSERT(insres.second);
-			//logger.Info(TC("FILEMAPPING DUPLI: %s %llu (%llu)"), hint, lpTargetHandle->handle, g_fileMappings.size());
+			//logger.Info(TC("FILEMAPPING DUPLI: %s %llu (%llu)"), hint, targetHandle.mh, g_fileMappings.size());
 		}
 		#endif
 
@@ -623,7 +628,7 @@ namespace uba
 			return {};
 
 		FileMappingHandle h;
-		h.handle = (HANDLE)StringToValue(handleStr, handleStrEnd - handleStr);
+		h.mh = (HANDLE)StringToValue(handleStr, handleStrEnd - handleStr);
 		const tchar* mappingOffsetStr = handleStrEnd + 1;
 		u64 mappingOffset = StringToValue(mappingOffsetStr, wcslen(mappingOffsetStr));
 		
diff --git a/Engine/Source/Programs/UnrealBuildAccelerator/Common/Private/UbaProcess.cpp b/Engine/Source/Programs/UnrealBuildAccelerator/Common/Private/UbaProcess.cpp
--- a/Engine/Source/Programs/UnrealBuildAccelerator/Common/Private/UbaProcess.cpp
+++ b/Engine/Source/Programs/UnrealBuildAccelerator/Common/Private/UbaProcess.cpp
@@ -821,7 +821,7 @@ namespace uba
 		msg.attributes = DefaultAttributes(); // reader.ReadFileAttributes(); TODO
 		msg.deleteOnClose = reader.ReadBool();
 		msg.success = reader.ReadBool();
-		msg.mappingHandle = reader.ReadU64();
+		msg.mappingHandle = FileMappingHandle::FromU64(reader.ReadU64());
 		msg.mappingWritten = reader.ReadU64();
 		msg.newNameKey = reader.ReadStringKey();
 		if (msg.newNameKey != StringKeyZero)
@@ -1335,13 +1335,11 @@ namespace uba
 		StringKey key = reader.ReadStringKey();
 		StringBuffer<> fileName;
 		reader.ReadString(fileName);
-		u64 mappingHandle = reader.ReadU64();
+		FileMappingHandle mappingHandle = FileMappingHandle::FromU64(reader.ReadU64());
 		u64 mappingHandleSize = reader.ReadU64();
 
-		FileMappingHandle source;
-		source.FromU64(mappingHandle);
 		FileMappingHandle newHandle;
-		if (!DuplicateFileMapping(m_session.m_logger, nativeProcessHandle, source, GetCurrentProcessHandle(), &newHandle, 0, false, DUPLICATE_SAME_ACCESS, fileName.data))
+		if (!DuplicateFileMapping(m_session.m_logger, nativeProcessHandle, mappingHandle, GetCurrentProcessHandle(), newHandle, 0, false, DUPLICATE_SAME_ACCESS, fileName.data))
 		{
 			m_session.m_logger.Error(TC("Failed to duplicate handle for temp file (%s)"), fileName.data);
 			return true;
@@ -1694,7 +1692,7 @@ namespace uba
 			payload.cancelEvent = m_cancelEvent.GetHandle();
 			payload.writeEvent = m_writeEvent.GetHandle();
 			payload.readEvent = m_readEvent.GetHandle();
-			payload.communicationHandle = communicationHandle.handle;
+			payload.communicationHandle = communicationHandle.mh;
 			payload.communicationOffset = communicationOffset;
 			payload.rulesIndex = m_startInfo.rules->index;
 			payload.version = ProcessMessageVersion;
diff --git a/Engine/Source/Programs/UnrealBuildAccelerator/Common/Private/UbaSession.cpp b/Engine/Source/Programs/UnrealBuildAccelerator/Common/Private/UbaSession.cpp
--- a/Engine/Source/Programs/UnrealBuildAccelerator/Common/Private/UbaSession.cpp
+++ b/Engine/Source/Programs/UnrealBuildAccelerator/Common/Private/UbaSession.cpp
@@ -2353,11 +2353,11 @@ namespace uba
 			{
 				if (writtenFile.name != msgName)
 				{
-					UBA_ASSERT(msg.mappingHandle == 0 && !writtenFile.mappingHandle.IsValid());
+					UBA_ASSERT(!msg.mappingHandle.IsValid() && !writtenFile.mappingHandle.IsValid());
 					writtenFile.name = msgName;
 				}
 
-				if (!msg.mappingHandle || msg.mappingHandle == writtenFile.originalMappingHandle)
+				if (!msg.mappingHandle.IsValid() || msg.mappingHandle == writtenFile.originalMappingHandle)
 				{
 					if (msg.mappingWritten)
 					{
@@ -2381,7 +2381,7 @@ namespace uba
 			if (!m_runningRemote && HasVfs(startInfo.rootsHandle)) // For posix we write the dependency file directly to disk so we need to update it if vfs is enabled
 			{
 				bool escapeSpaces;
-				if (!msg.mappingHandle && startInfo.rules->ShouldDevirtualizeFile(file.name, escapeSpaces))
+				if (!msg.mappingHandle.IsValid() && startInfo.rules->ShouldDevirtualizeFile(file.name, escapeSpaces))
 				{
 					// TODO: On linux we don't use file mappings for outputs yet.. so we have to open the file and change it
 					FileAccessor readFile(m_logger, name);
@@ -2410,13 +2410,10 @@ namespace uba
 			{
 				writtenFile.name = msgName;
 				FileMappingHandle mappingHandle;
-				if (msg.mappingHandle)
-				{
-					FileMappingHandle source;
-					source.FromU64(msg.mappingHandle);
-					if (!DuplicateFileMapping(m_logger, msg.process.m_nativeProcessHandle, source, GetCurrentProcessHandle(), &mappingHandle, FILE_MAP_READ, false, 0, msgName))
+				if (msg.mappingHandle.IsValid())
+					if (!DuplicateFileMapping(m_logger, msg.process.m_nativeProcessHandle, msg.mappingHandle, GetCurrentProcessHandle(), mappingHandle, FILE_MAP_READ, false, 0, msgName))
 						return m_logger.Error(TC("Failed to duplicate file mapping handle for %s"), name);
-				}
+
 				writtenFile.key = key;
 				writtenFile.mappingHandle = mappingHandle;
 				writtenFile.mappingWritten = msg.mappingWritten;
@@ -3393,14 +3390,12 @@ namespace uba
 
 	bool Session::ExtractSymbolsFromObjectFile(const CloseFileMessage& msg, const tchar* fileName, u64 fileSize)
 	{
-		if (!msg.mappingHandle)
+		if (!msg.mappingHandle.IsValid())
 			return m_logger.Error(TC("Can't extract symbols from obj file that is written directly to disk (%s writing %s)"), msg.process.m_startInfo.application, fileName);
 
-		FileMappingHandle source;
-		source.FromU64(msg.mappingHandle);
 		FileMappingHandle objectFileMappingHandle;
 
-		if (!DuplicateFileMapping(m_logger, msg.process.m_nativeProcessHandle, source, GetCurrentProcessHandle(), &objectFileMappingHandle, FILE_MAP_ALL_ACCESS, false, 0, fileName))
+		if (!DuplicateFileMapping(m_logger, msg.process.m_nativeProcessHandle, msg.mappingHandle, GetCurrentProcessHandle(), objectFileMappingHandle, FILE_MAP_ALL_ACCESS, false, 0, fileName))
 			return m_logger.Error(TC("Failed to duplicate file mapping handle for %s"), fileName);
 		auto ofmh = MakeGuard([&]() { CloseFileMapping(m_logger, objectFileMappingHandle, fileName); });
 
diff --git a/Engine/Source/Programs/UnrealBuildAccelerator/Common/Private/UbaStorageClient.cpp b/Engine/Source/Programs/UnrealBuildAccelerator/Common/Private/UbaStorageClient.cpp
--- a/Engine/Source/Programs/UnrealBuildAccelerator/Common/Private/UbaStorageClient.cpp
+++ b/Engine/Source/Programs/UnrealBuildAccelerator/Common/Private/UbaStorageClient.cpp
@@ -858,7 +858,7 @@ namespace uba
 		if (keepMappingInMemory && isPersistentMapping && !localFile.casEntry.mappingHandle.IsValid())
 		{
 			FileMappingHandle mappingHandle2;
-			if (DuplicateFileMapping(m_logger, GetCurrentProcessHandle(), mappingHandle, GetCurrentProcessHandle(), &mappingHandle2, FILE_MAP_READ, false, 0, fileName))
+			if (DuplicateFileMapping(m_logger, GetCurrentProcessHandle(), mappingHandle, GetCurrentProcessHandle(), mappingHandle2, FILE_MAP_READ, false, 0, fileName))
 			{
 				localFile.casEntry.mappingHandle = mappingHandle2;
 				localFile.casEntry.size = fileSize;
diff --git a/Engine/Source/Programs/UnrealBuildAccelerator/Common/Public/UbaFileMapping.h b/Engine/Source/Programs/UnrealBuildAccelerator/Common/Public/UbaFileMapping.h
--- a/Engine/Source/Programs/UnrealBuildAccelerator/Common/Public/UbaFileMapping.h
+++ b/Engine/Source/Programs/UnrealBuildAccelerator/Common/Public/UbaFileMapping.h
@@ -11,16 +11,7 @@ namespace uba
 	class Logger;
 	class WorkManager;
 
-#if PLATFORM_WINDOWS
-	struct FileMappingHandle
-	{
-		HANDLE handle = 0;
-		bool operator==(const FileMappingHandle& o) const { return handle == o.handle; }
-		bool IsValid() const { return handle != 0; }
-		u64 ToU64() const { return u64(handle); }
-		void FromU64(u64 v) { handle = (HANDLE)v; }
-	};
-#else
+#if !PLATFORM_WINDOWS
 	inline constexpr u32 FILE_MAP_WRITE = 0x0002;
 	inline constexpr u32 FILE_MAP_READ = 0x0004;
 	inline constexpr u32 FILE_MAP_ALL_ACCESS = FILE_MAP_WRITE | FILE_MAP_READ;
@@ -32,19 +23,6 @@ namespace uba
 	#else
 	inline constexpr u32 SHM_MAX_FILENAME = 38;
 	#endif
-
-	struct FileMappingHandle
-	{
-		FileMappingHandle(int shmFd_ = -1, int lockFd_ = -1, u64 uid_ = ~u64(0)) : shmFd(shmFd_), lockFd(lockFd_), uid(uid_) {};
-		FileMappingHandle(const FileMappingHandle& o) { shmFd = o.shmFd; uid = o.uid; lockFd = o.lockFd; }
-		int shmFd;
-		int lockFd;
-		u64 uid;
-		bool operator==(const FileMappingHandle& o) const { return shmFd == o.shmFd; }
-		bool IsValid() const { return shmFd != -1; }
-		u64 ToU64() const { return uid; }
-		void FromU64(u64 v) { UBA_ASSERT(false); uid = v; }
-	};
 #endif
 
 	FileMappingHandle CreateMemoryMappingW(Logger& logger, u32 flProtect, u64 maxSize, const tchar* name, const tchar* hint);
@@ -53,7 +31,7 @@ namespace uba
 	bool MapViewCommit(void* address, u64 size);
 	bool UnmapViewOfFile(Logger& logger, const void* lpBaseAddress, u64 bytesToUnmap, const tchar* hint);
 	bool CloseFileMapping(Logger& logger, FileMappingHandle h, const tchar* hint);
-	bool DuplicateFileMapping(Logger& logger, ProcHandle hSourceProcessHandle, FileMappingHandle hSourceHandle, ProcHandle hTargetProcessHandle, FileMappingHandle* lpTargetHandle, u32 dwDesiredAccess, bool bInheritHandle, u32 dwOptions, const tchar* hint);
+	bool DuplicateFileMapping(Logger& logger, ProcHandle sourceProcessHandle, FileMappingHandle sourceHandle, ProcHandle targetProcessHandle, FileMappingHandle& targetHandle, u32 dwDesiredAccess, bool bInheritHandle, u32 dwOptions, const tchar* hint);
 	void MapMemoryCopy(void* dest, const void* source, u64 size);
 
 	struct MappedView
diff --git a/Engine/Source/Programs/UnrealBuildAccelerator/Common/Public/UbaSession.h b/Engine/Source/Programs/UnrealBuildAccelerator/Common/Public/UbaSession.h
--- a/Engine/Source/Programs/UnrealBuildAccelerator/Common/Public/UbaSession.h
+++ b/Engine/Source/Programs/UnrealBuildAccelerator/Common/Public/UbaSession.h
@@ -365,7 +365,7 @@ namespace uba
 		u32 attributes = 0;
 		bool deleteOnClose = false;
 		bool success = true;
-		u64 mappingHandle = 0;
+		FileMappingHandle mappingHandle;
 		u64 mappingWritten = 0;
 	};
 
@@ -512,7 +512,7 @@ namespace uba
 		TString name;
 		FileMappingHandle mappingHandle;
 		u64 mappingWritten = 0;
-		u64 originalMappingHandle = 0;
+		FileMappingHandle originalMappingHandle;
 		u64 lastWriteTime = 0;
 		u32 attributes = 0;
 	};
diff --git a/Engine/Source/Programs/UnrealBuildAccelerator/Core/Public/UbaPlatform.h b/Engine/Source/Programs/UnrealBuildAccelerator/Core/Public/UbaPlatform.h
--- a/Engine/Source/Programs/UnrealBuildAccelerator/Core/Public/UbaPlatform.h
+++ b/Engine/Source/Programs/UnrealBuildAccelerator/Core/Public/UbaPlatform.h
@@ -141,6 +141,15 @@ namespace uba
 	#define TStrcat_s wcscat_s
 	#define TStrdup _wcsdup
 	#define UBA_NOINLINE __declspec(noinline)
+	struct FileMappingHandle
+	{
+		bool operator==(const FileMappingHandle& o) const { return mh == o.mh; }
+		bool IsValid() const { return mh != 0; }
+		u64 ToU64() const { return u64(fh) | (u64(mh) << 32); }
+		static FileMappingHandle FromU64(u64 v) { return { (HANDLE)(v & ~0u), (HANDLE)(v >> 32) }; }
+		HANDLE fh = 0;
+		HANDLE mh = 0;
+	};
 #else
 	inline constexpr tchar PathSeparator = '/';
 	inline constexpr tchar NonPathSeparator = '\\';
@@ -172,5 +181,17 @@ namespace uba
 	#define UBA_NOINLINE
     #define ERROR_SUCCESS 0
 	#define DUPLICATE_SAME_ACCESS 0
+	struct FileMappingHandle
+	{
+		FileMappingHandle(int shmFd_ = -1, int lockFd_ = -1, u64 uid_ = ~u64(0)) : shmFd(shmFd_), lockFd(lockFd_), uid(uid_) {};
+		FileMappingHandle(const FileMappingHandle& o) { shmFd = o.shmFd; uid = o.uid; lockFd = o.lockFd; }
+		bool operator==(const FileMappingHandle& o) const { return shmFd == o.shmFd; }
+		bool IsValid() const { return shmFd != -1; }
+		u64 ToU64() const { return uid; }
+		static FileMappingHandle FromU64(u64 v) { UBA_ASSERT(false); return FileMappingHandle(-1, -1, v); }
+		int shmFd;
+		int lockFd;
+		u64 uid;
+	};
 #endif // PLATFORM_WINDOWS
 }
diff --git a/Engine/Source/Programs/UnrealBuildAccelerator/Detours/Private/PosixOS/UbaDetoursFunctionsPosix.cpp b/Engine/Source/Programs/UnrealBuildAccelerator/Detours/Private/PosixOS/UbaDetoursFunctionsPosix.cpp
--- a/Engine/Source/Programs/UnrealBuildAccelerator/Detours/Private/PosixOS/UbaDetoursFunctionsPosix.cpp
+++ b/Engine/Source/Programs/UnrealBuildAccelerator/Detours/Private/PosixOS/UbaDetoursFunctionsPosix.cpp
@@ -563,7 +563,7 @@ void Shared_close(int fd, const TrueClose& trueClose)
 	if (--fo->refCount)
 		return;
 
-	u64 mappingHandle = 0;
+	FileMappingHandle mappingHandle;
 	u64 mappingWritten = 0;
 	FileInfo& fi = *fo->fileInfo;
 	const tchar* path = fi.name;
@@ -2676,7 +2676,7 @@ namespace uba
 				if (!fo->closeId)
 					continue;
 				TRUE_WRAPPER(close)(kv.first);
-				u64 mappingHandle = 0;
+				FileMappingHandle mappingHandle;
 				u64 mappingWritten = 0;
 				FileInfo& fi = *fo->fileInfo;
 				const tchar* path = fi.name;
diff --git a/Engine/Source/Programs/UnrealBuildAccelerator/Detours/Private/UbaDetoursFileMappingTable.cpp b/Engine/Source/Programs/UnrealBuildAccelerator/Detours/Private/UbaDetoursFileMappingTable.cpp
--- a/Engine/Source/Programs/UnrealBuildAccelerator/Detours/Private/UbaDetoursFileMappingTable.cpp
+++ b/Engine/Source/Programs/UnrealBuildAccelerator/Detours/Private/UbaDetoursFileMappingTable.cpp
@@ -118,7 +118,7 @@ namespace uba
 		return tableOffset;
 	}
 
-	void Rpc_UpdateCloseHandle(const tchar* handleName, u32 closeId, bool deleteOnClose, const tchar* newName, u64 mappingHandle, u64 mappingWritten, bool success)
+	void Rpc_UpdateCloseHandle(const tchar* handleName, u32 closeId, bool deleteOnClose, const tchar* newName, const FileMappingHandle& mappingHandle, u64 mappingWritten, bool success)
 	{
 		u32 directoryTableSize;
 		{
@@ -131,7 +131,7 @@ namespace uba
 			//writer.WriteU32(attributes); // TODO THIS
 			writer.WriteBool(deleteOnClose);
 			writer.WriteBool(success);
-			writer.WriteU64(u64(mappingHandle));
+			writer.WriteU64(mappingHandle.ToU64());
 			writer.WriteU64(mappingWritten);
 			if (*newName)
 			{
diff --git a/Engine/Source/Programs/UnrealBuildAccelerator/Detours/Private/UbaDetoursFileMappingTable.h b/Engine/Source/Programs/UnrealBuildAccelerator/Detours/Private/UbaDetoursFileMappingTable.h
--- a/Engine/Source/Programs/UnrealBuildAccelerator/Detours/Private/UbaDetoursFileMappingTable.h
+++ b/Engine/Source/Programs/UnrealBuildAccelerator/Detours/Private/UbaDetoursFileMappingTable.h
@@ -89,7 +89,7 @@ namespace uba
 
 	void Rpc_CreateFileW(const tchar* fileName, const StringKey& fileNameKey, u8 accessFlags, tchar* outNewName, u64 newNameCapacity, u64& outSize, u32& outCloseId, bool lock);
 	u32  Rpc_UpdateDirectory(const StringKey& dirKey, const tchar* dirName, u64 dirNameLen, bool lockDirTable = true);
-	void Rpc_UpdateCloseHandle(const tchar* handleName, u32 closeId, bool deleteOnClose, const tchar* newName, u64 mappingHandle, u64 mappingWritten, bool success);
+	void Rpc_UpdateCloseHandle(const tchar* handleName, u32 closeId, bool deleteOnClose, const tchar* newName, const FileMappingHandle& mappingHandle, u64 mappingWritten, bool success);
 	void Rpc_UpdateTables();
 	u32  Rpc_GetEntryOffset(const StringKey& entryNameKey, const tchar* entryName, u64 entryNameLen, bool checkIfDir = false);
 	void Rpc_GetFullFileName(const tchar*& path, u64& pathLen, StringBufferBase& tempBuf, bool useVirtualName, const tchar* const* loaderPaths = nullptr);
diff --git a/Engine/Source/Programs/UnrealBuildAccelerator/Detours/Private/Windows/UbaDetoursFunctionsKernelBase.inl b/Engine/Source/Programs/UnrealBuildAccelerator/Detours/Private/Windows/UbaDetoursFunctionsKernelBase.inl
--- a/Engine/Source/Programs/UnrealBuildAccelerator/Detours/Private/Windows/UbaDetoursFunctionsKernelBase.inl
+++ b/Engine/Source/Programs/UnrealBuildAccelerator/Detours/Private/Windows/UbaDetoursFunctionsKernelBase.inl
@@ -1426,7 +1426,7 @@ BOOL Detoured_CopyFileExW(LPCWSTR lpExistingFileName, LPCWSTR lpNewFileName, LPP
 		return res;
 
 	bool deleteOnClose = res == false; // If failing to copy we set deleteOnClose
-	Rpc_UpdateCloseHandle(newToName.data, closeId, deleteOnClose, L"", 0, 0, true);
+	Rpc_UpdateCloseHandle(newToName.data, closeId, deleteOnClose, L"", {}, 0, true);
 
 	return res;
 }
@@ -1483,7 +1483,7 @@ BOOL Detoured_CreateHardLinkW(LPCWSTR lpFileName, LPCWSTR lpExistingFileName, LP
 
 	// We need to report the new file that has been added (and we must do it _after_ it has been copied
 	if (closeId)
-		Rpc_UpdateCloseHandle(newToName.data, closeId, false, L"", 0, 0, true);
+		Rpc_UpdateCloseHandle(newToName.data, closeId, false, L"", {}, 0, true);
 
 	return res;
 }
diff --git a/Engine/Source/Programs/UnrealBuildAccelerator/Detours/Private/Windows/UbaDetoursFunctionsNtDll.inl b/Engine/Source/Programs/UnrealBuildAccelerator/Detours/Private/Windows/UbaDetoursFunctionsNtDll.inl
--- a/Engine/Source/Programs/UnrealBuildAccelerator/Detours/Private/Windows/UbaDetoursFunctionsNtDll.inl
+++ b/Engine/Source/Programs/UnrealBuildAccelerator/Detours/Private/Windows/UbaDetoursFunctionsNtDll.inl
@@ -1102,7 +1102,7 @@ NTSTATUS NTAPI Shared_NtCreateFile(bool IsCreateFunc, PHANDLE hFileHandle, ACCES
 			UBA_ASSERT(handleStrEnd);
 			return STATUS_UNSUCCESSFUL;
 		}
-		HANDLE mappingHandle = (HANDLE)StringToValue(handleStr, handleStrEnd - handleStr);
+		HANDLE mappingHandle = FileMappingHandle::FromU64(StringToValue(handleStr, handleStrEnd - handleStr)).mh;
 		const wchar_t* mappingOffsetStr = handleStrEnd + 1;
 		u64 mappingOffset = StringToValue(mappingOffsetStr, wcslen(mappingOffsetStr));
 		info.trueFileMapOffset = mappingOffset;
@@ -1149,21 +1149,22 @@ NTSTATUS NTAPI Shared_NtCreateFile(bool IsCreateFunc, PHANDLE hFileHandle, ACCES
 					writer.WriteString(fileNameLower);
 					writer.Flush();
 					BinaryReader reader;
-					u64 mappingHandle = reader.ReadU64();
+					FileMappingHandle mappingHandle = FileMappingHandle::FromU64(reader.ReadU64());
 					u64 mappingHandleSize = reader.ReadU64();
 					pcs.Leave();
-					if (mappingHandle)
+					if (mappingHandle.IsValid())
 					{
 						info.memoryFile = new MemoryFile(nullptr, false);
 						MemoryFile& mf = *info.memoryFile;
-						True_DuplicateHandle(g_hostProcess, (HANDLE)mappingHandle, GetCurrentProcess(), &mf.mappingHandle, 0, FALSE, DUPLICATE_SAME_ACCESS);
-						UBA_ASSERTF(mf.mappingHandle, L"DuplicateHandle failed when opening temp file %ls (%u)", fileName.data, GetLastError());
+						HANDLE newHandle = 0;
+						True_DuplicateHandle(g_hostProcess, mappingHandle.mh, GetCurrentProcess(), &newHandle, 0, FALSE, DUPLICATE_SAME_ACCESS);
+						UBA_ASSERTF(newHandle, L"DuplicateHandle failed when opening temp file %ls (%u)", fileName.data, GetLastError());
 						mf.writtenSize = mappingHandleSize;
 						TimerScope ts2(g_kernelStats.mapViewOfFile);
-						mf.baseAddress = (u8*)True_MapViewOfFile(mf.mappingHandle, FILE_MAP_READ | FILE_MAP_ALL_ACCESS, 0, 0, mappingHandleSize);
+						mf.baseAddress = (u8*)True_MapViewOfFile(newHandle, FILE_MAP_READ | FILE_MAP_ALL_ACCESS, 0, 0, mappingHandleSize);
 						UBA_ASSERTF(mf.baseAddress, L"MapViewOfFile failed when opening temp file %ls (%u)", fileName.data, GetLastError());
-						//mf.committedSize = mappingHandleSize;
 						mf.reserveSize = FileTypeMaxSize(fileName, isSystemOrTempFile);
+						mf.mappingHandle.mh = newHandle;
 					}
 					else if (isWrite)
 					{
@@ -1265,7 +1266,7 @@ NTSTATUS NTAPI Shared_NtCreateFile(bool IsCreateFunc, PHANDLE hFileHandle, ACCES
 		if (closeId)
 		{
 			info.lastDesiredAccess = lastDesiredAccess;
-			Rpc_UpdateCloseHandle(L"", closeId, false, L"", 0, 0, false);
+			Rpc_UpdateCloseHandle(L"", closeId, false, L"", {}, 0, false);
 		}
 		DEBUG_LOG_TRUE(funcName, L"%ls (%ls) (%ls) -> %ls", isWriteStr, lpFileName, (fileName.data != lpFileName ? fileName.data : L""), ToString(res).data);
 		return res;
@@ -1401,25 +1402,25 @@ NTSTATUS NTAPI Detoured_NtClose(HANDLE handle)
 		return res;
 	}
 
-	HANDLE mappingHandle = 0;
+	FileMappingHandle mappingHandle;
 	u64 mappingWritten = 0;
 	FileInfo& fi = *fo->fileInfo;
 	const wchar_t* path = fi.name;
 	wchar_t temp[512];
-	if (fi.memoryFile)
+	if (MemoryFile* mf = fi.memoryFile)
 	{
 		if (IsWrite(fo->desiredAccess, 0))
 		{
 			// TODO: There are race conditions in this code. There could be other file handles accessing the same piece of memory (allthough unlikely)
-			u64 alignedWritten = AlignUp(fi.memoryFile->writtenSize, 64 * 1024);
-			if (alignedWritten < fi.memoryFile->committedSize)
+			u64 alignedWritten = AlignUp(mf->writtenSize, 64 * 1024);
+			if (alignedWritten < mf->committedSize)
 			{
-				u64 decommitSize = u64(fi.memoryFile->committedSize - alignedWritten);
-				if (fi.memoryFile->isLocalOnly)
+				u64 decommitSize = u64(mf->committedSize - alignedWritten);
+				if (mf->isLocalOnly)
 				{
 #pragma warning(push)
 #pragma warning(disable:6250)
-					if (!::VirtualFree(fi.memoryFile->baseAddress + alignedWritten, decommitSize, MEM_DECOMMIT))
+					if (!::VirtualFree(mf->baseAddress + alignedWritten, decommitSize, MEM_DECOMMIT))
 						ToInvestigate(L"Failed to decommit memory (%u)", GetLastError());
 #pragma warning(pop)
 				}
@@ -1427,25 +1428,25 @@ NTSTATUS NTAPI Detoured_NtClose(HANDLE handle)
 				{
 					// Speculative change. According to stackoverflow (I know) this can hint the system that this memory is not needed anymore.
 					// Building UnrealEditor and friends put huge pressure on committed space and anything that can reduce that is valuable
-					if (!::VirtualUnlock(fi.memoryFile->baseAddress + alignedWritten, decommitSize))
+					if (!::VirtualUnlock(mf->baseAddress + alignedWritten, decommitSize))
 						if (GetLastError() != ERROR_NOT_LOCKED)
 							ToInvestigate(L"Failed to unlock memory (%u)", GetLastError());
 				}
-				fi.memoryFile->committedSize = alignedWritten;
+				mf->committedSize = alignedWritten;
 			}
 		}
 
-		mappingHandle = fi.memoryFile->mappingHandle;
-		mappingWritten = fi.memoryFile->writtenSize;
+		mappingHandle = mf->mappingHandle;
+		mappingWritten = mf->writtenSize;
 
 		u32 orginalNameLen = TStrlen(fi.originalName);
 		if (IsOutputFile(StringView(fi.originalName, orginalNameLen), IsWrite(fo->desiredAccess, 0), fo->deleteOnClose) && !g_rules->IsThrowAway(StringView(fi.originalName, orginalNameLen), g_runningRemote))
 		{
 			// Need to report this file to host so it can be tracked in directory table
-			if (!fi.memoryFile->isReported)
+			if (!mf->isReported)
 			{
 				path = temp;
-				fi.memoryFile->isReported = true;
+				mf->isReported = true;
 				const wchar_t* fileName = fi.originalName;
 				if (!fo->newName.empty())
 					fileName = fo->newName.c_str();
@@ -1479,11 +1480,10 @@ NTSTATUS NTAPI Detoured_NtClose(HANDLE handle)
 				fo->ownsFileInfo = false;
 				fo->newName.clear();
 			}
-
 		}
 		else if (NeedsSharedMemory(fi.originalName) && IsWrite(fo->desiredAccess, 0))
 		{
-			UBA_ASSERT(!fi.memoryFile->isLocalOnly);
+			UBA_ASSERT(!mf->isLocalOnly);
 			StringBuffer<> fixedName;
 			FixPath(fixedName, path);
 
@@ -1493,7 +1493,7 @@ NTSTATUS NTAPI Detoured_NtClose(HANDLE handle)
 			writer.WriteByte(MessageType_CreateTempFile);
 			writer.WriteStringKey(ToStringKeyLower(fixedName));
 			writer.WriteString(fixedName);
-			writer.WriteU64((u64)mappingHandle);
+			writer.WriteU64(mappingHandle.ToU64());
 			writer.WriteU64(mappingWritten);
 			writer.Flush();
 			BinaryReader reader;
@@ -1506,7 +1506,7 @@ NTSTATUS NTAPI Detoured_NtClose(HANDLE handle)
 
 	if (fo->closeId)
 	{
-		Rpc_UpdateCloseHandle(path, fo->closeId, fo->deleteOnClose, fo->newName.c_str(), (u64)mappingHandle, mappingWritten, true);
+		Rpc_UpdateCloseHandle(path, fo->closeId, fo->deleteOnClose, fo->newName.c_str(), mappingHandle, mappingWritten, true);
 	}
 	else
 	{
diff --git a/Engine/Source/Programs/UnrealBuildAccelerator/Detours/Private/Windows/UbaDetoursFunctionsWin.cpp b/Engine/Source/Programs/UnrealBuildAccelerator/Detours/Private/Windows/UbaDetoursFunctionsWin.cpp
--- a/Engine/Source/Programs/UnrealBuildAccelerator/Detours/Private/Windows/UbaDetoursFunctionsWin.cpp
+++ b/Engine/Source/Programs/UnrealBuildAccelerator/Detours/Private/Windows/UbaDetoursFunctionsWin.cpp
@@ -245,6 +245,66 @@ const wchar_t* HandleToName(const DetouredHandle& dh)
 
 BlockAllocator<FileObject>& g_fileObjectAllocator = *new BlockAllocator<FileObject>(g_memoryBlock);
 
+MemoryFile::MemoryFile(u8* data, bool localOnly)
+:	baseAddress(data)
+,	isLocalOnly(localOnly)
+{
+}
+
+MemoryFile::MemoryFile(bool localOnly, u64 reserveSize_, bool isThrowAway_)
+:	isLocalOnly(localOnly)
+,	isThrowAway(isThrowAway_)
+{
+	if (!isThrowAway)
+		Reserve(reserveSize_);
+}
+
+void MemoryFile::Reserve(u64 reserveSize_)
+{
+	UBA_ASSERT(!isThrowAway);
+	reserveSize = reserveSize_;
+	if (isLocalOnly)
+	{
+		baseAddress = (u8*)VirtualAlloc(NULL, reserveSize, MEM_RESERVE, PAGE_READWRITE);
+		if (!baseAddress)
+			FatalError(1354, L"VirtualAlloc failed trying to reserve %llu. (Error code: %u)", reserveSize, GetLastError());
+		mappedSize = reserveSize;
+	}
+	else
+	{
+		HANDLE handle;
+		{
+			mappedSize = 32 * 1024 * 1024;
+			TimerScope ts(g_kernelStats.createFileMapping);
+			handle = True_CreateFileMappingW(INVALID_HANDLE_VALUE, NULL, PAGE_READWRITE | SEC_RESERVE, ToHigh(reserveSize), ToLow(reserveSize), NULL);
+			if (!handle)
+				FatalError(1348, L"CreateFileMappingW failed trying to reserve %llu. (Error code: %u)", reserveSize, GetLastError());
+		}
+		TimerScope ts(g_kernelStats.mapViewOfFile);
+		baseAddress = (u8*)True_MapViewOfFile(handle, FILE_MAP_READ | FILE_MAP_WRITE, 0, 0, mappedSize);
+		if (!baseAddress)
+			FatalError(1353, L"MapViewOfFile failed trying to map %llu. ReservedSize: %llu (Error code: %u)", mappedSize, reserveSize, GetLastError());
+		mappingHandle = { 0, handle };
+	}
+}
+
+void MemoryFile::Unreserve()
+{
+	if (isLocalOnly)
+	{
+		VirtualFree(baseAddress, 0, MEM_RELEASE);
+	}
+	else
+	{
+		True_UnmapViewOfFile(baseAddress);
+		CloseHandle(mappingHandle.mh);
+		UBA_ASSERT(!mappingHandle.fh);
+		mappingHandle.mh = 0;
+	}
+	baseAddress = nullptr;
+	committedSize = 0;
+}
+
 void MemoryFile::Write(DetouredHandle& handle, LPCVOID lpBuffer, u64 nNumberOfBytesToWrite)
 {
 	u64 newPos = handle.pos + nNumberOfBytesToWrite;
@@ -305,7 +365,7 @@ void MemoryFile::Remap(const DetouredHandle& handle, u64 size)
 	True_UnmapViewOfFile(baseAddress);
 	mappedSize = Min(reserveSize, AlignUp(Max(size, mappedSize * 4), g_pageSize));
 	TimerScope ts(g_kernelStats.mapViewOfFile);
-	baseAddress = (u8*)True_MapViewOfFile(mappingHandle, FILE_MAP_READ | FILE_MAP_WRITE, 0, 0, mappedSize);
+	baseAddress = (u8*)True_MapViewOfFile(mappingHandle.mh, FILE_MAP_READ | FILE_MAP_WRITE, 0, 0, mappedSize);
 	if (!baseAddress)
 		FatalError(1347, L"MapViewOfFile failed trying to map %llu for %ls. ReservedSize: %llu (Error code: %u)", mappedSize, HandleToName(handle), reserveSize, GetLastError());
 }
@@ -467,7 +527,8 @@ void CloseCaches()
 		if (info.memoryFile && !info.memoryFile->isLocalOnly)
 		{
 			UnmapViewOfFile(info.memoryFile->baseAddress);
-			CloseHandle(info.memoryFile->mappingHandle);
+			CloseHandle(info.memoryFile->mappingHandle.mh);
+			UBA_ASSERT(!info.memoryFile->mappingHandle.fh);
 		}
 		//if (info.memoryFile && info.memoryFile != &g_emptyMemoryFile)
 		//	delete info.memoryFile;
@@ -868,10 +929,10 @@ void Init(const DetoursPayload& payload, u64 startTime)
 		reader.ReadString(applicationBuffer);
 		reader.ReadString(workingDirBuffer);
 
-		directoryTableHandle = (HANDLE)reader.ReadU64();
+		directoryTableHandle = FileMappingHandle::FromU64(reader.ReadU64()).mh;
 		directoryTableSize = reader.ReadU32();
 		directoryTableCount = reader.ReadU32();
-		mappedFileTableHandle = (HANDLE)reader.ReadU64();
+		mappedFileTableHandle = FileMappingHandle::FromU64(reader.ReadU64()).mh;
 		mappedFileTableSize = reader.ReadU32();
 		mappedFileTableCount = reader.ReadU32();
 
diff --git a/Engine/Source/Programs/UnrealBuildAccelerator/Detours/Private/Windows/UbaDetoursUtilsWin.h b/Engine/Source/Programs/UnrealBuildAccelerator/Detours/Private/Windows/UbaDetoursUtilsWin.h
--- a/Engine/Source/Programs/UnrealBuildAccelerator/Detours/Private/Windows/UbaDetoursUtilsWin.h
+++ b/Engine/Source/Programs/UnrealBuildAccelerator/Detours/Private/Windows/UbaDetoursUtilsWin.h
@@ -60,56 +60,11 @@ namespace uba
 
 	struct MemoryFile
 	{
-		MemoryFile(u8* data = nullptr, bool localOnly = true) : baseAddress(data), isLocalOnly(localOnly) {}
-		MemoryFile(bool localOnly, u64 reserveSize_, bool isThrowAway_ = false) : isLocalOnly(localOnly), isThrowAway(isThrowAway_)
-		{
-			if (!isThrowAway)
-				Reserve(reserveSize_);
-		}
-
-		void Reserve(u64 reserveSize_)
-		{
-			UBA_ASSERT(!isThrowAway);
-			reserveSize = reserveSize_;
-			if (isLocalOnly)
-			{
-				baseAddress = (u8*)VirtualAlloc(NULL, reserveSize, MEM_RESERVE, PAGE_READWRITE);
-				if (!baseAddress)
-					FatalError(1354, L"VirtualAlloc failed trying to reserve %llu. (Error code: %u)", reserveSize, GetLastError());
-				mappedSize = reserveSize;
-			}
-			else
-			{
-				{
-					mappedSize = 32 * 1024 * 1024;
-					TimerScope ts(g_kernelStats.createFileMapping);
-					mappingHandle = True_CreateFileMappingW(INVALID_HANDLE_VALUE, NULL, PAGE_READWRITE | SEC_RESERVE, ToHigh(reserveSize), ToLow(reserveSize), NULL);
-					if (!mappingHandle)
-						FatalError(1348, L"CreateFileMappingW failed trying to reserve %llu. (Error code: %u)", reserveSize, GetLastError());
-				}
-				TimerScope ts(g_kernelStats.mapViewOfFile);
-				baseAddress = (u8*)True_MapViewOfFile(mappingHandle, FILE_MAP_READ | FILE_MAP_WRITE, 0, 0, mappedSize);
-				if (!baseAddress)
-					FatalError(1353, L"MapViewOfFile failed trying to map %llu. ReservedSize: %llu (Error code: %u)", mappedSize, reserveSize, GetLastError());
-			}
-		}
-
-		void Unreserve()
-		{
-			if (isLocalOnly)
-			{
-				VirtualFree(baseAddress, 0, MEM_RELEASE);
-			}
-			else
-			{
-				True_UnmapViewOfFile(baseAddress);
-				CloseHandle(mappingHandle);
-				mappingHandle = nullptr;
-			}
-			baseAddress = nullptr;
-			committedSize = 0;
-		}
+		MemoryFile(u8* data = nullptr, bool localOnly = true);
+		MemoryFile(bool localOnly, u64 reserveSize_, bool isThrowAway_ = false);
 
+		void Reserve(u64 reserveSize_);
+		void Unreserve();
 		void Write(struct DetouredHandle& handle, LPCVOID lpBuffer, u64 nNumberOfBytesToWrite);
 		void EnsureCommitted(const struct DetouredHandle& handle, u64 size);
 		void Remap(const struct DetouredHandle& handle, u64 size);
@@ -118,8 +73,8 @@ namespace uba
 		u64 fileTime = ~u64(0);
 		u32 volumeSerial = 0;
 
-		HANDLE mappingHandle = nullptr;
-		u8* baseAddress;
+		FileMappingHandle mappingHandle;
+		u8* baseAddress = nullptr;
 		u64 reserveSize = 0;
 		u64 mappedSize = 0;
 		u64 committedSize = 0;
