[UBA]
* Fixed visualizer compile error
* Fixed better handling in visualizer for corrupt uba files
* Added so it is possible to open a .uba file from a visualizer that is listening for new sessions

diff --git a/Engine/Source/Programs/UnrealBuildAccelerator/Core/Public/UbaBinaryReaderWriter.h b/Engine/Source/Programs/UnrealBuildAccelerator/Core/Public/UbaBinaryReaderWriter.h
--- a/Engine/Source/Programs/UnrealBuildAccelerator/Core/Public/UbaBinaryReaderWriter.h
+++ b/Engine/Source/Programs/UnrealBuildAccelerator/Core/Public/UbaBinaryReaderWriter.h
@@ -470,7 +470,7 @@ namespace uba
 		bool hasMoreBytes;
 		do
 		{
-			if (m_pos == m_end)
+			if (m_pos >= m_end)
 				return false;
 			u8 value = *m_pos++;
 			hasMoreBytes = value & 0x80;
@@ -484,7 +484,41 @@ namespace uba
 	bool BinaryReader::TryReadString(StringBufferBase& out)
 	{
 		#if PLATFORM_WINDOWS
-		ReadString(out); // TODO: Implement secure read
+		u64 charLen;
+		if (!TryRead7BitEncoded(charLen))
+			return false;
+		tchar* it = out.data + out.count;
+		u64 left = charLen;
+		while (left--)
+		{
+			if (m_pos >= m_end)
+				return false;
+			u8 a = *m_pos++;
+			if (a <= 127)
+			{
+				*it++ = a;
+				continue;
+			}
+			if (m_pos >= m_end)
+				return false;
+			u8 b = *m_pos++;
+			if (a >= 192 && a <= 223)
+			{
+				*it++ = (a-192)*64 + (b-128);
+				continue;
+			}
+			if (m_pos >= m_end)
+				return false;
+			u8 c = *m_pos++;
+			if (a >= 224 && a <= 239)
+			{
+				*it++ = (a-224)*4096 + (b-128)*64 + (c-128);
+				continue;
+			}
+			return false;
+		}
+		*it = 0;
+		out.count = u32(it - out.data);
 		#else
 		u64 charLen;
 		if (!TryRead7BitEncoded(charLen))
diff --git a/Engine/Source/Programs/UnrealBuildAccelerator/Visualizer/Private/UbaTraceReader.cpp b/Engine/Source/Programs/UnrealBuildAccelerator/Visualizer/Private/UbaTraceReader.cpp
--- a/Engine/Source/Programs/UnrealBuildAccelerator/Visualizer/Private/UbaTraceReader.cpp
+++ b/Engine/Source/Programs/UnrealBuildAccelerator/Visualizer/Private/UbaTraceReader.cpp
@@ -95,7 +95,7 @@ namespace uba
 		m_memoryPos = m_memoryBegin;
 		m_memoryEnd = m_memoryBegin + fileSize;
 
-		BinaryReader reader(m_memoryBegin);
+		BinaryReader reader(m_memoryBegin, 0, fileSize);
 
 		u32 traceSize = reader.ReadU32(); (void)traceSize;
 		u32 version = reader.ReadU32();
@@ -142,9 +142,8 @@ namespace uba
 	bool TraceReader::UpdateReadFile(TraceView& out, u64 maxTime, bool& outChanged)
 	{
 		bool res = true;
-		BinaryReader traceReader(m_memoryPos);
-		u64 left = u64(m_memoryEnd - m_memoryPos);
-		while (traceReader.GetPosition() < left)
+		BinaryReader traceReader(m_memoryPos, 0, m_memoryEnd - m_memoryPos);
+		while (traceReader.GetLeft())
 		{
 			u64 pos = traceReader.GetPosition();
 			if (!ReadTrace(out, traceReader, maxTime))
@@ -153,12 +152,9 @@ namespace uba
 				break;
 			}
 			if (pos == traceReader.GetPosition())
-			{
-				if (!traceReader.GetLeft())
-					out.finished = true;
 				break;
-			}
 		}
+		out.finished = !traceReader.GetLeft();
 		outChanged = traceReader.GetPosition() != 0 || !m_activeProcesses.empty();
 		m_memoryPos += traceReader.GetPosition();
 		return res;
@@ -234,7 +230,7 @@ namespace uba
 
 		if (*namedTrace && m_namedTrace != namedTrace)
 		{
-			m_memoryHandle.handle = ::OpenFileMappingW(PAGE_READWRITE, false, namedTrace);
+			m_memoryHandle.mh = ::OpenFileMappingW(PAGE_READWRITE, false, namedTrace);
 			if (!m_memoryHandle.IsValid())
 			{
 				if (!silentFail)
@@ -383,9 +379,11 @@ namespace uba
 		case TraceType_SessionAdded:
 		{
 			StringBuffer<128> sessionName;
-			reader.ReadString(sessionName);
+			if (!reader.TryReadString(sessionName))
+				return false;
 			StringBuffer<> sessionInfo;
-			reader.ReadString(sessionInfo);
+			if (!reader.TryReadString(sessionInfo))
+				return false;
 			Guid clientUid = ReadClientId(out, reader);
 			u32 sessionIndex = reader.ReadU32();
 
@@ -801,7 +799,9 @@ namespace uba
 		case TraceType_FileFetchLight:
 		{
 			Guid clientUid = ReadClientId(out, reader);
-			u64 fileSize = reader.Read7BitEncoded();
+			u64 fileSize;
+			if (!reader.TryRead7BitEncoded(fileSize))
+				return false;
 			if (auto session = GetSession(out, clientUid))
 			{
 				session->fetchedFiles.push_back({ CasKeyZero, fileSize, L"", time, time });
diff --git a/Engine/Source/Programs/UnrealBuildAccelerator/Visualizer/Private/UbaVisualizer.cpp b/Engine/Source/Programs/UnrealBuildAccelerator/Visualizer/Private/UbaVisualizer.cpp
--- a/Engine/Source/Programs/UnrealBuildAccelerator/Visualizer/Private/UbaVisualizer.cpp
+++ b/Engine/Source/Programs/UnrealBuildAccelerator/Visualizer/Private/UbaVisualizer.cpp
@@ -910,6 +910,51 @@ namespace uba
 	{
 		u64 playTime = GetPlayTime();
 
+		SetBkMode(hdc, TRANSPARENT);
+		SetTextColor(hdc, m_textColor);
+		SetBkColor(hdc, m_config.DarkMode ? RGB(70, 70, 70) : RGB(180, 180, 180));
+
+		auto DrawCenteredText = [&](const StringBuffer<>* lines, u32 lineCount, bool drawBackground)
+			{
+				for (u32 i=0; i!=lineCount; ++i)
+				{
+					RECT rect = { 0, 0, 0, 0 };
+					DrawTextW(hdc, lines[i].data, lines[i].count, &rect, DT_SINGLELINE|DT_NOCLIP|DT_CALCRECT);
+					int top = (clientRect.bottom - rect.bottom)/2 + i*20;
+					rect.left = (clientRect.right - rect.right)/2;
+					rect.right = (clientRect.right + rect.right)/2;
+					rect.bottom = top + (rect.bottom - rect.top);
+					rect.top = top;
+
+					if (drawBackground)
+					{
+						RECT fillRect = rect;
+						InflateRect(&fillRect, 4, 2);
+						FillRect(hdc, &fillRect, m_tooltipBackgroundBrush);
+					}
+					DrawTextW(hdc, lines[i].data, lines[i].count, &rect, DT_SINGLELINE|DT_NOCLIP);
+				}
+			};
+
+
+		if (m_traceView.sessions.empty() && !m_parentHwnd)
+		{
+			StringBuffer<> str[2];
+			str[1].Append(TC("Click here to open trace file"));
+			if (m_listenChannel.count)
+			{
+				str[0].Appendf(TC("Listening for new sessions on channel '%s'"), m_listenChannel.data);
+				str[1].Append(TC(" instead"));
+			}
+			else
+				str[0].Append(TC("No trace active"));
+
+			SetActiveFont(m_popupFont);
+			DrawCenteredText(str, 2, false);
+			return;
+		}
+
+
 		int posY = int(m_scrollPosY);
 		float scaleX = 50.0f*m_zoomValue*m_horizontalScaleValue;
 
@@ -921,15 +966,11 @@ namespace uba
 			progressRect.bottom -= m_defaultFont.height + 10;
 		}
 
-		SetBkMode(hdc, TRANSPARENT);
-		SetTextColor(hdc, m_textColor);
-
 		HDC textDC = CreateCompatibleDC(hdc);
 		SetTextColor(textDC, m_textColor);
 		SelectObject(textDC, m_processFont.handle);
 		SelectObject(textDC, GetStockObject(NULL_BRUSH));
 		SetBkMode(textDC, TRANSPARENT);
-		SetBkColor(hdc, m_config.DarkMode ? RGB(70, 70, 70) : RGB(180, 180, 180));
 
 		//TEXTMETRIC metric;
 		//GetTextMetrics(textDC, &metric);
@@ -1711,24 +1752,8 @@ namespace uba
 			str[0].Append(TC("Box Filter: ")).Append(m_filterString);
 			str[1].Appendf(TC("Box Count: %u"), visibleBoxes);
 
-			for (u32 i=0; i!=2; ++i)
-			{
-				SetActiveFont(m_popupFont);
-				RECT rect = { 0, 0, 0, 0 };
-				DrawTextW(hdc, str[i].data, str[i].count, &rect, DT_SINGLELINE|DT_NOCLIP|DT_CALCRECT);
-				int top = (progressRect.bottom - rect.bottom)/2 + i*20;
-				rect.left = (progressRect.right - rect.right)/2;
-				rect.right = (progressRect.right + rect.right)/2;
-				rect.bottom = top + (rect.bottom - rect.top);
-				rect.top = top;
-
-				RECT fillRect = rect;
-				InflateRect(&fillRect, 4, 2);
-				FillRect(hdc, &fillRect, m_tooltipBackgroundBrush);
-				DrawTextW(hdc, str[i].data, str[i].count, &rect, DT_SINGLELINE|DT_NOCLIP);
-
-
-			}
+			SetActiveFont(m_popupFont);
+			DrawCenteredText(str, 2, true);
 		}
 
 		if (m_config.showCursorLine && m_mouseOverWindow)
@@ -3460,6 +3485,36 @@ namespace uba
 
 		case WM_LBUTTONDOWN:
 		{
+			if (m_traceView.sessions.empty())
+			{
+				RECT r;
+				GetClientRect(hWnd, &r);
+				int centerHrz = r.right/2;
+				int centerVrt = r.bottom/2;
+				r = { centerHrz, centerVrt, centerHrz, centerVrt };
+				InflateRect(&r, 200, 50);
+				POINTS p = MAKEPOINTS(lParam);
+				if (PtInRect(&r, { p.x, p.y }))
+				{
+					OPENFILENAME ofn;
+					tchar szFile[MAX_PATH] = TC("");
+					ZeroMemory(&ofn, sizeof(ofn));
+					ofn.lStructSize = sizeof(ofn);
+					ofn.hwndOwner = m_hwnd;
+					ofn.lpstrFilter = L"Uba Files\0*.uba\0All Files\0*.*\0";
+					ofn.lpstrFile = szFile;
+					ofn.nMaxFile = MAX_PATH;
+					ofn.Flags = OFN_FILEMUSTEXIST | OFN_PATHMUSTEXIST;
+
+					if (GetOpenFileName(&ofn))
+					{
+						m_fileName.Append(szFile);
+						PostNewTrace(0, false);
+					}
+				}
+				break;
+			}
+
 			if (m_buttonSelected != ~0u)
 			{
 				bool* values = &m_config.showProgress;
