[UBA]
* Added iops to process stats

diff --git a/Engine/Source/Programs/UnrealBuildAccelerator/Common/Private/UbaProcess.cpp b/Engine/Source/Programs/UnrealBuildAccelerator/Common/Private/UbaProcess.cpp
--- a/Engine/Source/Programs/UnrealBuildAccelerator/Common/Private/UbaProcess.cpp
+++ b/Engine/Source/Programs/UnrealBuildAccelerator/Common/Private/UbaProcess.cpp
@@ -2026,7 +2026,18 @@ namespace uba
 
 #if PLATFORM_WINDOWS
 
-		auto closeHandleGuard = MakeGuard([&]() { CloseHandle((HANDLE)handle); });
+		auto closeHandleGuard = MakeGuard([&]()
+			{
+				IO_COUNTERS ioCounters;
+				if (GetProcessIoCounters((HANDLE)handle, &ioCounters))
+				{
+					m_processStats.iopsRead = ioCounters.ReadOperationCount;
+					m_processStats.iopsWrite = ioCounters.WriteOperationCount;
+					m_processStats.iopsOther = ioCounters.OtherOperationCount;
+				}
+
+				CloseHandle((HANDLE)handle);
+			});
 
 		PROCESS_MEMORY_COUNTERS mem;
 		if (GetProcessMemoryInfo((HANDLE)handle, &mem, sizeof(mem)))
diff --git a/Engine/Source/Programs/UnrealBuildAccelerator/Common/Private/UbaStats.cpp b/Engine/Source/Programs/UnrealBuildAccelerator/Common/Private/UbaStats.cpp
--- a/Engine/Source/Programs/UnrealBuildAccelerator/Common/Private/UbaStats.cpp
+++ b/Engine/Source/Programs/UnrealBuildAccelerator/Common/Private/UbaStats.cpp
@@ -47,6 +47,13 @@ namespace uba
 			if (detoursMemory)
 				logger.Info(TC("  DetoursMem                  %9s"), BytesToText(detoursMemory).str);
 		}
+		if (iopsRead)
+			logger.Info(TC("  IopsRead                    %9llu"), iopsRead.load());
+		if (iopsWrite)
+			logger.Info(TC("  IopsWrite                   %9llu"), iopsWrite.load());
+		if (iopsOther)
+			logger.Info(TC("  IopsOther                   %9llu"), iopsOther.load());
+
 		if (peakMemory)
 			logger.Info(TC("  PeakMem                     %9s"), BytesToText(peakMemory).str);
 		if (cpuTime)
@@ -98,6 +105,12 @@ namespace uba
 			cpuTime = reader.Read7BitEncoded();
 			detoursMemory = reader.Read7BitEncoded();
 			peakMemory = reader.Read7BitEncoded();
+			if (version >= 39)
+			{
+				iopsRead = reader.Read7BitEncoded();
+				iopsWrite = reader.Read7BitEncoded();
+				iopsOther = reader.Read7BitEncoded();
+			}
 			hostTotalTime = reader.Read7BitEncoded();
 		}
 		else
@@ -125,6 +138,9 @@ namespace uba
 		cpuTime += other.cpuTime;
 		detoursMemory = Max(detoursMemory, other.detoursMemory);
 		peakMemory = Max(peakMemory, other.peakMemory);
+		iopsRead += other.iopsRead;
+		iopsWrite += other.iopsWrite;
+		iopsOther += other.iopsOther;
 		hostTotalTime += other.hostTotalTime;
 	}
 
diff --git a/Engine/Source/Programs/UnrealBuildAccelerator/Common/Public/UbaTrace.h b/Engine/Source/Programs/UnrealBuildAccelerator/Common/Public/UbaTrace.h
--- a/Engine/Source/Programs/UnrealBuildAccelerator/Common/Public/UbaTrace.h
+++ b/Engine/Source/Programs/UnrealBuildAccelerator/Common/Public/UbaTrace.h
@@ -68,7 +68,7 @@ namespace uba
 		#undef UBA_TRACE_TYPE
 	};
 
-	static constexpr u32 TraceVersion = 38;
+	static constexpr u32 TraceVersion = 39;
 	static constexpr u32 TraceReadCompatibilityVersion = 6;
 
 	class Trace : public WorkTracker
diff --git a/Engine/Source/Programs/UnrealBuildAccelerator/Core/Public/UbaProcessStats.h b/Engine/Source/Programs/UnrealBuildAccelerator/Core/Public/UbaProcessStats.h
--- a/Engine/Source/Programs/UnrealBuildAccelerator/Core/Public/UbaProcessStats.h
+++ b/Engine/Source/Programs/UnrealBuildAccelerator/Core/Public/UbaProcessStats.h
@@ -67,6 +67,10 @@ namespace uba
 		AtomicU64 detoursMemory;
 		AtomicU64 peakMemory;
 
+		AtomicU64 iopsRead;
+		AtomicU64 iopsWrite;
+		AtomicU64 iopsOther;
+
 		AtomicU64 hostTotalTime;
 
 		void Print(Logger& logger, u64 frequency = GetFrequency());
@@ -97,6 +101,9 @@ namespace uba
 			writer.Write7BitEncoded(cpuTime);
 			writer.Write7BitEncoded(detoursMemory);
 			writer.Write7BitEncoded(peakMemory);
+			writer.Write7BitEncoded(iopsRead);
+			writer.Write7BitEncoded(iopsWrite);
+			writer.Write7BitEncoded(iopsOther);
 			writer.Write7BitEncoded(hostTotalTime);
 		}
 	};
