[UBA]
* Added handling of SetFileInformationByHandle for setting end of file or allocation info when file is a memory file
* Fixed write file counter for trace

diff --git a/Engine/Source/Programs/UnrealBuildAccelerator/Common/Private/UbaProcess.cpp b/Engine/Source/Programs/UnrealBuildAccelerator/Common/Private/UbaProcess.cpp
--- a/Engine/Source/Programs/UnrealBuildAccelerator/Common/Private/UbaProcess.cpp
+++ b/Engine/Source/Programs/UnrealBuildAccelerator/Common/Private/UbaProcess.cpp
@@ -1396,6 +1396,9 @@ namespace uba
 				files.push_back(&kv.second);
 		}
 
+		// We want file count to match number of files.. and it is actually fine that count is 0
+		m_processStats.writeFiles.count += u32(files.size() - 1);
+
 		if (!m_session.WriteFilesToDisk(*this, files.data(), u32(files.size())))
 			return false;
 
diff --git a/Engine/Source/Programs/UnrealBuildAccelerator/Core/Private/UbaApplicationRules.cpp b/Engine/Source/Programs/UnrealBuildAccelerator/Core/Private/UbaApplicationRules.cpp
--- a/Engine/Source/Programs/UnrealBuildAccelerator/Core/Private/UbaApplicationRules.cpp
+++ b/Engine/Source/Programs/UnrealBuildAccelerator/Core/Private/UbaApplicationRules.cpp
@@ -230,6 +230,27 @@ namespace uba
 		{
 			return false//Super::IsOutputFile(fileName)
 				|| fileName.Contains(TC(".exe.tmp"))
+				|| fileName.Contains(TC(".dll.tmp"))
+				|| fileName.Contains(TC(".pdb.tmp"));
+		}
+	};
+
+	class ApplicationRulesRadLinkExe : public ApplicationRulesVcLink
+	{
+		using Super = ApplicationRulesVcLink;
+
+		virtual u64 FileTypeMaxSize(const StringView& file, bool isSystemOrTempFile) const override
+		{
+			if (file.Contains(TC(".pdb.tmp")))
+				return 14ull * 1024 * 1024 * 1024; // This is ridiculous
+			return ApplicationRulesVcLink::FileTypeMaxSize(file, isSystemOrTempFile);
+		}
+
+		virtual bool IsOutputFile(const StringView& fileName) const override
+		{
+			return false//Super::IsOutputFile(fileName)
+				|| fileName.Contains(TC(".exe.tmp"))
+				|| fileName.Contains(TC(".dll.tmp"))
 				|| fileName.Contains(TC(".pdb.tmp"));
 		}
 	};
@@ -635,7 +656,7 @@ namespace uba
 			{ ApplicationHash(TC("clang-cl.exe")),				new ApplicationRulesClangPlusPlusExe() },
 			{ ApplicationHash(TC("verse-clang-cl.exe")),		new ApplicationRulesClangPlusPlusExe() },
 			{ ApplicationHash(TC("ispc.exe")),					new ApplicationRulesISPCExe() },
-			{ ApplicationHash(TC("radlink.exe")),				new ApplicationRulesVcLink() },
+			{ ApplicationHash(TC("radlink.exe")),				new ApplicationRulesRadLinkExe() },
 			{ 3340509542,										new ApplicationRulesClangPlusPlusExePlatform1() },
 			{ 4113554641,										new ApplicationRulesLdExePlatform1() }, // Must be index 15
 			{ 1752955744,										new ApplicationRulesLdExePlatform1() },
diff --git a/Engine/Source/Programs/UnrealBuildAccelerator/Detours/Private/Windows/UbaDetoursFunctionsKernelBase.inl b/Engine/Source/Programs/UnrealBuildAccelerator/Detours/Private/Windows/UbaDetoursFunctionsKernelBase.inl
--- a/Engine/Source/Programs/UnrealBuildAccelerator/Detours/Private/Windows/UbaDetoursFunctionsKernelBase.inl
+++ b/Engine/Source/Programs/UnrealBuildAccelerator/Detours/Private/Windows/UbaDetoursFunctionsKernelBase.inl
@@ -2272,6 +2272,9 @@ BOOL Detoured_SetFileInformationByHandle(HANDLE hFile, FILE_INFO_BY_HANDLE_CLASS
 	}
 
 	auto& dh = asDetouredHandle(hFile);
+	FileObject& fo = *dh.fileObject;
+	FileInfo& fi = *fo.fileInfo;
+
 	if (FileInformationClass == FileRenameInfo)
 	{
 		DEBUG_LOG_TRUE(L"SetFileInformationByHandle", L"%llu (FileRenameInfo)", uintptr_t(hFile));
@@ -2290,29 +2293,48 @@ BOOL Detoured_SetFileInformationByHandle(HANDLE hFile, FILE_INFO_BY_HANDLE_CLASS
 		if (info.DeleteFileW)
 		{
 			DEBUG_LOG_DETOURED(L"SetFileInformationByHandle", L"File is set to be deleted on close (%ls)", HandleToName(hFile));
-			dh.fileObject->deleteOnClose = true;
+			fo.deleteOnClose = true;
 		}
-		else if (dh.fileObject->deleteOnClose)
+		else if (fo.deleteOnClose)
 		{
 			DEBUG_LOG_DETOURED(L"SetFileInformationByHandle", L"File is set to NOT be deleted on close (%ls)", HandleToName(hFile));
-			dh.fileObject->deleteOnClose = false;
+			fo.deleteOnClose = false;
 		}
 		else
 		{
 			DEBUG_LOG_DETOURED(L"SetFileInformationByHandle", L"%llu (FileDispositionInfo %u)", uintptr_t(hFile), info.DeleteFileW);
 		}
 
-		if (dh.fileObject->fileInfo->memoryFile)
+		if (fo.fileInfo->memoryFile)
 			return true;
 
 		DEBUG_LOG_TRUE(L"SetFileInformationByHandle", L"%llu (FileDispositionInfo)", uintptr_t(hFile));
 		return True_SetFileInformationByHandle(hFile, FileInformationClass, lpFileInformation, dwBufferSize); // In here to be tabbed in log
 	}
-	else
+	else if (FileInformationClass == FileAllocationInfo)
 	{
-		DEBUG_LOG_TRUE(L"SetFileInformationByHandle", L"%llu (%u)", uintptr_t(hFile), FileInformationClass);
-		return True_SetFileInformationByHandle(hFile, FileInformationClass, lpFileInformation, dwBufferSize); // In here to be tabbed in log
+		if (MemoryFile* mf = fi.memoryFile)
+		{
+			auto& info = *(FILE_ALLOCATION_INFO*)lpFileInformation;
+			mf->EnsureCommitted(dh, info.AllocationSize.QuadPart);
+			return true;
+		}
+	}
+	else if (FileInformationClass == FileEndOfFileInfo)
+	{
+		if (MemoryFile* mf = fi.memoryFile)
+		{
+			auto& info = *(FILE_END_OF_FILE_INFO*)lpFileInformation;
+			mf->writtenSize = info.EndOfFile.QuadPart;
+			mf->isReported = false;
+			mf->EnsureCommitted(dh, mf->writtenSize);
+			SetLastError(ERROR_SUCCESS);
+			return true;
+		}
 	}
+
+	DEBUG_LOG_TRUE(L"SetFileInformationByHandle", L"%llu (%u)", uintptr_t(hFile), FileInformationClass);
+	return True_SetFileInformationByHandle(hFile, FileInformationClass, lpFileInformation, dwBufferSize); // In here to be tabbed in log
 }
 
 HANDLE Detoured_CreateFileMappingW(HANDLE hFile, LPSECURITY_ATTRIBUTES lpFileMappingAttributes, DWORD flProtect, DWORD dwMaximumSizeHigh, DWORD dwMaximumSizeLow, LPCWSTR lpName)
