[UBA]
* Added some files as output for radlink.exe application rules
* Added so .lib files are cached in memory maps

diff --git a/Engine/Source/Programs/UnrealBuildAccelerator/Common/Private/UbaSession.cpp b/Engine/Source/Programs/UnrealBuildAccelerator/Common/Private/UbaSession.cpp
--- a/Engine/Source/Programs/UnrealBuildAccelerator/Common/Private/UbaSession.cpp
+++ b/Engine/Source/Programs/UnrealBuildAccelerator/Common/Private/UbaSession.cpp
@@ -2071,6 +2071,8 @@ namespace uba
 			return 4 * 1024; // clang seems to need 4k alignment? Is it a coincidence it works or what is happening inside the code? (msvc works with alignment 1byte here)
 		if (fileName.EndsWith(TC(".h.obj")))
 			return 4 * 1024;
+		if (fileName.EndsWith(TC(".lib")))
+			return 4 * 1024;
 		//if (fileName.EndsWith(TC(".lib")))
 		//	return 4 * 1024;
 		//if (fileName.EndsWith(TC(".rc2.res")))
diff --git a/Engine/Source/Programs/UnrealBuildAccelerator/Core/Private/UbaApplicationRules.cpp b/Engine/Source/Programs/UnrealBuildAccelerator/Core/Private/UbaApplicationRules.cpp
--- a/Engine/Source/Programs/UnrealBuildAccelerator/Core/Private/UbaApplicationRules.cpp
+++ b/Engine/Source/Programs/UnrealBuildAccelerator/Core/Private/UbaApplicationRules.cpp
@@ -241,7 +241,7 @@ namespace uba
 
 		virtual u64 FileTypeMaxSize(const StringView& file, bool isSystemOrTempFile) const override
 		{
-			if (file.Contains(TC(".pdb.tmp")))
+			if (file.Contains(TC(".pdb.tmp")) || file.EndsWith(TC(".pdb")))
 				return 14ull * 1024 * 1024 * 1024; // This is ridiculous
 			return ApplicationRulesVcLink::FileTypeMaxSize(file, isSystemOrTempFile);
 		}
@@ -249,6 +249,9 @@ namespace uba
 		virtual bool IsOutputFile(const StringView& fileName) const override
 		{
 			return false//Super::IsOutputFile(fileName)
+				|| fileName.EndsWith(TC(".exe"))
+				|| fileName.EndsWith(TC(".dll"))
+				|| fileName.EndsWith(TC(".pdb"))
 				|| fileName.Contains(TC(".exe.tmp"))
 				|| fileName.Contains(TC(".dll.tmp"))
 				|| fileName.Contains(TC(".pdb.tmp"));
diff --git a/Engine/Source/Programs/UnrealBuildAccelerator/Detours/Private/Windows/UbaDetoursFunctionsKernelBase.inl b/Engine/Source/Programs/UnrealBuildAccelerator/Detours/Private/Windows/UbaDetoursFunctionsKernelBase.inl
--- a/Engine/Source/Programs/UnrealBuildAccelerator/Detours/Private/Windows/UbaDetoursFunctionsKernelBase.inl
+++ b/Engine/Source/Programs/UnrealBuildAccelerator/Detours/Private/Windows/UbaDetoursFunctionsKernelBase.inl
@@ -2317,6 +2317,7 @@ BOOL Detoured_SetFileInformationByHandle(HANDLE hFile, FILE_INFO_BY_HANDLE_CLASS
 		{
 			auto& info = *(FILE_ALLOCATION_INFO*)lpFileInformation;
 			mf->EnsureCommitted(dh, info.AllocationSize.QuadPart);
+			DEBUG_LOG_TRUE(L"SetFileInformationByHandle", L"%llu (FileAllocationInfo) Size: %llu", uintptr_t(hFile), info.AllocationSize.QuadPart);
 			return true;
 		}
 	}
@@ -2329,6 +2330,7 @@ BOOL Detoured_SetFileInformationByHandle(HANDLE hFile, FILE_INFO_BY_HANDLE_CLASS
 			mf->isReported = false;
 			mf->EnsureCommitted(dh, mf->writtenSize);
 			SetLastError(ERROR_SUCCESS);
+			DEBUG_LOG_TRUE(L"SetFileInformationByHandle", L"%llu (FileEndOfFileInfo) End: %llu", uintptr_t(hFile), info.EndOfFile.QuadPart);
 			return true;
 		}
 	}
