[UBA]
* Added default value to TrackWorkScope ctor for work color

diff --git a/Engine/Source/Programs/UnrealBuildAccelerator/Common/Public/UbaWorkManager.h b/Engine/Source/Programs/UnrealBuildAccelerator/Common/Public/UbaWorkManager.h
--- a/Engine/Source/Programs/UnrealBuildAccelerator/Common/Public/UbaWorkManager.h
+++ b/Engine/Source/Programs/UnrealBuildAccelerator/Common/Public/UbaWorkManager.h
@@ -81,7 +81,7 @@ namespace uba
 	struct TrackWorkScope
 	{
 		TrackWorkScope() : tracker(nullptr), id(0) {}
-		TrackWorkScope(WorkTracker& t, const StringView& desc, const Color& color) : tracker(&t), id(t.TrackWorkStart(desc, color)) {}
+		TrackWorkScope(WorkTracker& t, const StringView& desc, const Color& color = ColorWork) : tracker(&t), id(t.TrackWorkStart(desc, color)) {}
 		void AddHint(const StringView& hint) { if (tracker) tracker->TrackWorkHint(id, hint); }
 		~TrackWorkScope() { if (tracker) tracker->TrackWorkEnd(id); }
 		WorkTracker* tracker;
