[UBA]
* Added code for transferring ownership of entire mapped file from detoured process to host. Disabled by default behind define UBA_ENABLE_ON_DISK_FILE_MAPPINGS but can be enabled for experimentation

diff --git a/Engine/Source/Programs/UnrealBuildAccelerator/Common/Private/UbaSession.cpp b/Engine/Source/Programs/UnrealBuildAccelerator/Common/Private/UbaSession.cpp
--- a/Engine/Source/Programs/UnrealBuildAccelerator/Common/Private/UbaSession.cpp
+++ b/Engine/Source/Programs/UnrealBuildAccelerator/Common/Private/UbaSession.cpp
@@ -1801,7 +1801,7 @@ namespace uba
 			u8 id = reader.ReadByte();
 			reader.SkipString();
 			StringBuffer<> rootPath;
-			reader.ReadString(rootPath.Clear());
+			reader.ReadString(rootPath);
 			if (!out.RegisterRoot(m_logger, rootPath.data, true, id))
 				return false;
 		}
@@ -2413,7 +2413,7 @@ namespace uba
 				writtenFile.name = msgName;
 				FileMappingHandle mappingHandle;
 				if (msg.mappingHandle.IsValid())
-					if (!DuplicateFileMapping(m_logger, msg.process.m_nativeProcessHandle, msg.mappingHandle, GetCurrentProcessHandle(), mappingHandle, FILE_MAP_READ, false, 0, msgName))
+					if (!DuplicateFileMapping(m_logger, msg.process.m_nativeProcessHandle, msg.mappingHandle, GetCurrentProcessHandle(), mappingHandle, FILE_MAP_READ, false, DUPLICATE_SAME_ACCESS, msgName))
 						return m_logger.Error(TC("Failed to duplicate file mapping handle for %s"), name);
 
 				writtenFile.key = key;
@@ -2584,6 +2584,69 @@ namespace uba
 
 			if (ShouldWriteToDisk(file.name))
 			{
+				auto& rules = *process.m_startInfo.rules;
+				RootsHandle rootsHandle = process.GetStartInfo().rootsHandle;
+
+				bool storeCompressed = m_storeObjFilesCompressed && rules.StoreFileCompressed(file.name);
+				bool shouldDevirtualize = false;
+				bool escapeSpaces = false;
+				if (!storeCompressed)
+					shouldDevirtualize = HasVfs(rootsHandle) && rules.ShouldDevirtualizeFile(file.name, escapeSpaces);
+
+#if UBA_ENABLE_ON_DISK_FILE_MAPPINGS
+				if (!storeCompressed && !shouldDevirtualize && file.mappingHandle.fh)
+				{
+					bool success = true;
+					SetFilePointerEx(file.mappingHandle.fh, ToLargeInteger(file.mappingWritten), NULL, FILE_BEGIN);
+					::SetEndOfFile(file.mappingHandle.fh);
+
+					if (false)
+					{
+						TrackWorkScope work(*m_workManager, AsView(TC("Flush")));
+						work.AddHint(file.name);
+						FlushFileBuffers(file.mappingHandle.fh);
+					}
+
+					if (true)
+					{
+						FILE_DISPOSITION_INFO info;
+						info.DeleteFile = false;
+						if (!::SetFileInformationByHandle(file.mappingHandle.fh, FileDispositionInfo, &info, sizeof(info)))
+							m_logger.Warning(TC("Failed to remove delete-on-close"));
+					}
+
+					if (false)
+					{
+						StringBuffer<> name2;
+						name2.Append(TC("\\??\\")).Append(file.name);
+
+						u8 buffer[2048];
+						auto& info = *(FILE_RENAME_INFO*)buffer;;
+						info.ReplaceIfExists = true;
+						info.RootDirectory = 0;
+						info.FileNameLength = u32(name2.count*sizeof(tchar));
+						info.Flags = FILE_RENAME_FLAG_POSIX_SEMANTICS;
+						memcpy(info.FileName, name2.data, info.FileNameLength+sizeof(tchar));
+						TrackWorkScope work(*m_workManager, AsView(TC("Rename")));
+						work.AddHint(file.name);
+						success = SetFileInformationByHandle(file.mappingHandle.fh, FileRenameInfo, &info, sizeof(buffer));
+					}
+
+					if (success)
+					{
+						{
+							TrackWorkScope work(*m_workManager, AsView(TC("CloseFileMapping")));
+							work.AddHint(file.name);
+							CloseFileMapping(m_logger, file.mappingHandle, file.name.c_str());
+						}
+						file.mappingHandle = {};
+						m_storage.InvalidateCachedFileInfo(file.key);
+						return true;
+					}
+					m_logger.Warning(TC("SetFileInformationByHandle failed %s (%s)"), file.name.c_str(), LastErrorToText().data);
+				}
+#endif
+
 				// This is to kill I/O when writing lots of pdb/dlls in parallel
 				#if PLATFORM_WINDOWS
 				constexpr u32 bottleneckMax = 32;
@@ -2606,9 +2669,7 @@ namespace uba
 
 				auto memClose = MakeGuard([&](){ UnmapViewOfFile(m_logger, mem, fileSize, file.name.c_str()); });
 
-				auto& rules = *process.m_startInfo.rules;
-
-				if (m_storeObjFilesCompressed && rules.StoreFileCompressed(file.name))
+				if (storeCompressed)
 				{
 					Storage::WriteResult res;
 					CompressedObjFileHeader header { CalculateCasKey(mem, fileSize, true, m_workManager, file.name.c_str()) };
@@ -2622,9 +2683,7 @@ namespace uba
 				{
 					FileAccessor destinationFile(m_logger, file.name.c_str());
 
-					RootsHandle rootsHandle = process.GetStartInfo().rootsHandle;
-					bool escapeSpaces;
-					if (HasVfs(rootsHandle) && rules.ShouldDevirtualizeFile(file.name, escapeSpaces))
+					if (shouldDevirtualize)
 					{
 						// Need to turn paths back into local paths
 
@@ -2697,7 +2756,7 @@ namespace uba
 				m_workManager->AddWork([mh = file.mappingHandle, n = file.name, this](const WorkContext&)
 					{
 						CloseFileMapping(m_logger, mh, n.c_str());
-					}, 1, TC("ClosFileMapping"));
+					}, 1, TC("CloseFileMapping"));
 			}
 			else
 			{
diff --git a/Engine/Source/Programs/UnrealBuildAccelerator/Core/Public/UbaProtocol.h b/Engine/Source/Programs/UnrealBuildAccelerator/Core/Public/UbaProtocol.h
--- a/Engine/Source/Programs/UnrealBuildAccelerator/Core/Public/UbaProtocol.h
+++ b/Engine/Source/Programs/UnrealBuildAccelerator/Core/Public/UbaProtocol.h
@@ -56,6 +56,8 @@ namespace uba
 	inline constexpr u32 DirTableMemSize = 128 * 1024 * 1024;
 }
 
+#define UBA_ENABLE_ON_DISK_FILE_MAPPINGS 0
+
 // Currently only used for detoured process
 #if UBA_DEBUG
 #define UBA_DEBUG_LOG_ENABLED 1
diff --git a/Engine/Source/Programs/UnrealBuildAccelerator/Detours/Private/Windows/UbaDetoursFunctionsNtDll.inl b/Engine/Source/Programs/UnrealBuildAccelerator/Detours/Private/Windows/UbaDetoursFunctionsNtDll.inl
--- a/Engine/Source/Programs/UnrealBuildAccelerator/Detours/Private/Windows/UbaDetoursFunctionsNtDll.inl
+++ b/Engine/Source/Programs/UnrealBuildAccelerator/Detours/Private/Windows/UbaDetoursFunctionsNtDll.inl
@@ -1160,6 +1160,7 @@ NTSTATUS NTAPI Shared_NtCreateFile(bool IsCreateFunc, PHANDLE hFileHandle, ACCES
 						True_DuplicateHandle(g_hostProcess, mappingHandle.mh, GetCurrentProcess(), &newHandle, 0, FALSE, DUPLICATE_SAME_ACCESS);
 						UBA_ASSERTF(newHandle, L"DuplicateHandle failed when opening temp file %ls (%u)", fileName.data, GetLastError());
 						mf.writtenSize = mappingHandleSize;
+						mf.committedSize = mappingHandleSize;
 						TimerScope ts2(g_kernelStats.mapViewOfFile);
 						mf.baseAddress = (u8*)True_MapViewOfFile(newHandle, FILE_MAP_READ | FILE_MAP_ALL_ACCESS, 0, 0, mappingHandleSize);
 						UBA_ASSERTF(mf.baseAddress, L"MapViewOfFile failed when opening temp file %ls (%u)", fileName.data, GetLastError());
@@ -1168,7 +1169,7 @@ NTSTATUS NTAPI Shared_NtCreateFile(bool IsCreateFunc, PHANDLE hFileHandle, ACCES
 					}
 					else if (isWrite)
 					{
-						info.memoryFile = new MemoryFile(false, FileTypeMaxSize(fileName, isSystemOrTempFile));
+						info.memoryFile = new MemoryFile(false, FileTypeMaxSize(fileName, isSystemOrTempFile), false, fileName.data);
 					}
 					else
 					{
@@ -1189,7 +1190,7 @@ NTSTATUS NTAPI Shared_NtCreateFile(bool IsCreateFunc, PHANDLE hFileHandle, ACCES
 
 				bool isLocal = !IsOutputFile(fileName, isWrite, isDeleteOnClose);
 				//UBA_ASSERTF(CreateDisposition != FILE_OPEN || Contains(fileName.data, L"vctip_"), TC("Unsupported disposition %u for file %s"), CreateDisposition, fileName.data);
-				info.memoryFile = new MemoryFile(isLocal, FileTypeMaxSize(fileName, isSystemOrTempFile), isThrowAway);
+				info.memoryFile = new MemoryFile(isLocal, FileTypeMaxSize(fileName, isSystemOrTempFile), isThrowAway, fileName.data);
 			}
 
 			// TODO: Time should be in sync with host machine!
diff --git a/Engine/Source/Programs/UnrealBuildAccelerator/Detours/Private/Windows/UbaDetoursFunctionsWin.cpp b/Engine/Source/Programs/UnrealBuildAccelerator/Detours/Private/Windows/UbaDetoursFunctionsWin.cpp
--- a/Engine/Source/Programs/UnrealBuildAccelerator/Detours/Private/Windows/UbaDetoursFunctionsWin.cpp
+++ b/Engine/Source/Programs/UnrealBuildAccelerator/Detours/Private/Windows/UbaDetoursFunctionsWin.cpp
@@ -25,6 +25,8 @@
 #include <malloc.h>
 
 NTSTATUS NtCopyFileChunk(HANDLE Source, HANDLE Dest, HANDLE Event, PIO_STATUS_BLOCK IoStatusBlock, ULONG Length, PULONG SourceOffset, PULONG DestOffset, PULONG SourceKey, PULONG DestKey, ULONG Flags) { return 0; }
+#define NtCurrentProcess() ((HANDLE)-1)
+
 #define DETOURED_FUNCTION(Func) decltype(Func)* True_##Func = ::Func; 
 DETOURED_FUNCTIONS
 DETOURED_FUNCTIONS_MEMORY
@@ -110,6 +112,7 @@ DETOURED_FUNCTIONS_MEMORY_NON_WINE
 bool g_isDetachedProcess;
 bool g_isRunningWine;
 int g_uiLanguage;
+u32 g_processId;
 
 StringBuffer<256> g_exeDir;
 
@@ -251,15 +254,15 @@ MemoryFile::MemoryFile(u8* data, bool localOnly)
 {
 }
 
-MemoryFile::MemoryFile(bool localOnly, u64 reserveSize_, bool isThrowAway_)
+MemoryFile::MemoryFile(bool localOnly, u64 reserveSize_, bool isThrowAway_, const tchar* fileName)
 :	isLocalOnly(localOnly)
 ,	isThrowAway(isThrowAway_)
 {
 	if (!isThrowAway)
-		Reserve(reserveSize_);
+		Reserve(reserveSize_, fileName);
 }
 
-void MemoryFile::Reserve(u64 reserveSize_)
+void MemoryFile::Reserve(u64 reserveSize_, const tchar* fileName)
 {
 	UBA_ASSERT(!isThrowAway);
 	reserveSize = reserveSize_;
@@ -270,6 +273,50 @@ void MemoryFile::Reserve(u64 reserveSize_)
 			FatalError(1354, L"VirtualAlloc failed trying to reserve %llu. (Error code: %u)", reserveSize, GetLastError());
 		mappedSize = reserveSize;
 	}
+#if UBA_ENABLE_ON_DISK_FILE_MAPPINGS
+	else if (fileName && !g_runningRemote)
+	{
+		StringBuffer<> tempFileName;
+		tempFileName.Appendf(L"\\??\\").Append(fileName);
+		if (false)
+			tempFileName.Append(L".uba.tmp");
+
+		UNICODE_STRING us;
+		us.Length = USHORT(tempFileName.count*sizeof(tchar));
+		us.MaximumLength = us.Length;
+		us.Buffer = tempFileName.data;
+
+		SuppressDetourScope _;
+		{
+			TimerScope ts(g_kernelStats.createFile);
+			OBJECT_ATTRIBUTES oa;
+			InitializeObjectAttributes(&oa, &us, OBJ_CASE_INSENSITIVE, NULL, NULL);
+			IO_STATUS_BLOCK iosb;
+			ACCESS_MASK desiredAccess = GENERIC_READ | GENERIC_WRITE | DELETE | SYNCHRONIZE | FILE_WRITE_ATTRIBUTES;
+			ULONG shareAccess = FILE_SHARE_READ | FILE_SHARE_WRITE | FILE_SHARE_DELETE;
+			ULONG createOptions = FILE_SYNCHRONOUS_IO_NONALERT;// | FILE_DELETE_ON_CLOSE;
+			NTSTATUS status = ZwCreateFile(&mappingHandle.fh, desiredAccess, &oa, &iosb, 0, FILE_ATTRIBUTE_NORMAL, shareAccess, FILE_OVERWRITE_IF, createOptions, 0, 0);
+			if (status != STATUS_SUCCESS)
+				FatalError(1347, L"ZwCreateFile failed to create %s. (Error code: 0x%x)", us.Buffer, status);
+		}
+		FILE_DISPOSITION_INFO info;
+		info.DeleteFile = true;
+		if (!::SetFileInformationByHandle(mappingHandle.fh, FileDispositionInfo, &info, sizeof(info)))
+			FatalError(1347, L"SetFileInformationByHandle failed to set delete-on-close on %s. (Error code: 0x%x)", us.Buffer, GetLastError());
+		{
+			TimerScope ts(g_kernelStats.createFileMapping);
+			LARGE_INTEGER liMaxSize; liMaxSize.QuadPart = 2;
+			NTSTATUS status = ZwCreateSection(&mappingHandle.mh, SECTION_ALL_ACCESS, 0, &liMaxSize, PAGE_READWRITE, SEC_COMMIT, mappingHandle.fh);
+			if (status != STATUS_SUCCESS)
+				FatalError(1348, L"NtCreateSection failed to reserve %llu. (Error code: 0x%x)", reserveSize, status);
+		}
+		TimerScope ts(g_kernelStats.mapViewOfFile);
+		NTSTATUS status = ZwMapViewOfSection(mappingHandle.mh, NtCurrentProcess(), (void**)&baseAddress, 0, 0, NULL, &reserveSize, 2, MEM_RESERVE, PAGE_READWRITE);
+		if (status != STATUS_SUCCESS)
+			FatalError(1348, L"ZwMapViewOfSection failed trying to reserve %llu. (Error code: 0x%x)", reserveSize, status);
+		mappedSize = reserveSize;
+	}
+#endif
 	else
 	{
 		HANDLE handle;
@@ -298,8 +345,7 @@ void MemoryFile::Unreserve()
 	{
 		True_UnmapViewOfFile(baseAddress);
 		CloseHandle(mappingHandle.mh);
-		UBA_ASSERT(!mappingHandle.fh);
-		mappingHandle.mh = 0;
+		mappingHandle.mh = nullptr;
 	}
 	baseAddress = nullptr;
 	committedSize = 0;
@@ -354,14 +400,29 @@ void MemoryFile::EnsureCommitted(const DetouredHandle& handle, u64 size)
 			Remap(handle, size);
 	}
 
-	u64 toCommit = Min(reserveSize, AlignUp(size - committedSize, g_pageSize));
-	if (!VirtualAlloc(baseAddress + committedSize, toCommit, MEM_COMMIT, PAGE_READWRITE))
-		FatalError(1347, L"Failed to ensure virtual memory for %ls. MappedSize: %llu, CommittedSize: %llu RequestedSize: %llu. (%u)", HandleToName(handle), mappedSize, committedSize, size, GetLastError());
-	committedSize += toCommit;
+#if UBA_ENABLE_ON_DISK_FILE_MAPPINGS
+	if (mappingHandle.fh)
+	{
+		if (!committedSize)
+			committedSize = size;
+		else
+			committedSize = Min(reserveSize, Max(size, committedSize * 2));
+		auto lg = ToLargeInteger(committedSize);
+		NtExtendSection(mappingHandle.mh, &lg);
+	}
+	else
+#endif
+	{
+		u64 toCommit = Min(reserveSize, AlignUp(size - committedSize, g_pageSize));
+		if (!VirtualAlloc(baseAddress + committedSize, toCommit, MEM_COMMIT, PAGE_READWRITE))
+			FatalError(1347, L"Failed to ensure virtual memory for %ls. MappedSize: %llu, CommittedSize: %llu RequestedSize: %llu. (%u)", HandleToName(handle), mappedSize, committedSize, size, GetLastError());
+		committedSize += toCommit;
+	}
 }
 
 void MemoryFile::Remap(const DetouredHandle& handle, u64 size)
 {
+	UBA_ASSERT(!mappingHandle.fh);
 	True_UnmapViewOfFile(baseAddress);
 	mappedSize = Min(reserveSize, AlignUp(Max(size, mappedSize * 4), g_pageSize));
 	TimerScope ts(g_kernelStats.mapViewOfFile);
@@ -528,7 +589,7 @@ void CloseCaches()
 		{
 			UnmapViewOfFile(info.memoryFile->baseAddress);
 			CloseHandle(info.memoryFile->mappingHandle.mh);
-			UBA_ASSERT(!info.memoryFile->mappingHandle.fh);
+			CloseHandle(info.memoryFile->mappingHandle.fh);
 		}
 		//if (info.memoryFile && info.memoryFile != &g_emptyMemoryFile)
 		//	delete info.memoryFile;
@@ -895,6 +956,8 @@ void Init(const DetoursPayload& payload, u64 startTime)
 		g_stdHandle[2] = makeDetouredHandle(new DetouredHandle(HandleType_StdIn)); // STD_IN
 	}
 
+	g_processId = payload.processId;
+
 	if (payload.trackInputs)
 		g_trackInputsMem = (u8*)g_memoryBlock.Allocate(TrackInputsMemCapacity, 1, L"TrackInputs");
 	
diff --git a/Engine/Source/Programs/UnrealBuildAccelerator/Detours/Private/Windows/UbaDetoursFunctionsWin.h b/Engine/Source/Programs/UnrealBuildAccelerator/Detours/Private/Windows/UbaDetoursFunctionsWin.h
--- a/Engine/Source/Programs/UnrealBuildAccelerator/Detours/Private/Windows/UbaDetoursFunctionsWin.h
+++ b/Engine/Source/Programs/UnrealBuildAccelerator/Detours/Private/Windows/UbaDetoursFunctionsWin.h
@@ -400,8 +400,12 @@ extern "C" {
 	NTSTATUS NTAPI NtAlpcCreatePortSection(HANDLE PortHandle, ULONG Flags, HANDLE SectionHandle, SIZE_T SectionSize, PHANDLE AlpcSectionHandle, PSIZE_T ActualSectionSize);
 	NTSTATUS NTAPI NtAlpcSendWaitReceivePort(HANDLE PortHandle, DWORD Flags, PPORT_MESSAGE SendMessage_, PALPC_MESSAGE_ATTRIBUTES SendMessageAttributes, PPORT_MESSAGE ReceiveMessage, PSIZE_T BufferLength, PALPC_MESSAGE_ATTRIBUTES ReceiveMessageAttributes, PLARGE_INTEGER Timeout);
 	NTSTATUS NTAPI NtAlpcDisconnectPort(HANDLE PortHandle, ULONG Flags);
+	NTSTATUS NTAPI NtExtendSection(HANDLE, PLARGE_INTEGER);
 	NTSTATUS NTAPI ZwCreateFile(PHANDLE FileHandle, ACCESS_MASK DesiredAccess, POBJECT_ATTRIBUTES ObjectAttributes, PIO_STATUS_BLOCK IoStatusBlock, PLARGE_INTEGER AllocationSize, ULONG FileAttributes, ULONG ShareAccess, ULONG CreateDisposition, ULONG CreateOptions, PVOID EaBuffer, ULONG EaLength);
 	NTSTATUS NTAPI ZwOpenFile(PHANDLE FileHandle, ACCESS_MASK DesiredAccess, POBJECT_ATTRIBUTES ObjectAttributes, PIO_STATUS_BLOCK IoStatusBlock, ULONG ShareAccess, ULONG OpenOptions);
+	NTSTATUS NTAPI ZwClose(HANDLE Handle);
+	NTSTATUS NTAPI ZwMapViewOfSection(HANDLE SectionHandle, HANDLE ProcessHandle, PVOID* BaseAddress, ULONG_PTR ZeroBits, SIZE_T CommitSize, PLARGE_INTEGER SectionOffset, PSIZE_T ViewSize, DWORD InheritDisposition, ULONG AllocationType, ULONG Win32Protect);
+	NTSTATUS NTAPI ZwCreateSection(PHANDLE SectionHandle, ACCESS_MASK DesiredAccess, POBJECT_ATTRIBUTES ObjectAttributes, PLARGE_INTEGER MaximumSize, ULONG SectionPageProtection, ULONG AllocationAttributes, HANDLE FileHandle);
 	NTSTATUS NTAPI ZwQueryDirectoryFile(HANDLE FileHandle, HANDLE Event, PIO_APC_ROUTINE ApcRoutine, PVOID ApcContext, PIO_STATUS_BLOCK IoStatusBlock, PVOID FileInformation, ULONG Length, FILE_INFORMATION_CLASS FileInformationClass, BOOLEAN ReturnSingleEntry, PUNICODE_STRING FileName, BOOLEAN RestartScan);
 	NTSTATUS NTAPI ZwSetInformationFile(HANDLE FileHandle,PIO_STATUS_BLOCK IoStatusBlock, PVOID FileInformation, ULONG Length, FILE_INFORMATION_CLASS FileInformationClass);
 	PVOID WINAPI ResolveDelayLoadedAPI(PVOID ParentModuleBase, PCIMAGE_DELAYLOAD_DESCRIPTOR DelayloadDescriptor, void* FailureDllHook, void* FailureSystemHook, PIMAGE_THUNK_DATA ThunkAddress, ULONG Flags);
diff --git a/Engine/Source/Programs/UnrealBuildAccelerator/Detours/Private/Windows/UbaDetoursUtilsWin.h b/Engine/Source/Programs/UnrealBuildAccelerator/Detours/Private/Windows/UbaDetoursUtilsWin.h
--- a/Engine/Source/Programs/UnrealBuildAccelerator/Detours/Private/Windows/UbaDetoursUtilsWin.h
+++ b/Engine/Source/Programs/UnrealBuildAccelerator/Detours/Private/Windows/UbaDetoursUtilsWin.h
@@ -61,9 +61,9 @@ namespace uba
 	struct MemoryFile
 	{
 		MemoryFile(u8* data = nullptr, bool localOnly = true);
-		MemoryFile(bool localOnly, u64 reserveSize_, bool isThrowAway_ = false);
+		MemoryFile(bool localOnly, u64 reserveSize_, bool isThrowAway_ = false, const tchar* fileName = nullptr);
 
-		void Reserve(u64 reserveSize_);
+		void Reserve(u64 reserveSize_, const tchar* fileName = nullptr);
 		void Unreserve();
 		void Write(struct DetouredHandle& handle, LPCVOID lpBuffer, u64 nNumberOfBytesToWrite);
 		void EnsureCommitted(const struct DetouredHandle& handle, u64 size);
