[UBA]
* Reduced stack footprint when writing to cache

diff --git a/Engine/Source/Programs/UnrealBuildAccelerator/Common/Private/UbaCacheClient.cpp b/Engine/Source/Programs/UnrealBuildAccelerator/Common/Private/UbaCacheClient.cpp
--- a/Engine/Source/Programs/UnrealBuildAccelerator/Common/Private/UbaCacheClient.cpp
+++ b/Engine/Source/Programs/UnrealBuildAccelerator/Common/Private/UbaCacheClient.cpp
@@ -205,6 +205,7 @@ namespace uba
 			}
 
 			// For .exe and .dll we sometimes get relative paths so we need to expand them to full
+			#if PLATFORM_WINDOWS
 			if (path[1] != ':' && (path.EndsWith(TC(".dll")) || path.EndsWith(TC(".exe"))))
 			{
 				tchar temp[512];
@@ -216,7 +217,9 @@ namespace uba
 					return false;
 				}
 			}
-			else if (ShouldNormalize(path)) // Paths can be absolute in rsp files so we need to normalize those paths
+			else
+			#endif
+			if (ShouldNormalize(path)) // Paths can be absolute in rsp files so we need to normalize those paths
 			{
 				casKey = rootPaths.NormalizeAndHashFile(m_logger, path.data);
 				if (casKey == CasKeyZero)
@@ -986,8 +989,8 @@ namespace uba
 		auto doneGuard = MakeGuard([&]()
 			{
 				// Send done.. confirm to server
-				StackBinaryWriter<SendMaxSize> writer;
-				NetworkMessage msg(m_client, CacheServiceId, CacheMessageType_StoreEntryDone, writer);
+				StackBinaryWriter<64> writer;
+				NetworkMessage msg(m_client, CacheServiceId, CacheMessageType_StoreEntryDone, writer.Reset());
 				writer.Write7BitEncoded(MakeId(bucket.id));
 				writer.WriteCasKey(cmdKey);
 				writer.WriteBool(success);
diff --git a/Engine/Source/Programs/UnrealBuildAccelerator/Common/Public/UbaCacheClient.h b/Engine/Source/Programs/UnrealBuildAccelerator/Common/Public/UbaCacheClient.h
--- a/Engine/Source/Programs/UnrealBuildAccelerator/Common/Public/UbaCacheClient.h
+++ b/Engine/Source/Programs/UnrealBuildAccelerator/Common/Public/UbaCacheClient.h
@@ -71,9 +71,10 @@ namespace uba
 		struct Bucket;
 		u64 MakeId(u32 bucketId);
 
-		bool SendPathTable(Bucket& bucket, u32 requiredPathTableSize);
-		bool SendCasTable(Bucket& bucket, u32 requiredCasTableSize);
-		bool SendCacheEntry(Bucket& bucket, const RootPaths& rootPaths, const CasKey& cmdKey, const Map<u32, u32>& inputsStringToCasKey, const Map<u32, u32>& outputsStringToCasKey, const u8* logLines, u64 logLinesSize, u64& outBytesSent);
+		// No inlining because they use lots of stack space
+		UBA_NOINLINE bool SendPathTable(Bucket& bucket, u32 requiredPathTableSize);
+		UBA_NOINLINE bool SendCasTable(Bucket& bucket, u32 requiredCasTableSize);
+		UBA_NOINLINE bool SendCacheEntry(Bucket& bucket, const RootPaths& rootPaths, const CasKey& cmdKey, const Map<u32, u32>& inputsStringToCasKey, const Map<u32, u32>& outputsStringToCasKey, const u8* logLines, u64 logLinesSize, u64& outBytesSent);
 		bool FetchCasTable(Bucket& bucket, CacheStats& stats, u32 requiredCasTableOffset);
 		bool FetchFile(Bucket& bucket, const RootPaths& rootPaths, const ProcessStartInfo& info, CacheStats& cacheStats, StorageStats& storageStats, u32 casKeyOffset);
 		template<typename TableType> bool FetchCompactTable(u32 bucketId, TableType& table, u32 requiredTableOffset, u8 messageType);
diff --git a/Engine/Source/Programs/UnrealBuildAccelerator/Core/Public/UbaBinaryReaderWriter.h b/Engine/Source/Programs/UnrealBuildAccelerator/Core/Public/UbaBinaryReaderWriter.h
--- a/Engine/Source/Programs/UnrealBuildAccelerator/Core/Public/UbaBinaryReaderWriter.h
+++ b/Engine/Source/Programs/UnrealBuildAccelerator/Core/Public/UbaBinaryReaderWriter.h
@@ -47,7 +47,7 @@ namespace uba
 		void Flush(bool waitOnResponse = true);
 		BinaryWriter();
 
-	private:
+	protected:
 		u8* m_begin;
 		u8* m_pos;
 		u8* m_end;
@@ -57,6 +57,7 @@ namespace uba
 	struct StackBinaryWriter : BinaryWriter
 	{
 		StackBinaryWriter() : BinaryWriter(buffer, 0, Capacity) {}
+		StackBinaryWriter& Reset() { m_pos = m_begin; m_end = buffer + Capacity; return *this; }
 		u8 buffer[Capacity];
 	};
 
