[UBA]
* Fixed bug where cache client FetchCasTable could write to memory that was not committed
* Renamed variable in MemoryBlock from mapped to committed since that is what it is

diff --git a/Engine/Source/Programs/UnrealBuildAccelerator/Common/Private/UbaCacheClient.cpp b/Engine/Source/Programs/UnrealBuildAccelerator/Common/Private/UbaCacheClient.cpp
--- a/Engine/Source/Programs/UnrealBuildAccelerator/Common/Private/UbaCacheClient.cpp
+++ b/Engine/Source/Programs/UnrealBuildAccelerator/Common/Private/UbaCacheClient.cpp
@@ -1104,6 +1104,7 @@ namespace uba
 		return doneGuard.Execute();
 	}
 
+#pragma optimize("", off)
 	bool CacheClient::FetchCasTable(Bucket& bucket, CacheStats& stats, u32 requiredCasTableOffset)
 	{
 		if (HasEnoughData(bucket, requiredCasTableOffset))
@@ -1143,8 +1144,8 @@ namespace uba
 		u32 messageFetchSize = u32(m_client.GetMessageMaxSize() - m_client.GetMessageReceiveHeaderSize());
 		u32 totalFetchSize = requiredTableSize - tableSize;
 		u32 messageCount = (totalFetchSize + messageFetchSize - 1) / messageFetchSize;
-
-		u8* data = table.BeginCommit(totalFetchSize);
+		u32 commitSize = messageCount * SendMaxSize;
+		u8* data = table.BeginCommit(commitSize);
 
 		struct Entry
 		{
@@ -1181,6 +1182,8 @@ namespace uba
 		return true;
 	}
 
+#pragma optimize("", on)
+
 	bool CacheClient::HasEnoughData(Bucket& bucket, u32 requiredCasTableOffset)
 	{
 		u32 pathOffset;
diff --git a/Engine/Source/Programs/UnrealBuildAccelerator/Common/Private/UbaCompactTables.cpp b/Engine/Source/Programs/UnrealBuildAccelerator/Common/Private/UbaCompactTables.cpp
--- a/Engine/Source/Programs/UnrealBuildAccelerator/Common/Private/UbaCompactTables.cpp
+++ b/Engine/Source/Programs/UnrealBuildAccelerator/Common/Private/UbaCompactTables.cpp
@@ -598,8 +598,7 @@ namespace uba
 	{
 		if (!m_pathTableMem.memory)
 			m_pathTableMem.Init(PathTableMaxSize);
-		m_pathTableMem.ReserveNoLock(size, TC("CompactPathTable::BeginCommit"));
-		return (u8*)m_pathTableMem.memory + m_pathTableMem.writtenSize;
+		return (u8*)m_pathTableMem.CommitNoLock(size, TC("CompactPathTable::BeginCommit"));
 	}
 
 	void CompactPathTable::EndCommit(u8* data, u64 written)
@@ -844,8 +843,7 @@ namespace uba
 			m_casKeyTableMem.Init(CasKeyTableMaxSize);
 		if (!m_offsets.IsInitialized())
 			m_offsets.Init(CasKeyLookupInitialCount);
-		m_casKeyTableMem.ReserveNoLock(size, TC("CompactCasKeyTable::BeginCommit"));
-		return (u8*)m_casKeyTableMem.memory + m_casKeyTableMem.writtenSize;
+		return (u8*)m_casKeyTableMem.CommitNoLock(size, TC("CompactCasKeyTable::BeginCommit"));
 	}
 
 	void CompactCasKeyTable::EndCommit(u8* data, u64 written)
diff --git a/Engine/Source/Programs/UnrealBuildAccelerator/Common/Private/UbaStorageUtils.cpp b/Engine/Source/Programs/UnrealBuildAccelerator/Common/Private/UbaStorageUtils.cpp
--- a/Engine/Source/Programs/UnrealBuildAccelerator/Common/Private/UbaStorageUtils.cpp
+++ b/Engine/Source/Programs/UnrealBuildAccelerator/Common/Private/UbaStorageUtils.cpp
@@ -311,7 +311,7 @@ namespace uba
 				u32 uncompressedBlockSize = (u32)Min(left, maxUncompressedBlock);
 
 				u64 reserveSize = totalWritten + uncompressedBlockSize + diff + 8;
-				if (reserveSize > memoryBlock.mappedSize)
+				if (reserveSize > memoryBlock.committedSize)
 				{
 					u64 toAllocate = reserveSize - memoryBlock.writtenSize;
 					memoryBlock.Allocate(toAllocate, 1, hint);
diff --git a/Engine/Source/Programs/UnrealBuildAccelerator/Core/Private/UbaMemory.cpp b/Engine/Source/Programs/UnrealBuildAccelerator/Core/Private/UbaMemory.cpp
--- a/Engine/Source/Programs/UnrealBuildAccelerator/Core/Private/UbaMemory.cpp
+++ b/Engine/Source/Programs/UnrealBuildAccelerator/Core/Private/UbaMemory.cpp
@@ -102,14 +102,14 @@ namespace uba
 			FatalError(9882, TC("Ran out of reserved virtual address space. Reserved %llu, Needed %llu (%s)"), reserveSize, newPos, hint);
 
 		#if PLATFORM_WINDOWS
-		if (newPos > mappedSize)
+		if (newPos > committedSize)
 		{
-			u64 toCommit = AlignUp(newPos - mappedSize, MemoryBlock_ReserveAlign);
-			if (mappedSize + toCommit > reserveSize)
-				toCommit = reserveSize - mappedSize;
-			if (!VirtualAlloc(memory + mappedSize, toCommit, MEM_COMMIT, PAGE_READWRITE))
-				FatalError(9883, TC("Failed to commit virtual memory for memory block. Total size %llu (%u) (%s)"), mappedSize + toCommit, GetLastError(), hint);
-			mappedSize += toCommit;
+			u64 toCommit = AlignUp(newPos - committedSize, MemoryBlock_ReserveAlign);
+			if (committedSize + toCommit > reserveSize)
+				toCommit = reserveSize - committedSize;
+			if (!VirtualAlloc(memory + committedSize, toCommit, MEM_COMMIT, PAGE_READWRITE))
+				FatalError(9883, TC("Failed to commit virtual memory for memory block. Total size %llu (%u) (%s)"), committedSize + toCommit, GetLastError(), hint);
+			committedSize += toCommit;
 		}
 		#endif
 
@@ -118,19 +118,20 @@ namespace uba
 		return ret;
 	}
 
-	void MemoryBlock::ReserveNoLock(u64 bytes, const tchar* hint)
+	void* MemoryBlock::CommitNoLock(u64 bytes, const tchar* hint)
 	{
 		#if PLATFORM_WINDOWS
 		u64 newPos = Min(writtenSize + bytes, reserveSize);
-		if (newPos <= mappedSize)
-			return;
-		u64 toCommit = AlignUp(newPos - mappedSize, MemoryBlock_ReserveAlign);
-		if (mappedSize + toCommit > reserveSize)
-			toCommit = reserveSize - mappedSize;
-		if (!VirtualAlloc(memory + mappedSize, toCommit, MEM_COMMIT, PAGE_READWRITE))
-			FatalError(9883, TC("Failed to commit virtual memory for memory block. Total size %llu (%u) (%s)"), mappedSize + toCommit, GetLastError(), hint);
-		mappedSize += toCommit;
+		if (newPos <= committedSize)
+			return memory + writtenSize;
+		u64 toCommit = AlignUp(newPos - committedSize, MemoryBlock_ReserveAlign);
+		if (committedSize + toCommit > reserveSize)
+			toCommit = reserveSize - committedSize;
+		if (!VirtualAlloc(memory + committedSize, toCommit, MEM_COMMIT, PAGE_READWRITE))
+			FatalError(9883, TC("Failed to commit virtual memory for memory block. Total size %llu (%u) (%s)"), committedSize + toCommit, GetLastError(), hint);
+		committedSize += toCommit;
 		#endif
+		return memory + writtenSize;
 	}
 
 	void MemoryBlock::Free(void* p)
@@ -158,17 +159,17 @@ namespace uba
 		u8* m = memory;
 		u64 rs = reserveSize;
 		u64 ws = writtenSize;
-		u64 ms = mappedSize;
+		u64 ms = committedSize;
 
 		memory = other.memory;
 		reserveSize = other.reserveSize;
 		writtenSize = other.writtenSize;
-		mappedSize = other.mappedSize;
+		committedSize = other.committedSize;
 
 		other.memory = m;
 		other.reserveSize = rs;
 		other.writtenSize = ws;
-		other.mappedSize = ms;
+		other.committedSize = ms;
 	}
 
 	bool SupportsHugePages()
diff --git a/Engine/Source/Programs/UnrealBuildAccelerator/Core/Public/UbaDirectoryTable.h b/Engine/Source/Programs/UnrealBuildAccelerator/Core/Public/UbaDirectoryTable.h
--- a/Engine/Source/Programs/UnrealBuildAccelerator/Core/Public/UbaDirectoryTable.h
+++ b/Engine/Source/Programs/UnrealBuildAccelerator/Core/Public/UbaDirectoryTable.h
@@ -19,7 +19,7 @@ namespace uba
 		{
 			m_memory = mem;
 			m_lookup.reserve(tableCount + 100);
-			m_memoryBlock.ReserveNoLock(tableCount*(sizeof(GrowingUnorderedMap<StringKey, Directory>::value_type)+16), TC(""));
+			m_memoryBlock.CommitNoLock(tableCount*(sizeof(GrowingUnorderedMap<StringKey, Directory>::value_type)+16), TC(""));
 			ParseDirectoryTable(tableSize);
 		}
 
diff --git a/Engine/Source/Programs/UnrealBuildAccelerator/Core/Public/UbaMemory.h b/Engine/Source/Programs/UnrealBuildAccelerator/Core/Public/UbaMemory.h
--- a/Engine/Source/Programs/UnrealBuildAccelerator/Core/Public/UbaMemory.h
+++ b/Engine/Source/Programs/UnrealBuildAccelerator/Core/Public/UbaMemory.h
@@ -100,7 +100,7 @@ namespace uba
 		void Deinit();
 		void* Allocate(u64 bytes, u64 alignment, const tchar* hint);
 		void* AllocateNoLock(u64 bytes, u64 alignment, const tchar* hint);
-		void ReserveNoLock(u64 bytes, const tchar* hint);
+		void* CommitNoLock(u64 bytes, const tchar* hint);
 		void Free(void* p);
 		StringView Strdup(const StringView& str);
 		tchar* Strdup(const tchar* str);
@@ -111,7 +111,7 @@ namespace uba
 		u8* memory;
 		u64 reserveSize = 0;
 		u64 writtenSize = 0;
-		u64 mappedSize = 0;
+		u64 committedSize = 0;
 	};
 
 
diff --git a/Engine/Source/Programs/UnrealBuildAccelerator/Detours/Private/UbaDetoursFileMappingTable.cpp b/Engine/Source/Programs/UnrealBuildAccelerator/Detours/Private/UbaDetoursFileMappingTable.cpp
--- a/Engine/Source/Programs/UnrealBuildAccelerator/Detours/Private/UbaDetoursFileMappingTable.cpp
+++ b/Engine/Source/Programs/UnrealBuildAccelerator/Detours/Private/UbaDetoursFileMappingTable.cpp
@@ -18,7 +18,7 @@ namespace uba
 	{
 		m_mem = mem;
 		m_lookup.reserve(tableCount + 100);
-		m_memoryBlock.ReserveNoLock(tableCount*(sizeof(GrowingUnorderedMap<StringKey, FileInfo>::value_type)+16), TC(""));
+		m_memoryBlock.CommitNoLock(tableCount*(sizeof(GrowingUnorderedMap<StringKey, FileInfo>::value_type)+16), TC(""));
 		ParseNoLock(tableSize);
 	}
 
