[UBA]
* Fixed nobuffering file write path
* Fixed so Bottleneck logic is fifo
* Added options for switching to use WriteFile instead of file mappings when writing files and also an option to decide at what file size the writing becomes overlapped+unbuffered

diff --git a/Engine/Source/Programs/UnrealBuildAccelerator/Common/Private/UbaFileAccessor.cpp b/Engine/Source/Programs/UnrealBuildAccelerator/Common/Private/UbaFileAccessor.cpp
--- a/Engine/Source/Programs/UnrealBuildAccelerator/Common/Private/UbaFileAccessor.cpp
+++ b/Engine/Source/Programs/UnrealBuildAccelerator/Common/Private/UbaFileAccessor.cpp
@@ -11,7 +11,7 @@
 #include <copyfile.h>
 #endif
 
-#define UBA_USE_WRITE_THROUGH 0
+#define UBA_USE_SPARSE_FILE 0
 
 namespace uba
 {
@@ -75,10 +75,6 @@ namespace uba
 		}
 		#endif
 
-		#if UBA_USE_WRITE_THROUGH
-		flagsAndAttributes |= FILE_FLAG_NO_BUFFERING | FILE_FLAG_WRITE_THROUGH;
-		#endif
-
 		u32 createDisp = CREATE_ALWAYS;
 		u32 dwDesiredAccess = GENERIC_WRITE | DELETE;
 		if (allowRead)
@@ -118,20 +114,37 @@ namespace uba
 		if (!SetDeleteOnClose(m_logger, m_fileName, m_fileHandle, true))
 			return false;
 
-		if (flagsAndAttributes & FILE_FLAG_OVERLAPPED)
+		if (fileSize != ~u64(0))
 		{
-			if (fileSize != ~u64(0))
+			u64 allocSize = m_size;
+			if ((m_flagsAndAttributes & FILE_FLAG_NO_BUFFERING) != 0)
+				allocSize = AlignUp(fileSize, 4*1024);
+
+			if ((m_flagsAndAttributes & FILE_FLAG_OVERLAPPED) != 0)
+				(u64&)m_fileHandle |= OverlappedIoFlag;
+
+			#if UBA_USE_SPARSE_FILE
 			{
+				HANDLE handle = asHANDLE(m_fileHandle);
 				DWORD dwTemp;
-				if (!::DeviceIoControl(asHANDLE(m_fileHandle), FSCTL_SET_SPARSE, NULL, 0, NULL, 0, &dwTemp, NULL))
+				if (!::DeviceIoControl(handle, FSCTL_SET_SPARSE, NULL, 0, NULL, 0, &dwTemp, NULL))
 				{
 					DWORD lastError = GetLastError();
 					if (lastError != ERROR_INVALID_FUNCTION) // Some file systems don't support this
 						return m_logger.Error(TC("Failed to make file %s sparse (%s)"), realFileName, LastErrorToText(lastError).data);
 				}
-				SetEndOfFile(m_logger, realFileName, m_fileHandle, fileSize);
+
+				if (!SetEndOfFile(m_logger, realFileName, m_fileHandle, allocSize))
+					return false;
+			}
+			#else
+			{
+				FILE_ALLOCATION_INFO info;
+				info.AllocationSize = ToLargeInteger(allocSize);
+				if (!::SetFileInformationByHandle(asHANDLE(m_fileHandle), FileAllocationInfo, &info, sizeof(info)))
+					return m_logger.Error(TC("SetFileInformationByHandle (FileDispositionInfo) failed on %llu %s (%s)"), uintptr_t(m_fileHandle), m_fileName, LastErrorToText().data);
 			}
-			(u64&)m_fileHandle |= OverlappedIoFlag;
+			#endif
 		}
 #endif
 		m_isWrite = true;
@@ -170,51 +183,50 @@ namespace uba
 		return InternalClose(true, lastWriteTime);
 	}
 
-	bool FileAccessor::Write(const void* data, u64 dataLen, u64 offset)
+	bool FileAccessor::Write(const void* data, u64 dataLen, u64 offset, bool lastWrite)
 	{
 		auto& stats = KernelStats::GetCurrent();
 		ExtendedTimerScope ts(stats.writeFile);
 
 		if (!m_isWrite)
-			return false;
+			return m_logger.Error(TC("File %s is not opened for write"), m_fileName);
 
 		stats.writeFile.bytes += dataLen;
 
-		#if UBA_USE_WRITE_THROUGH
-		bool useWriteThrough = true;
-		#endif
-
-
-		#if UBA_USE_WRITE_THROUGH
-		if (useWriteThrough)
+#if PLATFORM_WINDOWS
+		bool noBuffering = (m_flagsAndAttributes & FILE_FLAG_NO_BUFFERING) != 0;
+		if (noBuffering)
 		{
-			if (m_writeThroughBufferPos)
+			if (m_writeThroughBufferSize)
 			{
+				if (offset != m_writeThroughBufferPos)
+					return m_logger.Error(TC("NoBuffering require file to be written sequentially or at 4k sizes (%s)"), m_fileName);
+
 				UBA_ASSERT(!((u64)m_fileHandle & OverlappedIoFlag) || !offset);
-				u64 toWrite = Min(WriteUnit - m_writeThroughBufferPos, dataLen);
-				memcpy(m_writeThroughBuffer + m_writeThroughBufferPos, data, toWrite);
-				m_writeThroughBufferPos += u32(toWrite);
+				u64 toWrite = Min(WriteUnit - m_writeThroughBufferSize, dataLen);
+				memcpy(m_writeThroughBuffer + m_writeThroughBufferSize, data, toWrite);
+				m_writeThroughBufferSize += u32(toWrite);
 				dataLen -= toWrite;
 				data = (const u8*)data + toWrite;
-				if (m_writeThroughBufferPos < WriteUnit)
+				if (m_writeThroughBufferSize < WriteUnit)
 					return true;
-				m_writeThroughBufferPos = 0;
+				m_writeThroughBufferSize = 0;
 				if (!Write(m_writeThroughBuffer, WriteUnit, 0))
 					return false;
 			}
 		}
-		#endif
 
-#if PLATFORM_WINDOWS
 		if ((u64)m_fileHandle & OverlappedIoFlag)
 		{
 			constexpr u64 BlockSize = 1024 * 1024;
-			constexpr u64 BlockCount = 256;
+			constexpr u64 BlockCount = 32;
 			OVERLAPPED ol[BlockCount];
 			Event ev[BlockCount];
+			u64 writeOffset = offset;
 			u64 writeLeft = dataLen;
 			u8* pos = (u8*)data;
 			u64 i = 0;
+			u8 tailBuffer[WriteUnit];
 
 			auto WaitAndCheckError = [&](u64 index)
 				{
@@ -228,23 +240,43 @@ namespace uba
 
 			auto eg = MakeGuard([&]()
 				{
-					u64 index = i % BlockCount;
-					if (i > BlockCount)
-						for (u64 j = index; j != BlockCount; ++j)
-							if (!WaitAndCheckError(j))
-								return false;
-					for (u64 j = 0; j != index; ++j)
+					for (u64 j=0, e=Min(i, BlockCount); j!=e; ++j)
 						if (!WaitAndCheckError(j))
 							return false;
 					return true;
 				});
 
+
+			#if 0//UBA_USE_SPARSE_FILE
+			u8 experiment[WriteUnit];
+			if (m_writeThroughBufferPos == 0 && lastWrite && dataLen > WriteUnit*4)
+			{
+				u64 lastOffset = (dataLen / WriteUnit) * WriteUnit - WriteUnit;
+				if (lastOffset)
+				{
+					ev[0].Create(true);
+					ol[0] = {};
+					ol[0].hEvent = ev[0].GetHandle();
+					ol[0].Offset = ToLow(lastOffset);
+					ol[0].OffsetHigh = ToHigh(lastOffset);
+					if (!::WriteFile(asHANDLE(m_fileHandle), experiment, WriteUnit, NULL, ol))
+					{
+						u32 lastError = GetLastError();
+						if (lastError != ERROR_IO_PENDING)
+							return m_logger.Error(L"FAILED!: %s", LastErrorToText(lastError).data);
+					}
+					if (!WaitAndCheckError(0))
+						return false;
+				}
+			}
+			#endif
+
 			while (writeLeft)
 			{
 				u64 index = i % BlockCount;
 
 				if (i < BlockCount)
-					ev[i].Create(false);
+					ev[i].Create(true);
 				else
 				{
 					if (!WaitAndCheckError(index))
@@ -254,31 +286,36 @@ namespace uba
 				u64 toWrite = Min(writeLeft, BlockSize);
 				u64 toActuallyWrite = toWrite;
 
-				#if UBA_USE_WRITE_THROUGH
-				UBA_ASSERT(!m_writeThroughBufferPos);
-				if (useWriteThrough)
+				if (noBuffering && toWrite < BlockSize)
 				{
-					if (toWrite < BlockSize)
+					toActuallyWrite = (toWrite / WriteUnit) * WriteUnit;
+					if (!toActuallyWrite)
 					{
-						toActuallyWrite = (toWrite / WriteUnit) * WriteUnit;
-						if (!toActuallyWrite)
+						if (lastWrite)
+						{
+							memcpy(tailBuffer, pos, toWrite);
+							pos = tailBuffer;
+							toActuallyWrite = WriteUnit;
+							m_writeThroughBufferSize = 0;
+						}
+						else
 						{
 							if (!m_writeThroughBuffer)
 								m_writeThroughBuffer = new u8[WriteUnit];
 							memcpy(m_writeThroughBuffer, pos, toWrite);
-							m_writeThroughBufferPos = u32(toWrite);
+							m_writeThroughBufferSize = u32(toWrite);
 							break;
 						}
-						else
-							toWrite = toActuallyWrite;
 					}
+					else
+						toWrite = toActuallyWrite;
 				}
-				#endif
 
 				ol[index] = {};
 				ol[index].hEvent = ev[index].GetHandle();
-				ol[index].Offset = ToLow(offset + i * BlockSize);
-				ol[index].OffsetHigh = ToHigh(offset + i * BlockSize);
+				ol[index].Offset = ToLow(writeOffset);
+				ol[index].OffsetHigh = ToHigh(writeOffset);
+				ev[index].Reset();
 
 				if (!::WriteFile(asHANDLE(m_fileHandle), pos, u32(toActuallyWrite), NULL, ol + index))
 				{
@@ -287,18 +324,17 @@ namespace uba
 						return m_logger.Error(L"FAILED!: %s", LastErrorToText(lastError).data);
 				}
 				++i;
+				writeOffset += toActuallyWrite;
 				pos += toWrite;
 				writeLeft -= toWrite;
 			}
 
 			if (!eg.Execute())
 				return false;
+
+			m_writeThroughBufferPos = u64(pos - (const u8*)data);
 			return true;
 		}
-
-		#if UBA_USE_WRITE_THROUGH
-		UBA_ASSERT(!useWriteThrough || (GetFilePointer(m_logger, m_fileHandle) % WriteUnit) == 0);
-		#endif
 #endif
 
 		u64 writeLeft = dataLen;
@@ -308,30 +344,10 @@ namespace uba
 			u32 toWrite = u32(Min(writeLeft, 256llu * 1024 * 1024));
 			u32 toActuallyWrite = toWrite;
 
-			#if UBA_USE_WRITE_THROUGH
-			UBA_ASSERT(!m_writeThroughBufferPos);
-			if (useWriteThrough)
-			{
-				toActuallyWrite = (toWrite / WriteUnit) * WriteUnit;
-				if (!toActuallyWrite)
-				{
-					if (!m_writeThroughBuffer)
-						m_writeThroughBuffer = new u8[WriteUnit];
-					memcpy(m_writeThroughBuffer, pos, toWrite);
-					m_writeThroughBufferPos = toWrite;
-					return true;
-				}
-			}
-			#endif
-
 #if PLATFORM_WINDOWS
 			DWORD written;
 			if (!::WriteFile(asHANDLE(m_fileHandle), pos, toActuallyWrite, &written, NULL))
 			{
-				#if UBA_USE_WRITE_THROUGH
-				UBA_ASSERT(!useWriteThrough || (u64(pos) % WriteUnit) == 0);
-				#endif
-
 				DWORD lastError = GetLastError();
 				m_logger.Error(TC("ERROR writing file %s writing %u bytes (%llu bytes written out of %llu. FilePos: %llu) (%s)"), m_fileName, toActuallyWrite, (dataLen - writeLeft), dataLen, GetFilePointer(m_logger, m_fileHandle), LastErrorToText(lastError).data);
 
@@ -400,27 +416,6 @@ namespace uba
 
 	bool FileAccessor::InternalClose(bool success, u64* lastWriteTime)
 	{
-#if UBA_USE_WRITE_THROUGH
-		if (m_writeThroughBufferPos)
-		{
-			u32 bufferPos = m_writeThroughBufferPos;
-			m_writeThroughBufferPos = 0;
-
-			u64 pos = GetFilePointer(m_logger, m_fileHandle);
-			if (pos == ~0ull)
-				return false;
-			if (!Write(m_writeThroughBuffer, WriteUnit))
-				return false;
-			if (!SetEndOfFile(m_logger, m_fileName, m_fileHandle, pos + bufferPos))
-				return false;
-		}
-		if (m_writeThroughBuffer)
-		{
-			delete[] m_writeThroughBuffer;
-			m_writeThroughBuffer = nullptr;
-		}
-#endif
-
 		if (m_data)
 		{
 			if (!UnmapViewOfFile(m_logger, m_data, m_size, m_fileName))
@@ -452,6 +447,25 @@ namespace uba
 		if (!m_isWrite)
 			return closeFile.Execute();
 
+		if (m_writeThroughBufferSize)
+		{
+			m_writeThroughBufferSize = 0;
+			if (!Write(m_writeThroughBuffer, WriteUnit, m_writeThroughBufferPos))
+				return false;
+		}
+
+		if ((m_flagsAndAttributes & FILE_FLAG_NO_BUFFERING) != 0)
+			if (!SetEndOfFile(m_logger, m_fileName, m_fileHandle, m_size))
+				return false;
+
+		if (m_writeThroughBuffer)
+		{
+			delete[] m_writeThroughBuffer;
+			m_writeThroughBuffer = nullptr;
+		}
+
+
+
 		#if !PLATFORM_WINDOWS
 		if (m_tempPath)
 			realFileName = tempFile.Append(m_tempPath).Append("Temp_").AppendValue(m_tempFileIndex).data;
diff --git a/Engine/Source/Programs/UnrealBuildAccelerator/Common/Private/UbaSession.cpp b/Engine/Source/Programs/UnrealBuildAccelerator/Common/Private/UbaSession.cpp
--- a/Engine/Source/Programs/UnrealBuildAccelerator/Common/Private/UbaSession.cpp
+++ b/Engine/Source/Programs/UnrealBuildAccelerator/Common/Private/UbaSession.cpp
@@ -1,7 +1,6 @@
 // Copyright Epic Games, Inc. All Rights Reserved.
 
 #include "UbaSession.h"
-#include "UbaBottleneck.h"
 #include "UbaCompressedObjFileHeader.h"
 #include "UbaConfig.h"
 #include "UbaFileAccessor.h"
@@ -14,6 +13,7 @@
 #include "UbaProtocol.h"
 #include "UbaStorageUtils.h"
 #include "UbaWorkManager.h"
+#include <span>
 
 #if PLATFORM_WINDOWS
 #include "UbaWinBinDependencyParser.h"
@@ -231,6 +231,10 @@ namespace uba
 		table.GetValueAsBool(extractObjFilesSymbols, TC("ExtractObjFilesSymbols"));
 		table.GetValueAsBool(useFakeVolumeSerial, TC("UseFakeVolumeSerial"));
 		table.GetValueAsU32(traceReserveSizeMb, TC("TraceReserveSizeMb"));
+		table.GetValueAsU32(writeFilesBottleneck, TC("WriteFilesBottleneck"));
+		table.GetValueAsU32(writeFilesFileMapMaxMb, TC("WriteFilesFileMapMaxMB"));
+		table.GetValueAsU32(writeFilesNoBufferingMinMb, TC("WriteFilesNoBufferingMinMB"));
+
 	}
 
 	void Session::AddEnvironmentVariableNoLock(const tchar* key, const tchar* value)
@@ -1109,6 +1113,9 @@ namespace uba
 	,	m_fileMappingBuffer(m_logger, workManager)
 	,	m_processCommunicationAllocator(m_logger, TC("CommunicationAllocator"))
 	,	m_trace(info.logWriter)
+	,	m_writeFilesBottleneck(info.writeFilesBottleneck)
+	,	m_writeFilesFileMapMax(u64(info.writeFilesFileMapMaxMb)*1024*1024)
+	,	m_writeFilesNoBufferingMin(u64(info.writeFilesNoBufferingMinMb)*1024*1024)
 	{
 		if (info.useUniqueId)
 		{
@@ -2577,6 +2584,8 @@ namespace uba
 
 	bool Session::WriteFilesToDisk(ProcessImpl& process, WrittenFile** files, u32 fileCount)
 	{
+		if (!fileCount)
+			return true;
 
 		auto writeFile = [&](WrittenFile& file)
 		{
@@ -2647,19 +2656,6 @@ namespace uba
 				}
 #endif
 
-				// This is to kill I/O when writing lots of pdb/dlls in parallel
-				#if PLATFORM_WINDOWS
-				constexpr u32 bottleneckMax = 32;
-				bool shouldBottleneck = false;//useOverlap;
-				static Bottleneck bottleneck(bottleneckMax);
-				auto bng = MakeGuard([&]() { if (shouldBottleneck) bottleneck.Leave(); });
-				#endif
-
-				#if PLATFORM_WINDOWS
-				shouldBottleneck = true;
-				bottleneck.Enter();
-				#endif
-
 				u64 fileSize = file.mappingWritten;
 				u8* mem = MapViewOfFile(m_logger, file.mappingHandle, FILE_MAP_READ, 0, fileSize);
 				if (!mem)
@@ -2699,17 +2695,15 @@ namespace uba
 					}
 					else
 					{
-
 						// Seems like best combo (for windows at least) is to use writes with overlap and max 16 at the same time.
 						// On one machine we get twice as fast without overlap if no bottleneck. On another machine (ntfs compression on) we get twice as slow without overlap
 						// Both machines behaves well with overlap AND bottleneck. Both machine are 128 logical core thread rippers.
-						bool useFileMapForWrite = fileSize > 0; // ::CreateFileMappingW does not work for zero-length files
-						bool useOverlap = false;//fileSize > 8 * 1024 * 1024;
-
+						bool useFileMapForWrite = fileSize && fileSize <= m_writeFilesFileMapMax; // ::CreateFileMappingW does not work for zero-length files
+						bool useOverlap = !useFileMapForWrite && fileSize >= m_writeFilesNoBufferingMin;//fileSize > 2*1024*1024;
 
 						u32 attributes = DefaultAttributes();
 						if (useOverlap)
-							attributes |= FILE_FLAG_OVERLAPPED;
+							attributes |= (FILE_FLAG_OVERLAPPED | FILE_FLAG_NO_BUFFERING);
 
 						if (useFileMapForWrite)
 						{
@@ -2721,14 +2715,7 @@ namespace uba
 						{
 							if (!destinationFile.CreateWrite(false, attributes, fileSize, m_tempPath.data))
 								return false;
-
-							#if PLATFORM_WINDOWS
-							//shouldBottleneck = fileSize > 64 * 1024 * 1024;
-							//if (shouldBottleneck)
-							//	bottleneck.Enter();
-							#endif
-
-							if (!destinationFile.Write(mem, fileSize))
+							if (!destinationFile.Write(mem, fileSize, 0, true))
 								return false;
 						}
 					}
@@ -2789,28 +2776,19 @@ namespace uba
 			return true;
 		};
 
-		for (u32 i=0; i!=fileCount; ++i)
-			if (!writeFile(*files[i]))
-				return false;
-		/*
-		Event events[32];
-		UBA_ASSERT(fileCount < 32);
-
-
-		for (u32 i=0; i!=fileCount; ++i)
-		{
-			events[i].Create(true);
+		// This is to not kill I/O when writing lots of pdb/dlls in parallel
+		#if PLATFORM_WINDOWS
+		BottleneckScope scope(m_writeFilesBottleneck);
+		#endif
 
-			m_workManager->AddWork([&, ii = i]()
-				{
-					writeFile(*files[ii]);
-					events[ii].Set();
-				}, 1, TC(""));
-		}
-		for (u32 i=0; i!=fileCount; ++i)
-			events[i].IsSet();
-		*/
-		return true;
+		Atomic<bool> success = true;
+		auto span = std::span(files, fileCount);
+		m_workManager->ParallelFor(fileCount - 1, span, [&](auto& it)
+			{
+				if (!writeFile(**it))
+					success = false;
+			});
+		return success;
 	}
 
 	bool Session::AllocFailed(Process& process, const tchar* allocType, u32 error)
diff --git a/Engine/Source/Programs/UnrealBuildAccelerator/Common/Public/UbaBottleneck.h b/Engine/Source/Programs/UnrealBuildAccelerator/Common/Public/UbaBottleneck.h
--- a/Engine/Source/Programs/UnrealBuildAccelerator/Common/Public/UbaBottleneck.h
+++ b/Engine/Source/Programs/UnrealBuildAccelerator/Common/Public/UbaBottleneck.h
@@ -7,51 +7,71 @@
 
 namespace uba
 {
+	struct BottleneckTicket
+	{
+		Event ev;
+		BottleneckTicket* next = nullptr;
+	};
 
 	struct Bottleneck
 	{
-		Bottleneck(u32 mc) : m_underMax(true), m_activeCount(0), m_maxCount(mc) {}
+		Bottleneck(u32 mc) : m_maxCount(mc) {}
 
-		void Enter()
+		void Enter(BottleneckTicket& ticket)
 		{
 			SCOPED_WRITE_LOCK(m_lock, lock);
-			while (true)
+			if (m_activeCount < m_maxCount)
 			{
-				if (m_activeCount < m_maxCount)
-				{
-					++m_activeCount;
-					if (m_activeCount == m_maxCount)
-						m_underMax.Reset();
-					break;
-				}
-
-				lock.Leave();
-				m_underMax.IsSet();
-				lock.Enter();
+				UBA_ASSERT(!m_first);
+				++m_activeCount;
+				return;
 			}
+
+			if (!m_first)
+				m_first = &ticket;
+			else
+				m_last->next = &ticket;
+			m_last = &ticket;
+
+			if (!ticket.ev.IsCreated())
+				ticket.ev.Create(true);
+
+			lock.Leave();
+
+			ticket.ev.IsSet();
 		}
 
-		void Leave()
+		void Leave(BottleneckTicket& ticket)
 		{
 			SCOPED_WRITE_LOCK(m_lock, lock);
-			if (m_activeCount == m_maxCount)
-				m_underMax.Set();
-			--m_activeCount;
+			BottleneckTicket* first = m_first;
+			if (!first)
+			{
+				UBA_ASSERT(m_activeCount > 0);
+				--m_activeCount;
+				return;
+			}
+			m_first = first->next;
+			if (!m_first)
+				m_last = nullptr;
+			first->ev.Set();
 		}
 
 	private:
 		ReaderWriterLock m_lock;
-		Event m_underMax;
-		u32 m_activeCount;
+		BottleneckTicket* m_first = nullptr;
+		BottleneckTicket* m_last = nullptr;
+		u32 m_activeCount = 0;
 		u32 m_maxCount;
 	};
 
 	struct BottleneckScope
 	{
-		BottleneckScope(Bottleneck& b) : bottleneck(b) { b.Enter(); }
-		~BottleneckScope() { bottleneck.Leave(); }
+		BottleneckScope(Bottleneck& b) : bottleneck(b) { b.Enter(ticket); }
+		~BottleneckScope() { bottleneck.Leave(ticket); }
 
 		Bottleneck& bottleneck;
+		BottleneckTicket ticket;
 	};
 
 }
diff --git a/Engine/Source/Programs/UnrealBuildAccelerator/Common/Public/UbaFileAccessor.h b/Engine/Source/Programs/UnrealBuildAccelerator/Common/Public/UbaFileAccessor.h
--- a/Engine/Source/Programs/UnrealBuildAccelerator/Common/Public/UbaFileAccessor.h
+++ b/Engine/Source/Programs/UnrealBuildAccelerator/Common/Public/UbaFileAccessor.h
@@ -18,7 +18,7 @@ namespace uba
 		bool CreateMemoryWrite(bool allowRead = false, u32 flagsAndAttributes = DefaultAttributes(), u64 size = 0, const tchar* tempPath = nullptr);
 		bool Close(u64* lastWriteTime = nullptr);
 
-		bool Write(const void* data, u64 dataLen, u64 offset = 0);
+		bool Write(const void* data, u64 dataLen, u64 offset = 0, bool lastWrite = false);
 
 		bool OpenRead();
 		bool OpenMemoryRead(u64 offset = 0, bool errorOnFail = true);
@@ -48,7 +48,8 @@ namespace uba
 		u32 m_tempFileIndex = 0;
 		#endif
 
-		u32 m_writeThroughBufferPos = 0;
+		u32 m_writeThroughBufferSize = 0;
 		u8* m_writeThroughBuffer = nullptr;
+		u64 m_writeThroughBufferPos = 0;
 	};
 }
\ No newline at end of file
diff --git a/Engine/Source/Programs/UnrealBuildAccelerator/Common/Public/UbaSession.h b/Engine/Source/Programs/UnrealBuildAccelerator/Common/Public/UbaSession.h
--- a/Engine/Source/Programs/UnrealBuildAccelerator/Common/Public/UbaSession.h
+++ b/Engine/Source/Programs/UnrealBuildAccelerator/Common/Public/UbaSession.h
@@ -2,6 +2,7 @@
 
 #pragma once
 
+#include "UbaBottleneck.h"
 #include "UbaFile.h"
 #include "UbaDirectoryTable.h"
 #include "UbaFileMapping.h"
@@ -295,6 +296,10 @@ namespace uba
 		GetNextProcessFunction m_getNextProcessFunction;
 		CustomServiceFunction m_customServiceFunction;
 
+		Bottleneck m_writeFilesBottleneck;
+		u64 m_writeFilesFileMapMax;
+		u64 m_writeFilesNoBufferingMin;
+
 		#if UBA_DEBUG_LOGGER
 		Logger* m_debugLogger = nullptr;
 		#endif
diff --git a/Engine/Source/Programs/UnrealBuildAccelerator/Common/Public/UbaSessionCreateInfo.h b/Engine/Source/Programs/UnrealBuildAccelerator/Common/Public/UbaSessionCreateInfo.h
--- a/Engine/Source/Programs/UnrealBuildAccelerator/Common/Public/UbaSessionCreateInfo.h
+++ b/Engine/Source/Programs/UnrealBuildAccelerator/Common/Public/UbaSessionCreateInfo.h
@@ -42,5 +42,8 @@ namespace uba
 		u64 deleteSessionsOlderThanSeconds = 12 * 60 * 60; // Delete session folders older than 12 hours by default . Set to 0 to not delete or 1 to delete all
 		u64 keepOutputFileMemoryMapsThreshold = 256 * 1024; // If allowMemoryMaps is true, output files will be kept in memory if smaller than this size
 		u32 traceReserveSizeMb = 128;			// Memory reserved for trace file (in mb)
+		u32 writeFilesBottleneck = 16;			// When writing files to disk this number will control how many that can write in parallel
+		u32 writeFilesFileMapMaxMb = 100000;	// If file size is smaller than this number, then use memory maps for writing to disk
+		u32 writeFilesNoBufferingMinMb = 16;	// If file size is at least this number, then overlapped io with no buffering will be used
 	};
 }
