[UBA]
* Changed so uba exception handler does continue search afterwards. This to make sure we get crashdumps if wer is set up

diff --git a/Engine/Source/Programs/UnrealBuildAccelerator/Common/Private/UbaLogger.cpp b/Engine/Source/Programs/UnrealBuildAccelerator/Common/Private/UbaLogger.cpp
--- a/Engine/Source/Programs/UnrealBuildAccelerator/Common/Private/UbaLogger.cpp
+++ b/Engine/Source/Programs/UnrealBuildAccelerator/Common/Private/UbaLogger.cpp
@@ -63,7 +63,7 @@ namespace uba
 		fflush(stdout);
 
 #if PLATFORM_WINDOWS
-#if UBA_ASSERT_MESSAGEBOX
+	#if UBA_ASSERT_MESSAGEBOX
 		int ret = MessageBoxW(GetConsoleWindow(), sb.data, TC("Assert"), MB_ABORTRETRYIGNORE);
 		if (ret != IDABORT)
 		{
@@ -74,17 +74,18 @@ namespace uba
 
 		SetFocus(GetConsoleWindow());
 		SetActiveWindow(GetConsoleWindow());
-#else
+	#else
 		if (IsDebuggerPresent())
 			DebugBreak();
-#endif
+	#endif
 
 
 		// *(int*)nullptr = 42; // Use this to force a crashdump instead of exiting
-
-		ExitProcess(terminateCode);
+		if (allowTerminate)
+			ExitProcess(terminateCode);
 #else
-		_Exit(-1);
+		if (allowTerminate)
+			_Exit(-1);
 #endif
 	}
 
diff --git a/Engine/Source/Programs/UnrealBuildAccelerator/Core/Private/UbaPlatform.cpp b/Engine/Source/Programs/UnrealBuildAccelerator/Core/Private/UbaPlatform.cpp
--- a/Engine/Source/Programs/UnrealBuildAccelerator/Core/Private/UbaPlatform.cpp
+++ b/Engine/Source/Programs/UnrealBuildAccelerator/Core/Private/UbaPlatform.cpp
@@ -561,8 +561,8 @@ namespace uba
 		else
 			text.Appendf(L"ERROR: Unhandled Exception (Code: 0x%x)", exceptionCode);
 
-		UbaAssert(text.data, nullptr, 0, nullptr, true, exceptionCode, contextPtr);
-		return EXCEPTION_EXECUTE_HANDLER;
+		UbaAssert(text.data, nullptr, 0, nullptr, false, exceptionCode, contextPtr);
+		return EXCEPTION_CONTINUE_SEARCH;
 	}
 
 	DWORD UbaExceptionHandlerAllowAccessViolation(DWORD code, _EXCEPTION_POINTERS* info, const tchar* fileName)
