[UBA]
* Disabled SetFileInformationByHandle(FileAllocationInfo)  on wine machines since it fails there it seems

diff --git a/Engine/Source/Programs/UnrealBuildAccelerator/Common/Private/UbaFileAccessor.cpp b/Engine/Source/Programs/UnrealBuildAccelerator/Common/Private/UbaFileAccessor.cpp
--- a/Engine/Source/Programs/UnrealBuildAccelerator/Common/Private/UbaFileAccessor.cpp
+++ b/Engine/Source/Programs/UnrealBuildAccelerator/Common/Private/UbaFileAccessor.cpp
@@ -114,7 +114,7 @@ namespace uba
 		if (!SetDeleteOnClose(m_logger, m_fileName, m_fileHandle, true))
 			return false;
 
-		if (fileSize != ~u64(0))
+		if (fileSize != 0)
 		{
 			u64 allocSize = m_size;
 			if ((m_flagsAndAttributes & FILE_FLAG_NO_BUFFERING) != 0)
@@ -142,7 +142,8 @@ namespace uba
 				FILE_ALLOCATION_INFO info;
 				info.AllocationSize = ToLargeInteger(allocSize);
 				if (!::SetFileInformationByHandle(asHANDLE(m_fileHandle), FileAllocationInfo, &info, sizeof(info)))
-					return m_logger.Error(TC("SetFileInformationByHandle (FileDispositionInfo) failed on %llu %s (%s)"), uintptr_t(m_fileHandle), m_fileName, LastErrorToText().data);
+					if (!IsRunningWine())
+						return m_logger.Error(TC("SetFileInformationByHandle (FileDispositionInfo) failed on %llu %s (%s)"), uintptr_t(m_fileHandle), m_fileName, LastErrorToText().data);
 			}
 			#endif
 		}
