[UBA]
* Fixed stack overflow error when using cache server on mac

diff --git a/Engine/Source/Programs/UnrealBuildAccelerator/Common/Private/UbaCacheClient.cpp b/Engine/Source/Programs/UnrealBuildAccelerator/Common/Private/UbaCacheClient.cpp
--- a/Engine/Source/Programs/UnrealBuildAccelerator/Common/Private/UbaCacheClient.cpp
+++ b/Engine/Source/Programs/UnrealBuildAccelerator/Common/Private/UbaCacheClient.cpp
@@ -955,31 +955,8 @@ namespace uba
 	bool CacheClient::SendCacheEntry(Bucket& bucket, const RootPaths& rootPaths, const CasKey& cmdKey, const Map<u32, u32>& inputsStringToCasKey, const Map<u32, u32>& outputsStringToCasKey, const u8* logLines, u64 logLinesSize, u64& outBytesSent)
 	{
 		StackBinaryReader<1024> reader;
-		{
-			StackBinaryWriter<SendMaxSize> writer;
-
-			NetworkMessage msg(m_client, CacheServiceId, CacheMessageType_StoreEntry, writer);
-			writer.Write7BitEncoded(MakeId(bucket.id));
-			writer.WriteCasKey(cmdKey);
-
-			writer.Write7BitEncoded(inputsStringToCasKey.size());
-			writer.Write7BitEncoded(outputsStringToCasKey.size());
-			for (auto& kv : outputsStringToCasKey)
-				writer.Write7BitEncoded(kv.second);
-
-			for (auto& kv : inputsStringToCasKey)
-				writer.Write7BitEncoded(kv.second);
-
-			if (logLinesSize)
-				if (writer.GetCapacityLeft() > logLinesSize + Get7BitEncodedCount(logLinesSize))
-					writer.WriteBytes(logLines, logLinesSize);
-
-			if (!msg.Send(reader))
-			{
-				m_logger.Info(TC("Failed to send cache entry. CasTable: %u/%u PathTable: %u/%u"), bucket.casKeyTableSizeSent, bucket.sendCasKeyTable.GetSize(), bucket.pathTableSizeSent, bucket.sendPathTable.GetSize());
-				return false;
-			}
-		}
+		if (!SendCacheEntryMessage(reader, bucket, cmdKey, inputsStringToCasKey, outputsStringToCasKey, logLines, logLinesSize))
+			return false;
 
 		// Server has all content for caskeys.. upload is done
 		if (!reader.GetLeft())
@@ -1104,7 +1081,32 @@ namespace uba
 		return doneGuard.Execute();
 	}
 
-#pragma optimize("", off)
+	bool CacheClient::SendCacheEntryMessage(BinaryReader& out, Bucket& bucket, const CasKey& cmdKey, const Map<u32, u32>& inputsStringToCasKey, const Map<u32, u32>& outputsStringToCasKey, const u8* logLines, u64 logLinesSize)
+	{
+		StackBinaryWriter<SendMaxSize> writer;
+
+		NetworkMessage msg(m_client, CacheServiceId, CacheMessageType_StoreEntry, writer);
+		writer.Write7BitEncoded(MakeId(bucket.id));
+		writer.WriteCasKey(cmdKey);
+
+		writer.Write7BitEncoded(inputsStringToCasKey.size());
+		writer.Write7BitEncoded(outputsStringToCasKey.size());
+		for (auto& kv : outputsStringToCasKey)
+			writer.Write7BitEncoded(kv.second);
+
+		for (auto& kv : inputsStringToCasKey)
+			writer.Write7BitEncoded(kv.second);
+
+		if (logLinesSize)
+			if (writer.GetCapacityLeft() > logLinesSize + Get7BitEncodedCount(logLinesSize))
+				writer.WriteBytes(logLines, logLinesSize);
+
+		if (msg.Send(out))
+			return true;
+		m_logger.Info(TC("Failed to send cache entry. CasTable: %u/%u PathTable: %u/%u"), bucket.casKeyTableSizeSent, bucket.sendCasKeyTable.GetSize(), bucket.pathTableSizeSent, bucket.sendPathTable.GetSize());
+		return false;
+	}
+
 	bool CacheClient::FetchCasTable(Bucket& bucket, CacheStats& stats, u32 requiredCasTableOffset)
 	{
 		if (HasEnoughData(bucket, requiredCasTableOffset))
@@ -1182,8 +1184,6 @@ namespace uba
 		return true;
 	}
 
-#pragma optimize("", on)
-
 	bool CacheClient::HasEnoughData(Bucket& bucket, u32 requiredCasTableOffset)
 	{
 		u32 pathOffset;
diff --git a/Engine/Source/Programs/UnrealBuildAccelerator/Common/Public/UbaCacheClient.h b/Engine/Source/Programs/UnrealBuildAccelerator/Common/Public/UbaCacheClient.h
--- a/Engine/Source/Programs/UnrealBuildAccelerator/Common/Public/UbaCacheClient.h
+++ b/Engine/Source/Programs/UnrealBuildAccelerator/Common/Public/UbaCacheClient.h
@@ -75,6 +75,7 @@ namespace uba
 		UBA_NOINLINE bool SendPathTable(Bucket& bucket, u32 requiredPathTableSize);
 		UBA_NOINLINE bool SendCasTable(Bucket& bucket, u32 requiredCasTableSize);
 		UBA_NOINLINE bool SendCacheEntry(Bucket& bucket, const RootPaths& rootPaths, const CasKey& cmdKey, const Map<u32, u32>& inputsStringToCasKey, const Map<u32, u32>& outputsStringToCasKey, const u8* logLines, u64 logLinesSize, u64& outBytesSent);
+		UBA_NOINLINE bool SendCacheEntryMessage(BinaryReader& out, Bucket& bucket, const CasKey& cmdKey, const Map<u32, u32>& inputsStringToCasKey, const Map<u32, u32>& outputsStringToCasKey, const u8* logLines, u64 logLinesSize);
 		bool FetchCasTable(Bucket& bucket, CacheStats& stats, u32 requiredCasTableOffset);
 		bool FetchFile(Bucket& bucket, const RootPaths& rootPaths, const ProcessStartInfo& info, CacheStats& cacheStats, StorageStats& storageStats, u32 casKeyOffset);
 		template<typename TableType> bool FetchCompactTable(u32 bucketId, TableType& table, u32 requiredTableOffset, u8 messageType);
@@ -84,7 +85,7 @@ namespace uba
 		bool ReportUsedEntry(Vector<ProcessLogLine>& outLogLines, bool ownedLogLines, Bucket& bucket, const CasKey& cmdKey, u32 entryId);
 		bool PopulateLogLines(Vector<ProcessLogLine>& outLogLines, const u8* mem, u64 memLen);
 
-		CasKey GetCmdKey(const RootPaths& rootPaths, const ProcessStartInfo& info, bool report, u32 bucketId);
+		UBA_NOINLINE CasKey GetCmdKey(const RootPaths& rootPaths, const ProcessStartInfo& info, bool report, u32 bucketId);
 		bool DevirtualizePath(StringBufferBase& inOut, RootsHandle rootsHandle);
 		bool ShouldNormalize(const StringBufferBase& path);
 
diff --git a/Engine/Source/Programs/UnrealBuildAccelerator/Core/Public/UbaPlatform.h b/Engine/Source/Programs/UnrealBuildAccelerator/Core/Public/UbaPlatform.h
--- a/Engine/Source/Programs/UnrealBuildAccelerator/Core/Public/UbaPlatform.h
+++ b/Engine/Source/Programs/UnrealBuildAccelerator/Core/Public/UbaPlatform.h
@@ -178,7 +178,7 @@ namespace uba
 	#else
 	inline constexpr bool CaseInsensitiveFs = true;
 	#endif
-	#define UBA_NOINLINE
+	#define UBA_NOINLINE __attribute__((noinline))
     #define ERROR_SUCCESS 0
 	#define DUPLICATE_SAME_ACCESS 0
 	struct FileMappingHandle
@@ -188,7 +188,7 @@ namespace uba
 		bool operator==(const FileMappingHandle& o) const { return shmFd == o.shmFd; }
 		bool IsValid() const { return shmFd != -1; }
 		u64 ToU64() const { return uid; }
-		static FileMappingHandle FromU64(u64 v) { UBA_ASSERT(false); return FileMappingHandle(-1, -1, v); }
+		static FileMappingHandle FromU64(u64 v) { return FileMappingHandle(-1, -1, v); }
 		int shmFd;
 		int lockFd;
 		u64 uid;
