[UBA]
* Added so detailed trace also enabled tracing for cache client

diff --git a/Engine/Source/Programs/UnrealBuildAccelerator/Common/Private/UbaCacheClient.cpp b/Engine/Source/Programs/UnrealBuildAccelerator/Common/Private/UbaCacheClient.cpp
--- a/Engine/Source/Programs/UnrealBuildAccelerator/Common/Private/UbaCacheClient.cpp
+++ b/Engine/Source/Programs/UnrealBuildAccelerator/Common/Private/UbaCacheClient.cpp
@@ -132,6 +132,9 @@ namespace uba
 			{
 				m_connected = false;
 			});
+
+		if (m_session.HasDetailedtrace())
+			m_client.SetWorkTracker(&m_session.GetTrace());
 	}
 
 	CacheClient::~CacheClient() = default;
diff --git a/Engine/Source/Programs/UnrealBuildAccelerator/Common/Public/UbaSession.h b/Engine/Source/Programs/UnrealBuildAccelerator/Common/Public/UbaSession.h
--- a/Engine/Source/Programs/UnrealBuildAccelerator/Common/Public/UbaSession.h
+++ b/Engine/Source/Programs/UnrealBuildAccelerator/Common/Public/UbaSession.h
@@ -81,7 +81,8 @@ namespace uba
 		const ApplicationRules* GetRules(const ProcessStartInfo& si); // Get application rules used for process
 		const tchar* GetTempPath(); // Path for temp files used for current session
 		const tchar* GetRootDir(); // Root dir for session files
-		bool ShouldStoreObjFilesCompressed() { return m_storeObjFilesCompressed; }
+		bool ShouldStoreObjFilesCompressed() const { return m_storeObjFilesCompressed; }
+		bool HasDetailedtrace() const { return m_detailedTrace; }
 
 		virtual ~Session();
 
