[UBA]
* Added more features to work trace system so it is now possible to do more detailed performance tracking
* Added trace scopes to cache client that is enabled by setting defines UBA_TRACE_WRITE_CACHE and UBA_TRACE_FETCH_CACHE

diff --git a/Engine/Source/Programs/UnrealBuildAccelerator/Common/Private/UbaCacheClient.cpp b/Engine/Source/Programs/UnrealBuildAccelerator/Common/Private/UbaCacheClient.cpp
--- a/Engine/Source/Programs/UnrealBuildAccelerator/Common/Private/UbaCacheClient.cpp
+++ b/Engine/Source/Programs/UnrealBuildAccelerator/Common/Private/UbaCacheClient.cpp
@@ -18,6 +18,21 @@
 #define UBA_LOG_WRITE_CACHE_INFO 0 // 0 = Disabled, 1 = Normal, 2 = Detailed
 #define UBA_LOG_FETCH_CACHE_INFO 0 // 0 = Disabled, 1 = Misses, 2 = Both misses and hits
 
+#define UBA_TRACE_WRITE_CACHE 0
+#define UBA_TRACE_FETCH_CACHE 0
+
+#if UBA_TRACE_WRITE_CACHE
+#define UBA_TRACE_WRITE_HINT(text) tws.AddHint(AsView(TC(text)));
+#else
+#define UBA_TRACE_WRITE_HINT(text)
+#endif
+
+#if UBA_TRACE_FETCH_CACHE
+#define UBA_TRACE_FETCH_HINT(text) tws.AddHint(AsView(TC(text)));
+#else
+#define UBA_TRACE_FETCH_HINT(text)
+#endif
+
 namespace uba
 {
 	struct DowngradedLogger : public LoggerWithWriter
@@ -155,6 +170,12 @@ namespace uba
 		if (!inputsSize)
 			return false;
 
+		#if UBA_TRACE_WRITE_CACHE
+		TrackWorkScope tws(m_client, AsView(TC("WriteToCache")));
+		#else
+		TrackWorkScope tws;
+		#endif
+
 		CasKey cmdKey = GetCmdKey(rootPaths, info, false, bucketId);
 		if (cmdKey == CasKeyZero)
 		{
@@ -187,6 +208,8 @@ namespace uba
 
 		TString qualifiedPath;
 
+		UBA_TRACE_WRITE_HINT("TraverseInputsOutputs");
+
 		// Traverse all inputs and outputs. to create cache entry that we can send to server
 		while (true)
 		{
@@ -358,16 +381,20 @@ namespace uba
 		if (outputsStringToCasKey.empty())
 			m_logger.Warning(TC("NO OUTPUTS FROM process %s"), info.GetDescription()); 
 
+		UBA_TRACE_WRITE_HINT("SendPathTable");
+
 		// Make sure server has enough of the path table to be able to resolve offsets from cache entry
 		if (!SendPathTable(bucket, requiredPathTableSize))
 			return false;
 
+		UBA_TRACE_WRITE_HINT("SendCasTable");
+
 		// Make sure server has enough of the cas table to be able to resolve offsets from cache entry
 		if (!SendCasTable(bucket, requiredCasTableSize))
 			return false;
 
 		// actual cache entry now when we know server has the needed tables
-		if (!SendCacheEntry(bucket, rootPaths, cmdKey, inputsStringToCasKey, outputsStringToCasKey, logLines, logLinesSize, bytesSent))
+		if (!SendCacheEntry(tws, bucket, rootPaths, cmdKey, inputsStringToCasKey, outputsStringToCasKey, logLines, logLinesSize, bytesSent))
 			return false;
 
 
@@ -417,6 +444,12 @@ namespace uba
 		if (!m_connected)
 			return false;
 
+		#if UBA_TRACE_FETCH_CACHE
+		TrackWorkScope tws(m_client, AsView(TC("FetchFromCache")));
+		#else
+		TrackWorkScope tws;
+		#endif
+
 		CacheStats cacheStats;
 		StorageStats storageStats;
 		KernelStats kernelStats;
@@ -451,6 +484,8 @@ namespace uba
 		bucketsLock.Leave();
 
 		{
+			UBA_TRACE_FETCH_HINT("FetchEntries")
+
 			TimerScope ts(cacheStats.fetchEntries);
 			// Fetch entries.. server will provide as many as fits. TODO: Should it be possible to ask for more entries?
 			StackBinaryWriter<32> writer;
@@ -565,6 +600,7 @@ namespace uba
 
 		// Create ranges out of shared offsets that matches local state
 		{
+			UBA_TRACE_FETCH_HINT("TestSharedMatch")
 			TimerScope ts(cacheStats.testEntry);
 			u64 sharedSize = reader.Read7BitEncoded();
 
@@ -619,6 +655,7 @@ namespace uba
 
 		// Read entries
 		{
+			UBA_TRACE_FETCH_HINT("TestEntriesMatch")
 			--cacheStats.testEntry.count; // Remove the shared one
 
 			BinaryReader entryReader(reader.GetPositionData(), 0, reader.GetLeft());
@@ -675,6 +712,7 @@ namespace uba
 				if (!m_useCacheHit)
 					return false;
 
+				UBA_TRACE_FETCH_HINT("ReportUsedEntry")
 
 				if (logLinesType == LogLinesType_Shared)
 					if (!PopulateLogLines(outResult.logLines, sharedLogLines, sharedLogLinesSize))
@@ -690,20 +728,36 @@ namespace uba
 				while (outputsReader.GetLeft())
 					casKeyOffsets.push_back(u32(outputsReader.Read7BitEncoded()));
 
+				UBA_TRACE_FETCH_HINT("FetchTableAndFiles")
+
 				if (!casKeyOffsets.empty())
 				{
 					Atomic<bool> fetchSuccess = true;
 					m_storage.m_workManager->ParallelFor(u32(casKeyOffsets.size() - 1), casKeyOffsets, [&](auto& it)
 						{
+							#if UBA_TRACE_FETCH_CACHE
+							u64 time = GetTime();
+							#endif
 							u32 casKeyOffset = *it;
 							if (!FetchCasTable(bucket, cacheStats, casKeyOffset))
 							{
 								fetchSuccess = false;
 								return;
 							}
+
+							#if UBA_TRACE_FETCH_CACHE
+							tws.AddHint(AsView(TC("FetchCasTable")), time);
+							time = GetTime();
+							#endif
+
 							TimerScope fts(cacheStats.fetchOutput);
 							if (!FetchFile(bucket, rootPaths, info, cacheStats, storageStats, casKeyOffset))
 								fetchSuccess = false;
+
+							#if UBA_TRACE_FETCH_CACHE
+							tws.AddHint(AsView(TC("FetchFile")), time);
+							#endif
+
 						}, AsView(TC("CacheFetchOutput")));
 					if (!fetchSuccess)
 						return false;
@@ -955,9 +1009,12 @@ namespace uba
 		return true;
 	}
 
-	bool CacheClient::SendCacheEntry(Bucket& bucket, const RootPaths& rootPaths, const CasKey& cmdKey, const Map<u32, u32>& inputsStringToCasKey, const Map<u32, u32>& outputsStringToCasKey, const u8* logLines, u64 logLinesSize, u64& outBytesSent)
+	bool CacheClient::SendCacheEntry(TrackWorkScope& tws, Bucket& bucket, const RootPaths& rootPaths, const CasKey& cmdKey, const Map<u32, u32>& inputsStringToCasKey, const Map<u32, u32>& outputsStringToCasKey, const u8* logLines, u64 logLinesSize, u64& outBytesSent)
 	{
 		StackBinaryReader<1024> reader;
+
+		UBA_TRACE_WRITE_HINT("SendCacheEntryMessage");
+
 		if (!SendCacheEntryMessage(reader, bucket, cmdKey, inputsStringToCasKey, outputsStringToCasKey, logLines, logLinesSize))
 			return false;
 
@@ -969,6 +1026,7 @@ namespace uba
 		auto doneGuard = MakeGuard([&]()
 			{
 				// Send done.. confirm to server
+				UBA_TRACE_WRITE_HINT("SendCacheDone");
 				StackBinaryWriter<64> writer;
 				NetworkMessage msg(m_client, CacheServiceId, CacheMessageType_StoreEntryDone, writer.Reset());
 				writer.Write7BitEncoded(MakeId(bucket.id));
@@ -994,6 +1052,8 @@ namespace uba
 			StorageImpl::CasEntry* casEntry;
 			if (m_storage.HasCasFile(casKey, &casEntry))
 			{
+				UBA_TRACE_WRITE_HINT("SendFile");
+
 				UBA_ASSERT(!IsNormalized(casKey));
 				StringBuffer<> casKeyFileName;
 				if (!m_storage.GetCasFileName(casKeyFileName, casKey))
@@ -1028,6 +1088,8 @@ namespace uba
 			}
 			else // If we don't have the cas key it should be one of the normalized files.... otherwise there is a bug
 			{
+				UBA_TRACE_WRITE_HINT("SendNormalizedFile");
+
 				if (!IsNormalized(casKey))
 					return m_logger.Error(TC("Can't find output file %s to send to cache server"), path.data);
 
diff --git a/Engine/Source/Programs/UnrealBuildAccelerator/Common/Private/UbaTrace.cpp b/Engine/Source/Programs/UnrealBuildAccelerator/Common/Private/UbaTrace.cpp
--- a/Engine/Source/Programs/UnrealBuildAccelerator/Common/Private/UbaTrace.cpp
+++ b/Engine/Source/Programs/UnrealBuildAccelerator/Common/Private/UbaTrace.cpp
@@ -137,9 +137,9 @@ namespace uba
 		return workId;
 	}
 
-	void Trace::TrackWorkHint(u32 id, const StringView& hint)
+	void Trace::TrackWorkHint(u32 id, const StringView& hint, u64 startTime)
 	{
-		WorkHint(id, hint);
+		WorkHint(id, hint, startTime);
 	}
 
 	void Trace::TrackWorkEnd(u32 id)
@@ -385,12 +385,13 @@ namespace uba
 		writer.WriteU32(color);
 	}
 
-	void Trace::WorkHint(u32 workIndex, const StringView& hint)
+	void Trace::WorkHint(u32 workIndex, const StringView& hint, u64 startTime)
 	{
 		u32 stringIndex = AddString(hint);
 		BEGIN_TRACE_ENTRY(WorkHint);
 		writer.Write7BitEncoded(workIndex);
 		writer.Write7BitEncoded(stringIndex);
+		writer.Write7BitEncoded(startTime ? (startTime - m_startTime) : 0);
 	}
 
 	void Trace::WorkEnd(u32 workIndex)
diff --git a/Engine/Source/Programs/UnrealBuildAccelerator/Common/Private/UbaWorkManager.cpp b/Engine/Source/Programs/UnrealBuildAccelerator/Common/Private/UbaWorkManager.cpp
--- a/Engine/Source/Programs/UnrealBuildAccelerator/Common/Private/UbaWorkManager.cpp
+++ b/Engine/Source/Programs/UnrealBuildAccelerator/Common/Private/UbaWorkManager.cpp
@@ -16,10 +16,10 @@ namespace uba
 		return 0;
 	}
 
-	void WorkManager::TrackWorkHint(u32 id, const StringView& hint)
+	void WorkManager::TrackWorkHint(u32 id, const StringView& hint, u64 startTime)
 	{
 		if (auto t = m_workTracker.load())
-			return t->TrackWorkHint(id, hint);
+			return t->TrackWorkHint(id, hint, startTime);
 	}
 
 	void WorkManager::TrackWorkEnd(u32 id)
diff --git a/Engine/Source/Programs/UnrealBuildAccelerator/Common/Public/UbaCacheClient.h b/Engine/Source/Programs/UnrealBuildAccelerator/Common/Public/UbaCacheClient.h
--- a/Engine/Source/Programs/UnrealBuildAccelerator/Common/Public/UbaCacheClient.h
+++ b/Engine/Source/Programs/UnrealBuildAccelerator/Common/Public/UbaCacheClient.h
@@ -20,6 +20,7 @@ namespace uba
 	struct ProcessStartInfo;
 	struct StorageStats;
 	struct StringView;
+	struct TrackWorkScope;
 	using RootsHandle = u64;
 
 	struct CacheClientCreateInfo
@@ -74,7 +75,7 @@ namespace uba
 		// No inlining because they use lots of stack space
 		UBA_NOINLINE bool SendPathTable(Bucket& bucket, u32 requiredPathTableSize);
 		UBA_NOINLINE bool SendCasTable(Bucket& bucket, u32 requiredCasTableSize);
-		UBA_NOINLINE bool SendCacheEntry(Bucket& bucket, const RootPaths& rootPaths, const CasKey& cmdKey, const Map<u32, u32>& inputsStringToCasKey, const Map<u32, u32>& outputsStringToCasKey, const u8* logLines, u64 logLinesSize, u64& outBytesSent);
+		UBA_NOINLINE bool SendCacheEntry(TrackWorkScope& tws, Bucket& bucket, const RootPaths& rootPaths, const CasKey& cmdKey, const Map<u32, u32>& inputsStringToCasKey, const Map<u32, u32>& outputsStringToCasKey, const u8* logLines, u64 logLinesSize, u64& outBytesSent);
 		UBA_NOINLINE bool SendCacheEntryMessage(BinaryReader& out, Bucket& bucket, const CasKey& cmdKey, const Map<u32, u32>& inputsStringToCasKey, const Map<u32, u32>& outputsStringToCasKey, const u8* logLines, u64 logLinesSize);
 		bool FetchCasTable(Bucket& bucket, CacheStats& stats, u32 requiredCasTableOffset);
 		bool FetchFile(Bucket& bucket, const RootPaths& rootPaths, const ProcessStartInfo& info, CacheStats& cacheStats, StorageStats& storageStats, u32 casKeyOffset);
diff --git a/Engine/Source/Programs/UnrealBuildAccelerator/Common/Public/UbaTrace.h b/Engine/Source/Programs/UnrealBuildAccelerator/Common/Public/UbaTrace.h
--- a/Engine/Source/Programs/UnrealBuildAccelerator/Common/Public/UbaTrace.h
+++ b/Engine/Source/Programs/UnrealBuildAccelerator/Common/Public/UbaTrace.h
@@ -99,7 +99,7 @@ namespace uba
 		void FileStoreBegin(u32 clientId, const CasKey& key, u64 size, const StringView& hint, bool detailed);
 		void FileStoreEnd(u32 clientId, const CasKey& key);
 		void WorkBegin(u32 workIndex, const StringView& desc, const Color& color);
-		void WorkHint(u32 workIndex, const StringView& hint);
+		void WorkHint(u32 workIndex, const StringView& hint, u64 startTime);
 		void WorkEnd(u32 workIndex);
 		void ProgressUpdate(u32 processesTotal, u32 processesDone, u32 errorCount);
 		void StatusUpdate(u32 statusRow, u32 statusColumn, const tchar* statusText, LogEntryType statusType, const tchar* statusLink);
@@ -113,7 +113,7 @@ namespace uba
 		bool StopWrite(const tchar* writeFileName);
 
 		virtual u32 TrackWorkStart(const StringView& desc, const Color& color) final override;
-		virtual void TrackWorkHint(u32 id, const StringView& hint) final override;
+		virtual void TrackWorkHint(u32 id, const StringView& hint, u64 startTime = 0) final override;
 		virtual void TrackWorkEnd(u32 id) final override;
 
 	private:
diff --git a/Engine/Source/Programs/UnrealBuildAccelerator/Common/Public/UbaWorkManager.h b/Engine/Source/Programs/UnrealBuildAccelerator/Common/Public/UbaWorkManager.h
--- a/Engine/Source/Programs/UnrealBuildAccelerator/Common/Public/UbaWorkManager.h
+++ b/Engine/Source/Programs/UnrealBuildAccelerator/Common/Public/UbaWorkManager.h
@@ -19,7 +19,7 @@ namespace uba
 	{
 	public:
 		virtual u32 TrackWorkStart(const StringView& desc, const Color& color) = 0;
-		virtual void TrackWorkHint(u32 id, const StringView& hint) = 0;
+		virtual void TrackWorkHint(u32 id, const StringView& hint, u64 startTime = 0) = 0;
 		virtual void TrackWorkEnd(u32 id) = 0;
 	};
 
@@ -37,7 +37,7 @@ namespace uba
 		virtual void DoWork(u32 count = 1) = 0;
 
 		u32 TrackWorkStart(const StringView& desc, const Color& color);
-		void TrackWorkHint(u32 id, const StringView& hint);
+		void TrackWorkHint(u32 id, const StringView& hint, u64 startTime = 0);
 		void TrackWorkEnd(u32 id);
 
 		void SetWorkTracker(WorkTracker* workTracker) { m_workTracker = workTracker; }
@@ -82,7 +82,7 @@ namespace uba
 	{
 		TrackWorkScope() : tracker(nullptr), id(0) {}
 		TrackWorkScope(WorkTracker& t, const StringView& desc, const Color& color = ColorWork) : tracker(&t), id(t.TrackWorkStart(desc, color)) {}
-		void AddHint(const StringView& hint) { if (tracker) tracker->TrackWorkHint(id, hint); }
+		void AddHint(const StringView& hint, u64 startTime = 0) { if (tracker) tracker->TrackWorkHint(id, hint, startTime); }
 		~TrackWorkScope() { if (tracker) tracker->TrackWorkEnd(id); }
 		WorkTracker* tracker;
 		u32 id;
diff --git a/Engine/Source/Programs/UnrealBuildAccelerator/Visualizer/Private/UbaTraceReader.cpp b/Engine/Source/Programs/UnrealBuildAccelerator/Visualizer/Private/UbaTraceReader.cpp
--- a/Engine/Source/Programs/UnrealBuildAccelerator/Visualizer/Private/UbaTraceReader.cpp
+++ b/Engine/Source/Programs/UnrealBuildAccelerator/Visualizer/Private/UbaTraceReader.cpp
@@ -1165,12 +1165,13 @@ namespace uba
 		{
 			u32 workIndex = u32(reader.Read7BitEncoded());
 			u64 stringIndex = reader.Read7BitEncoded();
+			u64 startTime = reader.Read7BitEncoded();
 			auto findIt = m_activeWorkRecords.find(workIndex);
 			if (findIt == m_activeWorkRecords.end())
 				return false;
 			WorkRecordLocation active = findIt->second;
 			TraceView::WorkRecord& record = out.workTracks[active.track].records[active.index];
-			record.entries.emplace_back(time, out.strings[stringIndex]); 
+			record.entries.emplace_back(time, startTime, out.strings[stringIndex]); 
 			break;
 		}
 		default:
diff --git a/Engine/Source/Programs/UnrealBuildAccelerator/Visualizer/Private/UbaTraceReader.h b/Engine/Source/Programs/UnrealBuildAccelerator/Visualizer/Private/UbaTraceReader.h
--- a/Engine/Source/Programs/UnrealBuildAccelerator/Visualizer/Private/UbaTraceReader.h
+++ b/Engine/Source/Programs/UnrealBuildAccelerator/Visualizer/Private/UbaTraceReader.h
@@ -53,6 +53,7 @@ namespace uba
 		struct WorkRecordLogEntry
 		{
 			u64 time = 0;
+			u64 startTime = 0;
 			const tchar* text;
 		};
 
diff --git a/Engine/Source/Programs/UnrealBuildAccelerator/Visualizer/Private/UbaVisualizer.cpp b/Engine/Source/Programs/UnrealBuildAccelerator/Visualizer/Private/UbaVisualizer.cpp
--- a/Engine/Source/Programs/UnrealBuildAccelerator/Visualizer/Private/UbaVisualizer.cpp
+++ b/Engine/Source/Programs/UnrealBuildAccelerator/Visualizer/Private/UbaVisualizer.cpp
@@ -2058,8 +2058,34 @@ namespace uba
 			logger.Info(L"  Duration:  %ls", TimeToText(duration, true).str);
 			if (!selectedWork.entries.empty())
 				logger.AddSpace();
-			for (auto& e : selectedWork.entries)
-				logger.Info(L"   %ls (%ls)", e.text, TimeToText(e.time - selectedWork.start).str);
+			for (u32 i=0,e=u32(selectedWork.entries.size()); i!=e; ++i)
+			{
+				auto& entry = selectedWork.entries[i];
+				u64 time = 0;
+				if (i == 0)
+				{
+					u64 start = entry.startTime ? entry.startTime : entry.time;
+					if (TimeToMs(start - selectedWork.start) > 1)
+						logger.Info(L"   Start (%ls)", TimeToText(start - selectedWork.start).str);
+				}
+
+				if (entry.startTime)
+					time = entry.time - entry.startTime;
+				
+				if (!time)
+					for (u32 j=i+1;j<e;++j)
+					{
+						auto& next = selectedWork.entries[j];
+						if (next.startTime)
+							continue;
+						time = next.time - entry.time;
+						break;
+					}
+				
+				if (!time && selectedWork.stop != ~0u)
+					time = selectedWork.stop - entry.time;
+				logger.Info(L"%s  %ls (%ls)", (entry.startTime ? L" " : L""), entry.text, TimeToText(time).str);
+			}
 			logger.AddSpace();
 			logger.DrawAtCursor();
 		}
