[UBA]
* Added TrackHintScope to make it easier to add performance tracking

diff --git a/Engine/Source/Programs/UnrealBuildAccelerator/Common/Private/UbaCacheClient.cpp b/Engine/Source/Programs/UnrealBuildAccelerator/Common/Private/UbaCacheClient.cpp
--- a/Engine/Source/Programs/UnrealBuildAccelerator/Common/Private/UbaCacheClient.cpp
+++ b/Engine/Source/Programs/UnrealBuildAccelerator/Common/Private/UbaCacheClient.cpp
@@ -29,8 +29,10 @@
 
 #if UBA_TRACE_FETCH_CACHE
 #define UBA_TRACE_FETCH_HINT(text) tws.AddHint(AsView(TC(text)));
+#define UBA_TRACE_FETCH_HINT_SCOPE(text) TrackHintScope ths(tws, AsView(TC(text)));
 #else
 #define UBA_TRACE_FETCH_HINT(text)
+#define UBA_TRACE_FETCH_HINT_SCOPE(text)
 #endif
 
 namespace uba
@@ -735,29 +737,21 @@ namespace uba
 					Atomic<bool> fetchSuccess = true;
 					m_storage.m_workManager->ParallelFor(u32(casKeyOffsets.size() - 1), casKeyOffsets, [&](auto& it)
 						{
-							#if UBA_TRACE_FETCH_CACHE
-							u64 time = GetTime();
-							#endif
 							u32 casKeyOffset = *it;
-							if (!FetchCasTable(bucket, cacheStats, casKeyOffset))
 							{
-								fetchSuccess = false;
-								return;
+								UBA_TRACE_FETCH_HINT_SCOPE("FetchCasTable")
+								if (!FetchCasTable(bucket, cacheStats, casKeyOffset))
+								{
+									fetchSuccess = false;
+									return;
+								}
 							}
 
-							#if UBA_TRACE_FETCH_CACHE
-							tws.AddHint(AsView(TC("FetchCasTable")), time);
-							time = GetTime();
-							#endif
-
+							UBA_TRACE_FETCH_HINT_SCOPE("FetchFile")
 							TimerScope fts(cacheStats.fetchOutput);
 							if (!FetchFile(bucket, rootPaths, info, cacheStats, storageStats, casKeyOffset))
 								fetchSuccess = false;
 
-							#if UBA_TRACE_FETCH_CACHE
-							tws.AddHint(AsView(TC("FetchFile")), time);
-							#endif
-
 						}, AsView(TC("CacheFetchOutput")));
 					if (!fetchSuccess)
 						return false;
diff --git a/Engine/Source/Programs/UnrealBuildAccelerator/Common/Public/UbaWorkManager.h b/Engine/Source/Programs/UnrealBuildAccelerator/Common/Public/UbaWorkManager.h
--- a/Engine/Source/Programs/UnrealBuildAccelerator/Common/Public/UbaWorkManager.h
+++ b/Engine/Source/Programs/UnrealBuildAccelerator/Common/Public/UbaWorkManager.h
@@ -4,6 +4,7 @@
 
 #include "UbaEvent.h"
 #include "UbaStringBuffer.h"
+#include "UbaTimer.h"
 
 namespace uba
 {
@@ -88,7 +89,14 @@ namespace uba
 		u32 id;
 	};
 
-
+	struct TrackHintScope
+	{
+		TrackHintScope(TrackWorkScope& t, const StringView& h) : tws(t), hint(h), startTime(GetTime()) {}
+		~TrackHintScope() { tws.AddHint(hint, startTime); }
+		TrackWorkScope& tws;
+		const StringView& hint;
+		u64 startTime;
+	};
 
 	template<typename TContainer, typename TFunc>
 	void WorkManager::ParallelFor(u32 workCount, TContainer& container, const TFunc& func, const StringView& description, bool highPriority)
