[UBA]
* Fixed so times for work hints are showing properly when trace is produced on mac.
* Added more work trace scopes in cache client

diff --git a/Engine/Source/Programs/UnrealBuildAccelerator/Common/Private/UbaCacheClient.cpp b/Engine/Source/Programs/UnrealBuildAccelerator/Common/Private/UbaCacheClient.cpp
--- a/Engine/Source/Programs/UnrealBuildAccelerator/Common/Private/UbaCacheClient.cpp
+++ b/Engine/Source/Programs/UnrealBuildAccelerator/Common/Private/UbaCacheClient.cpp
@@ -541,7 +541,7 @@ namespace uba
 					cachedIsMatch = &insres.first->second;
 				}
 
-				if (!FetchCasTable(bucket, cacheStats, casKeyOffset))
+				if (!FetchCasTable(tws, bucket, cacheStats, casKeyOffset))
 					return false;
 
 				if (!GetLocalPathAndCasKey(bucket, rootPaths, path, cacheCasKey, bucket.serverCasKeyTable, bucket.serverPathTable, casKeyOffset))
@@ -738,13 +738,10 @@ namespace uba
 					m_storage.m_workManager->ParallelFor(u32(casKeyOffsets.size() - 1), casKeyOffsets, [&](auto& it)
 						{
 							u32 casKeyOffset = *it;
+							if (!FetchCasTable(tws, bucket, cacheStats, casKeyOffset))
 							{
-								UBA_TRACE_FETCH_HINT_SCOPE("FetchCasTable")
-								if (!FetchCasTable(bucket, cacheStats, casKeyOffset))
-								{
-									fetchSuccess = false;
-									return;
-								}
+								fetchSuccess = false;
+								return;
 							}
 
 							UBA_TRACE_FETCH_HINT_SCOPE("FetchFile")
@@ -1166,11 +1163,13 @@ namespace uba
 		return false;
 	}
 
-	bool CacheClient::FetchCasTable(Bucket& bucket, CacheStats& stats, u32 requiredCasTableOffset)
+	bool CacheClient::FetchCasTable(TrackWorkScope& tws, Bucket& bucket, CacheStats& stats, u32 requiredCasTableOffset)
 	{
 		if (HasEnoughData(bucket, requiredCasTableOffset))
 			return true;
 
+		UBA_TRACE_FETCH_HINT_SCOPE("FetchCasTable")
+
 		if (requiredCasTableOffset > 256*1024*1024)
 			return m_logger.Warning(TC("Cas table offset %u too large. Cache entry corrupt (Bucket %llu)"), bucket.id);
 
diff --git a/Engine/Source/Programs/UnrealBuildAccelerator/Common/Public/UbaCacheClient.h b/Engine/Source/Programs/UnrealBuildAccelerator/Common/Public/UbaCacheClient.h
--- a/Engine/Source/Programs/UnrealBuildAccelerator/Common/Public/UbaCacheClient.h
+++ b/Engine/Source/Programs/UnrealBuildAccelerator/Common/Public/UbaCacheClient.h
@@ -77,7 +77,7 @@ namespace uba
 		UBA_NOINLINE bool SendCasTable(Bucket& bucket, u32 requiredCasTableSize);
 		UBA_NOINLINE bool SendCacheEntry(TrackWorkScope& tws, Bucket& bucket, const RootPaths& rootPaths, const CasKey& cmdKey, const Map<u32, u32>& inputsStringToCasKey, const Map<u32, u32>& outputsStringToCasKey, const u8* logLines, u64 logLinesSize, u64& outBytesSent);
 		UBA_NOINLINE bool SendCacheEntryMessage(BinaryReader& out, Bucket& bucket, const CasKey& cmdKey, const Map<u32, u32>& inputsStringToCasKey, const Map<u32, u32>& outputsStringToCasKey, const u8* logLines, u64 logLinesSize);
-		bool FetchCasTable(Bucket& bucket, CacheStats& stats, u32 requiredCasTableOffset);
+		bool FetchCasTable(TrackWorkScope& tws, Bucket& bucket, CacheStats& stats, u32 requiredCasTableOffset);
 		bool FetchFile(Bucket& bucket, const RootPaths& rootPaths, const ProcessStartInfo& info, CacheStats& cacheStats, StorageStats& storageStats, u32 casKeyOffset);
 		template<typename TableType> bool FetchCompactTable(u32 bucketId, TableType& table, u32 requiredTableOffset, u8 messageType);
 		bool HasEnoughData(Bucket& bucket, u32 requiredCasTableOffset);
diff --git a/Engine/Source/Programs/UnrealBuildAccelerator/Visualizer/Private/UbaTraceReader.cpp b/Engine/Source/Programs/UnrealBuildAccelerator/Visualizer/Private/UbaTraceReader.cpp
--- a/Engine/Source/Programs/UnrealBuildAccelerator/Visualizer/Private/UbaTraceReader.cpp
+++ b/Engine/Source/Programs/UnrealBuildAccelerator/Visualizer/Private/UbaTraceReader.cpp
@@ -1165,7 +1165,7 @@ namespace uba
 		{
 			u32 workIndex = u32(reader.Read7BitEncoded());
 			u64 stringIndex = reader.Read7BitEncoded();
-			u64 startTime = reader.Read7BitEncoded();
+			u64 startTime = ConvertTime(out, reader.Read7BitEncoded());
 			auto findIt = m_activeWorkRecords.find(workIndex);
 			if (findIt == m_activeWorkRecords.end())
 				return false;
diff --git a/Engine/Source/Programs/UnrealBuildAccelerator/Visualizer/Private/UbaVisualizer.cpp b/Engine/Source/Programs/UnrealBuildAccelerator/Visualizer/Private/UbaVisualizer.cpp
--- a/Engine/Source/Programs/UnrealBuildAccelerator/Visualizer/Private/UbaVisualizer.cpp
+++ b/Engine/Source/Programs/UnrealBuildAccelerator/Visualizer/Private/UbaVisualizer.cpp
@@ -307,6 +307,7 @@ namespace uba
 		m_looping = true;
 		if (!StartHwndThread())
 			return true;
+		m_usingNamed = true;
 		m_newTraceName.Append(namedTrace);
 		PostNewTrace(0, false);
 		return true;
@@ -3223,7 +3224,7 @@ namespace uba
 				m_traceView.finished = m_replay == 0;
 				PostNewTitle(GetTitlePrefix(title).Appendf(L"%s (v%u) - %s", m_fileName.data, m_traceView.version, GetWorldTime(0).data));
 			}
-			else
+			else if (m_usingNamed)
 			{
 				if (!m_trace.StartReadNamed(m_traceView, m_newTraceName.data, true, m_replay != 0))
 					return false;
@@ -3292,7 +3293,7 @@ namespace uba
 					if (!m_trace.UpdateReadClient(m_traceView, *m_client, changed))
 						m_clientDisconnect.Set();
 				}
-				else
+				else if (m_usingNamed)
 				{
 					if (!m_trace.UpdateReadNamed(m_traceView, m_replay ? timeOffset : ~u64(0), changed))
 						m_listenTimeout.Set();
@@ -3576,7 +3577,7 @@ namespace uba
 					m_traceView.finished = false;
 					if (!m_fileName.IsEmpty())
 						m_trace.UpdateReadFile(m_traceView, time, changed);
-					else
+					else if (m_usingNamed)
 						m_trace.UpdateReadNamed(m_traceView, time, changed);
 
 					m_pauseStart = m_startTime + time;
diff --git a/Engine/Source/Programs/UnrealBuildAccelerator/Visualizer/Private/UbaVisualizer.h b/Engine/Source/Programs/UnrealBuildAccelerator/Visualizer/Private/UbaVisualizer.h
--- a/Engine/Source/Programs/UnrealBuildAccelerator/Visualizer/Private/UbaVisualizer.h
+++ b/Engine/Source/Programs/UnrealBuildAccelerator/Visualizer/Private/UbaVisualizer.h
@@ -257,6 +257,7 @@ namespace uba
 		u32 m_workTrack = ~0u;
 		u32 m_workIndex = ~0u;
 		
+		bool m_usingNamed = false;
 		bool m_mouseOverWindow = false;
 		bool m_showPopup = false;
 		bool m_locked = false;
