[UBA]
* Changed access to m_fileTableLookup and m_casLookup to use readlock first since these are usually already populated so 99.9% of all finds will be successful

diff --git a/Engine/Source/Programs/UnrealBuildAccelerator/Common/Private/UbaStorage.cpp b/Engine/Source/Programs/UnrealBuildAccelerator/Common/Private/UbaStorage.cpp
--- a/Engine/Source/Programs/UnrealBuildAccelerator/Common/Private/UbaStorage.cpp
+++ b/Engine/Source/Programs/UnrealBuildAccelerator/Common/Private/UbaStorage.cpp
@@ -910,10 +910,7 @@ namespace uba
 	bool StorageImpl::AddCasFile(StringKey fileNameKey, const tchar* fileName, const CasKey& casKey, bool deferCreation, bool fileIsCompressed)
 	{
 		UBA_ASSERTF(IsCompressed(casKey) == m_storeCompressed, TC("CasKey compress mode must match storage compress mode. Compress mode Key: %i, Store: %i Cas: %s (%s)"), (IsCompressed(casKey) ? 1 : 0), (m_storeCompressed ? 1 : 0), CasKeyString(casKey).str, fileName);
-		SCOPED_WRITE_LOCK(m_casLookupLock, lookupLock);
-		auto insres = m_casLookup.try_emplace(casKey, casKey);
-		CasEntry& casEntry = insres.first->second;
-		lookupLock.Leave();
+		CasEntry& casEntry = GetOrCreateCasEntry(casKey);
 		CasEntryAccessed(casEntry);
 
 		SCOPED_WRITE_LOCK(casEntry.lock, entryLock);
@@ -1269,12 +1266,28 @@ namespace uba
 		return logger.Error(TC("Needs exclusive access to storage %s. Another process is running"), rootDir.data);
 	}
 
-	StorageImpl::FileEntry& StorageImpl::GetOrCreateFileEntry(StringKey fileNameKey)
+	StorageImpl::FileEntry& StorageImpl::GetOrCreateFileEntry(const StringKey& fileNameKey)
 	{
-		SCOPED_WRITE_LOCK(m_fileTableLookupLock, lock);
+		SCOPED_READ_LOCK(m_fileTableLookupLock, lock);
+		auto findIt = m_fileTableLookup.find(fileNameKey);
+		if (findIt != m_fileTableLookup.end())
+			return findIt->second;
+		lock.Leave();
+		SCOPED_WRITE_LOCK(m_fileTableLookupLock, lock2);
 		return m_fileTableLookup.try_emplace(fileNameKey).first->second;
 	}
 
+	StorageImpl::CasEntry& StorageImpl::GetOrCreateCasEntry(const CasKey& casKey)
+	{
+		SCOPED_READ_LOCK(m_casLookupLock, lock);
+		auto findIt = m_casLookup.find(casKey);
+		if (findIt != m_casLookup.end())
+			return findIt->second;
+		lock.Leave();
+		SCOPED_WRITE_LOCK(m_casLookupLock, lock2);
+		return m_casLookup.try_emplace(casKey, casKey).first->second;
+	}
+
 	bool StorageImpl::LoadCasTable(bool logStats, bool alwaysCheckAllFiles)
 	{
 		if (!m_writeToDisk)
@@ -2114,10 +2127,7 @@ namespace uba
 
 	bool StorageImpl::EnsureCasFile(const CasKey& casKey, const tchar* fileName)
 	{
-		SCOPED_WRITE_LOCK(m_casLookupLock, lookupLock);
-		auto insres = m_casLookup.try_emplace(casKey, casKey);
-		CasEntry& casEntry = insres.first->second;
-		lookupLock.Leave();
+		CasEntry& casEntry = GetOrCreateCasEntry(casKey);
 		CasEntryAccessed(casEntry);
 
 		SCOPED_WRITE_LOCK(casEntry.lock, entryLock);
diff --git a/Engine/Source/Programs/UnrealBuildAccelerator/Common/Public/UbaStorage.h b/Engine/Source/Programs/UnrealBuildAccelerator/Common/Public/UbaStorage.h
--- a/Engine/Source/Programs/UnrealBuildAccelerator/Common/Public/UbaStorage.h
+++ b/Engine/Source/Programs/UnrealBuildAccelerator/Common/Public/UbaStorage.h
@@ -185,8 +185,10 @@ namespace uba
 		void HandleOverflow(UnorderedSet<CasKey>* outDeletedFiles);
 
 		struct FileEntry;
-		FileEntry& GetOrCreateFileEntry(StringKey fileNameKey);
+		FileEntry& GetOrCreateFileEntry(const StringKey& fileNameKey);
 
+		struct CasEntry;
+		CasEntry& GetOrCreateCasEntry(const CasKey& casKey);
 
 		WorkManager* m_workManager;
 		MutableLogger m_logger;
