[UBA]
* Fixed macos compile error

diff --git a/Engine/Source/Programs/UnrealBuildAccelerator/Common/Public/UbaStorage.h b/Engine/Source/Programs/UnrealBuildAccelerator/Common/Public/UbaStorage.h
--- a/Engine/Source/Programs/UnrealBuildAccelerator/Common/Public/UbaStorage.h
+++ b/Engine/Source/Programs/UnrealBuildAccelerator/Common/Public/UbaStorage.h
@@ -186,8 +186,6 @@ namespace uba
 
 		struct FileEntry;
 		FileEntry& GetOrCreateFileEntry(const StringKey& fileNameKey);
-
-		struct CasEntry;
 		CasEntry& GetOrCreateCasEntry(const CasKey& casKey);
 
 		WorkManager* m_workManager;
