[UBA]
* Optimizations related to readerwriterlocks on non windows platforms
* Added more profiling scopes in cache client

diff --git a/Engine/Source/Programs/UnrealBuildAccelerator/Common/Private/UbaCacheClient.cpp b/Engine/Source/Programs/UnrealBuildAccelerator/Common/Private/UbaCacheClient.cpp
--- a/Engine/Source/Programs/UnrealBuildAccelerator/Common/Private/UbaCacheClient.cpp
+++ b/Engine/Source/Programs/UnrealBuildAccelerator/Common/Private/UbaCacheClient.cpp
@@ -521,6 +521,11 @@ namespace uba
 		UnorderedMap<StringKey, CasKey> normalizedCasKeys;
 		UnorderedMap<u32, bool> isCasKeyMatchCache;
 
+		#if UBA_TRACE_FETCH_CACHE
+		u64 storeTime = 0;
+		auto addStoreKeyHint = [&]() { tws.AddHint(AsView(TC("StoreCasKey")), GetTime() - storeTime); storeTime = 0; };
+		#endif
+
 		auto IsCasKeyMatch = [&](bool& outIsMatch, u32 casKeyOffset, u32 entryIndex, bool useLookup)
 			{
 				outIsMatch = false;
@@ -553,6 +558,7 @@ namespace uba
 					auto insres2 = normalizedCasKeys.try_emplace(ToStringKeyNoCheck(path.data, path.count));
 					if (insres2.second)
 					{
+						UBA_TRACE_FETCH_HINT_SCOPE("NormalizeAndHash")
 						TimerScope ts(cacheStats.normalizeFile);
 						localCasKey = rootPaths.NormalizeAndHashFile(m_logger, path.data);
 						if (localCasKey != CasKeyZero)
@@ -575,7 +581,17 @@ namespace uba
 						PreparseDirectory(fileNameKey, path);
 
 					bool fileIsCompressed = IsFileCompressed(info, path, false);
+
+					#if UBA_TRACE_FETCH_CACHE
+					u64 t = GetTime();
+					#endif
+
 					m_storage.StoreCasKey(localCasKey, fileNameKey, path.data, CasKeyZero, fileIsCompressed);
+
+					#if UBA_TRACE_FETCH_CACHE
+					storeTime += GetTime() - t;
+					#endif
+
 					UBA_ASSERT(localCasKey == CasKeyZero || IsCompressed(localCasKey));
 				}
 
@@ -653,6 +669,10 @@ namespace uba
 				range.begin = 0;
 				range.end = 0;
 			}
+
+			#if UBA_TRACE_FETCH_CACHE
+			addStoreKeyHint();
+			#endif
 		}
 
 		// Read entries
@@ -711,6 +731,10 @@ namespace uba
 						continue;
 				}
 
+				#if UBA_TRACE_FETCH_CACHE
+				addStoreKeyHint();
+				#endif
+
 				if (!m_useCacheHit)
 					return false;
 
@@ -1467,16 +1491,29 @@ namespace uba
 		dirHasher.Update(path.data, path.count);
 		StringKey pathKey = ToStringKey(dirHasher);
 
-		SCOPED_WRITE_LOCK(m_directoryPreparserLock, preparserLock);
-		auto insres = m_directoryPreparser.try_emplace(pathKey);
-		PreparedDir& dir = insres.first->second;
-		preparserLock.Leave();
+		PreparedDir* dir = nullptr;
+		#if !PLATFORM_WINDOWS
+		{
+			SCOPED_READ_LOCK(m_directoryPreparserLock, preparserLock);
+			auto findIt = m_directoryPreparser.find(pathKey);
+			if (findIt != m_directoryPreparser.end())
+				dir = &findIt->second;
+		}
+		#endif
 
-		SCOPED_WRITE_LOCK(dir.lock, preparserLock2);
-		if (dir.done)
+		if (!dir)
+		{
+			SCOPED_WRITE_LOCK(m_directoryPreparserLock, preparserLock);
+			dir = &m_directoryPreparser.try_emplace(pathKey).first->second;
+		}
+
+		if (dir->done)
 			return;
 
-		dir.done = true;
+		SCOPED_WRITE_LOCK(dir->lock, lock);
+		if (dir->done)
+			return;
+		dir->done = true;
 
 		// It is likely this folder has already been handled by session if this file is verified
 		if (m_storage.IsFileVerified(fileNameKey))
diff --git a/Engine/Source/Programs/UnrealBuildAccelerator/Common/Private/UbaStorage.cpp b/Engine/Source/Programs/UnrealBuildAccelerator/Common/Private/UbaStorage.cpp
--- a/Engine/Source/Programs/UnrealBuildAccelerator/Common/Private/UbaStorage.cpp
+++ b/Engine/Source/Programs/UnrealBuildAccelerator/Common/Private/UbaStorage.cpp
@@ -2004,6 +2004,18 @@ namespace uba
 	bool StorageImpl::StoreCasKey(CasKey& out, const StringKey& fileNameKey, const tchar* fileName, const CasKey& casKeyOverride, bool fileIsCompressed)
 	{
 		FileEntry& fileEntry = GetOrCreateFileEntry(fileNameKey);
+
+		#if !PLATFORM_WINDOWS
+		{
+			SCOPED_READ_LOCK(fileEntry.lock, entryLock);
+			if (fileEntry.verified)
+			{
+				out = fileEntry.casKey;
+				return true;
+			}
+		}
+		#endif
+
 		SCOPED_WRITE_LOCK(fileEntry.lock, entryLock);
 	
 		if (fileEntry.verified)
diff --git a/Engine/Source/Programs/UnrealBuildAccelerator/Common/Public/UbaCacheClient.h b/Engine/Source/Programs/UnrealBuildAccelerator/Common/Public/UbaCacheClient.h
--- a/Engine/Source/Programs/UnrealBuildAccelerator/Common/Public/UbaCacheClient.h
+++ b/Engine/Source/Programs/UnrealBuildAccelerator/Common/Public/UbaCacheClient.h
@@ -114,7 +114,7 @@ namespace uba
 		ReaderWriterLock m_sendOneAtTheTimeLock;
 
 		ReaderWriterLock m_directoryPreparserLock;
-		struct PreparedDir { ReaderWriterLock lock; bool done = false; };
+		struct PreparedDir { ReaderWriterLock lock; Atomic<bool> done; };
 		UnorderedMap<StringKey, PreparedDir> m_directoryPreparser;
 
 		CacheClient(const CacheClient&) = delete;
diff --git a/Engine/Source/Programs/UnrealBuildAccelerator/Core/Private/UbaSynchronization.cpp b/Engine/Source/Programs/UnrealBuildAccelerator/Core/Private/UbaSynchronization.cpp
--- a/Engine/Source/Programs/UnrealBuildAccelerator/Core/Private/UbaSynchronization.cpp
+++ b/Engine/Source/Programs/UnrealBuildAccelerator/Core/Private/UbaSynchronization.cpp
@@ -3,10 +3,12 @@
 #include "UbaSynchronization.h"
 #include "UbaPlatform.h"
 #include "UbaStringBuffer.h"
+#include <shared_mutex>
+
+#define UBA_USE_SHARED_MUTEX 0
 
 namespace uba
 {
-
 	CriticalSection::CriticalSection()
 	{
 		#if PLATFORM_WINDOWS
@@ -64,6 +66,10 @@ namespace uba
 		static_assert(sizeof(data) >= sizeof(SRWLOCK));
 		static_assert(alignof(SRWLOCK) == alignof(ReaderWriterLock));
 		InitializeSRWLock((SRWLOCK*)&data);
+		#elif UBA_USE_SHARED_MUTEX
+		static_assert(alignof(std::shared_mutex) == alignof(ReaderWriterLock));
+		static_assert(sizeof(data) >= sizeof(std::shared_mutex));
+		new (data) std::shared_mutex();
 		#else
 		static_assert(alignof(pthread_rwlock_t) == alignof(ReaderWriterLock));
 		static_assert(sizeof(data) >= sizeof(pthread_rwlock_t));
@@ -81,6 +87,8 @@ namespace uba
 		else
 			UBA_ASSERT(false);
 		#endif
+		#elif UBA_USE_SHARED_MUTEX
+		((std::shared_mutex*)data)->~shared_mutex();
 		#else
 		int res = pthread_rwlock_destroy((pthread_rwlock_t*)data);(void)res;
 		UBA_ASSERTF(res == 0, TC("pthread_rwlock_destroy failed: %i"), res);
@@ -91,6 +99,8 @@ namespace uba
 	{
 		#if PLATFORM_WINDOWS
 		AcquireSRWLockShared((SRWLOCK*)&data);
+		#elif UBA_USE_SHARED_MUTEX
+		((std::shared_mutex*)data)->lock_shared();
 		#else
 		pthread_rwlock_rdlock((pthread_rwlock_t*)data);
 		#endif
@@ -100,6 +110,8 @@ namespace uba
 	{
 		#if PLATFORM_WINDOWS
 		ReleaseSRWLockShared((SRWLOCK*)&data);
+		#elif UBA_USE_SHARED_MUTEX
+		((std::shared_mutex*)data)->unlock_shared();
 		#else
 		pthread_rwlock_unlock((pthread_rwlock_t*)data);
 		#endif
@@ -109,6 +121,8 @@ namespace uba
 	{
 		#if PLATFORM_WINDOWS
 		AcquireSRWLockExclusive((SRWLOCK*)&data);
+		#elif UBA_USE_SHARED_MUTEX
+		((std::shared_mutex*)data)->lock();
 		#else
 		pthread_rwlock_wrlock((pthread_rwlock_t*)data);
 		#endif
@@ -118,6 +132,8 @@ namespace uba
 	{
 		#if PLATFORM_WINDOWS
 		return TryAcquireSRWLockExclusive((SRWLOCK*)&data) != 0;
+		#elif UBA_USE_SHARED_MUTEX
+		return ((std::shared_mutex*)data)->try_lock();
 		#else
 		return pthread_rwlock_trywrlock((pthread_rwlock_t*)data) == 0;
 		#endif
@@ -127,6 +143,8 @@ namespace uba
 	{
 		#if PLATFORM_WINDOWS
 		ReleaseSRWLockExclusive((SRWLOCK*)&data);
+		#elif UBA_USE_SHARED_MUTEX
+		((std::shared_mutex*)data)->unlock();
 		#else
 		pthread_rwlock_unlock((pthread_rwlock_t*)data);
 		#endif
