[UBA]
* Collapse work scopes in ubavisualizer

diff --git a/Engine/Source/Programs/UnrealBuildAccelerator/Visualizer/Private/UbaTraceReader.cpp b/Engine/Source/Programs/UnrealBuildAccelerator/Visualizer/Private/UbaTraceReader.cpp
--- a/Engine/Source/Programs/UnrealBuildAccelerator/Visualizer/Private/UbaTraceReader.cpp
+++ b/Engine/Source/Programs/UnrealBuildAccelerator/Visualizer/Private/UbaTraceReader.cpp
@@ -1169,9 +1169,36 @@ namespace uba
 			auto findIt = m_activeWorkRecords.find(workIndex);
 			if (findIt == m_activeWorkRecords.end())
 				return false;
+			const tchar* text = out.strings[stringIndex];
 			WorkRecordLocation active = findIt->second;
 			TraceView::WorkRecord& record = out.workTracks[active.track].records[active.index];
-			record.entries.emplace_back(time, startTime, out.strings[stringIndex]); 
+
+			bool handled = false;
+			// Collapse if already exists
+			if (startTime)
+			{
+				for (auto rit=record.entries.rbegin(), rend=record.entries.rend(); rit!=rend; ++rit)
+				{
+					auto& entry = *rit;
+					if (!entry.startTime)
+						break;
+					if (entry.text != text)
+						continue;
+					++entry.count;
+					u64 entryTime = entry.time - entry.startTime;
+					u64 newTime = time - startTime;
+					if (newTime > entryTime)
+					{
+						entry.time = time;
+						entry.startTime = startTime;
+					}
+					handled = true;
+					break;
+				}
+			}
+
+			if (!handled)
+				record.entries.emplace_back(time, startTime, text); 
 			break;
 		}
 		default:
diff --git a/Engine/Source/Programs/UnrealBuildAccelerator/Visualizer/Private/UbaTraceReader.h b/Engine/Source/Programs/UnrealBuildAccelerator/Visualizer/Private/UbaTraceReader.h
--- a/Engine/Source/Programs/UnrealBuildAccelerator/Visualizer/Private/UbaTraceReader.h
+++ b/Engine/Source/Programs/UnrealBuildAccelerator/Visualizer/Private/UbaTraceReader.h
@@ -55,6 +55,7 @@ namespace uba
 			u64 time = 0;
 			u64 startTime = 0;
 			const tchar* text;
+			u32 count = 1;
 		};
 
 		struct WorkRecord
diff --git a/Engine/Source/Programs/UnrealBuildAccelerator/Visualizer/Private/UbaVisualizer.cpp b/Engine/Source/Programs/UnrealBuildAccelerator/Visualizer/Private/UbaVisualizer.cpp
--- a/Engine/Source/Programs/UnrealBuildAccelerator/Visualizer/Private/UbaVisualizer.cpp
+++ b/Engine/Source/Programs/UnrealBuildAccelerator/Visualizer/Private/UbaVisualizer.cpp
@@ -2085,7 +2085,10 @@ namespace uba
 				
 				if (!time && selectedWork.stop != ~0u)
 					time = selectedWork.stop - entry.time;
-				logger.Info(L"%s  %ls (%ls)", (entry.startTime ? L" " : L""), entry.text, TimeToText(time).str);
+				if (entry.count == 1)
+					logger.Info(L"%s  %ls (%ls)", (entry.startTime ? L" " : L""), entry.text, TimeToText(time).str);
+				else
+					logger.Info(L"%s  %ls (%ls %u)", (entry.startTime ? L" " : L""), entry.text, TimeToText(time).str, entry.count);
 			}
 			logger.AddSpace();
 			logger.DrawAtCursor();
