[UBA]
* Added special error message for when UbaDetours.dll and application does not match

diff --git a/Engine/Source/Programs/UnrealBuildAccelerator/Common/Private/UbaProcess.cpp b/Engine/Source/Programs/UnrealBuildAccelerator/Common/Private/UbaProcess.cpp
--- a/Engine/Source/Programs/UnrealBuildAccelerator/Common/Private/UbaProcess.cpp
+++ b/Engine/Source/Programs/UnrealBuildAccelerator/Common/Private/UbaProcess.cpp
@@ -300,7 +300,9 @@ namespace uba
 
 				ImageInfo imageInfo;
 				bool machineIsArm = IsArmBinary;
-				if (!GetImageInfo(imageInfo, m_session.m_logger, m_realApplication.c_str(), true))
+				if (exitCode == 1398)
+					err.Appendf(TC("UbaDetours.dll has a different version than " UBA_BINARY));
+				else if (!GetImageInfo(imageInfo, m_session.m_logger, m_realApplication.c_str(), true))
 					err.Appendf(TC("Failed to load %s"), m_realApplication.c_str());
 				else if (!imageInfo.is64Bit)
 					err.Append(TC("Doesn't seem to be a 64-bit executable"));
@@ -311,7 +313,8 @@ namespace uba
 				else if (!machineIsArm && imageInfo.isArm64)
 					err.Appendf(TC("Machine is x64 and image is arm64"));
 
-				err.Appendf(TC(" (GetExitCodeProcess returned 0x%x (%s Size: %llu, CasKey: %s)"), exitCode, m_realApplication.c_str(), imageInfo.fileSize, CasKeyString(imageInfo.key).str);
+				if (exitCode != 1398)
+					err.Appendf(TC(" (GetExitCodeProcess returned 0x%x (%s Size: %llu, CasKey: %s)"), exitCode, m_realApplication.c_str(), imageInfo.fileSize, CasKeyString(imageInfo.key).str);
 			}
 
 			if (err.IsEmpty())
diff --git a/Engine/Source/Programs/UnrealBuildAccelerator/Core/Public/UbaProtocol.h b/Engine/Source/Programs/UnrealBuildAccelerator/Core/Public/UbaProtocol.h
--- a/Engine/Source/Programs/UnrealBuildAccelerator/Core/Public/UbaProtocol.h
+++ b/Engine/Source/Programs/UnrealBuildAccelerator/Core/Public/UbaProtocol.h
@@ -48,7 +48,7 @@ namespace uba
 		#undef UBA_PROCESS_MESSAGE
 	};
 
-	inline constexpr u32 ProcessMessageVersion = 1339;
+	inline constexpr u32 ProcessMessageVersion = 1340;
 
 	inline constexpr u32 CommunicationMemSize = IsWindows ? 64*1024 : 64*1024*2; // Macos expands some commandlines to be crazy long
 
diff --git a/Engine/Source/Programs/UnrealBuildAccelerator/Detours/Private/Windows/UbaDetoursMainWin.cpp b/Engine/Source/Programs/UnrealBuildAccelerator/Detours/Private/Windows/UbaDetoursMainWin.cpp
--- a/Engine/Source/Programs/UnrealBuildAccelerator/Detours/Private/Windows/UbaDetoursMainWin.cpp
+++ b/Engine/Source/Programs/UnrealBuildAccelerator/Detours/Private/Windows/UbaDetoursMainWin.cpp
@@ -154,7 +154,7 @@ BOOL WINAPI DllMain(HINSTANCE hinst, DWORD dwReason, LPVOID reserved)
 		if (!payload)
 			TerminateCurrentProcess(1342);
 		if (payload->version != ProcessMessageVersion)
-			TerminateCurrentProcess(1354);
+			TerminateCurrentProcess(1398);
 
 		PreInit(*payload);
 
diff --git a/Engine/Source/Programs/UnrealBuildAccelerator/UbaAgent.Target.cs b/Engine/Source/Programs/UnrealBuildAccelerator/UbaAgent.Target.cs
--- a/Engine/Source/Programs/UnrealBuildAccelerator/UbaAgent.Target.cs
+++ b/Engine/Source/Programs/UnrealBuildAccelerator/UbaAgent.Target.cs
@@ -52,10 +52,13 @@ public class UbaAgentTarget : TargetRules
 
 		var BinariesFolder = Path.Combine("Binaries", Target.Platform.ToString(), "UnrealBuildAccelerator");
 
+		string BinaryName = null;
+
 		if (Target.Platform.IsInGroup(UnrealPlatformGroup.Windows))
 		{
 			string BinaryExt = Rules.bShouldCompileAsDLL ? ".dll" : ".exe";
-			Rules.OutputFile = Path.Combine(BinariesFolder, Target.Architectures.SingleArchitecture.bIsX64 ? "x64" : "arm64", $"{Rules.LaunchModuleName}{BinaryExt}");
+			BinaryName = $"{Rules.LaunchModuleName}{BinaryExt}";
+			Rules.OutputFile = Path.Combine(BinariesFolder, Target.Architectures.SingleArchitecture.bIsX64 ? "x64" : "arm64", BinaryName);
 
 			Rules.GlobalDefinitions.AddRange(new string[] {
 				"UBA_DETOURS_LIBRARY=L\"UbaDetours.dll\"",
@@ -72,7 +75,8 @@ public class UbaAgentTarget : TargetRules
 		{
 			string BinaryPrefix = Rules.bShouldCompileAsDLL ? "lib" : string.Empty;
 			string BinaryExt = Rules.bShouldCompileAsDLL ? ".so" : string.Empty;
-			Rules.OutputFile = Path.Combine(BinariesFolder, $"{BinaryPrefix}{Rules.LaunchModuleName}{BinaryExt}");
+			BinaryName = $"{BinaryPrefix}{Rules.LaunchModuleName}{BinaryExt}";
+			Rules.OutputFile = Path.Combine(BinariesFolder, BinaryName);
 
 			Rules.GlobalDefinitions.AddRange(new string[] {
 				"UBA_DETOURS_LIBRARY=\"libUbaDetours.so\"",
@@ -97,7 +101,8 @@ public class UbaAgentTarget : TargetRules
 			}
 			string BinaryPrefix = Rules.bShouldCompileAsDLL ? "lib" : string.Empty;
 			string BinaryExt = Rules.bShouldCompileAsDLL ? ".dylib" : string.Empty;
-			Rules.OutputFile = Path.Combine(BinariesFolder, $"{BinaryPrefix}{Rules.LaunchModuleName}{BinaryExt}");
+			BinaryName = $"{BinaryPrefix}{Rules.LaunchModuleName}{BinaryExt}";
+			Rules.OutputFile = Path.Combine(BinariesFolder, BinaryName);
 			Rules.GlobalDefinitions.AddRange(new string[] {
 				"UBA_DETOURS_LIBRARY=\"libUbaDetours.dylib\"",
 				"UBA_AGENT_EXECUTABLE=\"UbaAgent\"",
@@ -108,6 +113,8 @@ public class UbaAgentTarget : TargetRules
 			Rules.GlobalDefinitions.Add("PLATFORM_MAC=0");
 		}
 
+		Rules.GlobalDefinitions.Add($"UBA_BINARY=\"{BinaryName}\"");
+
 		int useMiMalloc = 0;
 		if (Target.Platform.IsInGroup(UnrealPlatformGroup.Windows) || (!Rules.bShouldCompileAsDLL && Target.Platform.IsInGroup(UnrealPlatformGroup.Linux)))
 		{
