[UBA]
* Fixed bug running link.exe on executable where .lib file was not properly written if it already existed on disk (has been flying under the radar since noone uses the lib file

diff --git a/Engine/Source/Programs/UnrealBuildAccelerator/Common/Private/UbaSession.cpp b/Engine/Source/Programs/UnrealBuildAccelerator/Common/Private/UbaSession.cpp
--- a/Engine/Source/Programs/UnrealBuildAccelerator/Common/Private/UbaSession.cpp
+++ b/Engine/Source/Programs/UnrealBuildAccelerator/Common/Private/UbaSession.cpp
@@ -2751,7 +2751,6 @@ namespace uba
 				Storage::GetMappingString(name, file.mappingHandle, 0);
 				SCOPED_WRITE_LOCK(m_fileMappingTableLookupLock, lookupLock);
 				auto insres = m_fileMappingTableLookup.try_emplace(file.key);
-				UBA_ASSERTF(insres.second, TC("Failed to add %s to file mapping, already exists"), file.name.c_str());
 				FileMappingEntry& entry = insres.first->second;
 				lookupLock.Leave();
 				SCOPED_WRITE_LOCK(entry.lock, entryCs);
diff --git a/Engine/Source/Programs/UnrealBuildAccelerator/Detours/Private/Windows/UbaDetoursFunctionsKernelBase.inl b/Engine/Source/Programs/UnrealBuildAccelerator/Detours/Private/Windows/UbaDetoursFunctionsKernelBase.inl
--- a/Engine/Source/Programs/UnrealBuildAccelerator/Detours/Private/Windows/UbaDetoursFunctionsKernelBase.inl
+++ b/Engine/Source/Programs/UnrealBuildAccelerator/Detours/Private/Windows/UbaDetoursFunctionsKernelBase.inl
@@ -1583,9 +1583,21 @@ bool Shared_MoveFile(LPCWSTR lpExistingFileName, LPCWSTR lpNewFileName, DWORD dw
 			lock2.Leave();
 			FileInfo& destInfo = insres.first->second;
 			UBA_ASSERTF(!insres.second, TC("%s -> %s"), lpExistingFileName, lpNewFileName); // This is here just to get a chance to investigate this scenario.. might work
-			UBA_ASSERTF(!destInfo.trueFileMapHandle && (!destInfo.memoryFile || g_rules->IsThrowAway(dest, g_runningRemote)), TC("Moving file %s to %s that is an output file that is not a memory file is not supported"), source.data, lpNewFileName);
+
+			HANDLE tempHandle = INVALID_HANDLE_VALUE;
+			if (destInfo.isFileMap || (destInfo.memoryFile && destInfo.memoryFile->isLocalOnly)) // File has been read before, let's just ignore that and take the new memoryFile
+			{
+				destInfo.isFileMap = false;
+				destInfo.trueFileMapHandle = 0;
+				destInfo.trueFileMapOffset = 0;
+				tempHandle = CreateFile(lpNewFileName, GENERIC_WRITE, 0, NULL, CREATE_ALWAYS, FILE_ATTRIBUTE_NORMAL, NULL);
+				if (tempHandle == INVALID_HANDLE_VALUE)
+					return false;
+			}
 			destInfo.memoryFile = sourceInfo.memoryFile;
 			sourceInfo.memoryFile = nullptr;
+
+			CloseHandle(tempHandle);
 			DEBUG_LOG_DETOURED(L"MoveFileExW", L"(memfile->memfile) %ls to %ls -> Success", lpExistingFileName, lpNewFileName);
 			SetLastError(ERROR_SUCCESS);
 			return true;
