[UBA]
* Changed so event dtor preserves last error code
* Fixed bug where visualizer listen path didn't always work
* Changed MessageBox to logger.error (which also creates messageboxes in visualizer)

diff --git a/Engine/Source/Programs/UnrealBuildAccelerator/Common/Private/UbaTrace.cpp b/Engine/Source/Programs/UnrealBuildAccelerator/Common/Private/UbaTrace.cpp
--- a/Engine/Source/Programs/UnrealBuildAccelerator/Common/Private/UbaTrace.cpp
+++ b/Engine/Source/Programs/UnrealBuildAccelerator/Common/Private/UbaTrace.cpp
@@ -464,7 +464,7 @@ namespace uba
 		m_memHandle = uba::CreateMemoryMappingW(m_logger, PAGE_READWRITE, 256, channelName, TC("TraceChannel"));
 		if (!m_memHandle.IsValid())
 		{
-			MessageBox(NULL, TC("Failed to create file mapping for trace channel"), TC("UbaVisualizer"), 0);
+			m_logger.Error(TC("Failed to create file mapping %s for trace channel (%s)"), channelName, LastErrorToText().data);
 			return false;
 		}
 		bool isCreator = GetLastError() != ERROR_ALREADY_EXISTS;
@@ -474,7 +474,7 @@ namespace uba
 		m_mem = MapViewOfFile(m_logger, m_memHandle, FILE_MAP_WRITE, 0, 256);
 		if (!m_mem)
 		{
-			MessageBox(NULL, TC("Failed to map file mapping for uba trace channel"), TC("UbaVisualizer"), 0);
+			m_logger.Error(TC("Failed to map file mapping for uba trace channel"));
 			return false;
 		}
 
diff --git a/Engine/Source/Programs/UnrealBuildAccelerator/Core/Private/UbaEvent.cpp b/Engine/Source/Programs/UnrealBuildAccelerator/Core/Private/UbaEvent.cpp
--- a/Engine/Source/Programs/UnrealBuildAccelerator/Core/Private/UbaEvent.cpp
+++ b/Engine/Source/Programs/UnrealBuildAccelerator/Core/Private/UbaEvent.cpp
@@ -321,8 +321,13 @@ namespace uba
 	void Event::Destroy()
 	{
 		#if PLATFORM_WINDOWS
-		CloseHandle(m_ev);
-		m_ev = nullptr;
+		if (m_ev)
+		{
+			u32 lastError = GetLastError();
+			CloseHandle(m_ev);
+			SetLastError(lastError);
+			m_ev = nullptr;
+		}
 		#else
 		((EventImpl&)m_data).Destroy();
 		#endif
diff --git a/Engine/Source/Programs/UnrealBuildAccelerator/Visualizer/Private/UbaVisualizer.cpp b/Engine/Source/Programs/UnrealBuildAccelerator/Visualizer/Private/UbaVisualizer.cpp
--- a/Engine/Source/Programs/UnrealBuildAccelerator/Visualizer/Private/UbaVisualizer.cpp
+++ b/Engine/Source/Programs/UnrealBuildAccelerator/Visualizer/Private/UbaVisualizer.cpp
@@ -290,6 +290,7 @@ namespace uba
 				if (!traceName.Equals(m_newTraceName.data) && traceName.EndsWith(filter.data))
 				{
 					m_newTraceName.Clear().Append(traceName);
+					m_usingNamed = true;
 					PostNewTrace(0, false);
 				}
 			}
@@ -307,8 +308,8 @@ namespace uba
 		m_looping = true;
 		if (!StartHwndThread())
 			return true;
-		m_usingNamed = true;
 		m_newTraceName.Append(namedTrace);
+		m_usingNamed = true;
 		PostNewTrace(0, false);
 		return true;
 	}
