[UBA]
* Attempt at using less stack in FetchFromCache path

diff --git a/Engine/Source/Programs/UnrealBuildAccelerator/Common/Private/UbaCacheClient.cpp b/Engine/Source/Programs/UnrealBuildAccelerator/Common/Private/UbaCacheClient.cpp
--- a/Engine/Source/Programs/UnrealBuildAccelerator/Common/Private/UbaCacheClient.cpp
+++ b/Engine/Source/Programs/UnrealBuildAccelerator/Common/Private/UbaCacheClient.cpp
@@ -464,7 +464,12 @@ namespace uba
 		if (cmdKey == CasKeyZero)
 			return false;
 
+#if PLATFORM_MAC
+		u8* memory = (u8*)malloc(SendMaxSize);
+		auto g = MakeGuard([memory]() { free(memory); });
+#else
 		u8 memory[SendMaxSize];
+#endif
 
 		u32 fetchId = m_session.CreateProcessId();
 		m_session.GetTrace().CacheBeginFetch(fetchId, info.GetDescription());
@@ -472,14 +477,14 @@ namespace uba
 		auto tg = MakeGuard([&]()
 			{
 				cacheStats.testEntry.time -= (cacheStats.fetchCasTable.time + cacheStats.normalizeFile.time);
-				BinaryWriter writer(memory, 0, sizeof_array(memory));
+				BinaryWriter writer(memory, 0, SendMaxSize);
 				cacheStats.Write(writer);
 				storageStats.Write(writer);
 				kernelStats.Write(writer);
 				m_session.GetTrace().CacheEndFetch(fetchId, success, memory, writer.GetPosition());
 			});
 
-		BinaryReader reader(memory, 0, sizeof_array(memory));
+		BinaryReader reader(memory, 0, SendMaxSize);
 
 		SCOPED_WRITE_LOCK(m_bucketsLock, bucketsLock);
 		Bucket& bucket = m_buckets.try_emplace(bucketId, bucketId).first->second;
@@ -518,18 +523,14 @@ namespace uba
 		struct MissInfo { TString path; u32 entryIndex; CasKey cache; CasKey local; };
 		Vector<MissInfo> misses;
 
-		UnorderedMap<StringKey, CasKey> normalizedCasKeys;
-		UnorderedMap<u32, bool> isCasKeyMatchCache;
-
 		#if UBA_TRACE_FETCH_CACHE
 		u64 storeTime = 0;
 		auto addStoreKeyHint = [&]() { tws.AddHint(AsView(TC("StoreCasKey")), GetTime() - storeTime); storeTime = 0; };
 		#endif
 
-		auto IsCasKeyMatch = [&](bool& outIsMatch, u32 casKeyOffset, u32 entryIndex, bool useLookup)
+		auto IsCasKeyMatch = [&, normalizedCasKeys = UnorderedMap<StringKey, CasKey>(), isCasKeyMatchCache = UnorderedMap<u32, bool>()](bool& outIsMatch, u32 casKeyOffset, u32 entryIndex, bool useLookup) mutable
 			{
 				outIsMatch = false;
-				StringBuffer<MaxPath> path;
 
 				CasKey cacheCasKey;
 				CasKey localCasKey;
@@ -549,6 +550,7 @@ namespace uba
 				if (!FetchCasTable(tws, bucket, cacheStats, casKeyOffset))
 					return false;
 
+				StringBuffer<MaxPath> path;
 				if (!GetLocalPathAndCasKey(bucket, rootPaths, path, cacheCasKey, bucket.serverCasKeyTable, bucket.serverPathTable, casKeyOffset))
 					return false;
 				UBA_ASSERTF(IsCompressed(cacheCasKey), TC("Cache entry for %s has uncompressed cache key for path %s (%s)"), info.GetDescription(), path.data, CasKeyString(cacheCasKey).str);
@@ -1191,7 +1193,11 @@ namespace uba
 	{
 		if (HasEnoughData(bucket, requiredCasTableOffset))
 			return true;
+		return FetchCasTable2(tws, bucket, stats, requiredCasTableOffset);
+	}
 
+	bool CacheClient::FetchCasTable2(TrackWorkScope& tws, Bucket& bucket, CacheStats& stats, u32 requiredCasTableOffset)
+	{
 		UBA_TRACE_FETCH_HINT_SCOPE("FetchCasTable")
 
 		if (requiredCasTableOffset > 256*1024*1024)
diff --git a/Engine/Source/Programs/UnrealBuildAccelerator/Common/Public/UbaCacheClient.h b/Engine/Source/Programs/UnrealBuildAccelerator/Common/Public/UbaCacheClient.h
--- a/Engine/Source/Programs/UnrealBuildAccelerator/Common/Public/UbaCacheClient.h
+++ b/Engine/Source/Programs/UnrealBuildAccelerator/Common/Public/UbaCacheClient.h
@@ -78,12 +78,13 @@ namespace uba
 		UBA_NOINLINE bool SendCacheEntry(TrackWorkScope& tws, Bucket& bucket, const RootPaths& rootPaths, const CasKey& cmdKey, const Map<u32, u32>& inputsStringToCasKey, const Map<u32, u32>& outputsStringToCasKey, const u8* logLines, u64 logLinesSize, u64& outBytesSent);
 		UBA_NOINLINE bool SendCacheEntryMessage(BinaryReader& out, Bucket& bucket, const CasKey& cmdKey, const Map<u32, u32>& inputsStringToCasKey, const Map<u32, u32>& outputsStringToCasKey, const u8* logLines, u64 logLinesSize);
 		bool FetchCasTable(TrackWorkScope& tws, Bucket& bucket, CacheStats& stats, u32 requiredCasTableOffset);
-		bool FetchFile(Bucket& bucket, const RootPaths& rootPaths, const ProcessStartInfo& info, CacheStats& cacheStats, StorageStats& storageStats, u32 casKeyOffset);
+		UBA_NOINLINE bool FetchCasTable2(TrackWorkScope& tws, Bucket& bucket, CacheStats& stats, u32 requiredCasTableOffset);
+		UBA_NOINLINE bool FetchFile(Bucket& bucket, const RootPaths& rootPaths, const ProcessStartInfo& info, CacheStats& cacheStats, StorageStats& storageStats, u32 casKeyOffset);
 		template<typename TableType> bool FetchCompactTable(u32 bucketId, TableType& table, u32 requiredTableOffset, u8 messageType);
 		bool HasEnoughData(Bucket& bucket, u32 requiredCasTableOffset);
 		bool HasEnoughCasData(Bucket& bucket, u32 requiredCasTableOffset, u32& outPathOffset);
 		bool HasEnoughPathData(Bucket& bucket, u32 requiredPathTableOffset);
-		bool ReportUsedEntry(Vector<ProcessLogLine>& outLogLines, bool ownedLogLines, Bucket& bucket, const CasKey& cmdKey, u32 entryId);
+		UBA_NOINLINE bool ReportUsedEntry(Vector<ProcessLogLine>& outLogLines, bool ownedLogLines, Bucket& bucket, const CasKey& cmdKey, u32 entryId);
 		bool PopulateLogLines(Vector<ProcessLogLine>& outLogLines, const u8* mem, u64 memLen);
 
 		UBA_NOINLINE CasKey GetCmdKey(const RootPaths& rootPaths, const ProcessStartInfo& info, bool report, u32 bucketId);
@@ -92,7 +93,7 @@ namespace uba
 
 		bool GetLocalPathAndCasKey(Bucket& bucket, const RootPaths& rootPaths, StringBufferBase& outPath, CasKey& outKey, CompactCasKeyTable& casKeyTable, CompactPathTable& pathTable, u32 offset);
 		bool IsFileCompressed(const ProcessStartInfo& info, const StringView& filename, bool isOutput);
-		void PreparseDirectory(const StringKey& fileNameKey, const StringBufferBase& filePath);
+		UBA_NOINLINE void PreparseDirectory(const StringKey& fileNameKey, const StringBufferBase& filePath);
 
 		MutableLogger m_logger;
 		StorageImpl& m_storage;
