[UBA]
* Enabled linger (10 seconds) for mac since it seems like the socket stays around after crashes

diff --git a/Engine/Source/Programs/UnrealBuildAccelerator/Common/Private/UbaNetworkBackendTcp.cpp b/Engine/Source/Programs/UnrealBuildAccelerator/Common/Private/UbaNetworkBackendTcp.cpp
--- a/Engine/Source/Programs/UnrealBuildAccelerator/Common/Private/UbaNetworkBackendTcp.cpp
+++ b/Engine/Source/Programs/UnrealBuildAccelerator/Common/Private/UbaNetworkBackendTcp.cpp
@@ -94,6 +94,7 @@ namespace uba
 	bool SetKeepAlive(Logger& logger, SOCKET socket);
 	bool SetBlocking(Logger& logger, SOCKET socket, bool blocking);
 	bool SetTimeout(Logger& logger, SOCKET socket, u32 timeoutMs);
+	bool SetLinger(Logger& logger, SOCKET socket, u32 lingerSeconds);
 	bool SetRecvBuf(Logger& logger, SOCKET socket, u32 windowSize);
 	bool SetSendBuf(Logger& logger, SOCKET socket, u32 windowSize);
 	bool DisableNagle(Logger& logger, SOCKET socket);
@@ -446,6 +447,8 @@ namespace uba
 				continue;
 			}
 
+			SetLinger(logger, clientSocket, 10);
+
 			// Hard coding this values disable auto-tuning. Seems like creating multiple sockets is more efficient than playing around with these
 			//SetRecvBuf(logger, clientSocket, 50*1024*1024);
 			//SetSendBuf(logger, clientSocket, 50*1024*1024);
@@ -701,6 +704,8 @@ namespace uba
 		if (!SetKeepAlive(logger, socketFd))
 			return false;
 
+		SetLinger(logger, socketFd, 10);
+
 		// Hard coding this values disable auto-tuning. Seems like creating multiple sockets is more efficient than playing around with these
 		//SetRecvBuf(logger, socketFd, 50*1024*1024);
 		//SetSendBuf(logger, socketFd, 50*1024*1024);
@@ -792,6 +797,19 @@ namespace uba
 		return true;
 	}
 
+	bool SetLinger(Logger& logger, SOCKET socket, u32 lingerSeconds)
+	{
+#if PLATFORM_MAC // Mac does not seem to automatically close the socket if the process crashes
+		struct linger so_linger;
+		so_linger.l_onoff = 1;
+		so_linger.l_linger = u16(lingerSeconds);
+
+		if (setsockopt(socket, SOL_SOCKET, SO_LINGER, (const char*)&so_linger, sizeof(so_linger)) < 0)
+			return logger.Error(TC("setsockopt SO_LINGER error: (error: %s)"), LastErrorToText(WSAGetLastError()).data);
+#endif
+		return true;
+	}
+
 	bool SetRecvBuf(Logger& logger, SOCKET socket, u32 windowSize)
 	{
 		if (setsockopt(socket, SOL_SOCKET, SO_RCVBUF, (const char*)&windowSize, sizeof(windowSize)) == SOCKET_ERROR)
