[UBA]
* Fixed bug in include crawler logic after changing so FileMappingHandle can hold file handle
* Fixed so files in session bin dir is not kept in file->cas table (since every run has a different session bin dir)
* Fixed so server caskey calculations for known inputs are waited on at the end just before sending network response
* Fixed process wall time visualization bug in visualizer

diff --git a/Engine/Source/Programs/UnrealBuildAccelerator/Common/Private/UbaFileMapping.cpp b/Engine/Source/Programs/UnrealBuildAccelerator/Common/Private/UbaFileMapping.cpp
--- a/Engine/Source/Programs/UnrealBuildAccelerator/Common/Private/UbaFileMapping.cpp
+++ b/Engine/Source/Programs/UnrealBuildAccelerator/Common/Private/UbaFileMapping.cpp
@@ -627,8 +627,7 @@ namespace uba
 		if (!handleStrEnd)
 			return {};
 
-		FileMappingHandle h;
-		h.mh = (HANDLE)StringToValue(handleStr, handleStrEnd - handleStr);
+		FileMappingHandle h = FileMappingHandle::FromU64(StringToValue(handleStr, handleStrEnd - handleStr));
 		const tchar* mappingOffsetStr = handleStrEnd + 1;
 		u64 mappingOffset = StringToValue(mappingOffsetStr, wcslen(mappingOffsetStr));
 		
diff --git a/Engine/Source/Programs/UnrealBuildAccelerator/Common/Private/UbaSessionClient.cpp b/Engine/Source/Programs/UnrealBuildAccelerator/Common/Private/UbaSessionClient.cpp
--- a/Engine/Source/Programs/UnrealBuildAccelerator/Common/Private/UbaSessionClient.cpp
+++ b/Engine/Source/Programs/UnrealBuildAccelerator/Common/Private/UbaSessionClient.cpp
@@ -481,7 +481,9 @@ namespace uba
 				return false;
 		}
 
-		return m_storage.CopyOrLink(casKey, out.data, fileAttributes);
+		constexpr bool writeCompressed = false;
+		constexpr bool isTemp = true;
+		return m_storage.CopyOrLink(casKey, out.data, fileAttributes, writeCompressed, {}, isTemp);
 	}
 
 	bool SessionClient::ProcessThreadStart(ProcessImpl& process)
@@ -520,10 +522,10 @@ namespace uba
 
 		TrackWorkScope tws(m_client, AsView(TC("CreateFile")), ColorWork);
 		tws.AddHint(msg.fileName);
-		return CreateFileForRead(out, msg.fileName, msg.fileNameKey, msg.process.m_id, *msg.process.m_startInfo.rules);
+		return CreateFileForRead(out, tws, msg.fileName, msg.fileNameKey, msg.process.m_id, *msg.process.m_startInfo.rules);
 	}
 
-	bool SessionClient::CreateFileForRead(CreateFileResponse& out, const StringView& fileName, const StringKey& fileNameKey, u32 processId, const ApplicationRules& rules)
+	bool SessionClient::CreateFileForRead(CreateFileResponse& out, TrackWorkScope& tws, const StringView& fileName, const StringKey& fileNameKey, u32 processId, const ApplicationRules& rules)
 	{
 		CasKey casKey;
 		if (!GetCasKeyForFile(casKey, processId, fileName, fileNameKey))
@@ -620,6 +622,7 @@ namespace uba
 
 			if (newName[0] != '^')
 			{
+				TrackHintScope ths(tws, AsView(TC("CreateMemoryMap")));
 				if (!isDir && memoryMapAlignment)
 				{
 					if (retryEntry)
@@ -2113,7 +2116,7 @@ namespace uba
 					TrackWorkScope tws(m_client, AsView(TC("CrawlFile")), ColorWork);
 					tws.AddHint(fileName);
 
-					if (!CreateFileForRead(out, fileName, ToStringKey(fileName), processId, *rules))
+					if (!CreateFileForRead(out, tws, fileName, ToStringKey(fileName), processId, *rules))
 						return false;
 				}
 
diff --git a/Engine/Source/Programs/UnrealBuildAccelerator/Common/Private/UbaSessionServer.cpp b/Engine/Source/Programs/UnrealBuildAccelerator/Common/Private/UbaSessionServer.cpp
--- a/Engine/Source/Programs/UnrealBuildAccelerator/Common/Private/UbaSessionServer.cpp
+++ b/Engine/Source/Programs/UnrealBuildAccelerator/Common/Private/UbaSessionServer.cpp
@@ -1025,7 +1025,7 @@ namespace uba
 		if (!success)
 		{
 			auto logType = connectionInfo.ShouldDisconnect() ? LogEntryType_Info : LogEntryType_Warning;
-			m_logger.Logf(logType, TC("Failed to retrieve cas for %s from client (Needed to write %s)"), CasKeyString(casKey).str, destination.data);
+			m_logger.Logf(logType, TC("Failed to retrieve cas for %s from client %u (Needed to write %s)"), CasKeyString(casKey).str, connectionInfo.GetId(), destination.data);
 		}
 
 		bool shouldWriteToDisk = ShouldWriteToDisk(destination);
@@ -1250,7 +1250,7 @@ namespace uba
 		u32 weight32 = reader.ReadU32();
 		float availableWeight = *(float*)&weight32;
 
-		Vector<RemoteProcess::KnownInput*> knownInputsToSend;
+		Vector<RemoteProcess*> processesWithKnownInputsToSend;
 
 		float weightLeft = availableWeight;
 		u32 addCount = 0;
@@ -1306,12 +1306,8 @@ namespace uba
 			writer.WriteU32(process->m_processId);
 			startInfo.Write(writer, applicationOverride);
 
-
 			if (process->m_knownInputsDone.IsCreated())
-				if (process->m_knownInputsDone.IsSet(50*1000))
-					for (auto kiIt = process->m_knownInputs, kiEnd = kiIt + process->m_knownInputsCount; kiIt!=kiEnd; ++kiIt)
-						if (session.sentKeys.insert(kiIt->key).second)
-							knownInputsToSend.push_back(kiIt);
+				processesWithKnownInputsToSend.push_back(process);
 
 			++addCount;
 
@@ -1366,6 +1362,15 @@ namespace uba
 		writer.WriteU32(neededDirectoryTableSize);
 		writer.WriteU32(neededHashTableSize);
 
+
+		// Collect known inputs
+		Vector<RemoteProcess::KnownInput*> knownInputsToSend;
+		for (auto process : processesWithKnownInputsToSend)
+			if (process->m_knownInputsDone.IsSet(50*1000))
+				for (auto kiIt = process->m_knownInputs, kiEnd = kiIt + process->m_knownInputsCount; kiIt!=kiEnd; ++kiIt)
+					if (session.sentKeys.insert(kiIt->key).second)
+						knownInputsToSend.push_back(kiIt);
+
 		// Send caskeys of known inputs so client can start retrieving them straight away
 		u32 kiCapacity = u32(writer.GetCapacityLeft() - sizeof(u32)) / sizeof(RemoteProcess::KnownInput);
 		u32 toSendCount = Min(kiCapacity, u32(knownInputsToSend.size()));
diff --git a/Engine/Source/Programs/UnrealBuildAccelerator/Common/Private/UbaStorage.cpp b/Engine/Source/Programs/UnrealBuildAccelerator/Common/Private/UbaStorage.cpp
--- a/Engine/Source/Programs/UnrealBuildAccelerator/Common/Private/UbaStorage.cpp
+++ b/Engine/Source/Programs/UnrealBuildAccelerator/Common/Private/UbaStorage.cpp
@@ -1505,7 +1505,7 @@ namespace uba
 			{
 				FileEntry& fileEntry = pair.second;
 				//SCOPED_READ_LOCK(fileEntry.lock, entryLock); // TODO: Remove this? tsan complains about it
-				if (fileEntry.casKey == CasKeyZero || fileEntry.lastWritten == 0)
+				if (fileEntry.casKey == CasKeyZero || fileEntry.lastWritten == 0 || fileEntry.isTemp)
 					continue;
 				writer.WriteStringKey(pair.first);
 				writer.WriteU64(fileEntry.size);
@@ -2142,6 +2142,17 @@ namespace uba
 		CasEntry& casEntry = GetOrCreateCasEntry(casKey);
 		CasEntryAccessed(casEntry);
 
+		{
+			SCOPED_READ_LOCK(casEntry.lock, entryLock);
+			if (casEntry.verified)
+			{
+				if (casEntry.exists)
+					return true;
+				if (!fileName)
+					return false;
+			}
+		}
+
 		SCOPED_WRITE_LOCK(casEntry.lock, entryLock);
 
 		if (casEntry.verified)
@@ -2299,7 +2310,7 @@ namespace uba
 		return out != CasKeyZero;
 	}
 
-	bool StorageImpl::CopyOrLink(const CasKey& casKey, const tchar* destination, u32 fileAttributes, bool writeCompressed, const FormattingFunc& formattingFunc)
+	bool StorageImpl::CopyOrLink(const CasKey& casKey, const tchar* destination, u32 fileAttributes, bool writeCompressed, const FormattingFunc& formattingFunc, bool isTemp)
 	{
 		UBA_ASSERT(casKey != CasKeyZero);
 		UBA_ASSERT(fileAttributes);
@@ -2404,6 +2415,7 @@ namespace uba
 
 				SCOPED_WRITE_LOCK(fileEntry.lock, entryLock);
 				fileEntry.verified = false;
+				fileEntry.isTemp = isTemp;
 
 				// This is to reduce number of active CreateFiles.. seems like machines don't like tons of CreateFile at the same time
 				#if PLATFORM_WINDOWS
diff --git a/Engine/Source/Programs/UnrealBuildAccelerator/Common/Private/UbaStorageServer.cpp b/Engine/Source/Programs/UnrealBuildAccelerator/Common/Private/UbaStorageServer.cpp
--- a/Engine/Source/Programs/UnrealBuildAccelerator/Common/Private/UbaStorageServer.cpp
+++ b/Engine/Source/Programs/UnrealBuildAccelerator/Common/Private/UbaStorageServer.cpp
@@ -350,7 +350,7 @@ namespace uba
 						if (timeSpentMs > 8 * 60 * 1000)
 						{
 							// Took more than 5 minutes to respond
-							m_logger.Warning(TC("Took more than 8 minutes to respond to FetchBegin (%u).. is this some sort of hang or just host being half dead?"), todoRemoveMe);
+							m_logger.Warning(TC("Took more than 8 minutes to respond to FetchBegin to client %u (Hang code %u).."), connectionInfo.GetId(), todoRemoveMe);
 						}
 					});
 
diff --git a/Engine/Source/Programs/UnrealBuildAccelerator/Common/Public/UbaSessionClient.h b/Engine/Source/Programs/UnrealBuildAccelerator/Common/Public/UbaSessionClient.h
--- a/Engine/Source/Programs/UnrealBuildAccelerator/Common/Public/UbaSessionClient.h
+++ b/Engine/Source/Programs/UnrealBuildAccelerator/Common/Public/UbaSessionClient.h
@@ -81,7 +81,7 @@ namespace uba
 		struct InternalProcessStartInfo;
 		struct ModuleInfo;
 
-		bool CreateFileForRead(CreateFileResponse& out, const StringView& fileName, const StringKey& fileNameKey, u32 processId, const ApplicationRules& rules);
+		bool CreateFileForRead(CreateFileResponse& out, TrackWorkScope& tws, const StringView& fileName, const StringKey& fileNameKey, u32 processId, const ApplicationRules& rules);
 		bool GetCasKeyForFile(CasKey& out, u32 processId, const StringView& fileName, const StringKey& fileNameKey);
 		bool ReadModules(List<ModuleInfo>& outModules, u32 processId, const tchar* application);
 		bool EnsureBinaryFile(StringBufferBase& out, StringBufferBase& outVirtual, u32 processId, const StringBufferBase& fileName, const StringKey& fileNameKey, const tchar* applicationDir, const u8* loaderPaths, u32 loaderPathsSize);
diff --git a/Engine/Source/Programs/UnrealBuildAccelerator/Common/Public/UbaStorage.h b/Engine/Source/Programs/UnrealBuildAccelerator/Common/Public/UbaStorage.h
--- a/Engine/Source/Programs/UnrealBuildAccelerator/Common/Public/UbaStorage.h
+++ b/Engine/Source/Programs/UnrealBuildAccelerator/Common/Public/UbaStorage.h
@@ -50,7 +50,7 @@ namespace uba
 		virtual bool CalculateCasKey(CasKey& out, const tchar* fileName) = 0;
 
 		using FormattingFunc = Function<bool(MemoryBlock& destData, const void* sourceData, u64 sourceSize, const tchar* hint)>;
-		virtual bool CopyOrLink(const CasKey& casKey, const tchar* destination, u32 fileAttributes, bool writeCompressed = false, const FormattingFunc& formattingFunc = {}) = 0;
+		virtual bool CopyOrLink(const CasKey& casKey, const tchar* destination, u32 fileAttributes, bool writeCompressed = false, const FormattingFunc& formattingFunc = {}, bool isTemp = false) = 0;
 		virtual bool FakeCopy(const CasKey& casKey, const tchar* destination, u64 size = 0, u64 lastWritten = 0, bool deleteExisting = true) = 0;
 		virtual bool GetCasFileName(StringBufferBase& out, const CasKey& casKey) = 0;
 
@@ -138,7 +138,7 @@ namespace uba
 		virtual bool DropCasFile(const CasKey& casKey, bool forceDelete, const tchar* hint) override;
 		virtual bool ReportBadCasFile(const CasKey& casKey) override;
 		virtual bool CalculateCasKey(CasKey& out, const tchar* fileName) override;
-		virtual bool CopyOrLink(const CasKey& casKey, const tchar* destination, u32 fileAttributes, bool writeCompressed = false, const FormattingFunc& formattingFunc = {}) override;
+		virtual bool CopyOrLink(const CasKey& casKey, const tchar* destination, u32 fileAttributes, bool writeCompressed = false, const FormattingFunc& formattingFunc = {}, bool isTemp = false) override;
 		virtual bool FakeCopy(const CasKey& casKey, const tchar* destination, u64 size = 0, u64 lastWritten = 0, bool deleteExisting = true) override;
 		virtual bool GetCasFileName(StringBufferBase& out, const CasKey& casKey) override;
 		virtual MappedView MapView(const CasKey& casKey, const tchar* hint) override;
@@ -204,6 +204,7 @@ namespace uba
 			u64 size = 0;
 			u64 lastWritten = 0;
 			bool verified = false;
+			bool isTemp = false;
 		};
 		UnorderedMap<StringKey, FileEntry> m_fileTableLookup;
 
diff --git a/Engine/Source/Programs/UnrealBuildAccelerator/Visualizer/Private/UbaVisualizer.cpp b/Engine/Source/Programs/UnrealBuildAccelerator/Visualizer/Private/UbaVisualizer.cpp
--- a/Engine/Source/Programs/UnrealBuildAccelerator/Visualizer/Private/UbaVisualizer.cpp
+++ b/Engine/Source/Programs/UnrealBuildAccelerator/Visualizer/Private/UbaVisualizer.cpp
@@ -1915,7 +1915,7 @@ namespace uba
 			logger.Info(L"  %ls", process.description.c_str());
 			logger.Info(L"  Host:        %ls", m_processSelectedLocation.sessionIndex == 0 ? TC("local") : m_traceView.GetSession(m_processSelectedLocation).name.c_str());
 			logger.Info(L"  ProcessId:  %6u", process.id);
-			logger.Info(L"  Start:     %7ls (%ls)", TimeToText(process.start, true).str, GetWorldTime(selectedWork.start).data);
+			logger.Info(L"  Start:     %7ls (%ls)", TimeToText(process.start, true).str, GetWorldTime(process.start).data);
 			logger.Info(L"  Duration:  %7ls", TimeToText(duration, true).str);
 			if (!process.returnedReason.empty())
 				logger.Info(L"  Returned: %7s", process.returnedReason.data());
