[UBA]
* Changed Event to reuse event handles for windows since we are creating/destroying a ton of events during a compile with remote helpers. Should probably do this for non-windows too but need two versions because of the shared events needing to be directly in memory
* Changed init of srwlock to use assignment instead of function call

diff --git a/Engine/Source/Programs/UnrealBuildAccelerator/Core/Private/UbaEvent.cpp b/Engine/Source/Programs/UnrealBuildAccelerator/Core/Private/UbaEvent.cpp
--- a/Engine/Source/Programs/UnrealBuildAccelerator/Core/Private/UbaEvent.cpp
+++ b/Engine/Source/Programs/UnrealBuildAccelerator/Core/Private/UbaEvent.cpp
@@ -279,6 +279,19 @@ namespace uba
 			UBA_ASSERTF(res == 0, "pthread_mutex_unlock failed (error code %i)", res);
 		}
 	};
+
+#else
+	struct EventImpl
+	{
+		HANDLE handle;
+		union
+		{
+			EventImpl* next;
+			bool manualReset;
+		};
+	};
+	ReaderWriterLock g_firstEventLock[2];
+	EventImpl* g_firstEvent[2];
 #endif
 
 
@@ -311,8 +324,28 @@ namespace uba
 	bool Event::Create(bool manualReset, bool shared)
 	{
 		#if PLATFORM_WINDOWS
-		m_ev = CreateEvent(nullptr, manualReset, false, NULL);
-		return m_ev != nullptr;
+		EventImpl* impl = nullptr;
+		SCOPED_WRITE_LOCK(g_firstEventLock[manualReset], lock);
+		auto*& first = g_firstEvent[manualReset];
+		if (first)
+		{
+			impl = first;
+			first = impl->next;
+		}
+		lock.Leave();
+
+		if (!impl)
+		{
+			impl = new EventImpl;
+			impl->handle = CreateEvent(nullptr, manualReset, false, NULL);
+			UBA_ASSERT(impl->handle);
+		}
+		else
+			ResetEvent(impl->handle);
+
+		impl->manualReset = manualReset;
+		m_ev = impl;
+		return true;
 		#else
 		return ((EventImpl&)m_data).Create(manualReset, shared);
 		#endif
@@ -323,10 +356,13 @@ namespace uba
 		#if PLATFORM_WINDOWS
 		if (m_ev)
 		{
-			u32 lastError = GetLastError();
-			CloseHandle(m_ev);
-			SetLastError(lastError);
+			auto impl = (EventImpl*)m_ev;
 			m_ev = nullptr;
+			bool manualReset = impl->manualReset;
+			auto*& first = g_firstEvent[manualReset];
+			SCOPED_WRITE_LOCK(g_firstEventLock[manualReset], l);
+			impl->next = first;
+			first = impl;
 		}
 		#else
 		((EventImpl&)m_data).Destroy();
@@ -336,7 +372,9 @@ namespace uba
 	void Event::Set()
 	{
 		#if PLATFORM_WINDOWS
-		SetEvent(m_ev);
+		if (!m_ev)
+			return;
+		SetEvent(((EventImpl*)m_ev)->handle);
 		#else
 		((EventImpl&)m_data).Set();
 		#endif
@@ -345,7 +383,9 @@ namespace uba
 	void Event::Reset()
 	{
 		#if PLATFORM_WINDOWS
-		ResetEvent(m_ev);
+		if (!m_ev)
+			return;
+		ResetEvent(((EventImpl*)m_ev)->handle);
 		#else
 		((EventImpl&)m_data).Reset();
 		#endif
@@ -363,7 +403,9 @@ namespace uba
 	bool Event::IsSet(u32 timeOutMs)
 	{
 		#if PLATFORM_WINDOWS
-		return WaitForSingleObject(m_ev, timeOutMs) == 0;
+		if (!m_ev)
+			return false;
+		return WaitForSingleObject(((EventImpl*)m_ev)->handle, timeOutMs) == 0;
 		#else
 		return ((EventImpl&)m_data).IsSet(timeOutMs);
 		#endif
@@ -372,7 +414,9 @@ namespace uba
 	void* Event::GetHandle()
 	{
 		#if PLATFORM_WINDOWS
-		return m_ev;
+		if (!m_ev)
+			return nullptr;
+		return ((EventImpl*)m_ev)->handle;
 		#else
 		UBA_ASSERTF(false, "Event::GetHandle not available");
 		return 0;
diff --git a/Engine/Source/Programs/UnrealBuildAccelerator/Core/Private/UbaSynchronization.cpp b/Engine/Source/Programs/UnrealBuildAccelerator/Core/Private/UbaSynchronization.cpp
--- a/Engine/Source/Programs/UnrealBuildAccelerator/Core/Private/UbaSynchronization.cpp
+++ b/Engine/Source/Programs/UnrealBuildAccelerator/Core/Private/UbaSynchronization.cpp
@@ -65,7 +65,7 @@ namespace uba
 		#if PLATFORM_WINDOWS
 		static_assert(sizeof(data) >= sizeof(SRWLOCK));
 		static_assert(alignof(SRWLOCK) == alignof(ReaderWriterLock));
-		InitializeSRWLock((SRWLOCK*)&data);
+		*(SRWLOCK*)&data = SRWLOCK_INIT;
 		#elif UBA_USE_SHARED_MUTEX
 		static_assert(alignof(std::shared_mutex) == alignof(ReaderWriterLock));
 		static_assert(sizeof(data) >= sizeof(std::shared_mutex));
