[UBA]
* Changed storage proxy to use work for sending fetchbegin messages. This should reduce contention on worker threads.. Looks like throughput becomes better when having very busy proxy

diff --git a/Engine/Source/Programs/UnrealBuildAccelerator/Common/Private/UbaStorageClient.cpp b/Engine/Source/Programs/UnrealBuildAccelerator/Common/Private/UbaStorageClient.cpp
--- a/Engine/Source/Programs/UnrealBuildAccelerator/Common/Private/UbaStorageClient.cpp
+++ b/Engine/Source/Programs/UnrealBuildAccelerator/Common/Private/UbaStorageClient.cpp
@@ -568,7 +568,8 @@ namespace uba
 			}
 			else
 			{
-				UBA_ASSERT(isCompressed); // Not implemented. Receiving non compressed and want to store it compressed.
+				if (!isCompressed)
+					return m_logger.Error(TC("Code path not implemented. Receiving non compressed cas %s and want to store it compressed (%s)"), casFile.data, hint);
 
 				bool sendSegmentMessage = responseSize == 0;
 				u64 leftUncompressed = actualSize;
diff --git a/Engine/Source/Programs/UnrealBuildAccelerator/Common/Private/UbaStorageProxy.cpp b/Engine/Source/Programs/UnrealBuildAccelerator/Common/Private/UbaStorageProxy.cpp
--- a/Engine/Source/Programs/UnrealBuildAccelerator/Common/Private/UbaStorageProxy.cpp
+++ b/Engine/Source/Programs/UnrealBuildAccelerator/Common/Private/UbaStorageProxy.cpp
@@ -10,9 +10,9 @@
 
 namespace uba
 {
-	struct StorageProxy::MessageInFlight
+	struct StorageProxy::SegmentMessage
 	{
-		MessageInFlight(StorageProxy& p, FileEntry& f, u8* readBuffer, u32 fi)
+		SegmentMessage(StorageProxy& p, FileEntry& f, u8* readBuffer, u32 fi)
 		:	proxy(p)
 		,	file(f)
 		,	message(p.m_client, ServiceId, StorageMessageType_FetchSegment, writer)
@@ -121,270 +121,270 @@ namespace uba
 
 	bool StorageProxy::HandleMessage(const ConnectionInfo& connectionInfo, const WorkContext& workContext, MessageInfo& messageInfo, BinaryReader& reader, BinaryWriter& writer)
 	{
-		LoggerWithWriter logger(g_consoleLogWriter, TC(""));
-
-		StackBinaryWriter<1024> writer2;
-		StackBinaryReader<SendMaxSize> reader2;
-
 		switch (messageInfo.type)
 		{
 		case StorageMessageType_Connect:
-			{
-				StringBuffer<> clientName;
-				reader.ReadString(clientName);
-				u32 clientVersion = reader.ReadU32();
-				if (clientVersion != StorageNetworkVersion)
-				{
-					m_logger.Error(TC("Different network versions. Client: %u, Server: %u. Disconnecting"), clientVersion, StorageNetworkVersion);
-					return false;
-				}
-				bool isInProcessClient = reader.ReadBool();
-				if (isInProcessClient)
-					m_inProcessClientId = connectionInfo.GetId();
+			return HandleConnect(connectionInfo, reader, writer);
+		case StorageMessageType_FetchBegin:
+			return HandleFetchBegin(connectionInfo, workContext, messageInfo, reader, writer);
+		case StorageMessageType_FetchSegment:
+			return HandleFetchSegment(connectionInfo, workContext, messageInfo, reader, writer);
+		case StorageMessageType_FetchEnd:
+			return true;
+		default:
+			return HandleDefault(messageInfo, reader, writer);
+		}
+	}
 
-				//m_logger.Info(TC("%s connected"), clientName.data);
-				writer.WriteGuid(m_storageServerUid);
-				return true;
-			}
+	bool StorageProxy::HandleConnect(const ConnectionInfo& connectionInfo, BinaryReader& reader, BinaryWriter& writer)
+	{
+		StringBuffer<> clientName;
+		reader.ReadString(clientName);
+		u32 clientVersion = reader.ReadU32();
+		if (clientVersion != StorageNetworkVersion)
+		{
+			m_logger.Error(TC("Different network versions. Client: %u, Server: %u. Disconnecting"), clientVersion, StorageNetworkVersion);
+			return false;
+		}
+		bool isInProcessClient = reader.ReadBool();
+		if (isInProcessClient)
+			m_inProcessClientId = connectionInfo.GetId();
 
-		case StorageMessageType_FetchBegin:
-			{
-				u8 recvFlags = reader.ReadByte(); // Wants proxy
-				UBA_ASSERT((recvFlags & 2) == 0);(void)recvFlags;
+		//m_logger.Info(TC("%s connected"), clientName.data);
+		writer.WriteGuid(m_storageServerUid);
+		return true;
+	}
 
-				CasKey casKey = reader.ReadCasKey();
-				StringBuffer<> hint;
-				reader.ReadString(hint);
+	struct StorageProxy::BeginMessage
+	{
+		BeginMessage(StorageProxy& p, FileEntry& f, const tchar* h) : proxy(p), file(f), hint(h), message(p.m_client, ServiceId, StorageMessageType_FetchBegin, writer) {}
+		StorageProxy& proxy;
+		FileEntry& file;
+		TString hint;
+		struct DeferredResponse { u32 clientId; MessageInfo info; };
+		List<DeferredResponse> deferredResponses;
+		StackBinaryWriter<1024> writer;
+		StackBinaryReader<SendMaxSize> reader;
+		NetworkMessage message;
+	};
 
-				workContext.tracker.AddHint(StringView(hint).GetFileName());
+	bool StorageProxy::HandleFetchBegin(const ConnectionInfo& connectionInfo, const WorkContext& workContext, MessageInfo& messageInfo, BinaryReader& reader, BinaryWriter& writer)
+	{
+		u8 recvFlags = reader.ReadByte(); // Wants proxy
+		UBA_ASSERT((recvFlags & 2) == 0);(void)recvFlags;
 
-				SCOPED_WRITE_LOCK(m_filesLock, filesLock);
-				auto insres = m_files.try_emplace(casKey);
-				FileEntry& file = insres.first->second;
-				filesLock.Leave();
+		CasKey casKey = reader.ReadCasKey();
+		StringBuffer<> hint;
+		reader.ReadString(hint);
 
-				SCOPED_WRITE_LOCK(file.lock, fileLock);
+		workContext.tracker.AddHint(StringView(hint).GetFileName());
 
-				while (true)
-				{
-					if (file.memory)
-						break;
+		SCOPED_WRITE_LOCK(m_filesLock, filesLock);
+		FileEntry& file = m_files.try_emplace(casKey).first->second;
+		filesLock.Leave();
 
-					bool storeCompressed = true;
-					if (m_useLocalStorage && m_localStorage && IsCompressed(casKey) && m_inProcessClientId && connectionInfo.GetId() != m_inProcessClientId)
-					{
-						// We need to leave this lock here since the in-process storage client might be asking for this file too and then we can end up in a deadlock
-						fileLock.Leave();
+		u32 clientId = connectionInfo.GetId();
 
-						bool hasCas = m_localStorage->EnsureCasFile(casKey, nullptr);
-						StringBuffer<> casFile;
-						hasCas = hasCas && m_localStorage->GetCasFileName(casFile, casKey);
+		SCOPED_WRITE_LOCK(file.lock, fileLock);
+		if (file.memory || file.error)
+		{
+			fileLock.Leave();
+			return SendFetchBeginResponse(clientId, file, writer);
+		}
 
-						// Enter lock again, and also check if another thread might have already handled this file while we looked if it existed in local storage
-						fileLock.Enter();
-						if (file.memory)
-							break;
+		UBA_ASSERT(!file.memory);
+		if (!GetFileFromLocalStorage(clientId, casKey, file, fileLock))
+			return false;
 
-						if (hasCas)
-						{
-							FileAccessor sourceFile(m_logger, casFile.data);
-							if (sourceFile.OpenMemoryRead())
-							{
-								u64 fileSize = sourceFile.GetSize();
-								file.memory = new u8[fileSize];
-								if (!file.memory)
-									return false;
-								file.size = fileSize;
-								file.received = fileSize;
-								memcpy(file.memory, sourceFile.GetData(), fileSize);
-								file.available = true;
-							}
-						}
-					}
+		if (file.memory || file.error)
+		{
+			fileLock.Leave();
+			return SendFetchBeginResponse(clientId, file, writer);
+		}
 
-					if (!file.memory)
-					{
-						file.trackId = m_client.TrackWorkStart(AsView(TC("ProxyFetch")), ColorWork);
-						m_client.TrackWorkHint(file.trackId, StringView(hint).GetFileName());
+		if (file.beginMessage == nullptr)
+		{
+			file.casKey = casKey;
+			file.beginMessage = new BeginMessage(*this, file, hint.data);
 
-						NetworkMessage msg(m_client, ServiceId, StorageMessageType_FetchBegin, writer2);
-						writer2.WriteByte(2); // Does not want proxy but informs it is proxy
-						writer2.WriteCasKey(casKey);
-						writer2.WriteString(hint);
-						writer2.WriteBytes(reader.GetPositionData(), reader.GetLeft());
+			file.trackId = m_client.TrackWorkStart(AsView(TC("ProxyFetch")), ColorWork);
+			m_client.TrackWorkHint(file.trackId, StringView(hint).GetFileName());
 
-						SCOPED_READ_LOCK(m_largeFileLock, largeFileLock);
-						if (!msg.Send(reader2))
-						{
-							file.error = true;
-							return m_logger.Error(TC("FetchBegin failed for cas file %s (%s). Requested by %s"), CasKeyString(casKey).str, hint.data, GuidToString(connectionInfo.GetUid()).str);
-						}
-						largeFileLock.Leave();
-
-						BinaryReader tempReader(reader2.GetPositionData(), 0, reader2.GetLeft());
-						u16 fetchId = tempReader.ReadU16();
-						if (fetchId == 0)
-						{
-							file.error = true;
-							m_logger.Error(TC("FetchBegin failed for cas file %s (%s). Requested by %s"), CasKeyString(casKey).str, hint.data, GuidToString(connectionInfo.GetUid()).str);
-							writer.WriteU16(0);
-							return true;
-						}
-						u64 fileSize = tempReader.Read7BitEncoded();
-						file.memory = new u8[fileSize];
-						file.size = fileSize;
-
-						u8 flags = tempReader.ReadByte();
-						storeCompressed = (flags >> 0) & 1;
-						bool sendEnd = (flags >> 1) & 1;
-						u64 fetchedSize = tempReader.GetLeft();
-
-						memcpy(file.memory, tempReader.GetPositionData(), fetchedSize);
-
-						if (sendEnd && fetchedSize == fileSize)
-							SendEnd(casKey);
-
-						file.received = fetchedSize;
-						file.fetchId = fetchId;
-						file.sendEnd = sendEnd;
-					}
+			auto& writer2 = file.beginMessage->writer;
+			writer2.WriteByte(2); // Does not want proxy but informs it is proxy
+			writer2.WriteCasKey(casKey);
+			writer2.WriteString(hint);
+			writer2.WriteBytes(reader.GetPositionData(), reader.GetLeft());
 
-					file.storeCompressed = storeCompressed;
-					file.casKey = casKey;
+			SCOPED_READ_LOCK(m_largeFileLock, largeFileLock);
 
-					if (file.received < file.size)
-					{
-						u64 left = file.size - file.received;
-						u64 segmentSize = m_client.GetMessageMaxSize() - 5; // This is server response size - header.. TODO: Should be taken from server
-						u32 segmentCount = u32((left + segmentSize - 1) / segmentSize);
-						file.messagesInFlight.resize(segmentCount);
-						for (u32 i=0; i!=segmentCount; ++i)
-						{
-							u64 offset = file.received + segmentSize * i;
-							u8* memory = file.memory + offset;
-							auto mif = new MessageInFlight(*this, file, memory, i);
-							file.messagesInFlight[i] = mif;
-						}
-
-						// Move the additional messages to a job to be able to return this one quickly.
-						m_server.AddWork([f = &file, segmentCount, this](const WorkContext&)
-							{
-								SCOPED_WRITE_LOCK(m_largeFileLock, lock);
-								//TrackWorkScope tws(m_client, TC("SEGMENTS"));
-								auto& file = *f;
-								for (u32 i=0; i!=segmentCount; ++i)
-								{
-									auto mif = file.messagesInFlight[i];
-
-									bool res = mif->message.SendAsync(mif->reader, [](bool error, void* userData)
-										{
-											auto mif = (MessageInFlight*)userData;
-											mif->error = error;
-											mif->proxy.m_server.AddWork([mif](const WorkContext&) { mif->proxy.HandleReceivedData(*mif); }, 1, TC("ProxyWaitMsg"), ColorWork);
-										}, mif);
-									if (!res)
-									{
-										// TODO: Don't leak mif
-										mif->error = true;
-									}
-								}
-							}, 1, TC("ProxySpawnMsg"), ColorWork);
-					}
-					else
-					{
-						m_client.TrackWorkEnd(file.trackId);
-					}
-					break;
-				}
+			bool res = file.beginMessage->message.SendAsync(file.beginMessage->reader, [](bool error, void* userData)
+				{
+					auto m = (BeginMessage*)userData;
+					m->proxy.m_server.AddWork([m, error](const WorkContext&) { m->proxy.HandleFetchBeginReceived(*m, error); }, 1, TC("ProxyFetchBegin"), ColorWork);
+				}, file.beginMessage);
 
-				if (file.error)
-					return false;
-				fileLock.Leave();
+			if (!res)
+			{
+				delete file.beginMessage;
+				file.beginMessage = nullptr;
+				file.error = true;
+				return SendFetchBeginResponse(clientId, file, writer);
+			}
+		}
 
-				u16 fetchId = u16(~0);
+		auto& deferredResponse = file.beginMessage->deferredResponses.emplace_back();
+		deferredResponse.clientId = clientId;
+		deferredResponse.info = messageInfo;
+		messageInfo = {};
+		return true;
+	}
+
+	void StorageProxy::HandleFetchBeginReceived(BeginMessage& m, bool error)
+	{
+		auto& file = m.file;
+		u8* memory = nullptr;
+
+		auto sendResponses = MakeGuard([&]()
+			{
+				SCOPED_WRITE_LOCK(file.lock, fileLock);
+				file.error = error;
+				file.memory = memory;
+				auto f = file.beginMessage;
+				file.beginMessage = nullptr;
+				fileLock.Leave();
 
-				u64 headerSize = sizeof(u16) + Get7BitEncodedCount(file.size) + sizeof(u8);
-				u64 fetchedSize = Min(file.size, m_client.GetMessageMaxSize() - m_client.GetMessageReceiveHeaderSize() - headerSize);
+				if (!f)
+					return;
 
-				if (fetchedSize < file.size)
+				StackBinaryWriter<SendMaxSize> writer;
+				bool isFirst = true;
+				for (auto& deferredResponse : f->deferredResponses)
 				{
-					SCOPED_WRITE_LOCK(m_activeFetchesLock, lock);
-					fetchId = PopId();
-					auto res = m_activeFetches.try_emplace(fetchId);
-					UBA_ASSERT(res.second);
-					ActiveFetch& fetch = res.first->second;
-					fetch.clientId = connectionInfo.GetId();
-					lock.Leave();
-
-					fetch.fetchedSize = fetchedSize;
-					fetch.file = &file;
+					writer.Reset();
+					SendFetchBeginResponse(deferredResponse.clientId, file, writer, isFirst);
+					m_server.SendResponse(deferredResponse.info, writer.GetData(), u32(writer.GetPosition()));
+					isFirst = false;
 				}
+				delete f;
+			});
 
-				u8 flags = 0;
-				flags |= u8(file.storeCompressed) << 0;
 
-				writer.WriteU16(fetchId);
-				writer.Write7BitEncoded(file.size);
-				writer.WriteByte(flags);
-				writer.WriteBytes(file.memory, fetchedSize);
+		if (error || !m.message.ProcessAsyncResults(m.reader))
+			return;
 
-				return true;
-			}
-		case StorageMessageType_FetchSegment:
-			{
-				u16 fetchId = reader.ReadU16();
-				u32 fetchIndex = reader.ReadU32() - 1;
+		StringView hint(m.hint);
+
+		u16 fetchId = m.reader.ReadU16();
+		if (fetchId == 0)
+		{
+			error = true;
+			m_logger.Error(TC("FetchBegin failed for cas file %s (%s)."), CasKeyString(file.casKey).str, hint.data);
+			return;
+		}
+		u64 fileSize = m.reader.Read7BitEncoded();
+		file.size = fileSize;
 
-				SCOPED_READ_LOCK(m_activeFetchesLock, activeLock);
-				auto findIt = m_activeFetches.find(fetchId);
-				UBA_ASSERT(findIt != m_activeFetches.end());
-				ActiveFetch& fetch = findIt->second;
-				u32 clientId = fetch.clientId;
-				activeLock.Leave();
+		u8 flags = m.reader.ReadByte();
+		bool storeCompressed = (flags >> 0) & 1;
+		bool sendEnd = (flags >> 1) & 1;
+		u64 fetchedSize = m.reader.GetLeft();
 
-				FileEntry& file = *fetch.file;
-				SCOPED_WRITE_LOCK(file.lock, fileLock);
-				if (file.error)
-					return false;
+		memory = new u8[fileSize];
+		memcpy(memory, m.reader.GetPositionData(), fetchedSize);
+
+		file.received = fetchedSize;
+		file.fetchId = fetchId;
+		file.sendEnd = sendEnd;
+		file.storeCompressed = storeCompressed;
+
+		if (sendEnd && fetchedSize == fileSize)
+			SendEnd(file.casKey);
+
+		if (file.received == file.size)
+		{
+			m_client.TrackWorkEnd(file.trackId);
+			return;
+		}
 
-				if (!file.available)
+		u64 left = file.size - file.received;
+		u64 segmentSize = m_client.GetMessageMaxSize() - 5; // This is server response size - header.. TODO: Should be taken from server
+		u32 segmentCount = u32((left + segmentSize - 1) / segmentSize);
+		file.segmentMessages.resize(segmentCount);
+		for (u32 i=0; i!=segmentCount; ++i)
+		{
+			u64 offset = file.received + segmentSize * i;
+			file.segmentMessages[i] = new SegmentMessage(*this, file, memory + offset, i);
+		}
+
+		sendResponses.Execute();
+
+		// Move the additional messages to a job to be able to return this one quickly.
+		m_server.AddWork([f = &file, segmentCount, this](const WorkContext&)
+			{
+				SCOPED_WRITE_LOCK(m_largeFileLock, lock);
+				//TrackWorkScope tws(m_client, TC("SEGMENTS"));
+				auto& file = *f;
+				for (u32 i=0; i!=segmentCount; ++i)
 				{
-					if (auto mif = file.messagesInFlight[fetchIndex])
+					auto mif = file.segmentMessages[i];
+
+					bool res = mif->message.SendAsync(mif->reader, [](bool error, void* userData)
+						{
+							auto mif = (SegmentMessage*)userData;
+							mif->error = error;
+							mif->proxy.m_server.AddWork([mif](const WorkContext&) { mif->proxy.HandleFetchSegmentReceived(*mif); }, 1, TC("ProxyWaitMsg"), ColorWork);
+						}, mif);
+					if (!res)
 					{
-						UBA_ASSERT(clientId == connectionInfo.GetId());
-						mif->deferredResponses.push_back({clientId, fetchId, messageInfo});
-						messageInfo = {};
-						return true;
+						// TODO: Don't leak mif
+						mif->error = true;
 					}
 				}
-				fileLock.Leave();
+			}, 1, TC("ProxySpawnMsg"), ColorWork);
+	}
 
-				u64 headerSize = sizeof(u16) + Get7BitEncodedCount(file.size) + sizeof(u8);
-				u64 firstFetchSize = m_client.GetMessageMaxSize() - m_client.GetMessageReceiveHeaderSize() - headerSize;
-				u64 segmentSize = m_client.GetMessageMaxSize() - 5; // This is server response size - header.. TODO: Should be taken from server
-				u64 offset = firstFetchSize + segmentSize * (fetchIndex);
-				if (offset + segmentSize > file.size)
-					segmentSize = file.size - offset;
-				writer.WriteBytes(file.memory + offset, segmentSize);
-				return UpdateFetch(fetch.clientId, fetchId, segmentSize);
-			}
-		case StorageMessageType_FetchEnd:
-			{
-				return true;
-			}
-		default:
+
+	bool StorageProxy::HandleFetchSegment(const ConnectionInfo& connectionInfo, const WorkContext& workContext, MessageInfo& messageInfo, BinaryReader& reader, BinaryWriter& writer)
+	{
+		u16 fetchId = reader.ReadU16();
+		u32 fetchIndex = reader.ReadU32() - 1;
+
+		SCOPED_READ_LOCK(m_activeFetchesLock, activeLock);
+		auto findIt = m_activeFetches.find(fetchId);
+		UBA_ASSERT(findIt != m_activeFetches.end());
+		ActiveFetch& fetch = findIt->second;
+		u32 clientId = fetch.clientId;
+		activeLock.Leave();
+
+		FileEntry& file = *fetch.file;
+		SCOPED_WRITE_LOCK(file.lock, fileLock);
+		if (file.error)
+			return false;
+
+		if (!file.available)
+		{
+			if (auto mif = file.segmentMessages[fetchIndex])
 			{
-				NetworkMessage msg(m_client, ServiceId, messageInfo.type, writer2);
-				writer2.WriteBytes(reader.GetPositionData(), reader.GetLeft());
-				if (!msg.Send(reader2))
-					return false;
-				writer.WriteBytes(reader2.GetPositionData(), reader2.GetLeft());
+				UBA_ASSERT(clientId == connectionInfo.GetId());
+				mif->deferredResponses.push_back({clientId, fetchId, messageInfo});
+				messageInfo = {};
 				return true;
 			}
 		}
+		fileLock.Leave();
+
+		u64 headerSize = sizeof(u16) + Get7BitEncodedCount(file.size) + sizeof(u8);
+		u64 firstFetchSize = m_client.GetMessageMaxSize() - m_client.GetMessageReceiveHeaderSize() - headerSize;
+		u64 segmentSize = m_client.GetMessageMaxSize() - 5; // This is server response size - header.. TODO: Should be taken from server
+		u64 offset = firstFetchSize + segmentSize * (fetchIndex);
+		if (offset + segmentSize > file.size)
+			segmentSize = file.size - offset;
+		writer.WriteBytes(file.memory + offset, segmentSize);
+		return UpdateFetch(fetch.clientId, fetchId, segmentSize);
 	}
 
-	void StorageProxy::HandleReceivedData(MessageInFlight& mif)
+	void StorageProxy::HandleFetchSegmentReceived(SegmentMessage& mif)
 	{
 		auto& file = mif.file;
 		if (mif.error)
@@ -399,8 +399,8 @@ namespace uba
 
 		SCOPED_WRITE_LOCK(file.lock, fileLock);
 
-		UBA_ASSERT(file.messagesInFlight[mif.fetchIndex] == &mif);
-		file.messagesInFlight[mif.fetchIndex] = nullptr;
+		UBA_ASSERT(file.segmentMessages[mif.fetchIndex] == &mif);
+		file.segmentMessages[mif.fetchIndex] = nullptr;
 		file.received += mif.reader.GetLeft();
 		bool finished = file.received == file.size;
 		if (finished)
@@ -425,6 +425,94 @@ namespace uba
 		delete &mif;
 	}
 
+	bool StorageProxy::HandleDefault(MessageInfo& messageInfo, BinaryReader& reader, BinaryWriter& writer)
+	{
+		StackBinaryWriter<1024> writer2;
+		NetworkMessage msg(m_client, ServiceId, messageInfo.type, writer2);
+		writer2.WriteBytes(reader.GetPositionData(), reader.GetLeft());
+
+		StackBinaryReader<SendMaxSize> reader2;
+		if (!msg.Send(reader2))
+			return false;
+		writer.WriteBytes(reader2.GetPositionData(), reader2.GetLeft());
+		return true;
+	}
+
+	bool StorageProxy::GetFileFromLocalStorage(u32 clientId, const CasKey& casKey, FileEntry& file, ScopedWriteLock& fileLock)
+	{
+		if (!m_useLocalStorage || !m_localStorage || !IsCompressed(casKey) || !m_inProcessClientId || clientId == m_inProcessClientId)
+			return true;
+
+		// We need to leave this lock here since the in-process storage client might be asking for this file too and then we can end up in a deadlock
+		fileLock.Leave();
+
+		bool hasCas = m_localStorage->EnsureCasFile(casKey, nullptr);
+		StringBuffer<> casFile;
+		hasCas = hasCas && m_localStorage->GetCasFileName(casFile, casKey);
+
+		// Enter lock again, and also check if another thread might have already handled this file while we looked if it existed in local storage
+		fileLock.Enter();
+
+		if (file.memory || !hasCas)
+			return true;
+
+		FileAccessor sourceFile(m_logger, casFile.data);
+		if (!sourceFile.OpenMemoryRead())
+			return true;
+
+		u64 fileSize = sourceFile.GetSize();
+		file.memory = new u8[fileSize];
+		if (!file.memory)
+			return false;
+		file.error = false;
+		file.size = fileSize;
+		file.received = fileSize;
+		file.storeCompressed = true;
+		memcpy(file.memory, sourceFile.GetData(), fileSize);
+		file.available = true;
+		return true;
+	}
+
+	bool StorageProxy::SendFetchBeginResponse(u32 clientId, FileEntry& file, BinaryWriter& writer, bool writeData)
+	{
+		if (file.error)
+		{
+			writer.WriteU16(0);
+			return false;
+		}
+		u16 fetchId = u16(~0);
+
+		u64 headerSize = sizeof(u16) + Get7BitEncodedCount(file.size) + sizeof(u8);
+		u64 fetchedSize = Min(file.size, m_client.GetMessageMaxSize() - m_client.GetMessageReceiveHeaderSize() - headerSize);
+
+		if (fetchedSize < file.size)
+		{
+			SCOPED_WRITE_LOCK(m_activeFetchesLock, lock);
+			fetchId = PopId();
+			auto res = m_activeFetches.try_emplace(fetchId);
+			UBA_ASSERT(res.second);
+			ActiveFetch& fetch = res.first->second;
+			fetch.clientId = clientId;
+			lock.Leave();
+
+			fetch.fetchedSize = fetchedSize;
+			fetch.file = &file;
+		}
+
+		u8 flags = 0;
+		flags |= u8(file.storeCompressed) << 0;
+
+		writer.WriteU16(fetchId);
+		writer.Write7BitEncoded(file.size);
+		writer.WriteByte(flags);
+
+		if (writeData)
+			writer.WriteBytes(file.memory, fetchedSize);
+		else
+			writer.AllocWrite(fetchedSize);
+		return true;
+	}
+
 	bool StorageProxy::UpdateFetch(u32 clientId, u16 fetchId, u64 segmentSize)
 	{
 		SCOPED_WRITE_LOCK(m_activeFetchesLock, activeLock);
diff --git a/Engine/Source/Programs/UnrealBuildAccelerator/Common/Public/UbaStorageProxy.h b/Engine/Source/Programs/UnrealBuildAccelerator/Common/Public/UbaStorageProxy.h
--- a/Engine/Source/Programs/UnrealBuildAccelerator/Common/Public/UbaStorageProxy.h
+++ b/Engine/Source/Programs/UnrealBuildAccelerator/Common/Public/UbaStorageProxy.h
@@ -43,10 +43,21 @@ namespace uba
 		u16 PopId();
 		void PushId(u16 id);
 		bool HandleMessage(const ConnectionInfo& connectionInfo, const WorkContext& workContext, MessageInfo& messageInfo, BinaryReader& reader, BinaryWriter& writer);
+		bool HandleConnect(const ConnectionInfo& connectionInfo, BinaryReader& reader, BinaryWriter& writer);
 
-		struct MessageInFlight;
-		void HandleReceivedData(MessageInFlight& mif);
+		struct BeginMessage;
+		bool HandleFetchBegin(const ConnectionInfo& connectionInfo, const WorkContext& workContext, MessageInfo& messageInfo, BinaryReader& reader, BinaryWriter& writer);
+		void HandleFetchBeginReceived(BeginMessage& m, bool error);
 
+		struct SegmentMessage;
+		bool HandleFetchSegment(const ConnectionInfo& connectionInfo, const WorkContext& workContext, MessageInfo& messageInfo, BinaryReader& reader, BinaryWriter& writer);
+		void HandleFetchSegmentReceived(SegmentMessage& mif);
+
+		bool HandleDefault(MessageInfo& messageInfo, BinaryReader& reader, BinaryWriter& writer);
+
+		struct FileEntry;
+		bool GetFileFromLocalStorage(u32 clientId, const CasKey& casKey, FileEntry& file, ScopedWriteLock& fileLock);
+		bool SendFetchBeginResponse(u32 clientId, FileEntry& file, BinaryWriter& writer, bool writeData = true);
 		bool UpdateFetch(u32 clientId, u16 fetchId, u64 segmentSize);
 		bool SendEnd(const CasKey& key);
 
@@ -77,7 +88,8 @@ namespace uba
 			bool sendEnd = false;
 			bool error = false;
 			bool available = false;
-			Vector<MessageInFlight*> messagesInFlight;
+			BeginMessage* beginMessage = nullptr;
+			Vector<SegmentMessage*> segmentMessages;
 		};
 
 		ReaderWriterLock m_filesLock;
