[UBA]
* Added Futex. Will replace write-only rw locks with this. Futex is still rwlock on windows but is non-recursive critical section on linux/macos since internet says they are faster for contented locks..

diff --git a/Engine/Source/Programs/UnrealBuildAccelerator/Core/Private/UbaSynchronization.cpp b/Engine/Source/Programs/UnrealBuildAccelerator/Core/Private/UbaSynchronization.cpp
--- a/Engine/Source/Programs/UnrealBuildAccelerator/Core/Private/UbaSynchronization.cpp
+++ b/Engine/Source/Programs/UnrealBuildAccelerator/Core/Private/UbaSynchronization.cpp
@@ -9,7 +9,7 @@
 
 namespace uba
 {
-	CriticalSection::CriticalSection()
+	CriticalSection::CriticalSection(bool recursive)
 	{
 		#if PLATFORM_WINDOWS
 		static_assert(alignof(CRITICAL_SECTION) == alignof(CriticalSection));
@@ -18,11 +18,18 @@ namespace uba
 		#else
 		static_assert(alignof(pthread_mutex_t) == alignof(CriticalSection));
 		static_assert(sizeof(data) >= sizeof(pthread_mutex_t));
-		pthread_mutexattr_t attr;
-		pthread_mutexattr_init(&attr);
-		pthread_mutexattr_settype(&attr, PTHREAD_MUTEX_RECURSIVE);
-		int res = pthread_mutex_init((pthread_mutex_t*)data, &attr);(void)res;
-		UBA_ASSERTF(res == 0, TC("pthread_mutex_init failed: %i"), res);
+		int res;
+		if (recursive)
+		{
+			pthread_mutexattr_t attr;
+			pthread_mutexattr_init(&attr);
+			pthread_mutexattr_settype(&attr, PTHREAD_MUTEX_RECURSIVE);
+			res = pthread_mutex_init((pthread_mutex_t*)data, &attr);
+		}
+		else
+			res = pthread_mutex_init((pthread_mutex_t*)data, nullptr);
+
+		UBA_ASSERTF(res == 0, TC("pthread_mutex_init failed: %i"), res);(void)res;
 		#endif
 	}
 
diff --git a/Engine/Source/Programs/UnrealBuildAccelerator/Core/Public/UbaSynchronization.h b/Engine/Source/Programs/UnrealBuildAccelerator/Core/Public/UbaSynchronization.h
--- a/Engine/Source/Programs/UnrealBuildAccelerator/Core/Public/UbaSynchronization.h
+++ b/Engine/Source/Programs/UnrealBuildAccelerator/Core/Public/UbaSynchronization.h
@@ -16,6 +16,8 @@ namespace uba
 	template<typename Type>
 	using Atomic = std::atomic<Type>;
 
+	////////////////////////////////////////////////////////////////////////////////////////////////////
+
 	struct AtomicU64 : Atomic<u64>
 	{
 		AtomicU64(u64 initialValue = 0) : Atomic<u64>(initialValue) {}
@@ -24,17 +26,17 @@ namespace uba
 		void operator=(const AtomicU64& o) { store(o); }
 	};
 
+	////////////////////////////////////////////////////////////////////////////////////////////////////
+
 	class CriticalSection
 	{
 	public:
-		CriticalSection();
+		CriticalSection(bool recursive = true);
 		~CriticalSection();
 
 		void Enter();
 		void Leave();
 
-		template<class Functor> auto Scoped(const Functor& f);
-
 	private:
 		#if PLATFORM_WINDOWS
 		u64 data[5];
@@ -46,19 +48,7 @@ namespace uba
 		CriticalSection& operator=(const CriticalSection&) = delete;
 	};
 
-	class ScopedCriticalSection 
-	{
-	public:
-		ScopedCriticalSection(CriticalSection& cs) : m_cs(cs), m_active(true) { cs.Enter(); }
-		~ScopedCriticalSection() { Leave(); }
-		void Enter() { if (m_active) return; m_cs.Enter(); m_active = true; }
-		void Leave() { if (!m_active) return; m_cs.Leave(); m_active = false; }
-	private:
-		CriticalSection& m_cs;
-		bool m_active;
-	};
-
-	template<class Functor> auto CriticalSection::Scoped(const Functor& f) { ScopedCriticalSection c(*this); return f(); }
+	////////////////////////////////////////////////////////////////////////////////////////////////////
 
 	class ReaderWriterLock
 	{
@@ -88,10 +78,13 @@ namespace uba
 		ReaderWriterLock& operator=(const ReaderWriterLock&) = delete;
 	};
 
+	////////////////////////////////////////////////////////////////////////////////////////////////////
+
 	#if !UBA_TRACK_CONTENTION
 	#define SCOPED_READ_LOCK(readerWriterLock, name) ScopedReadLock name(readerWriterLock);
 	#define SCOPED_WRITE_LOCK(readerWriterLock, name) ScopedWriteLock name(readerWriterLock);
 	#define SCOPED_CRITICAL_SECTION(cs, name) ScopedCriticalSection name(cs);
+	#define SCOPED_FUTEX(cs, name) ScopedFutex name(cs);
 	#else
 	u64 GetTime();
 	struct ContentionTracker { void Add(u64 t) { time += t; ++count; }; Atomic<u64> time; Atomic<u64> count; const char* file; u64 line; };
@@ -105,9 +98,25 @@ namespace uba
 
 	#define SCOPED_READ_LOCK(readerWriterLock, name) SCOPED_LOCK_INTERNAL(readerWriterLock, ScopedReadLock, name)
 	#define SCOPED_WRITE_LOCK(readerWriterLock, name) SCOPED_LOCK_INTERNAL(readerWriterLock, ScopedWriteLock, name)
-	#define SCOPED_CRITICAL_SECTION(cs, name) SCOPED_LOCK_INTERNAL(cs, ScopedCriticalSection,  name)
+	#define SCOPED_CRITICAL_SECTION(cs, name) SCOPED_LOCK_INTERNAL(cs, ScopedCriticalSection, name)
+	#define SCOPED_FUTEX(futex, name) SCOPED_LOCK_INTERNAL(futex, ScopedFutex, name)
 	#endif
 
+	////////////////////////////////////////////////////////////////////////////////////////////////////
+
+	class ScopedCriticalSection 
+	{
+	public:
+		ScopedCriticalSection(CriticalSection& cs) : m_cs(cs), m_active(true) { cs.Enter(); }
+		~ScopedCriticalSection() { Leave(); }
+		void Enter() { if (m_active) return; m_cs.Enter(); m_active = true; }
+		void Leave() { if (!m_active) return; m_cs.Leave(); m_active = false; }
+	private:
+		CriticalSection& m_cs;
+		bool m_active;
+	};
+
+	////////////////////////////////////////////////////////////////////////////////////////////////////
 
 	class ScopedReadLock
 	{
@@ -121,6 +130,8 @@ namespace uba
 		bool m_active = true;
 	};
 
+	////////////////////////////////////////////////////////////////////////////////////////////////////
+
 	class ScopedWriteLock
 	{
 	public:
@@ -133,6 +144,19 @@ namespace uba
 		bool m_active = true;
 	};
 
+	////////////////////////////////////////////////////////////////////////////////////////////////////
+
+	#if PLATFORM_WINDOWS
+	using Futex = ReaderWriterLock;
+	using ScopedFutex = ScopedWriteLock;
+	#else
+	class Futex : public CriticalSection { Futex() : CriticalSection(false) {} };
+	using ScopedFutex = ScopedCriticalSection;
+	#endif
+
+
+	////////////////////////////////////////////////////////////////////////////////////////////////////
+
 	template<typename Lambda>
 	struct ScopeGuard
 	{
@@ -152,4 +176,6 @@ namespace uba
 	};
 	template<typename Lambda>
 	ScopeGuard<Lambda> MakeGuard(Lambda&& lambda) { return ScopeGuard<Lambda>(std::move(lambda)); }
+
+	////////////////////////////////////////////////////////////////////////////////////////////////////
 }
