[UBA]
* Changed readerwriterlocks to futexes because write contended rwlocks are really slow on linux/macos.

diff --git a/Engine/Source/Programs/UnrealBuildAccelerator/Agent/Private/UbaAgent.cpp b/Engine/Source/Programs/UnrealBuildAccelerator/Agent/Private/UbaAgent.cpp
--- a/Engine/Source/Programs/UnrealBuildAccelerator/Agent/Private/UbaAgent.cpp
+++ b/Engine/Source/Programs/UnrealBuildAccelerator/Agent/Private/UbaAgent.cpp
@@ -142,12 +142,12 @@ namespace uba
 		g_shouldExit = true;
 		g_ctrlPressed = true;
 
-		g_exitLock->EnterWrite();
+		g_exitLock->Enter();
 		if (g_logger)
 			g_logger->Info(TC("  Exiting..."));
 		if (g_sessionClient)
 			g_sessionClient->Stop();
-		g_exitLock->LeaveWrite();
+		g_exitLock->Leave();
 	}
 
 	#if PLATFORM_WINDOWS
@@ -714,10 +714,10 @@ namespace uba
 		FilteredLogWriter logWriter(g_consoleLogWriter, verbose ? LogEntryType_Debug : LogEntryType_Detail);
 		LoggerWithWriter logger(logWriter, TC(""));
 
-		g_exitLock->EnterWrite();
+		g_exitLock->Enter();
 		g_logger = &logger;
-		g_exitLock->LeaveWrite();
-		auto glg = MakeGuard([]() { g_exitLock->EnterWrite(); g_logger = nullptr; g_exitLock->LeaveWrite(); });
+		g_exitLock->Leave();
+		auto glg = MakeGuard([]() { g_exitLock->Enter(); g_logger = nullptr; g_exitLock->Leave(); });
 
 #if UBA_AUTO_UPDATE
 		if (waitProcessId != ~0u)
@@ -929,7 +929,7 @@ namespace uba
 
 		Vector<ProcessLogLine> logLines[2];
 		u32 logLinesIndex = 0;
-		ReaderWriterLock logLinesLock;
+		Futex logLinesLock;
 		Event logLinesAvailable(false);
 
 		auto processFinished = [&](const ProcessHandle& process)
@@ -941,7 +941,7 @@ namespace uba
 			const Vector<ProcessLogLine>& processLogLines = process.GetLogLines();
 			if (!processLogLines.empty())
 			{
-				SCOPED_WRITE_LOCK(logLinesLock, lock);
+				SCOPED_FUTEX(logLinesLock, lock);
 				for (auto& line : processLogLines)
 					logLines[logLinesIndex].push_back(line);
 				if (errorCode)
@@ -965,7 +965,7 @@ namespace uba
 					name.Appendf(TC(" (exit code: %u)"), errorCode);
 					entryType = LogEntryType_Error;
 				}
-				SCOPED_WRITE_LOCK(logLinesLock, lock);
+				SCOPED_FUTEX(logLinesLock, lock);
 				logLines[logLinesIndex].push_back({ TString(name.data), entryType });
 			}
 
@@ -1199,7 +1199,7 @@ namespace uba
 						logLinesAvailable.IsSet();
 						u32 logLinesIndexPrev;
 						{
-							SCOPED_WRITE_LOCK(logLinesLock, l);
+							SCOPED_FUTEX(logLinesLock, l);
 							logLinesIndexPrev = logLinesIndex;
 							logLinesIndex = (logLinesIndex + 1) % 2;
 						}
@@ -1216,9 +1216,9 @@ namespace uba
 
 			auto disconnectAndStopLoggingThread = MakeGuard([&]()
 			{
-				g_exitLock->EnterWrite();
+				g_exitLock->Enter();
 				g_sessionClient = nullptr;
-				g_exitLock->LeaveWrite();
+				g_exitLock->Leave();
 
 				networkBackend->StopListen();
 				storageClient->StopProxy();
diff --git a/Engine/Source/Programs/UnrealBuildAccelerator/CacheService/Private/UbaCacheService.cpp b/Engine/Source/Programs/UnrealBuildAccelerator/CacheService/Private/UbaCacheService.cpp
--- a/Engine/Source/Programs/UnrealBuildAccelerator/CacheService/Private/UbaCacheService.cpp
+++ b/Engine/Source/Programs/UnrealBuildAccelerator/CacheService/Private/UbaCacheService.cpp
@@ -76,10 +76,10 @@ namespace uba
 	{
 		g_shouldExit = true;
 
-		g_exitLock->EnterWrite();
+		g_exitLock->Enter();
 		if (g_logger)
 			g_logger->Info(TC("  Exiting..."));
-		g_exitLock->LeaveWrite();
+		g_exitLock->Leave();
 	}
 
 	#if PLATFORM_WINDOWS
@@ -217,10 +217,10 @@ namespace uba
 		FilteredLogWriter logWriter(g_consoleLogWriter, quiet ? LogEntryType_Info : LogEntryType_Detail);
 		LoggerWithWriter logger(logWriter, TC(""));
 
-		g_exitLock->EnterWrite();
+		g_exitLock->Enter();
 		g_logger = &logger;
-		g_exitLock->LeaveWrite();
-		auto glg = MakeGuard([]() { g_exitLock->EnterWrite(); g_logger = nullptr; g_exitLock->LeaveWrite(); });
+		g_exitLock->Leave();
+		auto glg = MakeGuard([]() { g_exitLock->Enter(); g_logger = nullptr; g_exitLock->Leave(); });
 
 		Config config;
 		if (!configFile.empty())
diff --git a/Engine/Source/Programs/UnrealBuildAccelerator/Common/Private/UbaCacheClient.cpp b/Engine/Source/Programs/UnrealBuildAccelerator/Common/Private/UbaCacheClient.cpp
--- a/Engine/Source/Programs/UnrealBuildAccelerator/Common/Private/UbaCacheClient.cpp
+++ b/Engine/Source/Programs/UnrealBuildAccelerator/Common/Private/UbaCacheClient.cpp
@@ -71,20 +71,20 @@ namespace uba
 
 		u32 id = 0;
 
-		ReaderWriterLock serverPathTableNetworkLock;
+		Futex serverPathTableNetworkLock;
 		CompactPathTable serverPathTable;
 		Atomic<u32> serverPathTableSize;
 
-		ReaderWriterLock serverCasKeyTableNetworkLock;
+		Futex serverCasKeyTableNetworkLock;
 		CompactCasKeyTable serverCasKeyTable;
 		Atomic<u32> serverCasKeyTableSize;
 
 		CompactPathTable sendPathTable;
-		ReaderWriterLock sendPathTableNetworkLock;
+		Futex sendPathTableNetworkLock;
 		u32 pathTableSizeSent = 0;
 
 		CompactCasKeyTable sendCasKeyTable;
-		ReaderWriterLock sendCasKeyTableNetworkLock;
+		Futex sendCasKeyTableNetworkLock;
 		u32 casKeyTableSizeSent = 0;
 	};
 
@@ -202,7 +202,7 @@ namespace uba
 		u32 requiredCasTableSize = 0;
 		bool success = true;
 
-		SCOPED_WRITE_LOCK(m_bucketsLock, bucketsLock);
+		SCOPED_FUTEX(m_bucketsLock, bucketsLock);
 		Bucket& bucket = m_buckets.try_emplace(bucketId, bucketId).first->second;
 		if (!bucket.sendPathTable.GetMemory())
 			bucket.sendPathTable.AddCommonStringSegments();
@@ -364,7 +364,7 @@ namespace uba
 						fa.GetFileInformationByHandle(fileInfo);
 
 						auto& fileEntry = m_storage.GetOrCreateFileEntry(CaseInsensitiveFs ? ToStringKeyLower(path) : ToStringKey(path));
-						SCOPED_READ_LOCK(fileEntry.lock, lock);
+						SCOPED_FUTEX_READ(fileEntry.lock, lock);
 
 						auto ToString = [](bool b) { return b ? TC("true") : TC("false"); };
 						return m_logger.Warning(TC("CasDb claims file %s has caskey %s but recalculating it gives us %s (FileEntry: %llu/%llu/%s, Real: %llu/%llu). Will not populate cache for %s"),
@@ -486,7 +486,7 @@ namespace uba
 
 		BinaryReader reader(memory, 0, SendMaxSize);
 
-		SCOPED_WRITE_LOCK(m_bucketsLock, bucketsLock);
+		SCOPED_FUTEX(m_bucketsLock, bucketsLock);
 		Bucket& bucket = m_buckets.try_emplace(bucketId, bucketId).first->second;
 		bucketsLock.Leave();
 
@@ -978,7 +978,7 @@ namespace uba
 
 	bool CacheClient::SendPathTable(Bucket& bucket, u32 requiredPathTableSize)
 	{
-		SCOPED_WRITE_LOCK(bucket.sendPathTableNetworkLock, lock);
+		SCOPED_FUTEX(bucket.sendPathTableNetworkLock, lock);
 		if (requiredPathTableSize <= bucket.pathTableSizeSent)
 			return true;
 
@@ -1003,7 +1003,7 @@ namespace uba
 
 	bool CacheClient::SendCasTable(Bucket& bucket, u32 requiredCasTableSize)
 	{
-		SCOPED_WRITE_LOCK(bucket.sendCasKeyTableNetworkLock, lock);
+		SCOPED_FUTEX(bucket.sendCasKeyTableNetworkLock, lock);
 		if (requiredCasTableSize <= bucket.casKeyTableSizeSent)
 			return true;
 
@@ -1206,7 +1206,7 @@ namespace uba
 		TimerScope ts2(stats.fetchCasTable);
 		u32 requiredPathTableOffset = 0;
 		{
-			SCOPED_WRITE_LOCK(bucket.serverCasKeyTableNetworkLock, lock);
+			SCOPED_FUTEX(bucket.serverCasKeyTableNetworkLock, lock);
 			while (!HasEnoughCasData(bucket, requiredCasTableOffset, requiredPathTableOffset))
 			{
 				if (!FetchCompactTable(bucket.id, bucket.serverCasKeyTable, requiredCasTableOffset + sizeof(CasKey) + 8, CacheMessageType_FetchCasTable2))
@@ -1215,7 +1215,7 @@ namespace uba
 			}
 		}
 		{
-			SCOPED_WRITE_LOCK(bucket.serverPathTableNetworkLock, lock);
+			SCOPED_FUTEX(bucket.serverPathTableNetworkLock, lock);
 			while (!HasEnoughPathData(bucket, requiredPathTableOffset))
 			{
 				u32 targetSize = requiredPathTableOffset + 200;
@@ -1498,25 +1498,17 @@ namespace uba
 		StringKey pathKey = ToStringKey(dirHasher);
 
 		PreparedDir* dir = nullptr;
-		#if !PLATFORM_WINDOWS
-		{
-			SCOPED_READ_LOCK(m_directoryPreparserLock, preparserLock);
-			auto findIt = m_directoryPreparser.find(pathKey);
-			if (findIt != m_directoryPreparser.end())
-				dir = &findIt->second;
-		}
-		#endif
 
 		if (!dir)
 		{
-			SCOPED_WRITE_LOCK(m_directoryPreparserLock, preparserLock);
+			SCOPED_FUTEX(m_directoryPreparserLock, preparserLock);
 			dir = &m_directoryPreparser.try_emplace(pathKey).first->second;
 		}
 
 		if (dir->done)
 			return;
 
-		SCOPED_WRITE_LOCK(dir->lock, lock);
+		SCOPED_FUTEX(dir->lock, lock);
 		if (dir->done)
 			return;
 		dir->done = true;
diff --git a/Engine/Source/Programs/UnrealBuildAccelerator/Common/Private/UbaCacheServer.cpp b/Engine/Source/Programs/UnrealBuildAccelerator/Common/Private/UbaCacheServer.cpp
--- a/Engine/Source/Programs/UnrealBuildAccelerator/Common/Private/UbaCacheServer.cpp
+++ b/Engine/Source/Programs/UnrealBuildAccelerator/Common/Private/UbaCacheServer.cpp
@@ -134,7 +134,7 @@ namespace uba
 		CompactPathTable pathTable;
 		CompactCasKeyTable casKeyTable;
 
-		ReaderWriterLock cacheEntryLookupLock;
+		Futex cacheEntryLookupLock;
 		UnorderedMap<CasKey, CacheEntry> cacheEntryLookup;
 		u64 id;
 		u32 index = ~0u;
@@ -584,7 +584,7 @@ namespace uba
 
 		do
 		{
-			SCOPED_WRITE_LOCK(m_connectionsLock, lock2);
+			SCOPED_FUTEX(m_connectionsLock, lock2);
 			if (!m_connections.empty())
 				return true;
 			m_isRunningMaintenance = true;
@@ -594,7 +594,7 @@ namespace uba
 
 			auto g = MakeGuard([&]()
 				{
-					SCOPED_WRITE_LOCK(m_connectionsLock, lock2);
+					SCOPED_FUTEX(m_connectionsLock, lock2);
 					m_isRunningMaintenance = false;
 				});
 
@@ -738,7 +738,7 @@ namespace uba
 			});
 
 
-		ReaderWriterLock globalStatsLock;
+		Futex globalStatsLock;
 		u64 now = GetSystemTimeAsFileTime();
 		u64 oldest = 0;
 		u64 oldestUsedTime = 0;
@@ -798,7 +798,7 @@ namespace uba
 
 				auto& deletedOffsets = context->deletedOffsets;
 
-				ReaderWriterLock bucketLock;
+				Futex bucketLock;
 				Vector<CasKey> keysToErase;
 
 
@@ -1002,7 +1002,7 @@ namespace uba
 
 
 					{
-						SCOPED_WRITE_LOCK(bucketLock, lock2);
+						SCOPED_FUTEX(bucketLock, lock2);
 						if (!bucket.oldestUsedTime || entriesOldestUsed < bucket.oldestUsedTime)
 							bucket.oldestUsedTime = entriesOldestUsed;
 						if (!bucketOldest || entriesOldest < bucketOldest)
@@ -1019,7 +1019,7 @@ namespace uba
 				});
 
 				{
-					SCOPED_WRITE_LOCK(globalStatsLock, l);
+					SCOPED_FUTEX(globalStatsLock, l);
 					if (!oldestUsedTime || bucketOldest < oldestUsedTime)
 						oldestUsedTime = bucket.oldestUsedTime;
 					if (!oldest || bucketOldest < oldest)
@@ -1239,7 +1239,7 @@ namespace uba
 				longestUnusedSeconds = GetFileTimeAsSeconds(now - m_creationTime - bucket.oldestUsedTime);
 			m_logger.Info(TC("    Bucket %u Done (%s). CacheEntries: %s (%s) PathTable: %s CasTable: %s LongestUnused: %s"), bucket.index, TimeToText(GetTime() - bucketStartTime).str, CountToText(bucket.totalEntryCount).str, BytesToText(bucket.totalEntrySize).str, BytesToText(bucket.m_pathTable.GetSize()).str, BytesToText(bucket.m_casKeyTable.GetSize()).str, TimeToText(MsToTime(longestUnusedSeconds*1000), true).str);
 
-			SCOPED_WRITE_LOCK(globalStatsLock, l);
+			SCOPED_FUTEX(globalStatsLock, l);
 			maxCommittedMemory = Max(maxCommittedMemory, memoryBlock.writtenSize);
 
 		}, AsView(TC("")), true);
@@ -1276,7 +1276,7 @@ namespace uba
 	{
 		if (!m_shutdownRequested)
 			return false;
-		SCOPED_READ_LOCK(m_connectionsLock, lock2);
+		SCOPED_FUTEX(m_connectionsLock, lock2);
 		if (!m_connections.empty() || m_addsSinceMaintenance)
 			return false;
 		return true;
@@ -1287,7 +1287,7 @@ namespace uba
 		StringBuffer<> log;
 		log.Appendf(TC("Client %u disconnected"), clientId);
 
-		SCOPED_WRITE_LOCK(m_connectionsLock, lock);
+		SCOPED_FUTEX(m_connectionsLock, lock);
 		auto it = m_connections.find(clientId);
 		if (it != m_connections.end())
 		{
@@ -1309,7 +1309,7 @@ namespace uba
 	{
 		u64 id = reader.Read7BitEncoded();
 		u32 bucketVersion = u32(id >> 35u);
-		SCOPED_WRITE_LOCK(m_connectionsLock, lock);
+		SCOPED_FUTEX(m_connectionsLock, lock);
 		auto& connection = m_connections[connectionInfo.GetId()];
 		if (outClientVersion)
 		{
@@ -1329,7 +1329,7 @@ namespace uba
 	CacheServer::Bucket& CacheServer::GetBucket(u64 id, const tchar* reason, bool addCommon)
 	{
 		u32 bucketVersion = u32(id >> 35u);
-		SCOPED_WRITE_LOCK(m_bucketsLock, bucketsLock);
+		SCOPED_FUTEX(m_bucketsLock, bucketsLock);
 		auto insres = m_buckets.try_emplace(id, id, bucketVersion);
 		auto& bucket = insres.first->second;
 		if (!insres.second)
@@ -1365,7 +1365,7 @@ namespace uba
 			if (reader.GetLeft())
 				hint = reader.ReadString();
 
-			SCOPED_WRITE_LOCK(m_connectionsLock, lock);
+			SCOPED_FUTEX(m_connectionsLock, lock);
 			if (m_isRunningMaintenance)
 			{
 				writer.WriteBool(false);
@@ -1414,7 +1414,7 @@ namespace uba
 			if (reader.GetLeft())
 				success = reader.ReadBool();
 
-			SCOPED_WRITE_LOCK(connectionBucket.cacheEntryLookupLock, lock2);
+			SCOPED_FUTEX(connectionBucket.cacheEntryLookupLock, lock2);
 			auto findIt = connectionBucket.cacheEntryLookup.find(cmdKey);
 			if (findIt == connectionBucket.cacheEntryLookup.end())
 				return true;
@@ -1711,7 +1711,7 @@ namespace uba
 		}
 		else
 		{
-			SCOPED_WRITE_LOCK(connectionBucket.cacheEntryLookupLock, lock3);
+			SCOPED_FUTEX(connectionBucket.cacheEntryLookupLock, lock3);
 			bool res = connectionBucket.cacheEntryLookup.try_emplace(cmdKey, std::move(newEntry)).second;
 			UBA_ASSERT(res);(void)res;
 		}
@@ -1773,7 +1773,7 @@ namespace uba
 
 		u32 clientVersion;
 		{
-			SCOPED_WRITE_LOCK(m_connectionsLock, lock);
+			SCOPED_FUTEX(m_connectionsLock, lock);
 			auto& conn = m_connections[clientId];
 			conn.fetchBuckets.insert(bucket.index);
 			++conn.fetchEntryCount;
@@ -1804,7 +1804,7 @@ namespace uba
 
 		u32 clientVersion;
 		{
-			SCOPED_WRITE_LOCK(m_connectionsLock, lock);
+			SCOPED_FUTEX(m_connectionsLock, lock);
 			auto& conn = m_connections[clientId];
 			++conn.fetchEntryHitCount;
 			clientVersion = conn.clientVersion;
@@ -1948,7 +1948,7 @@ namespace uba
 
 			Vector<u8> temp;
 
-			SCOPED_WRITE_LOCK(m_bucketsLock, bucketsLock);
+			SCOPED_FUTEX(m_bucketsLock, bucketsLock);
 			for (auto& kv : m_buckets)
 			{
 				Bucket& bucket = kv.second;
@@ -2037,7 +2037,7 @@ namespace uba
 			u32 index = 0;
 
 			{
-				SCOPED_WRITE_LOCK(m_bucketsLock, bucketsLock);
+				SCOPED_FUTEX(m_bucketsLock, bucketsLock);
 				for (auto& kv : m_buckets)
 				{
 					Bucket& bucket = kv.second;
@@ -2094,7 +2094,7 @@ namespace uba
 			writeLine(line.Clear().Appendf(TC("  Buckets:")));
 			{
 				u32 index = 0;
-				SCOPED_WRITE_LOCK(m_bucketsLock, bucketsLock);
+				SCOPED_FUTEX(m_bucketsLock, bucketsLock);
 				for (auto& kv : m_buckets)
 				{
 					bool printInvalid = true;
@@ -2260,7 +2260,7 @@ namespace uba
 			Vector<u8> temp2;
 			Vector<u8> temp3;
 
-			SCOPED_WRITE_LOCK(m_bucketsLock, bucketsLock);
+			SCOPED_FUTEX(m_bucketsLock, bucketsLock);
 			for (auto& kv : m_buckets)
 			{
 				Bucket& bucket = kv.second;
@@ -2275,7 +2275,7 @@ namespace uba
 		}
 		else if (command.Equals(TC("pathtable")))
 		{
-			SCOPED_WRITE_LOCK(m_bucketsLock, bucketsLock);
+			SCOPED_FUTEX(m_bucketsLock, bucketsLock);
 			for (auto& kv : m_buckets)
 			{
 				Bucket& bucket = kv.second;
@@ -2297,7 +2297,7 @@ namespace uba
 		}
 		else if (command.Equals(TC("usedpaths")))
 		{
-			SCOPED_WRITE_LOCK(m_bucketsLock, bucketsLock);
+			SCOPED_FUTEX(m_bucketsLock, bucketsLock);
 			for (auto& kv : m_buckets)
 			{
 				Bucket& bucket = kv.second;
@@ -2323,7 +2323,7 @@ namespace uba
 		}
 		else if (command.Equals(TC("usedsegments")))
 		{
-			SCOPED_WRITE_LOCK(m_bucketsLock, bucketsLock);
+			SCOPED_FUTEX(m_bucketsLock, bucketsLock);
 			for (auto& kv : m_buckets)
 			{
 				UnorderedMap<u32, u32> segments;
@@ -2368,7 +2368,7 @@ namespace uba
 			CacheEntries* largestEntries = nullptr;
 			Bucket* largestEntriesBucket = nullptr;
 			u64 largestSize = 0;
-			SCOPED_WRITE_LOCK(m_bucketsLock, bucketsLock);
+			SCOPED_FUTEX(m_bucketsLock, bucketsLock);
 			for (auto& kv : m_buckets)
 			{
 				Bucket& bucket = kv.second;
diff --git a/Engine/Source/Programs/UnrealBuildAccelerator/Common/Private/UbaCompactTables.cpp b/Engine/Source/Programs/UnrealBuildAccelerator/Common/Private/UbaCompactTables.cpp
--- a/Engine/Source/Programs/UnrealBuildAccelerator/Common/Private/UbaCompactTables.cpp
+++ b/Engine/Source/Programs/UnrealBuildAccelerator/Common/Private/UbaCompactTables.cpp
@@ -30,14 +30,14 @@ namespace uba
 
 	u32 CompactPathTable::Add(const tchar* str, u64 strLen)
 	{
-		SCOPED_WRITE_LOCK(m_lock, lock);
+		SCOPED_FUTEX(m_lock, lock);
 		u32 res = AddNoLock(str, strLen);
 		return res;
 	}
 
 	u32 CompactPathTable::Add(const tchar* str, u64 strLen, u32& outRequiredPathTableSize)
 	{
-		SCOPED_WRITE_LOCK(m_lock, lock);
+		SCOPED_FUTEX(m_lock, lock);
 		u32 res = AddNoLock(str, strLen);
 		outRequiredPathTableSize = u32(m_pathTableMem.writtenSize);
 		return res;
@@ -412,7 +412,7 @@ namespace uba
 	{
 		#if UBA_DEBUG
 		{
-			SCOPED_READ_LOCK(const_cast<CompactPathTable*>(this)->m_lock, lock)
+			SCOPED_FUTEX_READ(const_cast<CompactPathTable*>(this)->m_lock, lock)
 			UBA_ASSERTF(offset < m_pathTableMem.writtenSize, TC("Reading path key from offset %llu which is out of bounds (Max %llu)"), offset, m_pathTableMem.writtenSize);
 		}
 		#endif
@@ -536,7 +536,7 @@ namespace uba
 
 	u32 CompactPathTable::GetSize() const
 	{
-		SCOPED_READ_LOCK(m_lock, lock2)
+		SCOPED_FUTEX_READ(m_lock, lock2)
 		return u32(m_pathTableMem.writtenSize);
 	}
 
@@ -636,7 +636,7 @@ namespace uba
 
 	u32 CompactCasKeyTable::Add(const CasKey& casKey, u64 stringOffset, u32& outRequiredCasTableSize)
 	{
-		SCOPED_WRITE_LOCK(m_lock, lock2)
+		SCOPED_FUTEX(m_lock, lock2)
 		if (!m_casKeyTableMem.memory)
 			m_casKeyTableMem.Init(CasKeyTableMaxSize);
 		if (!m_offsets.IsInitialized())
@@ -786,7 +786,7 @@ namespace uba
 	{
 		#if UBA_DEBUG
 		{
-			SCOPED_READ_LOCK(const_cast<CompactCasKeyTable*>(this)->m_lock, lock)
+			SCOPED_FUTEX_READ(const_cast<CompactCasKeyTable*>(this)->m_lock, lock)
 			UBA_ASSERTF(offset + sizeof(CasKey) < m_casKeyTableMem.writtenSize, TC("Reading cas key from offset %llu which is out of bounds (Max %llu)"), offset + sizeof(CasKey), m_casKeyTableMem.writtenSize);
 		}
 		#endif
@@ -804,7 +804,7 @@ namespace uba
 
 	u32 CompactCasKeyTable::GetSize()
 	{
-		SCOPED_READ_LOCK(m_lock, lock2)
+		SCOPED_FUTEX_READ(m_lock, lock2)
 		return u32(m_casKeyTableMem.writtenSize);
 	}
 
diff --git a/Engine/Source/Programs/UnrealBuildAccelerator/Common/Private/UbaFile.cpp b/Engine/Source/Programs/UnrealBuildAccelerator/Common/Private/UbaFile.cpp
--- a/Engine/Source/Programs/UnrealBuildAccelerator/Common/Private/UbaFile.cpp
+++ b/Engine/Source/Programs/UnrealBuildAccelerator/Common/Private/UbaFile.cpp
@@ -867,10 +867,10 @@ namespace uba
 		TString key(dir, dir + dirLen);
 		dir = key.c_str();
 
-		SCOPED_WRITE_LOCK(m_createdDirsLock, lock);
+		SCOPED_FUTEX(m_createdDirsLock, lock);
 		CreatedDir& cd = m_createdDirs.try_emplace(key).first->second;
 		lock.Leave();
-		SCOPED_WRITE_LOCK(cd.lock, dirLock);
+		SCOPED_FUTEX(cd.lock, dirLock);
 		if (cd.handled)
 			return true;
 		cd.handled = true;
@@ -898,7 +898,7 @@ namespace uba
 
 	void DirectoryCache::Clear()
 	{
-		SCOPED_WRITE_LOCK(m_createdDirsLock, lock);
+		SCOPED_FUTEX(m_createdDirsLock, lock);
 		m_createdDirs.clear();
 	}
 
diff --git a/Engine/Source/Programs/UnrealBuildAccelerator/Common/Private/UbaFileMapping.cpp b/Engine/Source/Programs/UnrealBuildAccelerator/Common/Private/UbaFileMapping.cpp
--- a/Engine/Source/Programs/UnrealBuildAccelerator/Common/Private/UbaFileMapping.cpp
+++ b/Engine/Source/Programs/UnrealBuildAccelerator/Common/Private/UbaFileMapping.cpp
@@ -18,7 +18,7 @@
 namespace uba
 {
 	#if UBA_DEBUG_FILE_MAPPING
-	ReaderWriterLock g_fileMappingsLock;
+	Futex g_fileMappingsLock;
 	struct DebugFileMapping { TString hint; Atomic<u64> viewCount; };
 	UnorderedMap<HANDLE, DebugFileMapping> g_fileMappings;
 	UnorderedMap<const void*, HANDLE> g_viewMappings;
@@ -45,7 +45,7 @@ namespace uba
 		#if UBA_DEBUG_FILE_MAPPING
 		if (res)
 		{
-			SCOPED_WRITE_LOCK(g_fileMappingsLock, lock);
+			SCOPED_FUTEX(g_fileMappingsLock, lock);
 			auto insres = g_fileMappings.try_emplace(res, hint);(void)insres;
 			UBA_ASSERT(insres.second);
 			//logger.Info(TC("FILEMAPPING CREAT: %s %llu (%llu)"), hint, u64(res), g_fileMappings.size());
@@ -57,7 +57,7 @@ namespace uba
 #else
 	int asFileDescriptor(FileHandle fh);
 
-	ReaderWriterLock g_mappingUidCounterLock;
+	Futex g_mappingUidCounterLock;
 	Atomic<u64> g_mappingUidCounter;
 #endif
 
@@ -74,7 +74,7 @@ namespace uba
 		StringBuffer<64> lockDir;
 		lockDir.Append("/tmp/uba_shm_locks");
 
-		SCOPED_WRITE_LOCK(g_mappingUidCounterLock, lock);
+		SCOPED_FUTEX(g_mappingUidCounterLock, lock);
 		if (!g_mappingUidCounter)
 		{
 			// Create dir
@@ -259,7 +259,7 @@ namespace uba
 
 		#if UBA_DEBUG_FILE_MAPPING
 		{
-			SCOPED_WRITE_LOCK(g_fileMappingsLock, lock);
+			SCOPED_FUTEX(g_fileMappingsLock, lock);
 			auto findIt = g_fileMappings.find(fileMappingObject.handle);
 			UBA_ASSERTF(findIt != g_fileMappings.end(), TC("Mapping nonexisting handle: %llu"), u64(fileMappingObject.handle));
 			auto insres = g_viewMappings.try_emplace(res, fileMappingObject.handle);
@@ -287,7 +287,7 @@ namespace uba
 
 		#if UBA_DEBUG_FILE_MAPPING
 		{
-			SCOPED_WRITE_LOCK(g_fileMappingsLock, lock);
+			SCOPED_FUTEX(g_fileMappingsLock, lock);
 			auto findIt = g_viewMappings.find(lpBaseAddress);
 			UBA_ASSERT(findIt != g_viewMappings.end());
 			auto findIt2 = g_fileMappings.find(findIt->second);
@@ -314,7 +314,7 @@ namespace uba
 		#if UBA_DEBUG_FILE_MAPPING
 		if (h.IsValid())
 		{
-			SCOPED_WRITE_LOCK(g_fileMappingsLock, lock);
+			SCOPED_FUTEX(g_fileMappingsLock, lock);
 			auto findIt = g_fileMappings.find(h.handle);
 			UBA_ASSERTF(findIt != g_fileMappings.end(), TC("Handle: %llu"), u64(h.handle));
 			UBA_ASSERTF(!findIt->second.viewCount.load(), TC("Closing file mapping with %llu open views (%s)"), findIt->second.viewCount.load(), findIt->second.hint.c_str());
@@ -368,7 +368,7 @@ namespace uba
 
 		#if UBA_DEBUG_FILE_MAPPING
 		{
-			SCOPED_WRITE_LOCK(g_fileMappingsLock, lock);
+			SCOPED_FUTEX(g_fileMappingsLock, lock);
 			auto insres = g_fileMappings.try_emplace(targetHandle.mh, hint);(void)insres;
 			UBA_ASSERT(insres.second);
 			//logger.Info(TC("FILEMAPPING DUPLI: %s %llu (%llu)"), hint, targetHandle.mh, g_fileMappings.size());
@@ -523,7 +523,7 @@ namespace uba
 	{
 		while (true)
 		{
-			SCOPED_WRITE_LOCK(storage.availableFilesLock, lock);
+			SCOPED_FUTEX(storage.availableFilesLock, lock);
 			for (u32 i = storage.availableFilesCount; i!=0; --i)
 			{
 				u32 index = i - 1;
@@ -561,7 +561,7 @@ namespace uba
 
 	void FileMappingBuffer::PushFile(MappingStorage& storage, File* file)
 	{
-		SCOPED_WRITE_LOCK(storage.availableFilesLock, lock);
+		SCOPED_FUTEX(storage.availableFilesLock, lock);
 		storage.availableFiles[storage.availableFilesCount++] = file;
 		storage.availableFilesEvent.Set();
 	}
@@ -571,7 +571,7 @@ namespace uba
 		u32 locksTaken = 0;
 		while (locksTaken < storage.fileCount)
 		{
-			SCOPED_WRITE_LOCK(storage.availableFilesLock, lock);
+			SCOPED_FUTEX(storage.availableFilesLock, lock);
 			if (storage.availableFilesCount == 0)
 			{
 				lock.Leave();
@@ -724,7 +724,7 @@ namespace uba
 
 	FileMappingAllocator::Allocation FileMappingAllocator::Alloc(const tchar* hint)
 	{
-		SCOPED_WRITE_LOCK(m_mappingLock, lock);
+		SCOPED_FUTEX(m_mappingLock, lock);
 
 		u64 index = m_mappingCount;
 		bool needCommit = false;
@@ -769,7 +769,7 @@ namespace uba
 		if (!UnmapViewOfFile(m_logger, allocation.memory, m_blockSize, m_name))
 			m_logger.Error(TC("%s - Failed to unmap view of file (%s)"), m_name, LastErrorToText().data);
 		u64 index = allocation.offset / m_blockSize;
-		SCOPED_WRITE_LOCK(m_mappingLock, lock);
+		SCOPED_FUTEX(m_mappingLock, lock);
 		m_availableBlocks.insert(index);
 	}
 }
diff --git a/Engine/Source/Programs/UnrealBuildAccelerator/Common/Private/UbaLogger.cpp b/Engine/Source/Programs/UnrealBuildAccelerator/Common/Private/UbaLogger.cpp
--- a/Engine/Source/Programs/UnrealBuildAccelerator/Common/Private/UbaLogger.cpp
+++ b/Engine/Source/Programs/UnrealBuildAccelerator/Common/Private/UbaLogger.cpp
@@ -234,7 +234,7 @@ namespace uba
 		virtual void Log(LogEntryType type, const tchar* str, u32 strLen, const tchar* prefix = nullptr, u32 prefixLen = 0) override;
 	private:
 		void LogNoLock(LogEntryType type, const tchar* str, u32 strLen, const tchar* prefix, u32 prefixLen);
-		ReaderWriterLock m_lock;
+		Futex m_lock;
 #if PLATFORM_WINDOWS
 		HANDLE m_stdout = 0;
 		u32 m_defaultAttributes = 0;
@@ -269,7 +269,7 @@ namespace uba
 	void ConsoleLogWriter::BeginScope()
 	{
 		if (!t_consoleLogScopeCount++)
-			m_lock.EnterWrite();
+			m_lock.Enter();
 	}
 
 	void ConsoleLogWriter::EndScope()
@@ -280,14 +280,14 @@ namespace uba
 		if (!m_stdout)
 #endif
 			fflush(stdout);
-		m_lock.LeaveWrite();
+		m_lock.Leave();
 	}
 
 	void ConsoleLogWriter::Log(LogEntryType type, const tchar* str, u32 strLen, const tchar* prefix, u32 prefixLen)
 	{
 		if (t_consoleLogScopeCount)
 			return LogNoLock(type, str, strLen, prefix, prefixLen);
-		SCOPED_WRITE_LOCK(m_lock, lock);
+		SCOPED_FUTEX(m_lock, lock);
 		LogNoLock(type, str, strLen, prefix, prefixLen);
 #if PLATFORM_WINDOWS
 		if (!m_stdout)
@@ -399,7 +399,7 @@ namespace uba
 			if (!m_file)
 				return;
 			if (!t_debugLogScopeCount++)
-				m_logLock.EnterWrite();
+				m_logLock.Enter();
 		}
 
 		virtual void EndScope() override
@@ -410,7 +410,7 @@ namespace uba
 				return;
 			//
 
-			m_logLock.LeaveWrite();
+			m_logLock.Leave();
 		}
 
 		virtual void Log(LogEntryType type, const tchar* str, u32 strLen, const tchar* prefix = nullptr, u32 prefixLen = 0) override
@@ -419,7 +419,7 @@ namespace uba
 				return;
 			if (t_debugLogScopeCount)
 				return LogNoLock(type, str, strLen, prefix, prefixLen);
-			SCOPED_WRITE_LOCK(m_logLock, lock);
+			SCOPED_FUTEX(m_logLock, lock);
 			LogNoLock(type, str, strLen, prefix, prefixLen);
 		}
 
@@ -446,7 +446,7 @@ namespace uba
 
 		TString m_fileName;
 		FileAccessor* m_file = nullptr;
-		ReaderWriterLock m_logLock;
+		Futex m_logLock;
 
 	};
 
diff --git a/Engine/Source/Programs/UnrealBuildAccelerator/Common/Private/UbaNetworkBackendTcp.cpp b/Engine/Source/Programs/UnrealBuildAccelerator/Common/Private/UbaNetworkBackendTcp.cpp
--- a/Engine/Source/Programs/UnrealBuildAccelerator/Common/Private/UbaNetworkBackendTcp.cpp
+++ b/Engine/Source/Programs/UnrealBuildAccelerator/Common/Private/UbaNetworkBackendTcp.cpp
@@ -80,7 +80,7 @@ namespace uba
 		void* disconnectContext = nullptr;
 		DisconnectCallback* disconnectCallback = nullptr;
 
-		ReaderWriterLock sendLock;
+		Futex sendLock;
 		CriticalSection shutdownLock;
 
 		Thread recvThread;
@@ -169,7 +169,7 @@ namespace uba
 	{
 		StopListen();
 
-		SCOPED_WRITE_LOCK(m_connectionsLock, lock);
+		SCOPED_FUTEX(m_connectionsLock, lock);
 		for (auto& conn : m_connections)
 		{
 			SCOPED_CRITICAL_SECTION(conn.shutdownLock, lock2);
@@ -204,7 +204,7 @@ namespace uba
 		auto& conn = *(Connection*)connection;
 		sendContext.isUsed = true;
 
-		SCOPED_WRITE_LOCK(conn.sendLock, lock);
+		SCOPED_FUTEX(conn.sendLock, lock);
 		if (conn.socket == INVALID_SOCKET)
 			return false;
 		bool res = SendSocket(logger, conn.socket, data, dataSize);
@@ -268,7 +268,7 @@ namespace uba
 		if (!EnsureInitialized(logger))
 			return false;
 
-		SCOPED_WRITE_LOCK(m_listenEntriesLock, lock);
+		SCOPED_FUTEX(m_listenEntriesLock, lock);
 
 		auto prevListenEntryCount = int(m_listenEntries.size());
 
@@ -331,7 +331,7 @@ namespace uba
 
 	void NetworkBackendTcp::StopListen()
 	{
-		SCOPED_WRITE_LOCK(m_listenEntriesLock, lock);
+		SCOPED_FUTEX(m_listenEntriesLock, lock);
 		for (auto& e : m_listenEntries)
 			e.socket = INVALID_SOCKET;
 		for (auto& e : m_listenEntries)
@@ -453,7 +453,7 @@ namespace uba
 			//SetRecvBuf(logger, clientSocket, 50*1024*1024);
 			//SetSendBuf(logger, clientSocket, 50*1024*1024);
 
-			SCOPED_WRITE_LOCK(m_connectionsLock, lock);
+			SCOPED_FUTEX(m_connectionsLock, lock);
 			auto it = m_connections.emplace(m_connections.end(), logger, clientSocket);
 			auto& conn = *it;
 			conn.recvThread.Start([this, connPtr = &conn] { ThreadRecv(*connPtr); return 0; }, TC("UbaTcpRecv"));
@@ -464,7 +464,7 @@ namespace uba
 				ShutdownSocket(logger, clientSocket, TC("ThreadListen"));
 				conn.ready.Set();
 				conn.recvThread.Wait();
-				SCOPED_WRITE_LOCK(m_connectionsLock, lock2);
+				SCOPED_FUTEX(m_connectionsLock, lock2);
 				m_connections.erase(it);
 				continue;
 			}
@@ -540,7 +540,7 @@ namespace uba
 		SOCKET s = connection.socket;
 
 		{
-			SCOPED_WRITE_LOCK(connection.sendLock, lock);
+			SCOPED_FUTEX(connection.sendLock, lock);
 			connection.socket = INVALID_SOCKET;
 		}
 		if (auto cb = connection.disconnectCallback)
@@ -713,7 +713,7 @@ namespace uba
 		// Socket is good, cancel the socket close scope and break out of the loop.
 		socketClose.Cancel();
 
-		SCOPED_WRITE_LOCK(m_connectionsLock, lock);
+		SCOPED_FUTEX(m_connectionsLock, lock);
 		auto it = m_connections.emplace(m_connections.end(), logger, socketFd);
 		auto& conn = *it;
 		conn.recvThread.Start([this, connPtr = &conn] { ThreadRecv(*connPtr); return 0; }, TC("UbaTcpRecv"));
@@ -724,7 +724,7 @@ namespace uba
 			ShutdownSocket(logger, conn.socket, TC("Connect"));
 			conn.ready.Set();
 			conn.recvThread.Wait();
-			SCOPED_WRITE_LOCK(m_connectionsLock, lock2);
+			SCOPED_FUTEX(m_connectionsLock, lock2);
 			m_connections.erase(it);
 			return false;
 		}
@@ -740,7 +740,7 @@ namespace uba
 
 	void NetworkBackendTcp::DeleteConnection(void* connection)
 	{
-		SCOPED_WRITE_LOCK(m_connectionsLock, lock);
+		SCOPED_FUTEX(m_connectionsLock, lock);
 		for (auto it=m_connections.begin();it!=m_connections.end(); ++it)
 		{
 			Connection& c = *it;
diff --git a/Engine/Source/Programs/UnrealBuildAccelerator/Common/Private/UbaNetworkClient.cpp b/Engine/Source/Programs/UnrealBuildAccelerator/Common/Private/UbaNetworkClient.cpp
--- a/Engine/Source/Programs/UnrealBuildAccelerator/Common/Private/UbaNetworkClient.cpp
+++ b/Engine/Source/Programs/UnrealBuildAccelerator/Common/Private/UbaNetworkClient.cpp
@@ -108,7 +108,7 @@ namespace uba
 
 				if (!rc.error)
 				{
-					SCOPED_WRITE_LOCK(rc.client.m_serverUidLock, lock);
+					SCOPED_FUTEX(rc.client.m_serverUidLock, lock);
 					if (rc.client.m_serverUid == Guid())
 						rc.client.m_serverUid = serverUid;
 					else if (rc.client.m_serverUid != serverUid) // Seems like two different servers tried to connect to this client.. keep the first one and ignore the others
@@ -188,7 +188,7 @@ namespace uba
 		if (m_connectionCount.fetch_add(1) != 0)
 			return true;
 
-		SCOPED_WRITE_LOCK(m_onConnectedFunctionsLock, lock);
+		SCOPED_FUTEX(m_onConnectedFunctionsLock, lock);
 		for (auto& f : m_onConnectedFunctions)
 			f();
 		m_isConnected.Set();
@@ -218,7 +218,7 @@ namespace uba
 		connection->connected = 1;
 		connection->backend = &backend;
 
-		SCOPED_WRITE_LOCK(m_connectionsItLock, l); // Take this lock to make sure callbacks are set before connection is used
+		SCOPED_FUTEX(m_connectionsItLock, l); // Take this lock to make sure callbacks are set before connection is used
 		m_connectionsIt = --m_connections.end();
 
 		m_logger.Detail(TC("Connected to server... (0x%p)"), backendConnection);
@@ -391,7 +391,7 @@ namespace uba
 
 	void NetworkClient::RegisterOnConnected(const OnConnectedFunction& function)
 	{
-		SCOPED_WRITE_LOCK(m_onConnectedFunctionsLock, lock);
+		SCOPED_FUTEX(m_onConnectedFunctionsLock, lock);
 		m_onConnectedFunctions.push_back(function);
 		if (!m_isConnected.IsSet(0))
 			return;
@@ -471,7 +471,7 @@ namespace uba
 	bool NetworkClient::Send(NetworkMessage& message, void* response, u32 responseCapacity, bool async)
 	{
 		SCOPED_READ_LOCK(m_connectionsLock, connectionLock);
-		SCOPED_WRITE_LOCK(m_connectionsItLock, connectionItLock);
+		SCOPED_FUTEX(m_connectionsItLock, connectionItLock);
 		if (m_connectionsIt == m_connections.end())
 		{
 			if (m_isDisconnecting)
diff --git a/Engine/Source/Programs/UnrealBuildAccelerator/Common/Private/UbaNetworkServer.cpp b/Engine/Source/Programs/UnrealBuildAccelerator/Common/Private/UbaNetworkServer.cpp
--- a/Engine/Source/Programs/UnrealBuildAccelerator/Common/Private/UbaNetworkServer.cpp
+++ b/Engine/Source/Programs/UnrealBuildAccelerator/Common/Private/UbaNetworkServer.cpp
@@ -64,7 +64,7 @@ namespace uba
 		void Stop(NetworkServer& server)
 		{
 			m_loop = false;
-			SCOPED_WRITE_LOCK(server.m_availableWorkersLock, lock);
+			SCOPED_FUTEX(server.m_availableWorkersLock, lock);
 			while (m_inUse)
 			{
 				m_context->workAvailable.Set();
@@ -212,7 +212,7 @@ namespace uba
 			}
 			else
 			{
-				SCOPED_WRITE_LOCK(conn.m_server.m_cryptoKeysLock, lock);
+				SCOPED_FUTEX(conn.m_server.m_cryptoKeysLock, lock);
 				auto& keys = conn.m_server.m_cryptoKeys;
 				u64 time = GetTime();
 				for (auto it=keys.begin(); it!=keys.end();)
@@ -284,7 +284,7 @@ namespace uba
 			if (!conn.SendInitialResponse(0))
 				return false;
 
-			SCOPED_WRITE_LOCK(conn.m_shutdownLock, shutdownLock);
+			SCOPED_FUTEX(conn.m_shutdownLock, shutdownLock);
 
 			SCOPED_WRITE_LOCK(server.m_clientsLock, clientsLock);
 			u32 clientId = 0;
@@ -378,7 +378,7 @@ namespace uba
 
 		bool SetShouldDisconnect()
 		{
-			SCOPED_WRITE_LOCK(m_shutdownLock, lock);
+			SCOPED_FUTEX(m_shutdownLock, lock);
 			bool isConnected = !m_shouldDisconnect;
 			m_shouldDisconnect = true;
 			return isConnected;
@@ -392,7 +392,7 @@ namespace uba
 
 		void TestDisconnect()
 		{
-			SCOPED_WRITE_LOCK(m_shutdownLock, lock);
+			SCOPED_FUTEX(m_shutdownLock, lock);
 			if (!m_shouldDisconnect)
 				return;
 			if (m_disconnected)
@@ -412,7 +412,7 @@ namespace uba
 
 		NetworkServer& m_server;
 		NetworkBackend& m_backend;
-		ReaderWriterLock m_shutdownLock;
+		Futex m_shutdownLock;
 		Client* m_client = nullptr;
 		sockaddr m_remoteSockAddr;
 		CryptoKey m_cryptoKey;
@@ -459,7 +459,7 @@ namespace uba
 	bool ConnectionInfo::ShouldDisconnect() const
 	{
 		auto& conn = *(NetworkServer::Connection*)internalData;
-		SCOPED_WRITE_LOCK(conn.m_shutdownLock, lock);
+		SCOPED_FUTEX(conn.m_shutdownLock, lock);
 		return conn.m_shouldDisconnect;
 	}
 
@@ -575,7 +575,7 @@ namespace uba
 			while (true)
 			{
 				AdditionalWork work;
-				SCOPED_WRITE_LOCK(server.m_additionalWorkLock, lock);
+				SCOPED_FUTEX(server.m_additionalWorkLock, lock);
 				if (server.m_additionalWork.empty())
 					break;
 				work = std::move(server.m_additionalWork.front());
@@ -590,8 +590,8 @@ namespace uba
 			// is present before making ourself available to avoid
 			// a race where AddWork would not see this thread in the
 			// available list after adding some work.
-			SCOPED_WRITE_LOCK(server.m_availableWorkersLock, lock1);
-			SCOPED_READ_LOCK(server.m_additionalWorkLock, lock2);
+			SCOPED_FUTEX(server.m_availableWorkersLock, lock1);
+			SCOPED_FUTEX_READ(server.m_additionalWorkLock, lock2);
 			// Verify there is not additional work while we hold both lock
 			// and only add ourself as available if no additional work is present.
 			if (!server.m_additionalWork.empty())
@@ -669,12 +669,12 @@ namespace uba
 	void NetworkServer::DisconnectClients()
 	{
 		{
-			SCOPED_WRITE_LOCK(m_availableWorkersLock, lock);
+			SCOPED_FUTEX(m_availableWorkersLock, lock);
 			m_workersEnabled = false;
 			m_workerAvailable.Set();
 		}
 		{
-			SCOPED_WRITE_LOCK(m_addConnectionsLock, lock);
+			SCOPED_FUTEX(m_addConnectionsLock, lock);
 			m_addConnections.clear();
 		}
 
@@ -711,7 +711,7 @@ namespace uba
 		CryptoKey key = Crypto::CreateKey(m_logger, cryptoKey128);
 		if (key == InvalidCryptoKey)
 			return false;
-		SCOPED_WRITE_LOCK(m_cryptoKeysLock, lock);
+		SCOPED_FUTEX(m_cryptoKeysLock, lock);
 		m_cryptoKeys.push_back(CryptoEntry{key, expirationTime});
 		return true;
 	}
@@ -723,7 +723,7 @@ namespace uba
 
 	bool NetworkServer::AddClient(NetworkBackend& backend, const tchar* ip, u16 port, const u8* cryptoKey128)
 	{
-		SCOPED_WRITE_LOCK(m_addConnectionsLock, lock);
+		SCOPED_FUTEX(m_addConnectionsLock, lock);
 		if (!m_workersEnabled)
 			return false;
 
@@ -840,7 +840,7 @@ namespace uba
 
 	void NetworkServer::AddWork(const WorkFunction& work, u32 count, const tchar* desc, const Color& color, bool highPriority)
 	{
-		SCOPED_WRITE_LOCK(m_additionalWorkLock, lock);
+		SCOPED_FUTEX(m_additionalWorkLock, lock);
 		for (u32 i = 0; i != count; ++i)
 		{
 			if (highPriority)
@@ -858,7 +858,7 @@ namespace uba
 		}
 		lock.Leave();
 
-		SCOPED_WRITE_LOCK(m_availableWorkersLock, lock2);
+		SCOPED_FUTEX(m_availableWorkersLock, lock2);
 		if (!m_workersEnabled)
 			return;
 		while (count--)
@@ -941,12 +941,12 @@ namespace uba
 
 	bool NetworkServer::DoAdditionalWork()
 	{
-		SCOPED_WRITE_LOCK(m_additionalWorkLock, lock);
+		SCOPED_FUTEX(m_additionalWorkLock, lock);
 		if (m_additionalWork.empty())
 		{
 			lock.Leave();
 
-			SCOPED_WRITE_LOCK(m_availableWorkersLock, lock2);
+			SCOPED_FUTEX(m_availableWorkersLock, lock2);
 			if (m_createdWorkerCount != m_maxWorkerCount)
 				return false;
 			lock2.Leave();
@@ -1060,7 +1060,7 @@ namespace uba
 	{
 		while (true)
 		{
-			SCOPED_WRITE_LOCK(m_availableWorkersLock, lock);
+			SCOPED_FUTEX(m_availableWorkersLock, lock);
 			if (!m_workersEnabled)
 				return nullptr;
 			if (auto worker = PopWorkerNoLock())
@@ -1101,7 +1101,7 @@ namespace uba
 
 	void NetworkServer::PushWorker(Worker* worker)
 	{
-		SCOPED_WRITE_LOCK(m_availableWorkersLock, lock);
+		SCOPED_FUTEX(m_availableWorkersLock, lock);
 		PushWorkerNoLock(worker);
 	}
 
@@ -1127,7 +1127,7 @@ namespace uba
 
 	void NetworkServer::FlushWorkers()
 	{
-		SCOPED_WRITE_LOCK(m_availableWorkersLock, lock);
+		SCOPED_FUTEX(m_availableWorkersLock, lock);
 		while (auto worker = m_firstActiveWorker)
 		{
 			lock.Leave();
@@ -1207,7 +1207,7 @@ namespace uba
 				u32 toAdd = connectionCount - c.connectionCount;
 
 				auto connPtr = (NetworkServer::Connection*)connectionInfo.internalData;
-				SCOPED_WRITE_LOCK(m_addConnectionsLock, lock2);
+				SCOPED_FUTEX(m_addConnectionsLock, lock2);
 				for (u32 i = 0; i != toAdd; ++i)
 				{
 					m_addConnections.emplace_back([this, connPtr]()
diff --git a/Engine/Source/Programs/UnrealBuildAccelerator/Common/Private/UbaProcess.cpp b/Engine/Source/Programs/UnrealBuildAccelerator/Common/Private/UbaProcess.cpp
--- a/Engine/Source/Programs/UnrealBuildAccelerator/Common/Private/UbaProcess.cpp
+++ b/Engine/Source/Programs/UnrealBuildAccelerator/Common/Private/UbaProcess.cpp
@@ -144,7 +144,7 @@ namespace uba
 	,	m_readEvent(*(m_comMemory.memory ? (new (m_comMemory.memory + sizeof(Event)*2) Event) : (Event*)nullptr))
 		#endif
 	,	m_detourEnabled(detourEnabled)
-	,	m_writtenFilesLock(parent ? parent->m_writtenFilesLock : *new ReaderWriterLock())
+	,	m_writtenFilesLock(parent ? parent->m_writtenFilesLock : *new Futex())
 	,	m_writtenFiles(parent ? parent->m_writtenFiles : *new UnorderedMap<TString, WrittenFile>())
 	,	m_tempFilesLock(parent ? parent->m_tempFilesLock : *new ReaderWriterLock())
 	,	m_tempFiles(parent ? parent->m_tempFiles : *new UnorderedMap<StringKey, WrittenFile>)
@@ -162,7 +162,7 @@ namespace uba
 		UBA_ASSERT(m_refCount == 0);
 		{
 			#if !PLATFORM_WINDOWS
-			SCOPED_WRITE_LOCK(m_comMemoryLock, lock);
+			SCOPED_FUTEX(m_comMemoryLock, lock);
 			#endif
 			if (m_comMemory.memory)
 				m_cancelEvent.Set();
@@ -448,7 +448,7 @@ namespace uba
 		m_cancelEvent.Set();
 		#else
 		m_cancelled = true;
-		SCOPED_WRITE_LOCK(m_comMemoryLock, lock);
+		SCOPED_FUTEX(m_comMemoryLock, lock);
 		if (m_comMemory.memory)
 			m_cancelEvent.Set();
 		#endif
@@ -605,7 +605,7 @@ namespace uba
 		if (m_comMemory.memory)
 		{
 			#if !PLATFORM_WINDOWS
-			SCOPED_WRITE_LOCK(m_comMemoryLock, lock);
+			SCOPED_FUTEX(m_comMemoryLock, lock);
 			m_cancelEvent.~Event();
 			m_writeEvent.~Event();
 			m_readEvent.~Event();
@@ -738,7 +738,7 @@ namespace uba
 			m_session.m_logger.Log(LogEntryType_Warning, line.c_str(), u32(line.size()));
 		if (m_startInfo.logLineFunc)
 			m_startInfo.logLineFunc(m_startInfo.logLineUserData, line.c_str(), u32(line.size()), logType);
-		SCOPED_WRITE_LOCK(m_logLinesLock, l);
+		SCOPED_FUTEX(m_logLinesLock, l);
 		m_logLines.push_back({ std::move(line), logType });
 	}
 
@@ -1389,7 +1389,7 @@ namespace uba
 	{
 		Vector<WrittenFile*> files;
 		TimerScope ts(m_processStats.writeFiles);
-		SCOPED_WRITE_LOCK(m_writtenFilesLock, lock);
+		SCOPED_FUTEX(m_writtenFilesLock, lock);
 		for (auto& kv : m_writtenFiles)
 		{
 			if (kv.second.owner != this)
@@ -1449,7 +1449,7 @@ namespace uba
 
 	u32 ProcessImpl::InternalCreateProcess(void* environment, FileMappingHandle communicationHandle, u64 communicationOffset)
 	{
-		SCOPED_WRITE_LOCK(m_initLock, initLock);
+		SCOPED_FUTEX(m_initLock, initLock);
 		Logger& logger = m_session.m_logger;
 
 #if PLATFORM_WINDOWS
@@ -1463,7 +1463,7 @@ namespace uba
 		{
 			if (IsRunningArm())
 			{
-				SCOPED_WRITE_LOCK(m_session.m_isX64ApplicationLock, lock);
+				SCOPED_FUTEX(m_session.m_isX64ApplicationLock, lock);
 				auto insres = m_session.m_isX64Application.try_emplace(m_realApplication);
 				if (insres.second)
 				{
@@ -2019,7 +2019,7 @@ namespace uba
 
 	u32 ProcessImpl::InternalExitProcess(bool cancel)
 	{
-		SCOPED_WRITE_LOCK(m_initLock, lock);
+		SCOPED_FUTEX(m_initLock, lock);
 		Logger& logger = m_session.m_logger;
 
 		ProcHandle handle = m_nativeProcessHandle;
diff --git a/Engine/Source/Programs/UnrealBuildAccelerator/Common/Private/UbaSession.cpp b/Engine/Source/Programs/UnrealBuildAccelerator/Common/Private/UbaSession.cpp
--- a/Engine/Source/Programs/UnrealBuildAccelerator/Common/Private/UbaSession.cpp
+++ b/Engine/Source/Programs/UnrealBuildAccelerator/Common/Private/UbaSession.cpp
@@ -498,12 +498,12 @@ namespace uba
 	u32 Session::AddFileMapping(StringKey fileNameKey, const tchar* fileName, const tchar* newFileName, u64 fileSize)
 	{
 		UBA_ASSERT(fileNameKey != StringKeyZero);
-		SCOPED_WRITE_LOCK(m_fileMappingTableLookupLock, lookupLock);
+		SCOPED_FUTEX(m_fileMappingTableLookupLock, lookupLock);
 		auto insres = m_fileMappingTableLookup.try_emplace(fileNameKey);
 		FileMappingEntry& entry = insres.first->second;
 		lookupLock.Leave();
 
-		SCOPED_WRITE_LOCK(entry.lock, entryCs);
+		SCOPED_FUTEX(entry.lock, entryCs);
 
 		if (entry.handled)
 		{
@@ -532,12 +532,12 @@ namespace uba
 	{
 		TimerScope ts(Stats().waitMmapFromFile);
 
-		SCOPED_WRITE_LOCK(m_fileMappingTableLookupLock, lookupLock);
+		SCOPED_FUTEX(m_fileMappingTableLookupLock, lookupLock);
 		auto insres = m_fileMappingTableLookup.try_emplace(fileNameKey);
 		FileMappingEntry& entry = insres.first->second;
 		lookupLock.Leave();
 
-		SCOPED_WRITE_LOCK(entry.lock, entryCs);
+		SCOPED_FUTEX(entry.lock, entryCs);
 
 		if (entry.handled)
 		{
@@ -656,12 +656,12 @@ namespace uba
 		//workName.Append(TC("MM:")).Append(len > 30 ? fileName + (len - 30) : fileName);
 		//TrackWorkScope tws(*m_workManager, workName.data);
 
-		SCOPED_WRITE_LOCK(m_fileMappingTableLookupLock, lookupLock);
+		SCOPED_FUTEX(m_fileMappingTableLookupLock, lookupLock);
 		auto insres = m_fileMappingTableLookup.try_emplace(fileNameKey);
 		FileMappingEntry& entry = insres.first->second;
 		lookupLock.Leave();
 
-		SCOPED_WRITE_LOCK(entry.lock, entryCs);
+		SCOPED_FUTEX(entry.lock, entryCs);
 
 		if (entry.handled)
 		{
@@ -769,13 +769,13 @@ namespace uba
 			
 		#if 0//_DEBUG  // Bring this back, turned off right now because a few lines above the call to this method we add a mapping
 		{
-			SCOPED_WRITE_LOCK(m_fileMappingTableLookupLock, lookupLock);
+			SCOPED_FUTEX(m_fileMappingTableLookupLock, lookupLock);
 			auto findIt = m_fileMappingTableLookup.find(fileNameKey);
 			if (findIt != m_fileMappingTableLookup.end())
 			{
 				FileMappingEntry& entry = findIt->second;
 				lookupLock.Leave();
-				SCOPED_WRITE_LOCK(entry.lock, entryCs);
+				SCOPED_FUTEX(entry.lock, entryCs);
 				UBA_ASSERT(!entry.mapping);
 			}
 		}
@@ -1305,13 +1305,13 @@ namespace uba
 
 		m_storage.RegisterExternalFileMappingsProvider([this](Storage::ExternalFileMapping& out, StringKey fileNameKey, const tchar* fileName)
 			{
-				SCOPED_READ_LOCK(m_fileMappingTableLookupLock, lookupLock);
+				SCOPED_FUTEX_READ(m_fileMappingTableLookupLock, lookupLock);
 				auto findIt = m_fileMappingTableLookup.find(fileNameKey);
 				if (findIt == m_fileMappingTableLookup.end())
 					return false;
 				FileMappingEntry& entry = findIt->second;
 				lookupLock.Leave();
-				SCOPED_READ_LOCK(entry.lock, entryLock);
+				SCOPED_FUTEX_READ(entry.lock, entryLock);
 				if (!entry.handled || !entry.success || !entry.mapping.IsValid())
 					return false;
 				out.handle = entry.mapping;
@@ -1366,7 +1366,7 @@ namespace uba
 		{
 			Vector<ProcessHandle> processes;
 			{
-				SCOPED_WRITE_LOCK(m_processesLock, lock);
+				SCOPED_FUTEX(m_processesLock, lock);
 				isEmpty = m_processes.empty();
 				processes.reserve(m_processes.size());
 				for (auto& pair : m_processes)
@@ -1387,7 +1387,7 @@ namespace uba
 			#if PLATFORM_WINDOWS
 			if (m_processJobObject != NULL)
 			{
-				SCOPED_WRITE_LOCK(m_processJobObjectLock, lock);
+				SCOPED_FUTEX(m_processJobObjectLock, lock);
 				CloseHandle(m_processJobObject);
 				m_processJobObject = NULL;
 			}
@@ -1831,7 +1831,7 @@ namespace uba
 		if (!process.IsChild() || m_traceChildProcesses)
 			m_trace.ProcessAdded(sessionId, processId, ToView(startInfo.GetDescription()), ToView(startInfo.breadcrumbs));
 
-		SCOPED_WRITE_LOCK(m_processesLock, lock);
+		SCOPED_FUTEX(m_processesLock, lock);
 		bool success = m_processes.try_emplace(processId, ProcessHandle(&process)).second;
 		UBA_ASSERT(success);(void)success;
 	}
@@ -1857,12 +1857,12 @@ namespace uba
 			Vector<ProcessLogLine> emptyLines;
 			auto& logLines = (exitCode != 0 || m_detailedTrace) ? process.m_logLines : emptyLines;
 			m_trace.ProcessExited(id, exitCode, writer.GetData(), writer.GetPosition(), logLines);
-			SCOPED_WRITE_LOCK(m_processStatsLock, lock);
+			SCOPED_FUTEX(m_processStatsLock, lock);
 			m_processStats.Add(process.m_processStats);
 			m_stats.Add(process.m_sessionStats);
 		}
 
-		SCOPED_WRITE_LOCK(m_processesLock, lock);
+		SCOPED_FUTEX(m_processesLock, lock);
 		m_deadProcesses.emplace_back(&process); // Here to prevent Process thread call trigger a delete of Process which causes a deadlock
 		auto& stats = m_applicationStats[applicationName.data];
 		stats.count++;
@@ -1873,7 +1873,7 @@ namespace uba
 
 	void Session::FlushDeadProcesses()
 	{
-		SCOPED_WRITE_LOCK(m_processesLock, lock);
+		SCOPED_FUTEX(m_processesLock, lock);
 		Vector<ProcessHandle> deadProcesses;
 		deadProcesses.swap(m_deadProcesses);
 		lock.Leave();
@@ -1901,7 +1901,7 @@ namespace uba
 			out.mappedFileTableSize = m_fileMappingTableSize;
 		}
 		{
-			SCOPED_READ_LOCK(m_fileMappingTableLookupLock, l);
+			SCOPED_FUTEX_READ(m_fileMappingTableLookupLock, l);
 			out.mappedFileTableCount = (u32)m_fileMappingTableLookup.size();
 		}
 		return true;
@@ -1928,7 +1928,7 @@ namespace uba
 
 	u32 Session::GetActiveProcessCount()
 	{
-		SCOPED_READ_LOCK(m_processesLock, cs);
+		SCOPED_FUTEX_READ(m_processesLock, cs);
 		return u32(m_processes.size());
 	}
 
@@ -2093,7 +2093,7 @@ namespace uba
 
 	void* Session::GetProcessEnvironmentVariables()
 	{
-		SCOPED_WRITE_LOCK(m_environmentVariablesLock, lock);
+		SCOPED_FUTEX(m_environmentVariablesLock, lock);
 		if (!m_environmentVariables.empty())
 			return m_environmentVariables.data();
 
@@ -2262,7 +2262,7 @@ namespace uba
 
 		if (m_runningRemote && !fileName.StartsWith(m_tempPath.data))
 		{
-			SCOPED_WRITE_LOCK(m_outputFilesLock, lock);
+			SCOPED_FUTEX(m_outputFilesLock, lock);
 			auto insres = m_outputFiles.try_emplace(fileName.data);
 			if (insres.second)
 			{
@@ -2279,7 +2279,7 @@ namespace uba
 			out.fileName.Append(fileName);
 		}
 
-		SCOPED_WRITE_LOCK(m_activeFilesLock, lock);
+		SCOPED_FUTEX(m_activeFilesLock, lock);
 		u32 wantsOnCloseId = m_wantsOnCloseIdCounter++;
 		out.closeId = wantsOnCloseId;
 		auto insres = m_activeFiles.try_emplace(wantsOnCloseId);
@@ -2293,7 +2293,7 @@ namespace uba
 
 	void RemoveWrittenFile(ProcessImpl& process, const TString& name)
 	{
-		SCOPED_WRITE_LOCK(process.m_writtenFilesLock, writtenLock);
+		SCOPED_FUTEX(process.m_writtenFilesLock, writtenLock);
 		auto& writtenFiles = process.m_writtenFiles;
 		auto findIt = writtenFiles.find(name);
 		if (findIt == writtenFiles.end())
@@ -2308,7 +2308,7 @@ namespace uba
 
 	bool Session::CloseFile(CloseFileResponse& out, const CloseFileMessage& msg)
 	{
-		SCOPED_WRITE_LOCK(m_activeFilesLock, lock);
+		SCOPED_FUTEX(m_activeFilesLock, lock);
 		auto findIt = m_activeFiles.find(msg.closeId);
 		if (findIt == m_activeFiles.end())
 			return m_logger.Error(TC("This should not happen. Got unknown closeId %u - %s"), msg.closeId, msg.fileName.data);
@@ -2344,7 +2344,7 @@ namespace uba
 					msgName = msg.newName.data;
 			}
 
-			SCOPED_WRITE_LOCK(msg.process.m_writtenFilesLock, writtenLock);
+			SCOPED_FUTEX(msg.process.m_writtenFilesLock, writtenLock);
 			auto insres = msg.process.m_writtenFiles.try_emplace(name);
 			WrittenFile& writtenFile = insres.first->second;
 
@@ -2460,12 +2460,12 @@ namespace uba
 	{
 		if (msg.closeId != 0)
 		{
-			SCOPED_WRITE_LOCK(m_activeFilesLock, lock);
+			SCOPED_FUTEX(m_activeFilesLock, lock);
 			m_activeFiles.erase(msg.closeId);
 		}
 
 		{
-			SCOPED_WRITE_LOCK(m_outputFilesLock, lock);
+			SCOPED_FUTEX(m_outputFilesLock, lock);
 			m_outputFiles.erase(msg.fileName.data);
 		}
 
@@ -2482,7 +2482,7 @@ namespace uba
 		out.fromName.Append(msg.fromName);
 		out.toName.Append(msg.toName);
 
-		SCOPED_WRITE_LOCK(m_activeFilesLock, lock);
+		SCOPED_FUTEX(m_activeFilesLock, lock);
 		u32 closeId = m_wantsOnCloseIdCounter++;
 		if (!m_activeFiles.try_emplace(closeId, ActiveFile{ msg.toName.data, msg.toKey }).second)
 		{
@@ -2749,11 +2749,11 @@ namespace uba
 			{
 				StringBuffer<> name;
 				Storage::GetMappingString(name, file.mappingHandle, 0);
-				SCOPED_WRITE_LOCK(m_fileMappingTableLookupLock, lookupLock);
+				SCOPED_FUTEX(m_fileMappingTableLookupLock, lookupLock);
 				auto insres = m_fileMappingTableLookup.try_emplace(file.key);
 				FileMappingEntry& entry = insres.first->second;
 				lookupLock.Leave();
-				SCOPED_WRITE_LOCK(entry.lock, entryCs);
+				SCOPED_FUTEX(entry.lock, entryCs);
 				entry.handled = true;
 				entry.mapping = file.mappingHandle;
 				entry.mappingOffset = 0;
@@ -2941,7 +2941,7 @@ namespace uba
 
 		{
 			StringKey sourceFileKey = CaseInsensitiveFs ? ToStringKeyLower(ToView(sourceFile)) : ToStringKey(ToView(sourceFile));
-			SCOPED_WRITE_LOCK(m_virtualSourceFilesLock, virtualSourceFilesLock);
+			SCOPED_FUTEX(m_virtualSourceFilesLock, virtualSourceFilesLock);
 			auto insres = m_virtualSourceFiles.try_emplace(sourceFileKey);
 			if (insres.second)
 			{
@@ -2957,12 +2957,12 @@ namespace uba
 		}
 
 
-		SCOPED_WRITE_LOCK(m_fileMappingTableLookupLock, lookupLock);
+		SCOPED_FUTEX(m_fileMappingTableLookupLock, lookupLock);
 		auto insres = m_fileMappingTableLookup.try_emplace(fileNameKey);
 		FileMappingEntry& entry = insres.first->second;
 		lookupLock.Leave();
 
-		SCOPED_WRITE_LOCK(entry.lock, entryLock);
+		SCOPED_FUTEX(entry.lock, entryLock);
 		
 		if (entry.handled)
 			return m_logger.Error(TC("Virtual file %s has already been registered"), filePath.data);
diff --git a/Engine/Source/Programs/UnrealBuildAccelerator/Common/Private/UbaSessionClient.cpp b/Engine/Source/Programs/UnrealBuildAccelerator/Common/Private/UbaSessionClient.cpp
--- a/Engine/Source/Programs/UnrealBuildAccelerator/Common/Private/UbaSessionClient.cpp
+++ b/Engine/Source/Programs/UnrealBuildAccelerator/Common/Private/UbaSessionClient.cpp
@@ -218,11 +218,11 @@ namespace uba
 	bool SessionClient::GetCasKeyForFile(CasKey& out, u32 processId, const StringView& fileName, const StringKey& fileNameKey)
 	{
 		TimerScope waitTimer(Stats().waitGetFileMsg);
-		SCOPED_WRITE_LOCK(m_nameToHashLookupLock, lock);
+		SCOPED_FUTEX(m_nameToHashLookupLock, lock);
 		auto insres = m_nameToHashLookup.try_emplace(fileNameKey);
 		HashRec& rec = insres.first->second;
 		lock.Leave();
-		SCOPED_WRITE_LOCK(rec.lock, lock2);
+		SCOPED_FUTEX(rec.lock, lock2);
 		if (rec.key == CasKeyZero)//!rec.serverTime)
 		{
 			waitTimer.Cancel();
@@ -328,12 +328,12 @@ namespace uba
 		UBA_ASSERT(application && *application);
 		bool isAbsolute = IsAbsolutePath(startInfo.application);
 
-		SCOPED_WRITE_LOCK(m_handledApplicationEnvironmentsLock, environmentslock);
+		SCOPED_FUTEX(m_handledApplicationEnvironmentsLock, environmentslock);
 		auto insres = m_handledApplicationEnvironments.try_emplace(application);
 		environmentslock.Leave();
 
 		ApplicationEnvironment& appEnv = insres.first->second;
-		SCOPED_WRITE_LOCK(appEnv.lock, lock);
+		SCOPED_FUTEX(appEnv.lock, lock);
 
 		if (!appEnv.realApplication.empty())
 		{
@@ -456,7 +456,7 @@ namespace uba
 		StringBuffer<> lower;
 		lower.Append(applicationDir).Append(PathSeparator).Append(binaryName);
 		lower.MakeLower();
-		SCOPED_WRITE_LOCK(m_binFileLock, lock);
+		SCOPED_FUTEX(m_binFileLock, lock);
 
 		auto insres = m_writtenBinFiles.try_emplace(lower.data, casKey);
 		if (!insres.second)
@@ -544,7 +544,7 @@ namespace uba
 		// Code for doing retry if failing to decompress casfile. We've seen cases of corrupt cas files on clients
 		bool shouldRetry = true;
 		FileMappingEntry* retryEntry = nullptr;
-		auto retryEntryGuard = MakeGuard([&]() { if (retryEntry) retryEntry->lock.LeaveWrite(); });
+		auto retryEntryGuard = MakeGuard([&]() { if (retryEntry) retryEntry->lock.Leave(); });
 
 		while (true)
 		{
@@ -582,12 +582,12 @@ namespace uba
 					StorageStats& stats = m_storage.Stats();
 					TimerScope ts(stats.ensureCas);
 
-					SCOPED_WRITE_LOCK(m_fileMappingTableLookupLock, lookupLock);
+					SCOPED_FUTEX(m_fileMappingTableLookupLock, lookupLock);
 					auto insres = m_fileMappingTableLookup.try_emplace(fileNameKey);
 					FileMappingEntry& entry = insres.first->second;
 					lookupLock.Leave();
 
-					SCOPED_WRITE_LOCK(entry.lock, entryCs);
+					SCOPED_FUTEX(entry.lock, entryCs);
 					ts.Leave();
 
 					if (entry.handled)
@@ -637,10 +637,10 @@ namespace uba
 
 						// We need to take a lock around the file map entry since there might be another thread also wanting to map this
 						{
-							SCOPED_WRITE_LOCK(m_fileMappingTableLookupLock, lookupLock);
+							SCOPED_FUTEX(m_fileMappingTableLookupLock, lookupLock);
 							retryEntry = &m_fileMappingTableLookup.try_emplace(fileNameKey).first->second;
 							lookupLock.Leave();
-							retryEntry->lock.EnterWrite();
+							retryEntry->lock.Enter();
 							retryEntry->handled = false;
 						}
 
@@ -731,12 +731,12 @@ namespace uba
 		if (msg.closeId != 0)
 		{
 			UBA_ASSERTF(false, TC("This has not been tested properly"));
-			SCOPED_WRITE_LOCK(m_activeFilesLock, lock);
+			SCOPED_FUTEX(m_activeFilesLock, lock);
 			sendDelete = m_activeFiles.erase(msg.closeId) == 0;
 		}
 
 		{
-			SCOPED_WRITE_LOCK(m_outputFilesLock, lock);
+			SCOPED_FUTEX(m_outputFilesLock, lock);
 			sendDelete = m_outputFiles.erase(msg.fileName.data) == 0 && sendDelete;
 		}
 
@@ -777,7 +777,7 @@ namespace uba
 
 	bool SessionClient::CopyFile(CopyFileResponse& out, const CopyFileMessage& msg)
 	{
-		SCOPED_WRITE_LOCK(m_outputFilesLock, lock);
+		SCOPED_FUTEX(m_outputFilesLock, lock);
 		auto findIt = m_outputFiles.find(msg.fromName.data);
 		if (findIt == m_outputFiles.end())
 		{
@@ -826,7 +826,7 @@ namespace uba
 		auto& process = msg.process;
 
 		{
-			SCOPED_WRITE_LOCK(process.m_writtenFilesLock, lock);
+			SCOPED_FUTEX(process.m_writtenFilesLock, lock);
 			auto& writtenFiles = process.m_writtenFiles;
 			auto findIt = writtenFiles.find(fromName);
 			if (findIt != writtenFiles.end())
@@ -846,7 +846,7 @@ namespace uba
 
 		bool sendMove = true;
 		{
-			SCOPED_WRITE_LOCK(m_outputFilesLock, lock);
+			SCOPED_FUTEX(m_outputFilesLock, lock);
 			auto findIt = m_outputFiles.find(fromName);
 			if (findIt != m_outputFiles.end())
 			{
@@ -878,7 +878,7 @@ namespace uba
 		const tchar* fromName = msg.fileName.data;
 
 		{
-			SCOPED_WRITE_LOCK(msg.process.m_writtenFilesLock, lock);
+			SCOPED_FUTEX(msg.process.m_writtenFilesLock, lock);
 			auto& writtenFiles = msg.process.m_writtenFiles;
 			auto findIt = writtenFiles.find(fromName);
 			if (findIt != writtenFiles.end())
@@ -943,11 +943,11 @@ namespace uba
 		// TODO: There is a potential risk here where two different applications asks for the full name of a file
 		// and they have different bin/working dir.. and there are two versions of this file.
 
-		SCOPED_WRITE_LOCK(m_nameToNameLookupLock, lock);
+		SCOPED_FUTEX(m_nameToNameLookupLock, lock);
 		auto insres = m_nameToNameLookup.try_emplace(msg.fileName.data);
 		NameRec& rec = insres.first->second;
 		lock.Leave();
-		SCOPED_WRITE_LOCK(rec.lock, lock2);
+		SCOPED_FUTEX(rec.lock, lock2);
 
 		if (rec.handled)
 		{
@@ -1031,7 +1031,7 @@ namespace uba
 		ActiveUpdateDirectoryEntry* next = nullptr;
 		bool success = true;
 
-		static bool Wait(SessionClient& client, ActiveUpdateDirectoryEntry*& first, ScopedWriteLock& lock, u32 readPos, const tchar* hint)
+		static bool Wait(SessionClient& client, ActiveUpdateDirectoryEntry*& first, ScopedFutex& lock, u32 readPos, const tchar* hint)
 		{
 			ActiveUpdateDirectoryEntry item;
 			item.next = first;
@@ -1111,7 +1111,7 @@ namespace uba
 			u8* pos = m_directoryTableMem + readPos;
 			u32 toRead = u32(reader.GetLeft());
 
-			SCOPED_WRITE_LOCK(m_directoryTableLock, lock);
+			SCOPED_FUTEX(m_directoryTableLock, lock);
 
 			if (m_directoryTableError)
 				return false;
@@ -1213,14 +1213,14 @@ namespace uba
 
 		u32 addCount = 0;
 		BinaryReader r(m_nameToHashTableMem.memory, readStartPos, NameToHashMemSize);
-		SCOPED_WRITE_LOCK(m_nameToHashLookupLock, lock);
+		SCOPED_FUTEX(m_nameToHashLookupLock, lock);
 		while (r.GetPosition() < localTableSize)
 		{
 			StringKey name = r.ReadStringKey();
 			CasKey hash = r.ReadCasKey();
 
 			HashRec& rec = m_nameToHashLookup[name];
-			SCOPED_WRITE_LOCK(rec.lock, lock2);
+			SCOPED_FUTEX(rec.lock, lock2);
 			if (serverTime < rec.serverTime)
 				continue;
 			rec.key = hash;
@@ -1688,7 +1688,7 @@ namespace uba
 		{
 			ProcessRec(ProcessImpl* impl) : handle(impl) {}
 			ProcessHandle handle;
-			ReaderWriterLock lock;
+			Futex lock;
 			Atomic<bool> isKilled;
 			Atomic<bool> isDone;
 			float weight = 1.0f;
@@ -1727,8 +1727,8 @@ namespace uba
 					++it;
 					continue;
 				}
-				r.lock.EnterWrite();
-				r.lock.LeaveWrite();
+				r.lock.Enter();
+				r.lock.Leave();
 				it = activeProcesses.erase(it);
 			}
 
@@ -1777,7 +1777,7 @@ namespace uba
 					ProcessRec& rec = *it;
 					if (rec.isKilled || rec.isDone)
 						continue;
-					SCOPED_WRITE_LOCK(rec.lock, lock);
+					SCOPED_FUTEX(rec.lock, lock);
 					if (rec.isDone)
 						continue;
 					rec.handle.Cancel(true);
@@ -1913,7 +1913,7 @@ namespace uba
 								session.m_waitToSendEvent.Set();
 							});
 
-						SCOPED_WRITE_LOCK(rec->lock, lock);
+						SCOPED_FUTEX(rec->lock, lock);
 						auto doneGuard = MakeGuard([&]() { rec->isDone = true; session.m_waitToSendEvent.Set(); });
 
 						if (rec->isKilled)
@@ -2018,7 +2018,7 @@ namespace uba
 						break;
 					}
 
-					SCOPED_READ_LOCK(m_processesLock, processesLock);
+					SCOPED_FUTEX_READ(m_processesLock, processesLock);
 					if (m_processes.empty())
 						break;
 					processesLock.Leave();
@@ -2190,11 +2190,11 @@ namespace uba
 				StringKey dirKey = ToStringKey(dirName);
 
 				{	// Optimization to reduce number of messages
-					SCOPED_WRITE_LOCK(m_dirVisitedLock, l);
+					SCOPED_FUTEX(m_dirVisitedLock, l);
 					DirVisitedEntry& entry = m_dirVisited[dirKey];
 					l.Leave();
 				
-					SCOPED_WRITE_LOCK(entry.lock, l2);
+					SCOPED_FUTEX(entry.lock, l2);
 					if (!entry.handled)
 					{
 						ListDirectoryResponse out;
@@ -2328,7 +2328,7 @@ namespace uba
 
 		CasKey key = ToCasKey(CasKeyHasher().Update(data, size), false);
 
-		SCOPED_WRITE_LOCK(m_hostRunCacheLock, l);
+		SCOPED_FUTEX(m_hostRunCacheLock, l);
 		auto insres = m_hostRunCache.try_emplace(key);
 		auto& buffer = insres.first->second;
 		if (!insres.second)
@@ -2384,10 +2384,10 @@ namespace uba
 
 	bool SessionClient::FlushWrittenFiles(ProcessImpl& process)
 	{
-		SCOPED_WRITE_LOCK(process.m_writtenFilesLock, lock);
+		SCOPED_FUTEX(process.m_writtenFilesLock, lock);
 		bool success = SendFiles(process, process.m_processStats.sendFiles);
 		{
-			SCOPED_WRITE_LOCK(m_outputFilesLock, lock2);
+			SCOPED_FUTEX(m_outputFilesLock, lock2);
 			for (auto& kv : process.m_writtenFiles)
 				m_outputFiles.erase(kv.first);
 		}
diff --git a/Engine/Source/Programs/UnrealBuildAccelerator/Common/Private/UbaSessionServer.cpp b/Engine/Source/Programs/UnrealBuildAccelerator/Common/Private/UbaSessionServer.cpp
--- a/Engine/Source/Programs/UnrealBuildAccelerator/Common/Private/UbaSessionServer.cpp
+++ b/Engine/Source/Programs/UnrealBuildAccelerator/Common/Private/UbaSessionServer.cpp
@@ -63,7 +63,7 @@ namespace uba
 
 		void CallProcessExit(ProcessHandle& h)
 		{
-			SCOPED_WRITE_LOCK(m_exitedLock, lock);
+			SCOPED_FUTEX(m_exitedLock, lock);
 			if (!m_startInfo.exitedFunc)
 				return;
 			auto exitedFunc = m_startInfo.exitedFunc;
@@ -75,7 +75,7 @@ namespace uba
 
 		SessionServer* m_server;
 		ProcessStartInfoHolder m_startInfo;
-		ReaderWriterLock m_exitedLock;
+		Futex m_exitedLock;
 		u32 m_processId;
 		u32 m_exitCode = ~0u;
 		u64 m_processorTime = 0;
@@ -443,7 +443,7 @@ namespace uba
 		FixFileName(fileName, fileName_, workingDir.data);
 		StringKey fileNameKey = ToStringKey(fileName);
 		
-		SCOPED_WRITE_LOCK(m_customCasKeysLock, lock);
+		SCOPED_FUTEX(m_customCasKeysLock, lock);
 		auto insres = m_customCasKeys.try_emplace(fileNameKey);
 		CustomCasKey& customKey = insres.first->second;
 		customKey.casKey = CasKeyZero;
@@ -529,7 +529,7 @@ namespace uba
 		{
 			Vector<ProcessHandle> processes;
 			{
-				SCOPED_WRITE_LOCK(m_processesLock, lock);
+				SCOPED_FUTEX(m_processesLock, lock);
 				isEmpty = m_processes.empty();
 				processes.reserve(m_processes.size());
 				for (auto& pair : m_processes)
@@ -705,7 +705,7 @@ namespace uba
 
 		bool binAsVersion = clientAgentKey != CasKeyZero;
 		{
-			SCOPED_WRITE_LOCK(m_binKeysLock, lock);
+			SCOPED_FUTEX(m_binKeysLock, lock);
 
 			detoursBinaryKey[0] = m_detoursBinaryKey[0];
 			detoursBinaryKey[1] = m_detoursBinaryKey[1];
@@ -828,12 +828,12 @@ namespace uba
 		lookupStr.MakeLower();
 		StringKey lookupKey = ToStringKeyNoCheck(lookupStr.data, lookupStr.count);
 
-		SCOPED_WRITE_LOCK(m_applicationDataLock, lock);
+		SCOPED_FUTEX(m_applicationDataLock, lock);
 		auto insres = m_applicationData.try_emplace(lookupKey);
 		ApplicationData& data = insres.first->second;
 		lock.Leave();
 
-		SCOPED_WRITE_LOCK(data.lock, lock2);
+		SCOPED_FUTEX(data.lock, lock2);
 		if (!data.bytes.empty())
 		{
 			writer.WriteBytes(data.bytes.data(), data.bytes.size());
@@ -908,12 +908,12 @@ namespace uba
 		reader.ReadString(applicationName);
 		StringKey applicationKey = ToStringKeyLower(applicationName);
 
-		SCOPED_WRITE_LOCK(m_applicationDataLock, lock);
+		SCOPED_FUTEX(m_applicationDataLock, lock);
 		auto insres = m_applicationData.try_emplace(applicationKey);
 		ApplicationData& data = insres.first->second;
 		lock.Leave();
 				
-		SCOPED_WRITE_LOCK(data.lock, lock2);
+		SCOPED_FUTEX(data.lock, lock2);
 		if (!data.bytes.empty())
 		{
 			writer.WriteBytes(data.bytes.data(), data.bytes.size());
@@ -1110,7 +1110,7 @@ namespace uba
 			RegisterCreateFileForWrite(StringKeyZero, destination, shouldWriteToDisk, 0, 0, invalidateStorage);
 
 
-			SCOPED_READ_LOCK(m_processesLock, lock);
+			SCOPED_FUTEX_READ(m_processesLock, lock);
 			auto findIt = m_processes.find(processId);
 			if (findIt != m_processes.end())
 			{
@@ -1254,7 +1254,7 @@ namespace uba
 
 		float weightLeft = availableWeight;
 		u32 addCount = 0;
-		SCOPED_WRITE_LOCK(m_fillUpOneAtTheTimeLock, fillLock); // This is a lock to group files better (all clients connect at the same time);
+		SCOPED_FUTEX(m_fillUpOneAtTheTimeLock, fillLock); // This is a lock to group files better (all clients connect at the same time);
 		while (weightLeft > 0)
 		{
 			RemoteProcess* process = DequeueProcess(session, sessionId, connectionInfo.GetId());
@@ -1535,7 +1535,7 @@ namespace uba
 		u32 remotePos = reader.ReadU32();
 		u32 localPos;
 		{
-			SCOPED_READ_LOCK(m_trace.m_memoryLock, l);
+			SCOPED_FUTEX_READ(m_trace.m_memoryLock, l);
 			localPos = u32(m_trace.m_memoryPos);
 		}
 
@@ -1595,7 +1595,7 @@ namespace uba
 			return m_logger.Error(TC("Failed to find process for id %u when receiving GetNextProcess message"), processId);
 
 		auto& remoteProcess = *(RemoteProcess*)h.m_process;
-		SCOPED_WRITE_LOCK(remoteProcess.m_exitedLock, exitedLock);
+		SCOPED_FUTEX(remoteProcess.m_exitedLock, exitedLock);
 		NextProcessInfo nextProcess;
 		bool newProcess;
 		remoteProcess.m_exitCode = prevExitCode;
@@ -1624,7 +1624,7 @@ namespace uba
 		if (!h.IsValid())
 			return m_logger.Error(TC("Failed to find process for id %u when receiving custom message"), processId);
 		auto& remoteProcess = *(RemoteProcess*)h.m_process;
-		SCOPED_WRITE_LOCK(remoteProcess.m_exitedLock, exitedLock);
+		SCOPED_FUTEX(remoteProcess.m_exitedLock, exitedLock);
 		CustomMessage(remoteProcess, reader, writer);
 		return true;
 	}
@@ -1740,7 +1740,7 @@ namespace uba
 
 		bool deferCreation = true;
 		{
-			SCOPED_WRITE_LOCK(m_customCasKeysLock, lock);
+			SCOPED_FUTEX(m_customCasKeysLock, lock);
 			auto findIt = m_customCasKeys.find(fileNameKey);
 			if (findIt != m_customCasKeys.end())
 			{
@@ -1766,7 +1766,7 @@ namespace uba
 	{
 		auto& dirTable = m_directoryTable;
 
-		SCOPED_WRITE_LOCK(session.dirTablePosLock, lock2);
+		SCOPED_FUTEX(session.dirTablePosLock, lock2);
 
 		//m_logger.Info(TC("WritePos: %llu"), session.dirTablePos);
 		writer.WriteU32(session.dirTablePos); // We can figure out on the other side if everything was written based on if the message is full or not.
@@ -1833,7 +1833,7 @@ namespace uba
 			bool allGood = false;
 			while (memAvail >= m_memRequiredToSpawn)
 			{
-				SCOPED_WRITE_LOCK(m_waitingProcessesLock, lock);
+				SCOPED_FUTEX(m_waitingProcessesLock, lock);
 				WaitingProcess* wp = m_oldestWaitingProcess;
 				if (!wp)
 				{
@@ -1855,7 +1855,7 @@ namespace uba
 			{
 				lastMessageTime = time;
 				u32 delayCount = 0;
-				SCOPED_WRITE_LOCK(m_waitingProcessesLock, lock);
+				SCOPED_FUTEX(m_waitingProcessesLock, lock);
 				for (auto it = m_oldestWaitingProcess; it; it = it->next)
 					++delayCount;
 				lock.Leave();
@@ -1896,7 +1896,7 @@ namespace uba
 			{
 				u64 newestTime = 0;
 				ProcessImpl* newestProcess = nullptr;
-				SCOPED_WRITE_LOCK(m_processesLock, lock);
+				SCOPED_FUTEX(m_processesLock, lock);
 				for (auto& kv : m_processes)
 				{
 					ProcessHandle& h = kv.second;
@@ -1919,7 +1919,7 @@ namespace uba
 			}
 		}
 
-		SCOPED_WRITE_LOCK(m_waitingProcessesLock, lock);
+		SCOPED_FUTEX(m_waitingProcessesLock, lock);
 		for (auto it = m_oldestWaitingProcess; it; it = it->next)
 			it->event.Set();
 		m_oldestWaitingProcess = nullptr;
@@ -2002,7 +2002,7 @@ namespace uba
 			m_logger.Warning(TC("Cancelling remote active processes has not been tested. Notify devs"));
 
 			{
-				SCOPED_WRITE_LOCK(m_processesLock, lock);
+				SCOPED_FUTEX(m_processesLock, lock);
 				m_processes.erase(process->m_processId);
 			}
 		}
@@ -2013,7 +2013,7 @@ namespace uba
 
 	ProcessHandle SessionServer::ProcessRemoved(u32 processId)
 	{
-		SCOPED_WRITE_LOCK(m_processesLock, lock);
+		SCOPED_FUTEX(m_processesLock, lock);
 		auto findIt = m_processes.find(processId);
 		if (findIt == m_processes.end())
 			return {};
@@ -2024,7 +2024,7 @@ namespace uba
 
 	ProcessHandle SessionServer::GetProcess(u32 processId)
 	{
-		SCOPED_READ_LOCK(m_processesLock, lock);
+		SCOPED_FUTEX_READ(m_processesLock, lock);
 		auto findIt = m_processes.find(processId);
 		if (findIt == m_processes.end())
 			return {};
@@ -2034,7 +2034,7 @@ namespace uba
 	TString SessionServer::GetProcessDescription(u32 processId)
 	{
 		StringBuffer<512> str;
-		SCOPED_READ_LOCK(m_processesLock, lock);
+		SCOPED_FUTEX_READ(m_processesLock, lock);
 		auto findIt = m_processes.find(processId);
 		if (findIt == m_processes.end())
 			return str.Appendf(TC("<Process with id %u not found>"), processId).data;
@@ -2057,7 +2057,7 @@ namespace uba
 		WaitingProcess wp;
 		wp.event.Create(true);
 
-		SCOPED_WRITE_LOCK(m_waitingProcessesLock, lock);
+		SCOPED_FUTEX(m_waitingProcessesLock, lock);
 		if (m_memoryThreadEvent.IsSet(0))
 			return false;
 
diff --git a/Engine/Source/Programs/UnrealBuildAccelerator/Common/Private/UbaStorage.cpp b/Engine/Source/Programs/UnrealBuildAccelerator/Common/Private/UbaStorage.cpp
--- a/Engine/Source/Programs/UnrealBuildAccelerator/Common/Private/UbaStorage.cpp
+++ b/Engine/Source/Programs/UnrealBuildAccelerator/Common/Private/UbaStorage.cpp
@@ -27,7 +27,7 @@ namespace uba
 
 	u8* BufferSlots::Pop()
 	{
-		SCOPED_WRITE_LOCK(m_slotsLock, lock);
+		SCOPED_FUTEX(m_slotsLock, lock);
 		if (!m_slots.empty())
 		{
 			auto back = m_slots.back();
@@ -55,7 +55,7 @@ namespace uba
 		UBA_ASSERT(*(u32*)(slot + BufferSlotSize) == 0xdeadbeef);
 #endif
 
-		SCOPED_WRITE_LOCK(m_slotsLock, lock);
+		SCOPED_FUTEX(m_slotsLock, lock);
 		m_slots.push_back(slot);
 	}
 
@@ -124,7 +124,7 @@ namespace uba
 			//	return;
 		}
 
-		SCOPED_WRITE_LOCK(m_accessLock, lock);
+		SCOPED_FUTEX(m_accessLock, lock);
 
 		CasEntry* prevAccessed = casEntry.prevAccessed;
 		if (prevAccessed == nullptr)
@@ -151,7 +151,7 @@ namespace uba
 
 	void StorageImpl::CasEntryWritten(CasEntry& casEntry, u64 size)
 	{
-		SCOPED_WRITE_LOCK(m_accessLock, lock);
+		SCOPED_FUTEX(m_accessLock, lock);
 
 		m_casTotalBytes += size - casEntry.size;
 		m_casMaxBytes = Max(m_casTotalBytes, m_casMaxBytes);
@@ -183,10 +183,10 @@ namespace uba
 		for (CasEntry* it = m_oldestAccessed; it;)
 		{
 			CasEntry& ce = *it;
-			ce.lock.EnterWrite();
+			ce.lock.Enter();
 			if (ce.verified) // Can't remove these since they might be in mapping tables and actively being used. It is tempting to delete dropped cas files but we can't
 			{
-				ce.lock.LeaveWrite();
+				ce.lock.Leave();
 				break;
 			}
 			UBA_ASSERT(ce.key != CasKeyZero);
@@ -227,11 +227,11 @@ namespace uba
 					m_logger.Error(TC("Failed to delete %s while handling overflow (%s)"), casFile.data, LastErrorToText(error).data);
 					rec.casEntry.exists = true;
 					rec.casEntry.size = rec.size;
-					rec.casEntry.lock.LeaveWrite();
+					rec.casEntry.lock.Leave();
 
 					// TODO: Should this instead set some overflow
 					/*
-					SCOPED_WRITE_LOCK(m_accessLock, lock);
+					SCOPED_FUTEX(m_accessLock, lock);
 					m_casEvictedBytes -= rec.size;
 					--m_casEvictedCount;
 					m_casTotalBytes += rec.size;
@@ -245,13 +245,13 @@ namespace uba
 
 			casEntry.verified = true; // Verified to be deleted
 
-			rec.casEntry.lock.LeaveWrite();
+			rec.casEntry.lock.Leave();
 		}
 	}
 
 	void StorageImpl::CasEntryDeleted(CasEntry& casEntry, u64 size)
 	{
-		SCOPED_WRITE_LOCK(m_accessLock, lock);
+		SCOPED_FUTEX(m_accessLock, lock);
 		m_casTotalBytes -= size;
 		casEntry.size = 0;
 		DetachEntry(casEntry);
@@ -668,7 +668,7 @@ namespace uba
 				u8* readPos = nullptr;
 				u8* writePos = nullptr;
 
-				ReaderWriterLock lock;
+				Futex lock;
 				u64 decompressedSize = 0;
 				u64 decompressedLeft = 0;
 				u64 written = 0;
@@ -691,7 +691,7 @@ namespace uba
 					u64 lastWritten = 0;
 					while (true)
 					{
-						SCOPED_WRITE_LOCK(rec->lock, lock);
+						SCOPED_FUTEX(rec->lock, lock);
 						rec->written += lastWritten;
 						if (!rec->decompressedLeft)
 						{
@@ -1299,7 +1299,7 @@ namespace uba
 
 		CreateDirectory(m_rootDir.data);
 
-		SCOPED_WRITE_LOCK(m_casTableLoadSaveLock, loadSaveLock);
+		SCOPED_FUTEX(m_casTableLoadSaveLock, loadSaveLock);
 
 		UBA_ASSERT(!m_casTableLoaded);
 		m_casTableLoaded = true;
@@ -1461,7 +1461,7 @@ namespace uba
 
 	bool StorageImpl::SaveCasTable(bool deleteIsRunningfile, bool deleteDropped)
 	{
-		SCOPED_WRITE_LOCK(m_casTableLoadSaveLock, loadSaveLock);
+		SCOPED_FUTEX(m_casTableLoadSaveLock, loadSaveLock);
 		if (!m_casTableLoaded)
 			return true;
 
@@ -1481,7 +1481,7 @@ namespace uba
 			
 			SCOPED_WRITE_LOCK(m_fileTableLookupLock, fileTableLock);
 			SCOPED_WRITE_LOCK(m_casLookupLock, casLookupLock);
-			SCOPED_WRITE_LOCK(m_accessLock, accessLock);
+			SCOPED_FUTEX(m_accessLock, accessLock);
 
 			Vector<u8> buffer;
 
@@ -1596,7 +1596,7 @@ namespace uba
 		u32 entryCount = 0;
 		u32 errorCount = 0;
 		u64 newestWrittenError = 0;
-		ReaderWriterLock lock;
+		Futex lock;
 		TraverseAllCasFiles(casRoot.data, [&](const StringBufferBase& fullPath, const DirectoryEntry& e)
 			{
 				++entryCount;
@@ -1611,7 +1611,7 @@ namespace uba
 
 						auto reportError = MakeGuard([&]()
 							{
-								SCOPED_WRITE_LOCK(lock, l);
+								SCOPED_FUTEX(lock, l);
 								++errorCount;
 								if (lastWritten > newestWrittenError)
 									newestWrittenError = lastWritten;
@@ -1814,7 +1814,7 @@ namespace uba
 		FileEntry& fileEntry = findIt->second;
 		lookupLock.Leave();
 
-		SCOPED_WRITE_LOCK(fileEntry.lock, entryLock);
+		SCOPED_FUTEX(fileEntry.lock, entryLock);
 		
 		fileEntry.verified = fileEntry.lastWritten == verifiedLastWriteTime && fileEntry.size == verifiedSize && fileEntry.casKey != CasKeyInvalid;
 
@@ -1833,7 +1833,7 @@ namespace uba
 		FileEntry& fileEntry = findIt->second;
 		lookupLock.Leave();
 
-		SCOPED_WRITE_LOCK(fileEntry.lock, entryLock);
+		SCOPED_FUTEX(fileEntry.lock, entryLock);
 		fileEntry.verified = false;
 		fileEntry.casKey = CasKeyInvalid;
 		return true;
@@ -1848,7 +1848,7 @@ namespace uba
 		StringKey fileNameKey = ToStringKey(forKey);
 
 		FileEntry& fileEntry = GetOrCreateFileEntry(fileNameKey);
-		SCOPED_WRITE_LOCK(fileEntry.lock, entryLock);
+		SCOPED_FUTEX(fileEntry.lock, entryLock);
 	
 		if (fileEntry.verified)
 		{
@@ -1966,7 +1966,7 @@ namespace uba
 			return false;
 		FileEntry& fileEntry = findIt->second;
 		lookupLock.Leave();
-		SCOPED_READ_LOCK(fileEntry.lock, entryLock);
+		SCOPED_FUTEX_READ(fileEntry.lock, entryLock);
 		return fileEntry.verified;
 	}
 
@@ -1979,7 +1979,7 @@ namespace uba
 		FileEntry& fileEntry = findIt->second;
 		lookupLock.Leave();
 
-		SCOPED_WRITE_LOCK(fileEntry.lock, entryLock);
+		SCOPED_FUTEX(fileEntry.lock, entryLock);
 		if (fileEntry.verified)
 			return;
 		fileEntry.verified = fileEntry.lastWritten == verifiedLastWriteTime && fileEntry.size == verifiedSize && fileEntry.casKey != CasKeyInvalid;
@@ -2005,18 +2005,7 @@ namespace uba
 	{
 		FileEntry& fileEntry = GetOrCreateFileEntry(fileNameKey);
 
-		#if !PLATFORM_WINDOWS
-		{
-			SCOPED_READ_LOCK(fileEntry.lock, entryLock);
-			if (fileEntry.verified)
-			{
-				out = fileEntry.casKey;
-				return true;
-			}
-		}
-		#endif
-
-		SCOPED_WRITE_LOCK(fileEntry.lock, entryLock);
+		SCOPED_FUTEX(fileEntry.lock, entryLock);
 	
 		if (fileEntry.verified)
 		{
@@ -2319,7 +2308,7 @@ namespace uba
 		
 		u32 maxParallelCopyOrLink = m_maxParallelCopyOrLink;
 		{
-			SCOPED_WRITE_LOCK(m_activeCopyOrLinkLock, lock);
+			SCOPED_FUTEX(m_activeCopyOrLinkLock, lock);
 			if (m_activeCopyOrLink > maxParallelCopyOrLink)
 			{
 				TimerScope ts(stats.copyOrLinkWait);
@@ -2335,7 +2324,7 @@ namespace uba
 		
 		auto activeGuard = MakeGuard([&]()
 			{
-				SCOPED_WRITE_LOCK(m_activeCopyOrLinkLock, lock);
+				SCOPED_FUTEX(m_activeCopyOrLinkLock, lock);
 				--m_activeCopyOrLink;
 				m_activeCopyOrLinkEvent.Set();
 			});
@@ -2413,7 +2402,7 @@ namespace uba
 
 				FileAccessor destinationFile(m_logger, destination);
 
-				SCOPED_WRITE_LOCK(fileEntry.lock, entryLock);
+				SCOPED_FUTEX(fileEntry.lock, entryLock);
 				fileEntry.verified = false;
 				fileEntry.isTemp = isTemp;
 
@@ -2524,7 +2513,7 @@ namespace uba
 
 			if (casEntry->mappingHandle.IsValid())
 			{
-				SCOPED_WRITE_LOCK(fileEntry.lock, entryLock);
+				SCOPED_FUTEX(fileEntry.lock, entryLock);
 				fileEntry.verified = false;
 
 				MappedView mappedView = m_casDataBuffer.MapView(casEntry->mappingHandle, casEntry->mappingOffset, casEntry->mappingSize, TC(""));
@@ -2554,7 +2543,7 @@ namespace uba
 			if (!GetCasFileName(casFile, actualKey))
 				return false;
 
-			SCOPED_WRITE_LOCK(fileEntry.lock, entryLock);
+			SCOPED_FUTEX(fileEntry.lock, entryLock);
 			fileEntry.verified = false;
 
 			bool firstTry = true;
@@ -2610,7 +2599,7 @@ namespace uba
 			forKey.MakeLower();
 		StringKey key = ToStringKey(forKey);
 		FileEntry& fileEntry = GetOrCreateFileEntry(key);
-		SCOPED_WRITE_LOCK(fileEntry.lock, lock2);
+		SCOPED_FUTEX(fileEntry.lock, lock2);
 		fileEntry.casKey = casKey;
 		fileEntry.lastWritten = lastWritten;
 		fileEntry.size = size;
@@ -2842,7 +2831,7 @@ namespace uba
 		FileEntry& fileEntry = findIt->second;
 		lookupLock.Leave();
 
-		SCOPED_WRITE_LOCK(fileEntry.lock, entryLock);
+		SCOPED_FUTEX(fileEntry.lock, entryLock);
 		fileEntry.verified = false;
 
 		return DropCasFile(fileEntry.casKey, true, file);
diff --git a/Engine/Source/Programs/UnrealBuildAccelerator/Common/Private/UbaStorageClient.cpp b/Engine/Source/Programs/UnrealBuildAccelerator/Common/Private/UbaStorageClient.cpp
--- a/Engine/Source/Programs/UnrealBuildAccelerator/Common/Private/UbaStorageClient.cpp
+++ b/Engine/Source/Programs/UnrealBuildAccelerator/Common/Private/UbaStorageClient.cpp
@@ -129,10 +129,9 @@ namespace uba
 		u64 start = GetTime();
 
 		WorkManagerImpl workManager(workerCount, TC("UbaWorker(PopulateCasFromDirs)"));
-		ReaderWriterLock lock;
 		bool success = true;
 		UnorderedSet<u64> seenIds;
-		ReaderWriterLock seenIdsLock;
+		Futex seenIdsLock;
 
 		for (auto& dir : directories)
 			success = PopulateCasFromDirsRecursive(dir.c_str(), workManager, seenIds, seenIdsLock, shouldExit) && success;
@@ -238,7 +237,7 @@ namespace uba
 
 		StorageStats& stats = Stats();
 		CasEntry* casEntry = nullptr;
-		auto casEntryLock = MakeGuard([&]() { if (casEntry) casEntry->lock.LeaveWrite(); });
+		auto casEntryLock = MakeGuard([&]() { if (casEntry) casEntry->lock.Leave(); });
 		if (shouldStore)
 		{
 			TimerScope ts(stats.ensureCas);
@@ -250,7 +249,7 @@ namespace uba
 			casEntry = &m_casLookup.find(casKey)->second;
 			lock.Leave();
 
-			casEntry->lock.EnterWrite();
+			casEntry->lock.Enter();
 			if (casEntry->verified && casEntry->exists)
 				return true;
 
@@ -513,8 +512,8 @@ namespace uba
 			// So worst case in that scenario would be 10*48*64 = 30000 messages.
 			bool oneAtTheTime = false;//left > client->GetMessageMaxSize() * 2;
 			if (oneAtTheTime)
-				m_retrieveOneBatchAtTheTimeLock.EnterWrite();
-			auto oatg = MakeGuard([&]() { if (oneAtTheTime) m_retrieveOneBatchAtTheTimeLock.LeaveWrite(); });
+				m_retrieveOneBatchAtTheTimeLock.Enter();
+			auto oatg = MakeGuard([&]() { if (oneAtTheTime) m_retrieveOneBatchAtTheTimeLock.Leave(); });
 
 
 			if (isCompressed == IsCompressed(casKey))
@@ -953,7 +952,7 @@ namespace uba
 			m_proxyClient->client.PrintSummary(logger);
 	}
 
-	bool StorageClient::PopulateCasFromDirsRecursive(const tchar* dir, WorkManager& workManager, UnorderedSet<u64>& seenIds, ReaderWriterLock& seenIdsLock, const Function<bool()>& shouldExit)
+	bool StorageClient::PopulateCasFromDirsRecursive(const tchar* dir, WorkManager& workManager, UnorderedSet<u64>& seenIds, Futex& seenIdsLock, const Function<bool()>& shouldExit)
 	{
 		if (shouldExit && shouldExit())
 			return true;
@@ -966,7 +965,7 @@ namespace uba
 				fullPath.Resize(dirLen).Append(e.name);
 				if (IsDirectory(e.attributes))
 				{
-					SCOPED_WRITE_LOCK(seenIdsLock, lock);
+					SCOPED_FUTEX(seenIdsLock, lock);
 					if (!seenIds.insert(e.id).second)
 						return;
 					lock.Leave();
@@ -983,12 +982,12 @@ namespace uba
 					forKey.MakeLower();
 				StringKey fileNameKey = ToStringKey(forKey);
 				FileEntry& fileEntry = GetOrCreateFileEntry(fileNameKey);
-				fileEntry.lock.EnterWrite();
+				fileEntry.lock.Enter();
 				if (e.size == fileEntry.size && e.lastWritten == fileEntry.lastWritten)
 				{
 					fileEntry.verified = true;
 					fileEntry.casKey = AsCompressed(fileEntry.casKey, false); // TODO: Remove this when machines have flushed their db
-					fileEntry.lock.LeaveWrite();
+					fileEntry.lock.Leave();
 
 					SCOPED_WRITE_LOCK(m_localStorageFilesLock, lookupLock);
 					auto insres = m_localStorageFiles.try_emplace(fileEntry.casKey);
@@ -1005,7 +1004,7 @@ namespace uba
 
 				workManager.AddWork([&, fe = &fileEntry, lw = e.lastWritten, s = e.size, filePath = TString(fullPath.data)](const WorkContext& context)
 					{
-						auto feLockLeave = MakeGuard([fe]() { fe->lock.LeaveWrite(); });
+						auto feLockLeave = MakeGuard([fe]() { fe->lock.Leave(); });
 
 						if (shouldExit && shouldExit())
 							return;
diff --git a/Engine/Source/Programs/UnrealBuildAccelerator/Common/Private/UbaStorageProxy.cpp b/Engine/Source/Programs/UnrealBuildAccelerator/Common/Private/UbaStorageProxy.cpp
--- a/Engine/Source/Programs/UnrealBuildAccelerator/Common/Private/UbaStorageProxy.cpp
+++ b/Engine/Source/Programs/UnrealBuildAccelerator/Common/Private/UbaStorageProxy.cpp
@@ -179,13 +179,13 @@ namespace uba
 
 		workContext.tracker.AddHint(StringView(hint).GetFileName());
 
-		SCOPED_WRITE_LOCK(m_filesLock, filesLock);
+		SCOPED_FUTEX(m_filesLock, filesLock);
 		FileEntry& file = m_files.try_emplace(casKey).first->second;
 		filesLock.Leave();
 
 		u32 clientId = connectionInfo.GetId();
 
-		SCOPED_WRITE_LOCK(file.lock, fileLock);
+		SCOPED_FUTEX(file.lock, fileLock);
 		if (file.memory || file.error)
 		{
 			fileLock.Leave();
@@ -247,7 +247,7 @@ namespace uba
 
 		auto sendResponses = MakeGuard([&]()
 			{
-				SCOPED_WRITE_LOCK(file.lock, fileLock);
+				SCOPED_FUTEX(file.lock, fileLock);
 				file.error = error;
 				file.memory = memory;
 				auto f = file.beginMessage;
@@ -358,7 +358,7 @@ namespace uba
 		activeLock.Leave();
 
 		FileEntry& file = *fetch.file;
-		SCOPED_WRITE_LOCK(file.lock, fileLock);
+		SCOPED_FUTEX(file.lock, fileLock);
 		if (file.error)
 			return false;
 
@@ -389,7 +389,7 @@ namespace uba
 		auto& file = mif.file;
 		if (mif.error)
 		{
-			SCOPED_WRITE_LOCK(file.lock, fileLock);
+			SCOPED_FUTEX(file.lock, fileLock);
 			file.error = true;
 		}
 		else
@@ -397,7 +397,7 @@ namespace uba
 			mif.message.ProcessAsyncResults(mif.reader);
 		}
 
-		SCOPED_WRITE_LOCK(file.lock, fileLock);
+		SCOPED_FUTEX(file.lock, fileLock);
 
 		UBA_ASSERT(file.segmentMessages[mif.fetchIndex] == &mif);
 		file.segmentMessages[mif.fetchIndex] = nullptr;
@@ -438,7 +438,7 @@ namespace uba
 		return true;
 	}
 
-	bool StorageProxy::GetFileFromLocalStorage(u32 clientId, const CasKey& casKey, FileEntry& file, ScopedWriteLock& fileLock)
+	bool StorageProxy::GetFileFromLocalStorage(u32 clientId, const CasKey& casKey, FileEntry& file, ScopedFutex& fileLock)
 	{
 		if (!m_useLocalStorage || !m_localStorage || !IsCompressed(casKey) || !m_inProcessClientId || clientId == m_inProcessClientId)
 			return true;
diff --git a/Engine/Source/Programs/UnrealBuildAccelerator/Common/Private/UbaStorageServer.cpp b/Engine/Source/Programs/UnrealBuildAccelerator/Common/Private/UbaStorageServer.cpp
--- a/Engine/Source/Programs/UnrealBuildAccelerator/Common/Private/UbaStorageServer.cpp
+++ b/Engine/Source/Programs/UnrealBuildAccelerator/Common/Private/UbaStorageServer.cpp
@@ -43,7 +43,7 @@ namespace uba
 
 		m_server.RegisterOnClientConnected(ServiceId, [this](const Guid& clientUid, u32 clientId)
 			{
-				SCOPED_WRITE_LOCK(m_loadCasTableLock, lock);
+				SCOPED_FUTEX(m_loadCasTableLock, lock);
 				if (!m_casTableLoaded)
 					LoadCasTable(true);
 			});
@@ -94,14 +94,14 @@ namespace uba
 		//if (m_storeCompressed)
 		//	actualKey = AsCompressed(casKey, true);
 
-		SCOPED_WRITE_LOCK(m_waitEntriesLock, waitLock);
+		SCOPED_FUTEX(m_waitEntriesLock, waitLock);
 		WaitEntry& waitEntry = m_waitEntries[actualKey];
 		++waitEntry.refCount;
 		waitLock.Leave();
 
 		auto g = MakeGuard([&]()
 			{
-				SCOPED_WRITE_LOCK(m_waitEntriesLock, waitLock2);
+				SCOPED_FUTEX(m_waitEntriesLock, waitLock2);
 				if (!--waitEntry.refCount)
 					m_waitEntries.erase(actualKey);
 			});
@@ -461,7 +461,7 @@ namespace uba
 							{
 								FileEntry& fileEntry = findIt->second;
 								lookupLock.Leave();
-								SCOPED_READ_LOCK(fileEntry.lock, entryLock);
+								SCOPED_FUTEX_READ(fileEntry.lock, entryLock);
 								if (fileEntry.verified)
 									checkedCasKey = fileEntry.casKey;
 							}
@@ -959,7 +959,7 @@ namespace uba
 					stats.recvCasBytesComp += activeStore.fileSize;
 					stats.recvCasBytesRaw += activeStore.actualSize;
 
-					SCOPED_WRITE_LOCK(m_waitEntriesLock, waitLock);
+					SCOPED_FUTEX(m_waitEntriesLock, waitLock);
 					auto waitFindIt = m_waitEntries.find(casEntry.key);
 					if (waitFindIt != m_waitEntries.end())
 					{
@@ -1038,7 +1038,7 @@ namespace uba
 
 	u16 StorageServer::PopId()
 	{
-		SCOPED_WRITE_LOCK(m_availableIdsLock, lock);
+		SCOPED_FUTEX(m_availableIdsLock, lock);
 		if (m_availableIds.empty())
 		{
 			UBA_ASSERT(m_availableIdsHigh < u16(~0) - 1);
@@ -1051,7 +1051,7 @@ namespace uba
 
 	void StorageServer::PushId(u16 id)
 	{
-		SCOPED_WRITE_LOCK(m_availableIdsLock, lock);
+		SCOPED_FUTEX(m_availableIdsLock, lock);
 		m_availableIds.push_back(id);
 	}
 
diff --git a/Engine/Source/Programs/UnrealBuildAccelerator/Common/Private/UbaStorageUtils.cpp b/Engine/Source/Programs/UnrealBuildAccelerator/Common/Private/UbaStorageUtils.cpp
--- a/Engine/Source/Programs/UnrealBuildAccelerator/Common/Private/UbaStorageUtils.cpp
+++ b/Engine/Source/Programs/UnrealBuildAccelerator/Common/Private/UbaStorageUtils.cpp
@@ -363,7 +363,7 @@ namespace uba
 			};
 
 		bool hasSendOneAtTheTimeLock = false;
-		auto lockGuard = MakeGuard([&]() { if (hasSendOneAtTheTimeLock) m_sendOneAtTheTimeLock.LeaveWrite(); });
+		auto lockGuard = MakeGuard([&]() { if (hasSendOneAtTheTimeLock) m_sendOneAtTheTimeLock.Leave(); });
 
 		while (sendLeft)
 		{
@@ -395,7 +395,7 @@ namespace uba
 
 			if (isFirst && !isDone && m_sendOneBigFileAtTheTime)
 			{
-				m_sendOneAtTheTimeLock.EnterWrite();
+				m_sendOneAtTheTimeLock.Enter();
 				hasSendOneAtTheTimeLock = true;
 			}
 
diff --git a/Engine/Source/Programs/UnrealBuildAccelerator/Common/Private/UbaTrace.cpp b/Engine/Source/Programs/UnrealBuildAccelerator/Common/Private/UbaTrace.cpp
--- a/Engine/Source/Programs/UnrealBuildAccelerator/Common/Private/UbaTrace.cpp
+++ b/Engine/Source/Programs/UnrealBuildAccelerator/Common/Private/UbaTrace.cpp
@@ -24,9 +24,9 @@ namespace uba
 		FreeMemory();
 	}
 
-	struct Trace::WriterScope : ScopedWriteLock, BinaryWriter
+	struct Trace::WriterScope : ScopedFutex, BinaryWriter
 	{
-		WriterScope(Trace& trace) : ScopedWriteLock(trace.m_memoryLock), BinaryWriter(trace.m_memoryBegin, trace.m_memoryPos, trace.m_memoryCapacity), m_trace(trace)
+		WriterScope(Trace& trace) : ScopedFutex(trace.m_memoryLock), BinaryWriter(trace.m_memoryBegin, trace.m_memoryPos, trace.m_memoryCapacity), m_trace(trace)
 		{
 			EnsureMemory(TraceMessageMaxSize);
 		}
@@ -186,7 +186,7 @@ namespace uba
 		if (!m_memoryBegin)
 			return 0;
 
-		SCOPED_WRITE_LOCK(m_stringsLock, lock);
+		SCOPED_FUTEX(m_stringsLock, lock);
 		auto insres = m_strings.try_emplace(ToStringKeyNoCheck(string.data, string.count));
 		if (insres.second)
 		{
diff --git a/Engine/Source/Programs/UnrealBuildAccelerator/Common/Private/UbaWorkManager.cpp b/Engine/Source/Programs/UnrealBuildAccelerator/Common/Private/UbaWorkManager.cpp
--- a/Engine/Source/Programs/UnrealBuildAccelerator/Common/Private/UbaWorkManager.cpp
+++ b/Engine/Source/Programs/UnrealBuildAccelerator/Common/Private/UbaWorkManager.cpp
@@ -61,7 +61,7 @@ namespace uba
 				{
 					while (true)
 					{
-						SCOPED_WRITE_LOCK(manager.m_workLock, lock);
+						SCOPED_FUTEX(manager.m_workLock, lock);
 						if (manager.m_work.empty())
 							break;
 						Work work = manager.m_work.front();
@@ -76,8 +76,8 @@ namespace uba
 						work.func({tws});
 					}
 
-					SCOPED_WRITE_LOCK(manager.m_availableWorkersLock, lock1);
-					SCOPED_READ_LOCK(manager.m_workLock, lock2);
+					SCOPED_FUTEX(manager.m_availableWorkersLock, lock1);
+					SCOPED_FUTEX_READ(manager.m_workLock, lock2);
 					if (!manager.m_work.empty())
 						continue;
 
@@ -114,7 +114,7 @@ namespace uba
 	void WorkManagerImpl::AddWork(const WorkFunction& work, u32 count, const tchar* desc, const Color& color, bool highPriority)
 	{
 		UBA_ASSERT(!m_workers.empty());
-		SCOPED_WRITE_LOCK(m_workLock, lock);
+		SCOPED_FUTEX(m_workLock, lock);
 		bool trackWork = m_workTracker.load();
 		for (u32 i = 0; i != count; ++i)
 		{
@@ -133,7 +133,7 @@ namespace uba
 		}
 		lock.Leave();
 
-		SCOPED_WRITE_LOCK(m_availableWorkersLock, lock2);
+		SCOPED_FUTEX(m_availableWorkersLock, lock2);
 		while (count--)
 		{
 			Worker* worker = PopWorkerNoLock();
@@ -150,7 +150,7 @@ namespace uba
 
 	void WorkManagerImpl::PushWorker(Worker* worker)
 	{
-		SCOPED_WRITE_LOCK(m_availableWorkersLock, lock);
+		SCOPED_FUTEX(m_availableWorkersLock, lock);
 		PushWorkerNoLock(worker);
 	}
 
@@ -168,7 +168,7 @@ namespace uba
 	{
 		while (count--)
 		{
-			SCOPED_WRITE_LOCK(m_workLock, lock);
+			SCOPED_FUTEX(m_workLock, lock);
 			if (m_work.empty())
 				break;
 			Work work = m_work.front();
@@ -188,7 +188,7 @@ namespace uba
 	{
 		while (true)
 		{
-			SCOPED_READ_LOCK(m_workLock, lock);
+			SCOPED_FUTEX_READ(m_workLock, lock);
 			bool workEmpty = m_work.empty();
 			lock.Leave();
 			if (workEmpty)
diff --git a/Engine/Source/Programs/UnrealBuildAccelerator/Common/Public/UbaBottleneck.h b/Engine/Source/Programs/UnrealBuildAccelerator/Common/Public/UbaBottleneck.h
--- a/Engine/Source/Programs/UnrealBuildAccelerator/Common/Public/UbaBottleneck.h
+++ b/Engine/Source/Programs/UnrealBuildAccelerator/Common/Public/UbaBottleneck.h
@@ -19,7 +19,7 @@ namespace uba
 
 		void Enter(BottleneckTicket& ticket)
 		{
-			SCOPED_WRITE_LOCK(m_lock, lock);
+			SCOPED_FUTEX(m_lock, lock);
 			if (m_activeCount < m_maxCount)
 			{
 				UBA_ASSERT(!m_first);
@@ -43,7 +43,7 @@ namespace uba
 
 		void Leave(BottleneckTicket& ticket)
 		{
-			SCOPED_WRITE_LOCK(m_lock, lock);
+			SCOPED_FUTEX(m_lock, lock);
 			BottleneckTicket* first = m_first;
 			if (!first)
 			{
@@ -58,7 +58,7 @@ namespace uba
 		}
 
 	private:
-		ReaderWriterLock m_lock;
+		Futex m_lock;
 		BottleneckTicket* m_first = nullptr;
 		BottleneckTicket* m_last = nullptr;
 		u32 m_activeCount = 0;
diff --git a/Engine/Source/Programs/UnrealBuildAccelerator/Common/Public/UbaCacheClient.h b/Engine/Source/Programs/UnrealBuildAccelerator/Common/Public/UbaCacheClient.h
--- a/Engine/Source/Programs/UnrealBuildAccelerator/Common/Public/UbaCacheClient.h
+++ b/Engine/Source/Programs/UnrealBuildAccelerator/Common/Public/UbaCacheClient.h
@@ -109,13 +109,13 @@ namespace uba
 
 		Atomic<bool> m_connected;
 
-		ReaderWriterLock m_bucketsLock;
+		Futex m_bucketsLock;
 		UnorderedMap<u32, Bucket> m_buckets;
 
-		ReaderWriterLock m_sendOneAtTheTimeLock;
+		Futex m_sendOneAtTheTimeLock;
 
-		ReaderWriterLock m_directoryPreparserLock;
-		struct PreparedDir { ReaderWriterLock lock; Atomic<bool> done; };
+		Futex m_directoryPreparserLock;
+		struct PreparedDir { Futex lock; Atomic<bool> done; };
 		UnorderedMap<StringKey, PreparedDir> m_directoryPreparser;
 
 		CacheClient(const CacheClient&) = delete;
diff --git a/Engine/Source/Programs/UnrealBuildAccelerator/Common/Public/UbaCacheServer.h b/Engine/Source/Programs/UnrealBuildAccelerator/Common/Public/UbaCacheServer.h
--- a/Engine/Source/Programs/UnrealBuildAccelerator/Common/Public/UbaCacheServer.h
+++ b/Engine/Source/Programs/UnrealBuildAccelerator/Common/Public/UbaCacheServer.h
@@ -102,10 +102,10 @@ namespace uba
 		Atomic<bool> m_isRunningMaintenance;
 		Atomic<bool> m_bucketIsOverflowing;
 
-		ReaderWriterLock m_bucketsLock;
+		Futex m_bucketsLock;
 		Map<u64, Bucket> m_buckets;
 
-		ReaderWriterLock m_connectionsLock;
+		Futex m_connectionsLock;
 		Map<u32, Connection> m_connections;
 
 		Atomic<bool> m_shutdownRequested = false;
diff --git a/Engine/Source/Programs/UnrealBuildAccelerator/Common/Public/UbaCompactTables.h b/Engine/Source/Programs/UnrealBuildAccelerator/Common/Public/UbaCompactTables.h
--- a/Engine/Source/Programs/UnrealBuildAccelerator/Common/Public/UbaCompactTables.h
+++ b/Engine/Source/Programs/UnrealBuildAccelerator/Common/Public/UbaCompactTables.h
@@ -65,7 +65,7 @@ namespace uba
 		u32 InternalAddWithNewSegment(u32 offset, u32 parentOffset, bool separator, const StringView& segment);
 		void InternalAddWithExistingSegment(u32 offset, u32 parentOffset, bool separator, u32 segmentOffset);
 
-		ReaderWriterLock m_lock;
+		Futex m_lock;
 		MemoryBlock m_pathTableMem;
 		HashMap<StringKey, u32, true> m_offsets;
 		HashMap<StringKey, u32, true> m_segmentOffsets;
@@ -94,7 +94,6 @@ namespace uba
 
 		u8* GetMemory();
 		u32 GetSize();
-		ReaderWriterLock& GetLock() { return m_lock; }
 
 		void ReadMem(BinaryReader& reader, bool populateLookup);
 		u8* BeginCommit(u64 size);
@@ -108,7 +107,7 @@ namespace uba
 	private:
 		u32* InternalAdd(const CasKey& casKey, u64 stringOffset, bool& outAdded);
 
-		ReaderWriterLock m_lock;
+		Futex m_lock;
 		MemoryBlock m_casKeyTableMem;
 
 		struct StringAndKey
diff --git a/Engine/Source/Programs/UnrealBuildAccelerator/Common/Public/UbaFile.h b/Engine/Source/Programs/UnrealBuildAccelerator/Common/Public/UbaFile.h
--- a/Engine/Source/Programs/UnrealBuildAccelerator/Common/Public/UbaFile.h
+++ b/Engine/Source/Programs/UnrealBuildAccelerator/Common/Public/UbaFile.h
@@ -87,8 +87,8 @@ namespace uba
 		void Clear();
 
 	private:
-		ReaderWriterLock m_createdDirsLock;
-		struct CreatedDir { ReaderWriterLock lock; bool handled = false; };
+		Futex m_createdDirsLock;
+		struct CreatedDir { Futex lock; bool handled = false; };
 		UnorderedMap<TString, CreatedDir> m_createdDirs;
 	};
 
diff --git a/Engine/Source/Programs/UnrealBuildAccelerator/Common/Public/UbaFileMapping.h b/Engine/Source/Programs/UnrealBuildAccelerator/Common/Public/UbaFileMapping.h
--- a/Engine/Source/Programs/UnrealBuildAccelerator/Common/Public/UbaFileMapping.h
+++ b/Engine/Source/Programs/UnrealBuildAccelerator/Common/Public/UbaFileMapping.h
@@ -93,7 +93,7 @@ namespace uba
 			File files[128];
 			u32 fileCount = 0;
 
-			ReaderWriterLock availableFilesLock;
+			Futex availableFilesLock;
 			Event availableFilesEvent;
 			File* availableFiles[128] = { 0 };
 			u32 availableFilesCount = 0;
@@ -142,7 +142,7 @@ namespace uba
 		u64 m_capacity = 0;
 		u64 m_blockSize = 0;
 
-		ReaderWriterLock m_mappingLock;
+		Futex m_mappingLock;
 		FileMappingHandle m_mappingHandle;
 		u64 m_mappingCount = 0;
 
diff --git a/Engine/Source/Programs/UnrealBuildAccelerator/Common/Public/UbaIncludeCrawler.h b/Engine/Source/Programs/UnrealBuildAccelerator/Common/Public/UbaIncludeCrawler.h
--- a/Engine/Source/Programs/UnrealBuildAccelerator/Common/Public/UbaIncludeCrawler.h
+++ b/Engine/Source/Programs/UnrealBuildAccelerator/Common/Public/UbaIncludeCrawler.h
@@ -34,12 +34,12 @@ namespace uba
 		Logger& m_logger;
 		WorkManager& m_workManager;
 
-		ReaderWriterLock m_codeFilesHandledLock;
-		struct HandledFile { ReaderWriterLock lock; bool handled = false; };
+		Futex m_codeFilesHandledLock;
+		struct HandledFile { Futex lock; bool handled = false; };
 		UnorderedMap<StringKey, HandledFile> m_codeFilesHandled;
 
-		ReaderWriterLock m_pchLookupLock;
-		struct Pch { ReaderWriterLock lock; bool handled = false; UnorderedSet<StringKey> files; };
+		Futex m_pchLookupLock;
+		struct Pch { Futex lock; bool handled = false; UnorderedSet<StringKey> files; };
 		UnorderedMap<StringKey, Pch> m_pchLookup;
 
 		struct IncludeRoot
@@ -65,9 +65,9 @@ namespace uba
 			Pch* pch = nullptr;
 		};
 
-		ReaderWriterLock m_instancesLock;
+		Futex m_instancesLock;
 		List<Instance> m_instances;
-		ReaderWriterLock m_builtIncludesHandledLock;
+		Futex m_builtIncludesHandledLock;
 		UnorderedSet<u32> m_builtIncludesHandled;
 		FileExistsFunc m_fileExistsFunc;
 		TraverseFilesFunc m_traverseFilesFunc;
@@ -98,7 +98,7 @@ namespace uba
 
 	bool IncludeCrawler::Add(const tchar* rsp, const tchar* workDir, CreateFileFunc&& createFileFunc, DevirtualizePathFunc&& devirtualizePathFunc, const tchar* app, const ApplicationRules* rules)
 	{
-		SCOPED_WRITE_LOCK(m_instancesLock, lock);
+		SCOPED_FUTEX(m_instancesLock, lock);
 		Instance& instance = m_instances.emplace_back();
 		lock.Leave();
 
@@ -129,11 +129,11 @@ namespace uba
 
 						StringKey pchKey = ToStringKeyLower(fixedFile);
 
-						SCOPED_WRITE_LOCK(m_pchLookupLock, lock);
+						SCOPED_FUTEX(m_pchLookupLock, lock);
 						Pch& pch = m_pchLookup[pchKey];
 						lock.Leave();
 
-						SCOPED_WRITE_LOCK(pch.lock, lock2);
+						SCOPED_FUTEX(pch.lock, lock2);
 						if (!pch.handled)
 						{
 							fixedFile.Append(TC(".dep.json"));
@@ -160,7 +160,7 @@ namespace uba
 
 				if (rules && !instance.builtinIncludesDir.empty())
 				{
-					SCOPED_WRITE_LOCK(m_builtIncludesHandledLock, lock);
+					SCOPED_FUTEX(m_builtIncludesHandledLock, lock);
 					bool shouldHandle = m_builtIncludesHandled.insert(rules->index).second;
 					lock.Leave();
 					if (shouldHandle)
@@ -288,12 +288,12 @@ namespace uba
 							fixedFile.MakeLower();
 
 						StringKey key = ToStringKey(fixedFile);
-						SCOPED_WRITE_LOCK(m_codeFilesHandledLock, lock);
+						SCOPED_FUTEX(m_codeFilesHandledLock, lock);
 						auto insres = m_codeFilesHandled.try_emplace(key);
 						lock.Leave();
 
 						HandledFile& handledFile = insres.first->second;
-						SCOPED_WRITE_LOCK(handledFile.lock, fileLock);
+						SCOPED_FUTEX(handledFile.lock, fileLock);
 						if (!handledFile.handled)
 						{
 							handledFile.handled = true;
@@ -595,11 +595,11 @@ namespace uba
 			StringKey key = ToStringKey(fullPath);
 			if (instance.pch)
 			{
-				SCOPED_READ_LOCK(instance.pch->lock, lock);
+				SCOPED_FUTEX_READ(instance.pch->lock, lock);
 				if (instance.pch->files.find(key) != instance.pch->files.end())
 					return true;
 			}
-			SCOPED_WRITE_LOCK(m_codeFilesHandledLock, lock);
+			SCOPED_FUTEX(m_codeFilesHandledLock, lock);
 			if (!m_codeFilesHandled.try_emplace(key).second)
 				return true;
 		}
diff --git a/Engine/Source/Programs/UnrealBuildAccelerator/Common/Public/UbaNetworkBackendTcp.h b/Engine/Source/Programs/UnrealBuildAccelerator/Common/Public/UbaNetworkBackendTcp.h
--- a/Engine/Source/Programs/UnrealBuildAccelerator/Common/Public/UbaNetworkBackendTcp.h
+++ b/Engine/Source/Programs/UnrealBuildAccelerator/Common/Public/UbaNetworkBackendTcp.h
@@ -49,10 +49,10 @@ namespace uba
 		void ThreadRecv(Connection& connection);
 
 		LoggerWithWriter m_logger;
-		ReaderWriterLock m_listenEntriesLock;
+		Futex m_listenEntriesLock;
 		List<ListenEntry> m_listenEntries;
 
-		ReaderWriterLock m_connectionsLock;
+		Futex m_connectionsLock;
 		List<Connection> m_connections;
 
 		Atomic<u64> m_totalSend;
diff --git a/Engine/Source/Programs/UnrealBuildAccelerator/Common/Public/UbaNetworkClient.h b/Engine/Source/Programs/UnrealBuildAccelerator/Common/Public/UbaNetworkClient.h
--- a/Engine/Source/Programs/UnrealBuildAccelerator/Common/Public/UbaNetworkClient.h
+++ b/Engine/Source/Programs/UnrealBuildAccelerator/Common/Public/UbaNetworkClient.h
@@ -106,13 +106,13 @@ namespace uba
 		Atomic<bool> m_isDisconnecting;
 		Timer m_sendTimer;
 
-		ReaderWriterLock m_serverUidLock;
+		Futex m_serverUidLock;
 		Guid m_serverUid;
 
 		Event m_isConnected;
 		Event m_isOrWasConnected;
 		Atomic<u32> m_connectionCount;
-		ReaderWriterLock m_onConnectedFunctionsLock;
+		Futex m_onConnectedFunctionsLock;
 		Vector<OnConnectedFunction> m_onConnectedFunctions;
 		ReaderWriterLock m_onDisconnectedFunctionsLock;
 		Vector<OnDisconnectedFunction> m_onDisconnectedFunctions;
@@ -120,7 +120,7 @@ namespace uba
 
 		ReaderWriterLock m_connectionsLock;
 		List<Connection> m_connections;
-		ReaderWriterLock m_connectionsItLock;
+		Futex m_connectionsItLock;
 		List<Connection>::iterator m_connectionsIt;
 
 
diff --git a/Engine/Source/Programs/UnrealBuildAccelerator/Common/Public/UbaNetworkServer.h b/Engine/Source/Programs/UnrealBuildAccelerator/Common/Public/UbaNetworkServer.h
--- a/Engine/Source/Programs/UnrealBuildAccelerator/Common/Public/UbaNetworkServer.h
+++ b/Engine/Source/Programs/UnrealBuildAccelerator/Common/Public/UbaNetworkServer.h
@@ -128,7 +128,7 @@ namespace uba
 			CryptoKey key;
 			u64 expirationTime;
 		};
-		ReaderWriterLock m_cryptoKeysLock;
+		Futex m_cryptoKeysLock;
 		List<CryptoEntry> m_cryptoKeys;
 
 		Guid m_uid;
@@ -144,11 +144,11 @@ namespace uba
 
 		u32 m_maxWorkerCount = 0;
 
-		ReaderWriterLock m_additionalWorkLock;
+		Futex m_additionalWorkLock;
 		struct AdditionalWork { WorkFunction func; TString desc; };
 		List<AdditionalWork> m_additionalWork;
 
-		ReaderWriterLock m_availableWorkersLock;
+		Futex m_availableWorkersLock;
 		Worker* m_firstAvailableWorker = nullptr;
 		Worker* m_firstActiveWorker = nullptr;
 		Event m_workerAvailable;
@@ -160,7 +160,7 @@ namespace uba
 		bool m_logConnections = true;
 		bool m_allClientsDisconnected = false;
 
-		ReaderWriterLock m_addConnectionsLock;
+		Futex m_addConnectionsLock;
 		List<Thread> m_addConnections;
 
 		ReaderWriterLock m_connectionsLock;
diff --git a/Engine/Source/Programs/UnrealBuildAccelerator/Common/Public/UbaProcess.h b/Engine/Source/Programs/UnrealBuildAccelerator/Common/Public/UbaProcess.h
--- a/Engine/Source/Programs/UnrealBuildAccelerator/Common/Public/UbaProcess.h
+++ b/Engine/Source/Programs/UnrealBuildAccelerator/Common/Public/UbaProcess.h
@@ -102,7 +102,7 @@ namespace uba
 		ProcessStartInfoHolder m_startInfo;
 		Session& m_session;
 		ProcessImpl* m_parentProcess;
-		ReaderWriterLock m_initLock;
+		Futex m_initLock;
 		u32 m_id;
 		FileMappingAllocator::Allocation m_comMemory;
 
@@ -111,7 +111,7 @@ namespace uba
 		Event m_writeEvent;
 		Event m_readEvent;
 	#else
-		ReaderWriterLock m_comMemoryLock;
+		Futex m_comMemoryLock;
 		Atomic<bool> m_cancelled;
 		Event& m_cancelEvent;
 		Event& m_writeEvent;
@@ -145,12 +145,12 @@ namespace uba
 		const tchar* m_realWorkingDir = nullptr;
 		u64 m_startTime = 0;
 		Event m_waitForParent;
-		ReaderWriterLock m_logLinesLock;
+		Futex m_logLinesLock;
 		Vector<ProcessLogLine> m_logLines;
 		Vector<u8> m_trackedInputs;
 		Vector<u8> m_trackedOutputs;
 		Vector<ProcessHandle> m_childProcesses;
-		ReaderWriterLock& m_writtenFilesLock;
+		Futex& m_writtenFilesLock;
 		UnorderedMap<TString, WrittenFile>& m_writtenFiles;
 		ReaderWriterLock& m_tempFilesLock;
 		UnorderedMap<StringKey, WrittenFile>& m_tempFiles;
diff --git a/Engine/Source/Programs/UnrealBuildAccelerator/Common/Public/UbaSession.h b/Engine/Source/Programs/UnrealBuildAccelerator/Common/Public/UbaSession.h
--- a/Engine/Source/Programs/UnrealBuildAccelerator/Common/Public/UbaSession.h
+++ b/Engine/Source/Programs/UnrealBuildAccelerator/Common/Public/UbaSession.h
@@ -236,11 +236,10 @@ namespace uba
 		ReaderWriterLock m_fileMappingTableMemLock;
 		u8* m_fileMappingTableMem;
 		u32 m_fileMappingTableSize = 0;
-		ReaderWriterLock m_fileMappingTableLookupLock;
-		struct FileMappingEntry { ReaderWriterLock lock; FileMappingHandle mapping; u64 mappingOffset = 0; u64 size = 0; u64 lastWriteTime = 0; bool isDir = false; bool handled = false; bool success = false; bool owned = false; };
+		Futex m_fileMappingTableLookupLock;
+		struct FileMappingEntry { Futex lock; FileMappingHandle mapping; u64 mappingOffset = 0; u64 size = 0; u64 lastWriteTime = 0; bool isDir = false; bool handled = false; bool success = false; bool owned = false; };
 		UnorderedMap<StringKey, FileMappingEntry> m_fileMappingTableLookup;
 
-		ReaderWriterLock m_nameToHashTableHandleLock;
 		static constexpr u64 NameToHashMemSize = 48*1024*1024;
 		MemoryBlock m_nameToHashTableMem;
 
@@ -249,22 +248,22 @@ namespace uba
 		FileMappingAllocator m_processCommunicationAllocator;
 		std::string m_detoursLibrary[2]; // x64 and arm64
 
-		ReaderWriterLock m_processStatsLock;
+		Futex m_processStatsLock;
 		ProcessStats m_processStats;
 
-		ReaderWriterLock m_processesLock;
+		Futex m_processesLock;
 		UnorderedMap<u32, ProcessHandle> m_processes;
 		Vector<ProcessHandle> m_deadProcesses;
 		UnorderedMap<TString, Timer> m_applicationStats;
 
-		ReaderWriterLock m_outputFilesLock;
+		Futex m_outputFilesLock;
 		UnorderedMap<TString, TString> m_outputFiles;
 
-		ReaderWriterLock m_activeFilesLock;
+		Futex m_activeFilesLock;
 		struct ActiveFile { TString name; StringKey nameKey; };
 		UnorderedMap<u32, ActiveFile> m_activeFiles;
 
-		ReaderWriterLock m_virtualSourceFilesLock;
+		Futex m_virtualSourceFilesLock;
 		UnorderedMap<StringKey, FileMappingHandle> m_virtualSourceFiles;
 
 		ReaderWriterLock m_rootsLookupLock;
@@ -284,13 +283,13 @@ namespace uba
 		float m_cpuLoad = 0;
 
 		#if PLATFORM_WINDOWS
-		ReaderWriterLock m_processJobObjectLock;
+		Futex m_processJobObjectLock;
 		HANDLE m_processJobObject = NULL;
 		#endif
 
 		Vector<u8> m_environmentMemory;
 
-		ReaderWriterLock m_environmentVariablesLock;
+		Futex m_environmentVariablesLock;
 		Vector<tchar> m_environmentVariables;
 		UnorderedSet<const tchar*, HashStringNoCase, EqualStringNoCase> m_localEnvironmentVariables;
 
@@ -306,7 +305,7 @@ namespace uba
 		#endif
 
 		#if PLATFORM_WINDOWS
-		ReaderWriterLock m_isX64ApplicationLock;
+		Futex m_isX64ApplicationLock;
 		UnorderedMap<TString, bool> m_isX64Application;
 		#endif
 
diff --git a/Engine/Source/Programs/UnrealBuildAccelerator/Common/Public/UbaSessionClient.h b/Engine/Source/Programs/UnrealBuildAccelerator/Common/Public/UbaSessionClient.h
--- a/Engine/Source/Programs/UnrealBuildAccelerator/Common/Public/UbaSessionClient.h
+++ b/Engine/Source/Programs/UnrealBuildAccelerator/Common/Public/UbaSessionClient.h
@@ -139,23 +139,23 @@ namespace uba
 		Atomic<u32> m_maxProcessCount;
 		Atomic<float> m_cpuUsage;
 
-		struct ApplicationEnvironment { ReaderWriterLock lock; TString virtualApplication; TString realApplication; };
-		ReaderWriterLock m_handledApplicationEnvironmentsLock;
+		struct ApplicationEnvironment { Futex lock; TString virtualApplication; TString realApplication; };
+		Futex m_handledApplicationEnvironmentsLock;
 		UnorderedMap<TString, ApplicationEnvironment> m_handledApplicationEnvironments;
 
-		ReaderWriterLock m_binFileLock;
+		Futex m_binFileLock;
 		UnorderedMap<TString, CasKey> m_writtenBinFiles;
 
-		ReaderWriterLock m_nameToNameLookupLock;
-		struct NameRec { TString name; TString virtualName; ReaderWriterLock lock; bool handled = false; };
+		Futex m_nameToNameLookupLock;
+		struct NameRec { TString name; TString virtualName; Futex lock; bool handled = false; };
 		UnorderedMap<TString, NameRec> m_nameToNameLookup;
 
-		struct HashRec { CasKey key; u64 serverTime = 0; ReaderWriterLock lock; };
+		struct HashRec { CasKey key; u64 serverTime = 0; Futex lock; };
 		UnorderedMap<StringKey, HashRec> m_nameToHashLookup;
-		ReaderWriterLock m_nameToHashLookupLock;
+		Futex m_nameToHashLookupLock;
 		ReaderWriterLock m_nameToHashMemLock;
 
-		ReaderWriterLock m_directoryTableLock;
+		Futex m_directoryTableLock;
 		u32 m_directoryTableMemPos = 0;
 		bool m_directoryTableError = false;
 		struct ActiveUpdateDirectoryEntry;
@@ -179,12 +179,12 @@ namespace uba
 		Atomic<u64> m_memTotal;
 
 		UnorderedMap<CasKey, Vector<u8>> m_hostRunCache;
-		ReaderWriterLock m_hostRunCacheLock;
+		Futex m_hostRunCacheLock;
 
 		IncludeCrawler m_includeCrawler;
 		u32 m_dirtableParsedPosition = 0;
-		ReaderWriterLock m_dirVisitedLock;
-		struct DirVisitedEntry { ReaderWriterLock lock; bool handled = false; };
+		Futex m_dirVisitedLock;
+		struct DirVisitedEntry { Futex lock; bool handled = false; };
 		UnorderedMap<StringKey, DirVisitedEntry> m_dirVisited;
 	};
 }
diff --git a/Engine/Source/Programs/UnrealBuildAccelerator/Common/Public/UbaSessionServer.h b/Engine/Source/Programs/UnrealBuildAccelerator/Common/Public/UbaSessionServer.h
--- a/Engine/Source/Programs/UnrealBuildAccelerator/Common/Public/UbaSessionServer.h
+++ b/Engine/Source/Programs/UnrealBuildAccelerator/Common/Public/UbaSessionServer.h
@@ -134,7 +134,7 @@ namespace uba
 		u32 m_availableRemoteSlotCount = 0;
 		u32 m_connectionCount = 0;
 
-		ReaderWriterLock m_binKeysLock;
+		Futex m_binKeysLock;
 		CasKey m_detoursBinaryKey[2];
 		CasKey m_agentBinaryKey[2];
 
@@ -142,7 +142,7 @@ namespace uba
 		{
 			TString name;
 			UnorderedSet<CasKey> sentKeys;
-			ReaderWriterLock dirTablePosLock;
+			Futex dirTablePosLock;
 			u32 dirTablePos = 0;
 			u32 clientId = ~0u;
 			u32 processSlotCount = 1;
@@ -167,7 +167,7 @@ namespace uba
 			TString workingDir;
 			Vector<u8> trackedInputs;
 		};
-		ReaderWriterLock m_customCasKeysLock;
+		Futex m_customCasKeysLock;
 		UnorderedMap<StringKey, CustomCasKey> m_customCasKeys;
 
 		UnorderedMap<StringKey, CasKey> m_nameToHashLookup;
@@ -177,10 +177,10 @@ namespace uba
 		ReaderWriterLock m_receivedFilesLock;
 		UnorderedMap<StringKey, CasKey> m_receivedFiles;
 
-		ReaderWriterLock m_fillUpOneAtTheTimeLock;
+		Futex m_fillUpOneAtTheTimeLock;
 
-		ReaderWriterLock m_applicationDataLock;
-		struct ApplicationData { ReaderWriterLock lock; Vector<u8> bytes; };
+		Futex m_applicationDataLock;
+		struct ApplicationData { Futex lock; Vector<u8> bytes; };
 		UnorderedMap<StringKey, ApplicationData> m_applicationData;
 
 		NetworkTrafficProvider m_provider;
@@ -198,7 +198,7 @@ namespace uba
 		struct WaitingProcess { Event event; WaitingProcess* next = nullptr; };
 		WaitingProcess* m_oldestWaitingProcess = nullptr;
 		WaitingProcess* m_newestWaitingProcess = nullptr;
-		ReaderWriterLock m_waitingProcessesLock;
+		Futex m_waitingProcessesLock;
 		bool m_allowWaitOnMem = false;
 		bool m_allowKillOnMem = false;
 		bool m_remoteLogEnabled = false;
diff --git a/Engine/Source/Programs/UnrealBuildAccelerator/Common/Public/UbaStorage.h b/Engine/Source/Programs/UnrealBuildAccelerator/Common/Public/UbaStorage.h
--- a/Engine/Source/Programs/UnrealBuildAccelerator/Common/Public/UbaStorage.h
+++ b/Engine/Source/Programs/UnrealBuildAccelerator/Common/Public/UbaStorage.h
@@ -99,7 +99,7 @@ namespace uba
 
 		~BufferSlots();
 
-		ReaderWriterLock m_slotsLock;
+		Futex m_slotsLock;
 		Vector<u8*> m_slots;
 	};
 
@@ -199,7 +199,7 @@ namespace uba
 		ReaderWriterLock m_fileTableLookupLock;
 		struct FileEntry
 		{
-			ReaderWriterLock lock;
+			Futex lock;
 			CasKey casKey;
 			u64 size = 0;
 			u64 lastWritten = 0;
@@ -229,7 +229,7 @@ namespace uba
 		};
 
 		UnorderedMap<CasKey, CasEntry> m_casLookup;
-		ReaderWriterLock m_accessLock;
+		Futex m_accessLock;
 		CasEntry* m_newestAccessed = nullptr;
 		CasEntry* m_oldestAccessed = nullptr;
 		u64 m_casTotalBytes = 0;
@@ -249,11 +249,11 @@ namespace uba
 		ExternalFileMappingsProvider m_externalFileMappingsProvider;
 
 		u32 m_maxParallelCopyOrLink;
-		ReaderWriterLock m_activeCopyOrLinkLock;
+		Futex m_activeCopyOrLinkLock;
 		Event m_activeCopyOrLinkEvent;
 		u32 m_activeCopyOrLink = 0;
 
-		ReaderWriterLock m_casTableLoadSaveLock;
+		Futex m_casTableLoadSaveLock;
 		bool m_casTableLoaded = false;
 
 		FileMappingBuffer m_casDataBuffer;
diff --git a/Engine/Source/Programs/UnrealBuildAccelerator/Common/Public/UbaStorageClient.h b/Engine/Source/Programs/UnrealBuildAccelerator/Common/Public/UbaStorageClient.h
--- a/Engine/Source/Programs/UnrealBuildAccelerator/Common/Public/UbaStorageClient.h
+++ b/Engine/Source/Programs/UnrealBuildAccelerator/Common/Public/UbaStorageClient.h
@@ -62,7 +62,7 @@ namespace uba
 		virtual void PrintSummary(Logger& logger) override;
 
 	private:
-		bool PopulateCasFromDirsRecursive(const tchar* dir, WorkManager& workManager, UnorderedSet<u64>& seenIds, ReaderWriterLock& seenIdsLock, const Function<bool()>& shouldExit);
+		bool PopulateCasFromDirsRecursive(const tchar* dir, WorkManager& workManager, UnorderedSet<u64>& seenIds, Futex& seenIdsLock, const Function<bool()>& shouldExit);
 
 		NetworkClient& m_client;
 		bool m_sendCompressed;
@@ -79,8 +79,8 @@ namespace uba
 		struct LocalFile { CasEntry casEntry; TString fileName; Event hasBeenSent; LocalFile() : casEntry(CasKeyZero) {} };
 		UnorderedMap<CasKey, LocalFile> m_localStorageFiles;
 
-		ReaderWriterLock m_sendOneAtTheTimeLock;
-		ReaderWriterLock m_retrieveOneBatchAtTheTimeLock;
+		Futex m_sendOneAtTheTimeLock;
+		Futex m_retrieveOneBatchAtTheTimeLock;
 
 		static constexpr u8 ServiceId = StorageServiceId;
 
diff --git a/Engine/Source/Programs/UnrealBuildAccelerator/Common/Public/UbaStorageProxy.h b/Engine/Source/Programs/UnrealBuildAccelerator/Common/Public/UbaStorageProxy.h
--- a/Engine/Source/Programs/UnrealBuildAccelerator/Common/Public/UbaStorageProxy.h
+++ b/Engine/Source/Programs/UnrealBuildAccelerator/Common/Public/UbaStorageProxy.h
@@ -56,7 +56,7 @@ namespace uba
 		bool HandleDefault(MessageInfo& messageInfo, BinaryReader& reader, BinaryWriter& writer);
 
 		struct FileEntry;
-		bool GetFileFromLocalStorage(u32 clientId, const CasKey& casKey, FileEntry& file, ScopedWriteLock& fileLock);
+		bool GetFileFromLocalStorage(u32 clientId, const CasKey& casKey, FileEntry& file, ScopedFutex& fileLock);
 		bool SendFetchBeginResponse(u32 clientId, FileEntry& file, BinaryWriter& writer, bool writeData = true);
 		bool UpdateFetch(u32 clientId, u16 fetchId, u64 segmentSize);
 		bool SendEnd(const CasKey& key);
@@ -77,7 +77,7 @@ namespace uba
 
 		struct FileEntry
 		{
-			ReaderWriterLock lock;
+			Futex lock;
 			u8* memory = nullptr;
 			u64 size = 0;
 			Atomic<u64> received;
@@ -92,7 +92,7 @@ namespace uba
 			Vector<SegmentMessage*> segmentMessages;
 		};
 
-		ReaderWriterLock m_filesLock;
+		Futex m_filesLock;
 		UnorderedMap<CasKey, FileEntry> m_files;
 
 		struct ActiveFetch
diff --git a/Engine/Source/Programs/UnrealBuildAccelerator/Common/Public/UbaStorageServer.h b/Engine/Source/Programs/UnrealBuildAccelerator/Common/Public/UbaStorageServer.h
--- a/Engine/Source/Programs/UnrealBuildAccelerator/Common/Public/UbaStorageServer.h
+++ b/Engine/Source/Programs/UnrealBuildAccelerator/Common/Public/UbaStorageServer.h
@@ -59,7 +59,7 @@ namespace uba
 		static constexpr u8 ServiceId = StorageServiceId;
 
 
-		ReaderWriterLock m_waitEntriesLock;
+		Futex m_waitEntriesLock;
 
 		struct WaitEntry
 		{
@@ -106,7 +106,7 @@ namespace uba
 		Atomic<u32> m_activeUnmap;
 
 
-		ReaderWriterLock m_availableIdsLock;
+		Futex m_availableIdsLock;
 		Vector<u16> m_availableIds;
 		u16 m_availableIdsHigh = 1;
 
@@ -132,7 +132,7 @@ namespace uba
 		ReaderWriterLock m_connectionInfoLock;
 		UnorderedMap<u32, Info> m_connectionInfo;
 
-		ReaderWriterLock m_loadCasTableLock;
+		Futex m_loadCasTableLock;
 
 		Trace* m_trace = nullptr;
 
diff --git a/Engine/Source/Programs/UnrealBuildAccelerator/Common/Public/UbaStorageUtils.h b/Engine/Source/Programs/UnrealBuildAccelerator/Common/Public/UbaStorageUtils.h
--- a/Engine/Source/Programs/UnrealBuildAccelerator/Common/Public/UbaStorageUtils.h
+++ b/Engine/Source/Programs/UnrealBuildAccelerator/Common/Public/UbaStorageUtils.h
@@ -27,7 +27,7 @@ namespace uba
 		NetworkClient& m_client;
 		BufferSlots& m_bufferSlots;
 		StorageStats& m_stats;
-		ReaderWriterLock& m_sendOneAtTheTimeLock;
+		Futex& m_sendOneAtTheTimeLock;
 		u8 m_casCompressor = DefaultCompressor;
 		u8 m_casCompressionLevel = DefaultCompressionLevel;
 		bool m_sendOneBigFileAtTheTime = true;
diff --git a/Engine/Source/Programs/UnrealBuildAccelerator/Common/Public/UbaTrace.h b/Engine/Source/Programs/UnrealBuildAccelerator/Common/Public/UbaTrace.h
--- a/Engine/Source/Programs/UnrealBuildAccelerator/Common/Public/UbaTrace.h
+++ b/Engine/Source/Programs/UnrealBuildAccelerator/Common/Public/UbaTrace.h
@@ -125,7 +125,7 @@ namespace uba
 		LoggerWithWriter m_logger;
 		TString m_namedTrace;
 		TraceChannel m_channel;
-		ReaderWriterLock m_memoryLock;
+		Futex m_memoryLock;
 		FileMappingHandle m_memoryHandle;
 		u8* m_memoryBegin = nullptr;
 		u64 m_memoryPos = 0;
@@ -133,7 +133,7 @@ namespace uba
 		u64 m_memoryCapacity = 0;
 		u64 m_startTime = ~u64(0);
 
-		ReaderWriterLock m_stringsLock;
+		Futex m_stringsLock;
 		UnorderedMap<StringKey, u32> m_strings;
 
 		Atomic<u32> m_workCounter;
diff --git a/Engine/Source/Programs/UnrealBuildAccelerator/Common/Public/UbaWorkManager.h b/Engine/Source/Programs/UnrealBuildAccelerator/Common/Public/UbaWorkManager.h
--- a/Engine/Source/Programs/UnrealBuildAccelerator/Common/Public/UbaWorkManager.h
+++ b/Engine/Source/Programs/UnrealBuildAccelerator/Common/Public/UbaWorkManager.h
@@ -70,12 +70,12 @@ namespace uba
 
 		Vector<Worker*> m_workers;
 		struct Work { WorkFunction func; TString desc; };
-		ReaderWriterLock m_workLock;
+		Futex m_workLock;
 		List<Work> m_work;
 		Atomic<u32> m_activeWorkerCount;
 		Atomic<u32> m_workCounter;
 
-		ReaderWriterLock m_availableWorkersLock;
+		Futex m_availableWorkersLock;
 		Worker* m_firstAvailableWorker = nullptr;
 	};
 
@@ -110,7 +110,7 @@ namespace uba
 			u32 refCount;
 			u32 activeCount;
 			bool isDone;
-			ReaderWriterLock lock;
+			Futex lock;
 			Event* doneEvent;
 		};
 
@@ -129,7 +129,7 @@ namespace uba
 				u32 active = 0;
 				while (true)
 				{
-					SCOPED_WRITE_LOCK(context->lock, l);
+					SCOPED_FUTEX(context->lock, l);
 					context->activeCount -= active;
 					context->isDone = context->it == context->end;
 					if (context->isDone)
diff --git a/Engine/Source/Programs/UnrealBuildAccelerator/Core/Private/UbaEvent.cpp b/Engine/Source/Programs/UnrealBuildAccelerator/Core/Private/UbaEvent.cpp
--- a/Engine/Source/Programs/UnrealBuildAccelerator/Core/Private/UbaEvent.cpp
+++ b/Engine/Source/Programs/UnrealBuildAccelerator/Core/Private/UbaEvent.cpp
@@ -290,7 +290,7 @@ namespace uba
 			bool manualReset;
 		};
 	};
-	ReaderWriterLock g_firstEventLock[2];
+	Futex g_firstEventLock[2];
 	EventImpl* g_firstEvent[2];
 #endif
 
@@ -325,7 +325,7 @@ namespace uba
 	{
 		#if PLATFORM_WINDOWS
 		EventImpl* impl = nullptr;
-		SCOPED_WRITE_LOCK(g_firstEventLock[manualReset], lock);
+		SCOPED_FUTEX(g_firstEventLock[manualReset], lock);
 		auto*& first = g_firstEvent[manualReset];
 		if (first)
 		{
@@ -360,7 +360,7 @@ namespace uba
 			m_ev = nullptr;
 			bool manualReset = impl->manualReset;
 			auto*& first = g_firstEvent[manualReset];
-			SCOPED_WRITE_LOCK(g_firstEventLock[manualReset], l);
+			SCOPED_FUTEX(g_firstEventLock[manualReset], l);
 			impl->next = first;
 			first = impl;
 		}
diff --git a/Engine/Source/Programs/UnrealBuildAccelerator/Core/Private/UbaMemory.cpp b/Engine/Source/Programs/UnrealBuildAccelerator/Core/Private/UbaMemory.cpp
--- a/Engine/Source/Programs/UnrealBuildAccelerator/Core/Private/UbaMemory.cpp
+++ b/Engine/Source/Programs/UnrealBuildAccelerator/Core/Private/UbaMemory.cpp
@@ -89,7 +89,7 @@ namespace uba
 		if (!memory)
 			return aligned_alloc(alignment, bytes);
 
-		SCOPED_WRITE_LOCK(lock, l);
+		SCOPED_FUTEX(lock, l);
 		return AllocateNoLock(bytes, alignment, hint);
 	}
 
diff --git a/Engine/Source/Programs/UnrealBuildAccelerator/Core/Private/UbaSynchronization.cpp b/Engine/Source/Programs/UnrealBuildAccelerator/Core/Private/UbaSynchronization.cpp
--- a/Engine/Source/Programs/UnrealBuildAccelerator/Core/Private/UbaSynchronization.cpp
+++ b/Engine/Source/Programs/UnrealBuildAccelerator/Core/Private/UbaSynchronization.cpp
@@ -49,7 +49,7 @@ namespace uba
 		#endif
 	}
 
-	void CriticalSection::Enter()
+	void CriticalSection::Enter() const
 	{
 		#if PLATFORM_WINDOWS
 		EnterCriticalSection((CRITICAL_SECTION*)&data);
@@ -58,7 +58,7 @@ namespace uba
 		#endif
 	}
 
-	void CriticalSection::Leave()
+	void CriticalSection::Leave() const
 	{
 		#if PLATFORM_WINDOWS
 		LeaveCriticalSection((CRITICAL_SECTION*)&data);
@@ -124,7 +124,7 @@ namespace uba
 		#endif
 	}
 
-	void ReaderWriterLock::EnterWrite()
+	void ReaderWriterLock::Enter()
 	{
 		#if PLATFORM_WINDOWS
 		AcquireSRWLockExclusive((SRWLOCK*)&data);
@@ -135,7 +135,7 @@ namespace uba
 		#endif
 	}
 
-	bool ReaderWriterLock::TryEnterWrite()
+	bool ReaderWriterLock::TryEnter()
 	{
 		#if PLATFORM_WINDOWS
 		return TryAcquireSRWLockExclusive((SRWLOCK*)&data) != 0;
@@ -146,7 +146,7 @@ namespace uba
 		#endif
 	}
 
-	void ReaderWriterLock::LeaveWrite()
+	void ReaderWriterLock::Leave()
 	{
 		#if PLATFORM_WINDOWS
 		ReleaseSRWLockExclusive((SRWLOCK*)&data);
@@ -168,8 +168,8 @@ namespace uba
 
 	ContentionTracker& GetContentionTracker(const char* file, u64 line)
 	{
-		static ReaderWriterLock rwl;
-		ScopedWriteLock l(rwl);
+		static Futex rwl;
+		ScopedFutex l(rwl);
 		ContentionTracker& ct = GetContentionTrackerList().emplace_back();
 		ct.file = file;
 		ct.line = line;
diff --git a/Engine/Source/Programs/UnrealBuildAccelerator/Core/Public/UbaMemory.h b/Engine/Source/Programs/UnrealBuildAccelerator/Core/Public/UbaMemory.h
--- a/Engine/Source/Programs/UnrealBuildAccelerator/Core/Public/UbaMemory.h
+++ b/Engine/Source/Programs/UnrealBuildAccelerator/Core/Public/UbaMemory.h
@@ -107,7 +107,7 @@ namespace uba
 		void Swap(MemoryBlock& other);
 		
 
-		ReaderWriterLock lock;
+		Futex lock;
 		u8* memory;
 		u64 reserveSize = 0;
 		u64 writtenSize = 0;
@@ -201,7 +201,7 @@ namespace uba
 
 		void* Allocate()
 		{
-			SCOPED_WRITE_LOCK(m_lock, lock);
+			SCOPED_FUTEX(m_lock, lock);
 			if (m_nextFree)
 			{
 				void* ptr = (void*)m_nextFree;
@@ -217,12 +217,12 @@ namespace uba
 			#if UBA_DEBUG
 			memset(mem, 0xFE, sizeof(Type));
 			#endif
-			SCOPED_WRITE_LOCK(m_lock, lock);
+			SCOPED_FUTEX(m_lock, lock);
 			*(u64*)mem = m_nextFree;
 			m_nextFree = u64(mem);
 		}
 
-		ReaderWriterLock m_lock;
+		Futex m_lock;
 		MemoryBlock& m_memory;
 		u64 m_nextFree = 0;
 	};
@@ -289,6 +289,32 @@ namespace uba
 		return v;
 	}
 
+	#if 0
+	template<typename Key, typename Value, typename Hash = std::hash<Key>, typename EqualTo = std::equal_to<Key>>
+	struct SafeUnorderedMap
+	{
+		using InternalMap = UnorderedMap<Key, Value, Hash, EqualTo>;
+
+		template <typename... Args>
+		Value& GetOrCreate(const Key& key, Args&&... args)
+		{
+			SCOPED_READ_LOCK(lock, read);
+			auto findIt = map.find(key);
+			if (findIt != map.end())
+				return findIt->second;
+			read.Leave();
+			SCOPED_WRITE_LOCK(lock, write);
+			return map.try_emplace(key, std::forward<Args>(args)...).first->second;
+		}
+
+		InternalMap::iterator find(const Key& key) { return map.find(key); }
+		InternalMap::iterator end() { return map.end(); }
+
+		ReaderWriterLock lock;
+		InternalMap map;
+	};
+	#endif
+
 	bool SupportsHugePages();
 	u64 GetHugePageCount();
 }
diff --git a/Engine/Source/Programs/UnrealBuildAccelerator/Core/Public/UbaSynchronization.h b/Engine/Source/Programs/UnrealBuildAccelerator/Core/Public/UbaSynchronization.h
--- a/Engine/Source/Programs/UnrealBuildAccelerator/Core/Public/UbaSynchronization.h
+++ b/Engine/Source/Programs/UnrealBuildAccelerator/Core/Public/UbaSynchronization.h
@@ -34,8 +34,8 @@ namespace uba
 		CriticalSection(bool recursive = true);
 		~CriticalSection();
 
-		void Enter();
-		void Leave();
+		void Enter() const;
+		void Leave() const;
 
 	private:
 		#if PLATFORM_WINDOWS
@@ -59,10 +59,10 @@ namespace uba
 		void EnterRead() const;
 		void LeaveRead() const;
 
-		void EnterWrite();
-		void LeaveWrite();
+		void Enter();
+		void Leave();
 
-		bool TryEnterWrite();
+		bool TryEnter();
 
 	private:
 
@@ -85,6 +85,7 @@ namespace uba
 	#define SCOPED_WRITE_LOCK(readerWriterLock, name) ScopedWriteLock name(readerWriterLock);
 	#define SCOPED_CRITICAL_SECTION(cs, name) ScopedCriticalSection name(cs);
 	#define SCOPED_FUTEX(cs, name) ScopedFutex name(cs);
+	#define SCOPED_FUTEX_READ(cs, name) ScopedFutexRead name(cs);
 	#else
 	u64 GetTime();
 	struct ContentionTracker { void Add(u64 t) { time += t; ++count; }; Atomic<u64> time; Atomic<u64> count; const char* file; u64 line; };
@@ -100,6 +101,7 @@ namespace uba
 	#define SCOPED_WRITE_LOCK(readerWriterLock, name) SCOPED_LOCK_INTERNAL(readerWriterLock, ScopedWriteLock, name)
 	#define SCOPED_CRITICAL_SECTION(cs, name) SCOPED_LOCK_INTERNAL(cs, ScopedCriticalSection, name)
 	#define SCOPED_FUTEX(futex, name) SCOPED_LOCK_INTERNAL(futex, ScopedFutex, name)
+	#define SCOPED_FUTEX_READ(futex, name) SCOPED_LOCK_INTERNAL(futex, ScopedFutex, name)
 	#endif
 
 	////////////////////////////////////////////////////////////////////////////////////////////////////
@@ -107,12 +109,12 @@ namespace uba
 	class ScopedCriticalSection 
 	{
 	public:
-		ScopedCriticalSection(CriticalSection& cs) : m_cs(cs), m_active(true) { cs.Enter(); }
+		ScopedCriticalSection(const CriticalSection& cs) : m_cs(cs), m_active(true) { cs.Enter(); }
 		~ScopedCriticalSection() { Leave(); }
 		void Enter() { if (m_active) return; m_cs.Enter(); m_active = true; }
 		void Leave() { if (!m_active) return; m_cs.Leave(); m_active = false; }
 	private:
-		CriticalSection& m_cs;
+		const CriticalSection& m_cs;
 		bool m_active;
 	};
 
@@ -135,10 +137,10 @@ namespace uba
 	class ScopedWriteLock
 	{
 	public:
-		ScopedWriteLock(ReaderWriterLock& lock) : m_lock(lock) { lock.EnterWrite(); }
+		ScopedWriteLock(ReaderWriterLock& lock) : m_lock(lock) { lock.Enter(); }
 		~ScopedWriteLock() { Leave(); }
-		inline void Enter() { if (m_active) return; m_active = true; m_lock.EnterWrite(); }
-		inline void Leave() { if (!m_active) return; m_active = false; m_lock.LeaveWrite(); }
+		inline void Enter() { if (m_active) return; m_active = true; m_lock.Enter(); }
+		inline void Leave() { if (!m_active) return; m_active = false; m_lock.Leave(); }
 
 		ReaderWriterLock& m_lock;
 		bool m_active = true;
@@ -149,9 +151,11 @@ namespace uba
 	#if PLATFORM_WINDOWS
 	using Futex = ReaderWriterLock;
 	using ScopedFutex = ScopedWriteLock;
+	using ScopedFutexRead = ScopedReadLock;
 	#else
-	class Futex : public CriticalSection { Futex() : CriticalSection(false) {} };
+	class Futex : public CriticalSection { public: Futex() : CriticalSection(false) {} };
 	using ScopedFutex = ScopedCriticalSection;
+	using ScopedFutexRead = ScopedCriticalSection;
 	#endif
 
 
diff --git a/Engine/Source/Programs/UnrealBuildAccelerator/Detours/Private/UbaDetoursShared.cpp b/Engine/Source/Programs/UnrealBuildAccelerator/Detours/Private/UbaDetoursShared.cpp
--- a/Engine/Source/Programs/UnrealBuildAccelerator/Detours/Private/UbaDetoursShared.cpp
+++ b/Engine/Source/Programs/UnrealBuildAccelerator/Detours/Private/UbaDetoursShared.cpp
@@ -60,7 +60,7 @@ namespace uba
 	thread_local char t_b[LogBufSize];
 	thread_local u32 t_b_size;
 	thread_local u32 t_logScopeCount;
-	ReaderWriterLock g_logScopeLock;
+	Futex g_logScopeLock;
 
 	void GetPrefixExtra(StringBufferBase& out)
 	{
@@ -140,7 +140,7 @@ namespace uba
 	{
 		if (++t_logScopeCount > 1)
 			return;
-		//g_logScopeLock.EnterWrite(); // Deadlocks in a few places
+		//g_logScopeLock.Enter(); // Deadlocks in a few places
 	}
 	LogScope::~LogScope()
 	{
@@ -148,7 +148,7 @@ namespace uba
 			return;
 		if (t_b_size)
 			Flush();
-		//g_logScopeLock.LeaveWrite();
+		//g_logScopeLock.Leave();
 	}
 	void LogScope::Flush()
 	{
@@ -334,7 +334,7 @@ namespace uba
 		bool hasLock = false;
 		while (tryCount++ < 5)
 		{
-			hasLock = g_communicationLock.TryEnterWrite();
+			hasLock = g_communicationLock.TryEnter();
 			if (hasLock)
 				break;
 			Sleep(100);
@@ -353,7 +353,7 @@ namespace uba
 			out.Append(TC("\n   Failed to resolve callstack\n"));
 		}
 		// Note, we leave the lock even though we might not have it because we want to be able to report
-		g_communicationLock.LeaveWrite();
+		g_communicationLock.Leave();
 	}
 
 	//TODO: Implement SetConsoleTextAttribute.. clang is using it to color errors
diff --git a/Engine/Source/Programs/UnrealBuildAccelerator/Visualizer/Private/UbaTraceReader.h b/Engine/Source/Programs/UnrealBuildAccelerator/Visualizer/Private/UbaTraceReader.h
--- a/Engine/Source/Programs/UnrealBuildAccelerator/Visualizer/Private/UbaTraceReader.h
+++ b/Engine/Source/Programs/UnrealBuildAccelerator/Visualizer/Private/UbaTraceReader.h
@@ -212,7 +212,6 @@ namespace uba
 
 		Logger& m_logger;
 		TraceChannel m_channel;
-		ReaderWriterLock m_memoryLock;
 		FileMappingHandle m_memoryHandle;
 		TString m_namedTrace;
 		u8* m_memoryBegin = nullptr;
