[UBA]
* Fixed so mac/linux also have cache for events since they are expensive primitives

diff --git a/Engine/Source/Programs/UnrealBuildAccelerator/Common/Private/UbaProcess.cpp b/Engine/Source/Programs/UnrealBuildAccelerator/Common/Private/UbaProcess.cpp
--- a/Engine/Source/Programs/UnrealBuildAccelerator/Common/Private/UbaProcess.cpp
+++ b/Engine/Source/Programs/UnrealBuildAccelerator/Common/Private/UbaProcess.cpp
@@ -139,9 +139,9 @@ namespace uba
 	,	m_id(id)
 	,	m_comMemory(detourEnabled ? m_session.m_processCommunicationAllocator.Alloc(TC("")) : FileMappingAllocator::Allocation())
 		#if !PLATFORM_WINDOWS
-	,	m_cancelEvent(*(m_comMemory.memory ? (new (m_comMemory.memory) Event) : (Event*)nullptr))
-	,	m_writeEvent(*(m_comMemory.memory ? (new (m_comMemory.memory + sizeof(Event)) Event) : (Event*)nullptr))
-	,	m_readEvent(*(m_comMemory.memory ? (new (m_comMemory.memory + sizeof(Event)*2) Event) : (Event*)nullptr))
+	,	m_cancelEvent(*(m_comMemory.memory ? (new (m_comMemory.memory) SharedEvent) : (SharedEvent*)nullptr))
+	,	m_writeEvent(*(m_comMemory.memory ? (new (m_comMemory.memory + sizeof(SharedEvent)) SharedEvent) : (SharedEvent*)nullptr))
+	,	m_readEvent(*(m_comMemory.memory ? (new (m_comMemory.memory + sizeof(SharedEvent)*2) SharedEvent) : (SharedEvent*)nullptr))
 		#endif
 	,	m_detourEnabled(detourEnabled)
 	,	m_writtenFilesLock(parent ? parent->m_writtenFilesLock : *new Futex())
@@ -151,9 +151,9 @@ namespace uba
 	{
 		if (m_comMemory.memory)
 		{
-			m_cancelEvent.Create(true, true);
-			m_writeEvent.Create(false, true);
-			m_readEvent.Create(false, true);
+			m_cancelEvent.Create(true);
+			m_writeEvent.Create(false);
+			m_readEvent.Create(false);
 		}
 	}
 
@@ -173,9 +173,9 @@ namespace uba
 		if (m_comMemory.memory)
 		{
 			#if !PLATFORM_WINDOWS
-			m_cancelEvent.~Event();
-			m_writeEvent.~Event();
-			m_readEvent.~Event();
+			m_cancelEvent.~SharedEvent();
+			m_writeEvent.~SharedEvent();
+			m_readEvent.~SharedEvent();
 			#endif
 			m_session.m_processCommunicationAllocator.Free(m_comMemory);
 		}
@@ -496,8 +496,8 @@ namespace uba
 		u8* comMemory = m_comMemory.memory;
 		u64 comMemorySize = CommunicationMemSize;
 		#if !PLATFORM_WINDOWS
-		comMemory += sizeof(Event) * 3;
-		comMemorySize -= sizeof(Event) * 3;
+		comMemory += sizeof(SharedEvent) * 3;
+		comMemorySize -= sizeof(SharedEvent) * 3;
 		#endif
 
 		u32 retryCount = 0; // Do not allow retry
@@ -606,9 +606,9 @@ namespace uba
 		{
 			#if !PLATFORM_WINDOWS
 			SCOPED_FUTEX(m_comMemoryLock, lock);
-			m_cancelEvent.~Event();
-			m_writeEvent.~Event();
-			m_readEvent.~Event();
+			m_cancelEvent.~SharedEvent();
+			m_writeEvent.~SharedEvent();
+			m_readEvent.~SharedEvent();
 			#endif
 			m_session.m_processCommunicationAllocator.Free(m_comMemory);
 			m_comMemory = {};
diff --git a/Engine/Source/Programs/UnrealBuildAccelerator/Common/Public/UbaProcess.h b/Engine/Source/Programs/UnrealBuildAccelerator/Common/Public/UbaProcess.h
--- a/Engine/Source/Programs/UnrealBuildAccelerator/Common/Public/UbaProcess.h
+++ b/Engine/Source/Programs/UnrealBuildAccelerator/Common/Public/UbaProcess.h
@@ -113,9 +113,9 @@ namespace uba
 	#else
 		Futex m_comMemoryLock;
 		Atomic<bool> m_cancelled;
-		Event& m_cancelEvent;
-		Event& m_writeEvent;
-		Event& m_readEvent;
+		SharedEvent& m_cancelEvent;
+		SharedEvent& m_writeEvent;
+		SharedEvent& m_readEvent;
 
 		int m_stdOutPipe = -1;
 		int m_stdErrPipe = -1;
diff --git a/Engine/Source/Programs/UnrealBuildAccelerator/Core/Private/UbaEvent.cpp b/Engine/Source/Programs/UnrealBuildAccelerator/Core/Private/UbaEvent.cpp
--- a/Engine/Source/Programs/UnrealBuildAccelerator/Core/Private/UbaEvent.cpp
+++ b/Engine/Source/Programs/UnrealBuildAccelerator/Core/Private/UbaEvent.cpp
@@ -23,14 +23,14 @@ namespace uba
 		EventImpl() = default;
 		~EventImpl() { Destroy(); }
 
-		bool Create(bool manualReset, bool shared = false)
+		bool Create(bool mr, bool shared = false)
 		{
 			static_assert(Atomic<u32>::is_always_lock_free, "atomic<T> not lock free, can't work in shared mem");
 			static_assert(Atomic<TriggerType>::is_always_lock_free, "atomic<T> not lock free, can't work in shared mem");
 			static_assert(Atomic<bool>::is_always_lock_free, "atomic<T> not lock free, can't work in shared mem");
 
 			UBA_ASSERTF(!m_initialized, "Can't create already created Event");
-			m_manualReset = manualReset;
+			m_manualReset = mr;
 
 			pthread_mutexattr_t attrmutex;
 			if (pthread_mutexattr_init(&attrmutex) != 0)
@@ -257,6 +257,12 @@ namespace uba
 			return false;
 		}
 
+		union
+		{
+			EventImpl* next;
+			bool manualReset;
+		};
+
 	private:
 		enum TriggerType { TriggerType_None, TriggerType_One, TriggerType_All };
 
@@ -289,41 +295,30 @@ namespace uba
 			EventImpl* next;
 			bool manualReset;
 		};
+
+		bool Create(bool mr) { handle = CreateEvent(nullptr, mr, false, NULL); return handle != 0; }
+		void Set() { SetEvent(handle); }
+		void Reset() { ResetEvent(handle); }
+		bool IsSet(u32 timeOutMs) { return WaitForSingleObject(handle, timeOutMs) == 0;} 
 	};
-	Futex g_firstEventLock[2];
-	EventImpl* g_firstEvent[2];
 #endif
 
+	Futex g_firstEventLock[2];
+	EventImpl* g_firstEvent[2];
 
-	Event::Event()
-	{
-		#if PLATFORM_WINDOWS
-		m_ev = nullptr;
-		#else
-		static_assert(sizeof(m_data) >= sizeof(EventImpl));
-		new (m_data) EventImpl();
-		#endif
-	}
 
 	Event::Event(bool manualReset)
 	{
-		#if !PLATFORM_WINDOWS
-		new (m_data) EventImpl();
-		#endif
-		Create(manualReset, false);
+		Create(manualReset);
 	}
 
 	Event::~Event()
 	{
 		Destroy();
-		#if !PLATFORM_WINDOWS
-		((EventImpl&)m_data).~EventImpl();
-		#endif
 	}
 
-	bool Event::Create(bool manualReset, bool shared)
+	bool Event::Create(bool manualReset)
 	{
-		#if PLATFORM_WINDOWS
 		EventImpl* impl = nullptr;
 		SCOPED_FUTEX(g_firstEventLock[manualReset], lock);
 		auto*& first = g_firstEvent[manualReset];
@@ -336,91 +331,120 @@ namespace uba
 
 		if (!impl)
 		{
-			impl = new EventImpl;
-			impl->handle = CreateEvent(nullptr, manualReset, false, NULL);
-			UBA_ASSERT(impl->handle);
+			impl = new EventImpl();
+			impl->Create(manualReset);
 		}
 		else
-			ResetEvent(impl->handle);
+			impl->Reset();
 
 		impl->manualReset = manualReset;
-		m_ev = impl;
+		m_impl = impl;
 		return true;
-		#else
-		return ((EventImpl&)m_data).Create(manualReset, shared);
-		#endif
 	}
 
 	void Event::Destroy()
 	{
-		#if PLATFORM_WINDOWS
-		if (m_ev)
-		{
-			auto impl = (EventImpl*)m_ev;
-			m_ev = nullptr;
-			bool manualReset = impl->manualReset;
-			auto*& first = g_firstEvent[manualReset];
-			SCOPED_FUTEX(g_firstEventLock[manualReset], l);
-			impl->next = first;
-			first = impl;
-		}
-		#else
-		((EventImpl&)m_data).Destroy();
-		#endif
+		if (!m_impl)
+			return;
+		auto impl = m_impl;
+		m_impl = nullptr;
+		bool manualReset = impl->manualReset;
+		auto*& first = g_firstEvent[manualReset];
+		SCOPED_FUTEX(g_firstEventLock[manualReset], l);
+		impl->next = first;
+		first = impl;
 	}
 
 	void Event::Set()
 	{
-		#if PLATFORM_WINDOWS
-		if (!m_ev)
-			return;
-		SetEvent(((EventImpl*)m_ev)->handle);
-		#else
-		((EventImpl&)m_data).Set();
-		#endif
+		if (m_impl)
+			m_impl->Set();
 	}
 
 	void Event::Reset()
 	{
-		#if PLATFORM_WINDOWS
-		if (!m_ev)
-			return;
-		ResetEvent(((EventImpl*)m_ev)->handle);
-		#else
-		((EventImpl&)m_data).Reset();
-		#endif
+		if (m_impl)
+			m_impl->Reset();
 	}
 
 	bool Event::IsCreated()
 	{
-		#if PLATFORM_WINDOWS
-		return m_ev != nullptr;
-		#else
-		return ((EventImpl&)m_data).IsCreated();
-		#endif
+		return m_impl != nullptr;
 	}
 
 	bool Event::IsSet(u32 timeOutMs)
 	{
-		#if PLATFORM_WINDOWS
-		if (!m_ev)
-			return false;
-		return WaitForSingleObject(((EventImpl*)m_ev)->handle, timeOutMs) == 0;
-		#else
-		return ((EventImpl&)m_data).IsSet(timeOutMs);
-		#endif
+		if (m_impl)
+			return m_impl->IsSet(timeOutMs);
+		return false;
 	}
 
 	void* Event::GetHandle()
 	{
 		#if PLATFORM_WINDOWS
-		if (!m_ev)
+		if (!m_impl)
 			return nullptr;
-		return ((EventImpl*)m_ev)->handle;
+		return m_impl->handle;
 		#else
 		UBA_ASSERTF(false, "Event::GetHandle not available");
 		return 0;
 		#endif
 	}
 
+
+#if !PLATFORM_WINDOWS
+	SharedEvent::SharedEvent()
+	{
+		static_assert(sizeof(m_data) >= sizeof(EventImpl));
+		new (m_data) EventImpl();
+	}
+
+	SharedEvent::SharedEvent(bool manualReset)
+	{
+		new (m_data) EventImpl();
+		Create(manualReset);
+	}
+
+	SharedEvent::~SharedEvent()
+	{
+		Destroy();
+		((EventImpl&)m_data).~EventImpl();
+	}
+
+	bool SharedEvent::Create(bool manualReset)//, bool shared)
+	{
+		return ((EventImpl&)m_data).Create(manualReset, true);
+	}
+
+	void SharedEvent::Destroy()
+	{
+		((EventImpl&)m_data).Destroy();
+	}
+
+	void SharedEvent::Set()
+	{
+		((EventImpl&)m_data).Set();
+	}
+
+	void SharedEvent::Reset()
+	{
+		((EventImpl&)m_data).Reset();
+	}
+
+	bool SharedEvent::IsCreated()
+	{
+		return ((EventImpl&)m_data).IsCreated();
+	}
+
+	bool SharedEvent::IsSet(u32 timeOutMs)
+	{
+		return ((EventImpl&)m_data).IsSet(timeOutMs);
+	}
+
+	void* SharedEvent::GetHandle()
+	{
+		UBA_ASSERTF(false, "Event::GetHandle not available");
+		return 0;
+	}
+#endif
 }
diff --git a/Engine/Source/Programs/UnrealBuildAccelerator/Core/Public/UbaEvent.h b/Engine/Source/Programs/UnrealBuildAccelerator/Core/Public/UbaEvent.h
--- a/Engine/Source/Programs/UnrealBuildAccelerator/Core/Public/UbaEvent.h
+++ b/Engine/Source/Programs/UnrealBuildAccelerator/Core/Public/UbaEvent.h
@@ -9,10 +9,10 @@ namespace uba
 	class Event
 	{
 	public:
-		Event();
+		Event() = default;
 		Event(bool manualReset);
 		~Event();
-		bool Create(bool manualReset, bool shared = false);
+		bool Create(bool manualReset);
 		void Destroy();
 		void Set();
 		void Reset();
@@ -20,13 +20,29 @@ namespace uba
 		bool IsSet(u32 timeOutMs = ~0u);
 		void* GetHandle();
 	private:
-		#if PLATFORM_WINDOWS
-		void* m_ev;
-		#else
-		u64 m_data[16];
-		#endif
-
+		struct EventImpl* m_impl = nullptr;
 		Event(const Event&) = delete;
 		void operator=(const Event&) = delete;
 	};
+
+#if !PLATFORM_WINDOWS
+	class SharedEvent
+	{
+	public:
+		SharedEvent();
+		SharedEvent(bool manualReset);
+		~SharedEvent();
+		bool Create(bool manualReset);
+		void Destroy();
+		void Set();
+		void Reset();
+		bool IsCreated();
+		bool IsSet(u32 timeOutMs = ~0u);
+		void* GetHandle();
+	private:
+		u64 m_data[16];
+		SharedEvent(const SharedEvent&) = delete;
+		void operator=(const SharedEvent&) = delete;
+	};
+#endif
 }
diff --git a/Engine/Source/Programs/UnrealBuildAccelerator/Detours/Private/PosixOS/UbaDetoursMainPosix.cpp b/Engine/Source/Programs/UnrealBuildAccelerator/Detours/Private/PosixOS/UbaDetoursMainPosix.cpp
--- a/Engine/Source/Programs/UnrealBuildAccelerator/Detours/Private/PosixOS/UbaDetoursMainPosix.cpp
+++ b/Engine/Source/Programs/UnrealBuildAccelerator/Detours/Private/PosixOS/UbaDetoursMainPosix.cpp
@@ -15,9 +15,9 @@ namespace uba
 {
 	int g_comFd = -1;
 	int g_sessionPid = -1;
-	Event* g_cancelEvent;
-	Event* g_readEvent;
-	Event* g_writeEvent;
+	SharedEvent* g_cancelEvent;
+	SharedEvent* g_readEvent;
+	SharedEvent* g_writeEvent;
 	u8* g_messageMappingMem;
 
 	void PreInit(const char* logFile);
@@ -76,10 +76,10 @@ static void __attribute__((constructor(102))) PreInitCtor()
 		return;
 	}
 
-	g_cancelEvent = ((Event*)rptr);
-	g_readEvent = ((Event*)rptr) + 1;
-	g_writeEvent = ((Event*)rptr) + 2;
-	g_messageMappingMem = rptr + sizeof(Event) * 3;
+	g_cancelEvent = ((SharedEvent*)rptr);
+	g_readEvent = ((SharedEvent*)rptr) + 1;
+	g_writeEvent = ((SharedEvent*)rptr) + 2;
+	g_messageMappingMem = rptr + sizeof(SharedEvent) * 3;
 
 	AddExceptionHandler();
 
@@ -126,7 +126,7 @@ namespace uba
 	{
 		m_begin = g_messageMappingMem;
 		m_pos = m_begin;
-		m_end = m_begin + CommunicationMemSize - sizeof(Event) * 3;
+		m_end = m_begin + CommunicationMemSize - sizeof(SharedEvent) * 3;
 	}
 
 	void BinaryWriter::Flush(bool waitOnResponse)
@@ -155,6 +155,6 @@ namespace uba
 	{
 		m_begin = g_messageMappingMem;
 		m_pos = m_begin;
-		m_end = m_begin + CommunicationMemSize - sizeof(Event) * 3;
+		m_end = m_begin + CommunicationMemSize - sizeof(SharedEvent) * 3;
 	}
 }
diff --git a/Engine/Source/Programs/UnrealBuildAccelerator/Test/Private/UbaTestBasics.cpp b/Engine/Source/Programs/UnrealBuildAccelerator/Test/Private/UbaTestBasics.cpp
--- a/Engine/Source/Programs/UnrealBuildAccelerator/Test/Private/UbaTestBasics.cpp
+++ b/Engine/Source/Programs/UnrealBuildAccelerator/Test/Private/UbaTestBasics.cpp
@@ -50,49 +50,57 @@ namespace uba
 		return true;
 	}
 
-	bool TestEvents(LoggerWithWriter& logger, const StringBufferBase& rootDir)
+	template<class EventType>
+	bool TestEventsImpl(LoggerWithWriter& logger, const StringBufferBase& rootDir)
 	{
-		for (u32 i=0; i!=2; ++i)
-		{
-			Event ev;
-			if (!ev.Create(true, i == 1))
-				return logger.Error(TC("Failed to create event"));
+		EventType ev;
+		if (!ev.Create(true))
+			return logger.Error(TC("Failed to create event"));
 
-			Thread t([&]()
-				{
-					Sleep(500);
-					//logger.Info(TC("Setting event"));
-					ev.Set();
-					Sleep(500);
-					return true;
-				});
-
-			if (ev.IsSet(1))
-				return logger.Error(TC("Event was set after 1ms timeout where it should take 500ms"));
-
-			if (ev.IsSet(0))
-				return logger.Error(TC("Event was set after no timeout where it should take 500ms"));
-
-			//logger.Info(TC("Waiting for event"));
-			if (!ev.IsSet(2000))
-				return logger.Error(TC("Event was not set after 2000ms where it should take 500ms"));
-			//logger.Info(TC("Event was set"));
-
-			if (t.Wait(0))
-				return logger.Error(TC("Thread wait timed out. Should already be done after 2000ms"));
-
-			if (!t.Wait(2000))
-				return logger.Error(TC("Thread wait did not timed out should be done after 2000ms"));
-
-			#if 0 // Long time test... disabled by default
-			u64 time = GetTime();
-			Event ev2(true);
-			ev2.IsSet(10 * 60 * 1000);
-			if (TimeToMs((GetTime() - time) < 9 * 60 * 1000))
-				return logger.Error(TC("Event timeout was way too fast"));
-			#endif
-		}
+		Thread t([&]()
+			{
+				Sleep(500);
+				//logger.Info(TC("Setting event"));
+				ev.Set();
+				Sleep(500);
+				return true;
+			});
+
+		if (ev.IsSet(1))
+			return logger.Error(TC("Event was set after 1ms timeout where it should take 500ms"));
+
+		if (ev.IsSet(0))
+			return logger.Error(TC("Event was set after no timeout where it should take 500ms"));
+
+		//logger.Info(TC("Waiting for event"));
+		if (!ev.IsSet(2000))
+			return logger.Error(TC("Event was not set after 2000ms where it should take 500ms"));
+		//logger.Info(TC("Event was set"));
+
+		if (t.Wait(0))
+			return logger.Error(TC("Thread wait timed out. Should already be done after 2000ms"));
 
+		if (!t.Wait(2000))
+			return logger.Error(TC("Thread wait did not timed out should be done after 2000ms"));
+
+		#if 0 // Long time test... disabled by default
+		u64 time = GetTime();
+		EventType ev2(true);
+		ev2.IsSet(10 * 60 * 1000);
+		if (TimeToMs((GetTime() - time) < 9 * 60 * 1000))
+			return logger.Error(TC("Event timeout was way too fast"));
+		#endif
+		return true;
+	}
+
+	bool TestEvents(LoggerWithWriter& logger, const StringBufferBase& rootDir)
+	{
+		if (!TestEventsImpl<Event>(logger, rootDir))
+			return false;
+		#if !PLATFORM_WINDOWS
+		if (!TestEventsImpl<SharedEvent>(logger, rootDir))
+			return false;
+		#endif
 		return true;
 	}
 
