[UBA]
* Fixed tsan reported potential race condition in cache server when adding cache entries

diff --git a/Engine/Source/Programs/UnrealBuildAccelerator/Common/Private/UbaCacheServer.cpp b/Engine/Source/Programs/UnrealBuildAccelerator/Common/Private/UbaCacheServer.cpp
--- a/Engine/Source/Programs/UnrealBuildAccelerator/Common/Private/UbaCacheServer.cpp
+++ b/Engine/Source/Programs/UnrealBuildAccelerator/Common/Private/UbaCacheServer.cpp
@@ -1418,7 +1418,9 @@ namespace uba
 			auto findIt = connectionBucket.cacheEntryLookup.find(cmdKey);
 			if (findIt == connectionBucket.cacheEntryLookup.end())
 				return true;
-			auto g = MakeGuard([&]() { connectionBucket.cacheEntryLookup.erase(findIt); });
+			CacheEntry entry(std::move(findIt->second));
+			connectionBucket.cacheEntryLookup.erase(findIt);
+			lock2.Leave();
 			if (!success)
 				return true;
 
@@ -1431,7 +1433,7 @@ namespace uba
 			lock3.Leave();
 			
 			SCOPED_WRITE_LOCK(cacheEntries.lock, lock4);
-			cacheEntries.entries.emplace_front(std::move(findIt->second));
+			cacheEntries.entries.emplace_front(std::move(entry));
 			return true;
 		}
 		case CacheMessageType_FetchPathTable:
