[UBA]
* Attempt at fixing macos compile error caused by type getting 8 bytes larger

diff --git a/Engine/Source/Programs/UnrealBuildAccelerator/Core/Private/UbaEvent.cpp b/Engine/Source/Programs/UnrealBuildAccelerator/Core/Private/UbaEvent.cpp
--- a/Engine/Source/Programs/UnrealBuildAccelerator/Core/Private/UbaEvent.cpp
+++ b/Engine/Source/Programs/UnrealBuildAccelerator/Core/Private/UbaEvent.cpp
@@ -257,14 +257,14 @@ namespace uba
 			return false;
 		}
 
-		union
-		{
-			EventImpl* next;
-			bool manualReset;
-		};
+		bool GetManualReset() { return m_manualReset; }
+		void ResetManualReset(bool mr) {}
+
+
+		EventImpl* next;
 
 	private:
-		enum TriggerType { TriggerType_None, TriggerType_One, TriggerType_All };
+		enum TriggerType : u8 { TriggerType_None, TriggerType_One, TriggerType_All };
 
 		Atomic<bool> m_initialized;
 		Atomic<bool> m_manualReset;
@@ -273,13 +273,13 @@ namespace uba
 		pthread_mutex_t m_mutex;
 		pthread_cond_t m_condition;
 
-		inline void LockEventMutex()
+		void LockEventMutex()
 		{
 			int res = pthread_mutex_lock(&m_mutex);(void)res;
 			UBA_ASSERTF(res == 0, "pthread_mutex_lock failed (error code %i)", res);
 		}
 
-		inline void UnlockEventMutex()
+		void UnlockEventMutex()
 		{
 			int res = pthread_mutex_unlock(&m_mutex);
 			UBA_ASSERTF(res == 0, "pthread_mutex_unlock failed (error code %i)", res);
@@ -300,6 +300,8 @@ namespace uba
 		void Set() { SetEvent(handle); }
 		void Reset() { ResetEvent(handle); }
 		bool IsSet(u32 timeOutMs) { return WaitForSingleObject(handle, timeOutMs) == 0;} 
+		bool GetManualReset() { return manualReset; }
+		void ResetManualReset(bool mr) { manualReset = mr; }
 	};
 #endif
 
@@ -337,7 +339,7 @@ namespace uba
 		else
 			impl->Reset();
 
-		impl->manualReset = manualReset;
+		impl->ResetManualReset(manualReset);
 		m_impl = impl;
 		return true;
 	}
@@ -348,7 +350,7 @@ namespace uba
 			return;
 		auto impl = m_impl;
 		m_impl = nullptr;
-		bool manualReset = impl->manualReset;
+		bool manualReset = impl->GetManualReset();
 		auto*& first = g_firstEvent[manualReset];
 		SCOPED_FUTEX(g_firstEventLock[manualReset], l);
 		impl->next = first;
diff --git a/Engine/Source/Programs/UnrealBuildAccelerator/Core/Public/UbaEvent.h b/Engine/Source/Programs/UnrealBuildAccelerator/Core/Public/UbaEvent.h
--- a/Engine/Source/Programs/UnrealBuildAccelerator/Core/Public/UbaEvent.h
+++ b/Engine/Source/Programs/UnrealBuildAccelerator/Core/Public/UbaEvent.h
@@ -40,7 +40,11 @@ namespace uba
 		bool IsSet(u32 timeOutMs = ~0u);
 		void* GetHandle();
 	private:
+		#if PLATFORM_LINUX
+		u64 m_data[13];
+		#else
 		u64 m_data[16];
+		#endif
 		SharedEvent(const SharedEvent&) = delete;
 		void operator=(const SharedEvent&) = delete;
 	};
