[UBA]
* Changed size of event storage when building for arm target

diff --git a/Engine/Source/Programs/UnrealBuildAccelerator/Core/Public/UbaEvent.h b/Engine/Source/Programs/UnrealBuildAccelerator/Core/Public/UbaEvent.h
--- a/Engine/Source/Programs/UnrealBuildAccelerator/Core/Public/UbaEvent.h
+++ b/Engine/Source/Programs/UnrealBuildAccelerator/Core/Public/UbaEvent.h
@@ -40,10 +40,10 @@ namespace uba
 		bool IsSet(u32 timeOutMs = ~0u);
 		void* GetHandle();
 	private:
-		#if PLATFORM_LINUX
-		u64 m_data[13];
-		#else
+		#if defined(__aarch64__)
 		u64 m_data[16];
+		#else
+		u64 m_data[13];
 		#endif
 		SharedEvent(const SharedEvent&) = delete;
 		void operator=(const SharedEvent&) = delete;
