[UBA]
* Final fix for making all targets compile (seems like checking for __aarch64__ on mac does not work)

diff --git a/Engine/Source/Programs/UnrealBuildAccelerator/Core/Public/UbaEvent.h b/Engine/Source/Programs/UnrealBuildAccelerator/Core/Public/UbaEvent.h
--- a/Engine/Source/Programs/UnrealBuildAccelerator/Core/Public/UbaEvent.h
+++ b/Engine/Source/Programs/UnrealBuildAccelerator/Core/Public/UbaEvent.h
@@ -40,7 +40,7 @@ namespace uba
 		bool IsSet(u32 timeOutMs = ~0u);
 		void* GetHandle();
 	private:
-		#if defined(__aarch64__)
+		#if PLATFORM_MAC || defined(__aarch64__)
 		u64 m_data[16];
 		#else
 		u64 m_data[13];
