[UBA]
* Added support for interpreting hyperlinks in session description and make them clickable in visualizer
* Fixed absolute time in traces running through memory mappings
* Fixed so sessions with non-unique id still get unique trace ids

diff --git a/Engine/Source/Programs/UnrealBuildAccelerator/Common/Private/UbaSession.cpp b/Engine/Source/Programs/UnrealBuildAccelerator/Common/Private/UbaSession.cpp
--- a/Engine/Source/Programs/UnrealBuildAccelerator/Common/Private/UbaSession.cpp
+++ b/Engine/Source/Programs/UnrealBuildAccelerator/Common/Private/UbaSession.cpp
@@ -1275,6 +1275,13 @@ namespace uba
 			OwnerInfo ownerInfo = GetOwnerInfo();
 			if (ownerInfo.pid)
 				traceName.Appendf(TC("_%s%u"), ownerInfo.id, ownerInfo.pid);
+
+			if (!info.useUniqueId)
+			{
+				Guid guid;
+				CreateGuid(guid);
+				traceName.Append(GuidToString(guid).str);
+			}
 		}
 
 		if (!traceName.IsEmpty())
diff --git a/Engine/Source/Programs/UnrealBuildAccelerator/Visualizer/Private/UbaTraceReader.cpp b/Engine/Source/Programs/UnrealBuildAccelerator/Visualizer/Private/UbaTraceReader.cpp
--- a/Engine/Source/Programs/UnrealBuildAccelerator/Visualizer/Private/UbaTraceReader.cpp
+++ b/Engine/Source/Programs/UnrealBuildAccelerator/Visualizer/Private/UbaTraceReader.cpp
@@ -326,6 +326,8 @@ namespace uba
 			if (traceSystemStartTimeUs)
 				out.realStartTime = GetTime() - UsToTime(GetSystemTimeUs() - traceSystemStartTimeUs);
 
+			out.traceSystemStartTimeUs = traceSystemStartTimeUs;
+
 			out.startTime = out.realStartTime;
 			if (replay)
 				out.startTime = GetTime();
@@ -386,6 +388,23 @@ namespace uba
 				return false;
 			Guid clientUid = ReadClientId(out, reader);
 			u32 sessionIndex = reader.ReadU32();
+			TString hyperlink;
+
+			const tchar* hyperlinkPos = nullptr;
+			if (!sessionInfo.Contains(TC("http://"), true, &hyperlinkPos))
+				sessionInfo.Contains(TC("https://"), true, &hyperlinkPos);
+			if (hyperlinkPos)
+			{
+				const tchar* hyperlinkEnd = sessionInfo.data + sessionInfo.count;
+				if (const tchar* space = TStrchr(hyperlinkPos, ' '))
+					hyperlinkEnd = space;
+
+				hyperlink.assign(hyperlinkPos, hyperlinkEnd);
+				if (*hyperlinkEnd == ' ')
+					++hyperlinkEnd;
+				memmove(const_cast<tchar*>(hyperlinkPos), hyperlinkEnd, (sessionInfo.data + sessionInfo.count - hyperlinkEnd)*sizeof(tchar));
+				sessionInfo.Resize(sessionInfo.count - (hyperlinkEnd - hyperlinkPos));
+			}
 
 			StringBuffer<> fullName;
 			fullName.Append(sessionName).Append(L" (").Append(sessionInfo).Append(L")");
@@ -422,6 +441,7 @@ namespace uba
 			TraceView::Session& session = GetSession(out, sessionIndex);
 			session.name = sessionName.data;
 			session.fullName = fullName.data;
+			session.hyperlink = hyperlink;
 			session.clientUid = clientUid;
 
 			++out.activeSessionCount;
@@ -1104,7 +1124,7 @@ namespace uba
 		{
 			out.remoteExecutionDisabled = true;
 			if (!out.sessions.empty())
-				out.sessions[0].notification = TC("(Remote execution disabled)");
+				out.sessions[0].notification = TC("(Remote scheduling finished)");
 			break;
 		}
 		case TraceType_String:
diff --git a/Engine/Source/Programs/UnrealBuildAccelerator/Visualizer/Private/UbaTraceReader.h b/Engine/Source/Programs/UnrealBuildAccelerator/Visualizer/Private/UbaTraceReader.h
--- a/Engine/Source/Programs/UnrealBuildAccelerator/Visualizer/Private/UbaTraceReader.h
+++ b/Engine/Source/Programs/UnrealBuildAccelerator/Visualizer/Private/UbaTraceReader.h
@@ -96,6 +96,7 @@ namespace uba
 		{
 			TString name;
 			TString fullName;
+			TString hyperlink;
 			Guid clientUid;
 			Vector<Processor> processors;
 			Vector<SessionUpdate> updates;
@@ -108,6 +109,7 @@ namespace uba
 			u64 fetchedFilesBytes = 0;
 			u64 storedFilesBytes = 0;
 			u32 maxVisibleFiles = 0;
+			u32 fullNameWidth = 0;
 
 			float highestSendPerS = 0;
 			float highestRecvPerS = 0;
diff --git a/Engine/Source/Programs/UnrealBuildAccelerator/Visualizer/Private/UbaVisualizer.cpp b/Engine/Source/Programs/UnrealBuildAccelerator/Visualizer/Private/UbaVisualizer.cpp
--- a/Engine/Source/Programs/UnrealBuildAccelerator/Visualizer/Private/UbaVisualizer.cpp
+++ b/Engine/Source/Programs/UnrealBuildAccelerator/Visualizer/Private/UbaVisualizer.cpp
@@ -664,7 +664,8 @@ namespace uba
 				{
 					destroyWin();
 					m_looping = false;
-					m_listenTimeout.Set();
+					if (m_listenTimeout.IsCreated())
+						m_listenTimeout.Set();
 					break;
 				}
 			}
@@ -1235,6 +1236,10 @@ namespace uba
 
 					if (selected)
 						SetBkMode(hdc, OPAQUE);
+
+					SIZE textSize;
+					if (GetTextExtentPoint32W(hdc, text.data, text.count, &textSize))
+						session.fullNameWidth = u32(textSize.cx);
 					ExtTextOutW(hdc, 5, posY+2, ETO_CLIPPED, &rect, text.data, text.count, NULL);
 					if (selected)
 						SetBkMode(hdc, TRANSPARENT);
@@ -2098,7 +2103,6 @@ namespace uba
 		{
 			SetActiveFont(m_popupFont);
 			DrawTextLogger logger(m_hwnd, hdc, m_popupFont.height, m_tooltipBackgroundBrush);
-			logger.AddTextOffset(-10); // Remove spaces in the front
 			logger.AddWidth(3);
 			logger.AddSpace(2);
 
@@ -2106,10 +2110,13 @@ namespace uba
 			if (summary.empty())
 			{
 				if (m_traceView.finished)
-					logger.Info(L" Session summary not available on this trace version");
+					logger.Info(L" Session summary was never received");
 				else
 					logger.Info(L" Session summary not available until session is done");
 			}
+			else
+				logger.AddTextOffset(-10); // Remove spaces in the front
+
 			for (auto& line : summary)
 				logger.Log(LogEntryType_Info, line);
 			logger.AddSpace(3);
@@ -2753,7 +2760,7 @@ namespace uba
 			{
 				if (pos.y >= posY && pos.y < posY + m_sessionStepY)
 				{
-					if (pos.x < 500)
+					if (pos.x < int(session.fullNameWidth) + 5)
 					{
 						outResult.sessionSelectedIndex = sessionIndex;
 						return;
@@ -3300,7 +3307,8 @@ namespace uba
 				else if (m_usingNamed)
 				{
 					if (!m_trace.UpdateReadNamed(m_traceView, m_replay ? timeOffset : ~u64(0), changed))
-						m_listenTimeout.Set();
+						if (m_listenTimeout.IsCreated())
+							m_listenTimeout.Set();
 				}
 			}
 
@@ -3614,6 +3622,10 @@ namespace uba
 			{
 				ShellExecuteW(NULL, L"open", m_hyperLinkSelected.c_str(), NULL, NULL, SW_SHOW);
 			}
+			else if (m_sessionSelectedIndex != ~0u && !m_traceView.sessions[m_sessionSelectedIndex].hyperlink.empty())
+			{
+				ShellExecuteW(NULL, L"open", m_traceView.sessions[m_sessionSelectedIndex].hyperlink.c_str(), NULL, NULL, SW_SHOW);
+			}
 			else
 			{
 				POINTS p = MAKEPOINTS(lParam);
@@ -3626,7 +3638,7 @@ namespace uba
 		{
 			static HCURSOR arrow = LoadCursorW(NULL, IDC_ARROW);
 			static HCURSOR hand = LoadCursorW(NULL, IDC_HAND);
-			if (!m_hyperLinkSelected.empty())
+			if (!m_hyperLinkSelected.empty() || (m_sessionSelectedIndex != ~0u && !m_traceView.sessions[m_sessionSelectedIndex].hyperlink.empty()))
 				SetCursor(hand);
 			else
 				SetCursor(arrow);
