[UBA]
* Fixed so .isRunning file is deleted when cache server shuts down

diff --git a/Engine/Source/Programs/UnrealBuildAccelerator/CacheService/Private/UbaCacheService.cpp b/Engine/Source/Programs/UnrealBuildAccelerator/CacheService/Private/UbaCacheService.cpp
--- a/Engine/Source/Programs/UnrealBuildAccelerator/CacheService/Private/UbaCacheService.cpp
+++ b/Engine/Source/Programs/UnrealBuildAccelerator/CacheService/Private/UbaCacheService.cpp
@@ -347,6 +347,9 @@ namespace uba
 		}
 
 		cacheServer.RunMaintenance(false, {});
+
+		storageServer.DeleteIsRunningFile();
+
 		return 0;
 	}
 }
diff --git a/Engine/Source/Programs/UnrealBuildAccelerator/Common/Private/UbaStorage.cpp b/Engine/Source/Programs/UnrealBuildAccelerator/Common/Private/UbaStorage.cpp
--- a/Engine/Source/Programs/UnrealBuildAccelerator/Common/Private/UbaStorage.cpp
+++ b/Engine/Source/Programs/UnrealBuildAccelerator/Common/Private/UbaStorage.cpp
@@ -1201,6 +1201,18 @@ namespace uba
 			m_logger.Info(TC("Evicted %s from storage (%s). Estimated new storage is now %s (there might be files db is not aware of)"), BytesToText(before - after).str, TimeToText(GetTime() - startTime).str, BytesToText(after).str);
 	}
 
+	bool StorageImpl::DeleteIsRunningFile()
+	{
+		StringBuffer<256> isRunningName;
+		isRunningName.Append(m_rootDir).Append(TC(".isRunning"));
+		if (DeleteFileW(isRunningName.data))
+			return true;
+		u32 lastError = GetLastError();
+		if (lastError == ERROR_FILE_NOT_FOUND || lastError == ERROR_PATH_NOT_FOUND)
+			return true;
+		return m_logger.Warning(TC("Failed to delete %s (%s)"), isRunningName.data, LastErrorToText(lastError).data);
+	}
+
 	#if UBA_USE_MIMALLOC
 	void* Oodle_MallocAligned(OO_SINTa bytes, OO_S32 alignment)
 	{
@@ -1556,12 +1568,7 @@ namespace uba
 			return m_logger.Error(TC("Can't move file from %s to %s (%s)"), tempFileName.data, fileName.data, LastErrorToText().data);
 
 		if (deleteIsRunningfile)
-		{
-			tchar isRunningName[256];
-			TStrcpy_s(isRunningName, sizeof_array(isRunningName), m_rootDir.data);
-			TStrcat_s(isRunningName, sizeof_array(isRunningName), TC(".isRunning"));
-			DeleteFileW(isRunningName);
-		}
+			DeleteIsRunningFile();
 
 		if (m_overflowReported)
 			m_logger.Info(TC("Session needs at least %s to not overflow."), BytesToText(m_casMaxBytes).str);
diff --git a/Engine/Source/Programs/UnrealBuildAccelerator/Common/Public/UbaStorage.h b/Engine/Source/Programs/UnrealBuildAccelerator/Common/Public/UbaStorage.h
--- a/Engine/Source/Programs/UnrealBuildAccelerator/Common/Public/UbaStorage.h
+++ b/Engine/Source/Programs/UnrealBuildAccelerator/Common/Public/UbaStorage.h
@@ -184,6 +184,8 @@ namespace uba
 		bool CheckAllCasFiles(u64 checkContentOfFilesNewerThanTime = ~u64(0));
 		void HandleOverflow(UnorderedSet<CasKey>* outDeletedFiles);
 
+		bool DeleteIsRunningFile();
+
 		struct FileEntry;
 		FileEntry& GetOrCreateFileEntry(const StringKey& fileNameKey);
 		CasEntry& GetOrCreateCasEntry(const CasKey& casKey);
