[UBA]
* Added Scheduler IsEmpty that should be used to know if scheduler can be shut down.. don't use GetStats since it is not thread safe and just a function that can be used to visualize data

diff --git a/Engine/Source/Programs/UnrealBuildAccelerator/Common/Private/UbaScheduler.cpp b/Engine/Source/Programs/UnrealBuildAccelerator/Common/Private/UbaScheduler.cpp
--- a/Engine/Source/Programs/UnrealBuildAccelerator/Common/Private/UbaScheduler.cpp
+++ b/Engine/Source/Programs/UnrealBuildAccelerator/Common/Private/UbaScheduler.cpp
@@ -212,6 +212,12 @@ namespace uba
 		outQueued = m_queuedProcesses;
 	}
 
+	bool Scheduler::IsEmpty()
+	{
+		SCOPED_READ_LOCK(m_processEntriesLock, lock);
+		return m_processEntries.size() == m_finishedProcesses;
+	}
+
 	void Scheduler::SetProcessFinishedCallback(const Function<void(const ProcessHandle&)>& processFinished)
 	{
 		m_processFinished = processFinished;
@@ -617,9 +623,9 @@ namespace uba
 
 	void Scheduler::FinishProcess(const ProcessHandle& handle)
 	{
-		++m_finishedProcesses;
 		if (m_processFinished)
 			m_processFinished(handle);
+		++m_finishedProcesses;
 		m_session.UpdateProgress(m_queuedProcesses + m_activeLocalProcesses + m_activeRemoteProcesses + m_finishedProcesses, m_finishedProcesses, 0);
 	}
 
diff --git a/Engine/Source/Programs/UnrealBuildAccelerator/Common/Public/UbaScheduler.h b/Engine/Source/Programs/UnrealBuildAccelerator/Common/Public/UbaScheduler.h
--- a/Engine/Source/Programs/UnrealBuildAccelerator/Common/Public/UbaScheduler.h
+++ b/Engine/Source/Programs/UnrealBuildAccelerator/Common/Public/UbaScheduler.h
@@ -63,7 +63,9 @@ namespace uba
 
 		u32 EnqueueProcess(const EnqueueProcessInfo& info); // Returns index of process. Index is a rolling number
 
-		void GetStats(u32& outQueued, u32& outActiveLocal, u32& outActiveRemote, u32& outFinished);
+		void GetStats(u32& outQueued, u32& outActiveLocal, u32& outActiveRemote, u32& outFinished); // This is not "threadsafe".. which means that you are not guaranteed to get the exact correct state.
+
+		bool IsEmpty(); // Returns true if scheduler is entirely empty.. as in no processes are left in the system
 
 		bool EnqueueFromFile(const tchar* yamlFilename, const Function<void(EnqueueProcessInfo&)>& enqueued = {}); // Enqueue actions from file. Example of format of file is at the end of this file
 
diff --git a/Engine/Source/Programs/UnrealBuildAccelerator/Host/Private/UbaExports.cpp b/Engine/Source/Programs/UnrealBuildAccelerator/Host/Private/UbaExports.cpp
--- a/Engine/Source/Programs/UnrealBuildAccelerator/Host/Private/UbaExports.cpp
+++ b/Engine/Source/Programs/UnrealBuildAccelerator/Host/Private/UbaExports.cpp
@@ -411,7 +411,7 @@ extern "C"
 		info->allowWaitOnMem = allowWaitOnMem;
 		info->allowKillOnMem = allowKillOnMem;
 		info->storeObjFilesCompressed = storeObjFilesCompressed;
-		//info->remoteTraceEnabled = true;
+		info->remoteTraceEnabled = true;
 		info->remoteLogEnabled = true;
 		return info;
 	}
@@ -683,6 +683,11 @@ extern "C"
 		scheduler->GetStats(outQueued, outActiveLocal, outActiveRemote, outFinished);
 	}
 
+	bool Scheduler_IsEmpty(uba::Scheduler* scheduler)
+	{
+		return scheduler->IsEmpty();
+	}
+
 	uba::CacheClient* CacheClient_Create(uba::SessionServer* session, bool reportMissReason, const uba::tchar* crypto, const uba::tchar* hint)
 	{
 		using namespace uba;
diff --git a/Engine/Source/Programs/UnrealBuildAccelerator/Host/Public/UbaExports.h b/Engine/Source/Programs/UnrealBuildAccelerator/Host/Public/UbaExports.h
--- a/Engine/Source/Programs/UnrealBuildAccelerator/Host/Public/UbaExports.h
+++ b/Engine/Source/Programs/UnrealBuildAccelerator/Host/Public/UbaExports.h
@@ -141,6 +141,7 @@ extern "C"
 	UBA_API void Scheduler_Stop(uba::Scheduler* scheduler);
 	UBA_API void Scheduler_Destroy(uba::Scheduler* scheduler);
 	UBA_API void Scheduler_GetStats(uba::Scheduler* scheduler, uba::u32& outQueued, uba::u32& outActiveLocal, uba::u32& outActiveRemote, uba::u32& outFinished);
+	UBA_API bool Scheduler_IsEmpty(uba::Scheduler* scheduler);
 
 	// Cache
 	UBA_API uba::CacheClient* CacheClient_Create(uba::SessionServer* session, bool reportMissReason = false, const uba::tchar* crypto = nullptr, const uba::tchar* hint = nullptr);
