[UBA]
* Changed proxy code to be more aggressive towards disconnecting clients when failing to retrieve files from host

diff --git a/Engine/Source/Programs/UnrealBuildAccelerator/Common/Private/UbaNetworkServer.cpp b/Engine/Source/Programs/UnrealBuildAccelerator/Common/Private/UbaNetworkServer.cpp
--- a/Engine/Source/Programs/UnrealBuildAccelerator/Common/Private/UbaNetworkServer.cpp
+++ b/Engine/Source/Programs/UnrealBuildAccelerator/Common/Private/UbaNetworkServer.cpp
@@ -1012,7 +1012,8 @@ namespace uba
 		{
 			if (it.m_id != info.connectionId)
 				continue;
-			found = &it;
+			if (!it.m_disconnected)
+				found = &it;
 			break;
 		}
 		if (!found)
diff --git a/Engine/Source/Programs/UnrealBuildAccelerator/Common/Private/UbaStorageProxy.cpp b/Engine/Source/Programs/UnrealBuildAccelerator/Common/Private/UbaStorageProxy.cpp
--- a/Engine/Source/Programs/UnrealBuildAccelerator/Common/Private/UbaStorageProxy.cpp
+++ b/Engine/Source/Programs/UnrealBuildAccelerator/Common/Private/UbaStorageProxy.cpp
@@ -229,7 +229,7 @@ namespace uba
 				delete file.beginMessage;
 				file.beginMessage = nullptr;
 				file.error = true;
-				return SendFetchBeginResponse(clientId, file, writer);
+				return false;
 			}
 		}
 
@@ -248,23 +248,29 @@ namespace uba
 		auto sendResponses = MakeGuard([&]()
 			{
 				SCOPED_FUTEX(file.lock, fileLock);
-				file.error = error;
 				file.memory = memory;
 				auto f = file.beginMessage;
 				file.beginMessage = nullptr;
+				UBA_ASSERT(f);
+				file.error = error;
 				fileLock.Leave();
 
-				if (!f)
-					return;
-
-				StackBinaryWriter<SendMaxSize> writer;
-				bool isFirst = true;
-				for (auto& deferredResponse : f->deferredResponses)
+				if (file.error)
 				{
-					writer.Reset();
-					SendFetchBeginResponse(deferredResponse.clientId, file, writer, isFirst);
-					m_server.SendResponse(deferredResponse.info, writer.GetData(), u32(writer.GetPosition()));
-					isFirst = false;
+					for (auto& deferredResponse : f->deferredResponses)
+						m_server.SendResponse(deferredResponse.info, nullptr, 0);
+				}
+				else
+				{
+					StackBinaryWriter<SendMaxSize> writer;
+					bool isFirst = true;
+					for (auto& deferredResponse : f->deferredResponses)
+					{
+						writer.Reset();
+						SendFetchBeginResponse(deferredResponse.clientId, file, writer, isFirst);
+						m_server.SendResponse(deferredResponse.info, writer.GetData(), u32(writer.GetPosition()));
+						isFirst = false;
+					}
 				}
 				delete f;
 			});
