[UBA]
* Added error message when CalculateCasKey fails because can't open file

diff --git a/Engine/Source/Programs/UnrealBuildAccelerator/Common/Private/UbaStorage.cpp b/Engine/Source/Programs/UnrealBuildAccelerator/Common/Private/UbaStorage.cpp
--- a/Engine/Source/Programs/UnrealBuildAccelerator/Common/Private/UbaStorage.cpp
+++ b/Engine/Source/Programs/UnrealBuildAccelerator/Common/Private/UbaStorage.cpp
@@ -2295,7 +2295,8 @@ namespace uba
 	{
 		FileHandle fileHandle;
 		if (!OpenFileSequentialRead(m_logger, fileName, fileHandle))
-			return false;
+			return m_logger.Error(TC("[CalculateCasKey] OpenFileSequentialRead failed for %s (%s)"), fileName, LastErrorToText().data);
+
 		auto fileGuard = MakeGuard([&](){ CloseFile(fileName, fileHandle); });
 
 		u64 fileSize;
