[UBA]
* Disabled remote trace by default (was mistakenly submitted)

diff --git a/Engine/Source/Programs/UnrealBuildAccelerator/Host/Private/UbaExports.cpp b/Engine/Source/Programs/UnrealBuildAccelerator/Host/Private/UbaExports.cpp
--- a/Engine/Source/Programs/UnrealBuildAccelerator/Host/Private/UbaExports.cpp
+++ b/Engine/Source/Programs/UnrealBuildAccelerator/Host/Private/UbaExports.cpp
@@ -411,7 +411,7 @@ extern "C"
 		info->allowWaitOnMem = allowWaitOnMem;
 		info->allowKillOnMem = allowKillOnMem;
 		info->storeObjFilesCompressed = storeObjFilesCompressed;
-		info->remoteTraceEnabled = true;
+		//info->remoteTraceEnabled = true;
 		info->remoteLogEnabled = true;
 		return info;
 	}
