[UBA]
* Fixed bug where storage proxy could connect and briefly send messages to the wrong host. Connection was not disconnected after validation so other threads could end up using the connection
* Added more validation around uid of server
* Fixed crash if remote trace file was enabled
* Added code to test network connection where client listen and server connects

diff --git a/Engine/Source/Programs/UnrealBuildAccelerator/Common/Private/UbaNetworkClient.cpp b/Engine/Source/Programs/UnrealBuildAccelerator/Common/Private/UbaNetworkClient.cpp
--- a/Engine/Source/Programs/UnrealBuildAccelerator/Common/Private/UbaNetworkClient.cpp
+++ b/Engine/Source/Programs/UnrealBuildAccelerator/Common/Private/UbaNetworkClient.cpp
@@ -106,13 +106,16 @@ namespace uba
 				rc.error = *headerData;
 				Guid serverUid = *(Guid*)(headerData+1);
 
+				if (serverUid == Guid())
+					rc.error = 5;
+
 				if (!rc.error)
 				{
 					SCOPED_FUTEX(rc.client.m_serverUidLock, lock);
 					if (rc.client.m_serverUid == Guid())
 						rc.client.m_serverUid = serverUid;
 					else if (rc.client.m_serverUid != serverUid) // Seems like two different servers tried to connect to this client.. keep the first one and ignore the others
-						rc.error = 5;
+						rc.error = 6;
 				}
 
 				if (!rc.error)
@@ -183,7 +186,10 @@ namespace uba
 		}
 
 		if (rc.error == 5)
-			return m_logger.Warning(TC("A connection from a server with different uid was requested. Ignore"));
+			return m_logger.Error(TC("A connection to a server with uid zero was requested."));
+
+		if (rc.error == 6)
+			return m_logger.Warning(TC("A connection to a server with different uid was requested. Ignore"));
 
 		if (m_connectionCount.fetch_add(1) != 0)
 			return true;
diff --git a/Engine/Source/Programs/UnrealBuildAccelerator/Common/Private/UbaNetworkServer.cpp b/Engine/Source/Programs/UnrealBuildAccelerator/Common/Private/UbaNetworkServer.cpp
--- a/Engine/Source/Programs/UnrealBuildAccelerator/Common/Private/UbaNetworkServer.cpp
+++ b/Engine/Source/Programs/UnrealBuildAccelerator/Common/Private/UbaNetworkServer.cpp
@@ -763,6 +763,19 @@ namespace uba
 		return success;
 	}
 
+	bool NetworkServer::HasConnectInProgress()
+	{
+		SCOPED_FUTEX(m_addConnectionsLock, lock);
+		for (auto it = m_addConnections.begin(); it != m_addConnections.end();)
+		{
+			if (it->Wait(0))
+				it = m_addConnections.erase(it);
+			else
+				++it;
+		}
+		return !m_addConnections.empty();
+	}
+
 	void NetworkServer::PrintSummary(Logger& logger)
 	{
 		if (!m_maxActiveConnections)
diff --git a/Engine/Source/Programs/UnrealBuildAccelerator/Common/Private/UbaSessionClient.cpp b/Engine/Source/Programs/UnrealBuildAccelerator/Common/Private/UbaSessionClient.cpp
--- a/Engine/Source/Programs/UnrealBuildAccelerator/Common/Private/UbaSessionClient.cpp
+++ b/Engine/Source/Programs/UnrealBuildAccelerator/Common/Private/UbaSessionClient.cpp
@@ -2050,7 +2050,8 @@ namespace uba
 
 		if (m_shouldSendTraceToServer)
 		{
-			m_client.SetWorkTracker(nullptr);
+			// Can't disable this, it can cause race conditions
+			//m_client.SetWorkTracker(nullptr);
 
 			StopTraceThread();
 
diff --git a/Engine/Source/Programs/UnrealBuildAccelerator/Common/Private/UbaStorageClient.cpp b/Engine/Source/Programs/UnrealBuildAccelerator/Common/Private/UbaStorageClient.cpp
--- a/Engine/Source/Programs/UnrealBuildAccelerator/Common/Private/UbaStorageClient.cpp
+++ b/Engine/Source/Programs/UnrealBuildAccelerator/Common/Private/UbaStorageClient.cpp
@@ -404,7 +404,7 @@ namespace uba
 					proxy = new ProxyClient(ctorSuccess, ncci);
 					m_proxyClient = proxy;
 
-					auto disallowProxy = MakeGuard([&]() { proxy = nullptr; allowProxy = false; });
+					auto disallowProxy = MakeGuard([&]() { m_proxyClient->client.Disconnect(); proxy = nullptr; allowProxy = false; });
 
 					if (!ctorSuccess)
 						continue;
diff --git a/Engine/Source/Programs/UnrealBuildAccelerator/Common/Private/UbaTrace.cpp b/Engine/Source/Programs/UnrealBuildAccelerator/Common/Private/UbaTrace.cpp
--- a/Engine/Source/Programs/UnrealBuildAccelerator/Common/Private/UbaTrace.cpp
+++ b/Engine/Source/Programs/UnrealBuildAccelerator/Common/Private/UbaTrace.cpp
@@ -118,6 +118,8 @@ namespace uba
 
 		if (!writeFileName || !*writeFileName)
 			return true;
+
+		ScopedFutex lock(m_memoryLock);
 		FileAccessor traceFile(m_logger, writeFileName);
 		if (!traceFile.CreateWrite(false, DefaultAttributes(), 0, nullptr))
 			return false;
diff --git a/Engine/Source/Programs/UnrealBuildAccelerator/Common/Public/UbaNetworkServer.h b/Engine/Source/Programs/UnrealBuildAccelerator/Common/Public/UbaNetworkServer.h
--- a/Engine/Source/Programs/UnrealBuildAccelerator/Common/Public/UbaNetworkServer.h
+++ b/Engine/Source/Programs/UnrealBuildAccelerator/Common/Public/UbaNetworkServer.h
@@ -60,6 +60,7 @@ namespace uba
 		void SetClientsConfig(const Config& config);
 
 		bool AddClient(NetworkBackend& backend, const tchar* ip, u16 port = DefaultPort, const u8* cryptoKey128 = nullptr); // Adds a client that server will create one or more connections to (note this will return before we know if it was a success or not)
+		bool HasConnectInProgress();
 
 		void PrintSummary(Logger& logger);
 
diff --git a/Engine/Source/Programs/UnrealBuildAccelerator/Test/Private/UbaTest.h b/Engine/Source/Programs/UnrealBuildAccelerator/Test/Private/UbaTest.h
--- a/Engine/Source/Programs/UnrealBuildAccelerator/Test/Private/UbaTest.h
+++ b/Engine/Source/Programs/UnrealBuildAccelerator/Test/Private/UbaTest.h
@@ -19,7 +19,7 @@ namespace uba
 	using RunProcessFunction = Function<ProcessHandle(const ProcessStartInfo&)>;
 	using TestSessionFunction = Function<bool(LoggerWithWriter& logger, SessionServer& session, const tchar* workingDir, const RunProcessFunction& runProcess)>;
 	bool RunLocal(LoggerWithWriter& logger, const StringBufferBase& testRootDir, const TestSessionFunction& testFunc, bool enableDetour = true);
-	bool RunRemote(LoggerWithWriter& logger, const StringBufferBase& testRootDir, const TestSessionFunction& testFunc, bool deleteAll = true);
+	bool RunRemote(LoggerWithWriter& logger, const StringBufferBase& testRootDir, const TestSessionFunction& testFunc, bool deleteAll = true, bool serverShouldListen = true);
 	void GetTestAppPath(LoggerWithWriter& logger, StringBufferBase& out);
 	const tchar* GetSystemApplication();
 	const tchar* GetSystemArguments();
diff --git a/Engine/Source/Programs/UnrealBuildAccelerator/Test/Private/UbaTestSession.cpp b/Engine/Source/Programs/UnrealBuildAccelerator/Test/Private/UbaTestSession.cpp
--- a/Engine/Source/Programs/UnrealBuildAccelerator/Test/Private/UbaTestSession.cpp
+++ b/Engine/Source/Programs/UnrealBuildAccelerator/Test/Private/UbaTestSession.cpp
@@ -50,7 +50,7 @@ namespace uba
 		return testFunc(logger, session, workingDir.data, [&](const ProcessStartInfo& pi) { return session.RunProcess(pi, true, enableDetour); });
 	}
 
-	bool RunRemote(LoggerWithWriter& logger, const StringBufferBase& testRootDir, const TestSessionFunction& testFunc, bool deleteAll)
+	bool RunRemote(LoggerWithWriter& logger, const StringBufferBase& testRootDir, const TestSessionFunction& testFunc, bool deleteAll, bool serverShouldListen)
 	{
 		LogWriter& logWriter = logger.m_writer;
 		NetworkBackendTcp tcpBackend(logWriter);
@@ -88,10 +88,27 @@ namespace uba
 		server.SetClientsConfig(clientConfig);
 
 		u16 port = 1356;
-		if (!server.StartListen(tcpBackend, port))
-			return logger.Error(TC("Failed to listen"));
-		if (!client.Connect(tcpBackend, TC("127.0.0.1"), port))
-			return logger.Error(TC("Failed to connect"));
+
+		if (serverShouldListen)
+		{
+			if (!server.StartListen(tcpBackend, port))
+				return logger.Error(TC("Failed to listen"));
+			if (!client.Connect(tcpBackend, TC("127.0.0.1"), port))
+				return logger.Error(TC("Failed to connect"));
+			if (!client.Connect(tcpBackend, TC("127.0.0.1"), port))
+				return logger.Error(TC("Failed to connect"));
+		}
+		else
+		{
+			if (!client.StartListen(tcpBackend, port))
+				return logger.Error(TC("Failed to listen"));
+			if (!server.AddClient(tcpBackend, TC("127.0.0.1"), port))
+				return logger.Error(TC("Failed to connect"));
+			while (server.HasConnectInProgress())
+				Sleep(1);
+		}
+
+		auto disconnectGuard = MakeGuard([&]() { client.Disconnect(); });
 
 		Config config;
 		if (!client.FetchConfig(config))
@@ -117,7 +134,7 @@ namespace uba
 		auto& sessionClient = *new SessionClient(sessionClientInfo);
 		auto scg2 = MakeGuard([&]() { delete &sessionClient; });
 
-		auto cg = MakeGuard([&]() { sessionClient.Stop(); client.Disconnect(); });
+		auto cg = MakeGuard([&]() { sessionClient.Stop(); disconnectGuard.Execute(); });
 
 		StringBuffer<MaxPath> workingDir;
 		workingDir.Append(testRootDir).Append(TC("WorkingDir"));
