[UBA]
* Minor cache server optimization

diff --git a/Engine/Source/Programs/UnrealBuildAccelerator/Common/Private/UbaCacheServer.cpp b/Engine/Source/Programs/UnrealBuildAccelerator/Common/Private/UbaCacheServer.cpp
--- a/Engine/Source/Programs/UnrealBuildAccelerator/Common/Private/UbaCacheServer.cpp
+++ b/Engine/Source/Programs/UnrealBuildAccelerator/Common/Private/UbaCacheServer.cpp
@@ -453,7 +453,7 @@ namespace uba
 		FileAccessor file;
 	};
 
-	bool CacheServer::SaveBucket(u64 bucketId, Bucket& bucket)
+	bool CacheServer::SaveBucket(u64 bucketId, Bucket& bucket, Vector<u8>& temp)
 	{
 		u64 saveStart = GetTime();
 
@@ -483,8 +483,6 @@ namespace uba
 		u32 entryLookupCount = u32(bucket.m_cacheEntryLookup.size());
 		file.Write(entryLookupCount);
 
-		Vector<u8> temp;
-
 		for (auto& kv2 : bucket.m_cacheEntryLookup)
 		{
 			file.Write(kv2.first);
@@ -561,7 +559,7 @@ namespace uba
 				Bucket& bucket = it->second;
 				if (!bucket.needsSave)
 					return;
-				if (SaveBucket(it->first, bucket))
+				if (SaveBucket(it->first, bucket, temp))
 					bucket.needsSave = false;
 				else
 					success = false;
@@ -2417,8 +2415,8 @@ namespace uba
 
 		CasKey key;
 		bool deferCreation = false;
-		bool fileIsCompressed = false;
-		if (!m_storage.StoreCasFile(key, tempFile.data, CasKeyZero, deferCreation, fileIsCompressed))
+		bool checkIfCompressed = false;
+		if (!m_storage.StoreCasFile(key, tempFile.data, CasKeyZero, deferCreation, checkIfCompressed))
 			return false;
 
 		writer.WriteCasKey(key);
diff --git a/Engine/Source/Programs/UnrealBuildAccelerator/Common/Public/UbaCacheServer.h b/Engine/Source/Programs/UnrealBuildAccelerator/Common/Public/UbaCacheServer.h
--- a/Engine/Source/Programs/UnrealBuildAccelerator/Common/Public/UbaCacheServer.h
+++ b/Engine/Source/Programs/UnrealBuildAccelerator/Common/Public/UbaCacheServer.h
@@ -68,7 +68,7 @@ namespace uba
 
 		bool RunMaintenanceInternal(const Function<bool()>& shouldExit);
 		bool LoadBucket(Bucket& bucket, BinaryReader& reader, u32 databaseVersion, LoadStats& outStats);
-		bool SaveBucket(u64 bucketId, Bucket& bucket);
+		bool SaveBucket(u64 bucketId, Bucket& bucket, Vector<u8>& temp);
 		bool SaveNoLock();
 		void OnDisconnected(u32 clientId);
 
