[UBA] Fix null pointer exception in uba::GetKeyAndFixedName().

#rnx

diff --git a/Engine/Source/Programs/UnrealBuildAccelerator/Common/Private/UbaSession.cpp b/Engine/Source/Programs/UnrealBuildAccelerator/Common/Private/UbaSession.cpp
--- a/Engine/Source/Programs/UnrealBuildAccelerator/Common/Private/UbaSession.cpp
+++ b/Engine/Source/Programs/UnrealBuildAccelerator/Common/Private/UbaSession.cpp
@@ -1471,7 +1471,7 @@ namespace uba
 
 		StringKey dirKey;
 		StringBuffer<> dirNameForHash;
-		const tchar* baseFileName;
+		const tchar* baseFileName = nullptr;
 		GetDirKey(dirKey, dirNameForHash, baseFileName, outFixedFilePath);
 
 		if (CaseInsensitiveFs)
@@ -1479,12 +1479,16 @@ namespace uba
 
 		outHasher.Update(dirNameForHash);
 
-		StringBuffer<256> baseFileNameForHash;
-		baseFileNameForHash.Append(baseFileName);
-		if (CaseInsensitiveFs)
-			baseFileNameForHash.MakeLower();
+		if (baseFileName)
+		{
+			StringBuffer<256> baseFileNameForHash;
+			baseFileNameForHash.Append(baseFileName);
+			if (CaseInsensitiveFs)
+				baseFileNameForHash.MakeLower();
+
+			outHasher.Update(baseFileNameForHash);
+		}
 
-		outHasher.Update(baseFileNameForHash);
 		StringKey result = ToStringKey(outHasher);
 
 #if UBA_DEBUG
diff --git a/Engine/Source/Programs/UnrealBuildAccelerator/Common/Public/UbaBottleneck.h b/Engine/Source/Programs/UnrealBuildAccelerator/Common/Public/UbaBottleneck.h
--- a/Engine/Source/Programs/UnrealBuildAccelerator/Common/Public/UbaBottleneck.h
+++ b/Engine/Source/Programs/UnrealBuildAccelerator/Common/Public/UbaBottleneck.h
@@ -3,6 +3,7 @@
 #pragma once
 
 #include "UbaEvent.h"
+#include "UbaPlatform.h"
 #include "UbaSynchronization.h"
 
 namespace uba
