[UBA]
* Fixed bug where SearchPathW virtualized returned path when it was the system32 folder (should never be virtualized)

diff --git a/Engine/Source/Programs/UnrealBuildAccelerator/Detours/Private/Windows/UbaDetoursFunctionsKernelBase.inl b/Engine/Source/Programs/UnrealBuildAccelerator/Detours/Private/Windows/UbaDetoursFunctionsKernelBase.inl
--- a/Engine/Source/Programs/UnrealBuildAccelerator/Detours/Private/Windows/UbaDetoursFunctionsKernelBase.inl
+++ b/Engine/Source/Programs/UnrealBuildAccelerator/Detours/Private/Windows/UbaDetoursFunctionsKernelBase.inl
@@ -2667,7 +2667,7 @@ DWORD Detoured_SearchPathW(LPCWSTR lpPath, LPCWSTR lpFileName, LPCWSTR lpExtensi
 	}
 
 	DWORD res = True_SearchPathW(lpPath, lpFileName, lpExtension, nBufferLength, lpBuffer, lpFilePart);
-	if (res && IsVfsEnabled())
+	if (res && IsVfsEnabled() && !StartsWith(lpBuffer, g_systemRoot.data))
 	{
 		UBA_ASSERT(res < nBufferLength);
 		StringBuffer<> temp(lpBuffer);
