[UBA]
* Changed FillRect to PatBlt to fix artifacts happening over parsec/remote desktop

diff --git a/Engine/Source/Programs/UnrealBuildAccelerator/Visualizer/Private/UbaVisualizer.cpp b/Engine/Source/Programs/UnrealBuildAccelerator/Visualizer/Private/UbaVisualizer.cpp
--- a/Engine/Source/Programs/UnrealBuildAccelerator/Visualizer/Private/UbaVisualizer.cpp
+++ b/Engine/Source/Programs/UnrealBuildAccelerator/Visualizer/Private/UbaVisualizer.cpp
@@ -917,6 +917,17 @@ namespace uba
 		SetTextColor(hdc, m_textColor);
 		SetBkColor(hdc, m_config.DarkMode ? RGB(70, 70, 70) : RGB(180, 180, 180));
 
+		HBRUSH lastSelectedBrush = 0;
+		auto fillRect = [&](const RECT& r, HBRUSH b)
+			{
+				if (lastSelectedBrush != b)
+				{
+					SelectObject(hdc, b);
+					lastSelectedBrush = b;
+				}
+				PatBlt(hdc, r.left, r.top, r.right - r.left, r.bottom - r.top, PATCOPY);
+			};
+
 		auto DrawCenteredText = [&](const StringBuffer<>* lines, u32 lineCount, bool drawBackground)
 			{
 				for (u32 i=0; i!=lineCount; ++i)
@@ -931,9 +942,9 @@ namespace uba
 
 					if (drawBackground)
 					{
-						RECT fillRect = rect;
-						InflateRect(&fillRect, 4, 2);
-						FillRect(hdc, &fillRect, m_tooltipBackgroundBrush);
+						RECT r = rect;
+						InflateRect(&r, 4, 2);
+						fillRect(r, m_tooltipBackgroundBrush);
 					}
 					DrawTextW(hdc, lines[i].data, lines[i].count, &rect, DT_SINGLELINE|DT_NOCLIP);
 				}
@@ -981,7 +992,6 @@ namespace uba
 		HBITMAP nullBmp = CreateCompatibleBitmap(hdc, 1, 1);
 		HBITMAP oldBmp = (HBITMAP)SelectObject(textDC, nullBmp);
 		HBITMAP lastSelectedBitmap = 0;
-		//HBRUSH lastSelectedBrush = 0;
 
 		u64 lastStop = 0;
 
@@ -1022,10 +1032,10 @@ namespace uba
 			rect.right = rect.left + width;
 			rect.top = posY;
 			rect.bottom = posY + m_activeFont.height;
-			FillRect(hdc, &rect, m_processBrushes[0].inProgress);
+			fillRect(rect, m_processBrushes[0].inProgress);
 
 			rect.right = rect.left + int(progress*float(width));//durationMs/100;
-			FillRect(hdc, &rect, m_traceView.progressErrorCount ? m_processBrushes[0].error : m_processBrushes[0].success);
+			fillRect(rect, m_traceView.progressErrorCount ? m_processBrushes[0].error : m_processBrushes[0].success);
 
 			StringBuffer<> str;
 			str.Appendf(TC("%u%%    %u / %u"), u32(progress*100.0f), m_traceView.progressProcessesDone, m_traceView.progressProcessesTotal);
@@ -1101,7 +1111,7 @@ namespace uba
 					rect.right = right;
 					rect.top = posY;
 					rect.bottom = posY + height;
-					FillRect(hdc, &rect, inProgress ? m_processBrushes[selected].inProgress : m_processBrushes[selected].success);
+					fillRect(rect, inProgress ? m_processBrushes[selected].inProgress : m_processBrushes[selected].success);
 					return rect;
 				};
 
@@ -1422,7 +1432,7 @@ namespace uba
 								process.bitmapDirty = true;
 
 							--rect.top;
-							PaintProcessRect(process, hdc, rect, progressRect, selected, false);
+							PaintProcessRect(process, hdc, rect, progressRect, selected, false, lastSelectedBrush);
 							++rect.top;
 
 							int processWidth = rect.right - rect.left;
@@ -1455,7 +1465,7 @@ namespace uba
 									if (!done)
 										rect3.right = 256;
 
-									PaintProcessRect(process, textDC, rect3, rect2, selected, true);
+									PaintProcessRect(process, textDC, rect3, rect2, selected, true, lastSelectedBrush);
 
 									rect2.left += 3; // Move in text a bit
 									
@@ -1662,7 +1672,7 @@ namespace uba
 							auto clampRect = [&](RECT& r) { r.left = Min(Max(r.left, progressRect.left), progressRect.right); r.right = Max(Min(r.right, progressRect.right), progressRect.left); };
 
 							clampRect(rect);
-							FillRect(hdc, &rect, brush);
+							fillRect(rect, brush);
 							++rect.top;
 
 							int processWidth = rect.right - rect.left;
@@ -2287,7 +2297,7 @@ namespace uba
 		SetActiveFont(m_defaultFont);
 	}
 
-	void Visualizer::PaintProcessRect(TraceView::Process& process, HDC hdc, RECT rect, const RECT& progressRect, bool selected, bool writingBitmap)
+	void Visualizer::PaintProcessRect(TraceView::Process& process, HDC hdc, RECT rect, const RECT& progressRect, bool selected, bool writingBitmap, HBRUSH& lastSelectedBrush)
 	{
 		auto clampRect = [&](RECT& r) { r.left = Min(Max(r.left, progressRect.left), progressRect.right); r.right = Max(Min(r.right, progressRect.right), progressRect.left); };
 		bool done = process.stop != ~u64(0);
@@ -2304,12 +2314,22 @@ namespace uba
 
 		u64 writeFilesTime = process.writeFilesTime;
 
+		auto fillRect = [&](const RECT& r, HBRUSH b)
+			{
+				if (lastSelectedBrush != b)
+				{
+					SelectObject(hdc, b);
+					lastSelectedBrush = b;
+				}
+				PatBlt(hdc, r.left, r.top, r.right - r.left, r.bottom - r.top, PATCOPY);
+			};
+
 		if (!done || process.exitCode != 0 || !m_config.ShowReadWriteColors)
 		{
 			if (writingBitmap)
 				rect.right = 256;
 			clampRect(rect);
-			FillRect(hdc, &rect, brush);
+			fillRect(rect, brush);
 			return;
 		}
 
@@ -2326,7 +2346,7 @@ namespace uba
 			r2.right = r2.left + headSize;
 			clampRect(r2);
 			if (r2.left != r2.right)
-				FillRect(hdc, &r2, m_processBrushes[selected].recv);
+				fillRect(r2, m_processBrushes[selected].recv);
 		}
 
 		double sendPart = (double(ConvertTime(m_traceView, writeFilesTime)) / duration);
@@ -2337,12 +2357,12 @@ namespace uba
 			r2.left = r2.right - tailSize;
 			clampRect(r2);
 			if (r2.left != r2.right)
-				FillRect(hdc, &r2, m_processBrushes[selected].send);
+				fillRect(r2, m_processBrushes[selected].send);
 		}
 
 		clampRect(main);
 		if (main.left != main.right)
-			FillRect(hdc, &main, brush);
+			fillRect(main, brush);
 
 		//clampRect(rect);
 		/*
diff --git a/Engine/Source/Programs/UnrealBuildAccelerator/Visualizer/Private/UbaVisualizer.h b/Engine/Source/Programs/UnrealBuildAccelerator/Visualizer/Private/UbaVisualizer.h
--- a/Engine/Source/Programs/UnrealBuildAccelerator/Visualizer/Private/UbaVisualizer.h
+++ b/Engine/Source/Programs/UnrealBuildAccelerator/Visualizer/Private/UbaVisualizer.h
@@ -86,7 +86,7 @@ namespace uba
 		void PaintClient(const Function<void(HDC hdc, HDC memDC, RECT& clientRect)>& paintFunc);
 		void PaintAll(HDC hdc, const RECT& clientRect);
 		void PaintActiveProcesses(int& posY, const RECT& clientRect, const Function<void(TraceView::ProcessLocation&, u32, bool)>& drawProcess);
-		void PaintProcessRect(TraceView::Process& process, HDC hdc, RECT rect, const RECT& progressRect, bool selected, bool writingBitmap);
+		void PaintProcessRect(TraceView::Process& process, HDC hdc, RECT rect, const RECT& progressRect, bool selected, bool writingBitmap, HBRUSH& lastSelectedBrush);
 		void PaintTimeline(HDC hdc, const RECT& clientRect);
 		using DrawTextFunc = Function<void(const StringBufferBase& text, RECT& rect)>;
 		void PaintDetailedStats(int& posY, const RECT& progressRect, TraceView::Session& session, bool isRemote, u64 playTime, const DrawTextFunc& drawTextFunc);
