[UBA]
* Changed trace reader updates from arrays of structs to struct of arrays

diff --git a/Engine/Source/Programs/UnrealBuildAccelerator/Visualizer/Private/UbaTraceReader.cpp b/Engine/Source/Programs/UnrealBuildAccelerator/Visualizer/Private/UbaTraceReader.cpp
--- a/Engine/Source/Programs/UnrealBuildAccelerator/Visualizer/Private/UbaTraceReader.cpp
+++ b/Engine/Source/Programs/UnrealBuildAccelerator/Visualizer/Private/UbaTraceReader.cpp
@@ -504,14 +504,21 @@ namespace uba
 				session.prevRecv = 0;
 				session.memTotal = 0;
 				if (!session.updates.empty())
-					session.updates.push_back({ time, 0, 0, lastPing, memAvail, cpuLoad, connectionCount });
+				{
+					session.updates.push_back(time);
+					session.networkSend.push_back(0);
+					session.networkRecv.push_back(0);
+					session.ping.push_back(lastPing);
+					session.memAvail.push_back(memAvail);
+					session.cpuLoad.push_back(cpuLoad);
+					session.connectionCount.push_back(connectionCount);
+				}
 			}
 			else
 			{
-				auto& prevUpdate = session.updates.back();
-				session.prevSend = prevUpdate.send;
-				session.prevRecv = prevUpdate.recv;
-				session.prevUpdateTime = session.updates.back().time;
+				session.prevSend = session.networkSend.back();
+				session.prevRecv = session.networkRecv.back();
+				session.prevUpdateTime = session.updates.back();
 			}
 
 			if (totalSend < session.prevSend)
@@ -522,7 +529,13 @@ namespace uba
 			//session.lastPing = lastPing;
 			//session.memAvail = memAvail;
 			session.memTotal = memTotal;
-			session.updates.push_back({time, totalSend, totalRecv, lastPing, memAvail, cpuLoad, connectionCount});
+			session.updates.push_back(time);
+			session.networkSend.push_back(totalSend);
+			session.networkRecv.push_back(totalRecv);
+			session.ping.push_back(lastPing);
+			session.memAvail.push_back(memAvail);
+			session.cpuLoad.push_back(cpuLoad);
+			session.connectionCount.push_back(connectionCount);
 			
 			if (session.prevUpdateTime)
 			{
diff --git a/Engine/Source/Programs/UnrealBuildAccelerator/Visualizer/Private/UbaTraceReader.h b/Engine/Source/Programs/UnrealBuildAccelerator/Visualizer/Private/UbaTraceReader.h
--- a/Engine/Source/Programs/UnrealBuildAccelerator/Visualizer/Private/UbaTraceReader.h
+++ b/Engine/Source/Programs/UnrealBuildAccelerator/Visualizer/Private/UbaTraceReader.h
@@ -39,17 +39,6 @@ namespace uba
 			Vector<Process> processes;
 		};
 
-		struct SessionUpdate
-		{
-			u64 time;
-			u64 send;
-			u64 recv;
-			u64 ping;
-			u64 memAvail;
-			float cpuLoad;
-			u8 connectionCount;
-		};
-
 		struct WorkRecordLogEntry
 		{
 			u64 time = 0;
@@ -99,7 +88,14 @@ namespace uba
 			TString hyperlink;
 			Guid clientUid;
 			Vector<Processor> processors;
-			Vector<SessionUpdate> updates;
+			Vector<u64> updates;
+			Vector<u64> networkSend;
+			Vector<u64> networkRecv;
+			Vector<u64> ping;
+			Vector<u64> memAvail;
+			Vector<float> cpuLoad;
+			Vector<u16> connectionCount;
+
 			Vector<TString> summary;
 			UnorderedMap<CasKey, u32> fetchedFilesActive;
 			Vector<FileTransfer> fetchedFiles;
diff --git a/Engine/Source/Programs/UnrealBuildAccelerator/Visualizer/Private/UbaVisualizer.cpp b/Engine/Source/Programs/UnrealBuildAccelerator/Visualizer/Private/UbaVisualizer.cpp
--- a/Engine/Source/Programs/UnrealBuildAccelerator/Visualizer/Private/UbaVisualizer.cpp
+++ b/Engine/Source/Programs/UnrealBuildAccelerator/Visualizer/Private/UbaVisualizer.cpp
@@ -1184,10 +1184,10 @@ namespace uba
 		u32 visibleBoxes = 0;
 
 		TraceView::ProcessLocation processLocation { 0, 0, 0 };
-		for (u64 i = 0, e = m_traceView.sessions.size(); i != e; ++i)
+		for (u64 sessionIt = 0, sessionEnd = m_traceView.sessions.size(); sessionIt != sessionEnd; ++sessionIt)
 		{
-			bool isFirst = i == 0;
-			auto& session = *sortedSessions[i].session;
+			bool isFirst = sessionIt == 0;
+			auto& session = *sortedSessions[sessionIt].session;
 			bool hasUpdates = !session.updates.empty();
 			if (!isFirst)
 			{
@@ -1198,7 +1198,7 @@ namespace uba
 					continue;
 			}
 
-			processLocation.sessionIndex = sortedSessions[i].index;
+			processLocation.sessionIndex = sortedSessions[sessionIt].index;
 			if (!isFirst)
 				posY += 3;
 
@@ -1215,13 +1215,13 @@ namespace uba
 
 					if (hasUpdates && session.disconnectTime == ~u64(0))
 					{
-						u64 ping = session.updates.back().ping;
-						u64 memAvail = session.updates.back().memAvail;
-						float cpuLoad = session.updates.back().cpuLoad;
+						u64 ping = session.ping.back();
+						//u64 memAvail = session.updates.back().memAvail;
+						//float cpuLoad = session.updates.back().cpuLoad;
 
-						text.Appendf(L" - Cpu: %.1f%%", cpuLoad * 100.0f);
-						if (memAvail)
-							text.Appendf(L" Mem: %ls/%ls", BytesToText(session.memTotal - memAvail).str, BytesToText(session.memTotal).str);
+						//text.Appendf(L" - Cpu: %.1f%%", cpuLoad * 100.0f);
+						//if (memAvail)
+						//	text.Appendf(L" Mem: %ls/%ls", BytesToText(session.memTotal - memAvail).str, BytesToText(session.memTotal).str);
 						if (ping)
 							text.Appendf(L" Ping: %ls", TimeToText(ping, false, m_traceView.frequency).str);
 						if (!session.notification.empty())
@@ -1263,90 +1263,80 @@ namespace uba
 				if (posY + GraphHeight >= progressRect.top && posY + GraphHeight - 5 < progressRect.bottom)
 				{
 					int posX = int(m_scrollPosX) + progressRect.left;
-					bool isFirstUpdate = true;
-					u64 prevTime = 0;
-					u64 prevSend = 0;
-					u64 prevRecv = 0;
-					float prevCupLoad = 0;
 					int graphBaseY = posY + GraphHeight - 4;
-					int prevX = 0;
-					int prevSendY = 0;
-					int prevRecvY = 0;
-					int prevCpuY = 0;
-					int prevMemY = 0;
-					double sendScale = double(session.highestSendPerS) / (double(GraphHeight) - 2);
-					double recvScale = double(session.highestRecvPerS) / (double(GraphHeight) - 2);
-
-					for (auto& update : session.updates)
-					{
-						float cpuLoad = update.cpuLoad;
-						if (cpuLoad < 0 || cpuLoad > 1.0f)
-							cpuLoad = prevCupLoad;
-						else
-							prevCupLoad = cpuLoad;
-
-						auto updateSend = update.send;
-						auto updateRecv = update.recv;
-
-						int x = posX + int(float(TimeToS(update.time)) * scaleX);
-						int sendY = graphBaseY;
-						int recvY = graphBaseY;
-						int cpuY = graphBaseY - int(cpuLoad*(GraphHeight-2));
-						int memY = 0;
-						if (update.memAvail && session.memTotal)
-							memY = graphBaseY - int(double(session.memTotal - update.memAvail)*double(GraphHeight - 2)/double(session.memTotal));
-
-						double duration = TimeToS(update.time - prevTime);
-						if (update.time == 0)
-							isFirstUpdate = true;
-						else if (prevSend > updateSend || prevRecv > updateRecv)
-							isFirstUpdate = true;
-
-						double sendInvScaleY = duration * sendScale;
-						if (sendInvScaleY != 0 && prevSend)
-							sendY = graphBaseY - int(double(updateSend - prevSend) / sendInvScaleY);
-						if (prevSendY <= 0)
-							prevSendY = sendY;
-						double recvInvScaleY = duration * recvScale;
-						if (recvInvScaleY != 0 && prevRecv)
-							recvY = graphBaseY - int(double(updateRecv - prevRecv) / recvInvScaleY) - 1;
-						if (prevRecvY <= 0)
-							prevRecvY = recvY;
-
-						if (prevMemY <= 0)
-							prevMemY = memY;
-
-						if (!isFirstUpdate && x > clientRect.left && prevX <= clientRect.right)
+					double graphHeight = double(GraphHeight - 2);
+
+					auto drawGraph = [&](auto& values, auto maxValue, double scale, HPEN pen, bool accumulatingValue, int offsetY = 0)
 						{
-							if (m_config.showNetworkStats && updateSend != 0 && updateRecv != 0)
-							{
-								SelectObject(hdc, m_sendPen);
-								MoveToEx(hdc, prevX, prevSendY, NULL);
-								LineTo(hdc, x, sendY);
-								SelectObject(hdc, m_recvPen);
-								MoveToEx(hdc, prevX, prevRecvY, NULL);
-								LineTo(hdc, x, recvY);
-							}
-							if (m_config.showCpuMemStats)
+							SelectObject(hdc, pen);
+							bool isFirstUpdate = true;
+							bool isFirstDraw = true;
+							decltype(maxValue) prevValue = 0;
+							u64 prevTime = 0;
+							int prevX = 0;
+							int prevY = 0;
+							for (u64 i=0, e=session.updates.size(); i!=e; ++i)
 							{
-								SelectObject(hdc, m_cpuPen);
-								MoveToEx(hdc, prevX, prevCpuY, NULL);
-								LineTo(hdc, x, cpuY);
-								SelectObject(hdc, m_memPen);
-								MoveToEx(hdc, prevX, prevMemY, NULL);
-								LineTo(hdc, x, memY);
+								u64 updateTime = session.updates[i];
+								auto value = values[i];
+								int x = posX + int(float(TimeToS(updateTime)) * scaleX);
+								int y = graphBaseY;
+
+								double duration = TimeToS(updateTime - prevTime);
+								if (updateTime == 0)
+									isFirstUpdate = true;
+
+								if (accumulatingValue)
+								{
+									double invScaleY = duration * scale;
+									if (invScaleY != 0 && prevValue)
+										y -= int(double(value - prevValue) / invScaleY) + offsetY;
+								}
+								else
+								{
+									double div = 1.0f;
+									if (maxValue)
+										div = double(maxValue);
+									y -= int(double(maxValue - value*scale)*graphHeight/div) + offsetY;
+								}
+
+								if (prevValue <= 0)
+									prevY = y;
+
+								if (x > clientRect.left)
+								{
+									if (isFirstUpdate)
+										MoveToEx(hdc, x, x, NULL);
+									else
+									{
+										if (isFirstDraw)
+											MoveToEx(hdc, prevX, prevY, NULL);
+										LineTo(hdc, x, y);
+										isFirstDraw = false;
+									}
+								}
+								if (x > clientRect.right)
+									break;
+								isFirstUpdate = false;
+								prevX = x;
+								prevY = y;
+								prevValue = value;
+								prevTime = updateTime;
 							}
-						}
-						isFirstUpdate = false;
-						prevX = x;
-						prevSendY = sendY;
-						prevRecvY = recvY;
-						prevCpuY = cpuY;
-						prevMemY = memY;
-						prevTime = update.time;
-						prevSend = updateSend;
-						prevRecv = updateRecv;
+						};
+
+					if (m_config.showNetworkStats && session.highestSendPerS != 0 && session.highestRecvPerS != 0)
+					{
+						drawGraph(session.networkSend, 100000000000000ull, double(session.highestSendPerS) / graphHeight, m_sendPen, true);
+						drawGraph(session.networkRecv, 100000000000000ull, double(session.highestRecvPerS) / graphHeight, m_recvPen, true, 1);
+					}
+
+					if (m_config.showCpuMemStats)
+					{
+						drawGraph(session.cpuLoad, 0.0f, -1.0, m_cpuPen, false);
+						drawGraph(session.memAvail, session.memTotal, 1.0, m_memPen, false);
 					}
+
 				}
 				posY += GraphHeight;
 			}
@@ -2509,18 +2499,20 @@ namespace uba
 
 			if (!session.updates.empty())
 			{
-				auto& u = session.updates.back();
+				auto updateTime = session.updates.back();
+				auto updateSend = session.networkSend.back();
+				auto updateRecv = session.networkRecv.back();
 				u64 sendPerS = 0;
 				u64 recvPerS = 0;
-				float duration = TimeToS(u.time - session.prevUpdateTime);
+				float duration = TimeToS(updateTime - session.prevUpdateTime);
 				if (duration != 0)
 				{
-					sendPerS = u64(float(u.send - session.prevSend) / duration);
-					recvPerS = u64(float(u.recv - session.prevRecv) / duration);
+					sendPerS = u64(float(updateSend - session.prevSend) / duration);
+					recvPerS = u64(float(updateRecv - session.prevRecv) / duration);
 				}
-				drawText(L"ClientId: %u  TcpCount: %u", session.clientUid.data1, u.connectionCount);
-				drawText(L"Recv: %ls (%sps)", BytesToText(u.recv), BytesToText(recvPerS*8));
-				drawText(L"Send: %ls (%sps)", BytesToText(u.send), BytesToText(sendPerS*8));
+				drawText(L"ClientId: %u  TcpCount: %u", session.clientUid.data1, session.connectionCount.back());
+				drawText(L"Recv: %ls (%sps)", BytesToText(updateRecv), BytesToText(recvPerS*8));
+				drawText(L"Send: %ls (%sps)", BytesToText(updateSend), BytesToText(sendPerS*8));
 			}
 
 			if (session.disconnectTime == ~u64(0))
@@ -2587,19 +2579,21 @@ namespace uba
 
 			if (!session.updates.empty())
 			{
-				auto& u = session.updates.back();
-				if (u.send || u.recv)
+				auto updateTime = session.updates.back();
+				auto updateSend = session.networkSend.back();
+				auto updateRecv = session.networkRecv.back();
+				if (updateSend || updateRecv)
 				{
 					u64 sendPerS = 0;
 					u64 recvPerS = 0;
-					float duration = TimeToS(u.time - session.prevUpdateTime);
+					float duration = TimeToS(updateTime - session.prevUpdateTime);
 					if (duration != 0)
 					{
-						sendPerS = u64(float(u.send - session.prevSend) / duration);
-						recvPerS = u64(float(u.recv - session.prevRecv) / duration);
+						sendPerS = u64(float(updateSend - session.prevSend) / duration);
+						recvPerS = u64(float(updateRecv - session.prevRecv) / duration);
 					}
-					drawText(L"Recv: %ls (%sps)", BytesToText(u.recv), BytesToText(recvPerS));
-					drawText(L"Send: %ls (%sps)", BytesToText(u.send), BytesToText(sendPerS));
+					drawText(L"Recv: %ls (%sps)", BytesToText(updateRecv), BytesToText(recvPerS));
+					drawText(L"Send: %ls (%sps)", BytesToText(updateSend), BytesToText(sendPerS));
 				}
 			}
 		}
@@ -2759,10 +2753,10 @@ namespace uba
 		SessionRec sortedSessions[1024];
 		Populate(sortedSessions, m_traceView, m_config.SortActiveRemoteSessions);
 
-		for (u64 i = 0, e = m_traceView.sessions.size(); i != e; ++i)
+		for (u64 sessionIt = 0, sessionEnd = m_traceView.sessions.size(); sessionIt != sessionEnd; ++sessionIt)
 		{
-			bool isFirst = i == 0;
-			auto& session = *sortedSessions[i].session;
+			bool isFirst = sessionIt == 0;
+			auto& session = *sortedSessions[sessionIt].session;
 			bool hasUpdates = !session.updates.empty();
 			if (!isFirst)
 			{
@@ -2772,7 +2766,7 @@ namespace uba
 					continue;
 			}
 
-			u32 sessionIndex = sortedSessions[i].index;
+			u32 sessionIndex = sortedSessions[sessionIt].index;
 			if (!isFirst)
 				posY += 3;
 
@@ -2800,14 +2794,17 @@ namespace uba
 					u64 prevSend = 0;
 					u64 prevRecv = 0;
 					int prevX = 100000;
-					for (auto& update : session.updates)
+					for (u64 i=0, e=session.updates.size(); i!=e; ++i)
 					{
-						int x = posX + int(float(TimeToS(update.time)) * scaleX);
+						u64 updateTime = session.updates[i];
+						auto updateSend = session.networkSend[i];
+						auto updateRecv = session.networkRecv[i];
+						int x = posX + int(float(TimeToS(updateTime)) * scaleX);
 
-						if (prevSend > update.send || prevRecv > update.recv)
+						if (prevSend > updateSend || prevRecv > updateRecv)
 						{
-							prevSend = update.send;
-							prevRecv = update.recv;
+							prevSend = updateSend;
+							prevRecv = updateRecv;
 							prevX = x;
 							continue;
 						}
@@ -2815,21 +2812,21 @@ namespace uba
 						int hitOffset = (prevX - x)/2;
 						if (pos.x + hitOffset >= prevX && pos.x + hitOffset <= x)
 						{
-							double duration = TimeToS(update.time - prevTime);
-							outResult.stats.recvBytesPerSecond = u64(float(update.recv - prevRecv) / duration);
-							outResult.stats.sendBytesPerSecond = u64(float(update.send - prevSend) / duration);
-							outResult.stats.ping = update.ping;
-							outResult.stats.memAvail = update.memAvail;
-							outResult.stats.cpuLoad = update.cpuLoad;
+							double duration = TimeToS(updateTime - prevTime);
+							outResult.stats.recvBytesPerSecond = u64(float(updateRecv - prevRecv) / duration);
+							outResult.stats.sendBytesPerSecond = u64(float(updateSend - prevSend) / duration);
+							outResult.stats.ping = session.ping[i];
+							outResult.stats.memAvail = session.memAvail[i];
+							outResult.stats.cpuLoad = session.cpuLoad[i];
 							outResult.stats.memTotal = session.memTotal;
 							outResult.statsSelected = true;
 							return;
 						}
 
 						prevX = x;
-						prevTime = update.time;
-						prevSend = update.send;
-						prevRecv = update.recv;
+						prevTime = updateTime;
+						prevSend = updateSend;
+						prevRecv = updateRecv;
 					}
 					posY += GraphHeight;
 				}
@@ -2845,7 +2842,7 @@ namespace uba
 						if (pos.x >= rect.left && pos.x < rect.right && pos.y >= rect.top && pos.y < rect.bottom && text.StartsWith(TC("Fetched Files")))
 							outResult.fetchedFilesSelected = sessionIndex;
 					};
-				PaintDetailedStats(posY, progressRect, session, i != 0, playTime, drawText);
+				PaintDetailedStats(posY, progressRect, session, sessionIt != 0, playTime, drawText);
 			}
 
 			if (m_config.showProcessBars)
