[UBA]
*  Added more logging information when compressed file is corrupt while trying to create a memory map of it in the transient buffer

diff --git a/Engine/Source/Programs/UnrealBuildAccelerator/Cli/Private/UbaCli.cpp b/Engine/Source/Programs/UnrealBuildAccelerator/Cli/Private/UbaCli.cpp
--- a/Engine/Source/Programs/UnrealBuildAccelerator/Cli/Private/UbaCli.cpp
+++ b/Engine/Source/Programs/UnrealBuildAccelerator/Cli/Private/UbaCli.cpp
@@ -675,22 +675,24 @@ namespace uba
 			u64 fileSize = fa.GetSize();
 			u8* mem = fa.GetData();
 
+			if (fileSize < 16)
+				return logger.Error(TC("File %s is too small to be compressed. Requires at least 16 bytes"), testDecompress.c_str());
+
 			StorageCreateInfo storageInfo(g_rootDir.data, logWriter);
 			storageInfo.casCapacityBytes = 0;
 			storageInfo.storeCompressed = storeCompressed;
 			storageInfo.workManager = &workManager;
 			StorageImpl storage(storageInfo);
 
-			auto& h = *(CompressedObjFileHeader*)mem;
-			if (!h.IsValid())
-				return logger.Error(TC("File %s is not a compressed file"), testDecompress.c_str());
-
 			BinaryReader reader(mem, 0, fileSize);
-			reader.Skip(sizeof(CompressedObjFileHeader));
+
+			auto& h = *(CompressedObjFileHeader*)mem;
+			if (h.IsValid())
+				reader.Skip(sizeof(CompressedObjFileHeader));
 			u64 decompressedSize = reader.ReadU64();
 
 			if (testCompressOriginalSize && decompressedSize != testCompressOriginalSize)
-				return logger.Error(TC("Compressed file %s has wrong decompressed size."), testDecompress.c_str());
+				return logger.Error(TC("Compressed file %s has wrong decompressed size. (Is it compressed?)"), testDecompress.c_str());
 
 			StringBuffer<> dest;
 			dest.Append(storage.GetTempPath()).Append(TC("TestDecompress.tmp"));
diff --git a/Engine/Source/Programs/UnrealBuildAccelerator/Common/Private/UbaFileMapping.cpp b/Engine/Source/Programs/UnrealBuildAccelerator/Common/Private/UbaFileMapping.cpp
--- a/Engine/Source/Programs/UnrealBuildAccelerator/Common/Private/UbaFileMapping.cpp
+++ b/Engine/Source/Programs/UnrealBuildAccelerator/Common/Private/UbaFileMapping.cpp
@@ -399,6 +399,11 @@ namespace uba
 		CloseMappingStorage(m_storage[MappedView_Persistent]);
 	}
 
+	u64 FileMappingBuffer::GetFileMappingCapacity()
+	{
+		return (IsWindows ? 32ull : 8ull) * 1024 * 1024 * 1024; // Linux can't have larger than 8gb
+	}
+
 	bool FileMappingBuffer::AddTransient(const tchar* name)
 	{
 		MappingStorage& storage = m_storage[MappedView_Transient];
@@ -532,7 +537,7 @@ namespace uba
 				if (!file->handle.IsValid())
 				{
 					UBA_ASSERT(&storage == &m_storage[MappedView_Transient]);
-					u64 capacity = (IsWindows ? 32ull : 8ull) * 1024 * 1024 * 1024; // Linux can't have larger than 8gb
+					u64 capacity = GetFileMappingCapacity();
 					file-> handle = uba::CreateMemoryMappingW(m_logger, PAGE_READWRITE | SEC_RESERVE, capacity, nullptr, TC("FileMappingBuffer"));
 					if (!file->handle.IsValid())
 					{
diff --git a/Engine/Source/Programs/UnrealBuildAccelerator/Common/Private/UbaSession.cpp b/Engine/Source/Programs/UnrealBuildAccelerator/Common/Private/UbaSession.cpp
--- a/Engine/Source/Programs/UnrealBuildAccelerator/Common/Private/UbaSession.cpp
+++ b/Engine/Source/Programs/UnrealBuildAccelerator/Common/Private/UbaSession.cpp
@@ -528,7 +528,7 @@ namespace uba
 		return newSize;
 	}
 
-	bool Session::CreateMemoryMapFromFile(MemoryMap& out, StringKey fileNameKey, const tchar* fileName, bool isCompressed, u64 alignment)
+	bool Session::CreateMemoryMapFromFile(MemoryMap& out, StringKey fileNameKey, const tchar* fileName, bool isCompressed, u64 alignment, const tchar* hint)
 	{
 		TimerScope ts(Stats().waitMmapFromFile);
 
@@ -607,6 +607,9 @@ namespace uba
 		}
 		else
 		{
+			if (size > m_fileMappingBuffer.GetFileMappingCapacity())
+				return m_logger.Error(TC("File %s has a size (%llu) that is too large to fit in mapping buffer (%s)"), fileName, size, hint);
+
 			auto mappedView = m_fileMappingBuffer.AllocAndMapView(MappedView_Transient, size, alignment, fileName, false);
 			if (!mappedView.memory)
 				return false;
@@ -699,6 +702,9 @@ namespace uba
 			}
 			else
 			{
+				if (size > m_fileMappingBuffer.GetFileMappingCapacity())
+					return m_logger.Error(TC("File %s has a size (%llu) that is too large to fit in mapping buffer (CreateMemoryMapFromView)"), fileName, size);
+
 				auto mappedViewWrite = m_fileMappingBuffer.AllocAndMapView(MappedView_Transient, size, alignment, fileName);
 				if (!mappedViewWrite.memory)
 					return false;
@@ -2244,7 +2250,7 @@ namespace uba
 				if (u64 alignment = GetMemoryMapAlignment(fileName))
 				{
 					MemoryMap map;
-					if (CreateMemoryMapFromFile(map, fileNameKey, fileName.data, false, alignment))
+					if (CreateMemoryMapFromFile(map, fileNameKey, fileName.data, false, alignment, TC("CreateFile")))
 					{
 						out.size = map.size;
 						out.fileName.Append(map.name);
diff --git a/Engine/Source/Programs/UnrealBuildAccelerator/Common/Private/UbaSessionClient.cpp b/Engine/Source/Programs/UnrealBuildAccelerator/Common/Private/UbaSessionClient.cpp
--- a/Engine/Source/Programs/UnrealBuildAccelerator/Common/Private/UbaSessionClient.cpp
+++ b/Engine/Source/Programs/UnrealBuildAccelerator/Common/Private/UbaSessionClient.cpp
@@ -629,7 +629,7 @@ namespace uba
 						retryEntryGuard.Execute();
 
 					MemoryMap map;
-					if (!CreateMemoryMapFromFile(map, fileNameKey, newName.data, IsCompressed(newCasKey), memoryMapAlignment))
+					if (!CreateMemoryMapFromFile(map, fileNameKey, newName.data, IsCompressed(newCasKey), memoryMapAlignment, fileName.data))
 					{
 						if (!shouldRetry)
 							return false;
diff --git a/Engine/Source/Programs/UnrealBuildAccelerator/Common/Public/UbaFileMapping.h b/Engine/Source/Programs/UnrealBuildAccelerator/Common/Public/UbaFileMapping.h
--- a/Engine/Source/Programs/UnrealBuildAccelerator/Common/Public/UbaFileMapping.h
+++ b/Engine/Source/Programs/UnrealBuildAccelerator/Common/Public/UbaFileMapping.h
@@ -69,6 +69,8 @@ namespace uba
 		FileHandle GetPersistentFile(u32 index) { return m_storage[MappedView_Persistent].files[index].file; }
 		u64 GetPersistentSize(u32 index) { return m_storage[MappedView_Persistent].files[index].size; }
 
+		u64 GetFileMappingCapacity();
+
 	private:
 		struct File
 		{
diff --git a/Engine/Source/Programs/UnrealBuildAccelerator/Common/Public/UbaSession.h b/Engine/Source/Programs/UnrealBuildAccelerator/Common/Public/UbaSession.h
--- a/Engine/Source/Programs/UnrealBuildAccelerator/Common/Public/UbaSession.h
+++ b/Engine/Source/Programs/UnrealBuildAccelerator/Common/Public/UbaSession.h
@@ -127,7 +127,7 @@ namespace uba
 		u32 AddFileMapping(StringKey fileNameKey, const tchar* fileName, const tchar* newFileName, u64 fileSize = InvalidValue);
 		
 		struct MemoryMap { StringBuffer<128> name; u64 size = 0; };
-		bool CreateMemoryMapFromFile(MemoryMap& out, StringKey fileNameKey, const tchar* fileName, bool isCompressed, u64 alignment);
+		bool CreateMemoryMapFromFile(MemoryMap& out, StringKey fileNameKey, const tchar* fileName, bool isCompressed, u64 alignment, const tchar* hint);
 		bool CreateMemoryMapFromView(MemoryMap& out, StringKey fileNameKey, const tchar* fileName, const CasKey& casKey, u64 alignment);
 
 		bool RegisterCreateFileForWrite(StringKey fileNameKey, const StringView& fileName, bool registerRealFile, u64 fileSize = 0, u64 lastWriteTime = 0, bool invalidateStorage = true);
