[UBA]
* Added so drive write/read count/bytes are included in trace file if detailed trace is enabled

diff --git a/Engine/Source/Programs/UnrealBuildAccelerator/Common/Private/UbaFile.cpp b/Engine/Source/Programs/UnrealBuildAccelerator/Common/Private/UbaFile.cpp
--- a/Engine/Source/Programs/UnrealBuildAccelerator/Common/Private/UbaFile.cpp
+++ b/Engine/Source/Programs/UnrealBuildAccelerator/Common/Private/UbaFile.cpp
@@ -5,6 +5,7 @@
 #include "UbaEvent.h"
 #include "UbaPathUtils.h"
 #include "UbaProcessStats.h"
+#include <algorithm>
 
 #if PLATFORM_WINDOWS
 #include <io.h>
@@ -922,11 +923,19 @@ namespace uba
 		return true;
 	};
 
+	VolumeCache::~VolumeCache()
+	{
+		#if PLATFORM_WINDOWS
+		for (auto& volume : volumes)
+			CloseHandle(asHANDLE(volume.handle));
+		#endif
+	}
+
 	bool VolumeCache::Init(Logger& logger)
 	{
 		#if PLATFORM_WINDOWS
-		tchar rootPathName[MAX_PATH] = {0};
-		HANDLE hFind = FindFirstVolumeW(rootPathName, ARRAYSIZE(rootPathName));
+		tchar volumeName[MAX_PATH] = {0};
+		HANDLE hFind = FindFirstVolumeW(volumeName, sizeof_array(volumeName));
 
 		if (hFind == INVALID_HANDLE_VALUE)
 			return logger.Error(TC("FindFirstVolume failed (%s)"), LastErrorToText().data);
@@ -934,14 +943,41 @@ namespace uba
 		do
 		{
 			DWORD serialNumber = 0;
-			if (GetVolumeInformationW(rootPathName, nullptr, 0, &serialNumber, nullptr, nullptr, nullptr, 0))
-				volumeSerials.push_back(serialNumber);
+			if (!GetVolumeInformationW(volumeName, nullptr, 0, &serialNumber, nullptr, nullptr, nullptr, 0))
+				continue;
+			Volume& volume = volumes.emplace_back();
+			volume.serialNumber = serialNumber;
+
+
+			tchar driveLetters[MAX_PATH] = {0};
+			if (GetVolumePathNamesForVolumeNameW(volumeName, driveLetters, sizeof_array(driveLetters), NULL))
+			{
+				StringBuffer<128> drives;
+				for (tchar* p = driveLetters; *p; p += TStrlen(p) + 1)
+					drives.Append(p);
+				volume.drives = drives.data;
+			}
+
+			u32 volumeNameLen = TStrlen(volumeName);
+			if (volumeName[volumeNameLen - 1] == PathSeparator)
+				volumeName[volumeNameLen - 1] = 0;
+
+			volume.handle = (FileHandle)(u64)::CreateFileW(volumeName, 0, FILE_SHARE_READ|FILE_SHARE_WRITE, NULL, OPEN_EXISTING, 0, NULL);
 		}
-		while (FindNextVolumeW(hFind, rootPathName, ARRAYSIZE(rootPathName)));
+		while (FindNextVolumeW(hFind, volumeName, sizeof_array(volumeName)));
 
 		FindVolumeClose(hFind);
 
-		volumeSerials.push_back(0);
+		std::sort(volumes.begin(), volumes.end(), [](const Volume& a, const Volume& b)
+			{
+				bool aempt = a.drives.empty();
+				bool bempt = b.drives.empty();
+				if (aempt != bempt)
+					return bempt;
+				if (!aempt)
+					return a.drives < b.drives;
+				return a.serialNumber < b.serialNumber;
+			});
 		#endif
 		return true;
 	}
@@ -950,28 +986,74 @@ namespace uba
 	{
 		#if PLATFORM_WINDOWS
 		u32 index = 1;
-		for (auto s : volumeSerials)
+		for (auto& volume : volumes)
 		{
-			if (s == volumeSerial)
+			if (volume.serialNumber == volumeSerial)
 				return index;
 			++index;
 		}
+		if (!volumeSerial)
+			return index;
 		#endif
 		return volumeSerial;
 	}
 
 	void VolumeCache::Write(BinaryWriter& writer)
 	{
-		writer.WriteU16(u16(volumeSerials.size()));
-		for (u32 serial : volumeSerials)
-			writer.WriteU32(serial);
+		writer.WriteU16(u16(volumes.size()));
+		for (auto& volume : volumes)
+			writer.WriteU32(volume.serialNumber);
 	}
 
 	void VolumeCache::Read(BinaryReader& reader)
 	{
 		u32 count = reader.ReadU16();
-		volumeSerials.resize(count);
+		volumes.resize(count);
 		for (u32 i=0; i!=count;++i)
-			volumeSerials[i] = reader.ReadU32();
+			volumes[i].serialNumber = reader.ReadU32();
+	}
+
+	bool VolumeCache::Volume::UpdateStats(u8& outBusyPercent, u32& outReadCount, u64& outReadBytes, u32& outWriteCount, u64& outWriteBytes)
+	{
+		outBusyPercent = 0;
+		outReadCount = 0;
+		outReadBytes = 0;
+		outWriteCount = 0;
+		outWriteBytes = 0;
+
+		#if PLATFORM_WINDOWS
+		DISK_PERFORMANCE perf = {};
+		DWORD bytesReturned;
+		if (!::DeviceIoControl(asHANDLE(handle), IOCTL_DISK_PERFORMANCE, NULL, 0, &perf, sizeof(perf), &bytesReturned, NULL))
+		{
+			::CloseHandle(asHANDLE(handle));
+			handle = InvalidFileHandle;
+			return false;
+		}
+
+		u64 queryTime = perf.QueryTime.QuadPart;
+		u64 idleTime = perf.IdleTime.QuadPart;
+		u32 readCount = perf.ReadCount;
+		u32 writeCount = perf.WriteCount;
+		u64 readBytes = perf.BytesRead.QuadPart;
+		u64 writeBytes = perf.BytesWritten.QuadPart;
+
+		if (prevQueryTime)
+		{
+			outBusyPercent = (u8)(100.0 - 100.0*(double(idleTime - prevIdleTime) / double(queryTime - prevQueryTime)));
+			outReadCount = readCount - prevReadCount;
+			outReadBytes = readBytes - prevReadBytes;
+			outWriteCount = writeCount - prevWriteCount;
+			outWriteBytes = writeBytes - prevWriteBytes;
+		}
+		prevQueryTime = queryTime;
+		prevIdleTime = idleTime;
+
+		prevReadCount = readCount;
+		prevReadBytes = readBytes;
+		prevWriteCount = writeCount;
+		prevWriteBytes = writeBytes;
+		#endif
+		return true;
 	}
 }
diff --git a/Engine/Source/Programs/UnrealBuildAccelerator/Common/Private/UbaSessionServer.cpp b/Engine/Source/Programs/UnrealBuildAccelerator/Common/Private/UbaSessionServer.cpp
--- a/Engine/Source/Programs/UnrealBuildAccelerator/Common/Private/UbaSessionServer.cpp
+++ b/Engine/Source/Programs/UnrealBuildAccelerator/Common/Private/UbaSessionServer.cpp
@@ -2167,6 +2167,24 @@ namespace uba
 		float cpuLoad = UpdateCpuLoad();
 		u64 memAvail = m_memAvail;
 		u64 memTotal = m_memTotal;
+
+
+		if (m_detailedTrace)
+		{
+			for (auto& volume : m_volumeCache.volumes)
+			{
+				if (volume.drives.empty())
+					continue;
+				u8 busyPercent;
+				u32 readCount;
+				u64 readBytes;
+				u32 writeCount;
+				u64 writeBytes;
+				volume.UpdateStats(busyPercent, readCount, readBytes, writeCount, writeBytes);
+				m_trace.DriveUpdate(volume.drives[0], busyPercent, readCount, readBytes, writeCount, writeBytes);
+			}
+		}
+
 		m_trace.SessionUpdate(0, 0, serverSend, serverRecv, 0, memAvail, memTotal, cpuLoad);
 	}
 
diff --git a/Engine/Source/Programs/UnrealBuildAccelerator/Common/Private/UbaTrace.cpp b/Engine/Source/Programs/UnrealBuildAccelerator/Common/Private/UbaTrace.cpp
--- a/Engine/Source/Programs/UnrealBuildAccelerator/Common/Private/UbaTrace.cpp
+++ b/Engine/Source/Programs/UnrealBuildAccelerator/Common/Private/UbaTrace.cpp
@@ -203,6 +203,12 @@ namespace uba
 	}
 
 	#define BEGIN_TRACE_ENTRY(x) \
+		BEGIN_TRACE_ENTRY_IMPL(x, true)
+
+	#define BEGIN_TRACE_ENTRY_NOTIME(x) \
+		BEGIN_TRACE_ENTRY_IMPL(x, false)
+
+	#define BEGIN_TRACE_ENTRY_IMPL(x, writeTime) \
 		if (!m_memoryBegin) \
 			return; \
 		u64 time = GetTime() - m_startTime; \
@@ -210,7 +216,8 @@ namespace uba
 		if (!writer.IsValid()) \
 			return; \
 		writer.WriteByte(TraceType_##x); \
-		writer.Write7BitEncoded(time);
+		if (writeTime) \
+			writer.Write7BitEncoded(time);
 
 	void Trace::SessionAdded(u32 sessionId, u32 clientId, const StringView& name, const StringView& info)
 	{
@@ -420,6 +427,17 @@ namespace uba
 		writer.WriteString(statusLink ? statusLink : TC(""));
 	}
 
+	void Trace::DriveUpdate(const tchar drive, u8 busyPercent, u32 readCount, u64 readBytes, u32 writeCount, u64 writeBytes)
+	{
+		BEGIN_TRACE_ENTRY_NOTIME(DriveUpdate);
+		writer.WriteByte((u8)drive);
+		writer.WriteByte(busyPercent);
+		writer.Write7BitEncoded(readCount);
+		writer.Write7BitEncoded(readBytes);
+		writer.Write7BitEncoded(writeCount);
+		writer.Write7BitEncoded(writeBytes);
+	}
+
 	void Trace::RemoteExecutionDisabled()
 	{
 		BEGIN_TRACE_ENTRY(RemoteExecutionDisabled);
diff --git a/Engine/Source/Programs/UnrealBuildAccelerator/Common/Public/UbaFile.h b/Engine/Source/Programs/UnrealBuildAccelerator/Common/Public/UbaFile.h
--- a/Engine/Source/Programs/UnrealBuildAccelerator/Common/Public/UbaFile.h
+++ b/Engine/Source/Programs/UnrealBuildAccelerator/Common/Public/UbaFile.h
@@ -140,10 +140,27 @@ namespace uba
 
 	struct VolumeCache
 	{
+		~VolumeCache();
 		bool Init(Logger& logger);
 		u32 GetSerialIndex(u32 volumeSerial);
 		void Write(BinaryWriter& writer);
 		void Read(BinaryReader& reader);
-		Vector<u32> volumeSerials;
+
+		struct Volume
+		{
+			u32 serialNumber = 0;
+			TString drives;
+			FileHandle handle = InvalidFileHandle;
+			u64 prevQueryTime = 0;
+			u64 prevIdleTime = 0;
+			u32 prevReadCount = 0;
+			u32 prevWriteCount = 0;
+			u64 prevReadBytes = 0;
+			u64 prevWriteBytes = 0;
+
+			bool UpdateStats(u8& outBusyPercent, u32& outReadCount, u64& outReadBytes, u32& outWriteCount, u64& outWriteBytes);
+		};
+
+		Vector<Volume> volumes;
 	};
 }
diff --git a/Engine/Source/Programs/UnrealBuildAccelerator/Common/Public/UbaTrace.h b/Engine/Source/Programs/UnrealBuildAccelerator/Common/Public/UbaTrace.h
--- a/Engine/Source/Programs/UnrealBuildAccelerator/Common/Public/UbaTrace.h
+++ b/Engine/Source/Programs/UnrealBuildAccelerator/Common/Public/UbaTrace.h
@@ -60,6 +60,7 @@ namespace uba
 		UBA_TRACE_TYPE(FileFetchSize) \
 		UBA_TRACE_TYPE(ProcessBreadcrumbs) \
 		UBA_TRACE_TYPE(WorkHint) \
+		UBA_TRACE_TYPE(DriveUpdate) \
 
 	enum TraceType : u8
 	{
@@ -103,6 +104,7 @@ namespace uba
 		void WorkEnd(u32 workIndex);
 		void ProgressUpdate(u32 processesTotal, u32 processesDone, u32 errorCount);
 		void StatusUpdate(u32 statusRow, u32 statusColumn, const tchar* statusText, LogEntryType statusType, const tchar* statusLink);
+		void DriveUpdate(const tchar drive, u8 busyPercent, u32 readCount, u64 readBytes, u32 writeCount, u64 writeBytes);
 		void RemoteExecutionDisabled();
 
 		void CacheBeginFetch(u32 fetchId, const tchar* description);
diff --git a/Engine/Source/Programs/UnrealBuildAccelerator/Visualizer/Private/UbaTraceReader.cpp b/Engine/Source/Programs/UnrealBuildAccelerator/Visualizer/Private/UbaTraceReader.cpp
--- a/Engine/Source/Programs/UnrealBuildAccelerator/Visualizer/Private/UbaTraceReader.cpp
+++ b/Engine/Source/Programs/UnrealBuildAccelerator/Visualizer/Private/UbaTraceReader.cpp
@@ -356,7 +356,7 @@ namespace uba
 		u64 readPos = reader.GetPosition();
 		u8 traceType = reader.ReadByte();
 		u64 time = 0;
-		if (out.version >= 15 && traceType != TraceType_String)
+		if (out.version >= 15 && traceType != TraceType_String && traceType != TraceType_DriveUpdate)
 		{
 			time = ConvertTime(out, reader.Read7BitEncoded());
 			if (time > maxTime)
@@ -1088,6 +1088,30 @@ namespace uba
 			out.progressErrorCount = u32(reader.Read7BitEncoded());
 			break;
 		}
+		case TraceType_DriveUpdate:
+		{
+			u8 driveLetter = reader.ReadByte();
+			u8 busyPercent = reader.ReadByte();
+			u32 readCount = u32(reader.Read7BitEncoded());
+			u64 readBytes = reader.Read7BitEncoded();
+			u32 writeCount = u32(reader.Read7BitEncoded());
+			u64 writeBytes = reader.Read7BitEncoded();
+
+			if (out.sessions.empty())
+				break;
+			auto& drive = out.sessions[0].drives[driveLetter];
+			drive.busyHighest = Max(drive.busyHighest, busyPercent);
+			drive.busyPercent.push_back(busyPercent);
+			drive.totalReadCount += readCount;
+			drive.totalWriteCount += writeCount;
+			drive.totalReadBytes += readBytes;
+			drive.totalWriteBytes += writeBytes;
+			drive.readCount.push_back(readCount);
+			drive.readBytes.push_back(readBytes);
+			drive.writeCount.push_back(writeCount);
+			drive.writeBytes.push_back(writeBytes);
+			break;
+		}
 		case TraceType_StatusUpdate:
 		{
 			if (out.version < 32)
diff --git a/Engine/Source/Programs/UnrealBuildAccelerator/Visualizer/Private/UbaTraceReader.h b/Engine/Source/Programs/UnrealBuildAccelerator/Visualizer/Private/UbaTraceReader.h
--- a/Engine/Source/Programs/UnrealBuildAccelerator/Visualizer/Private/UbaTraceReader.h
+++ b/Engine/Source/Programs/UnrealBuildAccelerator/Visualizer/Private/UbaTraceReader.h
@@ -81,6 +81,20 @@ namespace uba
 			TString link;
 		};
 
+		struct Drive
+		{
+			u8 busyHighest = 0;
+			u32 totalReadCount = 0;
+			u32 totalWriteCount = 0;
+			u64 totalReadBytes = 0;
+			u64 totalWriteBytes = 0;
+			Vector<u8> busyPercent;
+			Vector<u32> readCount;
+			Vector<u32> writeCount;
+			Vector<u64> readBytes;
+			Vector<u64> writeBytes;
+		};
+
 		struct Session
 		{
 			TString name;
@@ -101,6 +115,7 @@ namespace uba
 			Vector<FileTransfer> fetchedFiles;
 			UnorderedMap<CasKey, u32> storedFilesActive;
 			Vector<FileTransfer> storedFiles;
+			Map<char, Drive> drives;
 			TString notification;
 			u64 fetchedFilesBytes = 0;
 			u64 storedFilesBytes = 0;
diff --git a/Engine/Source/Programs/UnrealBuildAccelerator/Visualizer/Private/UbaVisualizer.cpp b/Engine/Source/Programs/UnrealBuildAccelerator/Visualizer/Private/UbaVisualizer.cpp
--- a/Engine/Source/Programs/UnrealBuildAccelerator/Visualizer/Private/UbaVisualizer.cpp
+++ b/Engine/Source/Programs/UnrealBuildAccelerator/Visualizer/Private/UbaVisualizer.cpp
@@ -483,6 +483,7 @@ namespace uba
 			m_recvColor = RGB(0, 170, 255);
 			m_cpuColor = RGB(170, 170, 0);
 			m_memColor = RGB(170, 0, 255);
+			m_driveColor = RGB(170, 65, 55);
 		}
 		else
 		{
@@ -518,6 +519,7 @@ namespace uba
 			m_recvColor = RGB(63, 72, 204); // Blue
 			m_cpuColor = RGB(200, 130, 0); // Orange
 			m_memColor = RGB(170, 0, 255); // Purple
+			m_driveColor = RGB(255, 115, 96);
 		}
 
 		m_textPen = CreatePen(PS_SOLID, 1, m_textColor);
@@ -525,6 +527,7 @@ namespace uba
 		m_recvPen = CreatePen(PS_SOLID, 1, m_recvColor);
 		m_cpuPen = CreatePen(PS_SOLID, 1, m_cpuColor);
 		m_memPen = CreatePen(PS_SOLID, 1, m_memColor);
+		m_drivePen = CreatePen(PS_SOLID, 1, m_driveColor);
 	}
 
 	void Visualizer::ThreadLoop()
@@ -1275,7 +1278,8 @@ namespace uba
 							u64 prevTime = 0;
 							int prevX = 0;
 							int prevY = 0;
-							for (u64 i=0, e=session.updates.size(); i!=e; ++i)
+							u64 e = Min(values.size(), session.updates.size());
+							for (u64 i=0; i!=e; ++i)
 							{
 								u64 updateTime = session.updates[i];
 								auto value = values[i];
@@ -1337,6 +1341,10 @@ namespace uba
 						drawGraph(session.memAvail, session.memTotal, 1.0, m_memPen, false);
 					}
 
+					for (auto& pair : session.drives)
+						if (pair.second.busyHighest)
+							drawGraph(pair.second.busyPercent, 0, -0.01, m_drivePen, false);
+
 				}
 				posY += GraphHeight;
 			}
@@ -2106,7 +2114,8 @@ namespace uba
 			logger.AddWidth(3);
 			logger.AddSpace(2);
 
-			Vector<TString>& summary = m_traceView.sessions[m_sessionSelectedIndex].summary;
+			auto& session = m_traceView.sessions[m_sessionSelectedIndex];
+			Vector<TString>& summary = session.summary;
 			if (summary.empty())
 			{
 				if (m_traceView.finished)
@@ -2114,11 +2123,18 @@ namespace uba
 				else
 					logger.Info(L" Session summary not available until session is done");
 			}
-			else
-				logger.AddTextOffset(-10); // Remove spaces in the front
 
+			logger.AddTextOffset(-10); // Remove spaces in the front
 			for (auto& line : summary)
 				logger.Log(LogEntryType_Info, line);
+			logger.AddTextOffset(0);
+
+			for (auto& pair : session.drives)
+			{
+				auto& d = pair.second;
+				logger.Info(L"  %c: Read: %s (%u) Write: %s (%u)", pair.first, BytesToText(d.totalReadBytes).str, d.totalReadCount, BytesToText(d.totalWriteBytes).str, d.totalWriteCount);
+			}
+
 			logger.AddSpace(3);
 			logger.DrawAtCursor();
 		}
@@ -2129,10 +2145,18 @@ namespace uba
 			logger.AddSpace(3);
 			logger.SetColor(m_cpuColor).Info(L"  Cpu: %.1f%%", m_stats.cpuLoad * 100.0f);
 			logger.SetColor(m_memColor).Info(L"  Mem: %ls/%ls", BytesToText(m_stats.memTotal - m_stats.memAvail).str, BytesToText(m_stats.memTotal).str);
-			logger.SetColor(m_recvColor).Info(L"  Recv: %lsps", BytesToText(m_stats.recvBytesPerSecond*8).str);
-			logger.SetColor(m_sendColor).Info(L"  Send: %lsps", BytesToText(m_stats.sendBytesPerSecond*8).str);
+			if (m_stats.recvBytes || m_stats.sendBytes)
+			{
+				logger.SetColor(m_recvColor).Info(L"  Recv: %lsps", BytesToText(m_stats.recvBytesPerSecond*8).str);
+				logger.SetColor(m_sendColor).Info(L"  Send: %lsps", BytesToText(m_stats.sendBytesPerSecond*8).str);
+			}
 			if (m_stats.ping)
 				logger.Info(L"  Ping: %ls", TimeToText(m_stats.ping, false, m_traceView.frequency).str);
+
+
+			for (auto& pair : m_stats.drives)
+				logger.SetColor(m_driveColor).Info(L"  %c: %u%%", pair.first, pair.second);
+
 			logger.AddSpace(3);
 			logger.DrawAtCursor();
 		}
@@ -2813,6 +2837,8 @@ namespace uba
 						if (pos.x + hitOffset >= prevX && pos.x + hitOffset <= x)
 						{
 							double duration = TimeToS(updateTime - prevTime);
+							outResult.stats.recvBytes = updateRecv;
+							outResult.stats.sendBytes = updateSend;
 							outResult.stats.recvBytesPerSecond = u64(float(updateRecv - prevRecv) / duration);
 							outResult.stats.sendBytesPerSecond = u64(float(updateSend - prevSend) / duration);
 							outResult.stats.ping = session.ping[i];
@@ -2820,6 +2846,9 @@ namespace uba
 							outResult.stats.cpuLoad = session.cpuLoad[i];
 							outResult.stats.memTotal = session.memTotal;
 							outResult.statsSelected = true;
+
+							for (auto& pair : session.drives)
+								outResult.stats.drives[pair.first].busyPercent = pair.second.busyPercent[i];
 							return;
 						}
 
diff --git a/Engine/Source/Programs/UnrealBuildAccelerator/Visualizer/Private/UbaVisualizer.h b/Engine/Source/Programs/UnrealBuildAccelerator/Visualizer/Private/UbaVisualizer.h
--- a/Engine/Source/Programs/UnrealBuildAccelerator/Visualizer/Private/UbaVisualizer.h
+++ b/Engine/Source/Programs/UnrealBuildAccelerator/Visualizer/Private/UbaVisualizer.h
@@ -98,7 +98,14 @@ namespace uba
 			u64 ping = 0;
 			u64 memAvail = 0;
 			u64 memTotal = 0;
+			u64 recvBytes = 0;
+			u64 sendBytes = 0;
 			float cpuLoad = 0;
+			struct Drive
+			{
+				u8 busyPercent;
+			};
+			Map<char, Drive> drives;
 		};
 		struct HitTestResult
 		{
@@ -184,6 +191,7 @@ namespace uba
 		COLORREF m_recvColor = {};
 		COLORREF m_cpuColor = {};
 		COLORREF m_memColor = {};
+		COLORREF m_driveColor = {};
 		HBRUSH m_backgroundBrush = 0;
 		HBRUSH m_tooltipBackgroundBrush = 0;
 		HPEN m_textPen = 0;
@@ -192,6 +200,7 @@ namespace uba
 		HPEN m_recvPen = 0;
 		HPEN m_cpuPen = 0;
 		HPEN m_memPen = 0;
+		HPEN m_drivePen = 0;
 		HPEN m_processUpdatePen = 0;
 		HPEN m_checkboxPen = 0;
 		int m_boxHeight = 12;
