[UBA]
* Fixed problem where two dlls with different names but hardlinked to same file didn't work because they ended up being the same dll
* Changed const tchar* to StringView in WriteBinFile
* Added code for not calling GetModuleHandle since it could actually accidentally find the "wrong" file on the remote.
* Added some code to prepare for making shared memory files work between processes when connected to stderr/stdout
* Improved some error messages

diff --git a/Engine/Source/Programs/UnrealBuildAccelerator/Common/Private/UbaFileAccessor.cpp b/Engine/Source/Programs/UnrealBuildAccelerator/Common/Private/UbaFileAccessor.cpp
--- a/Engine/Source/Programs/UnrealBuildAccelerator/Common/Private/UbaFileAccessor.cpp
+++ b/Engine/Source/Programs/UnrealBuildAccelerator/Common/Private/UbaFileAccessor.cpp
@@ -107,7 +107,8 @@ namespace uba
 				}
 			}
 			#endif
-			return m_logger.Error(TC("ERROR opening file %s for write (%s%s)"), realFileName, LastErrorToText(lastError).data, additionalInfo.data);
+			const tchar* retryText = retryCount ? TC(" after retrying 10 times") : TC("");
+			return m_logger.Error(TC("ERROR opening file %s for write%s (%s%s)"), realFileName, retryText, LastErrorToText(lastError).data, additionalInfo.data);
 		}
 
 #if PLATFORM_WINDOWS
diff --git a/Engine/Source/Programs/UnrealBuildAccelerator/Common/Private/UbaSessionClient.cpp b/Engine/Source/Programs/UnrealBuildAccelerator/Common/Private/UbaSessionClient.cpp
--- a/Engine/Source/Programs/UnrealBuildAccelerator/Common/Private/UbaSessionClient.cpp
+++ b/Engine/Source/Programs/UnrealBuildAccelerator/Common/Private/UbaSessionClient.cpp
@@ -308,7 +308,7 @@ namespace uba
 		applicationDirLower.Append(applicationDir).MakeLower();
 		KeyToString keyStr(ToStringKey(applicationDirLower));
 
-		return WriteBinFile(out, destFile.data, newKey, keyStr, fileAttributes);
+		return WriteBinFile(out, destFile, newKey, keyStr, fileAttributes);
 	}
 
 	bool SessionClient::PrepareProcess(ProcessStartInfoHolder& startInfo, bool isChild, StringBufferBase& outRealApplication, const tchar*& outRealWorkingDir)
@@ -376,7 +376,7 @@ namespace uba
 				if (const tchar* lastSeparator = TStrrchr(moduleName, PathSeparator))
 					moduleName = lastSeparator + 1;
 				StringBuffer<MaxPath> temp;
-				if (!WriteBinFile(temp, moduleName, newCasKey, keyStr, m.attributes))
+				if (!WriteBinFile(temp, ToView(moduleName), newCasKey, keyStr, m.attributes))
 					success = false;
 			}, AsView(TC("FetchModule")), true);
 
@@ -446,7 +446,7 @@ namespace uba
 		return m_environmentVariables.data();
 	}
 
-	bool SessionClient::WriteBinFile(StringBufferBase& out, const tchar* binaryName, const CasKey& casKey, const KeyToString& applicationDir, u32 fileAttributes)
+	bool SessionClient::WriteBinFile(StringBufferBase& out, const StringView& binaryName, const CasKey& casKey, const KeyToString& applicationDir, u32 fileAttributes)
 	{
 		UBA_ASSERT(fileAttributes);
 
@@ -473,7 +473,7 @@ namespace uba
 		//if (GetFileAttributesW(out.data) != INVALID_FILE_ATTRIBUTES)
 		//	return true;
 
-		if (TStrchr(binaryName, PathSeparator))
+		if (TStrchr(binaryName.data, PathSeparator))
 		{
 			StringBuffer<> binaryDir;
 			binaryDir.AppendDir(out);
@@ -481,9 +481,12 @@ namespace uba
 				return false;
 		}
 
+		// Special hack to prevent two identical dlls from pointing to the same file (dlls are deduplicated if they are the same file in the backend)
+		bool allowHardlink = !binaryName.GetFileName().StartsWith(TC("c2"));
+
 		constexpr bool writeCompressed = false;
 		constexpr bool isTemp = true;
-		return m_storage.CopyOrLink(casKey, out.data, fileAttributes, writeCompressed, {}, isTemp);
+		return m_storage.CopyOrLink(casKey, out.data, fileAttributes, writeCompressed, {}, isTemp, allowHardlink);
 	}
 
 	bool SessionClient::ProcessThreadStart(ProcessImpl& process)
@@ -1302,7 +1305,7 @@ namespace uba
 				key.a = i;
 				KeyToString dir(key);
 
-				if (!WriteBinFile(detoursFile[i], UBA_DETOURS_LIBRARY, detoursBinaryKey[i], dir, DefaultAttributes()))
+				if (!WriteBinFile(detoursFile[i], AsView(UBA_DETOURS_LIBRARY), detoursBinaryKey[i], dir, DefaultAttributes()))
 					return;
 			}
 		}
diff --git a/Engine/Source/Programs/UnrealBuildAccelerator/Common/Private/UbaStorage.cpp b/Engine/Source/Programs/UnrealBuildAccelerator/Common/Private/UbaStorage.cpp
--- a/Engine/Source/Programs/UnrealBuildAccelerator/Common/Private/UbaStorage.cpp
+++ b/Engine/Source/Programs/UnrealBuildAccelerator/Common/Private/UbaStorage.cpp
@@ -2307,7 +2307,7 @@ namespace uba
 		return out != CasKeyZero;
 	}
 
-	bool StorageImpl::CopyOrLink(const CasKey& casKey, const tchar* destination, u32 fileAttributes, bool writeCompressed, const FormattingFunc& formattingFunc, bool isTemp)
+	bool StorageImpl::CopyOrLink(const CasKey& casKey, const tchar* destination, u32 fileAttributes, bool writeCompressed, const FormattingFunc& formattingFunc, bool isTemp, bool allowHardLink)
 	{
 		UBA_ASSERT(casKey != CasKeyZero);
 		UBA_ASSERT(fileAttributes);
@@ -2560,7 +2560,8 @@ namespace uba
 				bool success = false;
 
 				#if !PLATFORM_MAC // For some reason creating links on macos causes trouble when they are exec/dylibs and being executed.. sometimes it is like the link behaves like a symlink.. but not always
-				success = CreateHardLinkW(destination, casFile.data);
+				if (allowHardLink)
+					success = CreateHardLinkW(destination, casFile.data);
 				#endif
 
 				if (!success)
diff --git a/Engine/Source/Programs/UnrealBuildAccelerator/Common/Public/UbaSessionClient.h b/Engine/Source/Programs/UnrealBuildAccelerator/Common/Public/UbaSessionClient.h
--- a/Engine/Source/Programs/UnrealBuildAccelerator/Common/Public/UbaSessionClient.h
+++ b/Engine/Source/Programs/UnrealBuildAccelerator/Common/Public/UbaSessionClient.h
@@ -85,7 +85,7 @@ namespace uba
 		bool GetCasKeyForFile(CasKey& out, u32 processId, const StringView& fileName, const StringKey& fileNameKey);
 		bool ReadModules(List<ModuleInfo>& outModules, u32 processId, const tchar* application);
 		bool EnsureBinaryFile(StringBufferBase& out, StringBufferBase& outVirtual, u32 processId, const StringBufferBase& fileName, const StringKey& fileNameKey, const tchar* applicationDir, const u8* loaderPaths, u32 loaderPathsSize);
-		bool WriteBinFile(StringBufferBase& out, const tchar* binaryName, const CasKey& casKey, const KeyToString& applicationDir, u32 fileAttributes);
+		bool WriteBinFile(StringBufferBase& out, const StringView& binaryName, const CasKey& casKey, const KeyToString& applicationDir, u32 fileAttributes);
 		bool SendFiles(ProcessImpl& process, Timer& sendFiles);
 		bool SendFile(WrittenFile& source, const tchar* destination, u32 processId, bool keepMappingInMemory, bool compressed);
 		bool SendUpdateDirectoryTable(StackBinaryReader<SendMaxSize>& reader); // Note, reader is sent in to save stack space.
diff --git a/Engine/Source/Programs/UnrealBuildAccelerator/Common/Public/UbaStorage.h b/Engine/Source/Programs/UnrealBuildAccelerator/Common/Public/UbaStorage.h
--- a/Engine/Source/Programs/UnrealBuildAccelerator/Common/Public/UbaStorage.h
+++ b/Engine/Source/Programs/UnrealBuildAccelerator/Common/Public/UbaStorage.h
@@ -50,7 +50,7 @@ namespace uba
 		virtual bool CalculateCasKey(CasKey& out, const tchar* fileName) = 0;
 
 		using FormattingFunc = Function<bool(MemoryBlock& destData, const void* sourceData, u64 sourceSize, const tchar* hint)>;
-		virtual bool CopyOrLink(const CasKey& casKey, const tchar* destination, u32 fileAttributes, bool writeCompressed = false, const FormattingFunc& formattingFunc = {}, bool isTemp = false) = 0;
+		virtual bool CopyOrLink(const CasKey& casKey, const tchar* destination, u32 fileAttributes, bool writeCompressed = false, const FormattingFunc& formattingFunc = {}, bool isTemp = false, bool allowHardLink = true) = 0;
 		virtual bool FakeCopy(const CasKey& casKey, const tchar* destination, u64 size = 0, u64 lastWritten = 0, bool deleteExisting = true) = 0;
 		virtual bool GetCasFileName(StringBufferBase& out, const CasKey& casKey) = 0;
 
@@ -138,7 +138,7 @@ namespace uba
 		virtual bool DropCasFile(const CasKey& casKey, bool forceDelete, const tchar* hint) override;
 		virtual bool ReportBadCasFile(const CasKey& casKey) override;
 		virtual bool CalculateCasKey(CasKey& out, const tchar* fileName) override;
-		virtual bool CopyOrLink(const CasKey& casKey, const tchar* destination, u32 fileAttributes, bool writeCompressed = false, const FormattingFunc& formattingFunc = {}, bool isTemp = false) override;
+		virtual bool CopyOrLink(const CasKey& casKey, const tchar* destination, u32 fileAttributes, bool writeCompressed = false, const FormattingFunc& formattingFunc = {}, bool isTemp = false, bool allowHardLink = true) override;
 		virtual bool FakeCopy(const CasKey& casKey, const tchar* destination, u64 size = 0, u64 lastWritten = 0, bool deleteExisting = true) override;
 		virtual bool GetCasFileName(StringBufferBase& out, const CasKey& casKey) override;
 		virtual MappedView MapView(const CasKey& casKey, const tchar* hint) override;
diff --git a/Engine/Source/Programs/UnrealBuildAccelerator/Detours/Private/UbaDetoursShared.cpp b/Engine/Source/Programs/UnrealBuildAccelerator/Detours/Private/UbaDetoursShared.cpp
--- a/Engine/Source/Programs/UnrealBuildAccelerator/Detours/Private/UbaDetoursShared.cpp
+++ b/Engine/Source/Programs/UnrealBuildAccelerator/Detours/Private/UbaDetoursShared.cpp
@@ -303,6 +303,7 @@ namespace uba
 	void Rpc_WriteLog(const tchar* text, u64 textCharLength, bool printInSession, bool isError)
 	{
 		DEBUG_LOG(TC("LOG  %.*s"), u32(textCharLength), text); // TODO: Investigate, deadlocks on non-windows
+		// DEBUG_LOG(TC("LOG [%7u] %.*s"), GetCurrentThreadId(), u32(textCharLength), text);
 		TimerScope ts(g_stats.log);
 		SCOPED_WRITE_LOCK(g_communicationLock, pcs);
 		BinaryWriter writer;
diff --git a/Engine/Source/Programs/UnrealBuildAccelerator/Detours/Private/Windows/UbaDetoursFunctionsKernelBase.inl b/Engine/Source/Programs/UnrealBuildAccelerator/Detours/Private/Windows/UbaDetoursFunctionsKernelBase.inl
--- a/Engine/Source/Programs/UnrealBuildAccelerator/Detours/Private/Windows/UbaDetoursFunctionsKernelBase.inl
+++ b/Engine/Source/Programs/UnrealBuildAccelerator/Detours/Private/Windows/UbaDetoursFunctionsKernelBase.inl
@@ -2767,6 +2767,29 @@ HMODULE Recursive_LoadLibraryExW(LPCWSTR lpLibFileName, LPCWSTR originalName, DW
 	return res;
 }
 
+HMODULE GetModuleHandleNoSyscall(const tchar* moduleName)
+{
+#if 1
+	const tchar* moduleFileName = GetFileName(moduleName);
+	PPEB peb = (PPEB)__readgsqword(0x60);
+	PPEB_LDR_DATA ldr = peb->Ldr;
+	PLIST_ENTRY head = &ldr->InMemoryOrderModuleList;
+	PLIST_ENTRY entry = head->Flink;
+	while (entry != head)
+	{
+		PLDR_DATA_TABLE_ENTRY mod = CONTAINING_RECORD(entry, LDR_DATA_TABLE_ENTRY, InMemoryOrderLinks);
+		StringView view(mod->FullDllName.Buffer, mod->FullDllName.Length/2);
+		if (view.EndsWith(moduleFileName))
+			return (HMODULE)mod->DllBase;
+		entry = entry->Flink;
+	}
+	return NULL;  // Module not found
+#else
+	SuppressCreateFileDetourScope cfs;
+	return GetModuleHandle(moduleName);
+#endif
+}
+
 HMODULE Shared_LoadLibrary(LPCWSTR lpLibFileName, HANDLE hFile, DWORD dwFlags)
 {
 	if (!g_runningRemote && !g_trackInputsMem)
@@ -2780,10 +2803,7 @@ HMODULE Shared_LoadLibrary(LPCWSTR lpLibFileName, HANDLE hFile, DWORD dwFlags)
 		detourDll = g_runningRemote && GetFileAttributes(path.data) == INVALID_FILE_ATTRIBUTES; // It might be that remote machine actually doesn't have the file in system32. then we need to detour
 
 	if (detourDll)
-	{
-		SuppressCreateFileDetourScope cfs;
-		detourDll = !GetModuleHandle(lpLibFileName); // This internally can end up calling NtCreate and we don't want NtCreate to handle the download of the file because of paths
-	}
+		detourDll = !GetModuleHandleNoSyscall(lpLibFileName); // This internally can end up calling NtCreate and we don't want NtCreate to handle the download of the file because of paths
 
 	if (!detourDll)
 		return NULL;
@@ -3093,24 +3113,44 @@ BOOL Detoured_CreateProcessW(LPCWSTR lpApplicationName, LPWSTR lpCommandLine, LP
 		DEBUG_LOG_PIPE(L"CreateProcess", L"%ls %ls", application.data, originalCmd.data);
 	}
 
-	LPCSTR dlls[] = { dll };
+	auto handleFileDetour = [&](HANDLE& handle, HandleType type, const tchar* name, DWORD nativeHandleId)
+		{
+			if (!isDetouredHandle(handle))
+				return;
+			DetouredHandle& dh = asDetouredHandle(handle);
+			if (dh.type == type)
+			{
+				handle = g_isDetachedProcess ? 0 : True_GetStdHandle(nativeHandleId);
+				return;
+			}
+
+			handle = dh.trueHandle;
+			if (dh.trueHandle != INVALID_HANDLE_VALUE)
+				return;
+
+			auto assertGuard = MakeGuard([&]() { UBA_ASSERTF(false, L"%s is detoured and there is no proper handle (%s)", name, lpApplicationName); });
+			if (dh.type != HandleType_File)
+				return;
+			if (!dh.fileObject)
+				return;
+			auto& fo = *dh.fileObject;
+			if (!fo.fileInfo)
+				return;
+			FileInfo& fi = *fo.fileInfo;
+			auto mf = fi.memoryFile;
+			if (!mf)
+				return;
+			if (mf->isLocalOnly)
+				return;
+			// TODO: This is a memory file
+			assertGuard.Cancel();
+		};
+
+
+	handleFileDetour(lpStartupInfo->hStdError, HandleType_StdErr, TC("hStdError"), STD_ERROR_HANDLE);
+	handleFileDetour(lpStartupInfo->hStdOutput, HandleType_StdOut, TC("hStdOutput"), STD_OUTPUT_HANDLE);
+
 
-	if (isDetouredHandle(lpStartupInfo->hStdError))
-	{
-		DetouredHandle& dh = asDetouredHandle(lpStartupInfo->hStdError);
-		if (dh.type == HandleType_StdErr)
-			lpStartupInfo->hStdError = g_isDetachedProcess ? 0 : True_GetStdHandle(STD_ERROR_HANDLE);
-		else
-			UBA_ASSERTF(false, L"hStdError is detoured (%s)", lpApplicationName);
-	}
-	if (isDetouredHandle(lpStartupInfo->hStdOutput))
-	{
-		DetouredHandle& dh = asDetouredHandle(lpStartupInfo->hStdOutput);
-		if (dh.type == HandleType_StdOut)
-			lpStartupInfo->hStdOutput = g_isDetachedProcess ? 0 : True_GetStdHandle(STD_OUTPUT_HANDLE);
-		else
-			UBA_ASSERTF(false, L"hStdOutput is detoured (%s)", lpApplicationName);
-	}
 	if (isDetouredHandle(lpStartupInfo->hStdInput))
 	{
 		DetouredHandle& dh = asDetouredHandle(lpStartupInfo->hStdInput);
@@ -3139,6 +3179,9 @@ BOOL Detoured_CreateProcessW(LPCWSTR lpApplicationName, LPWSTR lpCommandLine, LP
 	u32 retryCount = 0;
 
 	++t_disallowDetour;
+
+	LPCSTR dlls[] = { dll };
+
 	while (true)
 	{
 		res = true;
