[UBA]
* Fixed win64 arm compile error

diff --git a/Engine/Source/Programs/UnrealBuildAccelerator/Detours/Private/Windows/UbaDetoursFunctionsKernelBase.inl b/Engine/Source/Programs/UnrealBuildAccelerator/Detours/Private/Windows/UbaDetoursFunctionsKernelBase.inl
--- a/Engine/Source/Programs/UnrealBuildAccelerator/Detours/Private/Windows/UbaDetoursFunctionsKernelBase.inl
+++ b/Engine/Source/Programs/UnrealBuildAccelerator/Detours/Private/Windows/UbaDetoursFunctionsKernelBase.inl
@@ -2767,11 +2767,17 @@ HMODULE Recursive_LoadLibraryExW(LPCWSTR lpLibFileName, LPCWSTR originalName, DW
 	return res;
 }
 
+#if defined(_M_ARM64) 
+	#define ReadTebOffset(offset) *(ULONG_PTR *)((BYTE *)NtCurrentTeb() + offset)
+#else
+	#define ReadTebOffset(offset) (*(ULONG_PTR *)(__readgsqword(offset)))
+#endif
+
 HMODULE GetModuleHandleNoSyscall(const tchar* moduleName)
 {
 #if 1
 	const tchar* moduleFileName = GetFileName(moduleName);
-	PPEB peb = (PPEB)__readgsqword(0x60);
+	PPEB peb = (PPEB)ReadTebOffset(0x60);
 	PPEB_LDR_DATA ldr = peb->Ldr;
 	PLIST_ENTRY head = &ldr->InMemoryOrderModuleList;
 	PLIST_ENTRY entry = head->Flink;
