[UBA]
* Proper fix for arm64 compile error

diff --git a/Engine/Source/Programs/UnrealBuildAccelerator/Detours/Private/Windows/UbaDetoursFunctionsKernelBase.inl b/Engine/Source/Programs/UnrealBuildAccelerator/Detours/Private/Windows/UbaDetoursFunctionsKernelBase.inl
--- a/Engine/Source/Programs/UnrealBuildAccelerator/Detours/Private/Windows/UbaDetoursFunctionsKernelBase.inl
+++ b/Engine/Source/Programs/UnrealBuildAccelerator/Detours/Private/Windows/UbaDetoursFunctionsKernelBase.inl
@@ -2768,9 +2768,9 @@ HMODULE Recursive_LoadLibraryExW(LPCWSTR lpLibFileName, LPCWSTR originalName, DW
 }
 
 #if defined(_M_ARM64) 
-	#define ReadTebOffset(offset) *(ULONG_PTR *)((BYTE *)NtCurrentTeb() + offset)
+	#define ReadTebOffset(offset) (((BYTE*)NtCurrentTeb()) + offset)
 #else
-	#define ReadTebOffset(offset) (*(ULONG_PTR *)(__readgsqword(offset)))
+	#define ReadTebOffset(offset) __readgsqword(offset)
 #endif
 
 HMODULE GetModuleHandleNoSyscall(const tchar* moduleName)
