[UBA]
* UbaVisualizer - Fixed bug with background color being wrong on text

diff --git a/Engine/Source/Programs/UnrealBuildAccelerator/Visualizer/Private/UbaVisualizer.cpp b/Engine/Source/Programs/UnrealBuildAccelerator/Visualizer/Private/UbaVisualizer.cpp
--- a/Engine/Source/Programs/UnrealBuildAccelerator/Visualizer/Private/UbaVisualizer.cpp
+++ b/Engine/Source/Programs/UnrealBuildAccelerator/Visualizer/Private/UbaVisualizer.cpp
@@ -921,6 +921,7 @@ namespace uba
 		SetBkColor(hdc, m_config.DarkMode ? RGB(70, 70, 70) : RGB(180, 180, 180));
 
 		HBRUSH lastSelectedBrush = 0;
+		HBRUSH textLastSelectedBrush = 0;
 		auto fillRect = [&](const RECT& r, HBRUSH b)
 			{
 				if (lastSelectedBrush != b)
@@ -1463,7 +1464,7 @@ namespace uba
 									if (!done)
 										rect3.right = 256;
 
-									PaintProcessRect(process, textDC, rect3, rect2, selected, true, lastSelectedBrush);
+									PaintProcessRect(process, textDC, rect3, rect2, selected, true, textLastSelectedBrush);
 
 									rect2.left += 3; // Move in text a bit
 									
@@ -1700,7 +1701,8 @@ namespace uba
 
 									RECT rect2{ 0,int(work.bitmapOffset), 256, int(work.bitmapOffset) + m_processFont.height };
 
-									FillRect(textDC, &rect2, brush);
+									SelectObject(textDC, brush);
+									PatBlt(textDC, rect2.left, rect2.top, rect2.right - rect2.left, rect2.bottom - rect2.top, PATCOPY);
 									ExtTextOutW(textDC, rect2.left, rect2.top, ETO_CLIPPED, &rect2, work.description, int(wcslen(work.description)), NULL);
 									if (!selected)
 										work.bitmapDirty = false;
