[UBA]
* Changed so saving bucket to disk is part of update job.. since saving to file can take time and is not cpu intense it should hopefully give better wall time to save at the same time as buckets are being recreated
* Fixed so cache server only asks once for files with same cas content when storing from client
* Minor optimization when writing buckets to disk
* Increased timeout for failing to open file for write that is locked by something else

diff --git a/Engine/Source/Programs/UnrealBuildAccelerator/CacheService/Private/UbaCacheService.cpp b/Engine/Source/Programs/UnrealBuildAccelerator/CacheService/Private/UbaCacheService.cpp
--- a/Engine/Source/Programs/UnrealBuildAccelerator/CacheService/Private/UbaCacheService.cpp
+++ b/Engine/Source/Programs/UnrealBuildAccelerator/CacheService/Private/UbaCacheService.cpp
@@ -31,7 +31,7 @@ namespace uba
 			GetFullPathNameW(TC("~/" UE_APP_NAME), sizeof_array(buf), buf, nullptr);
 		return buf;
 		}();
-	u32				DefaultProcessorCount = []() { return GetLogicalProcessorCount(); }();
+	u32				DefaultWorkerCount = []() { return GetLogicalProcessorCount() + 4; }();
 
 	int PrintHelp(const tchar* message)
 	{
@@ -56,6 +56,7 @@ namespace uba
 		logger.Info(TC("  -nomaintenance          Skip all maintenance"));
 		logger.Info(TC("  -crash                  Force a crash (for testing)"));
 		logger.Info(TC("  -nosignalhandler        Will not hook up signal handler"));
+		logger.Info(TC("  -maxworkers=<number>    Max number of workers used by cache server. Defaults to \"%u\""), DefaultWorkerCount);
 
 		logger.Info(TC(" Example of how to register crypto key to cache server (when -http=80 is provided)"));
 		logger.Info(TC("   curl http://localhost/addcrypto?3f58aa57466db9999213456789123445"));
@@ -113,6 +114,7 @@ namespace uba
 		bool signalHandlerEnabled = true;
 		bool shouldCrash = false;
 		u32 expirationTimeSeconds = DefaultExpiration;
+		u32 maxWorkerCount = DefaultWorkerCount;
 		TString configFile;
 
 		for (int i=1; i!=argc; ++i)
@@ -179,6 +181,11 @@ namespace uba
 			{
 				signalHandlerEnabled = false;
 			}
+			else if (name.Equals(TC("-maxworkers")))
+			{
+				if (!value.Parse(maxWorkerCount))
+					return PrintHelp(TC("Invalid value for -maxworkers"));
+			}
 			else if (name.Equals(TC("-nomaintenance")))
 			{
 				maintenanceEnabled = false;
@@ -231,7 +238,7 @@ namespace uba
 		#if UBA_DEBUG
 		dbgStr = TC(" (DEBUG)");
 		#endif
-		logger.Info(TC("UbaCacheService v%s(%u)%s (Workers: %u, Rootdir: \"%s\", StoreCapacity: %uGb, Expiration: %s)"), Version, CacheNetworkVersion, dbgStr, GetLogicalProcessorCount(), g_rootDir.data, storageCapacityGb, TimeToText(MsToTime(expirationTimeSeconds)*1000, true).str);
+		logger.Info(TC("UbaCacheService v%s(%u)%s (Workers: %u, Rootdir: \"%s\", StoreCapacity: %uGb, Expiration: %s)"), Version, CacheNetworkVersion, dbgStr, maxWorkerCount, g_rootDir.data, storageCapacityGb, TimeToText(MsToTime(expirationTimeSeconds)*1000, true).str);
 
 		u64 maintenanceReserveSizeMb = 128;
 
@@ -268,7 +275,7 @@ namespace uba
 
 		NetworkServerCreateInfo nsci(logWriter);
 		nsci.Apply(config);
-		//nsci.workerCount = 4;
+		nsci.workerCount = maxWorkerCount;
 		nsci.logConnections = false; // Let the cache server report instead
 		nsci.receiveTimeoutSeconds = 2*60*60; // Two hours timeout
 		bool ctorSuccess = true;
diff --git a/Engine/Source/Programs/UnrealBuildAccelerator/Common/Private/UbaCacheServer.cpp b/Engine/Source/Programs/UnrealBuildAccelerator/Common/Private/UbaCacheServer.cpp
--- a/Engine/Source/Programs/UnrealBuildAccelerator/Common/Private/UbaCacheServer.cpp
+++ b/Engine/Source/Programs/UnrealBuildAccelerator/Common/Private/UbaCacheServer.cpp
@@ -153,13 +153,14 @@ namespace uba
 
 	struct CacheServer::Bucket
 	{
-		Bucket(u64 id, u32 version) : m_pathTable(IsCaseInsensitive(id), 0, 0, version) {}
+		Bucket(u64 id, u32 version) : m_id(id), m_pathTable(IsCaseInsensitive(id), 0, 0, version) {}
 		ReaderWriterLock m_cacheEntryLookupLock;
 		UnorderedMap<CasKey, CacheEntries> m_cacheEntryLookup;
 
 		CompactPathTable m_pathTable;
 		CompactCasKeyTable m_casKeyTable;
 
+		u64 m_id;
 		u64 totalEntryCount;
 		u64 totalEntrySize;
 		Atomic<bool> hasDeletedEntries;
@@ -397,6 +398,19 @@ namespace uba
 
 		void WriteBytes(const void* data, u64 size)
 		{
+			if (size > TempBufferSize)
+			{
+				if (tempBufferPos)
+				{
+					written += tempBufferPos;
+					success &= file.Write(tempBuffer, tempBufferPos);
+					tempBufferPos = 0;
+				}
+				success &= file.Write(data, size);
+				written += size;
+				return;
+			}
+
 			u8* readPos = (u8*)data;
 			u64 left = size;
 			while (left)
@@ -453,7 +467,7 @@ namespace uba
 		FileAccessor file;
 	};
 
-	bool CacheServer::SaveBucket(u64 bucketId, Bucket& bucket, Vector<u8>& temp)
+	bool CacheServer::SaveBucket(Bucket& bucket, Vector<u8>& temp)
 	{
 		u64 saveStart = GetTime();
 
@@ -463,7 +477,7 @@ namespace uba
 			return false;
 		bucketsDir.EnsureEndsWithSlash();
 		StringBuffer<MaxPath> bucketsFile(bucketsDir);
-		bucketsFile.AppendValue(bucketId);
+		bucketsFile.AppendValue(bucket.m_id);
 
 		FileWriter file(m_logger, bucketsFile.data);
 		
@@ -511,33 +525,51 @@ namespace uba
 
 	bool CacheServer::SaveNoLock()
 	{
-		if (m_dbfileDirty)
-		{
-			StringBuffer<MaxPath> fileName(m_rootDir);
-			fileName.EnsureEndsWithSlash().Append(TC("cachedb"));
+		if (!SaveDbNoLock())
+			return false;
 
-			FileWriter file(m_logger, fileName.data);
-		
-			if (!file.Create())
-				return false;
+		Atomic<bool> success = true;
 
-			file.Write(CacheFileVersion);
+		DeleteEmptyBuckets();
 
-			file.Write(m_creationTime);
+		m_server.ParallelFor(GetBucketWorkerCount(), m_buckets, [&, temp = Vector<u8>()](auto& it) mutable
+			{
+				Bucket& bucket = it->second;
+				if (!bucket.needsSave)
+					return;
+				if (SaveBucket(bucket, temp))
+					bucket.needsSave = false;
+				else
+					success = false;
+			});
 
-			if (!file.Close())
-				return false;
-			m_dbfileDirty = false;
-		}
+		return success;
+	}
 
-		StringBuffer<MaxPath> bucketsDir(m_rootDir);
-		bucketsDir.EnsureEndsWithSlash().Append(TC("buckets"));
-		if (!m_storage.CreateDirectory(bucketsDir.data))
+	bool CacheServer::SaveDbNoLock()
+	{
+		if (!m_dbfileDirty)
+			return true;
+		StringBuffer<MaxPath> fileName(m_rootDir);
+		fileName.EnsureEndsWithSlash().Append(TC("cachedb"));
+
+		FileWriter file(m_logger, fileName.data);
+		
+		if (!file.Create())
 			return false;
-		bucketsDir.EnsureEndsWithSlash();
 
-		Atomic<bool> success = true;
+		file.Write(CacheFileVersion);
+
+		file.Write(m_creationTime);
 
+		if (!file.Close())
+			return false;
+		m_dbfileDirty = false;
+		return true;
+	}
+
+	bool CacheServer::DeleteEmptyBuckets()
+	{
 		for (auto it=m_buckets.begin(); it!=m_buckets.end();)
 		{
 			Bucket& bucket = it->second;
@@ -553,19 +585,7 @@ namespace uba
 			m_logger.Detail(TC("    Bucket %u was empty. Deleted"), bucket.index);
 			it = m_buckets.erase(it);
 		}
-
-		m_server.ParallelFor(GetBucketWorkerCount(), m_buckets, [&, temp = Vector<u8>()](auto& it) mutable
-			{
-				Bucket& bucket = it->second;
-				if (!bucket.needsSave)
-					return;
-				if (SaveBucket(it->first, bucket, temp))
-					bucket.needsSave = false;
-				else
-					success = false;
-			});
-
-		return success;
+		return true;
 	}
 
 	void CacheServer::SetForceFullMaintenance()
@@ -1080,6 +1100,17 @@ namespace uba
 		if (shouldExit && shouldExit())
 			return true;
 
+
+		bool shouldSave = allowSave && (entriesAdded || deletedCasCount || deleteEntryCount || forceAllSteps);
+
+		Event saveCasEvent(true);
+		Event saveDbEvent(true);
+		if (shouldSave)
+		{
+			m_server.AddWork([&](const WorkContext&) { m_storage.SaveCasTable(false, false); saveCasEvent.Set(); }, 1, TC("SaveCas"));
+			m_server.AddWork([&](const WorkContext&) { SaveDbNoLock(); saveDbEvent.Set(); }, 1, TC("SaveDb"));
+		}
+
 		u64 maxCommittedMemory = 0;
 
 		m_server.ParallelFor(workerCountToUseForBuckets, sortedBuckets, [&](auto& it)
@@ -1087,6 +1118,17 @@ namespace uba
 			u64 bucketStartTime = GetTime();
 
 			Bucket& bucket = **it;
+
+			auto saveGuard = MakeGuard([&]()
+				{
+					if (shouldSave && !bucket.m_cacheEntryLookup.empty())
+					{
+						Vector<u8> temp;
+						if (SaveBucket(bucket, temp))
+							bucket.needsSave = false;
+					}
+				});
+
 			auto deleteContext = MakeGuard([&]() { delete bucket.m_maintenanceContext; bucket.m_maintenanceContext = nullptr; });
 
 			if (!bucket.hasDeletedEntries && !forceAllSteps)
@@ -1230,16 +1272,16 @@ namespace uba
 				m_logger.Detail(TC("    Bucket %u Updated cache entries with new tables (%s)"), bucket.index, TimeToText(GetTime() - updateEntriesStart).str);
 			}
 
+			{
+				SCOPED_FUTEX(globalStatsLock, l);
+				maxCommittedMemory = Max(maxCommittedMemory, memoryBlock.writtenSize);
+			}
 			bucket.needsSave = true;
 
 			u64 longestUnusedSeconds = 0;
 			if (bucket.oldestUsedTime)
 				longestUnusedSeconds = GetFileTimeAsSeconds(now - m_creationTime - bucket.oldestUsedTime);
 			m_logger.Info(TC("    Bucket %u Done (%s). CacheEntries: %s (%s) PathTable: %s CasTable: %s LongestUnused: %s"), bucket.index, TimeToText(GetTime() - bucketStartTime).str, CountToText(bucket.totalEntryCount).str, BytesToText(bucket.totalEntrySize).str, BytesToText(bucket.m_pathTable.GetSize()).str, BytesToText(bucket.m_casKeyTable.GetSize()).str, TimeToText(MsToTime(longestUnusedSeconds*1000), true).str);
-
-			SCOPED_FUTEX(globalStatsLock, l);
-			maxCommittedMemory = Max(maxCommittedMemory, memoryBlock.writtenSize);
-
 		}, AsView(TC("")), true);
 
 		// Need to make sure all cas entries are dropped before saving cas table
@@ -1249,15 +1291,11 @@ namespace uba
 		if (TimeToMs(dropCasDuration) > 10)
 			m_logger.Detail(TC("  Done deleting cas files (%s)"), TimeToText(dropCasDuration).str);
 
-		if (allowSave && (entriesAdded || deletedCasCount || deleteEntryCount || forceAllSteps))
+		if (shouldSave)
 		{
-			u64 saveStart = GetTime();
-			m_logger.Detail(TC("  Saving to disk (v%u)"), CacheFileVersion);
-			Event saveCasEvent(true);
-			m_server.AddWork([&](const WorkContext&) { m_storage.SaveCasTable(false, false); saveCasEvent.Set(); }, 1, TC("SaveCas"));
-			SaveNoLock();
+			DeleteEmptyBuckets();
 			saveCasEvent.IsSet();
-			m_logger.Detail(TC("  Save Done (%s)"), TimeToText(GetTime() - saveStart).str);
+			saveDbEvent.IsSet();
 		}
 
 		u64 oldestTime = oldest ? GetFileTimeAsTime(now - (m_creationTime + oldest)) : 0;
@@ -1646,6 +1684,7 @@ namespace uba
 		bool hasAllContent = true;
 		reader.SetPosition(outputStartOffset);
 		u64 left = outputCount;
+		UnorderedSet<CasKey> requestedOutputs;
 		while (left--)
 		{
 			u32 outputOffset = u32(reader.Read7BitEncoded());
@@ -1657,11 +1696,12 @@ namespace uba
 
 			bytesForOutput += Get7BitEncodedCount(casKeyOffset);
 
-			if (!m_storage.EnsureCasFile(casKey, nullptr))
-			{
-				writer.Write7BitEncoded(outputOffset);
-				hasAllContent = false;
-			}
+			if (!requestedOutputs.insert(casKey).second)
+				continue;
+			if (m_storage.EnsureCasFile(casKey, nullptr))
+				continue;
+			writer.Write7BitEncoded(outputOffset);
+			hasAllContent = false;
 		}
 
 		// Write outputs into buffer for entry
diff --git a/Engine/Source/Programs/UnrealBuildAccelerator/Common/Private/UbaFileAccessor.cpp b/Engine/Source/Programs/UnrealBuildAccelerator/Common/Private/UbaFileAccessor.cpp
--- a/Engine/Source/Programs/UnrealBuildAccelerator/Common/Private/UbaFileAccessor.cpp
+++ b/Engine/Source/Programs/UnrealBuildAccelerator/Common/Private/UbaFileAccessor.cpp
@@ -88,7 +88,10 @@ namespace uba
 			if (m_fileHandle != InvalidFileHandle)
 			{
 				if (retryCount)
-					m_logger.Info(TC("Had to retry %u times to open file %s for write (because is was being used%s)"), retryCount, realFileName, additionalInfo.data);
+				{
+					LogEntryType logType = retryCount > 10 ? LogEntryType_Warning : LogEntryType_Info;
+					m_logger.Logf(logType, TC("Had to retry for %u seconds to open file %s for write (because it was being used%s)"), retryCount/2, realFileName, additionalInfo.data);
+				}
 				break;
 			}
 			u32 lastError = GetLastError();
@@ -99,7 +102,7 @@ namespace uba
 				if (retryCount == 1)
 					GetProcessHoldingFile(additionalInfo, m_fileName);
 
-				if (retryCount < 10)
+				if (retryCount < 40)
 				{
 					Sleep(500);
 					++retryCount;
@@ -107,7 +110,7 @@ namespace uba
 				}
 			}
 			#endif
-			const tchar* retryText = retryCount ? TC(" after retrying 10 times") : TC("");
+			const tchar* retryText = retryCount ? TC(" after retrying for 20 seconds") : TC("");
 			return m_logger.Error(TC("ERROR opening file %s for write%s (%s%s)"), realFileName, retryText, LastErrorToText(lastError).data, additionalInfo.data);
 		}
 
@@ -466,8 +469,6 @@ namespace uba
 			m_writeThroughBuffer = nullptr;
 		}
 
-
-
 		#if !PLATFORM_WINDOWS
 		if (m_tempPath)
 			realFileName = tempFile.Append(m_tempPath).Append("Temp_").AppendValue(m_tempFileIndex).data;
diff --git a/Engine/Source/Programs/UnrealBuildAccelerator/Common/Public/UbaCacheServer.h b/Engine/Source/Programs/UnrealBuildAccelerator/Common/Public/UbaCacheServer.h
--- a/Engine/Source/Programs/UnrealBuildAccelerator/Common/Public/UbaCacheServer.h
+++ b/Engine/Source/Programs/UnrealBuildAccelerator/Common/Public/UbaCacheServer.h
@@ -68,8 +68,10 @@ namespace uba
 
 		bool RunMaintenanceInternal(const Function<bool()>& shouldExit);
 		bool LoadBucket(Bucket& bucket, BinaryReader& reader, u32 databaseVersion, LoadStats& outStats);
-		bool SaveBucket(u64 bucketId, Bucket& bucket, Vector<u8>& temp);
+		bool SaveBucket(Bucket& bucket, Vector<u8>& temp);
 		bool SaveNoLock();
+		bool SaveDbNoLock();
+		bool DeleteEmptyBuckets();
 		void OnDisconnected(u32 clientId);
 
 		struct Connection;
