[UBA]
* Fixed linux compile warning

diff --git a/Engine/Source/Programs/UnrealBuildAccelerator/Common/Private/UbaCacheServer.cpp b/Engine/Source/Programs/UnrealBuildAccelerator/Common/Private/UbaCacheServer.cpp
--- a/Engine/Source/Programs/UnrealBuildAccelerator/Common/Private/UbaCacheServer.cpp
+++ b/Engine/Source/Programs/UnrealBuildAccelerator/Common/Private/UbaCacheServer.cpp
@@ -153,7 +153,7 @@ namespace uba
 
 	struct CacheServer::Bucket
 	{
-		Bucket(u64 id, u32 version) : m_id(id), m_pathTable(IsCaseInsensitive(id), 0, 0, version) {}
+		Bucket(u64 id, u32 version) : m_pathTable(IsCaseInsensitive(id), 0, 0, version), m_id(id) {}
 		ReaderWriterLock m_cacheEntryLookupLock;
 		UnorderedMap<CasKey, CacheEntries> m_cacheEntryLookup;
 
