[UBA]
* Fixed a bug where all buckets were always saved

diff --git a/Engine/Source/Programs/UnrealBuildAccelerator/Common/Private/UbaCacheServer.cpp b/Engine/Source/Programs/UnrealBuildAccelerator/Common/Private/UbaCacheServer.cpp
--- a/Engine/Source/Programs/UnrealBuildAccelerator/Common/Private/UbaCacheServer.cpp
+++ b/Engine/Source/Programs/UnrealBuildAccelerator/Common/Private/UbaCacheServer.cpp
@@ -1121,7 +1121,7 @@ namespace uba
 
 			auto saveGuard = MakeGuard([&]()
 				{
-					if (shouldSave && !bucket.m_cacheEntryLookup.empty())
+					if (shouldSave && bucket.needsSave && !bucket.m_cacheEntryLookup.empty())
 					{
 						Vector<u8> temp;
 						if (SaveBucket(bucket, temp))
