[UBA]
* Moved ping message to trace thread instead of request-process thread on client
* Changed error about not finding process when receiving files to be a warning.

diff --git a/Engine/Source/Programs/UnrealBuildAccelerator/Common/Private/UbaSession.cpp b/Engine/Source/Programs/UnrealBuildAccelerator/Common/Private/UbaSession.cpp
--- a/Engine/Source/Programs/UnrealBuildAccelerator/Common/Private/UbaSession.cpp
+++ b/Engine/Source/Programs/UnrealBuildAccelerator/Common/Private/UbaSession.cpp
@@ -1975,8 +1975,7 @@ namespace uba
 		StringBuffer<> systemInfo;
 		GetSystemInfo(systemInfo);
 		m_trace.SessionAdded(0, {}, buf, systemInfo);
-		m_traceThreadEvent.Create(true);
-		m_traceThread.Start([this]() { ThreadTraceLoop(); return 0; }, TC("UbaTraceLoop"));
+		StartTraceThread();
 	}
 
 	bool Session::StopTrace(const tchar* writeFile)
@@ -1985,6 +1984,12 @@ namespace uba
 		return m_trace.StopWrite(writeFile);
 	}
 
+	void Session::StartTraceThread()
+	{
+		m_traceThreadEvent.Create(true);
+		m_traceThread.Start([this]() { ThreadTraceLoop(); return 0; }, TC("UbaTraceLoop"));
+	}
+
 	void Session::StopTraceThread()
 	{
 		m_traceThreadEvent.Set();
diff --git a/Engine/Source/Programs/UnrealBuildAccelerator/Common/Private/UbaSessionClient.cpp b/Engine/Source/Programs/UnrealBuildAccelerator/Common/Private/UbaSessionClient.cpp
--- a/Engine/Source/Programs/UnrealBuildAccelerator/Common/Private/UbaSessionClient.cpp
+++ b/Engine/Source/Programs/UnrealBuildAccelerator/Common/Private/UbaSessionClient.cpp
@@ -1345,6 +1345,10 @@ namespace uba
 			}
 			StartTrace(nullptr, 256);
 		}
+		else
+		{
+			StartTraceThread(); // We still use trace thread for ping messages
+		}
 
 		BuildEnvironmentVariables(reader);
 
@@ -2001,8 +2005,6 @@ namespace uba
 				firstCall = false;
 			}
 
-			SendPing(memAvail, memTotal);
-
 			//SendUpdateNameToHashTable(); // It is always nice to populate this at certain cadence since it might speed up running processes queries
 
 			m_waitToSendEvent.IsSet(waitTimeoutMs);
@@ -2051,13 +2053,13 @@ namespace uba
 		m_client.FlushWork();
 
 
+		StopTraceThread();
+
 		if (m_shouldSendTraceToServer)
 		{
 			// Can't disable this, it can cause race conditions
 			//m_client.SetWorkTracker(nullptr);
 
-			StopTraceThread();
-
 			StackBinaryWriter<SendMaxSize> writer;
 			WriteSummary(writer, [&](Logger& logger)
 				{
@@ -2582,6 +2584,12 @@ namespace uba
 
 	void SessionClient::TraceSessionUpdate()
 	{
+		if (m_memAvail)
+			SendPing(m_memAvail, m_memTotal);
+
+		if (!m_trace.IsWriting())
+			return;
+
 		u64 send;
 		u64 recv;
 		if (auto backend = m_client.GetFirstConnectionBackend())
diff --git a/Engine/Source/Programs/UnrealBuildAccelerator/Common/Private/UbaSessionServer.cpp b/Engine/Source/Programs/UnrealBuildAccelerator/Common/Private/UbaSessionServer.cpp
--- a/Engine/Source/Programs/UnrealBuildAccelerator/Common/Private/UbaSessionServer.cpp
+++ b/Engine/Source/Programs/UnrealBuildAccelerator/Common/Private/UbaSessionServer.cpp
@@ -1028,6 +1028,8 @@ namespace uba
 			m_logger.Logf(logType, TC("Failed to retrieve cas for %s from client %u (Needed to write %s)"), CasKeyString(casKey).str, connectionInfo.GetId(), destination.data);
 		}
 
+		auto writeResponse = MakeGuard([&]() { writer.WriteBool(success); });
+
 		bool shouldWriteToDisk = ShouldWriteToDisk(destination);
 		if (success)
 		{
@@ -1039,7 +1041,6 @@ namespace uba
 					m_logger.Error(TC("Failed to copy cas from %s to %s"), CasKeyString(casKey).str, logPath.data);
 				else if (!m_storage.DropCasFile(casKey, false, logPath.data))
 					m_logger.Error(TC("Failed to drop cas %s"), CasKeyString(casKey).str);
-				writer.WriteBool(true);
 				return true;
 			}
 
@@ -1060,7 +1061,6 @@ namespace uba
 				ubaPath.Append(TC(".uba"));
 				m_storage.CopyOrLink(casKey, ubaPath.data, attributes);
 				m_storage.DropCasFile(casKey, false, ubaPath.data);
-				writer.WriteBool(true);
 				return true;
 			}
 
@@ -1070,7 +1070,10 @@ namespace uba
 
 				ProcessHandle h = GetProcess(processId);
 				if (!h.IsValid())
-					return m_logger.Error(TC("Failed to find process for id %u when receiving SendFileToServer message"), processId);
+				{
+					success = false;
+					return m_logger.Warning(TC("Failed to find process for id %u when receiving SendFileToServer message"), processId);
+				}
 
 				RootsHandle rootsHandle = h.GetStartInfo().rootsHandle;
 
@@ -1102,7 +1105,6 @@ namespace uba
 				m_receivedFiles.try_emplace(destinationKey, casKey);
 			}
 		}
-		writer.WriteBool(success); 
 
 		if (success)
 		{
diff --git a/Engine/Source/Programs/UnrealBuildAccelerator/Common/Public/UbaSession.h b/Engine/Source/Programs/UnrealBuildAccelerator/Common/Public/UbaSession.h
--- a/Engine/Source/Programs/UnrealBuildAccelerator/Common/Public/UbaSession.h
+++ b/Engine/Source/Programs/UnrealBuildAccelerator/Common/Public/UbaSession.h
@@ -109,6 +109,7 @@ namespace uba
 		void PrintProcessStats(ProcessStats& stats, const tchar* logName);
 		void StartTrace(const tchar* traceName, u32 traceReserveSizeMb);
 		bool StopTrace(const tchar* writeFile);
+		void StartTraceThread();
 		void StopTraceThread();
 		u32 GetFileMappingSize();
 		u32 GetMemoryMapAlignment(const StringView& fileName) const;
