[UBA]
* Removed code where pch obj files were not compressed when obj file compression was turned on
* Changed so UbaCli is always supporting reading compressed obj files (it is used by live coding)

diff --git a/Engine/Source/Programs/UnrealBuildAccelerator/Cli/Private/UbaCli.cpp b/Engine/Source/Programs/UnrealBuildAccelerator/Cli/Private/UbaCli.cpp
--- a/Engine/Source/Programs/UnrealBuildAccelerator/Cli/Private/UbaCli.cpp
+++ b/Engine/Source/Programs/UnrealBuildAccelerator/Cli/Private/UbaCli.cpp
@@ -868,7 +868,8 @@ namespace uba
 		//info.shouldWriteToDisk = shouldWriteToDisk;
 		info.rootDir = g_rootDir.data;
 		//info.traceName.Append(TC("TESTTRACE"));
-		//info.storeObjFilesCompressed = true;
+		//info.storeObjFilesCompressed = false;
+		info.readObjFilesCompressed = true;
 		//info.extractObjFilesSymbols = true;
 		#if UBA_DEBUG_LOG_ENABLED
 		info.remoteLogEnabled = true;
diff --git a/Engine/Source/Programs/UnrealBuildAccelerator/Common/Private/UbaProcess.cpp b/Engine/Source/Programs/UnrealBuildAccelerator/Common/Private/UbaProcess.cpp
--- a/Engine/Source/Programs/UnrealBuildAccelerator/Common/Private/UbaProcess.cpp
+++ b/Engine/Source/Programs/UnrealBuildAccelerator/Common/Private/UbaProcess.cpp
@@ -1710,7 +1710,7 @@ namespace uba
 			payload.trackInputs = m_startInfo.trackInputs;
 			payload.useCustomAllocator = allowCustomAllocator && m_startInfo.useCustomAllocator && m_startInfo.rules->AllowMiMalloc();
 			payload.isRunningWine = IsRunningWine();
-			payload.storeObjFilesCompressed = m_session.m_storeObjFilesCompressed;
+			payload.readObjFilesCompressed = m_session.m_readObjFilesCompressed;
 			payload.uiLanguage = m_startInfo.uiLanguage;
 			if (*m_startInfo.logFile)
 			{
diff --git a/Engine/Source/Programs/UnrealBuildAccelerator/Common/Private/UbaSession.cpp b/Engine/Source/Programs/UnrealBuildAccelerator/Common/Private/UbaSession.cpp
--- a/Engine/Source/Programs/UnrealBuildAccelerator/Common/Private/UbaSession.cpp
+++ b/Engine/Source/Programs/UnrealBuildAccelerator/Common/Private/UbaSession.cpp
@@ -1155,6 +1155,7 @@ namespace uba
 		UBA_ASSERTF(m_shouldWriteToDisk || m_allowMemoryMaps, TC("Can't disable both should write to disk and allow memory maps"));
 
 		m_storeObjFilesCompressed = info.storeObjFilesCompressed && IsWindows; // Non-windows not implemented (yet)
+		m_readObjFilesCompressed = m_storeObjFilesCompressed || (info.readObjFilesCompressed && IsWindows);
 		m_allowLocalDetour = info.allowLocalDetour;
 
 		m_detailedTrace = info.detailedTrace;
diff --git a/Engine/Source/Programs/UnrealBuildAccelerator/Common/Public/UbaSession.h b/Engine/Source/Programs/UnrealBuildAccelerator/Common/Public/UbaSession.h
--- a/Engine/Source/Programs/UnrealBuildAccelerator/Common/Public/UbaSession.h
+++ b/Engine/Source/Programs/UnrealBuildAccelerator/Common/Public/UbaSession.h
@@ -217,6 +217,7 @@ namespace uba
 		bool m_traceChildProcesses;
 		bool m_logToFile;
 		bool m_storeObjFilesCompressed;
+		bool m_readObjFilesCompressed;
 		bool m_allowLocalDetour;
 
 		u64 m_keepOutputFileMemoryMapsThreshold;
diff --git a/Engine/Source/Programs/UnrealBuildAccelerator/Common/Public/UbaSessionCreateInfo.h b/Engine/Source/Programs/UnrealBuildAccelerator/Common/Public/UbaSessionCreateInfo.h
--- a/Engine/Source/Programs/UnrealBuildAccelerator/Common/Public/UbaSessionCreateInfo.h
+++ b/Engine/Source/Programs/UnrealBuildAccelerator/Common/Public/UbaSessionCreateInfo.h
@@ -35,6 +35,7 @@ namespace uba
 		bool detailedTrace = false;				// Enable detailed trace to include jobs, individual file I/O etc in trace dump
 		bool traceChildProcesses = false;		// Trace and visualize child processes so they can be seen in visualizer
 		bool storeObjFilesCompressed = false;	// Compiler will write .obj files compressed to disk and linker will decompress when reading
+		bool readObjFilesCompressed = false;	// Set to true to support reading compressed .obj files. This flag will not compress new obj files
 		bool allowLocalDetour = true;			// Allow local processes to be detoured. If this is false it is up to outside logic to register all created/deleted files
 		bool extractObjFilesSymbols = false;	// Will extract import/export symbols and clear export table from obj file.
 		bool treatTempDirAsEmpty = true;		// If this is true, directory table will always see temp as empty. Use sub directories of temp instead
diff --git a/Engine/Source/Programs/UnrealBuildAccelerator/Core/Private/UbaApplicationRules.cpp b/Engine/Source/Programs/UnrealBuildAccelerator/Core/Private/UbaApplicationRules.cpp
--- a/Engine/Source/Programs/UnrealBuildAccelerator/Core/Private/UbaApplicationRules.cpp
+++ b/Engine/Source/Programs/UnrealBuildAccelerator/Core/Private/UbaApplicationRules.cpp
@@ -63,7 +63,7 @@ namespace uba
 
 	bool IsCompressedObjFile(const StringView& fileName)
 	{
-		return fileName.EndsWith(TC(".obj")) && !fileName.Contains(TC("pch."));
+		return fileName.EndsWith(TC(".obj"));
 	}
 
 	class ApplicationRulesClExe : public ApplicationRulesVC
@@ -106,8 +106,6 @@ namespace uba
 
 		virtual bool StoreFileCompressed(const StringView& fileName) const override
 		{
-			// We don't compress pch obj files because they are used by live-coding.
-			// TODO: Add support in live coding module to load compressed obj files instead
 			return IsCompressedObjFile(fileName);
 		}
 
diff --git a/Engine/Source/Programs/UnrealBuildAccelerator/Core/Public/UbaDetoursPayload.h b/Engine/Source/Programs/UnrealBuildAccelerator/Core/Public/UbaDetoursPayload.h
--- a/Engine/Source/Programs/UnrealBuildAccelerator/Core/Public/UbaDetoursPayload.h
+++ b/Engine/Source/Programs/UnrealBuildAccelerator/Core/Public/UbaDetoursPayload.h
@@ -26,7 +26,7 @@ namespace uba
 		bool allowKeepFilesInMemory = IsWindows;
 		bool allowOutputFiles = IsWindows;
 		bool suppressLogging = false;
-		bool storeObjFilesCompressed = false;
+		bool readObjFilesCompressed = false;
 		int uiLanguage = 0;
 		StringBuffer<256> logFile;
 	};
diff --git a/Engine/Source/Programs/UnrealBuildAccelerator/Detours/Private/Windows/UbaDetoursFunctionsWin.cpp b/Engine/Source/Programs/UnrealBuildAccelerator/Detours/Private/Windows/UbaDetoursFunctionsWin.cpp
--- a/Engine/Source/Programs/UnrealBuildAccelerator/Detours/Private/Windows/UbaDetoursFunctionsWin.cpp
+++ b/Engine/Source/Programs/UnrealBuildAccelerator/Detours/Private/Windows/UbaDetoursFunctionsWin.cpp
@@ -1108,7 +1108,7 @@ void Init(const DetoursPayload& payload, u64 startTime)
 	g_stats.attach.time += GetTime() - startTime;
 	g_stats.attach.count = 1;
 
-	if (payload.storeObjFilesCompressed && g_rules->ShouldDecompressFiles(StringView()))
+	if (payload.readObjFilesCompressed && g_rules->ShouldDecompressFiles(StringView()))
 	{
 		bool preprocess = true;//g_rules->index == SpecialRulesIndex_LdExePlatform1;
 		g_fileDecompressor.Start(cmdLine, preprocess);
