[UBA]
* Added support for creating trace file around maintenance logic on UbaCacheServer
* Added support for batches to ParallelFor
* Changed so bucket cache entry parallel for loop is taking batches of 100
* Changed so dropcas is batched up instead of taking one at the time
* Added TCV("") to replace AsView(TC(""))

diff --git a/Engine/Source/Programs/UnrealBuildAccelerator/Common/Private/UbaCacheClient.cpp b/Engine/Source/Programs/UnrealBuildAccelerator/Common/Private/UbaCacheClient.cpp
--- a/Engine/Source/Programs/UnrealBuildAccelerator/Common/Private/UbaCacheClient.cpp
+++ b/Engine/Source/Programs/UnrealBuildAccelerator/Common/Private/UbaCacheClient.cpp
@@ -761,7 +761,7 @@ namespace uba
 				if (!casKeyOffsets.empty())
 				{
 					Atomic<bool> fetchSuccess = true;
-					m_storage.m_workManager->ParallelFor(u32(casKeyOffsets.size() - 1), casKeyOffsets, [&](auto& it)
+					m_storage.m_workManager->ParallelFor(u32(casKeyOffsets.size() - 1), casKeyOffsets, [&](const WorkContext&, auto& it)
 						{
 							u32 casKeyOffset = *it;
 							if (!FetchCasTable(tws, bucket, cacheStats, casKeyOffset))
diff --git a/Engine/Source/Programs/UnrealBuildAccelerator/Common/Private/UbaCacheServer.cpp b/Engine/Source/Programs/UnrealBuildAccelerator/Common/Private/UbaCacheServer.cpp
--- a/Engine/Source/Programs/UnrealBuildAccelerator/Common/Private/UbaCacheServer.cpp
+++ b/Engine/Source/Programs/UnrealBuildAccelerator/Common/Private/UbaCacheServer.cpp
@@ -9,6 +9,7 @@
 #include "UbaHashMap.h"
 #include "UbaNetworkServer.h"
 #include "UbaStorageServer.h"
+#include "UbaTrace.h"
 #include <algorithm>
 
 //#include <oodle2.h>
@@ -18,6 +19,8 @@
 // - Sort buckets by last maintenance time to make sure the long ones always get a slot first
 // - Change so save happens when bucket is done in same work to minimize latency for the long ones
 
+#define UBA_TRACK_WORK_SCOPE(name) //TrackWorkScope tws(m_server, TCV(name));
+#define UBA_CACHE_SERVER_ADD_HINT(name, start) //wc.tracker.AddHint(TCV(name), start);
 
 #if PLATFORM_WINDOWS
 #define UBA_FORCEINLINE __forceinline
@@ -186,11 +189,11 @@ namespace uba
 	{
 		switch (type)
 		{
-			#define UBA_CACHE_MESSAGE(x) case CacheMessageType_##x: return AsView(TC("")#x);
+			#define UBA_CACHE_MESSAGE(x) case CacheMessageType_##x: return TCV(#x);
 			UBA_CACHE_MESSAGES
 			#undef UBA_CACHE_MESSAGE
 		default:
-			return AsView(TC("Unknown")); // Should never happen
+			return TCV("Unknown"); // Should never happen
 		}
 	}
 
@@ -299,7 +302,7 @@ namespace uba
 					GetBucket(id, TC("Loading"), false);
 				});
 
-			m_server.ParallelFor(GetBucketWorkerCount(), m_buckets, [&](auto& it)
+			m_server.ParallelFor(GetBucketWorkerCount(), m_buckets, [&](const WorkContext&, auto& it)
 				{
 					u64 key = it->first;
 					Bucket& bucket = it->second;
@@ -532,7 +535,7 @@ namespace uba
 
 		DeleteEmptyBuckets();
 
-		m_server.ParallelFor(GetBucketWorkerCount(), m_buckets, [&, temp = Vector<u8>()](auto& it) mutable
+		m_server.ParallelFor(GetBucketWorkerCount(), m_buckets, [&, temp = Vector<u8>()](const WorkContext&, auto& it) mutable
 			{
 				Bucket& bucket = it->second;
 				if (!bucket.needsSave)
@@ -660,6 +663,30 @@ namespace uba
 		bool forceAllSteps = m_forceAllSteps;
 		m_forceAllSteps = false;
 
+		Trace trace(m_logger.m_writer);
+		if (forceAllSteps)
+		{
+			u64 traceReserveSize = 128ull*1024*1024;
+			trace.StartWrite(nullptr, traceReserveSize);
+			trace.SessionAdded(0, 0, TCV("CacheServer"), {});
+			trace.ProcessAdded(0, 0, TCV("Maintenance"), {});
+
+			m_server.SetWorkTracker(&trace);
+		}
+		auto endTrace = MakeGuard([&]()
+			{
+				m_server.SetWorkTracker(nullptr);
+				StackBinaryWriter<1024> writer;
+				ProcessStats().Write(writer);
+				SessionStats().Write(writer);
+				StorageStats().Write(writer);
+				KernelStats().Write(writer);
+				trace.ProcessExited(0, 0, writer.GetData(), writer.GetPosition(), {});
+				StringBuffer<> traceFile(m_rootDir);
+				traceFile.Append(TC("UbaCacheServer.uba"));
+				trace.StopWrite(traceFile.data);
+			});
+
 		u32 addsSinceMaintenance = m_addsSinceMaintenance;
 		bool entriesAdded = addsSinceMaintenance != 0;
 		m_addsSinceMaintenance = 0;
@@ -700,6 +727,7 @@ namespace uba
 
 		u64 totalCasCount;
 		{
+			UBA_TRACK_WORK_SCOPE("CollectCas");
 			u64 collectCasStartTime = GetTime();
 
 			u32 removedNonExisting = 0;
@@ -780,6 +808,7 @@ namespace uba
 				auto& context = *bucket.m_maintenanceContext;
 				if (!context.isInitialized)
 				{
+					UBA_TRACK_WORK_SCOPE("InitContext");
 					if (!context.memoryBlock.Init(m_maintenanceReserveSize, nullptr, true)) // Try to use large blocks
 						context.memoryBlock.Init(m_maintenanceReserveSize, nullptr, false);
 					context.deletedOffsets.Init(context.memoryBlock, bucket.m_casKeyTable.GetSize(), TC("DeletedOffsets"));
@@ -796,8 +825,10 @@ namespace uba
 			totalEntryCount = 0;
 
 			// Traverse all buckets in parallel
-			m_server.ParallelFor(workerCountToUseForBuckets, sortedBuckets, [&](auto& it)
+			m_server.ParallelFor(workerCountToUseForBuckets, sortedBuckets, [&](const WorkContext&, auto& it)
 			{
+				UBA_TRACK_WORK_SCOPE("TraverseBucket");
+
 				Bucket& bucket = **it;
 				auto context = bucket.m_maintenanceContext;
 				if (!context)
@@ -854,8 +885,9 @@ namespace uba
 				bucket.totalEntrySize = 0;
 
 				// Loop through all cache entries in parallel
-				m_server.ParallelFor(workerCountToUse, bucket.m_cacheEntryLookup, [&, touchedCas = Vector<Atomic<bool>*>()](auto& li) mutable
+				m_server.ParallelFor<100>(workerCountToUse, bucket.m_cacheEntryLookup, [&, touchedCas = Vector<Atomic<bool>*>()](const WorkContext&, auto& li) mutable
 				{
+					UBA_TRACK_WORK_SCOPE("TraverseEntry");
 					CacheEntries& entries = li->second;
 
 					bool checkInputsForDeletes = false;
@@ -1034,7 +1066,7 @@ namespace uba
 							bucket.totalEntryCount += entries.entries.size();
 						}
 					}
-				});
+				}, TCV("ParallelForEntries1"));
 
 				{
 					SCOPED_FUTEX(globalStatsLock, l);
@@ -1048,7 +1080,7 @@ namespace uba
 					bucket.m_cacheEntryLookup.erase(key);
 
 				totalEntryCount += bucket.totalEntryCount;
-			}, AsView(TC("")), true);
+			}, TCV("ParallelForBucket1"), true);
 
 			// Reset deleted cas files and update it again..
 			deletedCasFiles.clear();
@@ -1072,11 +1104,26 @@ namespace uba
 			// Add drop cas as work so it can run in the background
 			if (allowSave)
 			{
+				Vector<CasKey> casKeysBatch;
+				auto createBatchWork = [&]()
+					{
+						activeDropCount += casKeysBatch.size();
+						m_server.AddWork([&, ckb = casKeysBatch](const WorkContext&)
+							{
+								for (auto& key : ckb)
+									m_storage.DropCasFile(key, true, TC(""));
+								activeDropCount -= ckb.size();
+							}, 1, TC("DropCasKeysBatch"));
+						casKeysBatch.resize(0);
+					};
 				for (auto& casKey : deletedCasFiles)
 				{
-					++activeDropCount;
-					m_server.AddWork([&, key = casKey](const WorkContext&) { m_storage.DropCasFile(key, true, TC("")); --activeDropCount; }, 1, TC(""));
+					casKeysBatch.push_back(casKey);
+					if (casKeysBatch.size() > 100)
+						createBatchWork();
 				}
+				if (!casKeysBatch.empty())
+					createBatchWork();
 			}
 			++deleteIteration;
 		}
@@ -1113,7 +1160,7 @@ namespace uba
 
 		u64 maxCommittedMemory = 0;
 
-		m_server.ParallelFor(workerCountToUseForBuckets, sortedBuckets, [&](auto& it)
+		m_server.ParallelFor(workerCountToUseForBuckets, sortedBuckets, [&](const WorkContext& wc, auto& it)
 		{
 			u64 bucketStartTime = GetTime();
 
@@ -1123,9 +1170,11 @@ namespace uba
 				{
 					if (shouldSave && bucket.needsSave && !bucket.m_cacheEntryLookup.empty())
 					{
+						u64 saveStart = GetTime();(void)saveStart;
 						Vector<u8> temp;
 						if (SaveBucket(bucket, temp))
 							bucket.needsSave = false;
+						UBA_CACHE_SERVER_ADD_HINT("SaveBucket", saveStart);
 					}
 				});
 
@@ -1168,6 +1217,7 @@ namespace uba
 			}
 			u64 usedCasKeyOffsetsCount = usedCasKeyOffsets.CountSetBits();
 
+			UBA_CACHE_SERVER_ADD_HINT("CollectUsedCasKeys", collectUsedCasKeysStart);
 			m_logger.Detail(TC("    Bucket %u Collected %s used caskeys. (%s)"), bucket.index, CountToText(usedCasKeyOffsetsCount).str, TimeToText(GetTime() - collectUsedCasKeysStart).str);
 
 			u64 recreatePathTableStart = GetTime();
@@ -1213,6 +1263,7 @@ namespace uba
 				});
 				pathTable.Swap(newPathTable);
 			}
+			UBA_CACHE_SERVER_ADD_HINT("RecreatePathTable", recreatePathTableStart);
 			m_logger.Detail(TC("    Bucket %u Recreated path table. %s -> %s (%s)"), bucket.index, BytesToText(oldSize).str, BytesToText(pathTable.GetSize()).str, TimeToText(GetTime() - recreatePathTableStart).str);
 
 			// Build new caskey table based on used offsets
@@ -1238,6 +1289,7 @@ namespace uba
 				});
 				casKeyTable.Swap(newCasKeyTable);
 			}
+			UBA_CACHE_SERVER_ADD_HINT("RecreateCasKeyTable", recreateCasKeyTableStart);
 			m_logger.Detail(TC("    Bucket %u Recreated caskey table. %s -> %s (%s)"), bucket.index, BytesToText(oldSize).str, BytesToText(bucket.m_casKeyTable.GetSize()).str, TimeToText(GetTime() - recreateCasKeyTableStart).str);
 
 			if (casKeyTable.GetSize() >= m_bucketCasTableMaxSize) // Still overflowing?
@@ -1248,10 +1300,10 @@ namespace uba
 				// Update all casKeyOffsets
 				u64 updateEntriesStart = GetTime();
 
-				m_server.ParallelFor(workerCountToUse, bucket.m_cacheEntryLookup, [&, temp = Vector<u32>(), temp2 = Vector<u8>(), temp3 = Vector<u8>()](auto& it) mutable
+				m_server.ParallelFor<100>(workerCountToUse, bucket.m_cacheEntryLookup, [&, temp = Vector<u32>(), temp2 = Vector<u8>(), temp3 = Vector<u8>()](const WorkContext&, auto& it) mutable
 					{
 						it->second.UpdateEntries(m_logger, oldToNewCasKeyOffset, temp, temp2, temp3);
-					});
+					}, TCV("ParallelForEntries2"));
 
 				#if 0
 				u8* mem = bucket.m_pathTable.GetMemory();
@@ -1268,7 +1320,7 @@ namespace uba
 				}
 				#endif
 
-
+				UBA_CACHE_SERVER_ADD_HINT("UpdateCacheEntries", updateEntriesStart);
 				m_logger.Detail(TC("    Bucket %u Updated cache entries with new tables (%s)"), bucket.index, TimeToText(GetTime() - updateEntriesStart).str);
 			}
 
@@ -1282,7 +1334,7 @@ namespace uba
 			if (bucket.oldestUsedTime)
 				longestUnusedSeconds = GetFileTimeAsSeconds(now - m_creationTime - bucket.oldestUsedTime);
 			m_logger.Info(TC("    Bucket %u Done (%s). CacheEntries: %s (%s) PathTable: %s CasTable: %s LongestUnused: %s"), bucket.index, TimeToText(GetTime() - bucketStartTime).str, CountToText(bucket.totalEntryCount).str, BytesToText(bucket.totalEntrySize).str, BytesToText(bucket.m_pathTable.GetSize()).str, BytesToText(bucket.m_casKeyTable.GetSize()).str, TimeToText(MsToTime(longestUnusedSeconds*1000), true).str);
-		}, AsView(TC("")), true);
+		}, TCV("ParallelForBucket2"), true);
 
 		// Need to make sure all cas entries are dropped before saving cas table
 		u64 dropStartTime = GetTime();
@@ -1980,7 +2032,7 @@ namespace uba
 
 		if (command.Equals(TC("content")))
 		{
-			writeLine(AsView(TC("UbaCache server summary")));
+			writeLine(TCV("UbaCache server summary"));
 
 			StringBufferBase& filterString = additionalInfo;
 
@@ -2069,7 +2121,7 @@ namespace uba
 		}
 		else if (command.Equals(TC("status")))
 		{
-			writeLine(AsView(TC("UbaCacheServer status")));
+			writeLine(TCV("UbaCacheServer status"));
 			writeLine(line.Clear().Appendf(TC("  CreationTime: %s ago"), TimeToText(GetFileTimeAsTime(GetSystemTimeAsFileTime() - m_creationTime), true).str));
 			writeLine(line.Clear().Appendf(TC("  UpTime: %s"), TimeToText(GetTime() - m_bootTime, true).str));
 			writeLine(line.Clear().Appendf(TC("  Longest maintenance: %s"), TimeToText(m_longestMaintenance).str));
@@ -2328,7 +2380,7 @@ namespace uba
 				bucket.m_pathTable.TraversePaths([&](const StringView& path)
 					{
 						if (!sortedPaths.emplace(path.data).second)
-							writeLine(AsView(TC("EEEEERRRRRROOOORRRR!!!!!!")));
+							writeLine(TCV("EEEEERRRRRROOOORRRR!!!!!!"));
 					});
 
 				for (auto& path : sortedPaths)
diff --git a/Engine/Source/Programs/UnrealBuildAccelerator/Common/Private/UbaNetworkServer.cpp b/Engine/Source/Programs/UnrealBuildAccelerator/Common/Private/UbaNetworkServer.cpp
--- a/Engine/Source/Programs/UnrealBuildAccelerator/Common/Private/UbaNetworkServer.cpp
+++ b/Engine/Source/Programs/UnrealBuildAccelerator/Common/Private/UbaNetworkServer.cpp
@@ -582,7 +582,11 @@ namespace uba
 				server.m_additionalWork.pop_front();
 				lock.Leave();
 
+				#if UBA_TRACK_WORK
 				TrackWorkScope tws(server, work.desc, ColorWork);
+				#else
+				TrackWorkScope tws;
+				#endif
 				work.func({tws});
 			}
 
@@ -1007,7 +1011,11 @@ namespace uba
 		m_additionalWork.pop_front();
 		lock.Leave();
 
+		#if UBA_TRACK_WORK
 		TrackWorkScope tws(*this, work.desc, ColorWork);
+		#else
+		TrackWorkScope tws;
+		#endif
 
 		work.func({tws});
 
diff --git a/Engine/Source/Programs/UnrealBuildAccelerator/Common/Private/UbaSession.cpp b/Engine/Source/Programs/UnrealBuildAccelerator/Common/Private/UbaSession.cpp
--- a/Engine/Source/Programs/UnrealBuildAccelerator/Common/Private/UbaSession.cpp
+++ b/Engine/Source/Programs/UnrealBuildAccelerator/Common/Private/UbaSession.cpp
@@ -2805,7 +2805,7 @@ namespace uba
 
 		Atomic<bool> success = true;
 		auto span = std::span(files, fileCount);
-		m_workManager->ParallelFor(fileCount - 1, span, [&](auto& it)
+		m_workManager->ParallelFor(fileCount - 1, span, [&](const WorkContext&, auto& it)
 			{
 				if (!writeFile(**it))
 					success = false;
diff --git a/Engine/Source/Programs/UnrealBuildAccelerator/Common/Private/UbaSessionClient.cpp b/Engine/Source/Programs/UnrealBuildAccelerator/Common/Private/UbaSessionClient.cpp
--- a/Engine/Source/Programs/UnrealBuildAccelerator/Common/Private/UbaSessionClient.cpp
+++ b/Engine/Source/Programs/UnrealBuildAccelerator/Common/Private/UbaSessionClient.cpp
@@ -356,7 +356,7 @@ namespace uba
 
 		Atomic<bool> success = true;
 
-		m_client.ParallelFor(u32(modules.size()), modules, [&](auto& it)
+		m_client.ParallelFor(u32(modules.size()), modules, [&](const WorkContext&, auto& it)
 			{
 				auto& m = *it;
 				TrackWorkScope tws(m_client, AsView(TC("FetchModule")), ColorWork);
diff --git a/Engine/Source/Programs/UnrealBuildAccelerator/Common/Private/UbaSessionServer.cpp b/Engine/Source/Programs/UnrealBuildAccelerator/Common/Private/UbaSessionServer.cpp
--- a/Engine/Source/Programs/UnrealBuildAccelerator/Common/Private/UbaSessionServer.cpp
+++ b/Engine/Source/Programs/UnrealBuildAccelerator/Common/Private/UbaSessionServer.cpp
@@ -352,6 +352,7 @@ namespace uba
 						Container(const tchar* b, u32 c) : ptr(b), count(c) {}
 						iterator begin() const { return iterator(ptr, 0); }
 						iterator end() const { return iterator(ptr, count); }
+						u32 size() { return count; }
 						const tchar* ptr;
 						u32 count;
 					};
@@ -361,7 +362,7 @@ namespace uba
 					Atomic<u32> keysIndex;
 					const TString& workingDir = remoteProcess->m_startInfo.workingDirStr;
 
-					m_server.ParallelFor(knownInputsCount, container, [&](auto& it)
+					m_server.ParallelFor(knownInputsCount, container, [&](const WorkContext&, auto& it)
 						{
 							//TrackWorkScope tws2(m_trace, AsView(TC("KnownInputs")), ColorWork);
 							StringBuffer<> fileName;
diff --git a/Engine/Source/Programs/UnrealBuildAccelerator/Common/Private/UbaWorkManager.cpp b/Engine/Source/Programs/UnrealBuildAccelerator/Common/Private/UbaWorkManager.cpp
--- a/Engine/Source/Programs/UnrealBuildAccelerator/Common/Private/UbaWorkManager.cpp
+++ b/Engine/Source/Programs/UnrealBuildAccelerator/Common/Private/UbaWorkManager.cpp
@@ -5,8 +5,6 @@
 #include "UbaPlatform.h"
 #include "UbaThread.h"
 
-#define UBA_TRACK_WORK 0
-
 namespace uba
 {
 	u32 WorkManager::TrackWorkStart(const StringView& desc, const Color& color)
diff --git a/Engine/Source/Programs/UnrealBuildAccelerator/Common/Public/UbaWorkManager.h b/Engine/Source/Programs/UnrealBuildAccelerator/Common/Public/UbaWorkManager.h
--- a/Engine/Source/Programs/UnrealBuildAccelerator/Common/Public/UbaWorkManager.h
+++ b/Engine/Source/Programs/UnrealBuildAccelerator/Common/Public/UbaWorkManager.h
@@ -6,6 +6,8 @@
 #include "UbaStringBuffer.h"
 #include "UbaTimer.h"
 
+#define UBA_TRACK_WORK 1
+
 namespace uba
 {
 	struct StringView;
@@ -44,8 +46,8 @@ namespace uba
 		void SetWorkTracker(WorkTracker* workTracker) { m_workTracker = workTracker; }
 		WorkTracker* GetWorkTracker() { return m_workTracker; }
 
-		template<typename TContainer, typename TFunc>
-		void ParallelFor(u32 workCount, TContainer& container, const TFunc& func, const StringView& description = AsView(TC("")), bool highPriority = false);
+		template<u32 BatchSize = 1, typename TContainer, typename TFunc>
+		void ParallelFor(u32 workCount, TContainer& container, TFunc&& func, const StringView& description = AsView(TC("")), bool highPriority = false);
 
 	protected:
 		Atomic<WorkTracker*> m_workTracker;
@@ -98,15 +100,36 @@ namespace uba
 		u64 startTime;
 	};
 
-	template<typename TContainer, typename TFunc>
-	void WorkManager::ParallelFor(u32 workCount, TContainer& container, const TFunc& func, const StringView& description, bool highPriority)
+	template<u32 BatchSize, typename TContainer, typename TFunc>
+	void WorkManager::ParallelFor(u32 workCount, TContainer& container, TFunc&& func, const StringView& description, bool highPriority)
 	{
 		#if !defined( __clang_analyzer__ ) // Static analyzer claims context can leak but it can only leak when process terminates (and then it doesn't matter)
 
+		typedef typename TContainer::iterator iterator;
+
+		auto size = container.size();
+
+		if (size <= BatchSize)
+		{
+			#if UBA_TRACK_WORK
+			TrackWorkScope tws(*this, description, ColorWork);
+			#else
+			TrackWorkScope tws;
+			#endif
+			WorkContext wc{tws};
+			for (iterator i=container.begin(), e=container.end(); i!=e; i++)
+				func(wc, i);
+			return;
+		}
+
+		workCount = Min(workCount, u32(size - 1)/BatchSize);
+
+		Event doneEvent(true);
+
 		struct Context
 		{
-			typename TContainer::iterator it;
-			typename TContainer::iterator end;
+			iterator it;
+			iterator end;
 			u32 refCount;
 			u32 activeCount;
 			bool isDone;
@@ -114,8 +137,6 @@ namespace uba
 			Event* doneEvent;
 		};
 
-		Event doneEvent(true);
-
 		auto context = new Context();
 		context->it = container.begin();
 		context->end = container.end();
@@ -124,8 +145,10 @@ namespace uba
 		context->isDone = false;
 		context->doneEvent = &doneEvent;
 
-		auto work = [context, funcCopy = func](const WorkContext&) mutable
+		auto work = [context, funcCopy = func](const WorkContext& wc) mutable
 			{
+				iterator itBatch[BatchSize];
+
 				u32 active = 0;
 				while (true)
 				{
@@ -146,20 +169,35 @@ namespace uba
 						return;
 					}
 
-					auto it2 = context->it++;
+					itBatch[0] = context->it++;
 					active = 1;
+
+					if constexpr (BatchSize > 1)
+						for (;active < BatchSize && context->it != context->end;)
+							itBatch[active++] = context->it++;
+
 					context->activeCount += active;
 					l.Leave();
 
-					funcCopy(it2);
+					funcCopy(wc, itBatch[0]);
+
+					if constexpr (BatchSize > 1)
+						for (u32 i=1; i!=active; ++i)
+							funcCopy(wc, itBatch[i]);
 				}
 			};
 		
 		AddWork(work, workCount, description.data, ColorWork, highPriority);
+
 		{
+			#if UBA_TRACK_WORK
 			TrackWorkScope tws(*this, description, ColorWork);
+			#else
+			TrackWorkScope tws;
+			#endif
 			work({tws});
 		}
+
 		doneEvent.IsSet();
 
 		#endif
diff --git a/Engine/Source/Programs/UnrealBuildAccelerator/Core/Public/UbaStringBuffer.h b/Engine/Source/Programs/UnrealBuildAccelerator/Core/Public/UbaStringBuffer.h
--- a/Engine/Source/Programs/UnrealBuildAccelerator/Core/Public/UbaStringBuffer.h
+++ b/Engine/Source/Programs/UnrealBuildAccelerator/Core/Public/UbaStringBuffer.h
@@ -126,6 +126,8 @@ namespace uba
 	template<typename T, size_t size>
 	StringView AsView(T (&buffer)[size]) { return StringView(buffer, size-1); }
 
+	#define TCV(x) AsView(TC(x))
+
 	struct LastErrorToText : StringBuffer<256>
 	{
 		LastErrorToText();
diff --git a/Engine/Source/Programs/UnrealBuildAccelerator/ObjTool/Private/UbaObjTool.cpp b/Engine/Source/Programs/UnrealBuildAccelerator/ObjTool/Private/UbaObjTool.cpp
--- a/Engine/Source/Programs/UnrealBuildAccelerator/ObjTool/Private/UbaObjTool.cpp
+++ b/Engine/Source/Programs/UnrealBuildAccelerator/ObjTool/Private/UbaObjTool.cpp
@@ -263,7 +263,7 @@ namespace uba
 
 			u32 workerCount = DefaultProcessorCount-1;
 			WorkManagerImpl workManager(workerCount, TC("UbaWorker(StripObjFiles)"));
-			workManager.ParallelFor(workerCount, objFilesDependencies, [&](auto& it)
+			workManager.ParallelFor(workerCount, objFilesDependencies, [&](const WorkContext&, auto& it)
 				{
 					const TString& exiFilename = *it;
 
@@ -285,7 +285,7 @@ namespace uba
 			Map<TString, ObjectFile*> objectFiles;
 			auto g = MakeGuard([&]() { for (auto& kv : objectFiles) delete kv.second; });
 
-			workManager.ParallelFor(workerCount, objFilesToStrip, [&](auto& it)
+			workManager.ParallelFor(workerCount, objFilesToStrip, [&](const WorkContext&, auto& it)
 				{
 					const TString& exiFilename = *it;
 					SymbolFile symbolFile;
@@ -339,7 +339,7 @@ namespace uba
 			objFiles.resize(objFilesForImpLib.size());
 			u32 workerCount = DefaultProcessorCount - 1;
 			WorkManagerImpl workManager(workerCount, TC("UbaWorker(Load)"));
-			workManager.ParallelFor(workerCount, objFilesForImpLib, [&](auto& it)
+			workManager.ParallelFor(workerCount, objFilesForImpLib, [&](const WorkContext&, auto& it)
 			{
 				auto& o = *it;
 				StringBuffer<> fixedPath;
