[UBA]
* Fixed so scrollbars are hidden in visualizer when in certain modes where they are not used

diff --git a/Engine/Source/Programs/UnrealBuildAccelerator/Visualizer/Private/UbaVisualizer.cpp b/Engine/Source/Programs/UnrealBuildAccelerator/Visualizer/Private/UbaVisualizer.cpp
--- a/Engine/Source/Programs/UnrealBuildAccelerator/Visualizer/Private/UbaVisualizer.cpp
+++ b/Engine/Source/Programs/UnrealBuildAccelerator/Visualizer/Private/UbaVisualizer.cpp
@@ -587,11 +587,19 @@ namespace uba
 		m_popupFont.height = 14;
 		//m_popupFont = (HFONT)GetStockObject(SYSTEM_FIXED_FONT);
 
-		DWORD windowStyle = WS_CAPTION | WS_SYSMENU | WS_THICKFRAME | WS_MAXIMIZEBOX | WS_MINIMIZEBOX | WS_CLIPCHILDREN | WS_VSCROLL | WS_HSCROLL;
+		DWORD scrollbarFlags = 0;
+		m_verticalScrollBarEnabled = !ActiveProcessesShouldFillHeight();
+		if (m_verticalScrollBarEnabled)
+			scrollbarFlags |= WS_VSCROLL;
+		m_horizontalScrollBarEnabled = !m_config.AutoScaleHorizontal;
+		if (m_horizontalScrollBarEnabled)
+			scrollbarFlags |= WS_HSCROLL;
+
+		DWORD windowStyle = WS_CAPTION | WS_SYSMENU | WS_THICKFRAME | WS_MAXIMIZEBOX | WS_MINIMIZEBOX | WS_CLIPCHILDREN | scrollbarFlags;
 
 		DWORD exStyle = 0;
 		if (m_config.parent)
-			windowStyle = WS_POPUP | WS_VSCROLL | WS_HSCROLL;
+			windowStyle = WS_POPUP | scrollbarFlags;
 
 		StringBuffer<> title;
 		GetTitlePrefix(title).Append(L"Initializing...");
@@ -626,7 +634,7 @@ namespace uba
 			exitCloak.Execute();
 
 			// If not child it will not propagate keyboard presses etc to parent
-			SetWindowLong(hwnd, GWL_STYLE, WS_CHILD|WS_VSCROLL|WS_HSCROLL);
+			SetWindowLong(hwnd, GWL_STYLE, WS_CHILD | scrollbarFlags);
 
 			m_parentHwnd = (HWND)(uintptr_t)m_config.parent;
 			if (!SetParent(hwnd, m_parentHwnd))
@@ -2232,7 +2240,7 @@ namespace uba
 
 		u32 maxHeight = clientRect.bottom;
 
-		bool fillHeight = !m_config.showDetailedData && !m_config.showTitleBars && !m_config.showCpuMemStats && !m_config.showNetworkStats && !m_config.showProcessBars;
+		bool fillHeight = ActiveProcessesShouldFillHeight();
 		if (fillHeight)
 		{
 			maxHeight = clientRect.bottom - posY;
@@ -3185,12 +3193,54 @@ namespace uba
 		si.nPage = int(rect.bottom);
 		si.nPos = -int(m_scrollPosY);
 		si.nTrackPos = 0;
-		SetScrollInfo(m_hwnd, SB_VERT, &si, redraw);
 
-		si.nMax = m_contentWidth;
-		si.nPage = rect.right;
-		si.nPos = -int(m_scrollPosX);
-		SetScrollInfo(m_hwnd, SB_HORZ, &si, redraw);
+		bool updateFrame = false;
+		if (ActiveProcessesShouldFillHeight())
+		{
+			if (m_verticalScrollBarEnabled)
+			{
+				SetWindowLong(m_hwnd, GWL_STYLE, GetWindowLong(m_hwnd, GWL_STYLE) & ~WS_VSCROLL);
+				m_verticalScrollBarEnabled = false;
+				updateFrame = true;
+			}
+		}
+		else
+		{
+			if (!m_verticalScrollBarEnabled)
+			{
+				SetWindowLong(m_hwnd, GWL_STYLE, GetWindowLong(m_hwnd, GWL_STYLE) | WS_VSCROLL);
+				m_verticalScrollBarEnabled = true;
+				updateFrame = true;
+			}
+			SetScrollInfo(m_hwnd, SB_VERT, &si, redraw);
+		}
+
+		if (m_config.AutoScaleHorizontal)
+		{
+			if (m_horizontalScrollBarEnabled)
+			{
+				SetWindowLong(m_hwnd, GWL_STYLE, GetWindowLong(m_hwnd, GWL_STYLE) & ~WS_HSCROLL);
+				m_horizontalScrollBarEnabled = false;
+				updateFrame = true;
+			}
+		}
+		else
+		{
+			if (!m_horizontalScrollBarEnabled)
+			{
+				SetWindowLong(m_hwnd, GWL_STYLE, GetWindowLong(m_hwnd, GWL_STYLE) | WS_HSCROLL);
+				m_horizontalScrollBarEnabled = true;
+				updateFrame = true;
+			}
+
+			si.nMax = m_contentWidth;
+			si.nPage = rect.right;
+			si.nPos = -int(m_scrollPosX);
+			SetScrollInfo(m_hwnd, SB_HORZ, &si, redraw);
+		}
+
+		if (updateFrame)
+			SetWindowPos(m_hwnd, NULL, 0, 0, 0, 0, SWP_NOSIZE | SWP_NOMOVE | SWP_NOZORDER | SWP_FRAMECHANGED);
 	}
 
 	void Visualizer::SetActiveFont(const Font& font)
@@ -3209,6 +3259,11 @@ namespace uba
 		DwmSetWindowAttribute(m_hwnd, attribute, &useDarkMode, sizeof(useDarkMode));
 	}
 
+	bool Visualizer::ActiveProcessesShouldFillHeight()
+	{
+		return !m_config.showDetailedData && !m_config.showTitleBars && !m_config.showCpuMemStats && !m_config.showNetworkStats && !m_config.showProcessBars;
+	}
+
 	StringBuffer<128> Visualizer::GetWorldTime(u64 time)
 	{
 		time_t rawTime = m_traceView.traceSystemStartTimeUs / 1'000'000;
@@ -3572,7 +3627,7 @@ namespace uba
 
 		case WM_LBUTTONDOWN:
 		{
-			if (m_traceView.sessions.empty())
+			if (!m_parentHwnd && m_traceView.sessions.empty())
 			{
 				RECT r;
 				GetClientRect(hWnd, &r);
@@ -3812,6 +3867,7 @@ namespace uba
 				break;
 			case Popup_AutoScaleHorizontal:
 				m_config.AutoScaleHorizontal = !m_config.AutoScaleHorizontal;
+				UpdateScrollbars(true);
 				Redraw(true);
 				break;
 			case Popup_LockTimelineToBottom:
diff --git a/Engine/Source/Programs/UnrealBuildAccelerator/Visualizer/Private/UbaVisualizer.h b/Engine/Source/Programs/UnrealBuildAccelerator/Visualizer/Private/UbaVisualizer.h
--- a/Engine/Source/Programs/UnrealBuildAccelerator/Visualizer/Private/UbaVisualizer.h
+++ b/Engine/Source/Programs/UnrealBuildAccelerator/Visualizer/Private/UbaVisualizer.h
@@ -160,6 +160,7 @@ namespace uba
 		void Redraw(bool now);
 		void SetActiveFont(const Font& font);
 		void UpdateTheme();
+		bool ActiveProcessesShouldFillHeight();
 		StringBuffer<128> GetWorldTime(u64 time);
 
 		StringBuffer<256> m_namedTrace;
@@ -281,6 +282,9 @@ namespace uba
 		float m_scrollAtAnchorY = 0;
 		int m_dragToScrollCounter = 0;
 
+		bool m_horizontalScrollBarEnabled = true;
+		bool m_verticalScrollBarEnabled = true;
+
 		TString m_filterString;
 
 		Thread m_thread;
