[UBA]
* Added GetFileBasicInformationByHandle and GetFileBasicInformation and use them when possible since they are much cheaper on wine

diff --git a/Engine/Source/Programs/UnrealBuildAccelerator/Common/Private/UbaCacheClient.cpp b/Engine/Source/Programs/UnrealBuildAccelerator/Common/Private/UbaCacheClient.cpp
--- a/Engine/Source/Programs/UnrealBuildAccelerator/Common/Private/UbaCacheClient.cpp
+++ b/Engine/Source/Programs/UnrealBuildAccelerator/Common/Private/UbaCacheClient.cpp
@@ -360,8 +360,8 @@ namespace uba
 
 					if (newKey != oldKey)
 					{
-						FileInformation fileInfo;
-						fa.GetFileInformationByHandle(fileInfo);
+						FileBasicInformation fileInfo;
+						fa.GetFileBasicInformationByHandle(fileInfo);
 
 						auto& fileEntry = m_storage.GetOrCreateFileEntry(CaseInsensitiveFs ? ToStringKeyLower(path) : ToStringKey(path));
 						SCOPED_FUTEX_READ(fileEntry.lock, lock);
diff --git a/Engine/Source/Programs/UnrealBuildAccelerator/Common/Private/UbaFile.cpp b/Engine/Source/Programs/UnrealBuildAccelerator/Common/Private/UbaFile.cpp
--- a/Engine/Source/Programs/UnrealBuildAccelerator/Common/Private/UbaFile.cpp
+++ b/Engine/Source/Programs/UnrealBuildAccelerator/Common/Private/UbaFile.cpp
@@ -13,6 +13,7 @@
 #include <WinCon.h>
 #include <RestartManager.h>
 #pragma comment(lib, "Rstrtmgr.lib")
+extern "C" NTSTATUS NTAPI NtQueryInformationFile(HANDLE FileHandle, PIO_STATUS_BLOCK IoStatusBlock, PVOID FileInformation, ULONG Length, FILE_INFORMATION_CLASS FileInformationClass);
 #else
 #include <dlfcn.h>
 #include <execinfo.h>
@@ -169,6 +170,65 @@ namespace uba
 		return ReportOpenFileError(logger, fileName, lastError);
 	}
 
+	bool GetFileBasicInformationByHandle(FileBasicInformation& out, Logger& logger, const tchar* fileName, FileHandle hFile)
+	{
+#if PLATFORM_WINDOWS
+		ExtendedTimerScope ts(KernelStats::GetCurrent().getFileInfo);
+
+		struct FILE_NETWORK_OPEN_INFORMATION
+		{
+			  LARGE_INTEGER CreationTime;
+			  LARGE_INTEGER LastAccessTime;
+			  LARGE_INTEGER LastWriteTime;
+			  LARGE_INTEGER ChangeTime;
+			  LARGE_INTEGER AllocationSize;
+			  LARGE_INTEGER EndOfFile;
+			  ULONG FileAttributes;
+		};
+
+		FILE_NETWORK_OPEN_INFORMATION info;
+		IO_STATUS_BLOCK b;
+		NTSTATUS res = NtQueryInformationFile(asHANDLE(hFile), &b, &info, sizeof(info), (FILE_INFORMATION_CLASS)34); // FileNetworkOpenInformation
+		if (res != STATUS_SUCCESS)
+			return logger.Error(TC("GetFileBasicInformationByHandle (NtQueryInformationFile) failed on %s (%u)"), fileName, res);
+		out.attributes = info.FileAttributes;
+		out.lastWriteTime = info.LastWriteTime.QuadPart;
+		out.size = info.EndOfFile.QuadPart;
+		return true;
+#else
+		FileInformation info;
+		if (!GetFileInformationByHandle(info, logger, fileName, hFile))
+			return false;
+		out.attributes = info.attributes;
+		out.lastWriteTime = info.lastWriteTime;
+		out.size = info.size;
+		return true;
+#endif
+	}
+
+	bool GetFileBasicInformation(FileBasicInformation& out, Logger& logger, const tchar* fileName)
+	{
+#if PLATFORM_WINDOWS
+		ExtendedTimerScope ts(KernelStats::GetCurrent().getFileInfo);
+		MAKE_LONG_FILENAME(fileName);
+		FileHandle h = uba::CreateFileW(fileName, 0, 0, OPEN_EXISTING, FILE_FLAG_BACKUP_SEMANTICS);
+		if (h == InvalidFileHandle)
+			return false;// logger.Error(TC("GetFileInformation: CreateFile failed for file %s (%s)"), fileName, LastErrorToText().data);
+		auto handleGuard = MakeGuard([&]() { uba::CloseFile(fileName, h); });
+		if (!GetFileBasicInformationByHandle(out, logger, fileName, h))
+			return false;// logger.Error(TC("Failed to get file information for %s while checking file added for write. This should not happen! (%s)"), fileName, LastErrorToText().data);
+		return true;
+#else
+		FileInformation info;
+		if (!GetFileInformation(info, logger, fileName))
+			return false;
+		out.attributes = info.attributes;
+		out.lastWriteTime = info.lastWriteTime;
+		out.size = info.size;
+		return true;
+#endif
+	}
+
 	bool GetFileInformationByHandle(FileInformation& out, Logger& logger, const tchar* fileName, FileHandle hFile)
 	{
 		ExtendedTimerScope ts(KernelStats::GetCurrent().getFileInfo);
diff --git a/Engine/Source/Programs/UnrealBuildAccelerator/Common/Private/UbaFileAccessor.cpp b/Engine/Source/Programs/UnrealBuildAccelerator/Common/Private/UbaFileAccessor.cpp
--- a/Engine/Source/Programs/UnrealBuildAccelerator/Common/Private/UbaFileAccessor.cpp
+++ b/Engine/Source/Programs/UnrealBuildAccelerator/Common/Private/UbaFileAccessor.cpp
@@ -389,8 +389,8 @@ namespace uba
 		if (!OpenFileSequentialRead(m_logger, m_fileName, m_fileHandle))
 			return errorOnFail ? m_logger.Error(TC("Failed to open file %s for read"), m_fileName) : false;
 
-		FileInformation info;
-		if (!GetFileInformationByHandle(info))
+		FileBasicInformation info;
+		if (!GetFileBasicInformationByHandle(info))
 			return m_logger.Error(TC("GetFileInformationByHandle failed on %s"), m_fileName);
 
 		m_size = info.size;
@@ -414,9 +414,9 @@ namespace uba
 		return true;
 	}
 
-	bool FileAccessor::GetFileInformationByHandle(FileInformation& out)
+	bool FileAccessor::GetFileBasicInformationByHandle(FileBasicInformation& out)
 	{
-		return uba::GetFileInformationByHandle(out, m_logger, m_fileName, m_fileHandle);
+		return uba::GetFileBasicInformationByHandle(out, m_logger, m_fileName, m_fileHandle);
 	}
 
 	bool FileAccessor::InternalClose(bool success, u64* lastWriteTime)
diff --git a/Engine/Source/Programs/UnrealBuildAccelerator/Common/Private/UbaSession.cpp b/Engine/Source/Programs/UnrealBuildAccelerator/Common/Private/UbaSession.cpp
--- a/Engine/Source/Programs/UnrealBuildAccelerator/Common/Private/UbaSession.cpp
+++ b/Engine/Source/Programs/UnrealBuildAccelerator/Common/Private/UbaSession.cpp
@@ -589,8 +589,8 @@ namespace uba
 			}
 			else
 			{
-				FileInformation info;
-				if (!GetFileInformationByHandle(info, m_logger, fileName, fileHandle))
+				FileBasicInformation info;
+				if (!GetFileBasicInformationByHandle(info, m_logger, fileName, fileHandle))
 					return false;
 
 				size = info.size;
@@ -2374,7 +2374,7 @@ namespace uba
 			if (m_allowOutputFiles && writtenFile.owner != nullptr && writtenFile.owner != &msg.process)
 			{
 				// This can happen when library has /GL (whole program optimization) but target has not.. then link.exe will restart
-				UBA_ASSERTF(false, TC("File %s changed owner.. should not happen (OldOwner: %s New owner: %s)"), name, writtenFile.owner->m_realApplication.c_str(), msg.process.m_realApplication.c_str());
+				//UBA_ASSERTF(false, TC("File %s changed owner.. should not happen (OldOwner: %s New owner: %s)"), name, writtenFile.owner->m_realApplication.c_str(), msg.process.m_realApplication.c_str());
 			}
 
 			writtenFile.owner = &msg.process;
diff --git a/Engine/Source/Programs/UnrealBuildAccelerator/Common/Private/UbaStorage.cpp b/Engine/Source/Programs/UnrealBuildAccelerator/Common/Private/UbaStorage.cpp
--- a/Engine/Source/Programs/UnrealBuildAccelerator/Common/Private/UbaStorage.cpp
+++ b/Engine/Source/Programs/UnrealBuildAccelerator/Common/Private/UbaStorage.cpp
@@ -1896,8 +1896,8 @@ namespace uba
 				return true;
 			}
 
-			FileInformation info;
-			if (!GetFileInformationByHandle(info, m_logger, fileName, fileHandle))
+			FileBasicInformation info;
+			if (!GetFileBasicInformationByHandle(info, m_logger, fileName, fileHandle))
 			{
 				fileEntry.casKey = CasKeyZero;
 				return m_logger.Error(TC("GetFileInformationByHandle failed on %s"), fileName);
@@ -2031,8 +2031,8 @@ namespace uba
 		}
 		auto fileGuard = MakeGuard([&](){ CloseFile(fileName, fileHandle); });
 
-		FileInformation info;
-		if (!GetFileInformationByHandle(info, m_logger, fileName, fileHandle))
+		FileBasicInformation info;
+		if (!GetFileBasicInformationByHandle(info, m_logger, fileName, fileHandle))
 		{
 			fileEntry.casKey = CasKeyZero;
 			return m_logger.Error(TC("GetFileInformationByHandle failed on %s"), fileName);
diff --git a/Engine/Source/Programs/UnrealBuildAccelerator/Common/Public/UbaFile.h b/Engine/Source/Programs/UnrealBuildAccelerator/Common/Public/UbaFile.h
--- a/Engine/Source/Programs/UnrealBuildAccelerator/Common/Public/UbaFile.h
+++ b/Engine/Source/Programs/UnrealBuildAccelerator/Common/Public/UbaFile.h
@@ -34,6 +34,15 @@ namespace uba
 	inline constexpr u64 OverlappedIoFlag = 0x0000'0001'0000'0000; // Only used by windows
 	#endif
 
+	struct FileBasicInformation
+	{
+		u32 attributes;
+		u64 lastWriteTime;
+		u64 size;
+	};
+	bool GetFileBasicInformationByHandle(FileBasicInformation& out, Logger& logger, const tchar* fileName, FileHandle hFile);
+	bool GetFileBasicInformation(FileBasicInformation& out, Logger& logger, const tchar* fileName);
+
 	struct FileInformation
 	{
 		u32 attributes;
diff --git a/Engine/Source/Programs/UnrealBuildAccelerator/Common/Public/UbaFileAccessor.h b/Engine/Source/Programs/UnrealBuildAccelerator/Common/Public/UbaFileAccessor.h
--- a/Engine/Source/Programs/UnrealBuildAccelerator/Common/Public/UbaFileAccessor.h
+++ b/Engine/Source/Programs/UnrealBuildAccelerator/Common/Public/UbaFileAccessor.h
@@ -30,7 +30,7 @@ namespace uba
 		inline u64 GetSize() { return m_size; }
 
 		// Can be called if file is opened
-		bool GetFileInformationByHandle(FileInformation& out);
+		bool GetFileBasicInformationByHandle(FileBasicInformation& out);
 
 	private:
 		bool InternalClose(bool success, u64* lastWriteTime);
