[UBA]
* Added timedout variable to socket connect call in visualizer to fix connect spin on wine

diff --git a/Engine/Source/Programs/UnrealBuildAccelerator/Visualizer/Private/UbaVisualizer.cpp b/Engine/Source/Programs/UnrealBuildAccelerator/Visualizer/Private/UbaVisualizer.cpp
--- a/Engine/Source/Programs/UnrealBuildAccelerator/Visualizer/Private/UbaVisualizer.cpp
+++ b/Engine/Source/Programs/UnrealBuildAccelerator/Visualizer/Private/UbaVisualizer.cpp
@@ -343,7 +343,8 @@ namespace uba
 			StringBuffer<> title;
 			PostNewTitle(GetTitlePrefix(title).Appendf(L"Trying to connect to %s:%u%s", host, port, dots + ((dotsCounter--) % 4)));
 
-			if (!m_client->Connect(backend, host, port))
+			bool timedOut;
+			if (!m_client->Connect(backend, host, port, &timedOut))
 				continue;
 
 			PostNewTitle(GetTitlePrefix(title).Appendf(L"Connected to %s:%u", host, port));
