[UBA}
* Visualizer - changed so icon turns into hand when hovering over open trace text
* Added more descriptive errors for HttpConnection
* Added assert if no description is provided for work

diff --git a/Engine/Source/Programs/UnrealBuildAccelerator/Common/Private/UbaNetworkBackendTcp.cpp b/Engine/Source/Programs/UnrealBuildAccelerator/Common/Private/UbaNetworkBackendTcp.cpp
--- a/Engine/Source/Programs/UnrealBuildAccelerator/Common/Private/UbaNetworkBackendTcp.cpp
+++ b/Engine/Source/Programs/UnrealBuildAccelerator/Common/Private/UbaNetworkBackendTcp.cpp
@@ -1111,19 +1111,22 @@ namespace uba
 
 		hostent* hostent = gethostbyname(host);
 		if (hostent == NULL)
-			return logger.Error(TC("HttpRequest: gethostbyname error"));
+			return logger.Error(TC("HttpConnection: gethostbyname error (%s)"), host);
 
-		unsigned long in_addr = inet_addr(inet_ntoa(*(struct in_addr*)*(hostent->h_addr_list)));
+		char* ntoaRes = inet_ntoa(*(struct in_addr*)*(hostent->h_addr_list));
+		if (!ntoaRes)
+			return logger.Error(TC("HttpConnection: inet_ntoa returned null"));
+		unsigned long in_addr = inet_addr(ntoaRes);
 		if (in_addr == INADDR_NONE)
-			return logger.Error(TC("HttpRequest: inet_addr error"));
+			return logger.Error(TC("HttpConnection: inet_addr returned INADDR_NONE (%s)"), ntoaRes);
 
 		protoent* protoent = getprotobyname("tcp");
 		if (protoent == NULL)
-			return logger.Error(TC("HttpRequest: socket error"));
+			return logger.Error(TC("HttpConnection: getprotobyname returned null for tcp"));
 
 		SOCKET sock = socket(AF_INET, SOCK_STREAM, protoent->p_proto);
-		if (sock == -1)
-			return logger.Error(TC("HttpRequest: socket error"));
+		if (sock == INVALID_SOCKET)
+			return logger.Error(TC("HttpConnection: socket failed (%s)"), LastErrorToText(WSAGetLastError()).data);
 		auto socketClose = MakeGuard([sock]() { closesocket(sock); });
 
 		if (m_connectTimeOutMs)
@@ -1135,7 +1138,7 @@ namespace uba
 		sockaddr_in.sin_port = htons(80);
 
 		if (connect(sock, (struct sockaddr*)&sockaddr_in, sizeof(sockaddr_in)) == -1)
-			return false;// logger.Error(TC("HttpRequest: connect error"));
+			return false;// logger.Error(TC("HttpConnection: connect error"));
 
 		socketClose.Cancel();
 
@@ -1168,7 +1171,7 @@ namespace uba
 		while (totalBytesSent < requestLen) {
 			int bytesSent = send(m_socket, request + totalBytesSent, requestLen - totalBytesSent, 0);
 			if (bytesSent == -1)
-				return logger.Error(TC("HttpRequest: send error"));
+				return logger.Error(TC("HttpConnection: socket send error (%s)"), host);
 			totalBytesSent += bytesSent;
 		}
 
@@ -1183,11 +1186,11 @@ namespace uba
 		while ((bytesRead = recv(m_socket, buf + readPos, sizeof(buf) - readPos, 0)) > 0)
 			readPos += bytesRead;
 
-		if (bytesRead == -1)
-			return logger.Error(TC("HttpRequest: recv error"));
+		if (bytesRead == SOCKET_ERROR)
+			return logger.Error(TC("HttpConnection: socket recv error %s (%s)"), LastErrorToText(WSAGetLastError()).data, m_host);
 
 		if (readPos == sizeof(buf))
-			return logger.Error(TC("HttpRequest: buffer overflow"));
+			return logger.Error(TC("HttpConnection: buffer overflow"));
 
 		buf[readPos] = 0;
 
@@ -1197,16 +1200,16 @@ namespace uba
 
 		char* firstSpace = strchr(buf, ' '); // After version (where status code starts)
 		if (!firstSpace)
-			return logger.Error(TC("HttpRequest: first space not found (read %u)"), readPos);
+			return logger.Error(TC("HttpConnection: first space not found (read %u)"), readPos);
 		char* secondSpace = strchr(firstSpace + 1, ' '); // after status code
 		if (!secondSpace)
-			return logger.Error(TC("HttpRequest: second space not found"));
+			return logger.Error(TC("HttpConnection: second space not found"));
 		*secondSpace = 0;
 		outStatusCode = strtoul(firstSpace + 1, nullptr, 10);
 
 		char* bodyStart = strstr(secondSpace + 1, "\r\n\r\n");
 		if (!bodyStart)
-			return logger.Error(TC("HttpRequest: no body found"));
+			return logger.Error(TC("HttpConnection: no body found"));
 		outResponse.Append(bodyStart + 4);
 		return true;
 	}
diff --git a/Engine/Source/Programs/UnrealBuildAccelerator/Common/Private/UbaNetworkServer.cpp b/Engine/Source/Programs/UnrealBuildAccelerator/Common/Private/UbaNetworkServer.cpp
--- a/Engine/Source/Programs/UnrealBuildAccelerator/Common/Private/UbaNetworkServer.cpp
+++ b/Engine/Source/Programs/UnrealBuildAccelerator/Common/Private/UbaNetworkServer.cpp
@@ -857,6 +857,7 @@ namespace uba
 
 	void NetworkServer::AddWork(const WorkFunction& work, u32 count, const tchar* desc, const Color& color, bool highPriority)
 	{
+		UBA_ASSERT(*desc);
 		SCOPED_FUTEX(m_additionalWorkLock, lock);
 		for (u32 i = 0; i != count; ++i)
 		{
diff --git a/Engine/Source/Programs/UnrealBuildAccelerator/Common/Private/UbaSession.cpp b/Engine/Source/Programs/UnrealBuildAccelerator/Common/Private/UbaSession.cpp
--- a/Engine/Source/Programs/UnrealBuildAccelerator/Common/Private/UbaSession.cpp
+++ b/Engine/Source/Programs/UnrealBuildAccelerator/Common/Private/UbaSession.cpp
@@ -2809,7 +2809,7 @@ namespace uba
 			{
 				if (!writeFile(**it))
 					success = false;
-			});
+			}, TCV("WriteFilesToDisk"));
 		return success;
 	}
 
diff --git a/Engine/Source/Programs/UnrealBuildAccelerator/Common/Private/UbaWorkManager.cpp b/Engine/Source/Programs/UnrealBuildAccelerator/Common/Private/UbaWorkManager.cpp
--- a/Engine/Source/Programs/UnrealBuildAccelerator/Common/Private/UbaWorkManager.cpp
+++ b/Engine/Source/Programs/UnrealBuildAccelerator/Common/Private/UbaWorkManager.cpp
@@ -111,6 +111,7 @@ namespace uba
 
 	void WorkManagerImpl::AddWork(const WorkFunction& work, u32 count, const tchar* desc, const Color& color, bool highPriority)
 	{
+		UBA_ASSERT(*desc);
 		UBA_ASSERT(!m_workers.empty());
 		SCOPED_FUTEX(m_workLock, lock);
 		bool trackWork = m_workTracker.load();
diff --git a/Engine/Source/Programs/UnrealBuildAccelerator/Visualizer/Private/UbaVisualizer.cpp b/Engine/Source/Programs/UnrealBuildAccelerator/Visualizer/Private/UbaVisualizer.cpp
--- a/Engine/Source/Programs/UnrealBuildAccelerator/Visualizer/Private/UbaVisualizer.cpp
+++ b/Engine/Source/Programs/UnrealBuildAccelerator/Visualizer/Private/UbaVisualizer.cpp
@@ -3635,7 +3635,7 @@ namespace uba
 				int centerHrz = r.right/2;
 				int centerVrt = r.bottom/2;
 				r = { centerHrz, centerVrt, centerHrz, centerVrt };
-				InflateRect(&r, 200, 50);
+				InflateRect(&r, 180, 40);
 				POINTS p = MAKEPOINTS(lParam);
 				if (PtInRect(&r, { p.x, p.y }))
 				{
@@ -3742,10 +3742,23 @@ namespace uba
 		{
 			static HCURSOR arrow = LoadCursorW(NULL, IDC_ARROW);
 			static HCURSOR hand = LoadCursorW(NULL, IDC_HAND);
-			if (!m_hyperLinkSelected.empty() || (m_sessionSelectedIndex != ~0u && !m_traceView.sessions[m_sessionSelectedIndex].hyperlink.empty()))
-				SetCursor(hand);
+			bool useHand = false;
+			if (!m_parentHwnd && m_traceView.sessions.empty())
+			{
+				RECT r;
+				GetClientRect(hWnd, &r);
+				int centerHrz = r.right/2;
+				int centerVrt = r.bottom/2;
+				r = { centerHrz, centerVrt, centerHrz, centerVrt };
+				InflateRect(&r, 180, 40);
+				POINT pt;
+				GetCursorPos(&pt);
+				ScreenToClient(hWnd, &pt);
+				useHand = PtInRect(&r, { pt.x, pt.y });
+			}
 			else
-				SetCursor(arrow);
+				useHand = !m_hyperLinkSelected.empty() || (m_sessionSelectedIndex != ~0u && !m_traceView.sessions[m_sessionSelectedIndex].hyperlink.empty());
+			SetCursor(useHand ? hand : arrow);
 			break;
 		}
 		case WM_LBUTTONUP:
