[UBA]
* Fixed code that was using SCOPED_READ_LOCK where they should be using SCOPED_FUTEX_READ

diff --git a/Engine/Source/Programs/UnrealBuildAccelerator/Common/Private/UbaProcess.cpp b/Engine/Source/Programs/UnrealBuildAccelerator/Common/Private/UbaProcess.cpp
--- a/Engine/Source/Programs/UnrealBuildAccelerator/Common/Private/UbaProcess.cpp
+++ b/Engine/Source/Programs/UnrealBuildAccelerator/Common/Private/UbaProcess.cpp
@@ -1536,7 +1536,7 @@ namespace uba
 			creationFlags |= EXTENDED_STARTUPINFO_PRESENT;
 
 
-			SCOPED_READ_LOCK(m_session.m_processJobObjectLock, jobObjectLock);
+			SCOPED_FUTEX_READ(m_session.m_processJobObjectLock, jobObjectLock);
 			if (!m_session.m_processJobObject)
 			{
 				m_cancelEvent.Set();
diff --git a/Engine/Source/Programs/UnrealBuildAccelerator/Common/Private/UbaSessionServer.cpp b/Engine/Source/Programs/UnrealBuildAccelerator/Common/Private/UbaSessionServer.cpp
--- a/Engine/Source/Programs/UnrealBuildAccelerator/Common/Private/UbaSessionServer.cpp
+++ b/Engine/Source/Programs/UnrealBuildAccelerator/Common/Private/UbaSessionServer.cpp
@@ -2320,12 +2320,12 @@ namespace uba
 
 				StringKey fileNameKey = ToStringKey(path);
 				{
-					SCOPED_READ_LOCK(m_fileMappingTableLookupLock, mlock);
+					SCOPED_FUTEX_READ(m_fileMappingTableLookupLock, mlock);
 					auto findIt = m_fileMappingTableLookup.find(fileNameKey);
 					if (findIt != m_fileMappingTableLookup.end())
 					{
 						auto& entry = findIt->second;
-						SCOPED_READ_LOCK(entry.lock, entryCs);
+						SCOPED_FUTEX_READ(entry.lock, entryCs);
 						logStr.Clear().Appendf(TC("File %s found in mapping table table."), path.data);
 						if (entry.handled)
 						{
