[UBA]
* Added so FileMappingBuffer keep anonymous file mappings mapped at all time instead of map/unmap sections of them. They still use the same amount of committed memory so it doesn't matter (can be disabled in config)
* Added tests for FileMappingBuffer

diff --git a/Engine/Source/Programs/UnrealBuildAccelerator/Common/Private/UbaFileMapping.cpp b/Engine/Source/Programs/UnrealBuildAccelerator/Common/Private/UbaFileMapping.cpp
--- a/Engine/Source/Programs/UnrealBuildAccelerator/Common/Private/UbaFileMapping.cpp
+++ b/Engine/Source/Programs/UnrealBuildAccelerator/Common/Private/UbaFileMapping.cpp
@@ -15,6 +15,8 @@
 
 #define UBA_DEBUG_FILE_MAPPING 0
 
+////////////////////////////////////////////////////////////////////////////////////////////////////
+
 namespace uba
 {
 	#if UBA_DEBUG_FILE_MAPPING
@@ -386,6 +388,8 @@ namespace uba
 		memcpy(dest, source, size);
 	}
 
+	////////////////////////////////////////////////////////////////////////////////////////////////////
+
 	FileMappingBuffer::FileMappingBuffer(Logger& logger, WorkManager* workManager)
 	:	m_logger(logger)
 	,	m_workManager(workManager)
@@ -404,15 +408,19 @@ namespace uba
 		return (IsWindows ? 32ull : 8ull) * 1024 * 1024 * 1024; // Linux can't have larger than 8gb
 	}
 
-	bool FileMappingBuffer::AddTransient(const tchar* name)
+	bool FileMappingBuffer::AddTransient(const tchar* name, bool keepMapped)
 	{
 		MappingStorage& storage = m_storage[MappedView_Transient];
+		File** prev = &storage.availableFile;
 		for (u32 i = 0; i != sizeof_array(storage.files); ++i)
 		{
-			File& file = storage.files[sizeof_array(storage.files) - (++storage.fileCount)];
+			File& file = storage.files[i];
 			file.name = name;
-			storage.availableFiles[storage.availableFilesCount++] = &file;
+			file.keepMapped = keepMapped;
+			*prev = &file;
+			prev = &file.next;
 		}
+		storage.fileCount = sizeof_array(storage.files);
 		return true;
 	}
 
@@ -446,6 +454,8 @@ namespace uba
 	{
 		MappingStorage& storage = m_storage[type];
 		File* file = PopFile(storage, size, alignment);
+		if (!file)
+			return {};
 
 		if (allowShrink)
 			return AllocAndMapViewNoLock(*file, size, alignment, hint);
@@ -473,11 +483,19 @@ namespace uba
 		}
 
 		u64 mapSize = alignedOffsetEnd - alignedOffsetStart;
-		u8* data = MapViewOfFile(m_logger, f.handle, FILE_MAP_WRITE, alignedOffsetStart, mapSize);
-		if (!data)
+		u8* data;
+		if (f.mappedMemory)
 		{
-			m_logger.Error(TC("%s - AllocAndMapView failed to map view of file for %s with size %llu and offset %llu (%s)"), f.name, hint, size, alignment, LastErrorToText().data);
-			return res;
+			data = f.mappedMemory + alignedOffsetStart;
+		}
+		else
+		{
+			data = MapViewOfFile(m_logger, f.handle, FILE_MAP_WRITE, alignedOffsetStart, mapSize);
+			if (!data)
+			{
+				m_logger.Error(TC("%s - AllocAndMapView failed to map view of file for %s with size %llu and offset %llu (%s)"), f.name, hint, size, alignment, LastErrorToText().data);
+				return res;
+			}
 		}
 
 		u64 committedBefore = AlignUp(offset, m_pageSize);
@@ -498,7 +516,7 @@ namespace uba
 		}
 
 		f.size = newOffset;
-		++f.activeMapCount;
+		//++f.activeMapCount;
 
 		res.handle = f.handle;
 		res.offset = offset;
@@ -529,37 +547,49 @@ namespace uba
 		while (true)
 		{
 			SCOPED_FUTEX(storage.availableFilesLock, lock);
-			for (u32 i = storage.availableFilesCount; i!=0; --i)
+			File* file = storage.availableFile;
+			if (!file)
 			{
-				u32 index = i - 1;
-				auto file = storage.availableFiles[index];
-				
-				if (!file->handle.IsValid())
+				if (storage.fullFileCount == storage.fileCount)
 				{
-					UBA_ASSERT(&storage == &m_storage[MappedView_Transient]);
-					u64 capacity = GetFileMappingCapacity();
-					file-> handle = uba::CreateMemoryMappingW(m_logger, PAGE_READWRITE | SEC_RESERVE, capacity, nullptr, TC("FileMappingBuffer"));
-					if (!file->handle.IsValid())
-					{
-						m_logger.Error(TC("%s - Failed to create memory map (%s)"), file->name, LastErrorToText().data);
-						return nullptr;
-					}
-					file->commitOnAlloc = true;
-					file->capacity = capacity;
+					m_logger.Error(TC("%s - All files are full!"), file->name);
+					return nullptr;
 				}
-				else
+				lock.Leave();
+				storage.availableFilesEvent.IsSet();
+				continue;
+			}
+			storage.availableFile = file->next;
+			lock.Leave();
+
+			if (!file->handle.IsValid())
+			{
+				UBA_ASSERT(&storage == &m_storage[MappedView_Transient]);
+				u64 capacity = GetFileMappingCapacity();
+				file->handle = uba::CreateMemoryMappingW(m_logger, PAGE_READWRITE | SEC_RESERVE, capacity, nullptr, TC("FileMappingBuffer"));
+				if (!file->handle.IsValid())
 				{
-					u64 newSize = AlignUp(file->size, alignment) + size;
-					if (newSize > file->capacity)
-						continue;
+					m_logger.Error(TC("%s - Failed to create memory map (%s)"), file->name, LastErrorToText().data);
+					return nullptr;
 				}
+				file->commitOnAlloc = true;
+				file->capacity = capacity;
+
+				if (file->keepMapped)
+					file->mappedMemory = MapViewOfFile(m_logger, file->handle, FILE_MAP_WRITE, 0, capacity);
 
-				storage.availableFiles[index] = storage.availableFiles[storage.availableFilesCount - 1];
-				--storage.availableFilesCount;
-				return file;
 			}
-			lock.Leave();
-			storage.availableFilesEvent.IsSet();
+			else
+			{
+				u64 newSize = AlignUp(file->size, alignment) + size;
+				if (newSize > file->capacity)
+				{
+					lock.Enter();
+					++storage.fullFileCount; // We will never add this file back to the available file list
+					continue;
+				}
+			}
+			return file;
 		}
 		return nullptr;
 	}
@@ -567,32 +597,43 @@ namespace uba
 	void FileMappingBuffer::PushFile(MappingStorage& storage, File* file)
 	{
 		SCOPED_FUTEX(storage.availableFilesLock, lock);
-		storage.availableFiles[storage.availableFilesCount++] = file;
+		file->next = storage.availableFile;
+		storage.availableFile = file;
 		storage.availableFilesEvent.Set();
 	}
 
 	void FileMappingBuffer::CloseMappingStorage(MappingStorage& storage)
 	{
-		u32 locksTaken = 0;
-		while (locksTaken < storage.fileCount)
+		SCOPED_FUTEX(storage.availableFilesLock, lock);
+		u32 filesTaken = storage.fullFileCount;
+		while (true)
 		{
-			SCOPED_FUTEX(storage.availableFilesLock, lock);
-			if (storage.availableFilesCount == 0)
+			while (storage.availableFile)
 			{
-				lock.Leave();
-				storage.availableFilesEvent.IsSet();
+				storage.availableFile = storage.availableFile->next;
+				++filesTaken;
 			}
-			++locksTaken;
+			lock.Leave();
+			if (filesTaken == storage.fileCount)
+				break;
+			storage.availableFilesEvent.IsSet();
+			lock.Enter();
 		}
 
 		for (u32 i=0; i!=storage.fileCount; ++i)
 		{
-			//UBA_ASSERT(!storage.files[i].activeMapCount);
-			CloseFileMapping(m_logger, storage.files[i].handle, TC("FileMappingBuffer"));
-			CloseFile(nullptr, storage.files[i].file);
+			File& file = storage.files[i];
+			if (file.mappedMemory)
+			{
+				UnmapViewOfFile(m_logger, file.mappedMemory, file.capacity, TC("FileMappingBuffer"));
+				file.mappedMemory = nullptr;
+			}
+			//UBA_ASSERT(!file.activeMapCount);
+			CloseFileMapping(m_logger, file.handle, TC("FileMappingBuffer"));
+			CloseFile(nullptr, file.file);
 		}
 		storage.fileCount = 0;
-		storage.availableFilesCount = 0;
+		storage.availableFile = nullptr;
 	}
 
 	MappedView FileMappingBuffer::MapView(FileMappingHandle handle, u64 offset, u64 size, const tchar* hint)
@@ -601,21 +642,28 @@ namespace uba
 		u8 storageIndex = 255;
 		File& file = GetFile(handle, storageIndex);
 
-		u64 alignedOffsetStart = AlignUp(offset - (m_pageSize - 1), m_pageSize);
-
-		u64 newOffset = offset + size;
-		u64 alignedOffsetEnd = AlignUp(newOffset, m_pageSize);
-
+		u64 alignedOffsetStart = 0;
 		MappedView res;
 
-		u8* data = MapViewOfFile(m_logger, handle, FILE_MAP_WRITE, alignedOffsetStart, alignedOffsetEnd - alignedOffsetStart);
-		if (!data)
+		u8* data;
+		if (file.mappedMemory)
 		{
-			m_logger.Error(TC("%s - MapView failed to map view of file for %s with size %llu and offset %llu (%s)"), file.name, hint, size, offset, LastErrorToText().data);
-			return res;
+			data = file.mappedMemory;
+		}
+		else
+		{
+			alignedOffsetStart = AlignUp(offset - (m_pageSize - 1), m_pageSize);
+			u64 newOffset = offset + size;
+			u64 alignedOffsetEnd = AlignUp(newOffset, m_pageSize);
+			data = MapViewOfFile(m_logger, handle, FILE_MAP_WRITE, alignedOffsetStart, alignedOffsetEnd - alignedOffsetStart);
+			if (!data)
+			{
+				m_logger.Error(TC("%s - MapView failed to map view of file for %s with size %llu and offset %llu (%s)"), file.name, hint, size, offset, LastErrorToText().data);
+				return res;
+			}
 		}
 
-		++file.activeMapCount;
+		//++file.activeMapCount;
 
 		res.handle = handle;
 		res.offset = offset;
@@ -657,10 +705,9 @@ namespace uba
 
 				u8* memory = view.memory - (view.offset - alignedOffsetStart);
 				u64 mapSize = alignedOffsetEnd - alignedOffsetStart;
-				if (!UnmapViewOfFile(m_logger, memory, mapSize, hint))
-				{
-					m_logger.Error(TC("%s - Failed to unmap view on address %llx (offset %llu) - %s (%s)"), file.name, u64(memory), view.offset, hint, LastErrorToText().data);
-				}
+				if (!file.mappedMemory)
+					if (!UnmapViewOfFile(m_logger, memory, mapSize, hint))
+						m_logger.Error(TC("%s - Failed to unmap view on address %llx (offset %llu) - %s (%s)"), file.name, u64(memory), view.offset, hint, LastErrorToText().data);
 
 				if (newSize != InvalidValue)
 				{
@@ -675,7 +722,7 @@ namespace uba
 					PushFile(storage, &file);
 				}
 
-				--file.activeMapCount;
+				//--file.activeMapCount;
 			};
 		if (m_workManager && newSize == InvalidValue)
 			m_workManager->AddWork([=, h = TString(hint_)](const WorkContext&) { unmap(h.c_str()); }, 1, TC("UnmapView"));
@@ -697,11 +744,7 @@ namespace uba
 	}
 
 
-
-
-
-
-
+	////////////////////////////////////////////////////////////////////////////////////////////////////
 
 	FileMappingAllocator::FileMappingAllocator(Logger& logger, const tchar* name)
 	:	m_logger(logger)
diff --git a/Engine/Source/Programs/UnrealBuildAccelerator/Common/Private/UbaSession.cpp b/Engine/Source/Programs/UnrealBuildAccelerator/Common/Private/UbaSession.cpp
--- a/Engine/Source/Programs/UnrealBuildAccelerator/Common/Private/UbaSession.cpp
+++ b/Engine/Source/Programs/UnrealBuildAccelerator/Common/Private/UbaSession.cpp
@@ -230,6 +230,7 @@ namespace uba
 		table.GetValueAsBool(allowLocalDetour, TC("AllowLocalDetour"));
 		table.GetValueAsBool(extractObjFilesSymbols, TC("ExtractObjFilesSymbols"));
 		table.GetValueAsBool(useFakeVolumeSerial, TC("UseFakeVolumeSerial"));
+		table.GetValueAsBool(keepTransientDataMapped, TC("KeepTransientDataMapped"));
 		table.GetValueAsU32(traceReserveSizeMb, TC("TraceReserveSizeMb"));
 		table.GetValueAsU32(writeFilesBottleneck, TC("WriteFilesBottleneck"));
 		table.GetValueAsU32(writeFilesFileMapMaxMb, TC("WriteFilesFileMapMaxMB"));
@@ -1251,7 +1252,7 @@ namespace uba
 		m_directoryTable.m_lookup.reserve(30000);
 		m_fileMappingTableLookup.reserve(70000);
 
-		m_fileMappingBuffer.AddTransient(TC("FileMappings"));
+		m_fileMappingBuffer.AddTransient(TC("FileMappings"), info.keepTransientDataMapped);
 
 		if (!m_processCommunicationAllocator.Init(CommunicationMemSize, CommunicationMemSize * 512))
 		{
diff --git a/Engine/Source/Programs/UnrealBuildAccelerator/Common/Private/UbaStorage.cpp b/Engine/Source/Programs/UnrealBuildAccelerator/Common/Private/UbaStorage.cpp
--- a/Engine/Source/Programs/UnrealBuildAccelerator/Common/Private/UbaStorage.cpp
+++ b/Engine/Source/Programs/UnrealBuildAccelerator/Common/Private/UbaStorage.cpp
@@ -79,6 +79,7 @@ namespace uba
 		table.GetValueAsString(rootDir, TC("RootDir"));
 		table.GetValueAsBool(asyncUnmapViewOfFile, TC("AsyncUnmapViewOfFile"));
 		table.GetValueAsU64(casCapacityBytes, TC("CasCapacityBytes"));
+		table.GetValueAsBool(keepTransientDataMapped, TC("KeepTransientDataMapped"));
 	}
 
 	struct FileMappingScope
@@ -1252,7 +1253,7 @@ namespace uba
 
 		m_rootDir.Append(TC("cas")).EnsureEndsWithSlash();
 
-		m_casDataBuffer.AddTransient(TC("CasData"));
+		m_casDataBuffer.AddTransient(TC("CasData"), info.keepTransientDataMapped);
 
 		m_maxParallelCopyOrLink = info.maxParallelCopyOrLink;
 
diff --git a/Engine/Source/Programs/UnrealBuildAccelerator/Common/Public/UbaFileMapping.h b/Engine/Source/Programs/UnrealBuildAccelerator/Common/Public/UbaFileMapping.h
--- a/Engine/Source/Programs/UnrealBuildAccelerator/Common/Public/UbaFileMapping.h
+++ b/Engine/Source/Programs/UnrealBuildAccelerator/Common/Public/UbaFileMapping.h
@@ -11,6 +11,8 @@ namespace uba
 	class Logger;
 	class WorkManager;
 
+	////////////////////////////////////////////////////////////////////////////////////////////////////
+
 #if !PLATFORM_WINDOWS
 	inline constexpr u32 FILE_MAP_WRITE = 0x0002;
 	inline constexpr u32 FILE_MAP_READ = 0x0004;
@@ -34,6 +36,8 @@ namespace uba
 	bool DuplicateFileMapping(Logger& logger, ProcHandle sourceProcessHandle, FileMappingHandle sourceHandle, ProcHandle targetProcessHandle, FileMappingHandle& targetHandle, u32 dwDesiredAccess, bool bInheritHandle, u32 dwOptions, const tchar* hint);
 	void MapMemoryCopy(void* dest, const void* source, u64 size);
 
+	////////////////////////////////////////////////////////////////////////////////////////////////////
+
 	struct MappedView
 	{
 		FileMappingHandle handle;
@@ -52,10 +56,10 @@ namespace uba
 	class FileMappingBuffer
 	{
 	public:
-		FileMappingBuffer(Logger& logger, WorkManager* workManager);
+		FileMappingBuffer(Logger& logger, WorkManager* workManager = nullptr);
 		~FileMappingBuffer();
 
-		bool AddTransient(const tchar* name);
+		bool AddTransient(const tchar* name, bool keepMapped);
 		bool AddPersistent(const tchar* name, FileHandle fileHandle, u64 size, u64 capacity);
 		void CloseDatabase();
 
@@ -74,13 +78,16 @@ namespace uba
 	private:
 		struct File
 		{
+			File* next = nullptr;
 			const tchar* name = nullptr;
 			FileHandle file = InvalidFileHandle;
 			FileMappingHandle handle;
 			u64 size = 0;
 			u64 capacity = 0;
-			Atomic<u64> activeMapCount;
+			u8* mappedMemory = nullptr;
+			//Atomic<u32> activeMapCount;
 			bool commitOnAlloc = false;
+			bool keepMapped = false;
 		};
 
 		MappedView AllocAndMapViewNoLock(File& file, u64 size, u64 alignment, const tchar* hint);
@@ -94,11 +101,11 @@ namespace uba
 		{
 			File files[128];
 			u32 fileCount = 0;
+			u32 fullFileCount = 0;
 
 			Futex availableFilesLock;
 			Event availableFilesEvent;
-			File* availableFiles[128] = { 0 };
-			u32 availableFilesCount = 0;
+			File* availableFile = nullptr;
 
 			MappingStorage() : availableFilesEvent(false) {}
 		};
@@ -113,9 +120,7 @@ namespace uba
 		void operator=(const FileMappingBuffer&) = delete;
 	};
 
-
-
-
+	////////////////////////////////////////////////////////////////////////////////////////////////////
 
 	class FileMappingAllocator
 	{
diff --git a/Engine/Source/Programs/UnrealBuildAccelerator/Common/Public/UbaSessionCreateInfo.h b/Engine/Source/Programs/UnrealBuildAccelerator/Common/Public/UbaSessionCreateInfo.h
--- a/Engine/Source/Programs/UnrealBuildAccelerator/Common/Public/UbaSessionCreateInfo.h
+++ b/Engine/Source/Programs/UnrealBuildAccelerator/Common/Public/UbaSessionCreateInfo.h
@@ -40,6 +40,7 @@ namespace uba
 		bool extractObjFilesSymbols = false;	// Will extract import/export symbols and clear export table from obj file.
 		bool treatTempDirAsEmpty = true;		// If this is true, directory table will always see temp as empty. Use sub directories of temp instead
 		bool useFakeVolumeSerial = true;		// Fake volume serials reduce the directory table size since serials end up being compressed down to one byte
+		bool keepTransientDataMapped = true;	// Will keep the transient data mapped in instead of doing map/unmap everytime we access it
 		u64 deleteSessionsOlderThanSeconds = 12 * 60 * 60; // Delete session folders older than 12 hours by default . Set to 0 to not delete or 1 to delete all
 		u64 keepOutputFileMemoryMapsThreshold = 256 * 1024; // If allowMemoryMaps is true, output files will be kept in memory if smaller than this size
 		u32 traceReserveSizeMb = 128;			// Memory reserved for trace file (in mb)
diff --git a/Engine/Source/Programs/UnrealBuildAccelerator/Common/Public/UbaStorage.h b/Engine/Source/Programs/UnrealBuildAccelerator/Common/Public/UbaStorage.h
--- a/Engine/Source/Programs/UnrealBuildAccelerator/Common/Public/UbaStorage.h
+++ b/Engine/Source/Programs/UnrealBuildAccelerator/Common/Public/UbaStorage.h
@@ -89,6 +89,7 @@ namespace uba
 		bool manuallyHandleOverflow = false;
 		bool asyncUnmapViewOfFile = true;
 		bool writeToDisk = true;
+		bool keepTransientDataMapped = true;	// Will keep the transient data mapped in instead of doing map/unmap everytime we access it
 		WorkManager* workManager = nullptr;
 	};
 
diff --git a/Engine/Source/Programs/UnrealBuildAccelerator/Test/Private/UbaTestAll.h b/Engine/Source/Programs/UnrealBuildAccelerator/Test/Private/UbaTestAll.h
--- a/Engine/Source/Programs/UnrealBuildAccelerator/Test/Private/UbaTestAll.h
+++ b/Engine/Source/Programs/UnrealBuildAccelerator/Test/Private/UbaTestAll.h
@@ -47,6 +47,7 @@ namespace uba
 		UBA_TEST(TestRegisterChanges) \
 		UBA_TEST(TestRegisterChangesRemote) \
 		UBA_TEST(TestDetouredClang) \
+		UBA_TEST(TestFileMappingBuffer) \
 		UBA_TEST_EXTRA \
 
 #define UBA_POSIX_TESTS \
diff --git a/Engine/Source/Programs/UnrealBuildAccelerator/Test/Private/UbaTestFileMappings.cpp b/Engine/Source/Programs/UnrealBuildAccelerator/Test/Private/UbaTestFileMappings.cpp
new file mode 100644
--- /dev/null
+++ b/Engine/Source/Programs/UnrealBuildAccelerator/Test/Private/UbaTestFileMappings.cpp
@@ -0,0 +1,25 @@
+// Copyright Epic Games, Inc. All Rights Reserved.
+
+#include "UbaLogger.h"
+#include "UbaFileMapping.h"
+
+namespace uba
+{
+	bool TestFileMappingBuffer(LoggerWithWriter& logger, const StringBufferBase& rootDir)
+	{
+
+		FileMappingBuffer mappingBuffer(logger);
+
+		mappingBuffer.AddTransient(TC("Temp"), false);
+
+		MappedView view = mappingBuffer.AllocAndMapView(MappedView_Transient, 1024, 1, TC("Foo"));
+
+		*(u64*)view.memory = 1337;
+		u64 value = *(u64*)view.memory;
+		if (value != 1337)
+			return false;
+		mappingBuffer.UnmapView(view, TC("Foo"));
+
+		return true;
+	}
+}
