[UBA]
* Changed so when running in wine we use NtQueryFullAttributesFile instead of NtQueryInformationFile because it is much faster. Seems like there is a huge penalty in wine to go from HANDLE to info than path to info

diff --git a/Engine/Source/Programs/UnrealBuildAccelerator/Common/Private/UbaFile.cpp b/Engine/Source/Programs/UnrealBuildAccelerator/Common/Private/UbaFile.cpp
--- a/Engine/Source/Programs/UnrealBuildAccelerator/Common/Private/UbaFile.cpp
+++ b/Engine/Source/Programs/UnrealBuildAccelerator/Common/Private/UbaFile.cpp
@@ -13,7 +13,9 @@
 #include <WinCon.h>
 #include <RestartManager.h>
 #pragma comment(lib, "Rstrtmgr.lib")
+struct FILE_NETWORK_OPEN_INFORMATION { LARGE_INTEGER CreationTime; LARGE_INTEGER LastAccessTime; LARGE_INTEGER LastWriteTime; LARGE_INTEGER ChangeTime; LARGE_INTEGER AllocationSize; LARGE_INTEGER EndOfFile; ULONG FileAttributes; };
 extern "C" NTSTATUS NTAPI NtQueryInformationFile(HANDLE FileHandle, PIO_STATUS_BLOCK IoStatusBlock, PVOID FileInformation, ULONG Length, FILE_INFORMATION_CLASS FileInformationClass);
+extern "C" NTSTATUS NTAPI NtQueryFullAttributesFile(const OBJECT_ATTRIBUTES *attr, FILE_NETWORK_OPEN_INFORMATION *info);
 #else
 #include <dlfcn.h>
 #include <execinfo.h>
@@ -74,16 +76,18 @@ namespace uba
 		return (HANDLE)(fh == InvalidFileHandle ? InvalidFileHandle : (fh & FileHandleFlagMask));
 	}
 
-	#define MAKE_LONG_FILENAME(fileName) \
+	#define PREFIX_FILENAME(fileName, prefix) \
 		UBA_ASSERT(TStrlen(fileName) < MaxPath); \
 		StringBuffer<MaxPath> STRING_JOIN(longName, __LINE__); \
 		if (IsAbsolutePath(fileName)) \
 		{ \
 			auto& lsb = STRING_JOIN(longName, __LINE__); \
-			lsb.Append(TC("\\\\?\\")); \
+			lsb.Append(TC(prefix)); \
 			FixPath(fileName, nullptr, 0, lsb); \
 			fileName = lsb.data; \
 		}
+	#define MAKE_LONG_FILENAME(fileName) PREFIX_FILENAME(fileName, "\\\\?\\")
+	#define MAKE_NT_FILENAME(fileName)  PREFIX_FILENAME(fileName, "\\??\\")
 
 #else
 	int asFileDescriptor(FileHandle fh)
@@ -175,22 +179,28 @@ namespace uba
 #if PLATFORM_WINDOWS
 		ExtendedTimerScope ts(KernelStats::GetCurrent().getFileInfo);
 
-		struct FILE_NETWORK_OPEN_INFORMATION
-		{
-			  LARGE_INTEGER CreationTime;
-			  LARGE_INTEGER LastAccessTime;
-			  LARGE_INTEGER LastWriteTime;
-			  LARGE_INTEGER ChangeTime;
-			  LARGE_INTEGER AllocationSize;
-			  LARGE_INTEGER EndOfFile;
-			  ULONG FileAttributes;
-		};
-
 		FILE_NETWORK_OPEN_INFORMATION info;
-		IO_STATUS_BLOCK b;
-		NTSTATUS res = NtQueryInformationFile(asHANDLE(hFile), &b, &info, sizeof(info), (FILE_INFORMATION_CLASS)34); // FileNetworkOpenInformation
-		if (res != STATUS_SUCCESS)
-			return logger.Error(TC("GetFileBasicInformationByHandle (NtQueryInformationFile) failed on %s (%u)"), fileName, res);
+		if (IsRunningWine())
+		{
+			MAKE_NT_FILENAME(fileName);
+			UNICODE_STRING us;
+			us.Length = USHORT(TStrlen(fileName)*sizeof(tchar));
+			us.MaximumLength = us.Length;
+			us.Buffer = (tchar*)fileName;
+			OBJECT_ATTRIBUTES oa;
+			InitializeObjectAttributes(&oa, &us, OBJ_CASE_INSENSITIVE, NULL, NULL);
+			NTSTATUS res = NtQueryFullAttributesFile(&oa, &info);
+			if (res != STATUS_SUCCESS)
+				return logger.Error(TC("GetFileBasicInformationByHandle (NtQueryFullAttributesFile) failed on %s (%u)"), fileName, res);
+		}
+		else
+		{
+			constexpr auto FileNetworkOpenInformation = (FILE_INFORMATION_CLASS)34;
+			IO_STATUS_BLOCK b;
+			NTSTATUS res = NtQueryInformationFile(asHANDLE(hFile), &b, &info, sizeof(info), FileNetworkOpenInformation);
+			if (res != STATUS_SUCCESS)
+				return logger.Error(TC("GetFileBasicInformationByHandle (NtQueryInformationFile) failed on %s (%u)"), fileName, res);
+		}
 		out.attributes = info.FileAttributes;
 		out.lastWriteTime = info.LastWriteTime.QuadPart;
 		out.size = info.EndOfFile.QuadPart;
@@ -206,29 +216,6 @@ namespace uba
 #endif
 	}
 
-	bool GetFileBasicInformation(FileBasicInformation& out, Logger& logger, const tchar* fileName)
-	{
-#if PLATFORM_WINDOWS
-		ExtendedTimerScope ts(KernelStats::GetCurrent().getFileInfo);
-		MAKE_LONG_FILENAME(fileName);
-		FileHandle h = uba::CreateFileW(fileName, 0, 0, OPEN_EXISTING, FILE_FLAG_BACKUP_SEMANTICS);
-		if (h == InvalidFileHandle)
-			return false;// logger.Error(TC("GetFileInformation: CreateFile failed for file %s (%s)"), fileName, LastErrorToText().data);
-		auto handleGuard = MakeGuard([&]() { uba::CloseFile(fileName, h); });
-		if (!GetFileBasicInformationByHandle(out, logger, fileName, h))
-			return false;// logger.Error(TC("Failed to get file information for %s while checking file added for write. This should not happen! (%s)"), fileName, LastErrorToText().data);
-		return true;
-#else
-		FileInformation info;
-		if (!GetFileInformation(info, logger, fileName))
-			return false;
-		out.attributes = info.attributes;
-		out.lastWriteTime = info.lastWriteTime;
-		out.size = info.size;
-		return true;
-#endif
-	}
-
 	bool GetFileInformationByHandle(FileInformation& out, Logger& logger, const tchar* fileName, FileHandle hFile)
 	{
 		ExtendedTimerScope ts(KernelStats::GetCurrent().getFileInfo);
diff --git a/Engine/Source/Programs/UnrealBuildAccelerator/Common/Private/UbaSession.cpp b/Engine/Source/Programs/UnrealBuildAccelerator/Common/Private/UbaSession.cpp
--- a/Engine/Source/Programs/UnrealBuildAccelerator/Common/Private/UbaSession.cpp
+++ b/Engine/Source/Programs/UnrealBuildAccelerator/Common/Private/UbaSession.cpp
@@ -561,7 +561,10 @@ namespace uba
 		entry.handled = true;
 
 		bool isDir = false;
-		FileHandle fileHandle = uba::CreateFileW(fileName, GENERIC_READ, FILE_SHARE_READ, OPEN_EXISTING, DefaultAttributes());
+		u32 attributes = DefaultAttributes();
+		if (!isCompressed)
+			attributes |= FILE_FLAG_SEQUENTIAL_SCAN;
+		FileHandle fileHandle = uba::CreateFileW(fileName, GENERIC_READ, FILE_SHARE_READ, OPEN_EXISTING, attributes);
 		if (fileHandle == InvalidFileHandle)
 		{
 			u32 error = GetLastError();
diff --git a/Engine/Source/Programs/UnrealBuildAccelerator/Common/Public/UbaFile.h b/Engine/Source/Programs/UnrealBuildAccelerator/Common/Public/UbaFile.h
--- a/Engine/Source/Programs/UnrealBuildAccelerator/Common/Public/UbaFile.h
+++ b/Engine/Source/Programs/UnrealBuildAccelerator/Common/Public/UbaFile.h
@@ -24,6 +24,7 @@ namespace uba
 	inline constexpr u32 FILE_SHARE_WRITE = 0x00000002;
 	inline constexpr u32 FILE_SHARE_READ = 0x00000001;
 	inline constexpr u32 FILE_FLAG_BACKUP_SEMANTICS = 0x02000000;
+	inline constexpr u32 FILE_FLAG_SEQUENTIAL_SCAN = 0;
 	inline constexpr u32 OPEN_EXISTING = 3;
 	inline constexpr u32 PAGE_READONLY = 0x02;
 	#endif
@@ -41,7 +42,6 @@ namespace uba
 		u64 size;
 	};
 	bool GetFileBasicInformationByHandle(FileBasicInformation& out, Logger& logger, const tchar* fileName, FileHandle hFile);
-	bool GetFileBasicInformation(FileBasicInformation& out, Logger& logger, const tchar* fileName);
 
 	struct FileInformation
 	{
