[UBA}
* Changed socket connect remote address timeout for local or loopback is 100ms

diff --git a/Engine/Source/Programs/UnrealBuildAccelerator/Common/Private/UbaNetworkBackendTcp.cpp b/Engine/Source/Programs/UnrealBuildAccelerator/Common/Private/UbaNetworkBackendTcp.cpp
--- a/Engine/Source/Programs/UnrealBuildAccelerator/Common/Private/UbaNetworkBackendTcp.cpp
+++ b/Engine/Source/Programs/UnrealBuildAccelerator/Common/Private/UbaNetworkBackendTcp.cpp
@@ -636,11 +636,9 @@ namespace uba
 		}
 #endif
 
-		// Return to blocking since we want select to block
-		if (!SetBlocking(logger, socketFd, true))
-			return false;
-
 		int timeoutMs = 2000;
+		if (nameHint && (Equals(nameHint, TC("localhost")) || Equals(nameHint, TC("127.0.0.1"))))
+			timeoutMs = 100;
 
 		WSAPOLLFD p;
 		p.fd = socketFd;
@@ -669,6 +667,9 @@ namespace uba
 			return false;
 		}
 
+		// Return to blocking since we want select to block
+		if (!SetBlocking(logger, socketFd, true))
+			return false;
 
 #if !PLATFORM_WINDOWS
 		// Before we send anything even though the
