[UBA]
* Rewrite of code handling compressed obj files. It is now more generic and stable for when files move from compressed to non-compressed and back. New version should handle some previous bugs related to remote linking

diff --git a/Engine/Source/Programs/UnrealBuildAccelerator/Common/Private/UbaCacheClient.cpp b/Engine/Source/Programs/UnrealBuildAccelerator/Common/Private/UbaCacheClient.cpp
--- a/Engine/Source/Programs/UnrealBuildAccelerator/Common/Private/UbaCacheClient.cpp
+++ b/Engine/Source/Programs/UnrealBuildAccelerator/Common/Private/UbaCacheClient.cpp
@@ -316,16 +316,15 @@ namespace uba
 			{
 				bool shouldValidate = (m_validateCacheWritesInput && !isOutput) || (m_validateCacheWritesOutput && isOutput);
 				bool deferCreation = true;
-				bool fileIsCompressed = IsFileCompressed(info, path, isOutput);
 
 				if (isOutput)
 				{
-					if (!m_storage.StoreCasFile(casKey, path.data, CasKeyZero, deferCreation, fileIsCompressed))
+					if (!m_storage.StoreCasFile(casKey, path.data, CasKeyZero, deferCreation))
 						return false;
 				}
 				else
 				{
-					if (!m_storage.StoreCasKey(casKey, path.data, CasKeyZero, fileIsCompressed))
+					if (!m_storage.StoreCasKey(casKey, path.data, CasKeyZero))
 						return false;
 				}
 
@@ -582,13 +581,11 @@ namespace uba
 					if (m_useDirectoryPreParsing)
 						PreparseDirectory(fileNameKey, path);
 
-					bool fileIsCompressed = IsFileCompressed(info, path, false);
-
 					#if UBA_TRACE_FETCH_CACHE
 					u64 t = GetTime();
 					#endif
 
-					m_storage.StoreCasKey(localCasKey, fileNameKey, path.data, CasKeyZero, fileIsCompressed);
+					m_storage.StoreCasKey(localCasKey, fileNameKey, path.data, CasKeyZero);
 
 					#if UBA_TRACE_FETCH_CACHE
 					storeTime += GetTime() - t;
@@ -878,8 +875,8 @@ namespace uba
 		else
 		{
 			DowngradedLogger logger(m_connected, m_logger.m_writer, TC("UbaCacheClientDownload"));
-			bool destinationIsCompressed = IsFileCompressed(info, path, true);
-			if (!fetcher.RetrieveFile(logger, m_client, casKey, path.data, destinationIsCompressed))
+			bool writeCompressed = m_session.ShouldStoreObjFilesCompressed() && g_globalRules.FileCanBeCompressed(path);
+			if (!fetcher.RetrieveFile(logger, m_client, casKey, path.data, writeCompressed))
 				return logger.Error(TC("Failed to download cache output %s for %s"), path.data, info.GetDescription());
 		}
 
@@ -1375,8 +1372,7 @@ namespace uba
 		{
 			// Add hash of application binary to key
 			CasKey applicationCasKey;
-			bool fileIsCompressed = false;
-			if (!m_storage.StoreCasKey(applicationCasKey, info.application, CasKeyZero, fileIsCompressed))
+			if (!m_storage.StoreCasKey(applicationCasKey, info.application, CasKeyZero))
 				return CasKeyZero;
 			if (report)
 				m_logger.Info(TC("   %s %s"), CasKeyString(applicationCasKey).str, info.application);
@@ -1469,19 +1465,6 @@ namespace uba
 		return true;
 	}
 
-	bool CacheClient::IsFileCompressed(const ProcessStartInfo& info, const StringView& filename, bool isOutput)
-	{
-		if (!m_session.ShouldStoreObjFilesCompressed())
-			return false;
-		auto rules = info.rules;
-		if (!rules)
-			rules = m_session.GetRules(info);
-		if (isOutput)
-			return rules->StoreFileCompressed(filename);
-		else
-			return rules->ShouldDecompressFiles(filename);
-	}
-
 	void CacheClient::PreparseDirectory(const StringKey& fileNameKey, const StringBufferBase& filePath)
 	{
 		const tchar* lastSep = filePath.Last(PathSeparator);
diff --git a/Engine/Source/Programs/UnrealBuildAccelerator/Common/Private/UbaCacheServer.cpp b/Engine/Source/Programs/UnrealBuildAccelerator/Common/Private/UbaCacheServer.cpp
--- a/Engine/Source/Programs/UnrealBuildAccelerator/Common/Private/UbaCacheServer.cpp
+++ b/Engine/Source/Programs/UnrealBuildAccelerator/Common/Private/UbaCacheServer.cpp
@@ -2507,8 +2507,7 @@ namespace uba
 
 		CasKey key;
 		bool deferCreation = false;
-		bool checkIfCompressed = false;
-		if (!m_storage.StoreCasFile(key, tempFile.data, CasKeyZero, deferCreation, checkIfCompressed))
+		if (!m_storage.StoreCasFile(key, tempFile.data, CasKeyZero, deferCreation))
 			return false;
 
 		writer.WriteCasKey(key);
diff --git a/Engine/Source/Programs/UnrealBuildAccelerator/Common/Private/UbaFile.cpp b/Engine/Source/Programs/UnrealBuildAccelerator/Common/Private/UbaFile.cpp
--- a/Engine/Source/Programs/UnrealBuildAccelerator/Common/Private/UbaFile.cpp
+++ b/Engine/Source/Programs/UnrealBuildAccelerator/Common/Private/UbaFile.cpp
@@ -324,6 +324,19 @@ namespace uba
 #endif
 	}
 
+	bool SetFilePointer(Logger& logger, const tchar* fileName, FileHandle handle, u64 position)
+	{
+#if PLATFORM_WINDOWS
+		if (!::SetFilePointerEx(asHANDLE(handle), ToLargeInteger(position), NULL, FILE_BEGIN))
+			return logger.Error(TC("SetFilePointerEx failed on %s (%s)"), fileName, LastErrorToText().data);
+		return true;
+#else
+		if (lseek(asFileDescriptor(hFile), position, SEEK_SET) != position)
+			return logger.Error("lseek to %llu failed for %s: %s\n", position, hint, strerror(errno));
+		return true;
+#endif
+	}
+
 	bool SetEndOfFile(Logger& logger, const tchar* fileName, FileHandle handle, u64 size)
 	{
 		ExtendedTimerScope ts(KernelStats::GetCurrent().setFileInfo);
diff --git a/Engine/Source/Programs/UnrealBuildAccelerator/Common/Private/UbaSession.cpp b/Engine/Source/Programs/UnrealBuildAccelerator/Common/Private/UbaSession.cpp
--- a/Engine/Source/Programs/UnrealBuildAccelerator/Common/Private/UbaSession.cpp
+++ b/Engine/Source/Programs/UnrealBuildAccelerator/Common/Private/UbaSession.cpp
@@ -1159,7 +1159,7 @@ namespace uba
 		UBA_ASSERTF(m_shouldWriteToDisk || m_allowMemoryMaps, TC("Can't disable both should write to disk and allow memory maps"));
 
 		m_storeObjFilesCompressed = info.storeObjFilesCompressed && IsWindows; // Non-windows not implemented (yet)
-		m_readObjFilesCompressed = m_storeObjFilesCompressed || (info.readObjFilesCompressed && IsWindows);
+		m_readObjFilesCompressed = (m_storeObjFilesCompressed || (info.readObjFilesCompressed && IsWindows)) && !runningRemote; // with remote we decompress the files into memory
 		m_allowLocalDetour = info.allowLocalDetour;
 
 		m_detailedTrace = info.detailedTrace;
@@ -2099,22 +2099,20 @@ namespace uba
 		// Reason is because there is additional overhead to use memory mappings.
 		// Upside is that all things that are memory mapped can be stored compressed in cas storage so it saves space.
 
-		if (runningRemote && fileName.EndsWith(TC(".pch")))
-			return 64 * 1024; // pch needs 64k alignment
 		if (fileName.EndsWith(TC(".h")) || fileName.EndsWith(TC(".inl")) || fileName.EndsWith(TC(".gch")))
 			return 4 * 1024; // clang seems to need 4k alignment? Is it a coincidence it works or what is happening inside the code? (msvc works with alignment 1byte here)
 		if (fileName.EndsWith(TC(".h.obj")))
 			return 4 * 1024;
 		if (fileName.EndsWith(TC(".lib")))
 			return 4 * 1024;
-		//if (fileName.EndsWith(TC(".lib")))
-		//	return 4 * 1024;
-		//if (fileName.EndsWith(TC(".rc2.res")))
-		//	return 64;
-		//if (fileName.EndsWith(TC(".rsp")))
-		//	return 4 * 1024; // rsp is read over and over again
-		//if (fileName.EndsWith(TC(".h.obj")) || fileName.EndsWith(TC(".lib")))
-		//	return 4 * 1024; // rsp is read over and over again
+
+		if (runningRemote) // We store these compressed to save space
+		{
+			if (fileName.EndsWith(TC(".obj")) || fileName.EndsWith(TC(".o")))
+				return 4 * 1024; // pch needs 64k alignment
+			if (fileName.EndsWith(TC(".pch")))
+				return 64 * 1024; // pch needs 64k alignment
+		}
 		return 0;
 	}
 
@@ -2623,7 +2621,7 @@ namespace uba
 				auto& rules = *process.m_startInfo.rules;
 				RootsHandle rootsHandle = process.GetStartInfo().rootsHandle;
 
-				bool storeCompressed = m_storeObjFilesCompressed && rules.StoreFileCompressed(file.name);
+				bool storeCompressed = m_storeObjFilesCompressed && g_globalRules.FileCanBeCompressed(file.name);
 				bool shouldDevirtualize = false;
 				bool escapeSpaces = false;
 				if (!storeCompressed)
diff --git a/Engine/Source/Programs/UnrealBuildAccelerator/Common/Private/UbaSessionServer.cpp b/Engine/Source/Programs/UnrealBuildAccelerator/Common/Private/UbaSessionServer.cpp
--- a/Engine/Source/Programs/UnrealBuildAccelerator/Common/Private/UbaSessionServer.cpp
+++ b/Engine/Source/Programs/UnrealBuildAccelerator/Common/Private/UbaSessionServer.cpp
@@ -491,8 +491,7 @@ namespace uba
 
 			CasKey casKey;
 			bool deferCreation = true;
-			bool fileIsCompressed = false;
-			if (!m_storage.StoreCasFile(casKey, path, CasKeyZero, deferCreation, fileIsCompressed))
+			if (!m_storage.StoreCasFile(casKey, path, CasKeyZero, deferCreation))
 				return false;
 			UBA_ASSERTF(casKey != CasKeyZero, TC("Failed to store cas for %s when calculating key for tracked inputs on %s"), path, fileName);
 			hasher.Update(&casKey, sizeof(CasKey));
@@ -713,7 +712,6 @@ namespace uba
 
 			StringBuffer<> detoursLib;
 			bool deferCreation = true;
-			bool fileIsCompressed = false;
 
 			// Handle x64 for both architectures and arm64 only for arm64 clients
 			for (u32 i=0; i!=(isClientArm ? 2u : 1u); ++i)
@@ -721,7 +719,7 @@ namespace uba
 				if (detoursBinaryKey[i] != CasKeyZero)
 					continue;
 				detoursLib.Clear().Append(m_detoursLibrary[i].c_str());
-				if (!m_storage.StoreCasFile(detoursBinaryKey[i], detoursLib.data, CasKeyZero, deferCreation, fileIsCompressed) || detoursBinaryKey[i] == CasKeyZero)
+				if (!m_storage.StoreCasFile(detoursBinaryKey[i], detoursLib.data, CasKeyZero, deferCreation) || detoursBinaryKey[i] == CasKeyZero)
 					return m_logger.Error(TC("Failed to create cas for %s"), detoursLib.data);
 				m_detoursBinaryKey[i] = detoursBinaryKey[i];
 			}
@@ -736,14 +734,14 @@ namespace uba
 						return false;
 				UBA_ASSERT(IsWindows);
 				agentDir.Append(PathSeparator).Append(UBA_AGENT_EXECUTABLE);
-				if (!m_storage.StoreCasFile(agentBinaryKey, agentDir.data, CasKeyZero, deferCreation, fileIsCompressed) || agentBinaryKey == CasKeyZero)
+				if (!m_storage.StoreCasFile(agentBinaryKey, agentDir.data, CasKeyZero, deferCreation) || agentBinaryKey == CasKeyZero)
 				{
 					// This is hacky but uba binary is not copied by ubt anymore.
 					StringBuffer<> dir2;
 					if (!GetAlternativeUbaPath(m_logger, dir2, agentDir, IsWindows && isClientArm))
 						return false;
 					dir2.Append(UBA_AGENT_EXECUTABLE);
-					if (!m_storage.StoreCasFile(agentBinaryKey, dir2.data, CasKeyZero, deferCreation, fileIsCompressed) || agentBinaryKey == CasKeyZero)
+					if (!m_storage.StoreCasFile(agentBinaryKey, dir2.data, CasKeyZero, deferCreation) || agentBinaryKey == CasKeyZero)
 						return m_logger.Error(TC("Failed to create cas for %s"), dir2.data);
 				}
 			}
@@ -1091,7 +1089,7 @@ namespace uba
 				}
 				else if (m_storeObjFilesCompressed)
 				{
-					writeCompressed = rules.StoreFileCompressed(destination);
+					writeCompressed = g_globalRules.FileCanBeCompressed(destination);
 				}
 				success = m_storage.CopyOrLink(casKey, destination.data, attributes, writeCompressed, formattingFunc);
 				if (!success)
@@ -1759,8 +1757,7 @@ namespace uba
 			}
 		}
 
-		bool fileIsCompressed = false;
-		if (!m_storage.StoreCasFile(out, fileName, casKeyOverride, deferCreation, fileIsCompressed)) // We can defer the creation of the cas file since client might already have it
+		if (!m_storage.StoreCasFile(out, fileName, casKeyOverride, deferCreation)) // We can defer the creation of the cas file since client might already have it
 			return false;
 		return true;//out != CasKeyZero;
 	}
@@ -2399,14 +2396,13 @@ namespace uba
 		StringBuffer<> dir;
 		if (GetDirectoryOfCurrentModule(m_logger, dir))
 		{
-			bool fileIsCompressed = false;
 			bool deferCreation = true;
 			auto ChangeToSymbolExtension = [](StringBufferBase& str) -> StringBufferBase& { IsWindows ? str.Resize(str.count - 3).Append("pdb") : str.Resize(str.count - 2).Append("debug"); return str; };
-			if (!m_storage.StoreCasFile(detoursSymbolsKey, ChangeToSymbolExtension(dir).data, CasKeyZero, deferCreation, fileIsCompressed) || detoursSymbolsKey == CasKeyZero)
+			if (!m_storage.StoreCasFile(detoursSymbolsKey, ChangeToSymbolExtension(dir).data, CasKeyZero, deferCreation) || detoursSymbolsKey == CasKeyZero)
 			{
 				StringBuffer<> dir2;
 				if (GetAlternativeUbaPath(m_logger, dir2, dir, IsWindows && isClientArm))
-					m_storage.StoreCasFile(detoursSymbolsKey, ChangeToSymbolExtension(dir2.Append(UBA_DETOURS_LIBRARY)).data, CasKeyZero, deferCreation, fileIsCompressed);
+					m_storage.StoreCasFile(detoursSymbolsKey, ChangeToSymbolExtension(dir2.Append(UBA_DETOURS_LIBRARY)).data, CasKeyZero, deferCreation);
 			}
 		}
 		writer.WriteCasKey(detoursSymbolsKey);
diff --git a/Engine/Source/Programs/UnrealBuildAccelerator/Common/Private/UbaStorage.cpp b/Engine/Source/Programs/UnrealBuildAccelerator/Common/Private/UbaStorage.cpp
--- a/Engine/Source/Programs/UnrealBuildAccelerator/Common/Private/UbaStorage.cpp
+++ b/Engine/Source/Programs/UnrealBuildAccelerator/Common/Private/UbaStorage.cpp
@@ -1,6 +1,7 @@
 // Copyright Epic Games, Inc. All Rights Reserved.
 
 #include "UbaStorage.h"
+#include "UbaApplicationRules.h"
 #include "UbaBottleneck.h"
 #include "UbaCompressedObjFileHeader.h"
 #include "UbaConfig.h"
@@ -291,36 +292,6 @@ namespace uba
 		m_externalFileMappingsProvider = std::move(provider);
 	}
 
-	bool StorageImpl::WriteCompressed(WriteResult& out, const StringKey& fromKey, const tchar* from, const tchar* to)
-	{
-		StorageStats& stats = Stats();
-		TimerScope ts(stats.createCas);
-
-		if (m_externalFileMappingsProvider)
-		{
-			ExternalFileMapping externalMapping;
-			if (m_externalFileMappingsProvider(externalMapping, fromKey, from))
-			{
-				FileMappingScope scope(m_logger);
-				u8* fileMem = scope.MapView(TC("WriteCompressed"), externalMapping.handle, externalMapping.offset, externalMapping.size, from);
-				if (!fileMem)
-					return false;
-				return StorageImpl::WriteCompressed(out, from, InvalidFileHandle, fileMem, externalMapping.size, to, nullptr, 0);
-			}
-		}
-
-		FileHandle readHandle;
-		if (!OpenFileSequentialRead(m_logger, from, readHandle))
-			return m_logger.Error(TC("[WriteCompressed] Failed to open file %s for read (%s)"), from, LastErrorToText().data);
-		auto rsg = MakeGuard([&]() { CloseFile(from, readHandle); });
-
-		u64 fileSize;
-		if (!uba::GetFileSizeEx(fileSize, readHandle))
-			return m_logger.Error(TC("[WriteCompressed] GetFileSize failed for %s (%s)"), from, LastErrorToText().data);
-
-		return WriteCompressed(out, from, readHandle, 0, fileSize, to, nullptr, 0);
-	}
-
 	bool StorageImpl::WriteCompressed(WriteResult& out, const tchar* from, FileHandle readHandle, u8* readMem, u64 fileSize, const tchar* to, const void* header, u64 headerSize, u64 lastWriteTime)
 	{
 		StorageStats& stats = Stats();
@@ -561,65 +532,86 @@ namespace uba
 		return true;
 	}
 
-	bool StorageImpl::WriteCasFileNoCheck(WriteResult& out, const StringKey& fileNameKey, const tchar* fileName, bool fileIsCompressed, const tchar* casFile, bool storeCompressed)
+	bool StorageImpl::WriteCasFileNoCheck(WriteResult& out, const StringKey& fileNameKey, const tchar* fileName, const CasKey& casKey, const tchar* casFile, bool storeCompressed)
 	{
-		if (storeCompressed && !fileIsCompressed)
-		{
-			if (!WriteCompressed(out, fileNameKey, fileName, casFile))
-				return false;
-		}
-		else if (fileIsCompressed && !storeCompressed)
-		{
-			return m_logger.Error(TC("Write compressed file to uncompressed store not implemented"));
-		}
-		else
-		{
-			FileHandle readHandle;
-			if (!OpenFileSequentialRead(m_logger, fileName, readHandle))
-				return m_logger.Error(TC("[WriteCasFileNoCheck] Failed to open file %s for read (%s)"), fileName, LastErrorToText().data);
-			auto fileGuard = MakeGuard([&](){ CloseFile(fileName, readHandle); });
+		StorageStats& stats = Stats();
+		TimerScope ts(stats.createCas);
 
-			u64 fileSize;
-			if (!uba::GetFileSizeEx(fileSize, readHandle))
-				return m_logger.Error(TC("[WriteCasFileNoCheck] GetFileSize failed for %s (%s)"), fileName, LastErrorToText().data);
+		FileHandle readHandle;
+		if (!OpenFileSequentialRead(m_logger, fileName, readHandle))
+			return m_logger.Error(TC("[WriteCasFileNoCheck] Failed to open file %s for read (%s)"), fileName, LastErrorToText().data);
+		auto fileGuard = MakeGuard([&](){ CloseFile(fileName, readHandle); });
 
-			FileAccessor destinationFile(m_logger, casFile);
-			if (!destinationFile.CreateWrite(false, DefaultAttributes(), 0, m_tempPath.data))
-				return false;
+		u64 readFileSize = 0;
+		if (!uba::GetFileSizeEx(readFileSize, readHandle))
+			return m_logger.Error(TC("[WriteCasFileNoCheck] GetFileSize failed for %s (%s)"), fileName, LastErrorToText().data);
 
-			if (fileIsCompressed)
+		bool readIsCompressed = false;
+		if (readFileSize >= sizeof(CompressedObjFileHeader) && g_globalRules.FileCanBeCompressed(ToView(fileName)))
+		{
+			CompressedObjFileHeader header(CasKeyZero);
+			if (!ReadFile(m_logger, fileName, readHandle, &header, sizeof(header)))
+				return m_logger.Error(TC("Failed to read header of compressed file %s (%s)"), fileName, LastErrorToText().data);
+			if (header.IsValid())
 			{
-				CompressedObjFileHeader header(CasKeyZero);
-				if (!ReadFile(m_logger, fileName, readHandle, &header, sizeof(header)))
+				if (AsCompressed(casKey, false) != AsCompressed(header.casKey, false))
+					return m_logger.Error(TC("Compressed file has different caskey than what was expected (%s)"), fileName);
+				readIsCompressed = true;
+				readFileSize -= sizeof(header);
+			}
+			else
+			{
+				if (!SetFilePointer(m_logger, fileName, readHandle, 0))
 					return false;
-				fileSize -= sizeof(header);
 			}
+		}
 
-			u8* slot= m_bufferSlots.Pop();
-			auto _ = MakeGuard([&](){ m_bufferSlots.Push(slot); });
-			u64 left = fileSize;
-
-			while (left)
+		if (!readIsCompressed && m_externalFileMappingsProvider)
+		{
+			ExternalFileMapping externalMapping;
+			if (m_externalFileMappingsProvider(externalMapping, fileNameKey, fileName))
 			{
-				u32 toRead = u32(Min(left, BufferSlotSize));
-				if (!ReadFile(m_logger, fileName, readHandle, slot, toRead))
-					return false;
-				if (!destinationFile.Write(slot, toRead))
+				FileMappingScope scope(m_logger);
+				u8* fileMem = scope.MapView(TC("WriteCompressed"), externalMapping.handle, externalMapping.offset, externalMapping.size, fileName);
+				if (!fileMem)
 					return false;
-
-				left -= toRead;
+				return StorageImpl::WriteCompressed(out, fileName, InvalidFileHandle, fileMem, externalMapping.size, casFile, nullptr, 0);
 			}
+		}
+
+		if (!readIsCompressed && storeCompressed)
+			return WriteCompressed(out, fileName, readHandle, 0, readFileSize, casFile, nullptr, 0);
 
-			if (!destinationFile.Close())
+		if (readIsCompressed && !storeCompressed)
+			return m_logger.Error(TC("[WriteCasFileNoCheck] Writing compressed file to non-compressed storage not supported (%s)"), fileName);
+
+		FileAccessor destinationFile(m_logger, casFile);
+		if (!destinationFile.CreateWrite(false, DefaultAttributes(), readFileSize, m_tempPath.data))
+			return false;
+
+		u8* slot= m_bufferSlots.Pop();
+		auto _ = MakeGuard([&](){ m_bufferSlots.Push(slot); });
+		u64 left = readFileSize;
+
+		while (left)
+		{
+			u32 toRead = u32(Min(left, BufferSlotSize));
+			if (!ReadFile(m_logger, fileName, readHandle, slot, toRead))
+				return false;
+			if (!destinationFile.Write(slot, toRead))
 				return false;
 
-			out.size = fileSize;
+			left -= toRead;
 		}
 
+		if (!destinationFile.Close())
+			return false;
+
+		out.size = readFileSize;
 		return true;
 	}
 
-	bool StorageImpl::WriteCasFile(WriteResult& out, const tchar* fileName, bool fileIsCompressed, const CasKey& casKey)
+	bool StorageImpl::WriteCasFile(WriteResult& out, const tchar* fileName, const CasKey& casKey)
 	{
 		UBA_ASSERT(IsCompressed(casKey) == m_storeCompressed);
 		StringBuffer<> casFile;
@@ -633,7 +625,7 @@ namespace uba
 		if (CaseInsensitiveFs)
 			temp.MakeLower();
 		StringKey fileNameKey = ToStringKey(temp);
-		return WriteCasFileNoCheck(out, fileNameKey, fileName, fileIsCompressed, casFile.data, IsCompressed(casKey));
+		return WriteCasFileNoCheck(out, fileNameKey, fileName, casKey, casFile.data, IsCompressed(casKey));
 	}
 
 	void StorageImpl::CasEntryAccessed(const CasKey& casKey)
@@ -908,7 +900,7 @@ namespace uba
 	}
 
 
-	bool StorageImpl::AddCasFile(StringKey fileNameKey, const tchar* fileName, const CasKey& casKey, bool deferCreation, bool fileIsCompressed)
+	bool StorageImpl::AddCasFile(StringKey fileNameKey, const tchar* fileName, const CasKey& casKey, bool deferCreation)
 	{
 		UBA_ASSERTF(IsCompressed(casKey) == m_storeCompressed, TC("CasKey compress mode must match storage compress mode. Compress mode Key: %i, Store: %i Cas: %s (%s)"), (IsCompressed(casKey) ? 1 : 0), (m_storeCompressed ? 1 : 0), CasKeyString(casKey).str, fileName);
 		CasEntry& casEntry = GetOrCreateCasEntry(casKey);
@@ -955,7 +947,6 @@ namespace uba
 			{
 				DeferedCasCreation& dcc = res.first->second;
 				dcc.fileName = fileName;
-				dcc.fileIsCompressed = fileIsCompressed;
 			}
 			auto res2 = m_deferredCasCreationLookupByName.try_emplace(fileNameKey, casKey);
 			if (res2.second)
@@ -969,7 +960,7 @@ namespace uba
 		casEntry.exists = false;
 
 		WriteResult res;
-		if (!WriteCasFileNoCheck(res, fileNameKey, fileName, fileIsCompressed, casFile.data, IsCompressed(casKey)))
+		if (!WriteCasFileNoCheck(res, fileNameKey, fileName, casKey, casFile.data, IsCompressed(casKey)))
 			return false;
 
 		casEntry.exists = true;
@@ -1847,7 +1838,7 @@ namespace uba
 		return true;
 	}
 
-	bool StorageImpl::StoreCasFile(CasKey& out, const tchar* fileName, const CasKey& casKeyOverride, bool deferCreation, bool fileIsCompressed)
+	bool StorageImpl::StoreCasFile(CasKey& out, const tchar* fileName, const CasKey& casKeyOverride, bool deferCreation)
 	{
 		StringBuffer<> forKey;
 		forKey.Append(fileName);
@@ -1857,7 +1848,7 @@ namespace uba
 
 		FileEntry& fileEntry = GetOrCreateFileEntry(fileNameKey);
 		SCOPED_FUTEX(fileEntry.lock, entryLock);
-	
+
 		if (fileEntry.verified)
 		{
 			UBA_ASSERT(fileEntry.casKey != CasKeyInvalid);
@@ -1865,7 +1856,7 @@ namespace uba
 			if (fileEntry.casKey != CasKeyZero)
 			{
 				UBA_ASSERT(casKeyOverride == CasKeyZero || casKeyOverride == fileEntry.casKey);
- 				if (!AddCasFile(fileNameKey, fileName, fileEntry.casKey, deferCreation, fileIsCompressed))
+ 				if (!AddCasFile(fileNameKey, fileName, fileEntry.casKey, deferCreation))
 					return false;
 			}
 			out = fileEntry.casKey;
@@ -1876,13 +1867,13 @@ namespace uba
 
 		ExternalFileMapping externalMapping;
 
-		FileHandle fileHandle = InvalidFileHandle;
-		auto fileGuard = MakeGuard([&](){ CloseFile(fileName, fileHandle); });
-
 		u64 fileSize = 0;
 		u64 lastWriteTime = 0;
 
-		if (!fileIsCompressed && m_externalFileMappingsProvider && m_externalFileMappingsProvider(externalMapping, fileNameKey, fileName))
+		FileHandle fileHandle = InvalidFileHandle;
+		auto fileGuard = MakeGuard([&](){ CloseFile(fileName, fileHandle); });
+
+		if (m_externalFileMappingsProvider && m_externalFileMappingsProvider(externalMapping, fileNameKey, fileName))
 		{
 			fileSize = externalMapping.size;
 			lastWriteTime = externalMapping.lastWriteTime;
@@ -1912,14 +1903,14 @@ namespace uba
 			if (casKeyOverride != CasKeyZero && casKeyOverride != fileEntry.casKey)
 			{
 				fileEntry.casKey = casKeyOverride;
-				if (!AddCasFile(fileNameKey, fileName, fileEntry.casKey, deferCreation, fileIsCompressed))
+				if (!AddCasFile(fileNameKey, fileName, fileEntry.casKey, deferCreation))
 					return false;
 				out = fileEntry.casKey;
 				return true;
 			}
 			if (fileSize == fileEntry.size && (lastWriteTime && lastWriteTime == fileEntry.lastWritten) && fileEntry.casKey != CasKeyInvalid)
 			{
-				if (!AddCasFile(fileNameKey, fileName, fileEntry.casKey, deferCreation, fileIsCompressed))
+				if (!AddCasFile(fileNameKey, fileName, fileEntry.casKey, deferCreation))
 					return false;
 				out = fileEntry.casKey;
 				return true;
@@ -1941,17 +1932,29 @@ namespace uba
 				else
 					fileEntry.casKey = CalculateCasKey(fileMem, externalMapping.size, m_storeCompressed);
 			}
-			else if (fileIsCompressed && fileSize > sizeof(CompressedObjFileHeader))
+			else
 			{
-				CompressedObjFileHeader header(CasKeyZero);
-				if (!ReadFile(m_logger, fileName, fileHandle, &header, sizeof(header)))
-					return m_logger.Error(TC("Failed to read header of compressed file %s (%s)"), fileName, LastErrorToText().data);
-				if (!header.IsValid())
-					return m_logger.Error(TC("File was expected to be compressed but was not %s (%s)"), fileName, LastErrorToText().data);
-				fileEntry.casKey = AsCompressed(header.casKey, m_storeCompressed);
+				bool handled = false;
+				if (fileSize >= sizeof(CompressedObjFileHeader) && g_globalRules.FileCanBeCompressed(ToView(fileName)))
+				{
+					CompressedObjFileHeader header(CasKeyZero);
+					if (!ReadFile(m_logger, fileName, fileHandle, &header, sizeof(header)))
+						return m_logger.Error(TC("Failed to read header of compressed file %s (%s)"), fileName, LastErrorToText().data);
+					if (header.IsValid())
+					{
+						fileEntry.casKey = AsCompressed(header.casKey, m_storeCompressed);
+						handled = true;
+					}
+					else
+					{
+						if (!SetFilePointer(m_logger, fileName, fileHandle, 0))
+							return false;
+					}
+				}
+
+				if (!handled)
+					fileEntry.casKey = CalculateCasKey(fileName, fileHandle, fileSize, m_storeCompressed);
 			}
-			else
-				fileEntry.casKey = CalculateCasKey(fileName, fileHandle, fileSize, m_storeCompressed);
 		}
 		else
 			fileEntry.casKey = casKeyOverride;
@@ -1959,7 +1962,7 @@ namespace uba
 		if (fileEntry.casKey == CasKeyZero)
 			return false;
 
-		if (!AddCasFile(fileNameKey, fileName, fileEntry.casKey, deferCreation, fileIsCompressed))
+		if (!AddCasFile(fileNameKey, fileName, fileEntry.casKey, deferCreation))
 			return false;
 
 		out = fileEntry.casKey;
@@ -1999,17 +2002,17 @@ namespace uba
 		return m_casLookup.find(casKey) != m_casLookup.end();
 	}
 
-	bool StorageImpl::StoreCasKey(CasKey& out, const tchar* fileName, const CasKey& casKeyOverride, bool fileIsCompressed)
+	bool StorageImpl::StoreCasKey(CasKey& out, const tchar* fileName, const CasKey& casKeyOverride)
 	{
 		StringBuffer<> forKey;
 		forKey.Append(fileName);
 		if (CaseInsensitiveFs)
 			forKey.MakeLower();
 		StringKey fileNameKey = ToStringKey(forKey);
-		return StoreCasKey(out, fileNameKey, fileName, casKeyOverride, fileIsCompressed);
+		return StoreCasKey(out, fileNameKey, fileName, casKeyOverride);
 	}
 
-	bool StorageImpl::StoreCasKey(CasKey& out, const StringKey& fileNameKey, const tchar* fileName, const CasKey& casKeyOverride, bool fileIsCompressed)
+	bool StorageImpl::StoreCasKey(CasKey& out, const StringKey& fileNameKey, const tchar* fileName, const CasKey& casKeyOverride)
 	{
 		FileEntry& fileEntry = GetOrCreateFileEntry(fileNameKey);
 
@@ -2061,17 +2064,26 @@ namespace uba
 		fileEntry.lastWritten = lastWritten;
 		if (casKeyOverride == CasKeyZero)
 		{
-			if (fileIsCompressed)
+			bool handled = false;
+			if (fileSize >= sizeof(CompressedObjFileHeader) && g_globalRules.FileCanBeCompressed(ToView(fileName)))
 			{
 				CompressedObjFileHeader header(CasKeyZero);
 				if (!ReadFile(m_logger, fileName, fileHandle, &header, sizeof(header)))
 					return m_logger.Error(TC("Failed to read header of compressed file %s (%s)"), fileName, LastErrorToText().data);
-				fileEntry.casKey = AsCompressed(header.casKey, m_storeCompressed);
+				if (header.IsValid())
+				{
+					fileEntry.casKey = AsCompressed(header.casKey, m_storeCompressed);
+					handled = true;
+				}
+				else
+				{
+					if (!SetFilePointer(m_logger, fileName, fileHandle, 0))
+						return false;
+				}
 			}
-			else
-			{
+
+			if (!handled)
 				fileEntry.casKey = CalculateCasKey(fileName, fileHandle, fileSize, m_storeCompressed);
-			}
 		}
 		else
 			fileEntry.casKey = casKeyOverride;
@@ -2113,7 +2125,6 @@ namespace uba
 			return false;
 		StringKey fileNameKey = findIt->second.fileNameKey;
 		StringBuffer<> deferredCreation(findIt->second.fileName);
-		bool fileIsCompressed = findIt->second.fileIsCompressed;
 		for (auto& it2 : findIt->second.names)
 			m_deferredCasCreationLookupByName.erase(it2);
 		m_deferredCasCreationLookup.erase(findIt);
@@ -2124,7 +2135,7 @@ namespace uba
 		casEntry.verified = true;
 		deferredLock.Leave();
 		WriteResult res;
-		if (!WriteCasFile(res, deferredCreation.data, fileIsCompressed, casKey))
+		if (!WriteCasFile(res, deferredCreation.data, casKey))
 			return m_logger.Error(TC("Failed to write deferred cas %s (%s)"), CasKeyString(casKey), deferredCreation.data);
 
 		casEntry.exists = true;
@@ -2173,8 +2184,7 @@ namespace uba
 		if (!fileName)
 			return false;
 		WriteResult res;
-		bool fileIsCompressed = false; // TODO: This needs to be supported
-		if (!WriteCasFile(res, fileName, fileIsCompressed, casKey))
+		if (!WriteCasFile(res, fileName, casKey))
 			return false;
 		casEntry.exists = true;
 		entryLock.Leave();
diff --git a/Engine/Source/Programs/UnrealBuildAccelerator/Common/Private/UbaStorageClient.cpp b/Engine/Source/Programs/UnrealBuildAccelerator/Common/Private/UbaStorageClient.cpp
--- a/Engine/Source/Programs/UnrealBuildAccelerator/Common/Private/UbaStorageClient.cpp
+++ b/Engine/Source/Programs/UnrealBuildAccelerator/Common/Private/UbaStorageClient.cpp
@@ -800,7 +800,7 @@ namespace uba
 		return true;
 	}
 
-	bool StorageClient::StoreCasFile(CasKey& out, const tchar* fileName, const CasKey& casKeyOverride, bool deferCreation, bool fileIsCompressed)
+	bool StorageClient::StoreCasFile(CasKey& out, const tchar* fileName, const CasKey& casKeyOverride, bool deferCreation)
 	{
 		UBA_ASSERTF(false, TC("This StoreCasFile function should not be used on the client side"));
 		return true;
diff --git a/Engine/Source/Programs/UnrealBuildAccelerator/Common/Private/UbaStorageServer.cpp b/Engine/Source/Programs/UnrealBuildAccelerator/Common/Private/UbaStorageServer.cpp
--- a/Engine/Source/Programs/UnrealBuildAccelerator/Common/Private/UbaStorageServer.cpp
+++ b/Engine/Source/Programs/UnrealBuildAccelerator/Common/Private/UbaStorageServer.cpp
@@ -485,8 +485,7 @@ namespace uba
 						}
 
 						bool deferCreation = false;
-						bool fileIsCompressed = false;
-						if (!AddCasFile(fileNameKey, hint.data, casKey, deferCreation, fileIsCompressed))
+						if (!AddCasFile(fileNameKey, hint.data, casKey, deferCreation))
 						{
 							m_logger.Error(TC("FetchBegin failed for cas file %s (%s). Can't add cas file to database"), CasKeyString(casKey).str, hint.data);
 							writer.WriteU16(0);
diff --git a/Engine/Source/Programs/UnrealBuildAccelerator/Common/Private/UbaStorageUtils.cpp b/Engine/Source/Programs/UnrealBuildAccelerator/Common/Private/UbaStorageUtils.cpp
--- a/Engine/Source/Programs/UnrealBuildAccelerator/Common/Private/UbaStorageUtils.cpp
+++ b/Engine/Source/Programs/UnrealBuildAccelerator/Common/Private/UbaStorageUtils.cpp
@@ -437,7 +437,7 @@ namespace uba
 		return sendEndMessage();
 	}
 
-	bool FileFetcher::RetrieveFile(Logger& logger, NetworkClient& client, const CasKey& casKey, const tchar* destination, bool destinationIsCompressed, MemoryBlock* destinationMem)
+	bool FileFetcher::RetrieveFile(Logger& logger, NetworkClient& client, const CasKey& casKey, const tchar* destination, bool writeCompressed, MemoryBlock* destinationMem)
 	{
 		TimerScope ts(m_stats.recvCas);
 		u8* slot = m_bufferSlots.Pop();
@@ -462,7 +462,7 @@ namespace uba
 			writer.WriteBool(false);
 			writer.WriteCasKey(casKey);
 			writer.WriteString(destination);
-			BinaryReader reader(slot + (destinationIsCompressed ? sizeof(CompressedObjFileHeader) : 0), 0, SendMaxSize); // Create some space in front of reader to write obj file header there if destination is compressed
+			BinaryReader reader(slot + (writeCompressed ? sizeof(CompressedObjFileHeader) : 0), 0, SendMaxSize); // Create some space in front of reader to write obj file header there if destination is compressed
 			if (!msg.Send(reader))
 				return logger.Error(TC("Failed to send fetch begin message for cas %s (%s). Error: %u"), CasKeyString(casKey).str, destination, msg.GetError());
 			sizeOfFirstMessage = u32(reader.GetLeft());
@@ -490,7 +490,7 @@ namespace uba
 		}
 
 		bytesReceived = fileSize;
-		sizeOnDisk = destinationIsCompressed ? (sizeof(CompressedObjFileHeader) + fileSize) : actualSize;
+		sizeOnDisk = writeCompressed ? (sizeof(CompressedObjFileHeader) + fileSize) : actualSize;
 
 		FileAccessor destinationFile(logger, destination);
 
@@ -539,7 +539,7 @@ namespace uba
 
 		u32 readIndex = 0;
 
-		if (destinationIsCompressed)
+		if (writeCompressed)
 		{
 			u8* source = slot + BufferSlotHalfSize;
 			u8* lastSource = readBuffer;
diff --git a/Engine/Source/Programs/UnrealBuildAccelerator/Common/Public/UbaCacheClient.h b/Engine/Source/Programs/UnrealBuildAccelerator/Common/Public/UbaCacheClient.h
--- a/Engine/Source/Programs/UnrealBuildAccelerator/Common/Public/UbaCacheClient.h
+++ b/Engine/Source/Programs/UnrealBuildAccelerator/Common/Public/UbaCacheClient.h
@@ -92,7 +92,6 @@ namespace uba
 		bool ShouldNormalize(const StringBufferBase& path);
 
 		bool GetLocalPathAndCasKey(Bucket& bucket, const RootPaths& rootPaths, StringBufferBase& outPath, CasKey& outKey, CompactCasKeyTable& casKeyTable, CompactPathTable& pathTable, u32 offset);
-		bool IsFileCompressed(const ProcessStartInfo& info, const StringView& filename, bool isOutput);
 		UBA_NOINLINE void PreparseDirectory(const StringKey& fileNameKey, const StringBufferBase& filePath);
 
 		MutableLogger m_logger;
diff --git a/Engine/Source/Programs/UnrealBuildAccelerator/Common/Public/UbaFile.h b/Engine/Source/Programs/UnrealBuildAccelerator/Common/Public/UbaFile.h
--- a/Engine/Source/Programs/UnrealBuildAccelerator/Common/Public/UbaFile.h
+++ b/Engine/Source/Programs/UnrealBuildAccelerator/Common/Public/UbaFile.h
@@ -57,6 +57,7 @@ namespace uba
 	bool ReadFile(Logger& logger, const tchar* fileName, FileHandle fileHandle, void* b, u64 bufferLen);
 	bool OpenFileSequentialRead(Logger& logger, const tchar* fileName, FileHandle& outHandle, bool fileNotFoundIsError = true, bool overlapped = false);
 	bool FileExists(Logger& logger, const tchar* fileName, u64* outSize = nullptr, u32* outAttributes = nullptr, u64* lastWriteTime = nullptr);
+	bool SetFilePointer(Logger& logger, const tchar* fileName, FileHandle handle, u64 position);
 	bool SetEndOfFile(Logger& logger, const tchar* fileName, FileHandle handle, u64 size);
 	bool GetDirectoryOfCurrentModule(Logger& logger, StringBufferBase& out);
 	bool DeleteAllFiles(Logger& logger, const tchar* dir, bool deleteDir = true, u32* count = nullptr);
diff --git a/Engine/Source/Programs/UnrealBuildAccelerator/Common/Public/UbaStorage.h b/Engine/Source/Programs/UnrealBuildAccelerator/Common/Public/UbaStorage.h
--- a/Engine/Source/Programs/UnrealBuildAccelerator/Common/Public/UbaStorage.h
+++ b/Engine/Source/Programs/UnrealBuildAccelerator/Common/Public/UbaStorage.h
@@ -43,7 +43,7 @@ namespace uba
 		virtual bool VerifyAndGetCachedFileInfo(CachedFileInfo& out, StringKey fileNameKey, u64 verifiedLastWriteTime, u64 verifiedSize) = 0;
 		virtual bool InvalidateCachedFileInfo(StringKey fileNameKey) = 0;
 
-		virtual bool StoreCasFile(CasKey& out, const tchar* fileName, const CasKey& casKeyOverride, bool deferCreation, bool fileIsCompressed) = 0;
+		virtual bool StoreCasFile(CasKey& out, const tchar* fileName, const CasKey& casKeyOverride, bool deferCreation) = 0;
 		virtual bool StoreCasFileClient(CasKey& out, StringKey fileNameKey, const tchar* fileName, FileMappingHandle mappingHandle, u64 mappingOffset, u64 fileSize, const tchar* hint, bool keepMappingInMemory = false, bool storeCompressed = true) = 0;
 		virtual bool DropCasFile(const CasKey& casKey, bool forceDelete, const tchar* hint) = 0;
 		virtual bool ReportBadCasFile(const CasKey& casKey) = 0;
@@ -134,7 +134,7 @@ namespace uba
 		virtual bool RetrieveCasFile(RetrieveResult& out, const CasKey& casKey, const tchar* hint, FileMappingBuffer* mappingBuffer = nullptr, u64 memoryMapAlignment = 1, bool allowProxy = true) override;
 		virtual bool VerifyAndGetCachedFileInfo(CachedFileInfo& out, StringKey fileNameKey, u64 verifiedLastWriteTime, u64 verifiedSize) override;
 		virtual bool InvalidateCachedFileInfo(StringKey fileNameKey) override;
-		virtual bool StoreCasFile(CasKey& out, const tchar* fileName, const CasKey& casKeyOverride, bool deferCreation, bool fileIsCompressed) override;
+		virtual bool StoreCasFile(CasKey& out, const tchar* fileName, const CasKey& casKeyOverride, bool deferCreation) override;
 		virtual bool StoreCasFileClient(CasKey& out, StringKey fileNameKey, const tchar* fileName, FileMappingHandle mappingHandle, u64 mappingOffset, u64 fileSize, const tchar* hint, bool keepMappingInMemory = false, bool storeCompressed = true) override;
 		virtual bool DropCasFile(const CasKey& casKey, bool forceDelete, const tchar* hint) override;
 		virtual bool ReportBadCasFile(const CasKey& casKey) override;
@@ -156,20 +156,19 @@ namespace uba
 		bool EnsureCasFile(const CasKey& casKey, const tchar* fileName);
 		CasKey CalculateCasKey(const tchar* fileName, FileHandle fileHandle, u64 fileSize, bool storeCompressed);
 		CasKey CalculateCasKey(u8* fileMem, u64 fileSize, bool storeCompressed);
-		bool StoreCasKey(CasKey& out, const tchar* fileName, const CasKey& casKeyOverride, bool fileIsCompressed);
-		bool StoreCasKey(CasKey& out, const StringKey& fileNameKey, const tchar* fileName, const CasKey& casKeyOverride, bool fileIsCompressed);
+		bool StoreCasKey(CasKey& out, const tchar* fileName, const CasKey& casKeyOverride);
+		bool StoreCasKey(CasKey& out, const StringKey& fileNameKey, const tchar* fileName, const CasKey& casKeyOverride);
 		bool IsFileVerified(const StringKey& fileNameKey);
 		void ReportFileInfoWeak(const StringKey& fileNameKey, u64 verifiedLastWriteTime, u64 verifiedSize);
 		bool HasBeenSeen(const CasKey& casKey);
 
 		virtual void RegisterExternalFileMappingsProvider(ExternalFileMappingsProvider&& provider) override;
-		bool WriteCompressed(WriteResult& out, const StringKey& fromKey, const tchar* from, const tchar* toFile);
 		virtual bool WriteCompressed(WriteResult& out, const tchar* from, FileHandle readHandle, u8* readMem, u64 fileSize, const tchar* toFile, const void* header, u64 headerSize, u64 lastWriteTime = 0) final;
 		bool WriteMemToCompressedFile(FileAccessor& destination, u32 workCount, const u8* uncompressedData, u64 fileSize, u64 maxUncompressedBlock, u64& totalWritten);
-		bool WriteCasFileNoCheck(WriteResult& out, const StringKey& fileNameKey, const tchar* fileName, bool fileIsCompressed, const tchar* casFile, bool storeCompressed);
-		bool WriteCasFile(WriteResult& out, const tchar* fileName, bool fileIsCompressed, const CasKey& casKey);
+		bool WriteCasFileNoCheck(WriteResult& out, const StringKey& fileNameKey, const tchar* fileName, const CasKey& casKey, const tchar* casFile, bool storeCompressed);
+		bool WriteCasFile(WriteResult& out, const tchar* fileName, const CasKey& casKey);
 		bool VerifyExisting(bool& outReturnValue, ScopedWriteLock& entryLock, const CasKey& casKey, CasEntry& casEntry, StringBufferBase& casFile);
-		bool AddCasFile(StringKey fileNameKey, const tchar* fileName, const CasKey& casKey, bool deferCreation, bool fileIsCompressed);
+		bool AddCasFile(StringKey fileNameKey, const tchar* fileName, const CasKey& casKey, bool deferCreation);
 		void CasEntryAccessed(const CasKey& casKey);
 		virtual bool IsDisallowedPath(const tchar* fileName);
 		virtual bool DecompressMemoryToMemory(u8* compressedData, u64 compressedSize, u8* writeData, u64 decompressedSize, const tchar* readHint, const tchar* writeHint) override;
@@ -263,7 +262,7 @@ namespace uba
 
 		ReaderWriterLock m_deferredCasCreationLookupLock;
 		using NameToCas = UnorderedMap<StringKey, CasKey>;
-		struct DeferedCasCreation { StringKey fileNameKey; TString fileName; bool fileIsCompressed; List<NameToCas::iterator> names; };
+		struct DeferedCasCreation { StringKey fileNameKey; TString fileName; List<NameToCas::iterator> names; };
 		UnorderedMap<CasKey, DeferedCasCreation> m_deferredCasCreationLookup;
 		NameToCas m_deferredCasCreationLookupByName;
 
diff --git a/Engine/Source/Programs/UnrealBuildAccelerator/Common/Public/UbaStorageClient.h b/Engine/Source/Programs/UnrealBuildAccelerator/Common/Public/UbaStorageClient.h
--- a/Engine/Source/Programs/UnrealBuildAccelerator/Common/Public/UbaStorageClient.h
+++ b/Engine/Source/Programs/UnrealBuildAccelerator/Common/Public/UbaStorageClient.h
@@ -55,7 +55,7 @@ namespace uba
 
 		virtual bool GetZone(StringBufferBase& out) override;
 		virtual bool RetrieveCasFile(RetrieveResult& out, const CasKey& casKey, const tchar* hint, FileMappingBuffer* mappingBuffer = nullptr, u64 memoryMapAlignment = 1, bool allowProxy = true) override;
-		virtual bool StoreCasFile(CasKey& out, const tchar* fileName, const CasKey& casKeyOverride, bool deferCreation, bool fileIsCompressed) override;
+		virtual bool StoreCasFile(CasKey& out, const tchar* fileName, const CasKey& casKeyOverride, bool deferCreation) override;
 		virtual bool StoreCasFileClient(CasKey& out, StringKey fileNameKey, const tchar* fileName, FileMappingHandle mappingFile, u64 mappingOffset, u64 fileSize, const tchar* hint, bool keepMappingInMemory = false, bool storeCompressed = true) override;
 		virtual bool HasCasFile(const CasKey& casKey, CasEntry** out = nullptr) override;
 		virtual void Ping() override;
diff --git a/Engine/Source/Programs/UnrealBuildAccelerator/Common/Public/UbaStorageUtils.h b/Engine/Source/Programs/UnrealBuildAccelerator/Common/Public/UbaStorageUtils.h
--- a/Engine/Source/Programs/UnrealBuildAccelerator/Common/Public/UbaStorageUtils.h
+++ b/Engine/Source/Programs/UnrealBuildAccelerator/Common/Public/UbaStorageUtils.h
@@ -36,7 +36,7 @@ namespace uba
 
 	struct FileFetcher
 	{
-		bool RetrieveFile(Logger& logger, NetworkClient& client, const CasKey& casKey, const tchar* destination, bool destinationIsCompressed, MemoryBlock* destinationMem = nullptr);
+		bool RetrieveFile(Logger& logger, NetworkClient& client, const CasKey& casKey, const tchar* destination, bool writeCompressed, MemoryBlock* destinationMem = nullptr);
 
 		BufferSlots& m_bufferSlots;
 		StorageStats& m_stats;
diff --git a/Engine/Source/Programs/UnrealBuildAccelerator/Core/Private/UbaApplicationRules.cpp b/Engine/Source/Programs/UnrealBuildAccelerator/Core/Private/UbaApplicationRules.cpp
--- a/Engine/Source/Programs/UnrealBuildAccelerator/Core/Private/UbaApplicationRules.cpp
+++ b/Engine/Source/Programs/UnrealBuildAccelerator/Core/Private/UbaApplicationRules.cpp
@@ -5,6 +5,14 @@
 
 namespace uba
 {
+	bool GlobalRules::FileCanBeCompressed(const StringView& fileName)
+	{
+		return IsWindows && (fileName.EndsWith(TC(".obj")) || fileName.EndsWith(TC(".o")));
+	}
+
+	GlobalRules g_globalRules;
+
+
 	class ApplicationRulesDotnet : public ApplicationRules
 	{
 		using Super = ApplicationRules;
@@ -61,11 +69,6 @@ namespace uba
 		}
 	};
 
-	bool IsCompressedObjFile(const StringView& fileName)
-	{
-		return fileName.EndsWith(TC(".obj"));
-	}
-
 	class ApplicationRulesClExe : public ApplicationRulesVC
 	{
 	public:
@@ -104,11 +107,6 @@ namespace uba
 			return true;
 		}
 
-		virtual bool StoreFileCompressed(const StringView& fileName) const override
-		{
-			return IsCompressedObjFile(fileName);
-		}
-
 		virtual bool ShouldExtractSymbols(const StringView& fileName) const override
 		{
 			return fileName.EndsWith(TC(".obj"));
@@ -184,9 +182,9 @@ namespace uba
 			return true;
 		}
 
-		virtual bool ShouldDecompressFiles(const StringView& fileName) const override
+		virtual bool CanDependOnCompressedFiles() override
 		{
-			return fileName.IsEmpty() || IsCompressedObjFile(fileName);
+			return true;
 		}
 	};
 
@@ -308,11 +306,6 @@ namespace uba
 			return true;
 		}
 
-		virtual bool StoreFileCompressed(const StringView& fileName) const override
-		{
-			return IsCompressedObjFile(fileName) || fileName.EndsWith(TC(".o"));
-		}
-
 		virtual bool ShouldExtractSymbols(const StringView& fileName) const override
 		{
 			return fileName.EndsWith(TC(".obj")) || fileName.EndsWith(TC(".o"));
@@ -344,9 +337,9 @@ namespace uba
 			return 14ull * 1024 * 1024 * 1024; // This is ridiculous (needed for asan targets)
 		}
 
-		virtual bool ShouldDecompressFiles(const StringView& fileName) const override
+		virtual bool CanDependOnCompressedFiles() override
 		{
-			return fileName.IsEmpty() || fileName.EndsWith(TC(".o"));
+			return true;
 		}
 	};
 
@@ -421,9 +414,9 @@ namespace uba
 			return Contains(file, TC("thinlto-")); // Used by a clang based platform's link time optimization pass. Shared from lto process back to linker process
 		}
 
-		virtual bool ShouldDecompressFiles(const StringView& fileName) const override
+		virtual bool CanDependOnCompressedFiles() override
 		{
-			return fileName.IsEmpty() || fileName.EndsWith(TC(".o"));
+			return true;
 		}
 	};
 
@@ -468,9 +461,9 @@ namespace uba
 		//	return fileName.Contains(TC("lto-llvm")); // Used by a clang based platform's link time optimization pass. Shared from lto process back to linker process
 		//}
 
-		virtual bool ShouldDecompressFiles(const StringView& fileName) const override
+		virtual bool CanDependOnCompressedFiles() override
 		{
-			return fileName.IsEmpty() || fileName.EndsWith(TC(".o"));
+			return true;
 		}
 	};
 
@@ -527,11 +520,6 @@ namespace uba
 			return true;
 		}
 
-		virtual bool StoreFileCompressed(const StringView& fileName) const override
-		{
-			return fileName.EndsWith(TC(".obj"));
-		}
-
 		virtual bool ShouldExtractSymbols(const StringView& fileName) const override
 		{
 			return fileName.EndsWith(TC(".obj")) || fileName.EndsWith(TC(".o"));
@@ -607,14 +595,9 @@ namespace uba
 				|| fileName.EndsWith(TC(".exp"));
 		}
 
-		virtual bool StoreFileCompressed(const StringView& fileName) const override
-		{
-			return fileName.EndsWith(TC(".obj")); // Extra obj file for msvc
-		}
-
-		virtual bool ShouldDecompressFiles(const StringView& fileName) const override
+		virtual bool CanDependOnCompressedFiles() override
 		{
-			return fileName.IsEmpty() || fileName.EndsWith(TC(".o")) || IsCompressedObjFile(fileName);
+			return true;
 		}
 	};
 
diff --git a/Engine/Source/Programs/UnrealBuildAccelerator/Core/Public/UbaApplicationRules.h b/Engine/Source/Programs/UnrealBuildAccelerator/Core/Public/UbaApplicationRules.h
--- a/Engine/Source/Programs/UnrealBuildAccelerator/Core/Public/UbaApplicationRules.h
+++ b/Engine/Source/Programs/UnrealBuildAccelerator/Core/Public/UbaApplicationRules.h
@@ -6,6 +6,14 @@
 
 namespace uba
 {
+	class GlobalRules
+	{
+	public:
+		bool FileCanBeCompressed(const StringView& fileName);
+	};
+
+	extern GlobalRules g_globalRules;
+
 
 	class ApplicationRules
 	{
@@ -129,15 +137,7 @@ namespace uba
 			return false;
 		}
 
-		// Return true if we want an output file to be stored compressed in destination.
-		// (Used to store .obj files compressed on disk)
-		virtual bool StoreFileCompressed(const StringView& fileName) const
-		{
-			return false;
-		}
-
-		// Return true to check if we need to decompress files when loading them (used in conjunction with StoreFileCompressed)
-		virtual bool ShouldDecompressFiles(const StringView& fileName) const
+		virtual bool CanDependOnCompressedFiles()
 		{
 			return false;
 		}
diff --git a/Engine/Source/Programs/UnrealBuildAccelerator/Detours/Private/Windows/UbaDetoursFileDecompressor.h b/Engine/Source/Programs/UnrealBuildAccelerator/Detours/Private/Windows/UbaDetoursFileDecompressor.h
--- a/Engine/Source/Programs/UnrealBuildAccelerator/Detours/Private/Windows/UbaDetoursFileDecompressor.h
+++ b/Engine/Source/Programs/UnrealBuildAccelerator/Detours/Private/Windows/UbaDetoursFileDecompressor.h
@@ -111,7 +111,7 @@ namespace uba
 
 			StringBuffer<> file;
 			file.Append(line);
-			if (!g_rules->ShouldDecompressFiles(file))
+			if (!g_globalRules.FileCanBeCompressed(file))
 				return;
 			StringBuffer<> fileFull;
 			FixPath(fileFull, file.data);
@@ -174,7 +174,7 @@ namespace uba
 		{
 			if (!m_isEnabled)
 				return false;
-			return g_rules->ShouldDecompressFiles(fileName);
+			return g_globalRules.FileCanBeCompressed(fileName);
 		}
 
 		bool IsDecompressed(const StringKey& fileNameKey, const StringView& fileName, const StringView& realFileName, FileInfo& info)
@@ -182,7 +182,7 @@ namespace uba
 			if (!m_isEnabled || m_isParsing)
 				return false;
 
-			if (!g_rules->ShouldDecompressFiles(fileName))
+			if (!g_globalRules.FileCanBeCompressed(fileName))
 				return false;
 
 			SCOPED_WRITE_LOCK(m_preloadedObjFilesLock, lock);
diff --git a/Engine/Source/Programs/UnrealBuildAccelerator/Detours/Private/Windows/UbaDetoursFunctionsWin.cpp b/Engine/Source/Programs/UnrealBuildAccelerator/Detours/Private/Windows/UbaDetoursFunctionsWin.cpp
--- a/Engine/Source/Programs/UnrealBuildAccelerator/Detours/Private/Windows/UbaDetoursFunctionsWin.cpp
+++ b/Engine/Source/Programs/UnrealBuildAccelerator/Detours/Private/Windows/UbaDetoursFunctionsWin.cpp
@@ -1108,7 +1108,7 @@ void Init(const DetoursPayload& payload, u64 startTime)
 	g_stats.attach.time += GetTime() - startTime;
 	g_stats.attach.count = 1;
 
-	if (payload.readObjFilesCompressed && g_rules->ShouldDecompressFiles(StringView()))
+	if (payload.readObjFilesCompressed && g_rules->CanDependOnCompressedFiles())
 	{
 		bool preprocess = true;//g_rules->index == SpecialRulesIndex_LdExePlatform1;
 		g_fileDecompressor.Start(cmdLine, preprocess);
diff --git a/Engine/Source/Programs/UnrealBuildAccelerator/Test/Private/UbaTestStorage.cpp b/Engine/Source/Programs/UnrealBuildAccelerator/Test/Private/UbaTestStorage.cpp
--- a/Engine/Source/Programs/UnrealBuildAccelerator/Test/Private/UbaTestStorage.cpp
+++ b/Engine/Source/Programs/UnrealBuildAccelerator/Test/Private/UbaTestStorage.cpp
@@ -28,8 +28,7 @@ namespace uba
 
 		CasKey key;
 		bool deferCreation = false;
-		bool fileIsCompressed = false;
-		if (!storage.StoreCasFile(key, detoursLib.data, CasKeyZero, deferCreation, fileIsCompressed))
+		if (!storage.StoreCasFile(key, detoursLib.data, CasKeyZero, deferCreation))
 			return logger.Error(TC("Failed to store file %s"), detoursLib.data);
 		if (key == CasKeyZero)
 			return logger.Error(TC("Failed to find file %s"), detoursLib.data);
