[UBA]
* Fixed UbaCli compile error
* Fixed linux compile error

diff --git a/Engine/Source/Programs/UnrealBuildAccelerator/Cli/Private/UbaCli.cpp b/Engine/Source/Programs/UnrealBuildAccelerator/Cli/Private/UbaCli.cpp
--- a/Engine/Source/Programs/UnrealBuildAccelerator/Cli/Private/UbaCli.cpp
+++ b/Engine/Source/Programs/UnrealBuildAccelerator/Cli/Private/UbaCli.cpp
@@ -540,7 +540,7 @@ namespace uba
 			storageInfo.storeCompressed = storeCompressed;
 			StorageImpl storage(storageInfo);
 			CasKey casKey;
-			if (!storage.StoreCasFile(casKey, addCas.c_str(), CasKeyZero, false, false))
+			if (!storage.StoreCasFile(casKey, addCas.c_str(), CasKeyZero, false))
 				return false;
 			exit = true;
 		}
diff --git a/Engine/Source/Programs/UnrealBuildAccelerator/Common/Private/UbaFile.cpp b/Engine/Source/Programs/UnrealBuildAccelerator/Common/Private/UbaFile.cpp
--- a/Engine/Source/Programs/UnrealBuildAccelerator/Common/Private/UbaFile.cpp
+++ b/Engine/Source/Programs/UnrealBuildAccelerator/Common/Private/UbaFile.cpp
@@ -331,8 +331,8 @@ namespace uba
 			return logger.Error(TC("SetFilePointerEx failed on %s (%s)"), fileName, LastErrorToText().data);
 		return true;
 #else
-		if (lseek(asFileDescriptor(hFile), position, SEEK_SET) != position)
-			return logger.Error("lseek to %llu failed for %s: %s\n", position, hint, strerror(errno));
+		if (lseek(asFileDescriptor(handle), position, SEEK_SET) != position)
+			return logger.Error("lseek to %llu failed for %s: %s", position, fileName, strerror(errno));
 		return true;
 #endif
 	}
