[UBA]
* Fixed so GetModuleFileNameW is not returning name of batch file when running a batch file but instead cmd.exe

diff --git a/Engine/Source/Programs/UnrealBuildAccelerator/Detours/Private/Windows/UbaDetoursFunctionsKernelBase.inl b/Engine/Source/Programs/UnrealBuildAccelerator/Detours/Private/Windows/UbaDetoursFunctionsKernelBase.inl
--- a/Engine/Source/Programs/UnrealBuildAccelerator/Detours/Private/Windows/UbaDetoursFunctionsKernelBase.inl
+++ b/Engine/Source/Programs/UnrealBuildAccelerator/Detours/Private/Windows/UbaDetoursFunctionsKernelBase.inl
@@ -1043,7 +1043,7 @@ BOOL Shared_GetFileAttributesExW(const tchar* caller, LPCWSTR lpFileName, GET_FI
 
 	memcpy(lpFileInformation, &attr.data, sizeof(WIN32_FILE_ATTRIBUTE_DATA));
 
-	DEBUG_LOG_DETOURED(caller, L"(%ls) -> %u", originalName, attr.exists);
+	DEBUG_LOG_DETOURED(caller, L"(%ls) -> %s", originalName, (attr.exists ? L"exists" : L"notfound"));
 	return attr.exists;
 }
 
@@ -1265,7 +1265,15 @@ DWORD Shared_GetModuleFileNameW(const tchar* func, HMODULE hModule, LPWSTR lpFil
 {
 	// If null we use the virtual application name
 	if (hModule == NULL)
+	{
+		if (g_virtualApplication.EndsWith(L".bat"))
+		{
+			DWORD res = True_GetModuleFileNameW(hModule, lpFilename, nSize);
+			DEBUG_LOG_TRUE(L"GetModuleFileNameW", L"%llu  %u (%ls) -> %u", uintptr_t(hModule), nSize, lpFilename, res);
+			return res;
+		}
 		return Shared_GetModuleFileNameInner(func, hModule, g_virtualApplication, lpFilename, nSize);
+	}
 
 	{
 		// Check if there are any stored paths from dynamically loaded dlls
