[UBA]
* Changed so when running batch files we don't add extra quotes
* Fixed logging cosmetics

diff --git a/Engine/Source/Programs/UnrealBuildAccelerator/Common/Private/UbaProcess.cpp b/Engine/Source/Programs/UnrealBuildAccelerator/Common/Private/UbaProcess.cpp
--- a/Engine/Source/Programs/UnrealBuildAccelerator/Common/Private/UbaProcess.cpp
+++ b/Engine/Source/Programs/UnrealBuildAccelerator/Common/Private/UbaProcess.cpp
@@ -1486,7 +1486,7 @@ namespace uba
 			m_session.VirtualizePath(application, m_startInfo.rootsHandle);
 
 			TString commandLine;
-			if (application.EndsWith(TC(".bat"))) // If there are quotes around arguments used by batch file we need to quote the entire thing 
+			if (application.EndsWith(TC(".bat")) && !IsRunningWine()) // If there are quotes around arguments used by batch file we need to quote the entire thing on windows
 				commandLine = TString(TC("\"\"")) + application.data + TC("\" ") + m_startInfo.arguments + TC("\"");
 			else
 				commandLine = TString(TC("\"")) + application.data + TC("\" ") + m_startInfo.arguments;
diff --git a/Engine/Source/Programs/UnrealBuildAccelerator/Detours/Private/Windows/UbaDetoursFunctionsKernelBase.inl b/Engine/Source/Programs/UnrealBuildAccelerator/Detours/Private/Windows/UbaDetoursFunctionsKernelBase.inl
--- a/Engine/Source/Programs/UnrealBuildAccelerator/Detours/Private/Windows/UbaDetoursFunctionsKernelBase.inl
+++ b/Engine/Source/Programs/UnrealBuildAccelerator/Detours/Private/Windows/UbaDetoursFunctionsKernelBase.inl
@@ -1043,7 +1043,7 @@ BOOL Shared_GetFileAttributesExW(const tchar* caller, LPCWSTR lpFileName, GET_FI
 
 	memcpy(lpFileInformation, &attr.data, sizeof(WIN32_FILE_ATTRIBUTE_DATA));
 
-	DEBUG_LOG_DETOURED(caller, L"(%ls) -> %s", originalName, (attr.exists ? L"exists" : L"notfound"));
+	DEBUG_LOG_DETOURED(caller, L"(%ls) -> %s", originalName, (attr.exists ? L"Exists" : L"NotFound"));
 	return attr.exists;
 }
 
