[UBA]
* Added support for getting exclusive access to casdb paths from the outside of uba

diff --git a/Engine/Source/Programs/UnrealBuildAccelerator/Common/Private/UbaNetworkBackendTcp.cpp b/Engine/Source/Programs/UnrealBuildAccelerator/Common/Private/UbaNetworkBackendTcp.cpp
--- a/Engine/Source/Programs/UnrealBuildAccelerator/Common/Private/UbaNetworkBackendTcp.cpp
+++ b/Engine/Source/Programs/UnrealBuildAccelerator/Common/Private/UbaNetworkBackendTcp.cpp
@@ -1188,7 +1188,7 @@ namespace uba
 			readPos += bytesRead;
 
 		if (bytesRead == SOCKET_ERROR)
-			return logger.Error(TC("HttpConnection: socket recv error %s (%s)"), LastErrorToText(WSAGetLastError()).data, m_host);
+			return logger.Error(TC("HttpConnection: socket recv error %s (%s %s)"), LastErrorToText(WSAGetLastError()).data, m_host, path);
 
 		if (readPos == sizeof(buf))
 			return logger.Error(TC("HttpConnection: buffer overflow"));
diff --git a/Engine/Source/Programs/UnrealBuildAccelerator/Common/Private/UbaStorage.cpp b/Engine/Source/Programs/UnrealBuildAccelerator/Common/Private/UbaStorage.cpp
--- a/Engine/Source/Programs/UnrealBuildAccelerator/Common/Private/UbaStorage.cpp
+++ b/Engine/Source/Programs/UnrealBuildAccelerator/Common/Private/UbaStorage.cpp
@@ -1231,6 +1231,8 @@ namespace uba
 		if (!m_writeToDisk)
 			m_casCapacityBytes = 0;
 
+		m_exclusiveMutex = info.exclusiveMutex;
+
 		m_rootDir.count = GetFullPathNameW(info.rootDir, m_rootDir.capacity, m_rootDir.data, NULL);
 		m_rootDir.Replace('/', PathSeparator).EnsureEndsWithSlash();
 
@@ -1259,15 +1261,30 @@ namespace uba
 	StorageImpl::~StorageImpl()
 	{
 		SaveCasTable(true, true);
+		CloseMutex(m_exclusiveMutex);
 	}
 
-	bool CheckExclusive(Logger& logger, const StringBufferBase& rootDir)
+	MutexHandle StorageImpl::GetExclusiveAccess(Logger& logger, const StringView& rootDir, bool reportError)
 	{
 		StringKey key = ToStringKeyNoCheck(rootDir.data, rootDir.count);
-		MutexHandle h = CreateMutexW(true, KeyToString(key).data);
-		if (h && GetLastError() != ERROR_ALREADY_EXISTS)
-			return true;
-		return logger.Error(TC("Needs exclusive access to storage %s. Another process is running"), rootDir.data);
+		KeyToString keyStr(key);
+		MutexHandle exclusiveMutex = CreateMutexW(true, keyStr.data);
+		u32 lastError = GetLastError();
+		if (exclusiveMutex == InvalidMutexHandle)
+		{
+			if (reportError)
+				logger.Error(TC("Failed to create mutex %s for path %s (%s)"), keyStr.data, rootDir.data, LastErrorToText(lastError).data);
+			return InvalidMutexHandle;
+		}
+
+		if (lastError != ERROR_ALREADY_EXISTS)
+			return exclusiveMutex;
+
+		CloseMutex(exclusiveMutex);
+		exclusiveMutex = InvalidMutexHandle;
+		if (reportError)
+			logger.Error(TC("Needs exclusive access to storage %s. Another process is running"), rootDir.data);
+		return InvalidMutexHandle;
 	}
 
 	StorageImpl::FileEntry& StorageImpl::GetOrCreateFileEntry(const StringKey& fileNameKey)
@@ -1297,9 +1314,12 @@ namespace uba
 		if (!m_writeToDisk)
 			return true;
 
-		static bool isExclusive = CheckExclusive(m_logger, m_rootDir);
-		if (!isExclusive)
-			return false;
+		if (m_exclusiveMutex == InvalidMutexHandle)
+		{
+			m_exclusiveMutex = GetExclusiveAccess(m_logger, m_rootDir, true);
+			if (m_exclusiveMutex == InvalidMutexHandle)
+				return false;
+		}
 
 		CreateDirectory(m_rootDir.data);
 
diff --git a/Engine/Source/Programs/UnrealBuildAccelerator/Common/Private/UbaTrace.cpp b/Engine/Source/Programs/UnrealBuildAccelerator/Common/Private/UbaTrace.cpp
--- a/Engine/Source/Programs/UnrealBuildAccelerator/Common/Private/UbaTrace.cpp
+++ b/Engine/Source/Programs/UnrealBuildAccelerator/Common/Private/UbaTrace.cpp
@@ -521,8 +521,7 @@ namespace uba
 			UnmapViewOfFile(m_logger, m_mem, 256, TC("TraceChannel"));
 		if (m_memHandle.IsValid())
 			CloseFileMapping(m_logger, m_memHandle, TC("TraceChannel"));
-		if (m_mutex)
-			CloseHandle((HANDLE)m_mutex);
+		CloseMutex(m_mutex);
 		#endif
 	}
 
@@ -530,7 +529,7 @@ namespace uba
 	{
 		#if PLATFORM_WINDOWS
 		WaitForSingleObject((HANDLE)m_mutex, INFINITE);
-		auto g = MakeGuard([this]() { ReleaseMutex((HANDLE)m_mutex); });
+		auto g = MakeGuard([this]() { ReleaseMutex(m_mutex); });
 		if (ifMatching)
 			if (!Equals((tchar*)m_mem, ifMatching))
 				return true;
@@ -544,7 +543,7 @@ namespace uba
 		#if PLATFORM_WINDOWS
 		WaitForSingleObject((HANDLE)m_mutex, INFINITE);
 		outTraceName.Append((tchar*)m_mem);
-		ReleaseMutex((HANDLE)m_mutex);
+		ReleaseMutex(m_mutex);
 		#endif
 		return true;
 	}
diff --git a/Engine/Source/Programs/UnrealBuildAccelerator/Common/Public/UbaStorage.h b/Engine/Source/Programs/UnrealBuildAccelerator/Common/Public/UbaStorage.h
--- a/Engine/Source/Programs/UnrealBuildAccelerator/Common/Public/UbaStorage.h
+++ b/Engine/Source/Programs/UnrealBuildAccelerator/Common/Public/UbaStorage.h
@@ -91,6 +91,7 @@ namespace uba
 		bool writeToDisk = true;
 		bool keepTransientDataMapped = true;	// Will keep the transient data mapped in instead of doing map/unmap everytime we access it
 		WorkManager* workManager = nullptr;
+		MutexHandle exclusiveMutex = InvalidMutexHandle; // Set this to hand over pre-exclusive access to storage. Storage will release mutex on shutdown
 	};
 
 	struct BufferSlots
@@ -184,6 +185,7 @@ namespace uba
 		bool CheckAllCasFiles(u64 checkContentOfFilesNewerThanTime = ~u64(0));
 		void HandleOverflow(UnorderedSet<CasKey>* outDeletedFiles);
 
+		static MutexHandle GetExclusiveAccess(Logger& logger, const StringView& rootDir, bool reportError);
 		bool DeleteIsRunningFile();
 
 		struct FileEntry;
@@ -247,6 +249,7 @@ namespace uba
 		bool m_asyncUnmapViewOfFile = true;
 		bool m_writeToDisk;
 
+		MutexHandle m_exclusiveMutex = InvalidMutexHandle;
 
 		ExternalFileMappingsProvider m_externalFileMappingsProvider;
 
diff --git a/Engine/Source/Programs/UnrealBuildAccelerator/Core/Private/UbaPlatform.cpp b/Engine/Source/Programs/UnrealBuildAccelerator/Core/Private/UbaPlatform.cpp
--- a/Engine/Source/Programs/UnrealBuildAccelerator/Core/Private/UbaPlatform.cpp
+++ b/Engine/Source/Programs/UnrealBuildAccelerator/Core/Private/UbaPlatform.cpp
@@ -441,6 +441,26 @@ namespace uba
 		#endif
 	}
 
+	void ReleaseMutex(MutexHandle mutex)
+	{
+		if (mutex == InvalidMutexHandle)
+			return;
+		#if PLATFORM_WINDOWS
+		::ReleaseMutex((HANDLE)mutex);
+		#else
+		#endif
+	}
+
+	void CloseMutex(MutexHandle mutex)
+	{
+		if (mutex == InvalidMutexHandle)
+			return;
+		#if PLATFORM_WINDOWS
+		::CloseHandle((HANDLE)mutex);
+		#else
+		#endif
+	}
+
 	u32 GetEnvironmentVariableW(const tchar* name, tchar* buffer, u32 nSize)
 	{
 		#if PLATFORM_WINDOWS
diff --git a/Engine/Source/Programs/UnrealBuildAccelerator/Core/Public/UbaPlatform.h b/Engine/Source/Programs/UnrealBuildAccelerator/Core/Public/UbaPlatform.h
--- a/Engine/Source/Programs/UnrealBuildAccelerator/Core/Public/UbaPlatform.h
+++ b/Engine/Source/Programs/UnrealBuildAccelerator/Core/Public/UbaPlatform.h
@@ -104,6 +104,8 @@ namespace uba
 	enum MutexHandle : u64 {};
 	inline constexpr MutexHandle InvalidMutexHandle = MutexHandle(0);
 	MutexHandle CreateMutexW(bool bInitialOwner, const tchar* name);
+	void ReleaseMutex(MutexHandle mutex);
+	void CloseMutex(MutexHandle mutex);
 
 	enum ProcHandle : u64 {};
 	inline constexpr ProcHandle InvalidProcHandle = ((ProcHandle)(u64)-1);
diff --git a/Engine/Source/Programs/UnrealBuildAccelerator/Host/Private/UbaExports.cpp b/Engine/Source/Programs/UnrealBuildAccelerator/Host/Private/UbaExports.cpp
--- a/Engine/Source/Programs/UnrealBuildAccelerator/Host/Private/UbaExports.cpp
+++ b/Engine/Source/Programs/UnrealBuildAccelerator/Host/Private/UbaExports.cpp
@@ -107,6 +107,16 @@ namespace uba
 		config.LoadFromFile(logger, fileName);
 		return config;
 	}
+
+	struct DowngradedLogger : public LoggerWithWriter
+	{
+		DowngradedLogger(LogWriter& writer, const tchar* prefix) : LoggerWithWriter(writer, prefix) {}
+		virtual void Log(LogEntryType type, const tchar* str, u32 strLen) override { LoggerWithWriter::Log(Max(LogEntryType_Info, type), str, strLen); }
+	};
+
+	// Not great with global but 
+	MutexHandle g_exclusiveMutex = InvalidMutexHandle;
+
 }
 
 extern "C"
@@ -283,8 +293,8 @@ extern "C"
 		AWS aws;
 		if (!info.zone || !*info.zone)
 		{
-			LoggerWithWriter logger(info.writer, TC(""));
-			if (aws.QueryAvailabilityZone(logger, info.rootDir))
+			DowngradedLogger downgradedLogger(info.writer, TC("AWS"));
+			if (aws.QueryAvailabilityZone(downgradedLogger, info.rootDir))
 				info.zone = aws.GetAvailabilityZone();
 		}
 		#endif
@@ -293,6 +303,10 @@ extern "C"
 		if (!info.zone || !*info.zone)
 			if (GetZone(zoneTemp))
 				info.zone = zoneTemp.data;
+
+		info.exclusiveMutex = g_exclusiveMutex;
+		g_exclusiveMutex = InvalidMutexHandle;
+
 		#endif
 
 		return new StorageServer(info);
@@ -866,6 +880,14 @@ extern "C"
 #endif
 	}
 
+	bool Uba_GetExclusiveAccess(const uba::tchar* path)
+	{
+		using namespace uba;
+		DowngradedLogger downgradedLogger(g_consoleLogWriter, TC("UbaGetExclusiveAccess"));
+		g_exclusiveMutex = StorageImpl::GetExclusiveAccess(downgradedLogger, ToView(path), false);
+		return g_exclusiveMutex != InvalidMutexHandle;
+	}
+
 	struct UbaInstance
 	{
 		uba::Scheduler* scheduler;
diff --git a/Engine/Source/Programs/UnrealBuildAccelerator/Host/Public/UbaExports.h b/Engine/Source/Programs/UnrealBuildAccelerator/Host/Public/UbaExports.h
--- a/Engine/Source/Programs/UnrealBuildAccelerator/Host/Public/UbaExports.h
+++ b/Engine/Source/Programs/UnrealBuildAccelerator/Host/Public/UbaExports.h
@@ -161,6 +161,8 @@ extern "C"
 	using ImportFunc = void(const uba::tchar* importName, void* userData);
 	UBA_API void Uba_FindImports(const uba::tchar* binary, ImportFunc* func, void* userData);
 
+	UBA_API bool Uba_GetExclusiveAccess(const uba::tchar* path); // Will take a global mutex if returns true and will hand it over to StorageServer on creation.
+
 	// High level interface using config file instead. Uses scheduler under the hood
 	UBA_API void* Uba_Create(const uba::tchar* configFile);
 	UBA_API uba::u32 Uba_RunProcess(void* uba, const uba::tchar* app, const uba::tchar* args, const uba::tchar* workDir, const uba::tchar* desc, void* userData, ProcessHandle_ExitCallback* exit);
