[UBA]
* Fixed bug introduced by compressed obj file refactoring causing external file mappings provider to stop working

diff --git a/Engine/Source/Programs/UnrealBuildAccelerator/Common/Private/UbaStorage.cpp b/Engine/Source/Programs/UnrealBuildAccelerator/Common/Private/UbaStorage.cpp
--- a/Engine/Source/Programs/UnrealBuildAccelerator/Common/Private/UbaStorage.cpp
+++ b/Engine/Source/Programs/UnrealBuildAccelerator/Common/Private/UbaStorage.cpp
@@ -537,6 +537,19 @@ namespace uba
 		StorageStats& stats = Stats();
 		TimerScope ts(stats.createCas);
 
+		if (m_externalFileMappingsProvider)
+		{
+			ExternalFileMapping externalMapping;
+			if (m_externalFileMappingsProvider(externalMapping, fileNameKey, fileName))
+			{
+				FileMappingScope scope(m_logger);
+				u8* fileMem = scope.MapView(TC("WriteCompressed"), externalMapping.handle, externalMapping.offset, externalMapping.size, fileName);
+				if (!fileMem)
+					return false;
+				return StorageImpl::WriteCompressed(out, fileName, InvalidFileHandle, fileMem, externalMapping.size, casFile, nullptr, 0);
+			}
+		}
+
 		FileHandle readHandle;
 		if (!OpenFileSequentialRead(m_logger, fileName, readHandle))
 			return m_logger.Error(TC("[WriteCasFileNoCheck] Failed to open file %s for read (%s)"), fileName, LastErrorToText().data);
@@ -566,19 +579,6 @@ namespace uba
 			}
 		}
 
-		if (!readIsCompressed && m_externalFileMappingsProvider)
-		{
-			ExternalFileMapping externalMapping;
-			if (m_externalFileMappingsProvider(externalMapping, fileNameKey, fileName))
-			{
-				FileMappingScope scope(m_logger);
-				u8* fileMem = scope.MapView(TC("WriteCompressed"), externalMapping.handle, externalMapping.offset, externalMapping.size, fileName);
-				if (!fileMem)
-					return false;
-				return StorageImpl::WriteCompressed(out, fileName, InvalidFileHandle, fileMem, externalMapping.size, casFile, nullptr, 0);
-			}
-		}
-
 		if (!readIsCompressed && storeCompressed)
 			return WriteCompressed(out, fileName, readHandle, 0, readFileSize, casFile, nullptr, 0);
 
