[UBA]
*  Fixed so detouring code handles NtQueryInformationFile with class FileNetworkOpenInformation

diff --git a/Engine/Source/Programs/UnrealBuildAccelerator/Common/Private/UbaFile.cpp b/Engine/Source/Programs/UnrealBuildAccelerator/Common/Private/UbaFile.cpp
--- a/Engine/Source/Programs/UnrealBuildAccelerator/Common/Private/UbaFile.cpp
+++ b/Engine/Source/Programs/UnrealBuildAccelerator/Common/Private/UbaFile.cpp
@@ -191,7 +191,7 @@ namespace uba
 			InitializeObjectAttributes(&oa, &us, OBJ_CASE_INSENSITIVE, NULL, NULL);
 			NTSTATUS res = NtQueryFullAttributesFile(&oa, &info);
 			if (res != STATUS_SUCCESS)
-				return logger.Error(TC("GetFileBasicInformationByHandle (NtQueryFullAttributesFile) failed on %s (%u)"), fileName, res);
+				return logger.Error(TC("GetFileBasicInformationByHandle (NtQueryFullAttributesFile) failed on %s (0x%11x)"), fileName, res);
 		}
 		else
 		{
@@ -199,7 +199,7 @@ namespace uba
 			IO_STATUS_BLOCK b;
 			NTSTATUS res = NtQueryInformationFile(asHANDLE(hFile), &b, &info, sizeof(info), FileNetworkOpenInformation);
 			if (res != STATUS_SUCCESS)
-				return logger.Error(TC("GetFileBasicInformationByHandle (NtQueryInformationFile) failed on %s (%u)"), fileName, res);
+				return logger.Error(TC("GetFileBasicInformationByHandle (NtQueryInformationFile) failed on %s (0x%11x)"), fileName, res);
 		}
 		out.attributes = info.FileAttributes;
 		out.lastWriteTime = info.LastWriteTime.QuadPart;
diff --git a/Engine/Source/Programs/UnrealBuildAccelerator/Detours/Private/Windows/UbaDetoursFunctionsNtDll.inl b/Engine/Source/Programs/UnrealBuildAccelerator/Detours/Private/Windows/UbaDetoursFunctionsNtDll.inl
--- a/Engine/Source/Programs/UnrealBuildAccelerator/Detours/Private/Windows/UbaDetoursFunctionsNtDll.inl
+++ b/Engine/Source/Programs/UnrealBuildAccelerator/Detours/Private/Windows/UbaDetoursFunctionsNtDll.inl
@@ -211,6 +211,35 @@ NTSTATUS Detoured_NtQueryInformationFile(HANDLE FileHandle, PIO_STATUS_BLOCK IoS
 				info.InternalInformation.IndexNumber.QuadPart = entryInfo.fileIndex;
 				return STATUS_SUCCESS;
 			}
+			if (FileInformationClass == 34) // FileNetworkOpenInformation
+			{
+				struct FILE_NETWORK_OPEN_INFORMATION
+				{
+					LARGE_INTEGER CreationTime;
+					LARGE_INTEGER LastAccessTime;
+					LARGE_INTEGER LastWriteTime;
+					LARGE_INTEGER ChangeTime;
+					LARGE_INTEGER AllocationSize;
+					LARGE_INTEGER EndOfFile;
+					ULONG FileAttributes;
+				};
+
+				UBA_ASSERT(dh.dirTableOffset != ~0u);
+				DirectoryTable::EntryInformation entryInfo;
+				g_directoryTable.GetEntryInformation(entryInfo, dh.dirTableOffset);
+				UBA_ASSERT(entryInfo.attributes != 0);
+
+				// TODO This code path is here to handle nodejs queries.. Is not properly implemented and miss things
+				auto& info = *(FILE_NETWORK_OPEN_INFORMATION*)FileInformation;
+				info.CreationTime.QuadPart = entryInfo.lastWrite;
+				info.LastAccessTime.QuadPart = entryInfo.lastWrite;
+				info.LastWriteTime.QuadPart = entryInfo.lastWrite;
+				info.ChangeTime.QuadPart = entryInfo.lastWrite;
+				info.AllocationSize.QuadPart = entryInfo.size;
+				info.EndOfFile.QuadPart = entryInfo.size;
+				info.FileAttributes = entryInfo.attributes;
+				return STATUS_SUCCESS;
+			}
 
 			//if (FileInformationClass == 9) // FileNameInformation
 			//{
