[UBA]
* Added size validation for source file used for virtual files

diff --git a/Engine/Source/Programs/UnrealBuildAccelerator/Common/Private/UbaSession.cpp b/Engine/Source/Programs/UnrealBuildAccelerator/Common/Private/UbaSession.cpp
--- a/Engine/Source/Programs/UnrealBuildAccelerator/Common/Private/UbaSession.cpp
+++ b/Engine/Source/Programs/UnrealBuildAccelerator/Common/Private/UbaSession.cpp
@@ -1354,7 +1354,7 @@ namespace uba
 		FlushDeadProcesses();
 
 		for (auto& kv : m_virtualSourceFiles)
-			CloseFileMapping(m_logger, kv.second, TC("VirtualFile"));
+			CloseFileMapping(m_logger, kv.second.mappingHandle, TC("VirtualFile"));
 
 		for (auto& i : m_fileMappingTableLookup)
 			if (i.second.owned)
@@ -2962,26 +2962,31 @@ namespace uba
 	{
 		TimerScope ts(Stats().createMmapFromFile);
 
-		FileMappingHandle mappingHandle;
+		VirtualSourceFile virtualFile;
 
 		{
 			StringKey sourceFileKey = CaseInsensitiveFs ? ToStringKeyLower(ToView(sourceFile)) : ToStringKey(ToView(sourceFile));
 			SCOPED_FUTEX(m_virtualSourceFilesLock, virtualSourceFilesLock);
-			auto insres = m_virtualSourceFiles.try_emplace(sourceFileKey);
+			auto insres = m_virtualSourceFiles.try_emplace(sourceFileKey); 
 			if (insres.second)
 			{
 				FileHandle fileHandle = uba::CreateFileW(sourceFile, GENERIC_READ, FILE_SHARE_READ, OPEN_EXISTING, DefaultAttributes());
 				if (fileHandle == InvalidFileHandle)
 					return m_logger.Error(TC("[RegisterVirtualFileInternal] CreateFileW for %s failed (%s)"), sourceFile, LastErrorToText().data);
+				u64 fileSize = 0;
+				if (!GetFileSizeEx(fileSize, fileHandle))
+					return m_logger.Error(TC("[RegisterVirtualFileInternal] GetFileSizeEx for %s failed (%s)"), sourceFile, LastErrorToText().data);
 				auto fg = MakeGuard([&]() { uba::CloseFile(sourceFile, fileHandle); });
-				insres.first->second = CreateFileMappingW(m_logger, fileHandle, PAGE_READONLY, 0, sourceFile);
+				virtualFile.mappingHandle = CreateFileMappingW(m_logger, fileHandle, PAGE_READONLY, fileSize, sourceFile);
+				virtualFile.size = fileSize;
+				insres.first->second = virtualFile;
 			}
-			mappingHandle = insres.first->second;
-			if (!mappingHandle.IsValid())
+			else
+				virtualFile = insres.first->second;
+			if (!virtualFile.mappingHandle.IsValid())
 				return m_logger.Error(TC("[RegisterVirtualFileInternal] CreateFileMapping for %s failed (%s)"), sourceFile, LastErrorToText().data);
 		}
 
-
 		SCOPED_FUTEX(m_fileMappingTableLookupLock, lookupLock);
 		auto insres = m_fileMappingTableLookup.try_emplace(fileNameKey);
 		FileMappingEntry& entry = insres.first->second;
@@ -2992,14 +2997,17 @@ namespace uba
 		if (entry.handled)
 			return m_logger.Error(TC("Virtual file %s has already been registered"), filePath.data);
 
-		entry.mapping = mappingHandle;
+		if (sourceSize + sourceOffset > virtualFile.size)
+			return m_logger.Error(TC("Virtual file offset(%llu)+size(%llu) outside source file size(%llu)"), sourceOffset, sourceSize, virtualFile.size, filePath.data);
+
+		entry.mapping = virtualFile.mappingHandle;
 		entry.mappingOffset = sourceOffset;
 		entry.size = sourceSize;
 		entry.handled = true;
 		entry.lastWriteTime = 0; // TODO: Take lastwritetime of source file?
 
 		StringBuffer<> mappingName;
-		Storage::GetMappingString(mappingName, mappingHandle, sourceOffset);
+		Storage::GetMappingString(mappingName, virtualFile.mappingHandle, sourceOffset);
 		entry.success = true;
 
 		SCOPED_WRITE_LOCK(m_fileMappingTableMemLock, lock);
diff --git a/Engine/Source/Programs/UnrealBuildAccelerator/Common/Public/UbaSession.h b/Engine/Source/Programs/UnrealBuildAccelerator/Common/Public/UbaSession.h
--- a/Engine/Source/Programs/UnrealBuildAccelerator/Common/Public/UbaSession.h
+++ b/Engine/Source/Programs/UnrealBuildAccelerator/Common/Public/UbaSession.h
@@ -266,7 +266,8 @@ namespace uba
 		UnorderedMap<u32, ActiveFile> m_activeFiles;
 
 		Futex m_virtualSourceFilesLock;
-		UnorderedMap<StringKey, FileMappingHandle> m_virtualSourceFiles;
+		struct VirtualSourceFile { FileMappingHandle mappingHandle; u64 size; };
+		UnorderedMap<StringKey, VirtualSourceFile> m_virtualSourceFiles;
 
 		ReaderWriterLock m_rootsLookupLock;
 		UnorderedMap<RootsHandle, RootsEntry> m_rootsLookup;
