[UBA]
* Added description to parallel for in UbaObjTool

diff --git a/Engine/Source/Programs/UnrealBuildAccelerator/ObjTool/Private/UbaObjTool.cpp b/Engine/Source/Programs/UnrealBuildAccelerator/ObjTool/Private/UbaObjTool.cpp
--- a/Engine/Source/Programs/UnrealBuildAccelerator/ObjTool/Private/UbaObjTool.cpp
+++ b/Engine/Source/Programs/UnrealBuildAccelerator/ObjTool/Private/UbaObjTool.cpp
@@ -275,7 +275,7 @@ namespace uba
 					}
 					SCOPED_CRITICAL_SECTION(cs, _);
 					allExternalImports.insert(symbolFile.imports.begin(), symbolFile.imports.end());
-				});
+				}, TCV("ObjFilesDeps"));
 			if (!success)
 				return -1;
 
@@ -297,7 +297,7 @@ namespace uba
 					SCOPED_CRITICAL_SECTION(cs, _);
 					allInternalImports.insert(symbolFile.imports.begin(), symbolFile.imports.end());
 					allExports.insert(symbolFile.exports.begin(), symbolFile.exports.end());
-				});
+				}, TCV("ObjFilesToStrip"));
 			if (!success)
 				return -1;
 
@@ -353,7 +353,7 @@ namespace uba
 				}
 				else
 					success = false;
-			});
+			}, TCV("OpenAndParse"));
 			if (!success)
 				return -1;
 
