[UBA]
* Fixed so linux/mac path fixup code handles paths that are in quotes
* Fixed include crawler handling some special options in rsp
* Added better error information when crawler fails to devirtualize path

diff --git a/Engine/Source/Programs/UnrealBuildAccelerator/Common/Private/UbaRootPaths.cpp b/Engine/Source/Programs/UnrealBuildAccelerator/Common/Private/UbaRootPaths.cpp
--- a/Engine/Source/Programs/UnrealBuildAccelerator/Common/Private/UbaRootPaths.cpp
+++ b/Engine/Source/Programs/UnrealBuildAccelerator/Common/Private/UbaRootPaths.cpp
@@ -81,6 +81,20 @@ namespace uba
 		return roots[index].path;
 	}
 
+	TString RootPaths::GetAllRoots() const
+	{
+		TString res;
+		bool isFirst = true;
+		for (auto& root : m_roots)
+		{
+			if (!isFirst)
+				res += ' ';
+			isFirst = false;
+			res += root.path;
+		}
+		return res;
+	}
+
 	CasKey RootPaths::NormalizeAndHashFile(Logger& logger, const tchar* filename, bool warnOnFileNotFound) const
 	{
 		FileAccessor file(logger, filename);
diff --git a/Engine/Source/Programs/UnrealBuildAccelerator/Common/Private/UbaSession.cpp b/Engine/Source/Programs/UnrealBuildAccelerator/Common/Private/UbaSession.cpp
--- a/Engine/Source/Programs/UnrealBuildAccelerator/Common/Private/UbaSession.cpp
+++ b/Engine/Source/Programs/UnrealBuildAccelerator/Common/Private/UbaSession.cpp
@@ -1741,7 +1741,7 @@ namespace uba
 			return true;
 		auto root = rootsEntry->roots.FindRoot(inOut);
 		if (!root)
-			return m_logger.Error(TC("Can't find root for path %s"), inOut.data);
+			return m_logger.Error(TC("Can't find root for path %s (Available roots: %s)"), inOut.data, rootsEntry->roots.GetAllRoots().c_str());
 
 		auto& path = rootsEntry->locals[root->index/PathsPerRoot];
 		StringBuffer<> temp;
diff --git a/Engine/Source/Programs/UnrealBuildAccelerator/Common/Private/UbaStorageClient.cpp b/Engine/Source/Programs/UnrealBuildAccelerator/Common/Private/UbaStorageClient.cpp
--- a/Engine/Source/Programs/UnrealBuildAccelerator/Common/Private/UbaStorageClient.cpp
+++ b/Engine/Source/Programs/UnrealBuildAccelerator/Common/Private/UbaStorageClient.cpp
@@ -972,7 +972,7 @@ namespace uba
 					workManager.AddWork([&, filePath = TString(fullPath.data)](const WorkContext& context)
 						{
 							PopulateCasFromDirsRecursive(filePath.c_str(), workManager, seenIds, seenIdsLock, shouldExit);
-						}, 1, TC(""), ColorWork);
+						}, 1, TC("PopulateCasFromDirsRecursive"), ColorWork);
 					return;
 				}
 
@@ -1032,7 +1032,7 @@ namespace uba
 							localFile.fileName = filePath;
 						}
 
-					}, 1, TC(""), ColorWork);
+					}, 1, TC("PrepopulateCasFromFile"), ColorWork);
 			});
 		return true;
 	}
diff --git a/Engine/Source/Programs/UnrealBuildAccelerator/Common/Public/UbaIncludeCrawler.h b/Engine/Source/Programs/UnrealBuildAccelerator/Common/Public/UbaIncludeCrawler.h
--- a/Engine/Source/Programs/UnrealBuildAccelerator/Common/Public/UbaIncludeCrawler.h
+++ b/Engine/Source/Programs/UnrealBuildAccelerator/Common/Public/UbaIncludeCrawler.h
@@ -76,11 +76,11 @@ namespace uba
 		using CodeFiles = List<CodeFile>;
 
 		inline bool ParseRsp(Instance& instance, const tchar* rsp, CodeFiles& outCodeFiles);
-		inline bool ParseRsp2(Instance& instance, const void* data, u64 dataSize, CodeFiles& outCodeFiles);
+		inline bool ParseRsp2(Instance& instance, const void* data, u64 dataSize, const tchar* rsp, CodeFiles& outCodeFiles);
 		inline bool ParsePch(Pch& pch, Instance& instance, const void* data, u64 dataSize);
 		inline bool ParseCodeFile(Instance& instance, const StringView& codeFile, bool parseDefines);
 		inline bool ParseCodeFile2(Instance& instance, const void* data, u64 dataSize, const StringView& codeFile, bool parseDefines);
-		inline bool HandleInclude(Instance& instance, const StringView& rootPath, const StringView& include, bool parseDefines);
+		inline bool HandleInclude(Instance& instance, const StringView& rootPath, const StringView& include, const StringView& codeFile, bool parseDefines);
 	};
 
 
@@ -167,7 +167,8 @@ namespace uba
 					{
 						StringBuffer<> fixedPath;
 						FixPath(instance.builtinIncludesDir.c_str(), instance.workDir.data(), instance.workDir.size(), fixedPath);
-						instance.devirtualizePathFunc(fixedPath);
+						if (!instance.devirtualizePathFunc(fixedPath))
+							m_logger.Info(TC("Failed to devirtualize path %s found in builtin includes"), fixedPath.data);
 						if (CaseInsensitiveFs)
 							fixedPath.MakeLower();
 						fixedPath.Append(PathSeparator);
@@ -193,16 +194,17 @@ namespace uba
 	{
 		StringBuffer<> fixedPath;
 		FixPath(rsp, instance.workDir.data(), instance.workDir.size(), fixedPath);
-		instance.devirtualizePathFunc(fixedPath);
+		if (!instance.devirtualizePathFunc(fixedPath))
+			m_logger.Info(TC("Failed to devirtualize path %s in %s"), fixedPath.data, rsp);
 		if (CaseInsensitiveFs)
 			fixedPath.MakeLower();
 		return instance.createFileFunc(fixedPath, [&](const void* data, u64 dataSize)
 			{
-				return ParseRsp2(instance, data, dataSize, outCodeFiles);
+				return ParseRsp2(instance, data, dataSize, rsp, outCodeFiles);
 			});
 	}
 
-	bool IncludeCrawler::ParseRsp2(Instance& instance, const void* data, u64 dataSize, CodeFiles& outCodeFiles)
+	bool IncludeCrawler::ParseRsp2(Instance& instance, const void* data, u64 dataSize, const tchar* rsp, CodeFiles& outCodeFiles)
 	{
 		StringBuffer<> prevArg;
 
@@ -210,7 +212,9 @@ namespace uba
 			{
 				StringBuffer<> fixedPath2;
 				FixPath(file, instance.workDir.data(), instance.workDir.size(), fixedPath2);
-				instance.devirtualizePathFunc(fixedPath2);
+				if (!instance.devirtualizePathFunc(fixedPath2))
+					m_logger.Info(TC("Failed to devirtualize path %s in %s"), fixedPath2.data, rsp);
+
 				if (CaseInsensitiveFs)
 					fixedPath2.MakeLower();
 				if (pushFront)
@@ -223,7 +227,8 @@ namespace uba
 			{
 				StringBuffer<> fixedPath2;
 				FixPath(path, instance.workDir.data(), instance.workDir.size(), fixedPath2);
-				instance.devirtualizePathFunc(fixedPath2);
+				if (!instance.devirtualizePathFunc(fixedPath2))
+					m_logger.Info(TC("Failed to devirtualize path %s in %s"), fixedPath2.data, rsp);
 				if (CaseInsensitiveFs)
 					fixedPath2.MakeLower();
 				fixedPath2.EnsureEndsWithSlash();
@@ -265,6 +270,8 @@ namespace uba
 							TC("-include-pch"),
 							TC("-vctoolsdir"),
 							TC("-Xclang"),
+							TC("-target"),
+							TC("-arch"),
 						};
 
 						for (auto option : optionsWithArg)
@@ -283,7 +290,9 @@ namespace uba
 						// Check if there is a precompiled header deps file
 						StringBuffer<> fixedFile;
 						FixPath(sb.data + 3, instance.workDir.data(), instance.workDir.size(), fixedFile);
-						instance.devirtualizePathFunc(fixedFile);
+						if (!instance.devirtualizePathFunc(fixedFile))
+							m_logger.Info(TC("Failed to devirtualize path %s in %s"), fixedFile.data, rsp);
+
 						if (CaseInsensitiveFs)
 							fixedFile.MakeLower();
 
@@ -503,7 +512,7 @@ namespace uba
 								{
 									StringBuffer<> localDir;
 									if (const tchar* lastSeparator = TStrrchr(codeFile2.c_str(), PathSeparator))
-										if (HandleInclude(instance, localDir.Append(codeFile2.c_str(), lastSeparator - codeFile2.c_str() + 1), include2, parseDefines))
+										if (HandleInclude(instance, localDir.Append(codeFile2.c_str(), lastSeparator - codeFile2.c_str() + 1), include2, codeFile2, parseDefines))
 											return;
 								}
 
@@ -517,7 +526,7 @@ namespace uba
 								{
 									if (root.bloomFilter.IsGuaranteedMiss(fileKey))
 										continue;
-									if (HandleInclude(instance, root.path, include2, parseDefines))
+									if (HandleInclude(instance, root.path, include2, codeFile2, parseDefines))
 										return;
 								}
 
@@ -576,12 +585,14 @@ namespace uba
 		return true;
 	}
 
-	bool IncludeCrawler::HandleInclude(Instance& instance, const StringView& rootPath, const StringView& include, bool parseDefines)
+	bool IncludeCrawler::HandleInclude(Instance& instance, const StringView& rootPath, const StringView& include, const StringView& codeFile, bool parseDefines)
 	{
 		StringBuffer<> fullPath;
 		FixPath(include.data, rootPath.data, rootPath.count, fullPath);
 		if (IsAbsolutePath(include.data))
-			instance.devirtualizePathFunc(fullPath);
+			if (!instance.devirtualizePathFunc(fullPath))
+				m_logger.Info(TC("Failed to devirtualize include path %s in %s"), fullPath.data, codeFile.data);
+
 		if (CaseInsensitiveFs)
 			fullPath.MakeLower();
 
diff --git a/Engine/Source/Programs/UnrealBuildAccelerator/Common/Public/UbaRootPaths.h b/Engine/Source/Programs/UnrealBuildAccelerator/Common/Public/UbaRootPaths.h
--- a/Engine/Source/Programs/UnrealBuildAccelerator/Common/Public/UbaRootPaths.h
+++ b/Engine/Source/Programs/UnrealBuildAccelerator/Common/Public/UbaRootPaths.h
@@ -30,6 +30,8 @@ namespace uba
 		const Root* FindRoot(const StringView& path) const;
 		const TString& GetRoot(u32 index) const;
 
+		TString GetAllRoots() const;
+
 		template<typename CharType, typename Func>
 		bool NormalizeString(Logger& logger, const CharType* str, u64 strLen, const Func& func, bool allowPathsWithoutRoot, const tchar* hint, const tchar* hint2 = TC("")) const;
 
diff --git a/Engine/Source/Programs/UnrealBuildAccelerator/Core/Public/UbaPathUtils.h b/Engine/Source/Programs/UnrealBuildAccelerator/Core/Public/UbaPathUtils.h
--- a/Engine/Source/Programs/UnrealBuildAccelerator/Core/Public/UbaPathUtils.h
+++ b/Engine/Source/Programs/UnrealBuildAccelerator/Core/Public/UbaPathUtils.h
@@ -230,6 +230,9 @@ namespace uba
 			*outBufferCharLen = u32(charLen - 1); // Remove terminator
 		return true;
 #else
+		if (fileName[0] == '\"')
+			++fileName;
+
 		StringBuffer<MaxPath> tmp;
 		if (fileName[0] == '~')
 		{
@@ -253,6 +256,8 @@ namespace uba
 		u64 len = strlen(fileName);
 		memcpy(buffer + memPos, fileName, len);
 		memPos += len;
+		if (buffer[memPos -1] == '\"')
+			--memPos;
 		buffer[memPos] = 0;
 		//printf("PATH: %s\n", buffer);
 		//fflush(stdout);
