[UBA]
* Attempt at fixing shutdown crash on osx

diff --git a/Engine/Source/Programs/UnrealBuildAccelerator/Detours/Private/PosixOS/UbaDetoursFunctionsPosix.cpp b/Engine/Source/Programs/UnrealBuildAccelerator/Detours/Private/PosixOS/UbaDetoursFunctionsPosix.cpp
--- a/Engine/Source/Programs/UnrealBuildAccelerator/Detours/Private/PosixOS/UbaDetoursFunctionsPosix.cpp
+++ b/Engine/Source/Programs/UnrealBuildAccelerator/Detours/Private/PosixOS/UbaDetoursFunctionsPosix.cpp
@@ -33,6 +33,7 @@ namespace uba
 
 	bool g_isDetouring;
 	bool g_isInitialized;
+	bool g_isCancelled;
 	void Deinit();
 }
 
@@ -2666,6 +2667,7 @@ namespace uba
 		g_isInitialized = false;
 		g_isDetouring = false;
 
+		if (!g_isCancelled)
 		{
 			SCOPED_WRITE_LOCK(g_fileHandlesLock, lock);
 			for (auto& kv : g_fileHandles)
diff --git a/Engine/Source/Programs/UnrealBuildAccelerator/Detours/Private/PosixOS/UbaDetoursMainPosix.cpp b/Engine/Source/Programs/UnrealBuildAccelerator/Detours/Private/PosixOS/UbaDetoursMainPosix.cpp
--- a/Engine/Source/Programs/UnrealBuildAccelerator/Detours/Private/PosixOS/UbaDetoursMainPosix.cpp
+++ b/Engine/Source/Programs/UnrealBuildAccelerator/Detours/Private/PosixOS/UbaDetoursMainPosix.cpp
@@ -19,6 +19,7 @@ namespace uba
 	SharedEvent* g_readEvent;
 	SharedEvent* g_writeEvent;
 	u8* g_messageMappingMem;
+	extern bool g_isCancelled;
 
 	void PreInit(const char* logFile);
 	void Init();
@@ -122,11 +123,13 @@ static void __attribute__((destructor(102))) InitDtor()
 
 namespace uba
 {
+	constexpr u32 WritableMemSize = CommunicationMemSize - sizeof(SharedEvent) * 3;
+
 	BinaryWriter::BinaryWriter()
 	{
 		m_begin = g_messageMappingMem;
 		m_pos = m_begin;
-		m_end = m_begin + CommunicationMemSize - sizeof(SharedEvent) * 3;
+		m_end = m_begin + WritableMemSize;
 	}
 
 	void BinaryWriter::Flush(bool waitOnResponse)
@@ -144,10 +147,16 @@ namespace uba
 			
 			// check if session process is gone
 			if (kill(g_sessionPid, 0) == -1 && errno == ESRCH)
+			{
+				g_isCancelled = true;
 				exit(1337);
-			
+			}
+
 			if (g_cancelEvent->IsSet(0))
+			{
+				g_isCancelled = true;
 				exit(1339);
+			}
 		} while (true);
 	}
 
@@ -155,6 +164,6 @@ namespace uba
 	{
 		m_begin = g_messageMappingMem;
 		m_pos = m_begin;
-		m_end = m_begin + CommunicationMemSize - sizeof(SharedEvent) * 3;
+		m_end = m_begin + WritableMemSize;
 	}
 }
