[UBA]
* Rewrote include crawler to have separate implementations for clang rsp and msvc rsp after running into problems with ios target. Also renamed it to dependency crawler
* Fixed so crawler does not run on child processes

diff --git a/Engine/Source/Programs/UnrealBuildAccelerator/Agent/Private/UbaAgent.cpp b/Engine/Source/Programs/UnrealBuildAccelerator/Agent/Private/UbaAgent.cpp
--- a/Engine/Source/Programs/UnrealBuildAccelerator/Agent/Private/UbaAgent.cpp
+++ b/Engine/Source/Programs/UnrealBuildAccelerator/Agent/Private/UbaAgent.cpp
@@ -1183,7 +1183,7 @@ namespace uba
 			info.downloadDetoursLib = downloadDetoursLib;
 			info.memWaitLoadPercent = u8(memWaitLoadPercent);
 			info.memKillLoadPercent = u8(memKillLoadPercent);
-			info.useIncludeCrawler = useCrawler;
+			info.useDependencyCrawler = useCrawler;
 
 			if (!quiet)
 				info.processFinished = processFinished;
diff --git a/Engine/Source/Programs/UnrealBuildAccelerator/Common/Private/UbaSessionClient.cpp b/Engine/Source/Programs/UnrealBuildAccelerator/Common/Private/UbaSessionClient.cpp
--- a/Engine/Source/Programs/UnrealBuildAccelerator/Common/Private/UbaSessionClient.cpp
+++ b/Engine/Source/Programs/UnrealBuildAccelerator/Common/Private/UbaSessionClient.cpp
@@ -22,7 +22,7 @@ namespace uba
 
 		if (const ConfigTable* table = config.GetTable(TC("Session")))
 		{
-			table->GetValueAsBool(useIncludeCrawler, TC("UseIncludeCrawler"));
+			table->GetValueAsBool(useDependencyCrawler, TC("UseDependencyCrawler"));
 		}
 	}
 
@@ -44,7 +44,7 @@ namespace uba
 	,	m_waitToSendEvent(false)
 	,	m_loop(true)
 	,	m_allowSpawn(true)
-	,	m_includeCrawler(m_logger, m_client)
+	,	m_dependencyCrawler(m_logger, m_client)
 	{
 		m_maxProcessCount = info.maxProcessCount;
 		m_dedicated = info.dedicated;
@@ -58,7 +58,7 @@ namespace uba
 		m_memKillLoadPercent = info.memKillLoadPercent;
 		m_processFinished = info.processFinished;
 
-		m_useIncludeCrawler = info.useIncludeCrawler;
+		m_useDependencyCrawler = info.useDependencyCrawler;
 
 		m_processIdCounter = ~0u / 2; // We set this value to a very high value.. because it will be used by child processes and we don't want id from server and child process id to collide
 
@@ -82,8 +82,8 @@ namespace uba
 
 
 
-		if (m_useIncludeCrawler)
-			m_includeCrawler.Init([this](const StringView& fileName, u32& outAttr) // FileExists
+		if (m_useDependencyCrawler)
+			m_dependencyCrawler.Init([this](const StringView& fileName, u32& outAttr) // FileExists
 				{
 					u32 tableOffset = 0;
 					if (!EntryExists(fileName, tableOffset))
@@ -95,7 +95,7 @@ namespace uba
 					outAttr = reader.ReadFileAttributes();
 					return true;
 				},
-			[this](const StringView& path, const IncludeCrawler::FileFunc& fileFunc) // TraverseFilesFunc
+			[this](const StringView& path, const DependencyCrawler::FileFunc& fileFunc) // TraverseFilesFunc
 				{
 					u32 tableOffset = 0;
 					if (!EntryExists(path, tableOffset))
@@ -513,7 +513,8 @@ namespace uba
 			if (!SendRootsHandle(rootsHandle))
 				return false;
 
-		RunIncludeCrawler(process);
+		if (!process.m_parentProcess)
+			RunDependencyCrawler(process);
 
 		return true;
 	}
@@ -2102,20 +2103,22 @@ namespace uba
 		return logLineCount;
 	}
 
-	void SessionClient::RunIncludeCrawler(ProcessImpl& process)
+	void SessionClient::RunDependencyCrawler(ProcessImpl& process)
 	{
-		if (!m_useIncludeCrawler)
+		if (!m_useDependencyCrawler)
 			return;
 
 		auto& startInfo = process.GetStartInfo();
-		if (!startInfo.rules->SupportsIncludeCrawler())
+
+		auto crawlerType = startInfo.rules->GetDependencyCrawlerType();
+		if (crawlerType == DependencyCrawlerType_None)
 			return;
 
 		const tchar* at = TStrchr(startInfo.arguments, '@');
 		if (!at)
 			return;
 
-		auto CreateFileFunc = [this, processId = process.GetId(), rules = startInfo.rules](const StringView& fileName, const IncludeCrawler::AccessFileFunc& func)
+		auto CreateFileFunc = [this, processId = process.GetId(), rules = startInfo.rules](const StringView& fileName, const DependencyCrawler::AccessFileFunc& func)
 			{
 				CreateFileResponse out;
 				{
@@ -2154,7 +2157,7 @@ namespace uba
 				return DevirtualizePath(inOut, rootsHandle);
 			};
 
-		m_includeCrawler.Add(at + 1, startInfo.workingDir, CreateFileFunc, DevirtualizePathFunc, startInfo.application, startInfo.rules);
+		m_dependencyCrawler.Add(at + 1, startInfo.workingDir, CreateFileFunc, DevirtualizePathFunc, startInfo.application, crawlerType, startInfo.rules->index);
 	}
 
 	bool SessionClient::ParseDirectoryTable()
@@ -2178,6 +2181,7 @@ namespace uba
 		else
 			key = ToStringKey(path);
 
+		StringBuffer<> dirName;
 		u32 tableOffset = 0;
 		DirectoryTable::Exists exists = m_directoryTable.EntryExists(key, path, true, &tableOffset);
 		if (exists == DirectoryTable::Exists_Maybe)
@@ -2189,7 +2193,6 @@ namespace uba
 
 			if (exists == DirectoryTable::Exists_Maybe)
 			{
-				StringBuffer<> dirName;
 				if (const tchar* lastSeparator = TStrrchr(path.data, PathSeparator))
 					dirName.Append(path.data, lastSeparator - path.data);
 				
@@ -2214,10 +2217,11 @@ namespace uba
 				exists = m_directoryTable.EntryExists(key, path, true, &tableOffset);
 			}
 		}
-		if (exists == DirectoryTable::Exists_No)
-			return false;
 
-		UBA_ASSERT(exists != DirectoryTable::Exists_Maybe);
+		UBA_ASSERTF(exists != DirectoryTable::Exists_Maybe, TC("This should not happen. Asking for directory %s"), dirName.data);
+
+		if (exists != DirectoryTable::Exists_Yes)
+			return false;
 
 		outTableOffset = tableOffset;
 
diff --git a/Engine/Source/Programs/UnrealBuildAccelerator/Common/Public/UbaIncludeCrawler.h b/Engine/Source/Programs/UnrealBuildAccelerator/Common/Public/UbaIncludeCrawler.h
--- a/Engine/Source/Programs/UnrealBuildAccelerator/Common/Public/UbaIncludeCrawler.h
+++ b/Engine/Source/Programs/UnrealBuildAccelerator/Common/Public/UbaIncludeCrawler.h
@@ -8,17 +8,19 @@
 #include "UbaProcessUtils.h"
 #include "UbaWorkManager.h"
 
+#define UBA_CRAWLER_DEBUG_PRINT(str1, str2) //wprintf(TC("") str1 ": %s\n", str2);
+
 // Note
-// IncludeCrawler is quite Unreal specific. It might work in different environments but is based on using response (.rsp) files.
+// DependencyCrawler is quite Unreal specific. It might work in different environments but is based on using response (.rsp) files.
 // It does parse some defines to figure out custom paths but code should also work without this
 // Ideally it should be a full fledged preprocessor parser but that would also not make it possible to run this in parallel
 
 namespace uba
 {
-	class IncludeCrawler
+	class DependencyCrawler
 	{
 	public:
-		inline IncludeCrawler(Logger& logger, WorkManager& workManager);
+		inline DependencyCrawler(Logger& logger, WorkManager& workManager);
 
 		using FileExistsFunc = Function<bool(const StringView& fileName, u32& outAttr)>;
 		using FileFunc = Function<void(const StringView& file, bool isDirectory)>;
@@ -29,7 +31,7 @@ namespace uba
 		using CreateFileFunc = Function<bool(const StringView& fileName, const AccessFileFunc& func)>;
 
 		inline void Init(FileExistsFunc&& fileExistsFunc, TraverseFilesFunc&& traverseFilesFunc);
-		inline bool Add(const tchar* rsp, const tchar* workDir, CreateFileFunc&& func, DevirtualizePathFunc&& devirtualizePathFunc, const tchar* app = nullptr, const ApplicationRules* rules = nullptr);
+		inline bool Add(const tchar* rsp, const tchar* workDir, CreateFileFunc&& func, DevirtualizePathFunc&& devirtualizePathFunc, const tchar* app, DependencyCrawlerType type, u32 ruleIndex);
 
 		Logger& m_logger;
 		WorkManager& m_workManager;
@@ -50,6 +52,7 @@ namespace uba
 
 		struct Instance
 		{
+			DependencyCrawlerType type;
 			TString application;
 			TString rsp;
 			TString workDir;
@@ -63,10 +66,18 @@ namespace uba
 			DevirtualizePathFunc devirtualizePathFunc;
 			Vector<IncludeRoot> includeRoots;
 			Pch* pch = nullptr;
+
+			Atomic<u32> refCount;
+		};
+
+		struct InstanceRef
+		{
+			InstanceRef(Instance& i) : instance(i) { ++instance.refCount; }
+			InstanceRef(const InstanceRef& i) : instance(i.instance) { ++instance.refCount; }
+			~InstanceRef() { if (!--instance.refCount) delete &instance; }
+			Instance& instance;
 		};
 
-		Futex m_instancesLock;
-		List<Instance> m_instances;
 		Futex m_builtIncludesHandledLock;
 		UnorderedSet<u32> m_builtIncludesHandled;
 		FileExistsFunc m_fileExistsFunc;
@@ -84,27 +95,27 @@ namespace uba
 	};
 
 
-	IncludeCrawler::IncludeCrawler(Logger& logger, WorkManager& workManager)
+	DependencyCrawler::DependencyCrawler(Logger& logger, WorkManager& workManager)
 	:	m_logger(logger)
 	,	m_workManager(workManager)
 	{
 	}
 
-	void IncludeCrawler::Init(FileExistsFunc&& fileExistsFunc, TraverseFilesFunc&& traverseFilesFunc)
+	void DependencyCrawler::Init(FileExistsFunc&& fileExistsFunc, TraverseFilesFunc&& traverseFilesFunc)
 	{
 		m_fileExistsFunc = std::move(fileExistsFunc);
 		m_traverseFilesFunc = std::move(traverseFilesFunc);
 	}
 
-	bool IncludeCrawler::Add(const tchar* rsp, const tchar* workDir, CreateFileFunc&& createFileFunc, DevirtualizePathFunc&& devirtualizePathFunc, const tchar* app, const ApplicationRules* rules)
+	bool DependencyCrawler::Add(const tchar* rsp, const tchar* workDir, CreateFileFunc&& createFileFunc, DevirtualizePathFunc&& devirtualizePathFunc, const tchar* app, DependencyCrawlerType type, u32 ruleIndex)
 	{
-		SCOPED_FUTEX(m_instancesLock, lock);
-		Instance& instance = m_instances.emplace_back();
-		lock.Leave();
+		auto& instance = *new Instance();
 
+		instance.type = type;
 		instance.application = app;
 		instance.rsp = rsp;
 		instance.workDir = workDir;
+		UBA_ASSERT(*workDir);
 		if (CaseInsensitiveFs)
 			ToLower(instance.workDir.data());
 		if (instance.workDir.back() != PathSeparator)
@@ -112,8 +123,9 @@ namespace uba
 		instance.createFileFunc = std::move(createFileFunc);
 		instance.devirtualizePathFunc = std::move(devirtualizePathFunc);
 
-		m_workManager.AddWork([this, &instance, rules](const WorkContext& context)
+		m_workManager.AddWork([this, ref = InstanceRef(instance), ruleIndex](const WorkContext& context)
 			{
+				auto& instance = ref.instance;
 				CodeFiles codeFiles;
 				if (!ParseRsp(instance, instance.rsp.c_str(), codeFiles))
 					return;
@@ -158,10 +170,10 @@ namespace uba
 					}
 				}
 
-				if (rules && !instance.builtinIncludesDir.empty())
+				if (!instance.builtinIncludesDir.empty())
 				{
 					SCOPED_FUTEX(m_builtIncludesHandledLock, lock);
-					bool shouldHandle = m_builtIncludesHandled.insert(rules->index).second;
+					bool shouldHandle = m_builtIncludesHandled.insert(ruleIndex).second;
 					lock.Leave();
 					if (shouldHandle)
 					{
@@ -178,9 +190,9 @@ namespace uba
 								if (isDirectory)
 									return;
 								codeFilePath.Resize(fixedPath.count).Append(file);
-								m_workManager.AddWork([this, &instance, codeFile = TString(codeFilePath.data)](const WorkContext& context)
+								m_workManager.AddWork([this, ref, codeFile = TString(codeFilePath.data)](const WorkContext& context)
 									{
-										instance.createFileFunc(codeFile, [&](const void* data, u64 dataSize) { return true; });
+										ref.instance.createFileFunc(codeFile, [&](const void* data, u64 dataSize) { return true; });
 									}, 1, TC("Crawler"), ColorWork);
 							});
 					}
@@ -190,7 +202,7 @@ namespace uba
 		return true;
 	}
 
-	bool IncludeCrawler::ParseRsp(Instance& instance, const tchar* rsp, CodeFiles& outCodeFiles)
+	bool DependencyCrawler::ParseRsp(Instance& instance, const tchar* rsp, CodeFiles& outCodeFiles)
 	{
 		StringBuffer<> fixedPath;
 		FixPath(rsp, instance.workDir.data(), instance.workDir.size(), fixedPath);
@@ -204,14 +216,13 @@ namespace uba
 			});
 	}
 
-	bool IncludeCrawler::ParseRsp2(Instance& instance, const void* data, u64 dataSize, const tchar* rsp, CodeFiles& outCodeFiles)
+	bool DependencyCrawler::ParseRsp2(Instance& instance, const void* data, u64 dataSize, const tchar* rsp, CodeFiles& outCodeFiles)
 	{
-		StringBuffer<> prevArg;
-
-		auto AddCodeFile = [&](const tchar* file, bool pushFront, bool hasPch)
+		auto AddCodeFile = [&](const StringView& file, bool pushFront, bool hasPch)
 			{
+				UBA_CRAWLER_DEBUG_PRINT("CODEFILE", file.data);
 				StringBuffer<> fixedPath2;
-				FixPath(file, instance.workDir.data(), instance.workDir.size(), fixedPath2);
+				FixPath(file.data, instance.workDir.data(), instance.workDir.size(), fixedPath2);
 				if (!instance.devirtualizePathFunc(fixedPath2))
 					m_logger.Info(TC("Failed to devirtualize path %s in %s"), fixedPath2.data, rsp);
 
@@ -223,10 +234,11 @@ namespace uba
 					outCodeFiles.push_back({TString(fixedPath2.data), hasPch});
 			};
 
-		auto AddRoot = [&](const tchar* path)
+		auto AddRoot = [&](const StringView& path)
 			{
+				UBA_CRAWLER_DEBUG_PRINT("ROOT", path.data);
 				StringBuffer<> fixedPath2;
-				FixPath(path, instance.workDir.data(), instance.workDir.size(), fixedPath2);
+				FixPath(path.data, instance.workDir.data(), instance.workDir.size(), fixedPath2);
 				if (!instance.devirtualizePathFunc(fixedPath2))
 					m_logger.Info(TC("Failed to devirtualize path %s in %s"), fixedPath2.data, rsp);
 				if (CaseInsensitiveFs)
@@ -242,111 +254,242 @@ namespace uba
 					instance.includeRoots.push_back({TString(fixedPath2.data), bloomFilter});
 			};
 
+		auto IgnoreOption = [](const StringView& path) { UBA_CRAWLER_DEBUG_PRINT("IGNORED", path.data); };
 
-		ParseArguments((const char*)data, dataSize, [&](const char* arg, u32 argLen)
+		if (instance.type == DependencyCrawlerType_MsvcCompiler)
+		{
+			struct MsvcOptionWithArg { const StringView name; bool addRoot; };
+			const MsvcOptionWithArg msvcOptionsWithArg[] = 
 			{
-				StringBuffer<> sb;
-				sb.Append(arg, argLen);
-				if (sb[0] == '-')
+				{ TCV("/I"), true },
+				{ TCV("/external:I"), true },
+				{ TCV("/imsvc"), true },
+				{ TCV("/experimental:log"), false },
+				{ TCV("/sourceDependencies"), false},
+			};
+
+			constexpr const tchar* clangClOptionsWithArg[] = 
+			{
+				TC("-D"),
+				TC("-x"),
+				TC("-o"),
+				TC("-include"),
+				TC("-include-pch"),
+				TC("-vctoolsdir"),
+				TC("-Xclang"),
+				TC("-target"),
+				TC("-arch"),
+			};
+
+			StringBuffer<> prevArg;
+			bool addRoot = false;
+			bool handled = false;
+			ParseArguments((const char*)data, dataSize, [&](const char* arg, u32 argLen)
 				{
-					prevArg.Clear();
-					if (sb.StartsWith(TC("-isystem"), false))
-						AddRoot(sb.data + 8);
-					else if (sb.StartsWith(TC("-I"), false))
-						AddRoot(sb.data + 2);
-					else if (sb.StartsWith(TC("-resource-dir"), false))
+					StringBuffer<> sb;
+					sb.Append(arg, argLen);
+
+					if (handled)
 					{
-						instance.builtinIncludesDir = sb.data + 14;
-						instance.builtinIncludesDir += TC("/include");
+						handled = false;
+						if (addRoot)
+							AddRoot(sb);
+						else
+							IgnoreOption(sb);
+						addRoot = false;
+						return;
 					}
-					else
+					else if (sb[0] == '/')
 					{
-						constexpr const tchar* optionsWithArg[] = 
-						{
-							TC("-D"),
-							TC("-x"),
-							TC("-o"),
-							TC("-include"),
-							TC("-include-pch"),
-							TC("-vctoolsdir"),
-							TC("-Xclang"),
-							TC("-target"),
-							TC("-arch"),
-						};
-
-						for (auto option : optionsWithArg)
+						for (auto& option : msvcOptionsWithArg)
 						{
-							if (!sb.Equals(option, false))
+							if (!option.name.Equals(sb))
 								continue;
-							prevArg = sb;
-							break;
+							addRoot = option.addRoot;
+							handled = true;
+							return;
 						}
-					}
-				}
-				else if (sb[0] == '/')
-				{
-					if (sb.StartsWith(TC("/FI")))
-					{
-						// Check if there is a precompiled header deps file
-						StringBuffer<> fixedFile;
-						FixPath(sb.data + 3, instance.workDir.data(), instance.workDir.size(), fixedFile);
-						if (!instance.devirtualizePathFunc(fixedFile))
-							m_logger.Info(TC("Failed to devirtualize path %s in %s"), fixedFile.data, rsp);
-
-						if (CaseInsensitiveFs)
-							fixedFile.MakeLower();
 
-						StringKey key = ToStringKey(fixedFile);
-						SCOPED_FUTEX(m_codeFilesHandledLock, lock);
-						auto insres = m_codeFilesHandled.try_emplace(key);
-						lock.Leave();
-
-						HandledFile& handledFile = insres.first->second;
-						SCOPED_FUTEX(handledFile.lock, fileLock);
-						if (!handledFile.handled)
+						if (sb.StartsWith(TC("/FI")))
 						{
-							handledFile.handled = true;
+							// Check if there is a precompiled header deps file
+							StringBuffer<> fixedFile;
+							FixPath(sb.data + 3, instance.workDir.data(), instance.workDir.size(), fixedFile);
+							if (!instance.devirtualizePathFunc(fixedFile))
+								m_logger.Info(TC("Failed to devirtualize path %s in %s"), fixedFile.data, rsp);
+
+							if (CaseInsensitiveFs)
+								fixedFile.MakeLower();
+
+							StringKey key = ToStringKey(fixedFile);
+							SCOPED_FUTEX(m_codeFilesHandledLock, lock);
+							auto insres = m_codeFilesHandled.try_emplace(key);
+							lock.Leave();
+
+							HandledFile& handledFile = insres.first->second;
+							SCOPED_FUTEX(handledFile.lock, fileLock);
+							if (!handledFile.handled)
+							{
+								handledFile.handled = true;
 
-							fixedFile.Append(TC(".pch"));
+								fixedFile.Append(TC(".pch"));
 
-							u32 attributes = 0;
-							bool hasPch = m_fileExistsFunc(fixedFile, attributes);
-							AddCodeFile(sb.data + 3, true, hasPch);
-						}
+								u32 attributes = 0;
+								bool hasPch = m_fileExistsFunc(fixedFile, attributes);
+								AddCodeFile(StringView(sb).Skip(3), true, hasPch);
+							}
 
-						prevArg.Clear();
+							prevArg.Clear();
+							return;
+						}
+						else if (sb.StartsWith(TC("/Yu")))
+						{
+							instance.usePch = true;
+							prevArg.Clear();
+							return;
+						}
+						else
+						{
+							IgnoreOption(sb);
+						}
+						prevArg = sb;
+					}
+					else if (sb[0] == '-')
+					{
+						for (auto option : clangClOptionsWithArg)
+							handled |= sb.Equals(option);
+						if (sb.StartsWith(TC("-resource-dir"), false))
+						{
+							instance.builtinIncludesDir = sb.data + 14;
+							instance.builtinIncludesDir += TC("/include");
+							UBA_CRAWLER_DEBUG_PRINT("RESOURCEDIR", sb.data + 14);
+						}
+						else
+							IgnoreOption(sb);
 					}
-					else if (sb.StartsWith(TC("/Yu")))
+					else if (sb[0] == '@')
 					{
-						instance.usePch = true;
+						ParseRsp(instance, sb.data + 1, outCodeFiles);
 					}
 					else
 					{
-						prevArg = sb;
+						AddCodeFile(sb, false, false);
 					}
-				}
-				else if (sb[0] == '@')
-				{
-					ParseRsp(instance, sb.data + 1, outCodeFiles);
-				}
-				else
+				});
+		}
+		else if (instance.type == DependencyCrawlerType_ClangCompiler)
+		{
+			enum ArgType { ArgType_None, ArgType_Ignore, ArgType_RootPath, ArgType_ResourceDir, ArgType_Code };
+			struct OptionWithArg { const StringView name; ArgType type; };
+			constexpr const OptionWithArg dashOptionsWithArg[] = 
+			{
+				{ TCV("-D"), ArgType_Ignore },
+				{ TCV("-x"), ArgType_Ignore },
+				{ TCV("-o"), ArgType_Ignore },
+				{ TCV("-include"), ArgType_Code },
+				{ TCV("-include-pch"), ArgType_Ignore },
+				{ TCV("-vctoolsdir"), ArgType_Ignore },
+				{ TCV("-Xclang"), ArgType_Ignore },
+				{ TCV("-target"), ArgType_Ignore },
+				{ TCV("-arch"), ArgType_Ignore },
+				{ TCV("--sysroot"), ArgType_RootPath},
+				{ TCV("-isystem"), ArgType_RootPath },
+				{ TCV("-isysroot"), ArgType_RootPath },
+				{ TCV("-internal-isystem"), ArgType_RootPath },
+				{ TCV("-I"), ArgType_RootPath },
+				{ TCV("-resource-dir"), ArgType_ResourceDir },
+				{ TCV("-dependency-file"), ArgType_Ignore },
+				{ TCV("-internal-externc-isystem"), ArgType_Ignore },
+				{ TCV("-MT"), ArgType_Ignore },
+			};
+
+			//StringBuffer<> prevArg;
+			ArgType type = ArgType_None;
+			ParseArguments((const char*)data, dataSize, [&](const char* arg, u32 argLen)
 				{
-					if (prevArg.Equals(TC("/I")) || prevArg.Equals(TC("/imsvc")) || prevArg.Equals(TC("/external:I")))
-						AddRoot(sb.data);
-					if (prevArg.Equals(TC("-include"), false))
+					StringBuffer<> sb;
+					sb.Append(arg, argLen);
+
+					if (type == ArgType_None)
 					{
-						if (!sb.EndsWith(TC(".gch")))
-							AddCodeFile(sb.data, true, false);
+						if (sb[0] == '-')
+						{
+							for (auto& option : dashOptionsWithArg)
+							{
+								if (!sb.Equals(option.name, false))
+									continue;
+								//prevArg = sb; // Arg comes in next parse call
+								type = option.type;
+								return;
+							}
+
+							if (auto equalsPos = TStrchr(sb.data, '='))
+							{
+								u32 equalsOffset = u32(equalsPos - sb.data);
+								StringView option2(sb.data, equalsOffset);
+								for (auto option : dashOptionsWithArg)
+								{
+									if (!option2.Equals(option.name, false))
+										continue;
+									++equalsOffset;
+									sb.Clear().Append(arg + equalsOffset, argLen - equalsOffset);
+									type = option.type;
+									//prevArg = option;
+									break;
+								}
+							}
+							else if (sb.StartsWith(TC("-I"), false))
+							{
+								type = ArgType_RootPath;
+								sb.Clear().Append(arg + 2, argLen - 2);
+							}
+							else if (sb.StartsWith(TC("-isystem"), false))
+							{
+								type = ArgType_RootPath;
+								sb.Clear().Append(arg + 8, argLen - 8);
+							}
+							if (type == ArgType_None)
+							{
+								IgnoreOption(sb);
+								return;
+							}
+						}
 					}
-					else if (prevArg.IsEmpty())
-						AddCodeFile(sb.data, false, false);
-					prevArg.Clear();
-				}
-			});
+
+					if (type != ArgType_None)
+					{
+						switch (type)
+						{
+						case ArgType_RootPath:
+							AddRoot(sb);
+							break;
+						case ArgType_Code:
+							if (!sb.EndsWith(TC(".gch")))
+								AddCodeFile(sb, true, false);
+							break;
+						case ArgType_ResourceDir:
+							instance.builtinIncludesDir = sb.data;
+							instance.builtinIncludesDir += TC("/include");
+							break;
+						default:
+							IgnoreOption(sb);
+						}
+						type = ArgType_None;
+						return;
+					}
+					else if (sb[0] == '@')
+						ParseRsp(instance, sb.data + 1, outCodeFiles);
+					else if (sb.Contains('/'))
+						AddCodeFile(sb, false, false);
+					else
+						IgnoreOption(sb);
+				});
+			}
+				
 		return true;
 	}
 
-	bool IncludeCrawler::ParsePch(Pch& pch, Instance& instance, const void* data, u64 dataSize)
+	bool DependencyCrawler::ParsePch(Pch& pch, Instance& instance, const void* data, u64 dataSize)
 	{
 		auto AddHandled = [&](const StringView& str)
 			{
@@ -394,7 +537,7 @@ namespace uba
 		return true;
 	}
 
-	bool IncludeCrawler::ParseCodeFile(Instance& instance, const StringView& codeFile, bool parseDefines)
+	bool DependencyCrawler::ParseCodeFile(Instance& instance, const StringView& codeFile, bool parseDefines)
 	{
 		return instance.createFileFunc(codeFile, [&](const void* data, u64 dataSize)
 			{
@@ -402,7 +545,7 @@ namespace uba
 			});
 	}
 
-	bool IncludeCrawler::ParseCodeFile2(Instance& instance, const void* data, u64 dataSize, const StringView& codeFile, bool parseDefines)
+	bool DependencyCrawler::ParseCodeFile2(Instance& instance, const void* data, u64 dataSize, const StringView& codeFile, bool parseDefines)
 	{
 		const char* it = (const char*)data;
 		const char* end = it + dataSize;
@@ -506,8 +649,9 @@ namespace uba
 					{
 						include.MakeLower().FixPathSeparators();
 						auto work = 
-						[this, &instance, isQuote, include2 = TString(include.data), codeFile2 = TString(codeFile.data), parseDefines](const WorkContext& context)
+						[this, ref = InstanceRef(instance), isQuote, include2 = TString(include.data), codeFile2 = TString(codeFile.data), parseDefines](const WorkContext& context)
 							{
+								auto& instance = ref.instance;
 								if (isQuote)
 								{
 									StringBuffer<> localDir;
@@ -585,8 +729,10 @@ namespace uba
 		return true;
 	}
 
-	bool IncludeCrawler::HandleInclude(Instance& instance, const StringView& rootPath, const StringView& include, const StringView& codeFile, bool parseDefines)
+	bool DependencyCrawler::HandleInclude(Instance& instance, const StringView& rootPath, const StringView& include, const StringView& codeFile, bool parseDefines)
 	{
+		if (include.EndsWith(TC(".ush")))
+			return true;
 		StringBuffer<> fullPath;
 		FixPath(include.data, rootPath.data, rootPath.count, fullPath);
 		if (IsAbsolutePath(include.data))
diff --git a/Engine/Source/Programs/UnrealBuildAccelerator/Common/Public/UbaSessionClient.h b/Engine/Source/Programs/UnrealBuildAccelerator/Common/Public/UbaSessionClient.h
--- a/Engine/Source/Programs/UnrealBuildAccelerator/Common/Public/UbaSessionClient.h
+++ b/Engine/Source/Programs/UnrealBuildAccelerator/Common/Public/UbaSessionClient.h
@@ -29,7 +29,7 @@ namespace uba
 		bool killRandom = false;
 		bool useStorage = true;
 		bool downloadDetoursLib = true;
-		bool useIncludeCrawler = false;	// IncludeCrawler will go wide and try to prefetch all headers and source files needed by the build
+		bool useDependencyCrawler = false;	// Dependency crawler will go wide and try to prefetch all dependencies needed by the process
 		Function<void(const ProcessHandle&)> processFinished;
 	};
 
@@ -104,7 +104,7 @@ namespace uba
 		void SendLogFileToServer(ProcessImpl& pi);
 		void GetLogFileName(StringBufferBase& out, const tchar* logFile, const tchar* arguments, u32 processId);
 		u32 WriteLogLines(BinaryWriter& writer, ProcessImpl& process);
-		void RunIncludeCrawler(ProcessImpl& process);
+		void RunDependencyCrawler(ProcessImpl& process);
 		bool ParseDirectoryTable();
 		bool EntryExists(const StringView& path, u32& outTableOffset);
 
@@ -132,7 +132,7 @@ namespace uba
 		bool m_shouldSendLogToServer = false;
 		bool m_shouldSendTraceToServer = false;
 		bool m_remoteExecutionEnabled = true;
-		bool m_useIncludeCrawler = false;
+		bool m_useDependencyCrawler = false;
 		
 		Atomic<const tchar*> m_terminationReason;
 		Atomic<u64> m_terminationTime;
@@ -181,7 +181,7 @@ namespace uba
 		UnorderedMap<CasKey, Vector<u8>> m_hostRunCache;
 		Futex m_hostRunCacheLock;
 
-		IncludeCrawler m_includeCrawler;
+		DependencyCrawler m_dependencyCrawler;
 		u32 m_dirtableParsedPosition = 0;
 		Futex m_dirVisitedLock;
 		struct DirVisitedEntry { Futex lock; bool handled = false; };
diff --git a/Engine/Source/Programs/UnrealBuildAccelerator/Core/Private/UbaApplicationRules.cpp b/Engine/Source/Programs/UnrealBuildAccelerator/Core/Private/UbaApplicationRules.cpp
--- a/Engine/Source/Programs/UnrealBuildAccelerator/Core/Private/UbaApplicationRules.cpp
+++ b/Engine/Source/Programs/UnrealBuildAccelerator/Core/Private/UbaApplicationRules.cpp
@@ -118,9 +118,9 @@ namespace uba
 			return fileName.EndsWith(TC(".dep.json"));
 		}
 
-		virtual bool SupportsIncludeCrawler() const override
+		virtual DependencyCrawlerType GetDependencyCrawlerType() const override
 		{
-			return true;
+			return DependencyCrawlerType_MsvcCompiler;
 		}
 	};
 
@@ -317,9 +317,9 @@ namespace uba
 			return fileName.EndsWith(TC(".d"));
 		}
 
-		virtual bool SupportsIncludeCrawler() const override
+		virtual DependencyCrawlerType GetDependencyCrawlerType() const override
 		{
-			return true;
+			return DependencyCrawlerType_ClangCompiler;
 		}
 	};
 
diff --git a/Engine/Source/Programs/UnrealBuildAccelerator/Core/Private/UbaStringBuffer.cpp b/Engine/Source/Programs/UnrealBuildAccelerator/Core/Private/UbaStringBuffer.cpp
--- a/Engine/Source/Programs/UnrealBuildAccelerator/Core/Private/UbaStringBuffer.cpp
+++ b/Engine/Source/Programs/UnrealBuildAccelerator/Core/Private/UbaStringBuffer.cpp
@@ -311,6 +311,16 @@ namespace uba
 		return TStrchr(data, c) != nullptr;
 	}
 
+	bool StringBufferBase::StartsWith(const StringView& str, bool ignoreCase) const
+	{
+		return str.count >= count && uba::StartsWith(data, str.data, ignoreCase);
+	}
+
+	bool StringBufferBase::Equals(const StringView& str, bool ignoreCase) const
+	{
+		return count == str.count && uba::Equals(data, str.data, count, ignoreCase);
+	}
+
 	const tchar* StringBufferBase::First(tchar c, u64 offset) const
 	{
 		return TStrchr(data + offset, c);
diff --git a/Engine/Source/Programs/UnrealBuildAccelerator/Core/Public/UbaApplicationRules.h b/Engine/Source/Programs/UnrealBuildAccelerator/Core/Public/UbaApplicationRules.h
--- a/Engine/Source/Programs/UnrealBuildAccelerator/Core/Public/UbaApplicationRules.h
+++ b/Engine/Source/Programs/UnrealBuildAccelerator/Core/Public/UbaApplicationRules.h
@@ -15,6 +15,14 @@ namespace uba
 	extern GlobalRules g_globalRules;
 
 
+	enum DependencyCrawlerType
+	{
+		DependencyCrawlerType_None,
+		DependencyCrawlerType_ClangCompiler,
+		DependencyCrawlerType_MsvcCompiler,
+	};
+
+
 	class ApplicationRules
 	{
 	public:
@@ -167,9 +175,9 @@ namespace uba
 			return false;
 		}
 
-		virtual bool SupportsIncludeCrawler() const
+		virtual DependencyCrawlerType GetDependencyCrawlerType() const
 		{
-			return false;
+			return DependencyCrawlerType_None;
 		}
 
 		u32 index = ~0u;
diff --git a/Engine/Source/Programs/UnrealBuildAccelerator/Core/Public/UbaStringBuffer.h b/Engine/Source/Programs/UnrealBuildAccelerator/Core/Public/UbaStringBuffer.h
--- a/Engine/Source/Programs/UnrealBuildAccelerator/Core/Public/UbaStringBuffer.h
+++ b/Engine/Source/Programs/UnrealBuildAccelerator/Core/Public/UbaStringBuffer.h
@@ -54,10 +54,12 @@ namespace uba
 
 		bool IsEmpty() const { return count == 0; }
 		bool StartsWith(const tchar* str, bool ignoreCase = true) const { return uba::StartsWith(data, str, ignoreCase); }
+		bool StartsWith(const StringView& str, bool ignoreCase = true) const;
 		bool EndsWith(const tchar* value, bool ignoreCase = true) const { return uba::EndsWith(data, count, value, ignoreCase); }
 		bool Contains(tchar c) const;
 		bool Contains(const tchar* str, bool ignoreCase = true, const tchar** pos = nullptr) const { return uba::Contains(data, str, ignoreCase, pos); }
 		bool Equals(const tchar* str, bool ignoreCase = true) const { return uba::Equals(data, str, ignoreCase); }
+		bool Equals(const StringView& str, bool ignoreCase = true) const;
 		const tchar* First(tchar c, u64 offset = 0) const;
 		const tchar* Last(tchar c, u64 offset = 0) const;
 		const tchar* GetFileName() const { return uba::GetFileName(data); }
@@ -101,11 +103,11 @@ namespace uba
 
 	struct StringView
 	{
-		StringView() : data(TC("")), count(0) {}
-		StringView(NoInitEnum) {}
-		StringView(const StringView&) = default;
-		StringView(const tchar* d, u32 c) : data(d), count(c) {}
-		StringView(const StringBufferBase& sb) : data(sb.data), count(sb.count) {}
+		constexpr StringView() : data(TC("")), count(0) {}
+		constexpr StringView(NoInitEnum) {}
+		constexpr StringView(const StringView&) = default;
+		constexpr StringView(const tchar* d, u32 c) : data(d), count(c) {}
+		constexpr StringView(const StringBufferBase& sb) : data(sb.data), count(sb.count) {}
 		StringView(const TString& str) : data(str.data()), count(u32(str.size())) {}
 
 		bool IsEmpty() const { return count == 0; }
@@ -113,7 +115,9 @@ namespace uba
 		bool EndsWith(const tchar* value, bool ignoreCase = true) const { return uba::EndsWith(data, count, value, ignoreCase); }
 		bool Contains(const tchar* str, bool ignoreCase = true) const { return uba::Contains(data, str, ignoreCase); }
 		bool Equals(const tchar* str, bool ignoreCase = true) const { return uba::Equals(data, str, ignoreCase); }
+		bool Equals(const StringView& str, bool ignoreCase = true) const { return count == str.count && uba::Equals(data, str.data, count, ignoreCase); }
 		StringView GetFileName() const { const tchar* str = uba::GetFileName(data); return StringView(str, count - u32(str - data)); }
+		StringView Skip(u32 skipCount) { return StringView(data + skipCount, count - skipCount); }
 
 		TString ToString() const { return TString(data, data + count); }
 
@@ -124,7 +128,7 @@ namespace uba
 	StringView ToView(const tchar* s);
 
 	template<typename T, size_t size>
-	StringView AsView(T (&buffer)[size]) { return StringView(buffer, size-1); }
+	constexpr StringView AsView(T (&buffer)[size]) { return StringView(buffer, size-1); }
 
 	#define TCV(x) AsView(TC(x))
 
diff --git a/Engine/Source/Programs/UnrealBuildAccelerator/Test/Private/UbaTestAll.h b/Engine/Source/Programs/UnrealBuildAccelerator/Test/Private/UbaTestAll.h
--- a/Engine/Source/Programs/UnrealBuildAccelerator/Test/Private/UbaTestAll.h
+++ b/Engine/Source/Programs/UnrealBuildAccelerator/Test/Private/UbaTestAll.h
@@ -70,7 +70,7 @@ namespace uba
 		UBA_TEST(TestStdOutLocal) \
 		UBA_TEST(TestStdOutViaCmd) \
 		UBA_TEST(TestVolumeCache) \
-		UBA_TEST(TestIncludeCrawler) \
+		UBA_TEST(TestDependencyCrawler) \
 
 
 #define UBA_LINUX_TESTS \
diff --git a/Engine/Source/Programs/UnrealBuildAccelerator/Test/Private/UbaTestIncludeCrawler.cpp b/Engine/Source/Programs/UnrealBuildAccelerator/Test/Private/UbaTestIncludeCrawler.cpp
--- a/Engine/Source/Programs/UnrealBuildAccelerator/Test/Private/UbaTestIncludeCrawler.cpp
+++ b/Engine/Source/Programs/UnrealBuildAccelerator/Test/Private/UbaTestIncludeCrawler.cpp
@@ -6,15 +6,15 @@
 
 namespace uba
 {
-	bool TestIncludeCrawler(LoggerWithWriter& logger, const StringBufferBase& rootDir)
+	bool TestDependencyCrawler(LoggerWithWriter& logger, const StringBufferBase& rootDir)
 	{
 #if 0 // Commented out for now because of hard coded paths below
 
 		WorkManagerImpl workManager(128);
 
-		IncludeCrawler crawler(logger, workManager);
+		DependencyCrawler crawler(logger, workManager);
 
-		auto CreateFileFunc = [&](const StringView& fileName, const IncludeCrawler::AccessFileFunc& func)
+		auto CreateFileFunc = [&](const StringView& fileName, const DependencyCrawler::AccessFileFunc& func)
 			{
 				FileAccessor fa(logger, fileName.data);
 				if (!fa.OpenMemoryRead())
@@ -27,7 +27,7 @@ namespace uba
 				return FileExists(logger, fileName.data, nullptr, &outAttr);
 			};
 
-		auto TraverseFilesFunc = [&](const StringView& path, const IncludeCrawler::FileFunc& fileFunc)
+		auto TraverseFilesFunc = [&](const StringView& path, const DependencyCrawler::FileFunc& fileFunc)
 			{
 				TraverseDir(logger, path, [&](const DirectoryEntry& entry)
 					{
@@ -43,13 +43,22 @@ namespace uba
 		crawler.Init(FileExistsFunc, TraverseFilesFunc);
 
 		const tchar app[] = TC("c:\\sdk\\AutoSDK\\HostWin64\\Win64\\LLVM\\18.1.8\\bin\\clang-cl.exe");
-		const tchar rsp[] = TC("e:\\dev\\fn\\Engine\\Intermediate\\Build\\Win64\\x64\\BlankProgram\\Development\\Core\\Module.Core.1.cpp.obj.rsp");
+
+		//DependencyCrawlerType type = DependencyCrawlerType_MsvcCompiler;
+		//const tchar rsp[] = TC("e:\\dev\\fn\\Engine\\Intermediate\\Build\\Win64\\x64\\BlankProgram\\Development\\Core\\Module.Core.1.cpp.obj.rsp");
+		
+		DependencyCrawlerType type = DependencyCrawlerType_ClangCompiler;
+		//const tchar rsp[] = TC("e:\\dev\\fn\\Engine\\Intermediate\\Build\\Linux\\x64\\BlankProgram\\Development\\Core\\Module.Core.1.cpp.o.rsp");
+		//const tchar rsp[] = TC("e:\\temp\\response-22b46d.txt");
+		const tchar rsp[] = TC("e:\\temp\\MiMalloc.c.o.rsp");
 		//const tchar rsp[] = TC("e:\\dev\\fn\\engine\\plugins\\experimental\\geometrycollectionplugin\\intermediate\\build\\win64\\x64\\unrealeditor\\development\\geometrycollectionnodes\\module.geometrycollectionnodes.3.cpp.obj.rsp");
 		//const tchar rsp[] = TC("e:\\dev\\fn\\Engine\\Plugins\\Runtime\\MeshModelingToolset\\Intermediate\\Build\\Win64\\x64\\UnrealEditor\\Development\\ModelingComponents\\Module.ModelingComponents.3.cpp.obj.rsp");
 		//const tchar rsp[] = TC("e:\\dev\\fn\\Engine\\Intermediate\\Build\\Win64\\x64\\UnrealEditor\\Development\\AnimGraphRuntime\\Module.AnimGraphRuntime.3.cpp.obj.rsp");
 
-		auto rules = GetApplicationRules()[10].rules;
-		crawler.Add(rsp, TC("e:\\dev\\fn\\Engine\\Source"), CreateFileFunc, [](StringBufferBase&) { return true; }, app, rules);
+		crawler.Add(rsp, TC("e:\\dev\\fn\\Engine\\Source"), CreateFileFunc, [](StringBufferBase&)
+			{
+				return true;
+			}, app, type, 10);
 
 		workManager.FlushWork();
 #endif
