[UBA]
* Renamed include crawler files to dependency crawler

diff --git a/Engine/Source/Programs/UnrealBuildAccelerator/Common/Public/UbaDependencyCrawler.h b/Engine/Source/Programs/UnrealBuildAccelerator/Common/Public/UbaDependencyCrawler.h
new file mode 100644
--- /dev/null
+++ b/Engine/Source/Programs/UnrealBuildAccelerator/Common/Public/UbaDependencyCrawler.h
@@ -0,0 +1,767 @@
+// Copyright Epic Games, Inc. All Rights Reserved.
+
+#pragma once
+
+#include "UbaApplicationRules.h"
+#include "UbaFileAccessor.h"
+#include "UbaPathUtils.h"
+#include "UbaProcessUtils.h"
+#include "UbaWorkManager.h"
+
+#define UBA_CRAWLER_DEBUG_PRINT(str1, str2) //wprintf(TC("") str1 ": %s\n", str2);
+
+// Note
+// DependencyCrawler is quite Unreal specific. It might work in different environments but is based on using response (.rsp) files.
+// It does parse some defines to figure out custom paths but code should also work without this
+// Ideally it should be a full fledged preprocessor parser but that would also not make it possible to run this in parallel
+
+namespace uba
+{
+	class DependencyCrawler
+	{
+	public:
+		inline DependencyCrawler(Logger& logger, WorkManager& workManager);
+
+		using FileExistsFunc = Function<bool(const StringView& fileName, u32& outAttr)>;
+		using FileFunc = Function<void(const StringView& file, bool isDirectory)>;
+		using TraverseFilesFunc = Function<void(const StringView& path, const FileFunc& fileFunc)>;
+
+		using AccessFileFunc = Function<bool(const void* data, u64 dataSize)>;
+		using DevirtualizePathFunc = Function<bool(StringBufferBase& inOut)>;
+		using CreateFileFunc = Function<bool(const StringView& fileName, const AccessFileFunc& func)>;
+
+		inline void Init(FileExistsFunc&& fileExistsFunc, TraverseFilesFunc&& traverseFilesFunc);
+		inline bool Add(const tchar* rsp, const tchar* workDir, CreateFileFunc&& func, DevirtualizePathFunc&& devirtualizePathFunc, const tchar* app, DependencyCrawlerType type, u32 ruleIndex);
+
+		Logger& m_logger;
+		WorkManager& m_workManager;
+
+		Futex m_codeFilesHandledLock;
+		struct HandledFile { Futex lock; bool handled = false; };
+		UnorderedMap<StringKey, HandledFile> m_codeFilesHandled;
+
+		Futex m_pchLookupLock;
+		struct Pch { Futex lock; bool handled = false; UnorderedSet<StringKey> files; };
+		UnorderedMap<StringKey, Pch> m_pchLookup;
+
+		struct IncludeRoot
+		{
+			TString path;
+			BloomFilter bloomFilter;
+		};
+
+		struct Instance
+		{
+			DependencyCrawlerType type;
+			TString application;
+			TString rsp;
+			TString workDir;
+			TString platform;
+			TString compiledPlatform;
+			TString overriddenPlatformName;
+			TString builtinIncludesDir;
+			bool platformIsExtension = false;
+			bool usePch = false;
+			CreateFileFunc createFileFunc;
+			DevirtualizePathFunc devirtualizePathFunc;
+			Vector<IncludeRoot> includeRoots;
+			Pch* pch = nullptr;
+
+			Atomic<u32> refCount;
+		};
+
+		struct InstanceRef
+		{
+			InstanceRef(Instance& i) : instance(i) { ++instance.refCount; }
+			InstanceRef(const InstanceRef& i) : instance(i.instance) { ++instance.refCount; }
+			~InstanceRef() { if (!--instance.refCount) delete &instance; }
+			Instance& instance;
+		};
+
+		Futex m_builtIncludesHandledLock;
+		UnorderedSet<u32> m_builtIncludesHandled;
+		FileExistsFunc m_fileExistsFunc;
+		TraverseFilesFunc m_traverseFilesFunc;
+
+		struct CodeFile { TString path; bool hasPch = false; };
+		using CodeFiles = List<CodeFile>;
+
+		inline bool ParseRsp(Instance& instance, const tchar* rsp, CodeFiles& outCodeFiles);
+		inline bool ParseRsp2(Instance& instance, const void* data, u64 dataSize, const tchar* rsp, CodeFiles& outCodeFiles);
+		inline bool ParsePch(Pch& pch, Instance& instance, const void* data, u64 dataSize);
+		inline bool ParseCodeFile(Instance& instance, const StringView& codeFile, bool parseDefines);
+		inline bool ParseCodeFile2(Instance& instance, const void* data, u64 dataSize, const StringView& codeFile, bool parseDefines);
+		inline bool HandleInclude(Instance& instance, const StringView& rootPath, const StringView& include, const StringView& codeFile, bool parseDefines);
+	};
+
+
+	DependencyCrawler::DependencyCrawler(Logger& logger, WorkManager& workManager)
+	:	m_logger(logger)
+	,	m_workManager(workManager)
+	{
+	}
+
+	void DependencyCrawler::Init(FileExistsFunc&& fileExistsFunc, TraverseFilesFunc&& traverseFilesFunc)
+	{
+		m_fileExistsFunc = std::move(fileExistsFunc);
+		m_traverseFilesFunc = std::move(traverseFilesFunc);
+	}
+
+	bool DependencyCrawler::Add(const tchar* rsp, const tchar* workDir, CreateFileFunc&& createFileFunc, DevirtualizePathFunc&& devirtualizePathFunc, const tchar* app, DependencyCrawlerType type, u32 ruleIndex)
+	{
+		auto& instance = *new Instance();
+
+		instance.type = type;
+		instance.application = app;
+		instance.rsp = rsp;
+		instance.workDir = workDir;
+		UBA_ASSERT(*workDir);
+		if (CaseInsensitiveFs)
+			ToLower(instance.workDir.data());
+		if (instance.workDir.back() != PathSeparator)
+			instance.workDir += PathSeparator;
+		instance.createFileFunc = std::move(createFileFunc);
+		instance.devirtualizePathFunc = std::move(devirtualizePathFunc);
+
+		m_workManager.AddWork([this, ref = InstanceRef(instance), ruleIndex](const WorkContext& context)
+			{
+				auto& instance = ref.instance;
+				CodeFiles codeFiles;
+				if (!ParseRsp(instance, instance.rsp.c_str(), codeFiles))
+					return;
+
+				for (auto& cf : codeFiles)
+				{
+					if (instance.usePch && cf.hasPch)
+					{
+						// Note, clang is always reading all the individual files that has been stored in the pch to validate content..
+						// So this code path is not implemented for clang.. if -fno-validate-ast-input-files-content is added we can implement this path
+						StringBuffer<> fixedFile;
+						FixPath(cf.path.c_str(), instance.workDir.data(), instance.workDir.size(), fixedFile);
+
+						StringKey pchKey = ToStringKeyLower(fixedFile);
+
+						SCOPED_FUTEX(m_pchLookupLock, lock);
+						Pch& pch = m_pchLookup[pchKey];
+						lock.Leave();
+
+						SCOPED_FUTEX(pch.lock, lock2);
+						if (!pch.handled)
+						{
+							fixedFile.Append(TC(".dep.json"));
+							instance.createFileFunc(fixedFile, [&](const void* data, u64 dataSize)
+							{
+								return ParsePch(pch, instance, data, dataSize);
+							});
+							pch.handled = true;
+						}
+						instance.pch = &pch;
+					}
+					else
+					{
+						ParseCodeFile(instance, cf.path, instance.platform.empty());
+						if (instance.platform.empty())
+						{
+							if (!instance.overriddenPlatformName.empty())
+								instance.platform = instance.overriddenPlatformName;
+							else
+								instance.platform = instance.compiledPlatform;
+						}
+					}
+				}
+
+				if (!instance.builtinIncludesDir.empty())
+				{
+					SCOPED_FUTEX(m_builtIncludesHandledLock, lock);
+					bool shouldHandle = m_builtIncludesHandled.insert(ruleIndex).second;
+					lock.Leave();
+					if (shouldHandle)
+					{
+						StringBuffer<> fixedPath;
+						FixPath(instance.builtinIncludesDir.c_str(), instance.workDir.data(), instance.workDir.size(), fixedPath);
+						if (!instance.devirtualizePathFunc(fixedPath))
+							m_logger.Info(TC("Failed to devirtualize path %s found in builtin includes"), fixedPath.data);
+						if (CaseInsensitiveFs)
+							fixedPath.MakeLower();
+						fixedPath.Append(PathSeparator);
+						StringBuffer<> codeFilePath(fixedPath);
+						m_traverseFilesFunc(fixedPath, [&](const StringView& file, bool isDirectory)
+							{
+								if (isDirectory)
+									return;
+								codeFilePath.Resize(fixedPath.count).Append(file);
+								m_workManager.AddWork([this, ref, codeFile = TString(codeFilePath.data)](const WorkContext& context)
+									{
+										ref.instance.createFileFunc(codeFile, [&](const void* data, u64 dataSize) { return true; });
+									}, 1, TC("Crawler"), ColorWork);
+							});
+					}
+				}
+
+			}, 1, TC("Crawler"), ColorWork);
+		return true;
+	}
+
+	bool DependencyCrawler::ParseRsp(Instance& instance, const tchar* rsp, CodeFiles& outCodeFiles)
+	{
+		StringBuffer<> fixedPath;
+		FixPath(rsp, instance.workDir.data(), instance.workDir.size(), fixedPath);
+		if (!instance.devirtualizePathFunc(fixedPath))
+			m_logger.Info(TC("Failed to devirtualize path %s in %s"), fixedPath.data, rsp);
+		if (CaseInsensitiveFs)
+			fixedPath.MakeLower();
+		return instance.createFileFunc(fixedPath, [&](const void* data, u64 dataSize)
+			{
+				return ParseRsp2(instance, data, dataSize, rsp, outCodeFiles);
+			});
+	}
+
+	bool DependencyCrawler::ParseRsp2(Instance& instance, const void* data, u64 dataSize, const tchar* rsp, CodeFiles& outCodeFiles)
+	{
+		auto AddCodeFile = [&](const StringView& file, bool pushFront, bool hasPch)
+			{
+				UBA_CRAWLER_DEBUG_PRINT("CODEFILE", file.data);
+				StringBuffer<> fixedPath2;
+				FixPath(file.data, instance.workDir.data(), instance.workDir.size(), fixedPath2);
+				if (!instance.devirtualizePathFunc(fixedPath2))
+					m_logger.Info(TC("Failed to devirtualize path %s in %s"), fixedPath2.data, rsp);
+
+				if (CaseInsensitiveFs)
+					fixedPath2.MakeLower();
+				if (pushFront)
+					outCodeFiles.push_front({TString(fixedPath2.data), hasPch});
+				else
+					outCodeFiles.push_back({TString(fixedPath2.data), hasPch});
+			};
+
+		auto AddRoot = [&](const StringView& path)
+			{
+				UBA_CRAWLER_DEBUG_PRINT("ROOT", path.data);
+				StringBuffer<> fixedPath2;
+				FixPath(path.data, instance.workDir.data(), instance.workDir.size(), fixedPath2);
+				if (!instance.devirtualizePathFunc(fixedPath2))
+					m_logger.Info(TC("Failed to devirtualize path %s in %s"), fixedPath2.data, rsp);
+				if (CaseInsensitiveFs)
+					fixedPath2.MakeLower();
+				fixedPath2.EnsureEndsWithSlash();
+
+				BloomFilter bloomFilter;
+				m_traverseFilesFunc(fixedPath2, [&](const StringView& file, bool isDirectory)
+					{
+						bloomFilter.Add(ToStringKey(file));
+					});
+				if (!bloomFilter.IsEmpty())
+					instance.includeRoots.push_back({TString(fixedPath2.data), bloomFilter});
+			};
+
+		auto IgnoreOption = [](const StringView& path) { UBA_CRAWLER_DEBUG_PRINT("IGNORED", path.data); };
+
+		if (instance.type == DependencyCrawlerType_MsvcCompiler)
+		{
+			struct MsvcOptionWithArg { const StringView name; bool addRoot; };
+			const MsvcOptionWithArg msvcOptionsWithArg[] = 
+			{
+				{ TCV("/I"), true },
+				{ TCV("/external:I"), true },
+				{ TCV("/imsvc"), true },
+				{ TCV("/experimental:log"), false },
+				{ TCV("/sourceDependencies"), false},
+			};
+
+			constexpr const tchar* clangClOptionsWithArg[] = 
+			{
+				TC("-D"),
+				TC("-x"),
+				TC("-o"),
+				TC("-include"),
+				TC("-include-pch"),
+				TC("-vctoolsdir"),
+				TC("-Xclang"),
+				TC("-target"),
+				TC("-arch"),
+			};
+
+			StringBuffer<> prevArg;
+			bool addRoot = false;
+			bool handled = false;
+			ParseArguments((const char*)data, dataSize, [&](const char* arg, u32 argLen)
+				{
+					StringBuffer<> sb;
+					sb.Append(arg, argLen);
+
+					if (handled)
+					{
+						handled = false;
+						if (addRoot)
+							AddRoot(sb);
+						else
+							IgnoreOption(sb);
+						addRoot = false;
+						return;
+					}
+					else if (sb[0] == '/')
+					{
+						for (auto& option : msvcOptionsWithArg)
+						{
+							if (!option.name.Equals(sb))
+								continue;
+							addRoot = option.addRoot;
+							handled = true;
+							return;
+						}
+
+						if (sb.StartsWith(TC("/FI")))
+						{
+							// Check if there is a precompiled header deps file
+							StringBuffer<> fixedFile;
+							FixPath(sb.data + 3, instance.workDir.data(), instance.workDir.size(), fixedFile);
+							if (!instance.devirtualizePathFunc(fixedFile))
+								m_logger.Info(TC("Failed to devirtualize path %s in %s"), fixedFile.data, rsp);
+
+							if (CaseInsensitiveFs)
+								fixedFile.MakeLower();
+
+							StringKey key = ToStringKey(fixedFile);
+							SCOPED_FUTEX(m_codeFilesHandledLock, lock);
+							auto insres = m_codeFilesHandled.try_emplace(key);
+							lock.Leave();
+
+							HandledFile& handledFile = insres.first->second;
+							SCOPED_FUTEX(handledFile.lock, fileLock);
+							if (!handledFile.handled)
+							{
+								handledFile.handled = true;
+
+								fixedFile.Append(TC(".pch"));
+
+								u32 attributes = 0;
+								bool hasPch = m_fileExistsFunc(fixedFile, attributes);
+								AddCodeFile(StringView(sb).Skip(3), true, hasPch);
+							}
+
+							prevArg.Clear();
+							return;
+						}
+						else if (sb.StartsWith(TC("/Yu")))
+						{
+							instance.usePch = true;
+							prevArg.Clear();
+							return;
+						}
+						else
+						{
+							IgnoreOption(sb);
+						}
+						prevArg = sb;
+					}
+					else if (sb[0] == '-')
+					{
+						for (auto option : clangClOptionsWithArg)
+							handled |= sb.Equals(option);
+						if (sb.StartsWith(TC("-resource-dir"), false))
+						{
+							instance.builtinIncludesDir = sb.data + 14;
+							instance.builtinIncludesDir += TC("/include");
+							UBA_CRAWLER_DEBUG_PRINT("RESOURCEDIR", sb.data + 14);
+						}
+						else
+							IgnoreOption(sb);
+					}
+					else if (sb[0] == '@')
+					{
+						ParseRsp(instance, sb.data + 1, outCodeFiles);
+					}
+					else
+					{
+						AddCodeFile(sb, false, false);
+					}
+				});
+		}
+		else if (instance.type == DependencyCrawlerType_ClangCompiler)
+		{
+			enum ArgType { ArgType_None, ArgType_Ignore, ArgType_RootPath, ArgType_ResourceDir, ArgType_Code };
+			struct OptionWithArg { const StringView name; ArgType type; };
+			constexpr const OptionWithArg dashOptionsWithArg[] = 
+			{
+				{ TCV("-D"), ArgType_Ignore },
+				{ TCV("-x"), ArgType_Ignore },
+				{ TCV("-o"), ArgType_Ignore },
+				{ TCV("-include"), ArgType_Code },
+				{ TCV("-include-pch"), ArgType_Ignore },
+				{ TCV("-vctoolsdir"), ArgType_Ignore },
+				{ TCV("-Xclang"), ArgType_Ignore },
+				{ TCV("-target"), ArgType_Ignore },
+				{ TCV("-arch"), ArgType_Ignore },
+				{ TCV("--sysroot"), ArgType_RootPath},
+				{ TCV("-isystem"), ArgType_RootPath },
+				{ TCV("-isysroot"), ArgType_RootPath },
+				{ TCV("-internal-isystem"), ArgType_RootPath },
+				{ TCV("-I"), ArgType_RootPath },
+				{ TCV("-resource-dir"), ArgType_ResourceDir },
+				{ TCV("-dependency-file"), ArgType_Ignore },
+				{ TCV("-internal-externc-isystem"), ArgType_Ignore },
+				{ TCV("-MT"), ArgType_Ignore },
+			};
+
+			//StringBuffer<> prevArg;
+			ArgType type = ArgType_None;
+			ParseArguments((const char*)data, dataSize, [&](const char* arg, u32 argLen)
+				{
+					StringBuffer<> sb;
+					sb.Append(arg, argLen);
+
+					if (type == ArgType_None)
+					{
+						if (sb[0] == '-')
+						{
+							for (auto& option : dashOptionsWithArg)
+							{
+								if (!sb.Equals(option.name, false))
+									continue;
+								//prevArg = sb; // Arg comes in next parse call
+								type = option.type;
+								return;
+							}
+
+							if (auto equalsPos = TStrchr(sb.data, '='))
+							{
+								u32 equalsOffset = u32(equalsPos - sb.data);
+								StringView option2(sb.data, equalsOffset);
+								for (auto option : dashOptionsWithArg)
+								{
+									if (!option2.Equals(option.name, false))
+										continue;
+									++equalsOffset;
+									sb.Clear().Append(arg + equalsOffset, argLen - equalsOffset);
+									type = option.type;
+									//prevArg = option;
+									break;
+								}
+							}
+							else if (sb.StartsWith(TC("-I"), false))
+							{
+								type = ArgType_RootPath;
+								sb.Clear().Append(arg + 2, argLen - 2);
+							}
+							else if (sb.StartsWith(TC("-isystem"), false))
+							{
+								type = ArgType_RootPath;
+								sb.Clear().Append(arg + 8, argLen - 8);
+							}
+							if (type == ArgType_None)
+							{
+								IgnoreOption(sb);
+								return;
+							}
+						}
+					}
+
+					if (type != ArgType_None)
+					{
+						switch (type)
+						{
+						case ArgType_RootPath:
+							AddRoot(sb);
+							break;
+						case ArgType_Code:
+							if (!sb.EndsWith(TC(".gch")))
+								AddCodeFile(sb, true, false);
+							break;
+						case ArgType_ResourceDir:
+							instance.builtinIncludesDir = sb.data;
+							instance.builtinIncludesDir += TC("/include");
+							break;
+						default:
+							IgnoreOption(sb);
+						}
+						type = ArgType_None;
+						return;
+					}
+					else if (sb[0] == '@')
+						ParseRsp(instance, sb.data + 1, outCodeFiles);
+					else if (sb.Contains('/'))
+						AddCodeFile(sb, false, false);
+					else
+						IgnoreOption(sb);
+				});
+			}
+				
+		return true;
+	}
+
+	bool DependencyCrawler::ParsePch(Pch& pch, Instance& instance, const void* data, u64 dataSize)
+	{
+		auto AddHandled = [&](const StringView& str)
+			{
+				StringBuffer<> fullPath;
+				FixPath(str.data, nullptr, 0, fullPath);
+				//instance.devirtualizePathFunc(fullPath); // These paths are never virtual
+				if (CaseInsensitiveFs)
+					fullPath.MakeLower();
+				StringKey key = ToStringKey(fullPath);
+				pch.files.insert(key);
+			};
+
+		StringBuffer<> line;
+		char lastChar = 0;
+
+		const char* it = (const char*)data;
+		const char* end = it + dataSize;
+		for (; it != end; lastChar = *it, ++it)
+		{
+			if ((line.count == 0 && (*it == ' ' || *it == '\t')) || *it == '\r')
+				continue;
+			if (*it == '\\' && lastChar == ' ') // Remove the "dir \dir" extra space in clang deps files
+			{
+				line.Resize(line.count - 1);
+				if (line[line.count-1] != ':')
+					AddHandled(line);
+				line.Clear();
+				continue;
+			}
+			
+			if (*it == '\n')
+			{
+				if (line.count > 3 && !line.Contains(TC("\":")))
+				{
+					AddHandled(line);
+				}
+				line.Clear();
+				continue;
+			}
+
+			if (*it == ' ' && lastChar == '\\')
+				--line.count;
+			line.Append(*it);
+		}
+		return true;
+	}
+
+	bool DependencyCrawler::ParseCodeFile(Instance& instance, const StringView& codeFile, bool parseDefines)
+	{
+		return instance.createFileFunc(codeFile, [&](const void* data, u64 dataSize)
+			{
+				return ParseCodeFile2(instance, data, dataSize, codeFile, parseDefines);
+			});
+	}
+
+	bool DependencyCrawler::ParseCodeFile2(Instance& instance, const void* data, u64 dataSize, const StringView& codeFile, bool parseDefines)
+	{
+		const char* it = (const char*)data;
+		const char* end = it + dataSize;
+		bool hasNonSpace = false;
+
+		while (it < end)
+		{
+			if (*it == '#' && !hasNonSpace)
+			{
+				++it;
+				while (*it == ' ' || *it == '\t')
+					++it;
+
+				if (strncmp(it, "include", 7) == 0)
+				{
+					it += 7;
+					while (*it == ' ' || *it == '\t')
+						++it;
+					StringBuffer<> include;
+					bool isQuote = false;
+
+					if (*it == '\"')
+					{
+						++it;
+						const char* includeStart = it;
+						const char* includeEnd = strchr(it, '\"');
+						it = includeEnd + 1;
+						include.Append(includeStart, u32(includeEnd - includeStart));
+						isQuote =true;
+					}
+					else if (*it == '<')
+					{
+						++it;
+						const char* includeStart = it;
+						const char* includeEnd = strchr(it, '>');
+						it = includeEnd + 1;
+						include.Append(includeStart, u32(includeEnd - includeStart));
+					}
+					else
+					{
+						const char* defineBegin = it;
+						const char* defineEnd = nullptr;
+						const char* argBegin = nullptr;
+						const char* argEnd = nullptr;
+
+						while (it < end)
+						{
+							if (*it == '\r' || *it == '\n')
+							{
+								if (!defineEnd)
+									defineEnd = it;
+								break;
+							}
+							else if (*it == ' ' || *it == '\t')
+							{
+								if (!defineEnd)
+									defineEnd = it;
+							}
+							else if (*it == '(')
+							{
+								defineEnd = it;
+								argBegin = it + 1;
+							}
+							else if (*it == ')')
+							{
+								argEnd = it;
+								break;
+							}
+							else if (!argBegin && defineEnd)
+								break;
+							++it;
+						}
+
+						u32 defineLen = u32(defineEnd - defineBegin);
+						if (argBegin)
+						{
+							if (strncmp(defineBegin, "UE_INLINE_GENERATED_CPP_BY_NAME", defineLen) == 0)
+							{
+								include.Append(argBegin, u32(argEnd - argBegin)).Append(TC(".gen.cpp"));
+							}
+							else if (strncmp(defineBegin, "COMPILED_PLATFORM_HEADER", defineLen) == 0)
+							{
+								if (!instance.platform.empty())
+								{
+									if (!instance.platformIsExtension)
+										include.Append(instance.platform).Append(PathSeparator);
+									include.Append(instance.platform).Append(argBegin, u32(argEnd - argBegin));
+								}
+							}
+						}
+						else if (defineBegin)
+						{
+							if (strncmp(defineBegin, "PER_MODULE_INLINE_FILE", defineLen) == 0)
+							{
+								include.Append(TC("HAL/PerModuleInline.inl")); // TODO: This is not correct.. need revisit
+							}
+						}
+					}
+
+					if (!include.IsEmpty())
+					{
+						include.MakeLower().FixPathSeparators();
+						auto work = 
+						[this, ref = InstanceRef(instance), isQuote, include2 = TString(include.data), codeFile2 = TString(codeFile.data), parseDefines](const WorkContext& context)
+							{
+								auto& instance = ref.instance;
+								if (isQuote)
+								{
+									StringBuffer<> localDir;
+									if (const tchar* lastSeparator = TStrrchr(codeFile2.c_str(), PathSeparator))
+										if (HandleInclude(instance, localDir.Append(codeFile2.c_str(), lastSeparator - codeFile2.c_str() + 1), include2, codeFile2, parseDefines))
+											return;
+								}
+
+								StringKey fileKey;
+								if (const tchar* firstSeparator = TStrchr(include2.c_str(), PathSeparator))
+									fileKey = ToStringKey(include2.c_str(), (firstSeparator - include2.c_str()));
+								else
+									fileKey = ToStringKey(include2);
+
+								for (auto& root : instance.includeRoots)
+								{
+									if (root.bloomFilter.IsGuaranteedMiss(fileKey))
+										continue;
+									if (HandleInclude(instance, root.path, include2, codeFile2, parseDefines))
+										return;
+								}
+
+							};
+						if (parseDefines)
+						{
+							TrackWorkScope tws;
+							work({tws});
+						}
+						else
+							m_workManager.AddWork(work, 1, TC("Crawler"), ColorWork);
+					}
+				}
+				else if (parseDefines && strncmp(it, "define", 6) == 0)
+				{
+					it += 6;
+
+					auto parseDefine = [&](const char* define, u32 defineLength, TString& out)
+						{
+							if (strncmp(it, define, defineLength) != 0)
+								return false;
+							it += defineLength;
+							while (it < end && (*it == '\t' || *it == ' '))
+								++it;
+							const char* platformBegin = it;
+							while (it < end && *it != '\t' && *it != ' ' && *it != '\r' && *it != '\n')
+								++it;
+							out.assign(platformBegin, it);
+							return true;
+						};
+
+					while (*it == ' ' || *it == '\t')
+						++it;
+					TString platformIsExtension;
+					if (parseDefine("UBT_COMPILED_PLATFORM", 21, instance.compiledPlatform))
+					{}
+					else if (parseDefine("OVERRIDE_PLATFORM_HEADER_NAME", 29, instance.overriddenPlatformName))
+					{}
+					else if (parseDefine("PLATFORM_IS_EXTENSION", 21, platformIsExtension))
+					{
+						instance.platformIsExtension = platformIsExtension != TC("0");
+					}
+				}
+			}
+			else if (*it == '\n')
+			{
+				hasNonSpace = false;
+			}
+			else if (*it != ' ' && *it != '\t')
+			{
+				hasNonSpace = true;
+			}
+			++it;
+		}
+
+		return true;
+	}
+
+	bool DependencyCrawler::HandleInclude(Instance& instance, const StringView& rootPath, const StringView& include, const StringView& codeFile, bool parseDefines)
+	{
+		if (include.EndsWith(TC(".ush")))
+			return true;
+		StringBuffer<> fullPath;
+		FixPath(include.data, rootPath.data, rootPath.count, fullPath);
+		if (IsAbsolutePath(include.data))
+			if (!instance.devirtualizePathFunc(fullPath))
+				m_logger.Info(TC("Failed to devirtualize include path %s in %s"), fullPath.data, codeFile.data);
+
+		if (CaseInsensitiveFs)
+			fullPath.MakeLower();
+
+		u32 attributes = 0;
+		if (!m_fileExistsFunc(fullPath, attributes))
+			return false;
+		if (IsDirectory(attributes))
+			return false;
+
+		{
+			StringKey key = ToStringKey(fullPath);
+			if (instance.pch)
+			{
+				SCOPED_FUTEX_READ(instance.pch->lock, lock);
+				if (instance.pch->files.find(key) != instance.pch->files.end())
+					return true;
+			}
+			SCOPED_FUTEX(m_codeFilesHandledLock, lock);
+			if (!m_codeFilesHandled.try_emplace(key).second)
+				return true;
+		}
+
+		ParseCodeFile(instance, fullPath, parseDefines);
+		return true;
+	}
+}
diff --git a/Engine/Source/Programs/UnrealBuildAccelerator/Common/Public/UbaIncludeCrawler.h b/Engine/Source/Programs/UnrealBuildAccelerator/Common/Public/UbaIncludeCrawler.h
deleted file mode 100644
diff --git a/Engine/Source/Programs/UnrealBuildAccelerator/Common/Public/UbaSessionClient.h b/Engine/Source/Programs/UnrealBuildAccelerator/Common/Public/UbaSessionClient.h
--- a/Engine/Source/Programs/UnrealBuildAccelerator/Common/Public/UbaSessionClient.h
+++ b/Engine/Source/Programs/UnrealBuildAccelerator/Common/Public/UbaSessionClient.h
@@ -2,7 +2,7 @@
 
 #pragma once
 
-#include "UbaIncludeCrawler.h"
+#include "UbaDependencyCrawler.h"
 #include "UbaNetwork.h"
 #include "UbaSession.h"
 
diff --git a/Engine/Source/Programs/UnrealBuildAccelerator/Test/Private/UbaTestDependencyCrawler.cpp b/Engine/Source/Programs/UnrealBuildAccelerator/Test/Private/UbaTestDependencyCrawler.cpp
new file mode 100644
--- /dev/null
+++ b/Engine/Source/Programs/UnrealBuildAccelerator/Test/Private/UbaTestDependencyCrawler.cpp
@@ -0,0 +1,67 @@
+// Copyright Epic Games, Inc. All Rights Reserved.
+
+#include "UbaDependencyCrawler.h"
+#include "UbaDirectoryIterator.h"
+#include "UbaLogger.h"
+
+namespace uba
+{
+	bool TestDependencyCrawler(LoggerWithWriter& logger, const StringBufferBase& rootDir)
+	{
+#if 0 // Commented out for now because of hard coded paths below
+
+		WorkManagerImpl workManager(128);
+
+		DependencyCrawler crawler(logger, workManager);
+
+		auto CreateFileFunc = [&](const StringView& fileName, const DependencyCrawler::AccessFileFunc& func)
+			{
+				FileAccessor fa(logger, fileName.data);
+				if (!fa.OpenMemoryRead())
+					return false;
+				return func(fa.GetData(), fa.GetSize());
+			};
+
+		auto FileExistsFunc = [&](const StringView& fileName, u32& outAttr)
+			{
+				return FileExists(logger, fileName.data, nullptr, &outAttr);
+			};
+
+		auto TraverseFilesFunc = [&](const StringView& path, const DependencyCrawler::FileFunc& fileFunc)
+			{
+				TraverseDir(logger, path, [&](const DirectoryEntry& entry)
+					{
+						StringBuffer<> name;
+						name.Append(entry.name, entry.nameLen);
+						if (CaseInsensitiveFs)
+							name.MakeLower();
+						fileFunc(name, IsDirectory(entry.attributes));
+					});
+
+			};
+
+		crawler.Init(FileExistsFunc, TraverseFilesFunc);
+
+		const tchar app[] = TC("c:\\sdk\\AutoSDK\\HostWin64\\Win64\\LLVM\\18.1.8\\bin\\clang-cl.exe");
+
+		//DependencyCrawlerType type = DependencyCrawlerType_MsvcCompiler;
+		//const tchar rsp[] = TC("e:\\dev\\fn\\Engine\\Intermediate\\Build\\Win64\\x64\\BlankProgram\\Development\\Core\\Module.Core.1.cpp.obj.rsp");
+		
+		DependencyCrawlerType type = DependencyCrawlerType_ClangCompiler;
+		//const tchar rsp[] = TC("e:\\dev\\fn\\Engine\\Intermediate\\Build\\Linux\\x64\\BlankProgram\\Development\\Core\\Module.Core.1.cpp.o.rsp");
+		//const tchar rsp[] = TC("e:\\temp\\response-22b46d.txt");
+		const tchar rsp[] = TC("e:\\temp\\MiMalloc.c.o.rsp");
+		//const tchar rsp[] = TC("e:\\dev\\fn\\engine\\plugins\\experimental\\geometrycollectionplugin\\intermediate\\build\\win64\\x64\\unrealeditor\\development\\geometrycollectionnodes\\module.geometrycollectionnodes.3.cpp.obj.rsp");
+		//const tchar rsp[] = TC("e:\\dev\\fn\\Engine\\Plugins\\Runtime\\MeshModelingToolset\\Intermediate\\Build\\Win64\\x64\\UnrealEditor\\Development\\ModelingComponents\\Module.ModelingComponents.3.cpp.obj.rsp");
+		//const tchar rsp[] = TC("e:\\dev\\fn\\Engine\\Intermediate\\Build\\Win64\\x64\\UnrealEditor\\Development\\AnimGraphRuntime\\Module.AnimGraphRuntime.3.cpp.obj.rsp");
+
+		crawler.Add(rsp, TC("e:\\dev\\fn\\Engine\\Source"), CreateFileFunc, [](StringBufferBase&)
+			{
+				return true;
+			}, app, type, 10);
+
+		workManager.FlushWork();
+#endif
+		return true;
+	}
+}
\ No newline at end of file
diff --git a/Engine/Source/Programs/UnrealBuildAccelerator/Test/Private/UbaTestIncludeCrawler.cpp b/Engine/Source/Programs/UnrealBuildAccelerator/Test/Private/UbaTestIncludeCrawler.cpp
deleted file mode 100644
