[UBA]
* Fixed include crawler for clang-cl.exe processes

diff --git a/Engine/Source/Programs/UnrealBuildAccelerator/Core/Private/UbaApplicationRules.cpp b/Engine/Source/Programs/UnrealBuildAccelerator/Core/Private/UbaApplicationRules.cpp
--- a/Engine/Source/Programs/UnrealBuildAccelerator/Core/Private/UbaApplicationRules.cpp
+++ b/Engine/Source/Programs/UnrealBuildAccelerator/Core/Private/UbaApplicationRules.cpp
@@ -323,6 +323,15 @@ namespace uba
 		}
 	};
 
+	class ApplicationRulesClangClExe : public ApplicationRulesClangPlusPlusExe
+	{
+	public:
+		virtual DependencyCrawlerType GetDependencyCrawlerType() const override
+		{
+			return DependencyCrawlerType_MsvcCompiler;
+		}
+	};
+
 	class ApplicationRulesLdLLdExe : public ApplicationRulesClang
 	{
 		using Super = ApplicationRulesClang;
@@ -637,7 +646,7 @@ namespace uba
 			{ ApplicationHash(TC("rc.exe")),					new ApplicationRulesVcLink() },
 			{ ApplicationHash(TC("lld-link.exe")),				new ApplicationRulesLldLinkExe() },
 			{ ApplicationHash(TC("clang++.exe")),				new ApplicationRulesClangPlusPlusExe() },
-			{ ApplicationHash(TC("clang-cl.exe")),				new ApplicationRulesClangPlusPlusExe() },
+			{ ApplicationHash(TC("clang-cl.exe")),				new ApplicationRulesClangClExe() },
 			{ ApplicationHash(TC("verse-clang-cl.exe")),		new ApplicationRulesClangPlusPlusExe() },
 			{ ApplicationHash(TC("ispc.exe")),					new ApplicationRulesISPCExe() },
 			{ ApplicationHash(TC("radlink.exe")),				new ApplicationRulesRadLinkExe() },
