[UBA]
* Detoured NtQueryFullAttributesFile since it is now used directly by uba itself and since we detour UbaObjTool we also need to detour it
* Fixed remote linking for win64 clang on wine machines

diff --git a/Engine/Source/Programs/UnrealBuildAccelerator/Common/Private/UbaFile.cpp b/Engine/Source/Programs/UnrealBuildAccelerator/Common/Private/UbaFile.cpp
--- a/Engine/Source/Programs/UnrealBuildAccelerator/Common/Private/UbaFile.cpp
+++ b/Engine/Source/Programs/UnrealBuildAccelerator/Common/Private/UbaFile.cpp
@@ -191,7 +191,7 @@ namespace uba
 			InitializeObjectAttributes(&oa, &us, OBJ_CASE_INSENSITIVE, NULL, NULL);
 			NTSTATUS res = NtQueryFullAttributesFile(&oa, &info);
 			if (res != STATUS_SUCCESS)
-				return logger.Error(TC("GetFileBasicInformationByHandle (NtQueryFullAttributesFile) failed on %s (0x%11x)"), fileName, res);
+				return logger.Error(TC("GetFileBasicInformationByHandle (NtQueryFullAttributesFile) failed on %s (0x%x)"), fileName, res);
 		}
 		else
 		{
@@ -199,7 +199,7 @@ namespace uba
 			IO_STATUS_BLOCK b;
 			NTSTATUS res = NtQueryInformationFile(asHANDLE(hFile), &b, &info, sizeof(info), FileNetworkOpenInformation);
 			if (res != STATUS_SUCCESS)
-				return logger.Error(TC("GetFileBasicInformationByHandle (NtQueryInformationFile) failed on %s (0x%11x)"), fileName, res);
+				return logger.Error(TC("GetFileBasicInformationByHandle (NtQueryInformationFile) failed on %s (0x%x)"), fileName, res);
 		}
 		out.attributes = info.FileAttributes;
 		out.lastWriteTime = info.LastWriteTime.QuadPart;
diff --git a/Engine/Source/Programs/UnrealBuildAccelerator/Detours/Private/PosixOS/UbaDetoursFunctionsPosix.cpp b/Engine/Source/Programs/UnrealBuildAccelerator/Detours/Private/PosixOS/UbaDetoursFunctionsPosix.cpp
--- a/Engine/Source/Programs/UnrealBuildAccelerator/Detours/Private/PosixOS/UbaDetoursFunctionsPosix.cpp
+++ b/Engine/Source/Programs/UnrealBuildAccelerator/Detours/Private/PosixOS/UbaDetoursFunctionsPosix.cpp
@@ -295,6 +295,9 @@ namespace uba
 	{
 		return CaseInsensitiveFs ? ToStringKeyLower(b) : ToStringKey(b);
 	}
+
+	bool CouldBeCompressedFile(const StringView& fileName) { return false; }
+
 }
 
 bool CanDetour2(const StringView& file)
diff --git a/Engine/Source/Programs/UnrealBuildAccelerator/Detours/Private/UbaDetoursShared.cpp b/Engine/Source/Programs/UnrealBuildAccelerator/Detours/Private/UbaDetoursShared.cpp
--- a/Engine/Source/Programs/UnrealBuildAccelerator/Detours/Private/UbaDetoursShared.cpp
+++ b/Engine/Source/Programs/UnrealBuildAccelerator/Detours/Private/UbaDetoursShared.cpp
@@ -407,6 +407,7 @@ namespace uba
 	{
 		StringBuffer<MaxPath> fileNameForKey;
 		fileNameForKey.Append(fileName);
+		StringView fileNameView(fileName, fileNameForKey.count);
 		if (CaseInsensitiveFs)
 			fileNameForKey.MakeLower();
 
@@ -416,7 +417,7 @@ namespace uba
 
 		memset(&outAttr.data, 0, sizeof(outAttr.data));
 
-		bool keepInMemory = KeepInMemory(StringView(fileName, fileNameForKey.count), false);
+		bool keepInMemory = KeepInMemory(fileNameView, false);
 		if (keepInMemory)
 		{
 			SCOPED_READ_LOCK(g_mappedFileTable.m_lookupLock, lock);
@@ -519,6 +520,17 @@ namespace uba
 
 					if (info.attributes)
 					{
+						u64 fileSize = info.size;
+
+						// Could be compressed and then directory table size is wrong
+						if (CouldBeCompressedFile(fileNameView))
+						{
+							SCOPED_READ_LOCK(g_mappedFileTable.m_lookupLock, lock);
+							auto findIt = g_mappedFileTable.m_lookup.find(fileNameKey);
+							if (findIt != g_mappedFileTable.m_lookup.end() && !findIt->second.deleted)
+								fileSize = findIt->second.size;
+						}
+
 						outAttr.useCache = true;
 						outAttr.exists = true;
 						outAttr.lastError = ErrorSuccess;
@@ -528,7 +540,7 @@ namespace uba
 						outAttr.volumeSerial = info.volumeSerial;
 
 #if PLATFORM_WINDOWS
-						LARGE_INTEGER li = ToLargeInteger(info.size);
+						LARGE_INTEGER li = ToLargeInteger(fileSize);
 						outAttr.data.dwFileAttributes = info.attributes;
 						outAttr.data.nFileSizeLow = li.LowPart;
 						outAttr.data.nFileSizeHigh = li.HighPart;
@@ -540,7 +552,7 @@ namespace uba
 						outAttr.data.st_mode = (mode_t)info.attributes;
 						outAttr.data.st_dev = info.volumeSerial;
 						outAttr.data.st_ino = info.fileIndex;
-						outAttr.data.st_size = info.size;
+						outAttr.data.st_size = fileSize;
 #endif
 					}
 					else
diff --git a/Engine/Source/Programs/UnrealBuildAccelerator/Detours/Private/UbaDetoursShared.h b/Engine/Source/Programs/UnrealBuildAccelerator/Detours/Private/UbaDetoursShared.h
--- a/Engine/Source/Programs/UnrealBuildAccelerator/Detours/Private/UbaDetoursShared.h
+++ b/Engine/Source/Programs/UnrealBuildAccelerator/Detours/Private/UbaDetoursShared.h
@@ -123,6 +123,7 @@ namespace uba
 
 	inline bool CanDetour(const tchar* file) { return !t_disallowDetour && g_rules->CanDetour(file, g_runningRemote); }
 	inline bool KeepInMemory(const StringView& fileName, bool isWrite) { return g_allowKeepFilesInMemory && g_rules->KeepInMemory(fileName, g_systemTemp.data, g_runningRemote, isWrite); }
+	bool CouldBeCompressedFile(const StringView& fileName);
 
 	void Shared_WriteConsole(const char* chars, u32 charCount, bool isError);
 	void Shared_WriteConsole(const wchar_t* chars, u32 charCount, bool isError);
diff --git a/Engine/Source/Programs/UnrealBuildAccelerator/Detours/Private/Windows/UbaDetoursFunctionsNtDll.inl b/Engine/Source/Programs/UnrealBuildAccelerator/Detours/Private/Windows/UbaDetoursFunctionsNtDll.inl
--- a/Engine/Source/Programs/UnrealBuildAccelerator/Detours/Private/Windows/UbaDetoursFunctionsNtDll.inl
+++ b/Engine/Source/Programs/UnrealBuildAccelerator/Detours/Private/Windows/UbaDetoursFunctionsNtDll.inl
@@ -1,5 +1,16 @@
 // Copyright Epic Games, Inc. All Rights Reserved.
 
+struct FILE_NETWORK_OPEN_INFORMATION
+{
+	LARGE_INTEGER CreationTime;
+	LARGE_INTEGER LastAccessTime;
+	LARGE_INTEGER LastWriteTime;
+	LARGE_INTEGER ChangeTime;
+	LARGE_INTEGER AllocationSize;
+	LARGE_INTEGER EndOfFile;
+	ULONG FileAttributes;
+};
+
 bool IsContentWrite(u32 desiredAccess, u32 createDisposition)
 {
 	if (desiredAccess & (FILE_WRITE_DATA | FILE_APPEND_DATA | GENERIC_WRITE))
@@ -213,17 +224,6 @@ NTSTATUS Detoured_NtQueryInformationFile(HANDLE FileHandle, PIO_STATUS_BLOCK IoS
 			}
 			if (FileInformationClass == 34) // FileNetworkOpenInformation
 			{
-				struct FILE_NETWORK_OPEN_INFORMATION
-				{
-					LARGE_INTEGER CreationTime;
-					LARGE_INTEGER LastAccessTime;
-					LARGE_INTEGER LastWriteTime;
-					LARGE_INTEGER ChangeTime;
-					LARGE_INTEGER AllocationSize;
-					LARGE_INTEGER EndOfFile;
-					ULONG FileAttributes;
-				};
-
 				UBA_ASSERT(dh.dirTableOffset != ~0u);
 				DirectoryTable::EntryInformation entryInfo;
 				g_directoryTable.GetEntryInformation(entryInfo, dh.dirTableOffset);
@@ -236,7 +236,10 @@ NTSTATUS Detoured_NtQueryInformationFile(HANDLE FileHandle, PIO_STATUS_BLOCK IoS
 				info.LastWriteTime.QuadPart = entryInfo.lastWrite;
 				info.ChangeTime.QuadPart = entryInfo.lastWrite;
 				info.AllocationSize.QuadPart = entryInfo.size;
-				info.EndOfFile.QuadPart = entryInfo.size;
+				u64 fileSize = dh.fileObject->fileInfo->size;
+				if (fileSize == InvalidValue)
+					fileSize = entryInfo.size;
+				info.EndOfFile.QuadPart = fileSize;
 				info.FileAttributes = entryInfo.attributes;
 				return STATUS_SUCCESS;
 			}
@@ -427,6 +430,65 @@ NTSTATUS NTAPI Detoured_NtQueryDirectoryFile(HANDLE FileHandle, HANDLE Event, PI
 	return res;
 }
 
+NTSTATUS Detoured_NtQueryFullAttributesFile(POBJECT_ATTRIBUTES ObjectAttributes, PVOID Attributes)
+{
+	auto fileName = (const wchar_t*)ObjectAttributes->ObjectName->Buffer;
+	u32 fileNameLen = ObjectAttributes->ObjectName->Length/sizeof(tchar);
+	UBA_ASSERT(fileName[fileNameLen] == 0);
+
+	if (!fileName || !StartsWith(fileName, L"\\??\\") || !CanDetour(fileName + 4) || wcschr(fileName, ',') != nullptr || Contains(fileName, L"::")) // Some weird .net path used by dotnet.exe ... ignore for now!
+	{
+		TimerScope ts(g_kernelStats.getFileInfo);
+		auto res = True_NtQueryFullAttributesFile(ObjectAttributes, Attributes);
+		DEBUG_LOG_TRUE(L"NtQueryFullAttributesFile", L"(%.*s) -> %s", fileNameLen, fileName, ToString(res).data);
+		return res;
+	}
+
+
+	StringBuffer<MaxPath> fixedName;
+	FixPath(fixedName, fileName + 4);
+	/*
+	if (fixedName.StartsWith(g_systemTemp.data))
+	{
+		auto res = True_NtQueryFullAttributesFile(ObjectAttributes, Attributes);
+		DEBUG_LOG_TRUE(L"NtQueryFullAttributesFile", L"(%.*s) -> %s", fileNameLen, fileName, ToString(res).data);
+		return res;
+	}
+	*/
+
+	DevirtualizePath(fixedName);
+
+	FileAttributes attr;
+	Shared_GetFileAttributes(attr, fixedName.data);
+
+	if (!attr.useCache)
+	{
+		auto res = True_NtQueryFullAttributesFile(ObjectAttributes, Attributes);
+		DEBUG_LOG_TRUE(L"NtQueryFullAttributesFile", L"(%.*s) -> %s", fileNameLen, fileName, ToString(res).data);
+		return res;
+	}
+
+	UBA_ASSERT(!ObjectAttributes->RootDirectory);
+
+	NTSTATUS res = STATUS_OBJECT_NAME_NOT_FOUND;
+	if (attr.exists && attr.lastError == ErrorSuccess)
+	{
+		WIN32_FILE_ATTRIBUTE_DATA& data = attr.data;
+		res = STATUS_SUCCESS;
+		auto& info = *(FILE_NETWORK_OPEN_INFORMATION*)Attributes;;
+		info.CreationTime = ToLargeInteger(data.ftCreationTime.dwHighDateTime, data.ftCreationTime.dwLowDateTime);
+		info.LastAccessTime = ToLargeInteger(data.ftLastAccessTime.dwHighDateTime, data.ftLastAccessTime.dwLowDateTime);
+		info.LastWriteTime = ToLargeInteger(data.ftLastWriteTime.dwHighDateTime, data.ftLastWriteTime.dwLowDateTime);
+		info.ChangeTime = ToLargeInteger(data.ftLastWriteTime.dwHighDateTime, data.ftLastWriteTime.dwLowDateTime);
+		info.AllocationSize = ToLargeInteger(data.nFileSizeHigh, data.nFileSizeLow);
+		info.EndOfFile = info.AllocationSize;
+		info.FileAttributes = data.dwFileAttributes;
+	}
+
+	DEBUG_LOG_DETOURED(L"NtQueryFullAttributesFile", L"(%.*s) -> %s", fileNameLen, fileName, ToString(res).data);
+	return res;
+}
+
 NTSTATUS Detoured_NtSetInformationFile(HANDLE FileHandle, PIO_STATUS_BLOCK IoStatusBlock, PVOID FileInformation, ULONG Length, FILE_INFORMATION_CLASS FileInformationClass)
 {
 	DETOURED_CALL(NtSetInformationFile);
@@ -1663,14 +1725,6 @@ NTSTATUS NTAPI Detoured_NtCreateIoCompletion(PHANDLE IoCompletionHandle, ACCESS_
 	return True_NtCreateIoCompletion(IoCompletionHandle, DesiredAccess, ObjectAttributes, Count);
 }
 
-NTSTATUS Detoured_NtQueryFullAttributesFile(POBJECT_ATTRIBUTES ObjectAttributes, PVOID Attributes)
-{
-	DETOURED_CALL(NtQueryFullAttributesFile);
-	NTSTATUS res = True_NtQueryFullAttributesFile(ObjectAttributes, Attributes);
-	DEBUG_LOG_TRUE(L"NtQueryFullAttributesFile", L"(%ls) -> %ls", ObjectAttributes->ObjectName->Buffer, ToString(res).data);
-	return res;
-}
-
 NTSTATUS NTAPI Detoured_NtFlushBuffersFileEx(HANDLE FileHandle, ULONG Flags, PVOID Parameters, ULONG ParametersSize, PIO_STATUS_BLOCK IoStatusBlock)
 {
 	DETOURED_CALL(NtFlushBuffersFileEx);
diff --git a/Engine/Source/Programs/UnrealBuildAccelerator/Detours/Private/Windows/UbaDetoursFunctionsWin.cpp b/Engine/Source/Programs/UnrealBuildAccelerator/Detours/Private/Windows/UbaDetoursFunctionsWin.cpp
--- a/Engine/Source/Programs/UnrealBuildAccelerator/Detours/Private/Windows/UbaDetoursFunctionsWin.cpp
+++ b/Engine/Source/Programs/UnrealBuildAccelerator/Detours/Private/Windows/UbaDetoursFunctionsWin.cpp
@@ -230,6 +230,11 @@ UnorderedMap<HMODULE, TString>& g_loadedModules = *new UnorderedMap<HMODULE, TSt
 u64 g_memoryFileIndexCounter = ~u64(0) - 1000000; // I really hope this will not collide with anything
 
 FileDecompressor& g_fileDecompressor = *new FileDecompressor(); // Leak for shutdown
+bool CouldBeCompressedFile(const StringView& fileName)
+{
+	return g_rules->CanDependOnCompressedFiles() && g_globalRules.FileCanBeCompressed(fileName);
+	//return g_fileDecompressor.CouldBeCompressedFile(fileName);
+}
 
 
 struct SuppressCreateFileDetourScope
