[UBA]
* Another fix for remote linking in clang toolchain when having compressed obj files

diff --git a/Engine/Source/Programs/UnrealBuildAccelerator/Detours/Private/Windows/UbaDetoursFunctionsNtDll.inl b/Engine/Source/Programs/UnrealBuildAccelerator/Detours/Private/Windows/UbaDetoursFunctionsNtDll.inl
--- a/Engine/Source/Programs/UnrealBuildAccelerator/Detours/Private/Windows/UbaDetoursFunctionsNtDll.inl
+++ b/Engine/Source/Programs/UnrealBuildAccelerator/Detours/Private/Windows/UbaDetoursFunctionsNtDll.inl
@@ -953,7 +953,7 @@ NTSTATUS NTAPI Shared_NtCreateFile(bool IsCreateFunc, PHANDLE hFileHandle, ACCES
 
 				bool isWriteAttributes = (DesiredAccess & FILE_WRITE_ATTRIBUTES) != 0;
 
-				if (allowEarlyOut && !useContent && !isWriteAttributes && !wasPreload && !g_fileDecompressor.CouldBeCompressedFile(fileName))
+				if (allowEarlyOut && !useContent && !isWriteAttributes && !wasPreload && !CouldBeCompressedFile(fileName))
 				{
 					auto dh = new DetouredHandle(HandleType_File);
 					dh->fileObject = new FileObject();
diff --git a/Engine/Source/Programs/UnrealBuildAccelerator/Detours/Private/Windows/UbaDetoursFunctionsWin.cpp b/Engine/Source/Programs/UnrealBuildAccelerator/Detours/Private/Windows/UbaDetoursFunctionsWin.cpp
--- a/Engine/Source/Programs/UnrealBuildAccelerator/Detours/Private/Windows/UbaDetoursFunctionsWin.cpp
+++ b/Engine/Source/Programs/UnrealBuildAccelerator/Detours/Private/Windows/UbaDetoursFunctionsWin.cpp
@@ -232,8 +232,10 @@ u64 g_memoryFileIndexCounter = ~u64(0) - 1000000; // I really hope this will not
 FileDecompressor& g_fileDecompressor = *new FileDecompressor(); // Leak for shutdown
 bool CouldBeCompressedFile(const StringView& fileName)
 {
-	return g_rules->CanDependOnCompressedFiles() && g_globalRules.FileCanBeCompressed(fileName);
-	//return g_fileDecompressor.CouldBeCompressedFile(fileName);
+	if (g_runningRemote)
+		return g_rules->CanDependOnCompressedFiles() && g_globalRules.FileCanBeCompressed(fileName);
+	else
+		return g_fileDecompressor.CouldBeCompressedFile(fileName);
 }
 
 
