[UBA]
* Added descriptions to some work
* Minor cleanup of Session::GetMemoryMapAlignment

diff --git a/Engine/Source/Programs/UnrealBuildAccelerator/Common/Private/UbaSession.cpp b/Engine/Source/Programs/UnrealBuildAccelerator/Common/Private/UbaSession.cpp
--- a/Engine/Source/Programs/UnrealBuildAccelerator/Common/Private/UbaSession.cpp
+++ b/Engine/Source/Programs/UnrealBuildAccelerator/Common/Private/UbaSession.cpp
@@ -2101,8 +2101,6 @@ namespace uba
 
 		if (fileName.EndsWith(TC(".h")) || fileName.EndsWith(TC(".inl")) || fileName.EndsWith(TC(".gch")))
 			return 4 * 1024; // clang seems to need 4k alignment? Is it a coincidence it works or what is happening inside the code? (msvc works with alignment 1byte here)
-		if (fileName.EndsWith(TC(".h.obj")))
-			return 4 * 1024;
 		if (fileName.EndsWith(TC(".lib")))
 			return 4 * 1024;
 
@@ -2113,6 +2111,11 @@ namespace uba
 			if (fileName.EndsWith(TC(".pch")))
 				return 64 * 1024; // pch needs 64k alignment
 		}
+		else
+		{
+			if (fileName.EndsWith(TC(".h.obj")))
+				return 4 * 1024;
+		}
 		return 0;
 	}
 
diff --git a/Engine/Source/Programs/UnrealBuildAccelerator/Common/Private/UbaStorage.cpp b/Engine/Source/Programs/UnrealBuildAccelerator/Common/Private/UbaStorage.cpp
--- a/Engine/Source/Programs/UnrealBuildAccelerator/Common/Private/UbaStorage.cpp
+++ b/Engine/Source/Programs/UnrealBuildAccelerator/Common/Private/UbaStorage.cpp
@@ -1667,7 +1667,7 @@ namespace uba
 						}
 						m_logger.Error(TC("Cas key and content mismatch for key %s (expected %s) (%s ago)"), CasKeyString(casKey).str, CasKeyString(checkedKey).str, timeStr.data);
 						
-					}, 1, TC(""));
+					}, 1, TC("CheckCasContent"));
 			});
 		m_logger.Info(TC("Validating %u entries..."), entryCount);
 
@@ -1744,7 +1744,7 @@ namespace uba
 								m_logger.Error(TC("CasKey mismatch for %s even though size and lastwritten were the same. Corrupt path table! (Correct: %s. Wrong: %s)"), p.data(), CasKeyString(casKey).str, CasKeyString(fe.casKey).str);
 							}
 
-						}, 1, TC(""));
+						}, 1, TC("CheckFileTable"));
 				});
 		}
 		workManager.FlushWork();
@@ -1805,7 +1805,7 @@ namespace uba
 							u32 deleteCountTemp = 0;
 							DeleteAllFiles(m_logger, fullPath.data, true, &deleteCountTemp);
 							atomicDeleteCount += deleteCountTemp;
-						}, 1, TC(""));
+						}, 1, TC("DeleteAllCas"));
 				});
 			workManager.FlushWork();
 			deleteCount += atomicDeleteCount;
