[UBA]
* Added code to fix shipping compile warning in detour module
* Fixed so Move from memory to file make sure to always write all data or return false

diff --git a/Engine/Source/Programs/UnrealBuildAccelerator/Detours/Private/Windows/UbaDetoursFunctionsKernelBase.inl b/Engine/Source/Programs/UnrealBuildAccelerator/Detours/Private/Windows/UbaDetoursFunctionsKernelBase.inl
--- a/Engine/Source/Programs/UnrealBuildAccelerator/Detours/Private/Windows/UbaDetoursFunctionsKernelBase.inl
+++ b/Engine/Source/Programs/UnrealBuildAccelerator/Detours/Private/Windows/UbaDetoursFunctionsKernelBase.inl
@@ -1621,12 +1621,16 @@ bool Shared_MoveFile(LPCWSTR lpExistingFileName, LPCWSTR lpNewFileName, DWORD dw
 		auto cg = MakeGuard([&]() { CloseHandle(h); });
 		UBA_ASSERT(sourceInfo.memoryFile->writtenSize < ~0u);
 		u32 toWrite = u32(sourceInfo.memoryFile->writtenSize);
-		DWORD written = 0;
-		if (!WriteFile(h, sourceInfo.memoryFile->baseAddress, toWrite, &written, NULL))
-			return false;
-
+		while (toWrite)
+		{
+			DWORD written = 0;
+			if (!WriteFile(h, sourceInfo.memoryFile->baseAddress, toWrite, &written, NULL))
+				return false;
+			toWrite -= written;
+		}
+		SetLastError(ERROR_SUCCESS);
 		sourceInfo.deleted = true;
-		return written == toWrite;
+		return true;
 	}
 
 	StringKey destKey = ToStringKeyLower(dest);
diff --git a/Engine/Source/Programs/UnrealBuildAccelerator/Detours/Private/Windows/UbaDetoursFunctionsNtDll.inl b/Engine/Source/Programs/UnrealBuildAccelerator/Detours/Private/Windows/UbaDetoursFunctionsNtDll.inl
--- a/Engine/Source/Programs/UnrealBuildAccelerator/Detours/Private/Windows/UbaDetoursFunctionsNtDll.inl
+++ b/Engine/Source/Programs/UnrealBuildAccelerator/Detours/Private/Windows/UbaDetoursFunctionsNtDll.inl
@@ -433,7 +433,7 @@ NTSTATUS NTAPI Detoured_NtQueryDirectoryFile(HANDLE FileHandle, HANDLE Event, PI
 NTSTATUS Detoured_NtQueryFullAttributesFile(POBJECT_ATTRIBUTES ObjectAttributes, PVOID Attributes)
 {
 	auto fileName = (const wchar_t*)ObjectAttributes->ObjectName->Buffer;
-	u32 fileNameLen = ObjectAttributes->ObjectName->Length/sizeof(tchar);
+	u32 fileNameLen = ObjectAttributes->ObjectName->Length/sizeof(tchar);(void)fileNameLen;
 	UBA_ASSERT(fileName[fileNameLen] == 0);
 
 	if (!fileName || !StartsWith(fileName, L"\\??\\") || !CanDetour(fileName + 4) || wcschr(fileName, ',') != nullptr || Contains(fileName, L"::")) // Some weird .net path used by dotnet.exe ... ignore for now!
