[UBA]
* Fixed bug where overlapped writes against memory file was not implemented
* Fixed so local processes that use process reuse write "written files" to disk before calling callback saying they're done

diff --git a/Engine/Source/Programs/UnrealBuildAccelerator/Common/Private/UbaProcess.cpp b/Engine/Source/Programs/UnrealBuildAccelerator/Common/Private/UbaProcess.cpp
--- a/Engine/Source/Programs/UnrealBuildAccelerator/Common/Private/UbaProcess.cpp
+++ b/Engine/Source/Programs/UnrealBuildAccelerator/Common/Private/UbaProcess.cpp
@@ -1264,6 +1264,8 @@ namespace uba
 		kernelStats.Write(statsWriter);
 		BinaryReader statsReader(statsWriter.GetData(), 0, statsWriter.GetPosition());
 
+		WriteFilesToDisk();
+
 		bool newProcess = false;
 		m_exitCode = prevExitCode;
 		m_messageSuccess = m_session.GetNextProcess(*this, newProcess, nextProcess, prevExitCode, statsReader) && m_messageSuccess;
@@ -1284,7 +1286,8 @@ namespace uba
 		m_logLines.clear();
 		m_trackedInputs.clear();
 		m_trackedOutputs.clear();
-		
+		m_writtenFiles.clear();
+
 		ClearTempFiles();
 		
 		m_processStats = {};
diff --git a/Engine/Source/Programs/UnrealBuildAccelerator/Detours/Private/Windows/UbaDetoursFunctionsKernelBase.inl b/Engine/Source/Programs/UnrealBuildAccelerator/Detours/Private/Windows/UbaDetoursFunctionsKernelBase.inl
--- a/Engine/Source/Programs/UnrealBuildAccelerator/Detours/Private/Windows/UbaDetoursFunctionsKernelBase.inl
+++ b/Engine/Source/Programs/UnrealBuildAccelerator/Detours/Private/Windows/UbaDetoursFunctionsKernelBase.inl
@@ -661,7 +661,24 @@ BOOL Detoured_WriteFile(HANDLE hFile, LPCVOID lpBuffer, DWORD nNumberOfBytesToWr
 		auto& fi = *fo.fileInfo;
 		if (MemoryFile* mf = fi.memoryFile)
 		{
-			mf->Write(dh, lpBuffer, nNumberOfBytesToWrite);
+			if (lpOverlapped)
+			{
+				u64 offset = ToLargeInteger(lpOverlapped->OffsetHigh, lpOverlapped->Offset).QuadPart;
+				u64 writtenSize = offset + nNumberOfBytesToWrite;
+				{
+					SCOPED_WRITE_LOCK(mf->lock, lock);
+					mf->EnsureCommitted(dh, writtenSize);
+					mf->writtenSize = Max(mf->writtenSize, writtenSize);
+					mf->isReported = false;
+				}
+				memcpy(mf->baseAddress + offset, lpBuffer, nNumberOfBytesToWrite);
+				//dh.pos = writtenSize;
+				if (lpOverlapped->hEvent)
+					SetEvent(lpOverlapped->hEvent);
+			}
+			else
+				mf->Write(dh, lpBuffer, nNumberOfBytesToWrite);
+
 			*lpNumberOfBytesWritten = nNumberOfBytesToWrite;
 			SetLastError(ERROR_SUCCESS);
 			DEBUG_LOG_DETOURED(L"WriteFile", L"(MEMORY) %llu (%ls) ToWrite: %u -> Success", uintptr_t(hFile), HandleToName(hFile), nNumberOfBytesToWrite);
diff --git a/Engine/Source/Programs/UnrealBuildAccelerator/Detours/Private/Windows/UbaDetoursUtilsWin.h b/Engine/Source/Programs/UnrealBuildAccelerator/Detours/Private/Windows/UbaDetoursUtilsWin.h
--- a/Engine/Source/Programs/UnrealBuildAccelerator/Detours/Private/Windows/UbaDetoursUtilsWin.h
+++ b/Engine/Source/Programs/UnrealBuildAccelerator/Detours/Private/Windows/UbaDetoursUtilsWin.h
@@ -79,6 +79,7 @@ namespace uba
 		u64 mappedSize = 0;
 		u64 committedSize = 0;
 		u64 writtenSize = 0;
+		ReaderWriterLock lock;
 		bool isLocalOnly;
 		bool isReported = false;
 		bool isThrowAway = false;
