[UBA]
* Fixed bug where Storage::DeleteAllFiles didn't go wide the way it was intended to do

diff --git a/Engine/Source/Programs/UnrealBuildAccelerator/Common/Private/UbaStorage.cpp b/Engine/Source/Programs/UnrealBuildAccelerator/Common/Private/UbaStorage.cpp
--- a/Engine/Source/Programs/UnrealBuildAccelerator/Common/Private/UbaStorage.cpp
+++ b/Engine/Source/Programs/UnrealBuildAccelerator/Common/Private/UbaStorage.cpp
@@ -1796,7 +1796,7 @@ namespace uba
 			Atomic<u32> atomicDeleteCount;
 			TraverseDir(m_logger, m_rootDir, [&](const DirectoryEntry& e)
 				{
-					if (IsDirectory(e.attributes))
+					if (!IsDirectory(e.attributes))
 						return;
 					workManager.AddWork([&, name = TString(e.name)](const WorkContext&)
 						{
