[UBA]
* Changed so files written by shadercompileworker with extension .tmp or .out are treated as output files

diff --git a/Engine/Source/Programs/UnrealBuildAccelerator/Core/Private/UbaApplicationRules.cpp b/Engine/Source/Programs/UnrealBuildAccelerator/Core/Private/UbaApplicationRules.cpp
--- a/Engine/Source/Programs/UnrealBuildAccelerator/Core/Private/UbaApplicationRules.cpp
+++ b/Engine/Source/Programs/UnrealBuildAccelerator/Core/Private/UbaApplicationRules.cpp
@@ -590,6 +590,11 @@ namespace uba
 			return file.Contains(TC(".uba."));
 		}
 
+		virtual bool IsOutputFile(const StringView& fileName) const override
+		{
+			return fileName.EndsWith(TC(".tmp")) || fileName.EndsWith(TC(".out"));
+		}
+
 		virtual bool AllowDetach() const override
 		{
 			return true;
