[UBA]
* Fixed so bright/dark green parts of boxes are written when using process reuse
* Optimized breadcrumb trace message by checking against running processes first before searching all

diff --git a/Engine/Source/Programs/UnrealBuildAccelerator/Visualizer/Private/UbaTraceReader.cpp b/Engine/Source/Programs/UnrealBuildAccelerator/Visualizer/Private/UbaTraceReader.cpp
--- a/Engine/Source/Programs/UnrealBuildAccelerator/Visualizer/Private/UbaTraceReader.cpp
+++ b/Engine/Source/Programs/UnrealBuildAccelerator/Visualizer/Private/UbaTraceReader.cpp
@@ -707,7 +707,7 @@ namespace uba
 			reader.ReadString(reason);
 			if (out.version < 15)
 				time = reader.Read7BitEncoded();
-			TraceView::ProcessLocation active = findIt->second;
+			TraceView::ProcessLocation& active = findIt->second;
 			auto& session = GetSession(out, active.sessionIndex);
 
 			auto& processes = session.processors[active.processorIndex].processes;
@@ -742,11 +742,14 @@ namespace uba
 			process.exitCode = 0u;
 			process.stop = time;
 			process.bitmapDirty = true;
+			process.createFilesTime = processStats.createFile.time;
+			process.writeFilesTime = Max(processStats.writeFiles.time, processStats.sendFiles.time);
+
 			bool isRemote = process.isRemote;
 
 			processes.emplace_back();
 			auto& process2 = processes.back();
-			m_activeProcesses[processId].processIndex = u32(processes.size() - 1);
+			active.processIndex = u32(processes.size() - 1);
 			process2.id = processId;
 			process2.description = reason.data;
 			process2.breadcrumbs = breadcrumbs;
@@ -1146,14 +1149,31 @@ namespace uba
 				breadcrumbs = reader.ReadLongString();
 			bool deleteOld = reader.ReadBool();
 
+			auto writeBreadcrumb = [&](TraceView::Process& process)
+				{
+					if (deleteOld)
+						process.breadcrumbs.clear();
+					else if (!process.breadcrumbs.empty())
+						process.breadcrumbs += '\n';
+					process.breadcrumbs += breadcrumbs;
+				};
+
+			auto findIt = m_activeProcesses.find(processId);
+			if (findIt != m_activeProcesses.end())
+			{
+				auto& active = findIt->second;
+				writeBreadcrumb(GetSession(out, active.sessionIndex).processors[active.processorIndex].processes[active.processIndex]);
+				break;
+			}
+
+			// Process is not active anymore, search for it the slow way
 			for (auto& session : out.sessions)
 				for (auto& processor : session.processors)
 					for (auto& process : processor.processes)
 						if (process.id == processId)
 						{
-							if (deleteOld)
-								process.breadcrumbs.clear();
-							process.breadcrumbs += breadcrumbs;
+							writeBreadcrumb(process);
+							return true;
 						}
 			break;
 		}
diff --git a/Engine/Source/Programs/UnrealBuildAccelerator/Visualizer/Private/UbaVisualizer.cpp b/Engine/Source/Programs/UnrealBuildAccelerator/Visualizer/Private/UbaVisualizer.cpp
--- a/Engine/Source/Programs/UnrealBuildAccelerator/Visualizer/Private/UbaVisualizer.cpp
+++ b/Engine/Source/Programs/UnrealBuildAccelerator/Visualizer/Private/UbaVisualizer.cpp
@@ -1429,7 +1429,7 @@ namespace uba
 							if (rect.right <= progressRect.left)
 								continue;
 
-							if (!m_filterString.empty() && !Contains(process.description.c_str(), m_filterString.c_str()))
+							if (!m_filterString.empty() && !Contains(process.description.c_str(), m_filterString.c_str()) && !Contains(process.breadcrumbs.c_str(), m_filterString.c_str()))
 								continue;
 							++visibleBoxes;
 
@@ -2924,7 +2924,7 @@ namespace uba
 							if (rect.right <= progressRect.left)
 								continue;
 
-							if (!m_filterString.empty() && !Contains(process.description.c_str(), m_filterString.c_str()))
+							if (!m_filterString.empty() && !Contains(process.description.c_str(), m_filterString.c_str()) && !Contains(process.breadcrumbs.c_str(), m_filterString.c_str()))
 								continue;
 
 							rect.right = Max(int(rect.right), left + 1);
