[UBA]
* Added GetFileBasicInformation function

diff --git a/Engine/Source/Programs/UnrealBuildAccelerator/Common/Private/UbaFile.cpp b/Engine/Source/Programs/UnrealBuildAccelerator/Common/Private/UbaFile.cpp
--- a/Engine/Source/Programs/UnrealBuildAccelerator/Common/Private/UbaFile.cpp
+++ b/Engine/Source/Programs/UnrealBuildAccelerator/Common/Private/UbaFile.cpp
@@ -174,33 +174,19 @@ namespace uba
 		return ReportOpenFileError(logger, fileName, lastError);
 	}
 
-	bool GetFileBasicInformationByHandle(FileBasicInformation& out, Logger& logger, const tchar* fileName, FileHandle hFile)
+	bool GetFileBasicInformationByHandle(FileBasicInformation& out, Logger& logger, const tchar* fileName, FileHandle hFile, bool errorOnFail)
 	{
 #if PLATFORM_WINDOWS
 		ExtendedTimerScope ts(KernelStats::GetCurrent().getFileInfo);
 
-		FILE_NETWORK_OPEN_INFORMATION info;
 		if (IsRunningWine())
-		{
-			MAKE_NT_FILENAME(fileName);
-			UNICODE_STRING us;
-			us.Length = USHORT(TStrlen(fileName)*sizeof(tchar));
-			us.MaximumLength = us.Length;
-			us.Buffer = (tchar*)fileName;
-			OBJECT_ATTRIBUTES oa;
-			InitializeObjectAttributes(&oa, &us, OBJ_CASE_INSENSITIVE, NULL, NULL);
-			NTSTATUS res = NtQueryFullAttributesFile(&oa, &info);
-			if (res != STATUS_SUCCESS)
-				return logger.Error(TC("GetFileBasicInformationByHandle (NtQueryFullAttributesFile) failed on %s (0x%x)"), fileName, res);
-		}
-		else
-		{
-			constexpr auto FileNetworkOpenInformation = (FILE_INFORMATION_CLASS)34;
-			IO_STATUS_BLOCK b;
-			NTSTATUS res = NtQueryInformationFile(asHANDLE(hFile), &b, &info, sizeof(info), FileNetworkOpenInformation);
-			if (res != STATUS_SUCCESS)
-				return logger.Error(TC("GetFileBasicInformationByHandle (NtQueryInformationFile) failed on %s (0x%x)"), fileName, res);
-		}
+			return GetFileBasicInformation(out, logger, fileName, errorOnFail);
+		constexpr auto FileNetworkOpenInformation = (FILE_INFORMATION_CLASS)34;
+		IO_STATUS_BLOCK b;
+		FILE_NETWORK_OPEN_INFORMATION info;
+		NTSTATUS res = NtQueryInformationFile(asHANDLE(hFile), &b, &info, sizeof(info), FileNetworkOpenInformation);
+		if (res != STATUS_SUCCESS)
+			return errorOnFail ? logger.Error(TC("GetFileBasicInformationByHandle (NtQueryInformationFile) failed on %s (0x%x)"), fileName, res) : false;
 		out.attributes = info.FileAttributes;
 		out.lastWriteTime = info.LastWriteTime.QuadPart;
 		out.size = info.EndOfFile.QuadPart;
@@ -216,6 +202,36 @@ namespace uba
 #endif
 	}
 
+	bool GetFileBasicInformation(FileBasicInformation& out, Logger& logger, const tchar* fileName, bool errorOnFail)
+	{
+#if PLATFORM_WINDOWS
+		ExtendedTimerScope ts(KernelStats::GetCurrent().getFileInfo);
+		FILE_NETWORK_OPEN_INFORMATION info;
+		MAKE_NT_FILENAME(fileName);
+		UNICODE_STRING us;
+		us.Length = USHORT(TStrlen(fileName)*sizeof(tchar));
+		us.MaximumLength = us.Length;
+		us.Buffer = (tchar*)fileName;
+		OBJECT_ATTRIBUTES oa;
+		InitializeObjectAttributes(&oa, &us, OBJ_CASE_INSENSITIVE, NULL, NULL);
+		NTSTATUS res = NtQueryFullAttributesFile(&oa, &info);
+		if (res != STATUS_SUCCESS)
+			return errorOnFail ? logger.Error(TC("GetFileBasicInformationByHandle (NtQueryFullAttributesFile) failed on %s (0x%x)"), fileName, res) : false;
+		out.attributes = info.FileAttributes;
+		out.lastWriteTime = info.LastWriteTime.QuadPart;
+		out.size = info.EndOfFile.QuadPart;
+		return true;
+#else
+		FileInformation info;
+		if (!GetFileInformation(info, logger, fileName))
+			return false;
+		out.attributes = info.attributes;
+		out.lastWriteTime = info.lastWriteTime;
+		out.size = info.size;
+		return true;
+#endif
+	}
+
 	bool GetFileInformationByHandle(FileInformation& out, Logger& logger, const tchar* fileName, FileHandle hFile)
 	{
 		ExtendedTimerScope ts(KernelStats::GetCurrent().getFileInfo);
diff --git a/Engine/Source/Programs/UnrealBuildAccelerator/Common/Public/UbaFile.h b/Engine/Source/Programs/UnrealBuildAccelerator/Common/Public/UbaFile.h
--- a/Engine/Source/Programs/UnrealBuildAccelerator/Common/Public/UbaFile.h
+++ b/Engine/Source/Programs/UnrealBuildAccelerator/Common/Public/UbaFile.h
@@ -41,7 +41,8 @@ namespace uba
 		u64 lastWriteTime;
 		u64 size;
 	};
-	bool GetFileBasicInformationByHandle(FileBasicInformation& out, Logger& logger, const tchar* fileName, FileHandle hFile);
+	bool GetFileBasicInformationByHandle(FileBasicInformation& out, Logger& logger, const tchar* fileName, FileHandle hFile, bool errorOnFail = true);
+	bool GetFileBasicInformation(FileBasicInformation& out, Logger& logger, const tchar* fileName, bool errorOnFail = true);
 
 	struct FileInformation
 	{
