[UBA]
* Added option to write out written files to trace file

diff --git a/Engine/Source/Programs/UnrealBuildAccelerator/Common/Private/UbaSession.cpp b/Engine/Source/Programs/UnrealBuildAccelerator/Common/Private/UbaSession.cpp
--- a/Engine/Source/Programs/UnrealBuildAccelerator/Common/Private/UbaSession.cpp
+++ b/Engine/Source/Programs/UnrealBuildAccelerator/Common/Private/UbaSession.cpp
@@ -226,6 +226,7 @@ namespace uba
 		table.GetValueAsBool(traceEnabled, TC("TraceEnabled"));
 		table.GetValueAsBool(detailedTrace, TC("DetailedTrace"));
 		table.GetValueAsBool(traceChildProcesses, TC("TraceChildProcesses"));
+		table.GetValueAsBool(traceWrittenFiles, TC("TraceWrittenFiles"));
 		table.GetValueAsBool(storeObjFilesCompressed, TC("StoreObjFilesCompressed"));
 		table.GetValueAsBool(allowLocalDetour, TC("AllowLocalDetour"));
 		table.GetValueAsBool(extractObjFilesSymbols, TC("ExtractObjFilesSymbols"));
@@ -1164,6 +1165,7 @@ namespace uba
 
 		m_detailedTrace = info.detailedTrace;
 		m_traceChildProcesses = info.traceChildProcesses;
+		m_traceWrittenFiles = info.traceWrittenFiles;
 		m_logToFile = info.logToFile;
 		if (info.extraInfo)
 			m_extraInfo = info.extraInfo;
@@ -2473,12 +2475,16 @@ namespace uba
 		if (!msg.newName.IsEmpty())
 		{
 			RegisterDeleteFile(file.nameKey, file.name);
-			RegisterCreateFileForWrite(msg.newNameKey, msg.newName, registerRealFile, fileSize, lastWriteTime);
+			if (RegisterCreateFileForWrite(msg.newNameKey, msg.newName, registerRealFile, fileSize, lastWriteTime) && registerRealFile)
+				TraceWrittenFile(msg.process.m_id, msg.newName, fileSize);
 		}
 		else if (msg.deleteOnClose)
 			RegisterDeleteFile(file.nameKey, file.name);
 		else
-			RegisterCreateFileForWrite(file.nameKey, file.name, registerRealFile, fileSize, lastWriteTime);
+		{
+			if (RegisterCreateFileForWrite(file.nameKey, file.name, registerRealFile, fileSize, lastWriteTime) && registerRealFile)
+				TraceWrittenFile(msg.process.m_id, file.name, fileSize);
+		}
 
 		out.directoryTableSize = GetDirectoryTableSize();
 		return true;
@@ -2522,6 +2528,24 @@ namespace uba
 
 	bool Session::MoveFile(MoveFileResponse& out, const MoveFileMessage& msg)
 	{
+		auto& process = msg.process;
+		{
+			SCOPED_FUTEX(process.m_writtenFilesLock, writtenLock);
+			auto findIt = process.m_writtenFiles.find(msg.fromName.data);
+			if (findIt != process.m_writtenFiles.end())
+			{
+				auto insres = process.m_writtenFiles.try_emplace(msg.toName.data);
+				UBA_ASSERT(insres.second);
+				WrittenFile& newFile = insres.first->second;
+				newFile = findIt->second;
+				newFile.name = msg.toName.data;
+				process.m_writtenFiles.erase(findIt);
+				out.errorCode = ERROR_SUCCESS;
+				out.result = true;
+				return true;
+			}
+		}
+
 		out.result = MoveFileExW(msg.fromName.data, msg.toName.data, msg.flags);
 		if (!out.result)
 		{
@@ -2529,7 +2553,8 @@ namespace uba
 			return true;
 		}
 		out.errorCode = ERROR_SUCCESS;
-		RegisterCreateFileForWrite(msg.toKey, msg.toName, true);
+		if (RegisterCreateFileForWrite(msg.toKey, msg.toName, true))
+			TraceWrittenFile(process.m_id, msg.toName, 0);
 		out.directoryTableSize = RegisterDeleteFile(msg.fromKey, msg.fromName);
 		return true;
 	}
@@ -2800,6 +2825,9 @@ namespace uba
 				m_fileMappingTableSize = (u32)newSize;
 			}
 			file.mappingHandle = {};
+
+			TraceWrittenFile(process.m_id, file.name, file.mappingWritten);
+
 			return true;
 		};
 
@@ -3588,6 +3616,26 @@ namespace uba
 		}
 	}
 
+	void Session::TraceWrittenFile(u32 processId, const StringView& file, u64 size)
+	{
+		if (!m_traceWrittenFiles)
+			return;
+		StringBuffer<> str(TC("WrittenFile: "));
+		str.Append(file);
+		if (!size)
+		{
+			size = InvalidValue;
+			FileBasicInformation info;
+			if (!GetFileBasicInformation(info, m_logger, file.data, false))
+				str.Append(TC(" (GetFileBasicInformation failed)")).Append(BytesToText(size).str).Append(TC(')'));
+			else
+				size = info.size;
+		}
+		if (size != InvalidValue)
+			str.Append(TC( " (size: ")).Append(BytesToText(size).str).Append(TC(')'));
+		m_trace.ProcessAddBreadcrumbs(processId, str, false);
+	}
+
 	void Session::TraceSessionUpdate()
 	{
 	}
diff --git a/Engine/Source/Programs/UnrealBuildAccelerator/Common/Private/UbaSessionServer.cpp b/Engine/Source/Programs/UnrealBuildAccelerator/Common/Private/UbaSessionServer.cpp
--- a/Engine/Source/Programs/UnrealBuildAccelerator/Common/Private/UbaSessionServer.cpp
+++ b/Engine/Source/Programs/UnrealBuildAccelerator/Common/Private/UbaSessionServer.cpp
@@ -1094,6 +1094,8 @@ namespace uba
 				success = m_storage.CopyOrLink(casKey, destination.data, attributes, writeCompressed, formattingFunc);
 				if (!success)
 					m_logger.Error(TC("Failed to copy cas from %s to %s (%s)"), CasKeyString(casKey).str, destination.data, GetProcessDescription(processId).c_str());
+				else
+					TraceWrittenFile(processId, destination);
 			}
 			else
 			{
diff --git a/Engine/Source/Programs/UnrealBuildAccelerator/Common/Public/UbaSession.h b/Engine/Source/Programs/UnrealBuildAccelerator/Common/Public/UbaSession.h
--- a/Engine/Source/Programs/UnrealBuildAccelerator/Common/Public/UbaSession.h
+++ b/Engine/Source/Programs/UnrealBuildAccelerator/Common/Public/UbaSession.h
@@ -188,6 +188,7 @@ namespace uba
 		bool DevirtualizeDepsFile(RootsHandle rootsHandle, MemoryBlock& destData, const void* sourceData, u64 sourceSize, bool escapeSpaces, const tchar* hint);
 		
 		void ThreadTraceLoop();
+		void TraceWrittenFile(u32 processId, const StringView& file, u64 size = 0);
 		virtual void TraceSessionUpdate();
 
 		Storage& m_storage;
@@ -215,6 +216,7 @@ namespace uba
 		bool m_shouldWriteToDisk;
 		bool m_detailedTrace;
 		bool m_traceChildProcesses;
+		bool m_traceWrittenFiles;
 		bool m_logToFile;
 		bool m_storeObjFilesCompressed;
 		bool m_readObjFilesCompressed;
diff --git a/Engine/Source/Programs/UnrealBuildAccelerator/Common/Public/UbaSessionCreateInfo.h b/Engine/Source/Programs/UnrealBuildAccelerator/Common/Public/UbaSessionCreateInfo.h
--- a/Engine/Source/Programs/UnrealBuildAccelerator/Common/Public/UbaSessionCreateInfo.h
+++ b/Engine/Source/Programs/UnrealBuildAccelerator/Common/Public/UbaSessionCreateInfo.h
@@ -34,6 +34,7 @@ namespace uba
 		bool traceEnabled = false;				// Set to true to always create in-memory trace data. Is not needed if traceName, traceOutputFile or launchVisualizer is set
 		bool detailedTrace = false;				// Enable detailed trace to include jobs, individual file I/O etc in trace dump
 		bool traceChildProcesses = false;		// Trace and visualize child processes so they can be seen in visualizer
+		bool traceWrittenFiles = false;			// Will add process output files to trace
 		bool storeObjFilesCompressed = false;	// Compiler will write .obj files compressed to disk and linker will decompress when reading
 		bool readObjFilesCompressed = false;	// Set to true to support reading compressed .obj files. This flag will not compress new obj files
 		bool allowLocalDetour = true;			// Allow local processes to be detoured. If this is false it is up to outside logic to register all created/deleted files
