[UBA]
* Fix for last change related to MoveFile where moving files not being memory mapped stopped working

diff --git a/Engine/Source/Programs/UnrealBuildAccelerator/Common/Private/UbaSession.cpp b/Engine/Source/Programs/UnrealBuildAccelerator/Common/Private/UbaSession.cpp
--- a/Engine/Source/Programs/UnrealBuildAccelerator/Common/Private/UbaSession.cpp
+++ b/Engine/Source/Programs/UnrealBuildAccelerator/Common/Private/UbaSession.cpp
@@ -2529,29 +2529,51 @@ namespace uba
 	bool Session::MoveFile(MoveFileResponse& out, const MoveFileMessage& msg)
 	{
 		auto& process = msg.process;
+		bool isMoved = false;
 		{
 			SCOPED_FUTEX(process.m_writtenFilesLock, writtenLock);
 			auto findIt = process.m_writtenFiles.find(msg.fromName.data);
 			if (findIt != process.m_writtenFiles.end())
 			{
+				auto& oldFile = findIt->second;
+				bool isMapping = oldFile.mappingHandle.IsValid();
+				if (!isMapping)
+				{
+					out.result = MoveFileExW(msg.fromName.data, msg.toName.data, msg.flags);
+					if (!out.result)
+					{
+						out.errorCode = GetLastError();
+						return true;
+					}
+					isMoved = true;
+				}
+
 				auto insres = process.m_writtenFiles.try_emplace(msg.toName.data);
 				UBA_ASSERT(insres.second);
 				WrittenFile& newFile = insres.first->second;
-				newFile = findIt->second;
+				newFile = oldFile;
 				newFile.name = msg.toName.data;
 				process.m_writtenFiles.erase(findIt);
-				out.errorCode = ERROR_SUCCESS;
-				out.result = true;
-				return true;
+
+				if (isMapping)
+				{
+					out.errorCode = ERROR_SUCCESS;
+					out.result = true;
+					return true;
+				}
 			}
 		}
 
-		out.result = MoveFileExW(msg.fromName.data, msg.toName.data, msg.flags);
-		if (!out.result)
+		if (!isMoved)
 		{
-			out.errorCode = GetLastError();
-			return true;
+			out.result = MoveFileExW(msg.fromName.data, msg.toName.data, msg.flags);
+			if (!out.result)
+			{
+				out.errorCode = GetLastError();
+				return true;
+			}
 		}
+
 		out.errorCode = ERROR_SUCCESS;
 		if (RegisterCreateFileForWrite(msg.toKey, msg.toName, true))
 			TraceWrittenFile(process.m_id, msg.toName, 0);
diff --git a/Engine/Source/Programs/UnrealBuildAccelerator/Detours/Private/Windows/UbaDetoursFunctionsKernelBase.inl b/Engine/Source/Programs/UnrealBuildAccelerator/Detours/Private/Windows/UbaDetoursFunctionsKernelBase.inl
--- a/Engine/Source/Programs/UnrealBuildAccelerator/Detours/Private/Windows/UbaDetoursFunctionsKernelBase.inl
+++ b/Engine/Source/Programs/UnrealBuildAccelerator/Detours/Private/Windows/UbaDetoursFunctionsKernelBase.inl
@@ -1638,11 +1638,13 @@ bool Shared_MoveFile(LPCWSTR lpExistingFileName, LPCWSTR lpNewFileName, DWORD dw
 		auto cg = MakeGuard([&]() { CloseHandle(h); });
 		UBA_ASSERT(sourceInfo.memoryFile->writtenSize < ~0u);
 		u32 toWrite = u32(sourceInfo.memoryFile->writtenSize);
+		u8* readPos = sourceInfo.memoryFile->baseAddress;
 		while (toWrite)
 		{
 			DWORD written = 0;
-			if (!WriteFile(h, sourceInfo.memoryFile->baseAddress, toWrite, &written, NULL))
+			if (!WriteFile(h, readPos, toWrite, &written, NULL))
 				return false;
+			readPos += written;
 			toWrite -= written;
 		}
 		SetLastError(ERROR_SUCCESS);
