[UBA]
* Added tests for traversedir
* Added handling of crawler trying to read directory as file.

diff --git a/Engine/Source/Programs/UnrealBuildAccelerator/Common/Private/UbaSessionClient.cpp b/Engine/Source/Programs/UnrealBuildAccelerator/Common/Private/UbaSessionClient.cpp
--- a/Engine/Source/Programs/UnrealBuildAccelerator/Common/Private/UbaSessionClient.cpp
+++ b/Engine/Source/Programs/UnrealBuildAccelerator/Common/Private/UbaSessionClient.cpp
@@ -2139,7 +2139,11 @@ namespace uba
 						return res;
 					}
 					else
-						m_logger.Warning(TC("Failed to open %s"), out.fileName.data);
+						return m_logger.Warning(TC("Failed to open %s"), out.fileName.data);
+				}
+				else if (out.fileName.Equals(TC("$d")))
+				{
+					return m_logger.Warning(TC("Trying to open directory %s as file"), fileName.data);
 				}
 				else
 				{
@@ -2147,9 +2151,8 @@ namespace uba
 					if (fa.OpenMemoryRead(0, false))
 						return func(fa.GetData(), fa.GetSize());
 					else
-						m_logger.Warning(TC("Failed to open %s"), out.fileName.data);
+						return m_logger.Warning(TC("Failed to open %s"), out.fileName.data);
 				}
-				return true;
 			};
 
 		auto DevirtualizePathFunc = [this, rootsHandle = startInfo.rootsHandle](StringBufferBase& inOut)
diff --git a/Engine/Source/Programs/UnrealBuildAccelerator/Test/Private/UbaTestAll.h b/Engine/Source/Programs/UnrealBuildAccelerator/Test/Private/UbaTestAll.h
--- a/Engine/Source/Programs/UnrealBuildAccelerator/Test/Private/UbaTestAll.h
+++ b/Engine/Source/Programs/UnrealBuildAccelerator/Test/Private/UbaTestAll.h
@@ -22,6 +22,7 @@ namespace uba
 		UBA_TEST(TestEvents) \
 		UBA_TEST(TestPaths) \
 		UBA_TEST(TestFiles) \
+		UBA_TEST(TestTraverseDir) \
 		UBA_TEST(TestOverlappedIO) \
 		UBA_TEST(TestMemoryBlock) \
 		UBA_TEST(TestParseArguments) \
diff --git a/Engine/Source/Programs/UnrealBuildAccelerator/Test/Private/UbaTestBasics.cpp b/Engine/Source/Programs/UnrealBuildAccelerator/Test/Private/UbaTestBasics.cpp
--- a/Engine/Source/Programs/UnrealBuildAccelerator/Test/Private/UbaTestBasics.cpp
+++ b/Engine/Source/Programs/UnrealBuildAccelerator/Test/Private/UbaTestBasics.cpp
@@ -257,10 +257,32 @@ namespace uba
 			return false;
 		#endif
 
-		LoggerWithWriter nullLogger(g_nullLogWriter);
-		if (TraverseDir(nullLogger, AsView(TC("TestDir2")), [&](const DirectoryEntry&) {}, true))
-			return logger.Error(TC("TraverseDir failed to report fail on non existing dir"));
+		return true;
+	}
 
+	bool TestTraverseDir(LoggerWithWriter& logger, const StringBufferBase& rootDir)
+	{
+		LoggerWithWriter nullLogger(g_nullLogWriter);
+		CHECK_TRUE(!TraverseDir(nullLogger, AsView(TC("TestDir2")), [&](const DirectoryEntry&) {}, true));
+		u32 foundCount = 0;
+		CHECK_TRUE(TraverseDir(nullLogger, rootDir, [&](const DirectoryEntry&) { ++foundCount; }, true));
+		CHECK_TRUE(foundCount == 0);
+
+		bool isFile = false;
+		StringBuffer<> entry(rootDir);
+		entry.Append(TC("Entry"));
+		FileAccessor fileHandle(logger, entry.data);
+		CHECK_TRUE(fileHandle.CreateWrite(false));
+		CHECK_TRUE(fileHandle.Close());
+		CHECK_TRUE(TraverseDir(nullLogger, rootDir, [&](const DirectoryEntry& e) { ++foundCount; isFile = !IsDirectory(e.attributes); }, true));
+		CHECK_TRUE(foundCount == 1);
+		CHECK_TRUE(isFile);
+		CHECK_TRUE(DeleteFileW(entry.data));
+		CHECK_TRUE(CreateDirectoryW(entry.data));
+		foundCount = 0;
+		CHECK_TRUE(TraverseDir(nullLogger, rootDir, [&](const DirectoryEntry& e) { ++foundCount; isFile = !IsDirectory(e.attributes); }, true));
+		CHECK_TRUE(foundCount == 1);
+		CHECK_TRUE(!isFile);
 		return true;
 	}
 
