[UBA]
* Fixed so verse-clangcl is using the correct dependency crawler type
* Changed so NtCreateFile retry for sharing violation only happens for .pch files.. since there are actually paths that tries this before doing something else (visual studio static analysis path)
* Added "/analyze:log" in ignore list for msvc dependency crawler

diff --git a/Engine/Source/Programs/UnrealBuildAccelerator/Common/Public/UbaDependencyCrawler.h b/Engine/Source/Programs/UnrealBuildAccelerator/Common/Public/UbaDependencyCrawler.h
--- a/Engine/Source/Programs/UnrealBuildAccelerator/Common/Public/UbaDependencyCrawler.h
+++ b/Engine/Source/Programs/UnrealBuildAccelerator/Common/Public/UbaDependencyCrawler.h
@@ -265,6 +265,7 @@ namespace uba
 				{ TCV("/external:I"), true },
 				{ TCV("/imsvc"), true },
 				{ TCV("/experimental:log"), false },
+				{ TCV("/analyze:log"), false },
 				{ TCV("/sourceDependencies"), false},
 			};
 
diff --git a/Engine/Source/Programs/UnrealBuildAccelerator/Core/Private/UbaApplicationRules.cpp b/Engine/Source/Programs/UnrealBuildAccelerator/Core/Private/UbaApplicationRules.cpp
--- a/Engine/Source/Programs/UnrealBuildAccelerator/Core/Private/UbaApplicationRules.cpp
+++ b/Engine/Source/Programs/UnrealBuildAccelerator/Core/Private/UbaApplicationRules.cpp
@@ -652,7 +652,7 @@ namespace uba
 			{ ApplicationHash(TC("lld-link.exe")),				new ApplicationRulesLldLinkExe() },
 			{ ApplicationHash(TC("clang++.exe")),				new ApplicationRulesClangPlusPlusExe() },
 			{ ApplicationHash(TC("clang-cl.exe")),				new ApplicationRulesClangClExe() },
-			{ ApplicationHash(TC("verse-clang-cl.exe")),		new ApplicationRulesClangPlusPlusExe() },
+			{ ApplicationHash(TC("verse-clang-cl.exe")),		new ApplicationRulesClangClExe() },
 			{ ApplicationHash(TC("ispc.exe")),					new ApplicationRulesISPCExe() },
 			{ ApplicationHash(TC("radlink.exe")),				new ApplicationRulesRadLinkExe() },
 			{ 3340509542,										new ApplicationRulesClangPlusPlusExePlatform1() },
diff --git a/Engine/Source/Programs/UnrealBuildAccelerator/Detours/Private/Windows/UbaDetoursFunctionsNtDll.inl b/Engine/Source/Programs/UnrealBuildAccelerator/Detours/Private/Windows/UbaDetoursFunctionsNtDll.inl
--- a/Engine/Source/Programs/UnrealBuildAccelerator/Detours/Private/Windows/UbaDetoursFunctionsNtDll.inl
+++ b/Engine/Source/Programs/UnrealBuildAccelerator/Detours/Private/Windows/UbaDetoursFunctionsNtDll.inl
@@ -712,6 +712,11 @@ NTSTATUS NTAPI Local_NtCreateFile(bool IsCreateFunc, PHANDLE hFileHandle, ACCESS
 		{
 			if (!--retriesLeft)
 				return res;
+
+			StringView fileName(ObjectAttributes->ObjectName->Buffer, ObjectAttributes->ObjectName->Length / 2);
+			if (!fileName.EndsWith(TC(".pch")))
+				return res;
+
 			#if UBA_DEBUG
 			StringBuffer<> b;
 			b.Appendf(L"Got access denied trying to open %.*s. Retrying in one second", ObjectAttributes->ObjectName->Length / 2, ObjectAttributes->ObjectName->Buffer);
