[UBA]
* Added more logging when crawler fails to open files
* Added "-F" handling to clang path since it is used by some platforms

diff --git a/Engine/Source/Programs/UnrealBuildAccelerator/Common/Private/UbaSession.cpp b/Engine/Source/Programs/UnrealBuildAccelerator/Common/Private/UbaSession.cpp
--- a/Engine/Source/Programs/UnrealBuildAccelerator/Common/Private/UbaSession.cpp
+++ b/Engine/Source/Programs/UnrealBuildAccelerator/Common/Private/UbaSession.cpp
@@ -320,7 +320,7 @@ namespace uba
 					memoryWriter.WriteString(e.name, e.nameLen);
 
 					#if UBA_DEBUG_TRACK_DIR
-					m_debugLogger->Info(TC("    %s (Size: %llu, Key: %s, Id: %llu)\n"), e.name, e.size, KeyToString(fileKey).data, e.id);
+					m_debugLogger->Info(TC("    %s (Size: %llu, Attr: %u, Key: %s, Id: %llu)\n"), e.name, e.size, e.attributes, KeyToString(fileKey).data, e.id);
 					#endif
 
 					u64 id = e.id;
@@ -928,7 +928,7 @@ namespace uba
 		}
 
 		#if UBA_DEBUG_TRACK_DIR
-		m_debugLogger->Info(TC("TRACKADD    %s (Size: %llu, Key: %s, Id: %llu)\n"), fileName.data, fileSize, KeyToString(fileNameKey).data, fileIndex);
+		m_debugLogger->Info(TC("TRACKADD    %s (Size: %llu, Attr: %u, Key: %s, Id: %llu)\n"), fileName.data, fileSize, attributes, KeyToString(fileNameKey).data, fileIndex);
 		#endif
 
 
diff --git a/Engine/Source/Programs/UnrealBuildAccelerator/Common/Public/UbaDependencyCrawler.h b/Engine/Source/Programs/UnrealBuildAccelerator/Common/Public/UbaDependencyCrawler.h
--- a/Engine/Source/Programs/UnrealBuildAccelerator/Common/Public/UbaDependencyCrawler.h
+++ b/Engine/Source/Programs/UnrealBuildAccelerator/Common/Public/UbaDependencyCrawler.h
@@ -86,10 +86,10 @@ namespace uba
 		struct CodeFile { TString path; bool hasPch = false; };
 		using CodeFiles = List<CodeFile>;
 
-		inline bool ParseRsp(Instance& instance, const tchar* rsp, CodeFiles& outCodeFiles);
+		inline bool ParseRsp(Instance& instance, const StringView& rsp, CodeFiles& outCodeFiles);
 		inline bool ParseRsp2(Instance& instance, const void* data, u64 dataSize, const tchar* rsp, CodeFiles& outCodeFiles);
 		inline bool ParsePch(Pch& pch, Instance& instance, const void* data, u64 dataSize);
-		inline bool ParseCodeFile(Instance& instance, const StringView& codeFile, bool parseDefines);
+		inline bool ParseCodeFile(Instance& instance, const StringView& codeFile, bool parseDefines, const StringView& caller);
 		inline bool ParseCodeFile2(Instance& instance, const void* data, u64 dataSize, const StringView& codeFile, bool parseDefines);
 		inline bool HandleInclude(Instance& instance, const StringView& rootPath, const StringView& include, const StringView& codeFile, bool parseDefines);
 	};
@@ -127,7 +127,7 @@ namespace uba
 			{
 				auto& instance = ref.instance;
 				CodeFiles codeFiles;
-				if (!ParseRsp(instance, instance.rsp.c_str(), codeFiles))
+				if (!ParseRsp(instance, instance.rsp, codeFiles))
 					return;
 
 				for (auto& cf : codeFiles)
@@ -159,7 +159,7 @@ namespace uba
 					}
 					else
 					{
-						ParseCodeFile(instance, cf.path, instance.platform.empty());
+						ParseCodeFile(instance, cf.path, instance.platform.empty(), instance.rsp);
 						if (instance.platform.empty())
 						{
 							if (!instance.overriddenPlatformName.empty())
@@ -192,7 +192,8 @@ namespace uba
 								codeFilePath.Resize(fixedPath.count).Append(file);
 								m_workManager.AddWork([this, ref, codeFile = TString(codeFilePath.data)](const WorkContext& context)
 									{
-										ref.instance.createFileFunc(codeFile, [&](const void* data, u64 dataSize) { return true; });
+										if (!ref.instance.createFileFunc(codeFile, [&](const void* data, u64 dataSize) { return true; }))
+											m_logger.Warning(TC("Failed to open file %s from builtindir"), codeFile.c_str());
 									}, 1, TC("Crawler"), ColorWork);
 							});
 					}
@@ -202,18 +203,17 @@ namespace uba
 		return true;
 	}
 
-	bool DependencyCrawler::ParseRsp(Instance& instance, const tchar* rsp, CodeFiles& outCodeFiles)
+	bool DependencyCrawler::ParseRsp(Instance& instance, const StringView& rsp, CodeFiles& outCodeFiles)
 	{
 		StringBuffer<> fixedPath;
-		FixPath(rsp, instance.workDir.data(), instance.workDir.size(), fixedPath);
+		FixPath(rsp.data, instance.workDir.data(), instance.workDir.size(), fixedPath);
 		if (!instance.devirtualizePathFunc(fixedPath))
-			m_logger.Info(TC("Failed to devirtualize path %s in %s"), fixedPath.data, rsp);
+			m_logger.Info(TC("Failed to devirtualize path %s in %s"), fixedPath.data, rsp.data);
 		if (CaseInsensitiveFs)
 			fixedPath.MakeLower();
-		return instance.createFileFunc(fixedPath, [&](const void* data, u64 dataSize)
-			{
-				return ParseRsp2(instance, data, dataSize, rsp, outCodeFiles);
-			});
+		if (!instance.createFileFunc(fixedPath, [&](const void* data, u64 dataSize) { return ParseRsp2(instance, data, dataSize, rsp.data, outCodeFiles); }))
+			return m_logger.Warning(TC("Failed to parse rsp %s"), rsp.data);
+		return true;
 	}
 
 	bool DependencyCrawler::ParseRsp2(Instance& instance, const void* data, u64 dataSize, const tchar* rsp, CodeFiles& outCodeFiles)
@@ -370,7 +370,7 @@ namespace uba
 					}
 					else if (sb[0] == '@')
 					{
-						ParseRsp(instance, sb.data + 1, outCodeFiles);
+						ParseRsp(instance, StringView(sb).Skip(1), outCodeFiles);
 					}
 					else
 					{
@@ -398,6 +398,7 @@ namespace uba
 				{ TCV("-isysroot"), ArgType_RootPath },
 				{ TCV("-internal-isystem"), ArgType_RootPath },
 				{ TCV("-I"), ArgType_RootPath },
+				{ TCV("-F"), ArgType_RootPath },
 				{ TCV("-resource-dir"), ArgType_ResourceDir },
 				{ TCV("-dependency-file"), ArgType_Ignore },
 				{ TCV("-internal-externc-isystem"), ArgType_Ignore },
@@ -479,7 +480,7 @@ namespace uba
 						return;
 					}
 					else if (sb[0] == '@')
-						ParseRsp(instance, sb.data + 1, outCodeFiles);
+						ParseRsp(instance, StringView(sb).Skip(1), outCodeFiles);
 					else if (sb.Contains('/'))
 						AddCodeFile(sb, false, false);
 					else
@@ -538,12 +539,11 @@ namespace uba
 		return true;
 	}
 
-	bool DependencyCrawler::ParseCodeFile(Instance& instance, const StringView& codeFile, bool parseDefines)
+	bool DependencyCrawler::ParseCodeFile(Instance& instance, const StringView& codeFile, bool parseDefines, const StringView& caller)
 	{
-		return instance.createFileFunc(codeFile, [&](const void* data, u64 dataSize)
-			{
-				return ParseCodeFile2(instance, data, dataSize, codeFile, parseDefines);
-			});
+		if (!instance.createFileFunc(codeFile, [&](const void* data, u64 dataSize) { return ParseCodeFile2(instance, data, dataSize, codeFile, parseDefines); }))
+			return m_logger.Warning(TC("Failed to parse code file %s found in %s"), codeFile.data, caller.data);
+		return true;
 	}
 
 	bool DependencyCrawler::ParseCodeFile2(Instance& instance, const void* data, u64 dataSize, const StringView& codeFile, bool parseDefines)
@@ -762,7 +762,6 @@ namespace uba
 				return true;
 		}
 
-		ParseCodeFile(instance, fullPath, parseDefines);
-		return true;
+		return ParseCodeFile(instance, fullPath, parseDefines, codeFile);
 	}
 }
