[UBA]
* Fixed hang when pressing ctrl-c on macos to exit UbaAgent

diff --git a/Engine/Source/Programs/UnrealBuildAccelerator/Agent/Private/UbaAgent.cpp b/Engine/Source/Programs/UnrealBuildAccelerator/Agent/Private/UbaAgent.cpp
--- a/Engine/Source/Programs/UnrealBuildAccelerator/Agent/Private/UbaAgent.cpp
+++ b/Engine/Source/Programs/UnrealBuildAccelerator/Agent/Private/UbaAgent.cpp
@@ -146,7 +146,7 @@ namespace uba
 		if (g_logger)
 			g_logger->Info(TC("  Exiting..."));
 		if (g_sessionClient)
-			g_sessionClient->Stop();
+			g_sessionClient->Stop(false);
 		g_exitLock->Leave();
 	}
 
diff --git a/Engine/Source/Programs/UnrealBuildAccelerator/Common/Private/UbaSessionClient.cpp b/Engine/Source/Programs/UnrealBuildAccelerator/Common/Private/UbaSessionClient.cpp
--- a/Engine/Source/Programs/UnrealBuildAccelerator/Common/Private/UbaSessionClient.cpp
+++ b/Engine/Source/Programs/UnrealBuildAccelerator/Common/Private/UbaSessionClient.cpp
@@ -159,11 +159,12 @@ namespace uba
 		return true;
 	}
 
-	void SessionClient::Stop()
+	void SessionClient::Stop(bool wait)
 	{
 		m_loop = false;
 		m_waitToSendEvent.Set();
-		m_loopThread.Wait();
+		if (wait)
+			m_loopThread.Wait();
 	}
 
 	bool SessionClient::Wait(u32 milliseconds, Event* wakeupEvent)
@@ -2143,6 +2144,7 @@ namespace uba
 				}
 				else if (out.fileName.Equals(TC("$d")))
 				{
+					// This can happen on apple targets.. crawler finds some includes that are not proper includes
 					return m_logger.Warning(TC("Trying to open directory %s as file"), fileName.data);
 				}
 				else
diff --git a/Engine/Source/Programs/UnrealBuildAccelerator/Common/Public/UbaSessionClient.h b/Engine/Source/Programs/UnrealBuildAccelerator/Common/Public/UbaSessionClient.h
--- a/Engine/Source/Programs/UnrealBuildAccelerator/Common/Public/UbaSessionClient.h
+++ b/Engine/Source/Programs/UnrealBuildAccelerator/Common/Public/UbaSessionClient.h
@@ -40,7 +40,7 @@ namespace uba
 		~SessionClient();
 
 		bool Start();
-		void Stop();
+		void Stop(bool wait = true);
 		bool Wait(u32 milliseconds = 0xFFFFFFFF, Event* wakeupEvent = nullptr);
 		void SendSummary(const Function<void(Logger&)>& extraInfo);
 		void SetIsTerminating(const tchar* reason = TC("Terminating"), u64 delayMs = 0); // Session stores pointer directly. Can't be temporary
