[UBA]
* Fixed bug where directory table did not properly handle invalid root paths
* Added commented out tests

diff --git a/Engine/Source/Programs/UnrealBuildAccelerator/Common/Private/UbaSession.cpp b/Engine/Source/Programs/UnrealBuildAccelerator/Common/Private/UbaSession.cpp
--- a/Engine/Source/Programs/UnrealBuildAccelerator/Common/Private/UbaSession.cpp
+++ b/Engine/Source/Programs/UnrealBuildAccelerator/Common/Private/UbaSession.cpp
@@ -359,7 +359,9 @@ namespace uba
 				#if UBA_DEBUG_TRACK_DIR
 				m_debugLogger->Info(TC("    FAILED (not existing?)\n"));
 				#endif
-				return false;
+
+				if (!IsWindows || dirPath.count > 3)
+					return false;
 			}
 		}
 		else
@@ -400,7 +402,6 @@ namespace uba
 			tableWriter.WriteStringKey(dirKey);
 			tableOffset = writePos + u32(tableWriter.GetPosition());
 			tableWriter.Write7BitEncoded(dir.tableOffset);
-			UBA_ASSERT(dirAttributes != 0);
 			tableWriter.WriteFileAttributes(dirAttributes);
 			tableWriter.WriteVolumeSerial(volumeSerialIndex);
 			tableWriter.WriteFileIndex(fileIndex);
diff --git a/Engine/Source/Programs/UnrealBuildAccelerator/Common/Private/UbaSessionClient.cpp b/Engine/Source/Programs/UnrealBuildAccelerator/Common/Private/UbaSessionClient.cpp
--- a/Engine/Source/Programs/UnrealBuildAccelerator/Common/Private/UbaSessionClient.cpp
+++ b/Engine/Source/Programs/UnrealBuildAccelerator/Common/Private/UbaSessionClient.cpp
@@ -85,15 +85,7 @@ namespace uba
 		if (m_useDependencyCrawler)
 			m_dependencyCrawler.Init([this](const StringView& fileName, u32& outAttr) // FileExists
 				{
-					u32 tableOffset = 0;
-					if (!EntryExists(fileName, tableOffset))
-						return false;
-					if ((tableOffset & 0x80000000) != 0)
-						return false;
-					BinaryReader reader(m_directoryTableMem, tableOffset, m_directoryTableMemPos);
-					reader.SkipString();
-					outAttr = reader.ReadFileAttributes();
-					return true;
+					return Exists(fileName, outAttr);
 				},
 			[this](const StringView& path, const DependencyCrawler::FileFunc& fileFunc) // TraverseFilesFunc
 				{
@@ -194,6 +186,15 @@ namespace uba
 		return m_bestPing;
 	}
 
+	bool SessionClient::Exists(const StringView& path, u32& outAttributes)
+	{
+		u32 tableOffset = 0;
+		if (!EntryExists(path, tableOffset))
+			return false;
+		outAttributes = m_directoryTable.GetAttributes(tableOffset);
+		return outAttributes != 0;
+	}
+
 	bool SessionClient::RetrieveCasFile(CasKey& outNewKey, u64& outSize, const CasKey& casKey, const tchar* hint, bool storeUncompressed, bool allowProxy)
 	{
 		//StringBuffer<> workName;
diff --git a/Engine/Source/Programs/UnrealBuildAccelerator/Common/Public/UbaSessionClient.h b/Engine/Source/Programs/UnrealBuildAccelerator/Common/Public/UbaSessionClient.h
--- a/Engine/Source/Programs/UnrealBuildAccelerator/Common/Public/UbaSessionClient.h
+++ b/Engine/Source/Programs/UnrealBuildAccelerator/Common/Public/UbaSessionClient.h
@@ -49,6 +49,8 @@ namespace uba
 
 		u64 GetBestPing();
 
+		bool Exists(const StringView& path, u32& outAttributes);
+
 	private:
 		bool RetrieveCasFile(CasKey& outNewKey, u64& outSize, const CasKey& casKey, const tchar* hint, bool storeUncompressed, bool allowProxy = true);
 
diff --git a/Engine/Source/Programs/UnrealBuildAccelerator/Core/Public/UbaDirectoryTable.h b/Engine/Source/Programs/UnrealBuildAccelerator/Core/Public/UbaDirectoryTable.h
--- a/Engine/Source/Programs/UnrealBuildAccelerator/Core/Public/UbaDirectoryTable.h
+++ b/Engine/Source/Programs/UnrealBuildAccelerator/Core/Public/UbaDirectoryTable.h
@@ -173,20 +173,27 @@ namespace uba
 				startSkip = 1;
 			}
 
-			// scan backwards first
-			const tchar* rit = entryName.data + entryName.count - startSkip;
+			// Scan backwards
+			const tchar* rend = entryName.data;
+			const tchar* rit = rend + entryName.count - startSkip;
+
 			bool inAncestor = false;
-			while (rit > entryName.data)
+			while (rit > rend)
 			{
 				if (*rit != PathSeparator)
 				{
 					--rit;
-					continue;
+
+					#if !PLATFORM_WINDOWS
+					if (rit != rend) // We want to test empty for non-windows
+					#endif
+						continue;
 				}
 
-				u64 sublen = u64(rit - entryName.data);
+				u64 sublen = u64(rit - rend);
+
 				StringKeyHasher ancestorHasher;
-				ancestorHasher.Update(entryName.data, sublen);
+				ancestorHasher.Update(rend, sublen);
 				StringKey ancestorKey = ToStringKey(ancestorHasher);
 				auto dirIt = m_lookup.find(ancestorKey);
 				if (dirIt != m_lookup.end())
@@ -244,6 +251,25 @@ namespace uba
 			u64 lastWrite;
 		};
 
+		u32 GetAttributes(u32 tableOffset)
+		{
+			if (tableOffset & 0x80000000)
+			{
+				tableOffset = tableOffset & ~0x80000000;
+				BinaryReader reader(m_memory, tableOffset);
+				u64 prevTableOffset = reader.Read7BitEncoded();
+				while (prevTableOffset != InvalidTableOffset)
+				{
+					reader.SetPosition(prevTableOffset);
+					prevTableOffset = reader.Read7BitEncoded();
+				}
+				return reader.ReadFileAttributes();
+			}
+			BinaryReader reader(m_memory, tableOffset);
+			reader.SkipString();
+			return reader.ReadFileAttributes();
+		}
+
 		u32 GetEntryInformation(EntryInformation& outInfo, u32 tableOffset, tchar* outFileName = nullptr, u32 fileNameCapacity = 0)
 		{
 			if (tableOffset & 0x80000000)
diff --git a/Engine/Source/Programs/UnrealBuildAccelerator/Detours/Private/PosixOS/UbaDetoursFunctionsPosix.cpp b/Engine/Source/Programs/UnrealBuildAccelerator/Detours/Private/PosixOS/UbaDetoursFunctionsPosix.cpp
--- a/Engine/Source/Programs/UnrealBuildAccelerator/Detours/Private/PosixOS/UbaDetoursFunctionsPosix.cpp
+++ b/Engine/Source/Programs/UnrealBuildAccelerator/Detours/Private/PosixOS/UbaDetoursFunctionsPosix.cpp
@@ -1236,7 +1236,7 @@ UBA_EXPORT DIR* UBA_WRAPPER(opendir)(const char* name)
 	if (insres.second)
 		if (g_directoryTable.EntryExistsNoLock(hash.key, forHash) != DirectoryTable::Exists_No)
 			Rpc_UpdateDirectory(hash.key, dirName.data, dirName.count, false);
-\
+
 	bool exists = false;
 	if (dir.tableOffset != InvalidTableOffset)
 	{
diff --git a/Engine/Source/Programs/UnrealBuildAccelerator/Test/Private/UbaTestAll.h b/Engine/Source/Programs/UnrealBuildAccelerator/Test/Private/UbaTestAll.h
--- a/Engine/Source/Programs/UnrealBuildAccelerator/Test/Private/UbaTestAll.h
+++ b/Engine/Source/Programs/UnrealBuildAccelerator/Test/Private/UbaTestAll.h
@@ -49,6 +49,7 @@ namespace uba
 		UBA_TEST(TestRegisterChangesRemote) \
 		UBA_TEST(TestDetouredClang) \
 		UBA_TEST(TestFileMappingBuffer) \
+		UBA_TEST(TestRemoteDirectoryTable) \
 		UBA_TEST_EXTRA \
 
 #define UBA_POSIX_TESTS \
diff --git a/Engine/Source/Programs/UnrealBuildAccelerator/Test/Private/UbaTestSession.cpp b/Engine/Source/Programs/UnrealBuildAccelerator/Test/Private/UbaTestSession.cpp
--- a/Engine/Source/Programs/UnrealBuildAccelerator/Test/Private/UbaTestSession.cpp
+++ b/Engine/Source/Programs/UnrealBuildAccelerator/Test/Private/UbaTestSession.cpp
@@ -50,7 +50,8 @@ namespace uba
 		return testFunc(logger, session, workingDir.data, [&](const ProcessStartInfo& pi) { return session.RunProcess(pi, true, enableDetour); });
 	}
 
-	bool RunRemote(LoggerWithWriter& logger, const StringBufferBase& testRootDir, const TestSessionFunction& testFunc, bool deleteAll, bool serverShouldListen)
+	using TestServerClientSessionFunction = Function<bool(LoggerWithWriter& logger, const StringView& workingDir, SessionServer& sessionServer, SessionClient& sessionClient)>;
+	bool SetupServerClientSession(LoggerWithWriter& logger, const StringBufferBase& testRootDir, bool deleteAll, bool serverShouldListen, const TestServerClientSessionFunction& testFunc)
 	{
 		LogWriter& logWriter = logger.m_writer;
 		NetworkBackendTcp tcpBackend(logWriter);
@@ -150,7 +151,16 @@ namespace uba
 
 
 		workingDir.EnsureEndsWithSlash();
-		return testFunc(logger, sessionServer, workingDir.data, [&](const ProcessStartInfo& pi) { return sessionServer.RunProcessRemote(pi); });
+		return testFunc(logger, workingDir, sessionServer, sessionClient);
+	}
+
+	bool RunRemote(LoggerWithWriter& logger, const StringBufferBase& testRootDir, const TestSessionFunction& testFunc, bool deleteAll, bool serverShouldListen)
+	{
+		return SetupServerClientSession(logger, testRootDir, deleteAll, serverShouldListen,
+			[&](LoggerWithWriter& logger, const StringView& workingDir, SessionServer& sessionServer, SessionClient& sessionClient)
+			{
+				return testFunc(logger, sessionServer, workingDir.data, [&](const ProcessStartInfo& pi) { return sessionServer.RunProcessRemote(pi); });
+			});
 	}
 
 	void GetTestAppPath(LoggerWithWriter& logger, StringBufferBase& out)
@@ -625,4 +635,29 @@ namespace uba
 			});
 	}
 
+
+	bool TestRemoteDirectoryTable(LoggerWithWriter& logger, const StringBufferBase& testRootDir)
+	{
+#if 0
+		return SetupServerClientSession(logger, testRootDir, true, true,
+			[&](LoggerWithWriter& logger, const StringView& workingDir, SessionServer& sessionServer, SessionClient& sessionClient)
+			{
+				u32 attributes;
+				#if PLATFORM_WINDOWS
+				CHECK_TRUE(sessionClient.Exists(TCV("c:\\"), attributes));
+				CHECK_TRUE(sessionClient.Exists(TCV("c:\\windows"), attributes));
+				CHECK_TRUE(IsDirectory(attributes));
+				CHECK_TRUE(!sessionClient.Exists(TCV("q:\\"), attributes))
+				CHECK_TRUE(!sessionClient.Exists(TCV("r:\\foo"), attributes))
+				#else
+				CHECK_TRUE(!sessionClient.Exists(TCV("/ergewrgergreg"), attributes));
+				CHECK_TRUE(!sessionClient.Exists(TCV("/ergergreg/h5r6tyh"), attributes));
+				#endif
+				return true;
+			});
+#else
+		return true;
+#endif
+	}
+
 }
