[UBA]
* Fixed unit test TestTraverseDir so it doesn't fail when running after other unit tests

diff --git a/Engine/Source/Programs/UnrealBuildAccelerator/Test/Private/UbaTestBasics.cpp b/Engine/Source/Programs/UnrealBuildAccelerator/Test/Private/UbaTestBasics.cpp
--- a/Engine/Source/Programs/UnrealBuildAccelerator/Test/Private/UbaTestBasics.cpp
+++ b/Engine/Source/Programs/UnrealBuildAccelerator/Test/Private/UbaTestBasics.cpp
@@ -265,22 +265,27 @@ namespace uba
 		LoggerWithWriter nullLogger(g_nullLogWriter);
 		CHECK_TRUE(!TraverseDir(nullLogger, AsView(TC("TestDir2")), [&](const DirectoryEntry&) {}, true));
 		u32 foundCount = 0;
-		CHECK_TRUE(TraverseDir(nullLogger, rootDir, [&](const DirectoryEntry&) { ++foundCount; }, true));
+
+		StringBuffer<> testDir(rootDir);
+		testDir.Append(TC("TraverseDir")).EnsureEndsWithSlash();
+		CHECK_TRUE(CreateDirectoryW(testDir.data));
+
+		CHECK_TRUE(TraverseDir(nullLogger, testDir, [&](const DirectoryEntry&) { ++foundCount; }, true));
 		CHECK_TRUE(foundCount == 0);
 
 		bool isFile = false;
-		StringBuffer<> entry(rootDir);
+		StringBuffer<> entry(testDir);
 		entry.Append(TC("Entry"));
 		FileAccessor fileHandle(logger, entry.data);
 		CHECK_TRUE(fileHandle.CreateWrite(false));
 		CHECK_TRUE(fileHandle.Close());
-		CHECK_TRUE(TraverseDir(nullLogger, rootDir, [&](const DirectoryEntry& e) { ++foundCount; isFile = !IsDirectory(e.attributes); }, true));
+		CHECK_TRUE(TraverseDir(nullLogger, testDir, [&](const DirectoryEntry& e) { ++foundCount; isFile = !IsDirectory(e.attributes); }, true));
 		CHECK_TRUE(foundCount == 1);
 		CHECK_TRUE(isFile);
 		CHECK_TRUE(DeleteFileW(entry.data));
 		CHECK_TRUE(CreateDirectoryW(entry.data));
 		foundCount = 0;
-		CHECK_TRUE(TraverseDir(nullLogger, rootDir, [&](const DirectoryEntry& e) { ++foundCount; isFile = !IsDirectory(e.attributes); }, true));
+		CHECK_TRUE(TraverseDir(nullLogger, testDir, [&](const DirectoryEntry& e) { ++foundCount; isFile = !IsDirectory(e.attributes); }, true));
 		CHECK_TRUE(foundCount == 1);
 		CHECK_TRUE(!isFile);
 		return true;
