[UBA}
* Another fix for for not handling root properly in directory table. Only repro on mac

diff --git a/Engine/Source/Programs/UnrealBuildAccelerator/Common/Private/UbaSession.cpp b/Engine/Source/Programs/UnrealBuildAccelerator/Common/Private/UbaSession.cpp
--- a/Engine/Source/Programs/UnrealBuildAccelerator/Common/Private/UbaSession.cpp
+++ b/Engine/Source/Programs/UnrealBuildAccelerator/Common/Private/UbaSession.cpp
@@ -467,14 +467,13 @@ namespace uba
 			break;
 		}
 
-		if (rit <= dirPath.data) // There are no path separators left, this is the drive
+		if (IsWindows && rit <= dirPath.data) // There are no path separators left, this is the drive
 		{
 			dirPath.count = 0;
 			return;
 		}
 
 		dirPath.count = u32(rit - dirPath.data);
-		//*(rit) = 0;
 
 		StringBuffer<> parentDirForHash;
 		parentDirForHash.Append(dirPath);
