[UBA]
* Changed some warnings/errors in SessionServer to info since they could happen in bad network conditions
* Added TraceIOEnabled which when enabled will write out I/O for drives on windows (enabled by default)
* Optimized UbaVisualizer by using PolyLine and ExtTextOut.
* Added timer to Bottleneck and added timer to report to see when uba is waiting on bottleneck
* Fixed some timers that were wrong

diff --git a/Engine/Source/Programs/UnrealBuildAccelerator/Common/Private/UbaFile.cpp b/Engine/Source/Programs/UnrealBuildAccelerator/Common/Private/UbaFile.cpp
--- a/Engine/Source/Programs/UnrealBuildAccelerator/Common/Private/UbaFile.cpp
+++ b/Engine/Source/Programs/UnrealBuildAccelerator/Common/Private/UbaFile.cpp
@@ -1116,7 +1116,9 @@ namespace uba
 
 		if (prevQueryTime)
 		{
-			outBusyPercent = (u8)(100.0 - 100.0*(double(idleTime - prevIdleTime) / double(queryTime - prevQueryTime)));
+			double busyPercent = 100.0 - 100.0*(double(idleTime - prevIdleTime) / double(queryTime - prevQueryTime));
+			busyPercent = Max(0.0, Min(100.0, busyPercent));
+			outBusyPercent = (u8)busyPercent;
 			outReadCount = readCount - prevReadCount;
 			outReadBytes = readBytes - prevReadBytes;
 			outWriteCount = writeCount - prevWriteCount;
diff --git a/Engine/Source/Programs/UnrealBuildAccelerator/Common/Private/UbaFileAccessor.cpp b/Engine/Source/Programs/UnrealBuildAccelerator/Common/Private/UbaFileAccessor.cpp
--- a/Engine/Source/Programs/UnrealBuildAccelerator/Common/Private/UbaFileAccessor.cpp
+++ b/Engine/Source/Programs/UnrealBuildAccelerator/Common/Private/UbaFileAccessor.cpp
@@ -192,12 +192,14 @@ namespace uba
 	{
 		auto& stats = KernelStats::GetCurrent();
 		ExtendedTimerScope ts(stats.writeFile);
-
 		if (!m_isWrite)
 			return m_logger.Error(TC("File %s is not opened for write"), m_fileName);
-
 		stats.writeFile.bytes += dataLen;
+		return InternalWrite(data, dataLen, offset, lastWrite);
+	}
 
+	bool FileAccessor::InternalWrite(const void* data, u64 dataLen, u64 offset, bool lastWrite)
+	{
 #if PLATFORM_WINDOWS
 		bool noBuffering = (m_flagsAndAttributes & FILE_FLAG_NO_BUFFERING) != 0;
 		if (noBuffering)
@@ -216,7 +218,7 @@ namespace uba
 				if (m_writeThroughBufferSize < WriteUnit)
 					return true;
 				m_writeThroughBufferSize = 0;
-				if (!Write(m_writeThroughBuffer, WriteUnit, 0))
+				if (!InternalWrite(m_writeThroughBuffer, WriteUnit, 0, lastWrite))
 					return false;
 			}
 		}
diff --git a/Engine/Source/Programs/UnrealBuildAccelerator/Common/Private/UbaFileMapping.cpp b/Engine/Source/Programs/UnrealBuildAccelerator/Common/Private/UbaFileMapping.cpp
--- a/Engine/Source/Programs/UnrealBuildAccelerator/Common/Private/UbaFileMapping.cpp
+++ b/Engine/Source/Programs/UnrealBuildAccelerator/Common/Private/UbaFileMapping.cpp
@@ -38,7 +38,8 @@ namespace uba
 		if (flProtect == PAGE_READWRITE)
 		{
 			// Experiment to try to prevent lock happening on AWS servers when lots of helpers are sending back obj files.
-			BottleneckScope bs(g_createFileHandleBottleneck);
+			Timer timer;
+			BottleneckScope bs(g_createFileHandleBottleneck, timer);
 			res = ::CreateFileMappingW(hFile, NULL, flProtect, dwMaximumSizeHigh, dwMaximumSizeLow, lpName);
 		}
 		else
diff --git a/Engine/Source/Programs/UnrealBuildAccelerator/Common/Private/UbaNetworkServer.cpp b/Engine/Source/Programs/UnrealBuildAccelerator/Common/Private/UbaNetworkServer.cpp
--- a/Engine/Source/Programs/UnrealBuildAccelerator/Common/Private/UbaNetworkServer.cpp
+++ b/Engine/Source/Programs/UnrealBuildAccelerator/Common/Private/UbaNetworkServer.cpp
@@ -333,6 +333,8 @@ namespace uba
 			Worker* worker = conn.m_server.PopWorker();
 			if (!worker)
 				return false;
+			if (!worker->m_context)
+				return conn.m_server.m_logger.Error(TC("Client %u - Popped worker which has no context"), conn.m_client->id);
 			auto& wc = *worker->m_context;
 			wc.id = messageId;
 			wc.serviceId = serviceId;
diff --git a/Engine/Source/Programs/UnrealBuildAccelerator/Common/Private/UbaSession.cpp b/Engine/Source/Programs/UnrealBuildAccelerator/Common/Private/UbaSession.cpp
--- a/Engine/Source/Programs/UnrealBuildAccelerator/Common/Private/UbaSession.cpp
+++ b/Engine/Source/Programs/UnrealBuildAccelerator/Common/Private/UbaSession.cpp
@@ -2855,7 +2855,7 @@ namespace uba
 
 		// This is to not kill I/O when writing lots of pdb/dlls in parallel
 		#if PLATFORM_WINDOWS
-		BottleneckScope scope(m_writeFilesBottleneck);
+		BottleneckScope scope(m_writeFilesBottleneck, Stats().waitBottleneck);
 		#endif
 
 		Atomic<bool> success = true;
diff --git a/Engine/Source/Programs/UnrealBuildAccelerator/Common/Private/UbaSessionServer.cpp b/Engine/Source/Programs/UnrealBuildAccelerator/Common/Private/UbaSessionServer.cpp
--- a/Engine/Source/Programs/UnrealBuildAccelerator/Common/Private/UbaSessionServer.cpp
+++ b/Engine/Source/Programs/UnrealBuildAccelerator/Common/Private/UbaSessionServer.cpp
@@ -105,6 +105,7 @@ namespace uba
 			table->GetValueAsBool(remoteLogEnabled, TC("RemoteLogEnabled"));
 			table->GetValueAsBool(remoteTraceEnabled, TC("RemoteTraceEnabled"));
 			table->GetValueAsBool(nameToHashTableEnabled, TC("NameToHashTableEnabled"));
+			table->GetValueAsBool(traceIOEnabled, TC("TraceIOEnabled"));
 		}
 	}
 
@@ -169,6 +170,7 @@ namespace uba
 		m_memKillLoadPercent = info.memKillLoadPercent;
 		m_remoteLogEnabled = info.remoteLogEnabled;
 		m_remoteTraceEnabled = info.remoteTraceEnabled;
+		m_traceIOEnabled = info.traceIOEnabled;
 
 		if (m_resetCas)
 			m_storage.Reset();
@@ -1071,7 +1073,8 @@ namespace uba
 				if (!h.IsValid())
 				{
 					success = false;
-					return m_logger.Warning(TC("Failed to find process for id %u when receiving SendFileToServer message"), processId);
+					m_logger.Info(TC("Failed to find process for id %u when receiving SendFileToServer message"), processId);
+					return false;
 				}
 
 				RootsHandle rootsHandle = h.GetStartInfo().rootsHandle;
@@ -1394,7 +1397,10 @@ namespace uba
 		u32 processId = u32(reader.Read7BitEncoded());
 		ProcessHandle h(GetProcess(processId));
 		if (!h.IsValid())
-			return m_logger.Error(TC("Failed to find process for id %u when receiving custom message"), processId);
+		{
+			m_logger.Info(TC("Failed to find process for id %u when receiving custom message"), processId);
+			return false;
+		}
 		auto& process = *(RemoteProcess*)h.m_process;
 		auto& inputs = process.m_trackedInputs;
 		u64 size = inputs.size();
@@ -1413,8 +1419,8 @@ namespace uba
 		ProcessHandle h = ProcessRemoved(processId);
 		if (!h.m_process)
 		{
-			m_logger.Warning(TC("Client finished process with id %u that is not found on server"), processId);
-			return true;
+			m_logger.Info(TC("Client finished process with id %u that is not found on server"), processId);
+			return false;
 		}
 		auto& process = *(RemoteProcess*)h.m_process;
 
@@ -1422,8 +1428,8 @@ namespace uba
 		if (!m_activeRemoteProcesses.erase(&process))
 		{
 			cs2.Leave();
-			m_logger.Warning(TC("Got finished process but process was not in active remote processes. Was there a disconnect happening directly after but executed before?"));
-			return true;
+			m_logger.Info(TC("Got finished process but process was not in active remote processes. Was there a disconnect happening directly after but executed before?"));
+			return false;
 		}
 		u32 sessionIndex = process.m_sessionId - 1;
 		if (sessionIndex >= m_clientSessions.size())
@@ -1595,7 +1601,10 @@ namespace uba
 		u32 prevExitCode = reader.ReadU32();
 		ProcessHandle h(GetProcess(processId));
 		if (!h.IsValid())
-			return m_logger.Error(TC("Failed to find process for id %u when receiving GetNextProcess message"), processId);
+		{
+			m_logger.Info(TC("Failed to find process for id %u when receiving GetNextProcess message"), processId);
+			return false;
+		}
 
 		auto& remoteProcess = *(RemoteProcess*)h.m_process;
 		SCOPED_FUTEX(remoteProcess.m_exitedLock, exitedLock);
@@ -1625,7 +1634,10 @@ namespace uba
 		u32 processId = reader.ReadU32();
 		ProcessHandle h(GetProcess(processId));
 		if (!h.IsValid())
-			return m_logger.Error(TC("Failed to find process for id %u when receiving custom message"), processId);
+		{
+			m_logger.Info(TC("Failed to find process for id %u when receiving custom message"), processId);
+			return false;
+		}
 		auto& remoteProcess = *(RemoteProcess*)h.m_process;
 		SCOPED_FUTEX(remoteProcess.m_exitedLock, exitedLock);
 		CustomMessage(remoteProcess, reader, writer);
@@ -1637,7 +1649,10 @@ namespace uba
 		u32 processId = reader.ReadU32();
 		ProcessHandle h(GetProcess(processId));
 		if (!h.IsValid())
-			return m_logger.Error(TC("Failed to find process for id %u when receiving update environment message"), processId);
+		{
+			m_logger.Info(TC("Failed to find process for id %u when receiving update environment message"), processId);
+			return false;
+		}
 		StringBuffer<> reason;
 		reader.ReadString(reason);
 		m_trace.ProcessEnvironmentUpdated(processId, reason, reader.GetPositionData(), reader.GetLeft(), ToView(h.GetStartInfo().breadcrumbs));
@@ -2171,7 +2186,7 @@ namespace uba
 		u64 memTotal = m_memTotal;
 
 
-		if (m_detailedTrace)
+		if (m_traceIOEnabled)
 		{
 			for (auto& volume : m_volumeCache.volumes)
 			{
@@ -2182,7 +2197,10 @@ namespace uba
 				u64 readBytes;
 				u32 writeCount;
 				u64 writeBytes;
-				volume.UpdateStats(busyPercent, readCount, readBytes, writeCount, writeBytes);
+				if (!volume.UpdateStats(busyPercent, readCount, readBytes, writeCount, writeBytes))
+					continue;
+				if (!busyPercent && !readCount && !readBytes && !writeCount && !writeBytes)
+					continue;
 				m_trace.DriveUpdate(volume.drives[0], busyPercent, readCount, readBytes, writeCount, writeBytes);
 			}
 		}
diff --git a/Engine/Source/Programs/UnrealBuildAccelerator/Common/Private/UbaStorage.cpp b/Engine/Source/Programs/UnrealBuildAccelerator/Common/Private/UbaStorage.cpp
--- a/Engine/Source/Programs/UnrealBuildAccelerator/Common/Private/UbaStorage.cpp
+++ b/Engine/Source/Programs/UnrealBuildAccelerator/Common/Private/UbaStorage.cpp
@@ -1220,7 +1220,7 @@ namespace uba
 	StorageImpl::StorageImpl(const StorageCreateInfo& info, const tchar* logPrefix)
 	:	m_workManager(info.workManager)
 	,	m_logger(info.writer, logPrefix)
-	,	m_activeCopyOrLinkEvent(false)
+	,	m_maxParallelCopyOrLinkBottleneck(info.maxParallelCopyOrLink)
 	,	m_casDataBuffer(m_logger, info.workManager)
 	{
 		m_casCapacityBytes = info.casCapacityBytes;
@@ -1248,8 +1248,6 @@ namespace uba
 
 		m_casDataBuffer.AddTransient(TC("CasData"), info.keepTransientDataMapped);
 
-		m_maxParallelCopyOrLink = info.maxParallelCopyOrLink;
-
 		m_casCompressor = DefaultCompressor;
 		m_casCompressionLevel = DefaultCompressionLevel;
 
@@ -2345,28 +2343,9 @@ namespace uba
 
 		auto& stats = Stats();
 		
-		u32 maxParallelCopyOrLink = m_maxParallelCopyOrLink;
-		{
-			SCOPED_FUTEX(m_activeCopyOrLinkLock, lock);
-			if (m_activeCopyOrLink > maxParallelCopyOrLink)
-			{
-				TimerScope ts(stats.copyOrLinkWait);
-				while (m_activeCopyOrLink > maxParallelCopyOrLink)
-				{
-					lock.Leave();
-					m_activeCopyOrLinkEvent.IsSet(2000);
-					lock.Enter();
-				}
-			}
-			++m_activeCopyOrLink;
-		}
-		
-		auto activeGuard = MakeGuard([&]()
-			{
-				SCOPED_FUTEX(m_activeCopyOrLinkLock, lock);
-				--m_activeCopyOrLink;
-				m_activeCopyOrLinkEvent.Set();
-			});
+		#if PLATFORM_WINDOWS
+		BottleneckScope bs(m_maxParallelCopyOrLinkBottleneck, stats.copyOrLinkWait);
+		#endif
 
 		StringBuffer<> forKey;
 		forKey.Append(destination);
@@ -2445,13 +2424,6 @@ namespace uba
 				fileEntry.verified = false;
 				fileEntry.isTemp = isTemp;
 
-				// This is to reduce number of active CreateFiles.. seems like machines don't like tons of CreateFile at the same time
-				#if PLATFORM_WINDOWS
-				constexpr u32 bottleneckMax = 16;
-				static Bottleneck bottleneck(bottleneckMax);
-				BottleneckScope scope(bottleneck);
-				#endif
-
 				u64 sizeOnDisk = decompressedSize;
 
 				auto decompressToMemory = [&](u8* destinationMemory)
diff --git a/Engine/Source/Programs/UnrealBuildAccelerator/Common/Public/UbaBottleneck.h b/Engine/Source/Programs/UnrealBuildAccelerator/Common/Public/UbaBottleneck.h
--- a/Engine/Source/Programs/UnrealBuildAccelerator/Common/Public/UbaBottleneck.h
+++ b/Engine/Source/Programs/UnrealBuildAccelerator/Common/Public/UbaBottleneck.h
@@ -5,6 +5,7 @@
 #include "UbaEvent.h"
 #include "UbaPlatform.h"
 #include "UbaSynchronization.h"
+#include "UbaTimer.h"
 
 namespace uba
 {
@@ -18,7 +19,7 @@ namespace uba
 	{
 		Bottleneck(u32 mc) : m_maxCount(mc) {}
 
-		void Enter(BottleneckTicket& ticket)
+		void Enter(BottleneckTicket& ticket, Timer& timer)
 		{
 			SCOPED_FUTEX(m_lock, lock);
 			if (m_activeCount < m_maxCount)
@@ -39,6 +40,7 @@ namespace uba
 
 			lock.Leave();
 
+			TimerScope ts(timer);
 			ticket.ev.IsSet();
 		}
 
@@ -68,7 +70,7 @@ namespace uba
 
 	struct BottleneckScope
 	{
-		BottleneckScope(Bottleneck& b) : bottleneck(b) { b.Enter(ticket); }
+		BottleneckScope(Bottleneck& b, Timer& timer) : bottleneck(b) { b.Enter(ticket, timer); }
 		~BottleneckScope() { bottleneck.Leave(ticket); }
 
 		Bottleneck& bottleneck;
diff --git a/Engine/Source/Programs/UnrealBuildAccelerator/Common/Public/UbaFileAccessor.h b/Engine/Source/Programs/UnrealBuildAccelerator/Common/Public/UbaFileAccessor.h
--- a/Engine/Source/Programs/UnrealBuildAccelerator/Common/Public/UbaFileAccessor.h
+++ b/Engine/Source/Programs/UnrealBuildAccelerator/Common/Public/UbaFileAccessor.h
@@ -33,6 +33,7 @@ namespace uba
 		bool GetFileBasicInformationByHandle(FileBasicInformation& out);
 
 	private:
+		bool InternalWrite(const void* data, u64 dataLen, u64 offset, bool lastWrite);
 		bool InternalClose(bool success, u64* lastWriteTime);
 		Logger& m_logger;
 		const tchar* m_fileName;
diff --git a/Engine/Source/Programs/UnrealBuildAccelerator/Common/Public/UbaSessionServer.h b/Engine/Source/Programs/UnrealBuildAccelerator/Common/Public/UbaSessionServer.h
--- a/Engine/Source/Programs/UnrealBuildAccelerator/Common/Public/UbaSessionServer.h
+++ b/Engine/Source/Programs/UnrealBuildAccelerator/Common/Public/UbaSessionServer.h
@@ -203,6 +203,7 @@ namespace uba
 		bool m_allowKillOnMem = false;
 		bool m_remoteLogEnabled = false;
 		bool m_remoteTraceEnabled = false;
+		bool m_traceIOEnabled = false;
 
 		SessionServer(const SessionServer&) = delete;
 		void operator=(const SessionServer&) = delete;
diff --git a/Engine/Source/Programs/UnrealBuildAccelerator/Common/Public/UbaSessionServerCreateInfo.h b/Engine/Source/Programs/UnrealBuildAccelerator/Common/Public/UbaSessionServerCreateInfo.h
--- a/Engine/Source/Programs/UnrealBuildAccelerator/Common/Public/UbaSessionServerCreateInfo.h
+++ b/Engine/Source/Programs/UnrealBuildAccelerator/Common/Public/UbaSessionServerCreateInfo.h
@@ -25,5 +25,6 @@ namespace uba
 		bool allowKillOnMem = false;
 		bool remoteLogEnabled = false; // If Uba is built in debug, then the logs will be sent back to server
 		bool remoteTraceEnabled = false; // If this is true, the agents will run trace and send the .uba file back to server
+		bool traceIOEnabled = true; // Enable IO information for storage 
 	};
 }
diff --git a/Engine/Source/Programs/UnrealBuildAccelerator/Common/Public/UbaStats.h b/Engine/Source/Programs/UnrealBuildAccelerator/Common/Public/UbaStats.h
--- a/Engine/Source/Programs/UnrealBuildAccelerator/Common/Public/UbaStats.h
+++ b/Engine/Source/Programs/UnrealBuildAccelerator/Common/Public/UbaStats.h
@@ -73,6 +73,7 @@ namespace uba
 		UBA_SESSION_STAT(Timer, createMmapFromFile, 12) \
 		UBA_SESSION_STAT(Timer, waitMmapFromFile, 12) \
 		UBA_SESSION_STAT(Timer, getLongNameMsg, 31) \
+		UBA_SESSION_STAT(Timer, waitBottleneck, 40) \
 
     struct SessionStats
 	{
diff --git a/Engine/Source/Programs/UnrealBuildAccelerator/Common/Public/UbaStorage.h b/Engine/Source/Programs/UnrealBuildAccelerator/Common/Public/UbaStorage.h
--- a/Engine/Source/Programs/UnrealBuildAccelerator/Common/Public/UbaStorage.h
+++ b/Engine/Source/Programs/UnrealBuildAccelerator/Common/Public/UbaStorage.h
@@ -2,6 +2,7 @@
 
 #pragma once
 
+#include "UbaBottleneck.h"
 #include "UbaFile.h"
 #include "UbaFileMapping.h"
 #include "UbaLogger.h"
@@ -253,10 +254,7 @@ namespace uba
 
 		ExternalFileMappingsProvider m_externalFileMappingsProvider;
 
-		u32 m_maxParallelCopyOrLink;
-		Futex m_activeCopyOrLinkLock;
-		Event m_activeCopyOrLinkEvent;
-		u32 m_activeCopyOrLink = 0;
+		Bottleneck m_maxParallelCopyOrLinkBottleneck;
 
 		Futex m_casTableLoadSaveLock;
 		bool m_casTableLoaded = false;
diff --git a/Engine/Source/Programs/UnrealBuildAccelerator/Common/Public/UbaTrace.h b/Engine/Source/Programs/UnrealBuildAccelerator/Common/Public/UbaTrace.h
--- a/Engine/Source/Programs/UnrealBuildAccelerator/Common/Public/UbaTrace.h
+++ b/Engine/Source/Programs/UnrealBuildAccelerator/Common/Public/UbaTrace.h
@@ -69,7 +69,7 @@ namespace uba
 		#undef UBA_TRACE_TYPE
 	};
 
-	static constexpr u32 TraceVersion = 39;
+	static constexpr u32 TraceVersion = 40;
 	static constexpr u32 TraceReadCompatibilityVersion = 6;
 
 	class Trace : public WorkTracker
diff --git a/Engine/Source/Programs/UnrealBuildAccelerator/Visualizer/Private/UbaTraceReader.cpp b/Engine/Source/Programs/UnrealBuildAccelerator/Visualizer/Private/UbaTraceReader.cpp
--- a/Engine/Source/Programs/UnrealBuildAccelerator/Visualizer/Private/UbaTraceReader.cpp
+++ b/Engine/Source/Programs/UnrealBuildAccelerator/Visualizer/Private/UbaTraceReader.cpp
@@ -537,6 +537,16 @@ namespace uba
 			session.cpuLoad.push_back(cpuLoad);
 			session.connectionCount.push_back(connectionCount);
 			
+			for (auto& pair : session.drives)
+			{
+				auto& drive = pair.second;
+				drive.busyPercent.resize(session.updates.size());
+				drive.readCount.resize(session.updates.size());
+				drive.writeCount.resize(session.updates.size());
+				drive.readBytes.resize(session.updates.size());
+				drive.writeBytes.resize(session.updates.size());
+			}
+
 			if (session.prevUpdateTime)
 			{
 				session.highestSendPerS = Max(session.highestSendPerS, float(totalSend - session.prevSend)/TimeToS(time - session.prevUpdateTime));
@@ -840,8 +850,8 @@ namespace uba
 				return false;
 			if (auto session = GetSession(out, clientUid))
 			{
-				session->fetchedFiles.push_back({ CasKeyZero, fileSize, L"", time, time });
 				session->fetchedFilesBytes += fileSize;
+				++session->fetchedFilesCount;
 			}
 			break;
 		}
@@ -926,6 +936,7 @@ namespace uba
 					auto& file = session->fetchedFiles[findIt->second];
 					file.stop = time;
 					session->fetchedFilesBytes += file.size;
+					++session->fetchedFilesCount;
 					session->fetchedFilesActive.erase(findIt);
 					break;
 				}
@@ -961,6 +972,7 @@ namespace uba
 				session->storedFilesActive.try_emplace(key, u32(session->storedFiles.size()));
 				session->storedFiles.push_back({ key, size, hint, time, ~u64(0) });
 				session->storedFilesBytes += size;
+				++session->storedFilesCount;
 			}
 			break;
 		}
@@ -970,8 +982,8 @@ namespace uba
 			u64 fileSize = reader.Read7BitEncoded();
 			if (auto session = GetSession(out, clientUid))
 			{
-				session->storedFiles.push_back({ CasKeyZero, fileSize, L"", time, time });
 				session->storedFilesBytes += fileSize;
+				++session->storedFilesCount;
 			}
 			break;
 		}
@@ -1102,8 +1114,21 @@ namespace uba
 
 			if (out.sessions.empty())
 				break;
-			auto& drive = out.sessions[0].drives[driveLetter];
+			auto& session = out.sessions[0];
+			auto& drive = session.drives[driveLetter];
+
+			if (drive.busyPercent.empty())
+			{
+				u64 updatesCount = session.updates.size(); // We get these before update
+				drive.busyPercent.resize(updatesCount);
+				drive.readCount.resize(updatesCount);
+				drive.writeCount.resize(updatesCount);
+				drive.readBytes.resize(updatesCount);
+				drive.writeBytes.resize(updatesCount);
+			}
+
 			drive.busyHighest = Max(drive.busyHighest, busyPercent);
+
 			drive.busyPercent.push_back(busyPercent);
 			drive.totalReadCount += readCount;
 			drive.totalWriteCount += writeCount;
diff --git a/Engine/Source/Programs/UnrealBuildAccelerator/Visualizer/Private/UbaTraceReader.h b/Engine/Source/Programs/UnrealBuildAccelerator/Visualizer/Private/UbaTraceReader.h
--- a/Engine/Source/Programs/UnrealBuildAccelerator/Visualizer/Private/UbaTraceReader.h
+++ b/Engine/Source/Programs/UnrealBuildAccelerator/Visualizer/Private/UbaTraceReader.h
@@ -119,6 +119,8 @@ namespace uba
 			TString notification;
 			u64 fetchedFilesBytes = 0;
 			u64 storedFilesBytes = 0;
+			u32 fetchedFilesCount = 0;
+			u32 storedFilesCount = 0;
 			u32 maxVisibleFiles = 0;
 			u32 fullNameWidth = 0;
 
diff --git a/Engine/Source/Programs/UnrealBuildAccelerator/Visualizer/Private/UbaVisualizer.cpp b/Engine/Source/Programs/UnrealBuildAccelerator/Visualizer/Private/UbaVisualizer.cpp
--- a/Engine/Source/Programs/UnrealBuildAccelerator/Visualizer/Private/UbaVisualizer.cpp
+++ b/Engine/Source/Programs/UnrealBuildAccelerator/Visualizer/Private/UbaVisualizer.cpp
@@ -174,7 +174,7 @@ namespace uba
 				tr.left += line.left;
 				tr.top += line.top;
 				SetTextColor(hdc, line.color);
-				DrawTextW(hdc, line.str.data(), u32(line.str.size()), &tr, DT_SINGLELINE);
+				DrawTextW(hdc, line.str.data(), u32(line.str.size()), &tr, DT_SINGLELINE|DT_NOPREFIX);
 			}
 		}
 
@@ -946,7 +946,7 @@ namespace uba
 				for (u32 i=0; i!=lineCount; ++i)
 				{
 					RECT rect = { 0, 0, 0, 0 };
-					DrawTextW(hdc, lines[i].data, lines[i].count, &rect, DT_SINGLELINE|DT_NOCLIP|DT_CALCRECT);
+					DrawTextW(hdc, lines[i].data, lines[i].count, &rect, DT_SINGLELINE|DT_NOPREFIX|DT_NOCLIP|DT_CALCRECT);
 					int top = (clientRect.bottom - rect.bottom)/2 + i*20;
 					rect.left = (clientRect.right - rect.right)/2;
 					rect.right = (clientRect.right + rect.right)/2;
@@ -959,7 +959,7 @@ namespace uba
 						InflateRect(&r, 4, 2);
 						fillRect(r, m_tooltipBackgroundBrush);
 					}
-					DrawTextW(hdc, lines[i].data, lines[i].count, &rect, DT_SINGLELINE|DT_NOCLIP);
+					DrawTextW(hdc, lines[i].data, lines[i].count, &rect, DT_SINGLELINE|DT_NOPREFIX|DT_NOCLIP);
 				}
 			};
 
@@ -1071,7 +1071,7 @@ namespace uba
 				str.Appendf(TC("Unsupported trace version %u (Versions supported are %u to %u)"), m_traceView.version, TraceReadCompatibilityVersion, TraceVersion);
 				HFONT handle = CreateFontW(-20, 0, 0, 0, FW_NORMAL, 0, 0, 0, ANSI_CHARSET, OUT_DEFAULT_PRECIS, CLIP_DEFAULT_PRECIS, CLEARTYPE_QUALITY, DEFAULT_PITCH, m_config.fontName.c_str());
 				SelectObject(hdc, handle);
-				DrawTextW(hdc, str.data, str.count, (RECT*)&clientRect, DT_CENTER|DT_SINGLELINE|DT_VCENTER);
+				DrawTextW(hdc, str.data, str.count, (RECT*)&clientRect, DT_CENTER|DT_NOPREFIX|DT_SINGLELINE|DT_VCENTER);
 				DeleteObject(handle);
 			}
 			else
@@ -1275,10 +1275,15 @@ namespace uba
 			{
 				if (posY + GraphHeight >= progressRect.top && posY + GraphHeight - 5 < progressRect.bottom)
 				{
+					bool isRemote = !isFirst;
 					int posX = int(m_scrollPosX) + progressRect.left;
+					if (isRemote)
+						posX -= 15; // TODO: This should be calculated instead.. remote graph stats are much older
 					int graphBaseY = posY + GraphHeight - 4;
 					double graphHeight = double(GraphHeight - 2);
 
+					Vector<POINT> line;
+
 					auto drawGraph = [&](auto& values, auto maxValue, double scale, HPEN pen, bool accumulatingValue, int offsetY = 0)
 						{
 							SelectObject(hdc, pen);
@@ -1289,6 +1294,7 @@ namespace uba
 							int prevX = 0;
 							int prevY = 0;
 							u64 e = Min(values.size(), session.updates.size());
+							line.clear();
 							for (u64 i=0; i!=e; ++i)
 							{
 								u64 updateTime = session.updates[i];
@@ -1320,12 +1326,15 @@ namespace uba
 								if (x > clientRect.left)
 								{
 									if (isFirstUpdate)
-										MoveToEx(hdc, x, x, NULL);
+									{
+										line.push_back({x, y}); //MoveToEx(hdc, x, y, NULL);
+										isFirstDraw = false;
+									}
 									else
 									{
 										if (isFirstDraw)
-											MoveToEx(hdc, prevX, prevY, NULL);
-										LineTo(hdc, x, y);
+											line.push_back({prevX, prevY}); // MoveToEx(hdc, prevX, prevY, NULL);
+										line.push_back({x, y}); // LineTo(hdc, x, y);
 										isFirstDraw = false;
 									}
 								}
@@ -1337,6 +1346,8 @@ namespace uba
 								prevValue = value;
 								prevTime = updateTime;
 							}
+							if (line.size() > 1)
+								Polyline(hdc, line.data(), int(line.size()));
 						};
 
 					if (m_config.showNetworkStats && session.highestSendPerS != 0 && session.highestRecvPerS != 0)
@@ -1361,14 +1372,28 @@ namespace uba
 
 			if (m_config.showDetailedData)
 			{
-				auto drawText = [&](const StringBufferBase& text, RECT& rect)
+				auto drawText = [&](const StringBufferBase& text, RECT& rect, u32* outWidth)
 					{
+						if (rect.top > progressRect.bottom)
+							return;
 						bool selected = m_fetchedFilesSelected == processLocation.sessionIndex && text.StartsWith(TC("Fetched Files"));
 						if (selected)
 							SetBkMode(hdc, OPAQUE);
-						DrawTextW(hdc, text.data, text.count, &rect, DT_SINGLELINE);
+						if (rect.bottom > progressRect.bottom)
+						{
+							rect.bottom = progressRect.bottom;
+							ExtTextOutW(hdc, rect.left, rect.top, ETO_CLIPPED, &rect, text.data, text.count, NULL);
+						}
+						else
+							ExtTextOutW(hdc, rect.left, rect.top, 0, NULL, text.data, text.count, NULL);
 						if (selected)
 							SetBkMode(hdc, TRANSPARENT);
+						if (outWidth)
+						{
+							SIZE s;
+							GetTextExtentPoint32W(hdc, text.data, text.count, &s);
+							*outWidth = u32(s.cx);
+						}
 					};
 				bool isRemote = processLocation.sessionIndex != 0;
 				PaintDetailedStats(posY, progressRect, session, isRemote, playTime, drawText);
@@ -1872,10 +1897,10 @@ namespace uba
 					SetTextColor(hdc, color);
 					RECT r{left, top, left+200, top+200};
 					auto strLen = TStrlen(text);
-					DrawTextW(hdc, text, strLen, &r, DT_SINGLELINE|DT_NOCLIP|DT_CALCRECT);
+					DrawTextW(hdc, text, strLen, &r, DT_SINGLELINE|DT_NOCLIP|DT_NOPREFIX|DT_CALCRECT);
 					left -= (r.right - r.left) + 5;
 					r.left = left;
-					DrawTextW(hdc, text, strLen, &r, DT_SINGLELINE|DT_NOCLIP);
+					DrawTextW(hdc, text, strLen, &r, DT_SINGLELINE|DT_NOCLIP|DT_NOPREFIX);
 				};
 
 			if (m_config.showCpuMemStats)
@@ -2166,7 +2191,7 @@ namespace uba
 
 
 			for (auto& pair : m_stats.drives)
-				logger.SetColor(m_driveColor).Info(L"  %c: %u%%", pair.first, pair.second);
+				logger.SetColor(m_driveColor).Info(L"  %c: %u%% R:%sps W:%sps", pair.first, pair.second.busyPercent, BytesToText(pair.second.readPerSecond).str, BytesToText(pair.second.writePerSecond).str);
 
 			logger.AddSpace(3);
 			logger.DrawAtCursor();
@@ -2441,6 +2466,8 @@ namespace uba
 		progressRect.left += m_progressRectLeft;
 
 		SelectObject(hdc, m_textPen);
+		Vector<DWORD> lines;
+		Vector<POINT> points;
 
 		while (true)
 		{
@@ -2474,21 +2501,25 @@ namespace uba
 					textRect.bottom = textRect.top + m_activeFont.height;
 					textRect.right = pos + 20;
 					textRect.left = pos - 20;
-					DrawTextW(hdc, buffer.data, buffer.count, &textRect, DT_SINGLELINE | DT_CENTER);
+					ExtTextOutW(hdc, textRect.left, textRect.top, 0, NULL, buffer.data, buffer.count, NULL);
 				}
 			}
 			if (!(number % lineStepSize))
 			{
-				MoveToEx(hdc, pos, top, NULL);
-				LineTo(hdc, pos, lineBottom);
+				points.push_back({pos, top});
+				points.push_back({pos, lineBottom});
+				lines.push_back(2);
 			}
 
 			++number;
 			++index;
 		}
 
-		MoveToEx(hdc, m_contentWidth, top - 25, NULL);
-		LineTo(hdc, m_contentWidth, top);
+		points.push_back({m_contentWidth, top - 25});
+		points.push_back({m_contentWidth, top});
+		lines.push_back(2);
+
+		PolyPolyline(hdc, points.data(), lines.data(), int(lines.size()));
 
 		/*
 		POINT cursorPos;
@@ -2508,29 +2539,51 @@ namespace uba
 		int stepY = m_activeFont.height;
 		int startPosY = posY;
 		int posX = progressRect.left + 5;
-		RECT textRect;
-		textRect.top = posY;
-		textRect.bottom = posY + 20;
-		textRect.left = posX;
-		textRect.right = posX + 1000;
+		int maxY = posY + stepY;
+		u32 maxTextWidth = 0;
 
+		RECT textRect;
 		auto drawText = [&](const wchar_t* format, ...)
 			{
 				textRect.top = posY;
-				textRect.bottom = posY + stepY;
+				textRect.bottom = posY + 20;
+				textRect.left = posX;
+				textRect.right = posX + 1000;
 				posY += stepY;
 				StringBuffer<> str;
 				va_list arg;
 				va_start(arg, format);
 				str.Append(format, arg);
 				va_end(arg);
-				drawTextFunc(str, textRect);
+				u32 lastTextWidth = 0;
+				drawTextFunc(str, textRect, &lastTextWidth);
+				maxTextWidth = Max(maxTextWidth, lastTextWidth);
+				maxY = Max(maxY, posY);
 			};
 
 		if (isRemote)
 		{
 			drawText(L"Finished Processes: %u", session.processExitedCount);
 			drawText(L"Active Processes: %u", session.processActiveCount);
+			drawText(L"ClientId: %u  TcpCount: %u", session.clientUid.data1, session.connectionCount.empty() ? 1 : session.connectionCount.back());
+
+			if (session.disconnectTime == ~u64(0))
+			{
+				if (session.proxyCreated)
+					drawText(L"Proxy(HOSTED): %ls", session.proxyName.c_str());
+				else if (!session.proxyName.empty())
+					drawText(L"Proxy: %ls", session.proxyName.c_str());
+				else
+					drawText(L"Proxy: None");
+			}
+
+			bool hasFileDetails = !session.fetchedFiles.empty() || !session.storedFiles.empty();
+
+			if (!hasFileDetails)
+			{
+				posY = startPosY;
+				posX += maxTextWidth + 15;
+			}
 
 			if (!session.updates.empty())
 			{
@@ -2545,21 +2598,9 @@ namespace uba
 					sendPerS = u64(float(updateSend - session.prevSend) / duration);
 					recvPerS = u64(float(updateRecv - session.prevRecv) / duration);
 				}
-				drawText(L"ClientId: %u  TcpCount: %u", session.clientUid.data1, session.connectionCount.back());
-				drawText(L"Recv: %ls (%sps)", BytesToText(updateRecv), BytesToText(recvPerS*8));
-				drawText(L"Send: %ls (%sps)", BytesToText(updateSend), BytesToText(sendPerS*8));
-			}
-
-			if (session.disconnectTime == ~u64(0))
-			{
-				if (session.proxyCreated)
-					drawText(L"Proxy(HOSTED): %ls", session.proxyName.c_str());
-				else if (!session.proxyName.empty())
-					drawText(L"Proxy: %ls", session.proxyName.c_str());
-				else
-					drawText(L"Proxy: None");
+				drawText(L"Recv: %ls (%sit/s)", BytesToText(updateRecv), BytesToText(recvPerS*8));
+				drawText(L"Send: %ls (%sit/s)", BytesToText(updateSend), BytesToText(sendPerS*8));
 			}
-			int posY1 = posY;
 
 			int fileWidth = 700;
 
@@ -2597,14 +2638,18 @@ namespace uba
 				fetchedFiles.push_back(session.fetchedFiles[kv.second]);
 			std::sort(fetchedFiles.begin(), fetchedFiles.end(), [](const TraceView::FileTransfer& a, const TraceView::FileTransfer& b) { return a.start > b.start; });
 
-			posY = startPosY;
-			posX += 150;
-			drawFiles(L"Fetched", fetchedFiles, session.fetchedFiles.size(), session.fetchedFilesBytes, session.fetchedFilesActive.size(), session.maxVisibleFiles);
-			int posY2 = posY;
-			posY = startPosY;
-			posX += fileWidth;
-			drawFiles(L"Stored", session.storedFiles, session.storedFiles.size(), session.storedFilesBytes, session.storedFilesActive.size(), session.maxVisibleFiles);
-			posY = Max(posY, Max(posY1, posY2));
+			if (hasFileDetails)
+			{
+				posY = startPosY;
+				posX += maxTextWidth + 15;
+			}
+			drawFiles(L"Fetched", fetchedFiles, session.fetchedFilesCount, session.fetchedFilesBytes, session.fetchedFilesActive.size(), session.maxVisibleFiles);
+			if (hasFileDetails)
+			{
+				posY = startPosY;
+				posX += fileWidth;
+			}
+			drawFiles(L"Stored", session.storedFiles, session.storedFilesCount, session.storedFilesBytes, session.storedFilesActive.size(), session.maxVisibleFiles);
 		}
 		else
 		{
@@ -2612,7 +2657,7 @@ namespace uba
 			drawText(L"Active Processes: %u (local: %u)", m_traceView.totalProcessActiveCount, session.processActiveCount);
 			drawText(L"Active Helpers: %u", Max(1u, m_traceView.activeSessionCount) - 1);
 
-			if (!session.updates.empty())
+			if (session.highestSendPerS || session.highestRecvPerS)
 			{
 				auto updateTime = session.updates.back();
 				auto updateSend = session.networkSend.back();
@@ -2627,11 +2672,38 @@ namespace uba
 						sendPerS = u64(float(updateSend - session.prevSend) / duration);
 						recvPerS = u64(float(updateRecv - session.prevRecv) / duration);
 					}
-					drawText(L"Recv: %ls (%sps)", BytesToText(updateRecv), BytesToText(recvPerS));
-					drawText(L"Send: %ls (%sps)", BytesToText(updateSend), BytesToText(sendPerS));
+
+					drawText(L"Recv: %ls (%sit/s)", BytesToText(updateRecv), BytesToText(recvPerS));
+					drawText(L"Send: %ls (%sit/s)", BytesToText(updateSend), BytesToText(sendPerS));
+				}
+			}
+
+			if (!session.updates.empty())
+			{
+				posY = startPosY;
+				posX += maxTextWidth + 10;
+
+				for (auto& pair : session.drives)
+				{
+					auto& drive = pair.second;
+
+					u64 readPerS = 0;
+					u64 writePerS = 0;
+
+					u64 updateCount = session.updates.size();
+					float duration = TimeToS(session.updates.back() - session.prevUpdateTime);
+					if (duration != 0 && updateCount > 1)
+					{
+						readPerS = u64(float(drive.readBytes.back()) / duration);
+						writePerS = u64(float(drive.writeBytes.back()) / duration);
+					}
+				
+					drawText(L"%c: Rd %ls (%s/s) Wr %ls (%s/s) ", pair.first, BytesToText(drive.totalReadBytes).str, BytesToText(readPerS).str, BytesToText(drive.totalWriteBytes).str, BytesToText(writePerS).str);
 				}
 			}
 		}
+
+		posY = maxY;
 	}
 
 	u64 Visualizer::GetPlayTime()
@@ -2825,9 +2897,6 @@ namespace uba
 				if (pos.y >= posY && pos.y < posY + GraphHeight)
 				{
 					int posX = int(m_scrollPosX) + progressRect.left;
-					u64 prevTime = 0;
-					u64 prevSend = 0;
-					u64 prevRecv = 0;
 					int prevX = 100000;
 					for (u64 i=0, e=session.updates.size(); i!=e; ++i)
 					{
@@ -2836,17 +2905,19 @@ namespace uba
 						auto updateRecv = session.networkRecv[i];
 						int x = posX + int(float(TimeToS(updateTime)) * scaleX);
 
-						if (prevSend > updateSend || prevRecv > updateRecv)
-						{
-							prevSend = updateSend;
-							prevRecv = updateRecv;
-							prevX = x;
-							continue;
-						}
-
 						int hitOffset = (prevX - x)/2;
 						if (pos.x + hitOffset >= prevX && pos.x + hitOffset <= x)
 						{
+							u64 prevTime = 0;
+							u64 prevSend = 0;
+							u64 prevRecv = 0;
+							if (i > 0)
+							{
+								prevTime = session.updates[i - 1];
+								prevSend = Min(session.networkSend[i - 1], updateSend);
+								prevRecv = Min(session.networkRecv[i - 1], updateRecv);
+							}
+
 							double duration = TimeToS(updateTime - prevTime);
 							outResult.stats.recvBytes = updateRecv;
 							outResult.stats.sendBytes = updateSend;
@@ -2859,14 +2930,19 @@ namespace uba
 							outResult.statsSelected = true;
 
 							for (auto& pair : session.drives)
-								outResult.stats.drives[pair.first].busyPercent = pair.second.busyPercent[i];
+							{
+								auto& updateBusy = pair.second.busyPercent[i];
+								u64 updateRead = pair.second.readBytes[i];
+								u64 updateWrite = pair.second.writeBytes[i];
+								auto& statsDrive = outResult.stats.drives[pair.first];
+								statsDrive.busyPercent = updateBusy;
+								statsDrive.readPerSecond = u64(float(updateRead) / duration);
+								statsDrive.writePerSecond = u64(float(updateWrite) / duration);
+							}
 							return;
 						}
 
 						prevX = x;
-						prevTime = updateTime;
-						prevSend = updateSend;
-						prevRecv = updateRecv;
 					}
 					posY += GraphHeight;
 				}
@@ -2876,9 +2952,8 @@ namespace uba
 
 			if (m_config.showDetailedData)
 			{
-				auto drawText = [&](const StringBufferBase& text, RECT& rect)
+				auto drawText = [&](const StringBufferBase& text, RECT& rect, u32* outWidth)
 					{
-						//DrawTextW(hdc, text.data, text.count, &rect, DT_SINGLELINE);
 						if (pos.x >= rect.left && pos.x < rect.right && pos.y >= rect.top && pos.y < rect.bottom && text.StartsWith(TC("Fetched Files")))
 							outResult.fetchedFilesSelected = sessionIndex;
 					};
@@ -3371,6 +3446,7 @@ namespace uba
 
 		case WM_PAINT:
 		{
+			u64 start = GetTime();
 			PaintClient([&](HDC hdc, HDC memDC, RECT& rect)
 				{
 					FillRect(memDC, &rect, m_backgroundBrush);
@@ -3379,6 +3455,7 @@ namespace uba
 					m_activeHdc = 0;
 					BitBlt(hdc, 0, 0, rect.right - rect.left, rect.bottom - rect.top, memDC, 0, 0, SRCCOPY);
 				});
+			m_lastPaintTimeMs = TimeToMs(GetTime() - start);
 			break;
 		}
 		case WM_SIZE:
@@ -3429,10 +3506,8 @@ namespace uba
 			{
 				UpdateScrollbars(true);
 
-				u64 startTime = GetTime();
-				RedrawWindow(m_hwnd, NULL, NULL, RDW_INVALIDATE|RDW_UPDATENOW);
-				u64 paintTimeMs = TimeToMs(GetTime() - startTime);
-				u32 waitTime = u32(Min(paintTimeMs * 5, 200ull));
+				RedrawWindow(m_hwnd, NULL, NULL, RDW_INVALIDATE); // Don't use RDW_UPDATENOW.. it will cause stutter
+				u32 waitTime = 60u;//u32(Min(m_lastPaintTimeMs * 5, 200ull));
 				if (!m_traceView.finished)
 					SetTimer(m_hwnd, 0, waitTime, NULL);
 			}
diff --git a/Engine/Source/Programs/UnrealBuildAccelerator/Visualizer/Private/UbaVisualizer.h b/Engine/Source/Programs/UnrealBuildAccelerator/Visualizer/Private/UbaVisualizer.h
--- a/Engine/Source/Programs/UnrealBuildAccelerator/Visualizer/Private/UbaVisualizer.h
+++ b/Engine/Source/Programs/UnrealBuildAccelerator/Visualizer/Private/UbaVisualizer.h
@@ -88,7 +88,7 @@ namespace uba
 		void PaintActiveProcesses(int& posY, const RECT& clientRect, const Function<void(TraceView::ProcessLocation&, u32, bool)>& drawProcess);
 		void PaintProcessRect(TraceView::Process& process, HDC hdc, RECT rect, const RECT& progressRect, bool selected, bool writingBitmap, HBRUSH& lastSelectedBrush);
 		void PaintTimeline(HDC hdc, const RECT& clientRect);
-		using DrawTextFunc = Function<void(const StringBufferBase& text, RECT& rect)>;
+		using DrawTextFunc = Function<void(const StringBufferBase& text, RECT& rect, u32* outWidth)>;
 		void PaintDetailedStats(int& posY, const RECT& progressRect, TraceView::Session& session, bool isRemote, u64 playTime, const DrawTextFunc& drawTextFunc);
 
 		struct Stats
@@ -104,6 +104,8 @@ namespace uba
 			struct Drive
 			{
 				u8 busyPercent;
+				u64 readPerSecond = 0;
+				u64 writePerSecond = 0;
 			};
 			Map<char, Drive> drives;
 		};
@@ -169,6 +171,8 @@ namespace uba
 		u64 m_startTime = 0;
 		u64 m_pauseTime = 0;
 
+		u64 m_lastPaintTimeMs;
+
 		struct ProcessBrushes
 		{
 			HBRUSH inProgress = 0;
