[UBA]
* Improved error description when failing to create/map memory mappings in detoured process

diff --git a/Engine/Source/Programs/UnrealBuildAccelerator/Detours/Private/Windows/UbaDetoursFunctionsWin.cpp b/Engine/Source/Programs/UnrealBuildAccelerator/Detours/Private/Windows/UbaDetoursFunctionsWin.cpp
--- a/Engine/Source/Programs/UnrealBuildAccelerator/Detours/Private/Windows/UbaDetoursFunctionsWin.cpp
+++ b/Engine/Source/Programs/UnrealBuildAccelerator/Detours/Private/Windows/UbaDetoursFunctionsWin.cpp
@@ -277,7 +277,7 @@ void MemoryFile::Reserve(u64 reserveSize_, const tchar* fileName)
 	{
 		baseAddress = (u8*)VirtualAlloc(NULL, reserveSize, MEM_RESERVE, PAGE_READWRITE);
 		if (!baseAddress)
-			FatalError(1354, L"VirtualAlloc failed trying to reserve %llu. (Error code: %u)", reserveSize, GetLastError());
+			FatalError(1354, L"VirtualAlloc failed trying to reserve %llu for %s. (Error code: %u)", reserveSize, fileName, GetLastError());
 		mappedSize = reserveSize;
 	}
 #if UBA_ENABLE_ON_DISK_FILE_MAPPINGS
@@ -332,12 +332,12 @@ void MemoryFile::Reserve(u64 reserveSize_, const tchar* fileName)
 			TimerScope ts(g_kernelStats.createFileMapping);
 			handle = True_CreateFileMappingW(INVALID_HANDLE_VALUE, NULL, PAGE_READWRITE | SEC_RESERVE, ToHigh(reserveSize), ToLow(reserveSize), NULL);
 			if (!handle)
-				FatalError(1348, L"CreateFileMappingW failed trying to reserve %llu. (Error code: %u)", reserveSize, GetLastError());
+				FatalError(1348, L"CreateFileMappingW failed trying to reserve %llu for %s. (Error code: %u)", reserveSize, fileName, GetLastError());
 		}
 		TimerScope ts(g_kernelStats.mapViewOfFile);
 		baseAddress = (u8*)True_MapViewOfFile(handle, FILE_MAP_READ | FILE_MAP_WRITE, 0, 0, mappedSize);
 		if (!baseAddress)
-			FatalError(1353, L"MapViewOfFile failed trying to map %llu. ReservedSize: %llu (Error code: %u)", mappedSize, reserveSize, GetLastError());
+			FatalError(1353, L"MapViewOfFile failed trying to map %llu for %s. ReservedSize: %llu (Error code: %u)", mappedSize, fileName, reserveSize, GetLastError());
 		mappingHandle = { 0, handle };
 	}
 }
