[UBA]
* Tweaked file type max sizes to save memory on wine (a reserved CreateFileMapping actually takes up space in wine on the tmp partition)

diff --git a/Engine/Source/Programs/UnrealBuildAccelerator/Core/Private/UbaApplicationRules.cpp b/Engine/Source/Programs/UnrealBuildAccelerator/Core/Private/UbaApplicationRules.cpp
--- a/Engine/Source/Programs/UnrealBuildAccelerator/Core/Private/UbaApplicationRules.cpp
+++ b/Engine/Source/Programs/UnrealBuildAccelerator/Core/Private/UbaApplicationRules.cpp
@@ -7,7 +7,11 @@ namespace uba
 {
 	bool GlobalRules::FileCanBeCompressed(const StringView& fileName)
 	{
-		return IsWindows && (fileName.EndsWith(TC(".obj")) || fileName.EndsWith(TC(".o")));
+		if (!IsWindows)
+			return false;
+		return fileName.EndsWith(TC(".obj"))
+			|| fileName.EndsWith(TC(".o"))
+			;
 	}
 
 	GlobalRules g_globalRules;
@@ -38,7 +42,7 @@ namespace uba
 		{
 			if (file.EndsWith(TC(".pdb")))
 				return 14ull * 1024 * 1024 * 1024; // This is ridiculous
-			if (file.EndsWith(TC(".json")) || file.EndsWith(TC(".exp")))
+			if (file.EndsWith(TC(".json")) || file.EndsWith(TC(".exp")) || file.EndsWith(TC(".sarif")) || file.EndsWith(TC(".res")))
 				return 32 * 1024 * 1024;
 			if (file.EndsWith(TC(".obj")) || (isSystemOrTempFile && file.Contains(TC("_cl_")))) // There are _huge_ obj files when building with -stresstestunity
 				return 1024 * 1024 * 1024;
@@ -154,6 +158,17 @@ namespace uba
 			 	|| EndsWith(file, TStrlen(file), TC(".manifest")); // lld-link.exe is using a different name for files shared with child processes
 		}
 
+		virtual u64 FileTypeMaxSize(const StringView& file, bool isSystemOrTempFile) const override
+		{
+			if (file.Contains(TC("lnk{")))
+				return 32 * 1024 * 1024;
+			if (file.EndsWith(TC(".lib"))) // This is import lib which is tiny but not sure if we output any full libs
+				return 512 * 1024 * 1024;
+			if (file.EndsWith(TC(".dll")) || file.EndsWith(TC(".exe")))
+				return 4ull * 1024 * 1024 * 1024;
+			return Super::FileTypeMaxSize(file, isSystemOrTempFile);
+		}
+
 		virtual bool IsRarelyRead(const StringView& file) const override
 		{
 			return file.EndsWith(TC(".exp"))
@@ -510,6 +525,8 @@ namespace uba
 
 	class ApplicationRulesISPCExe : public ApplicationRules
 	{
+		using Super = ApplicationRules;
+
 		virtual bool AllowDetach() const override
 		{
 			return true;
@@ -524,6 +541,15 @@ namespace uba
 				|| fileName.EndsWith(TC(".o")); // Used when compiling for linux
 		}
 
+		virtual u64 FileTypeMaxSize(const StringView& file, bool isSystemOrTempFile) const override
+		{
+			if (file.Contains(TC(".generated.dummy")) || file.EndsWith(TC(".ispc.txt")))
+				return 32ull * 1024 * 1024;
+			if (file.EndsWith(TC(".obj")) || file.EndsWith(TC(".o")))
+				return 128ull * 1024 * 1024;
+			return Super::FileTypeMaxSize(file, isSystemOrTempFile);
+		}
+
 		virtual bool IsCacheable() const override
 		{
 			return true;
