[UBA]
* Tweaked file max sizes even more for clang paths

diff --git a/Engine/Source/Programs/UnrealBuildAccelerator/Core/Private/UbaApplicationRules.cpp b/Engine/Source/Programs/UnrealBuildAccelerator/Core/Private/UbaApplicationRules.cpp
--- a/Engine/Source/Programs/UnrealBuildAccelerator/Core/Private/UbaApplicationRules.cpp
+++ b/Engine/Source/Programs/UnrealBuildAccelerator/Core/Private/UbaApplicationRules.cpp
@@ -44,7 +44,7 @@ namespace uba
 				return 14ull * 1024 * 1024 * 1024; // This is ridiculous
 			if (file.EndsWith(TC(".json")) || file.EndsWith(TC(".exp")) || file.EndsWith(TC(".sarif")) || file.EndsWith(TC(".res")))
 				return 32 * 1024 * 1024;
-			if (file.EndsWith(TC(".obj")) || (isSystemOrTempFile && file.Contains(TC("_cl_")))) // There are _huge_ obj files when building with -stresstestunity
+			if (file.EndsWith(TC(".ifc")) || file.EndsWith(TC(".obj")) || (isSystemOrTempFile && file.Contains(TC("_cl_")))) // There are _huge_ obj files when building with -stresstestunity
 				return 1024 * 1024 * 1024;
 			return Super::FileTypeMaxSize(file, isSystemOrTempFile);
 		}
@@ -81,6 +81,7 @@ namespace uba
 			return fileName.EndsWith(TC(".obj"))
 				|| fileName.EndsWith(TC(".dep.json"))
 				|| fileName.EndsWith(TC(".sarif"))
+				|| fileName.EndsWith(TC(".ifc"))
 				|| fileName.EndsWith(TC(".rc2.res")) // Not really an obj file.. 
 				;// || fileName.EndsWith(TC(".h.pch") // Not tested enough
 		}
@@ -274,15 +275,28 @@ namespace uba
 	class ApplicationRulesClang : public ApplicationRules
 	{
 	public:
+		using Super = ApplicationRules;
+
 		virtual bool IsExitCodeSuccess(u32 exitCode) const override
 		{
 			return exitCode == 0;
 		}
+
+		virtual u64 FileTypeMaxSize(const StringView& file, bool isSystemOrTempFile) const override
+		{
+			if (file.EndsWith(TC(".obj")) || file.EndsWith(TC(".o")) || file.EndsWith(TC(".o.tmp")) || file.EndsWith(TC(".obj.tmp")))
+				return 512 * 1024 * 1024;
+			if (file.EndsWith(TC(".d")))
+				return 32 * 1024 * 1024;
+			return Super::FileTypeMaxSize(file, isSystemOrTempFile);
+		}
 	};
 
 	class ApplicationRulesClangPlusPlusExe : public ApplicationRulesClang
 	{
 	public:
+		using Super = ApplicationRulesClang;
+
 		virtual bool AllowDetach() const override
 		{
 			return true;
diff --git a/Engine/Source/Programs/UnrealBuildAccelerator/Core/Public/UbaApplicationRules.h b/Engine/Source/Programs/UnrealBuildAccelerator/Core/Public/UbaApplicationRules.h
--- a/Engine/Source/Programs/UnrealBuildAccelerator/Core/Public/UbaApplicationRules.h
+++ b/Engine/Source/Programs/UnrealBuildAccelerator/Core/Public/UbaApplicationRules.h
@@ -69,7 +69,9 @@ namespace uba
 			return false;
 		}
 
-		// Max file size if using memory files.
+		// Max file size if using memory files. Setting this value to a proper size is important because of Wine on linux use this
+		// .. on windows this size can be very big since it is just virtual address space and nothing is committed...
+		// On wine (as of version 9.10) a ftruncate is done on the wine server causing large max sizes to take up actual disk space on the tmp partition
 		virtual u64 FileTypeMaxSize(const StringView& file, bool isSystemOrTempFile) const
 		{
 			return 8ull * 1024 * 1024 * 1024;
