[UBA]
* Moved all obj file decompression to session side and deleted code from detouring code.
* Added so pch is stored compressed when obj file compression is on
* Added support for freeing compressed file mappings. decompressed .obj and pch are now freed when not used anymore (and remapped if needed)
* Promoted dependency crawler from SessionClient to Session and enabled msvc link crawler on host side if compression is on (so we can decompress obj files in parallel)
*  Added early out in uba CloseFileMapping/UnmapViewOfFile if invalid handle/ptr
* Added debug code for tracking of creating/freeing file mappings

diff --git a/Engine/Source/Programs/UnrealBuildAccelerator/Common/Private/UbaFileMapping.cpp b/Engine/Source/Programs/UnrealBuildAccelerator/Common/Private/UbaFileMapping.cpp
--- a/Engine/Source/Programs/UnrealBuildAccelerator/Common/Private/UbaFileMapping.cpp
+++ b/Engine/Source/Programs/UnrealBuildAccelerator/Common/Private/UbaFileMapping.cpp
@@ -286,6 +286,9 @@ namespace uba
 
 	bool UnmapViewOfFile(Logger& logger, const void* lpBaseAddress, u64 bytesToUnmap, const tchar* hint)
 	{
+		if (!lpBaseAddress)
+			return true;
+
 		ExtendedTimerScope ts(KernelStats::GetCurrent().unmapViewOfFile);
 
 		#if UBA_DEBUG_FILE_MAPPING
@@ -314,16 +317,16 @@ namespace uba
 
 	bool CloseFileMapping(Logger& logger, FileMappingHandle h, const tchar* hint)
 	{
+		if (!h.IsValid())
+			return true;
+
 		#if UBA_DEBUG_FILE_MAPPING
-		if (h.IsValid())
-		{
-			SCOPED_FUTEX(g_fileMappingsLock, lock);
-			auto findIt = g_fileMappings.find(h.handle);
-			UBA_ASSERTF(findIt != g_fileMappings.end(), TC("Handle: %llu"), u64(h.handle));
-			UBA_ASSERTF(!findIt->second.viewCount.load(), TC("Closing file mapping with %llu open views (%s)"), findIt->second.viewCount.load(), findIt->second.hint.c_str());
-			//logger.Info(TC("FILEMAPPING CLOSE: %s %llu (%llu)"), findIt->second.hint.c_str(), u64(h.handle), g_fileMappings.size() - 1);
-			g_fileMappings.erase(findIt);
-		}
+		SCOPED_FUTEX(g_fileMappingsLock, lock);
+		auto findIt = g_fileMappings.find(h.handle);
+		UBA_ASSERTF(findIt != g_fileMappings.end(), TC("Handle: %llu"), u64(h.handle));
+		UBA_ASSERTF(!findIt->second.viewCount.load(), TC("Closing file mapping with %llu open views (%s)"), findIt->second.viewCount.load(), findIt->second.hint.c_str());
+		//logger.Info(TC("FILEMAPPING CLOSE: %s %llu (%llu)"), findIt->second.hint.c_str(), u64(h.handle), g_fileMappings.size() - 1);
+		g_fileMappings.erase(findIt);
 		#endif
 
 #if PLATFORM_WINDOWS
@@ -331,8 +334,6 @@ namespace uba
 			CloseHandle(h.fh);
 		return CloseHandle(h.mh);
 #else
-		if (h.shmFd == -1)
-			return true;
 		if (h.uid == ~u64(0))
 			return true;
 		if (close(h.shmFd) != 0)
diff --git a/Engine/Source/Programs/UnrealBuildAccelerator/Common/Private/UbaProcess.cpp b/Engine/Source/Programs/UnrealBuildAccelerator/Common/Private/UbaProcess.cpp
--- a/Engine/Source/Programs/UnrealBuildAccelerator/Common/Private/UbaProcess.cpp
+++ b/Engine/Source/Programs/UnrealBuildAccelerator/Common/Private/UbaProcess.cpp
@@ -33,7 +33,7 @@
 extern char **environ;
 #endif
 
-#define UBA_DEBUG_TRACK_PROCESS UBA_DEBUG_LOGGER
+#define UBA_DEBUG_TRACK_PROCESS 0 // UBA_DEBUG_LOGGER
 
 #if !PLATFORM_WINDOWS
 #define EXCEPTION_ACCESS_VIOLATION 128 + SIGSEGV
@@ -222,8 +222,11 @@ namespace uba
 		StringBuffer<> realApplication(m_startInfo.application);
 		const tchar* realWorkingDir = m_startInfo.workingDir;
 
-		if (!m_session.PrepareProcess(m_startInfo, m_parentProcess != nullptr, realApplication, realWorkingDir))
+		if (!m_session.PrepareProcess(*this, m_parentProcess != nullptr, realApplication, realWorkingDir))
+		{
+			m_exitCode = 44324;
 			return false;
+		}
 
 		m_realApplication = realApplication.data;
 		m_realWorkingDir = realWorkingDir;
@@ -1336,6 +1339,12 @@ namespace uba
 		return true;
 	}
 
+	bool ProcessImpl::HandleCheckRemapping(BinaryReader& reader, BinaryWriter& writer)
+	{
+		m_session.CheckRemapping(*this, reader, writer);
+		return true;
+	}
+
 	bool ProcessImpl::CreateTempFile(BinaryReader& reader, ProcHandle nativeProcessHandle, const tchar* application)
 	{
 		StringKey key = reader.ReadStringKey();
diff --git a/Engine/Source/Programs/UnrealBuildAccelerator/Common/Private/UbaSession.cpp b/Engine/Source/Programs/UnrealBuildAccelerator/Common/Private/UbaSession.cpp
--- a/Engine/Source/Programs/UnrealBuildAccelerator/Common/Private/UbaSession.cpp
+++ b/Engine/Source/Programs/UnrealBuildAccelerator/Common/Private/UbaSession.cpp
@@ -520,6 +520,10 @@ namespace uba
 		entry.mapping = {};
 		entry.handled = true;
 
+		#if UBA_DEBUG_TRACK_MAPPING
+		entry.name = newFileName;
+		#endif
+
 		SCOPED_WRITE_LOCK(m_fileMappingTableMemLock, lock);
 		BinaryWriter writer(m_fileMappingTableMem, m_fileMappingTableSize);
 		writer.WriteStringKey(fileNameKey);
@@ -530,7 +534,7 @@ namespace uba
 		return newSize;
 	}
 
-	bool Session::CreateMemoryMapFromFile(MemoryMap& out, StringKey fileNameKey, const tchar* fileName, bool isCompressed, u64 alignment, const tchar* hint)
+	bool Session::CreateMemoryMapFromFile(MemoryMap& out, StringKey fileNameKey, const tchar* fileName, bool isCompressedCas, u64 alignment, const tchar* hint, ProcessImpl* requestingProcess, bool canBeFreed)
 	{
 		TimerScope ts(Stats().waitMmapFromFile);
 
@@ -539,16 +543,31 @@ namespace uba
 		FileMappingEntry& entry = insres.first->second;
 		lookupLock.Leave();
 
-		SCOPED_FUTEX(entry.lock, entryCs);
+
+		auto updateRequestingProcess = [&]()
+			{
+				if (requestingProcess && entry.canBeFreed && !m_runningRemote)
+				{
+					SCOPED_FUTEX(requestingProcess->m_usedFileMappingsLock, lock);
+					if (!requestingProcess->m_hasExited && requestingProcess->m_usedFileMappings.insert(fileNameKey).second)
+						++entry.refCount;
+				}
+			};
+
+		SCOPED_FUTEX(entry.lock, entryLock);
 
 		if (entry.handled)
 		{
-			entryCs.Leave();
+			entryLock.Leave();
+
 			if (!entry.success)
 				return false;
 			out.size = entry.size;
 			if (entry.mapping.IsValid())
+			{
 				Storage::GetMappingString(out.name, entry.mapping, entry.mappingOffset);
+				updateRequestingProcess();
+			}
 			else
 				out.name.Append(entry.isDir ? TC("$d") : TC("$f"));
 			return true;
@@ -563,8 +582,6 @@ namespace uba
 
 		bool isDir = false;
 		u32 attributes = DefaultAttributes();
-		if (!isCompressed)
-			attributes |= FILE_FLAG_SEQUENTIAL_SCAN;
 		FileHandle fileHandle = uba::CreateFileW(fileName, GENERIC_READ, FILE_SHARE_READ, OPEN_EXISTING, attributes);
 		if (fileHandle == InvalidFileHandle)
 		{
@@ -584,10 +601,12 @@ namespace uba
 
 		u64 size = 0;
 		u64 lastWriteTime = 0;
+		u64 fileStartOffset = 0;
 
+		bool isCompressed = isCompressedCas;
 		if (!isDir)
 		{
-			if (isCompressed)
+			if (isCompressedCas)
 			{
 				if (!ReadFile(m_logger, fileName, fileHandle, &size, sizeof(u64)))
 					return m_logger.Error(TC("Failed to read first bytes from file %s (%s)"), fileName, LastErrorToText().data);
@@ -600,6 +619,25 @@ namespace uba
 
 				size = info.size;
 				lastWriteTime = info.lastWriteTime;
+
+				if (m_readObjFilesCompressed && info.size > sizeof(CompressedObjFileHeader) && g_globalRules.FileCanBeCompressed(ToView(fileName)))
+				{
+					CompressedObjFileHeader header(CasKeyZero);
+					if (!ReadFile(m_logger, fileName, fileHandle, &header, sizeof(header)))
+						return m_logger.Error(TC("Failed to read header of compressed file %s (%s)"), fileName, LastErrorToText().data);
+					if (header.IsValid())
+					{
+						fileStartOffset = sizeof(CompressedObjFileHeader);
+						isCompressed = true;
+						if (!ReadFile(m_logger, fileName, fileHandle, &size, sizeof(u64)))
+							return m_logger.Error(TC("Failed to read first bytes from file %s (%s)"), fileName, LastErrorToText().data);
+					}
+					else
+					{
+						if (!SetFilePointer(m_logger, fileName, fileHandle, 0))
+							return false;
+					}
+				}
 			}
 		}
 
@@ -615,27 +653,67 @@ namespace uba
 			if (size > m_fileMappingBuffer.GetFileMappingCapacity())
 				return m_logger.Error(TC("File %s has a size (%llu) that is too large to fit in mapping buffer (%s)"), fileName, size, hint);
 
-			auto mappedView = m_fileMappingBuffer.AllocAndMapView(MappedView_Transient, size, alignment, fileName, false);
-			if (!mappedView.memory)
-				return false;
-			auto unmapGuard = MakeGuard([&](){ m_fileMappingBuffer.UnmapView(mappedView, fileName); });
+			FileMappingHandle mapping;
+			u64 mappingOffset = 0;
+			u8* mappingMemory = nullptr;
+
+			MappedView mappedView;
+
+			auto ownedGuard1 = MakeGuard([&]() { CloseFileMapping(m_logger, mapping, fileName); });
+			auto ownedGuard2 = MakeGuard([&]() { UnmapViewOfFile(m_logger, mappingMemory, size, fileName); });
+			auto viewGuard = MakeGuard([&]() { m_fileMappingBuffer.UnmapView(mappedView, fileName); });
+
+			if (canBeFreed)
+			{
+				viewGuard.Cancel();
+				mapping = CreateMemoryMappingW(m_logger, PAGE_READWRITE, size, nullptr, fileName);
+				if (!mapping.IsValid())
+					return false;
+				mappingMemory = MapViewOfFile(m_logger, mapping, FILE_MAP_WRITE, 0, size);
+				if (!mappingMemory)
+					return false;
+			}
+			else
+			{
+				UBA_ASSERTF(alignment, TC("No alignment set when creating memory map for %s (%s)"), fileName, hint);
+				UBA_ASSERT(!entry.canBeFreed);
+				ownedGuard1.Cancel();
+				ownedGuard2.Cancel();
+				mappedView = m_fileMappingBuffer.AllocAndMapView(MappedView_Transient, size, alignment, fileName, false);
+				if (!mappedView.memory)
+					return false;
+				mapping = mappedView.handle;
+				mappingOffset = mappedView.offset;
+				mappingMemory = mappedView.memory;
+			}
 
 			if (isCompressed)
 			{
-				if (!m_storage.DecompressFileToMemory(fileName, fileHandle, mappedView.memory, size, TC("MappedMemory")))
+				if (!m_storage.DecompressFileToMemory(fileName, fileHandle, mappingMemory, size, TC("CreateMemoryMapFromFile"), fileStartOffset))
 					return false;
 			}
 			else
 			{
-				if (!ReadFile(m_logger, fileName, fileHandle, mappedView.memory, size))
+				if (!ReadFile(m_logger, fileName, fileHandle, mappingMemory, size))
 					return false;
 			}
 
-			unmapGuard.Execute();
+			ownedGuard2.Execute();
+			ownedGuard1.Cancel();
+			viewGuard.Execute();
 
-			entry.mappingOffset = mappedView.offset;
-			Storage::GetMappingString(out.name, mappedView.handle, mappedView.offset);
-			entry.mapping = mappedView.handle;
+			entry.canBeFreed = canBeFreed;
+			entry.mappingOffset = mappingOffset;
+			Storage::GetMappingString(out.name, mapping, mappingOffset);
+			entry.mapping = mapping;
+
+			if (canBeFreed)
+			{
+				entry.usedCount = 0;
+				entry.usedCountBeforeFree = g_globalRules.GetUsedCountBeforeFree(ToView(fileName));
+			}
+
+			updateRequestingProcess();
 		}
 
 		entry.success = true;
@@ -649,6 +727,11 @@ namespace uba
 			m_fileMappingTableSize = (u32)writer.GetPosition();
 		}
 
+		#if UBA_DEBUG_TRACK_MAPPING
+		entry.name = fileName;
+		m_debugLogger->Info(TC("Mapping created 0x%llx (%s) from file (%s) - %s\n"), u64(entry.mapping.mh), entry.name.c_str(), hint, TimeToText(GetTime() - ts2.start).str);
+		#endif
+
 		entry.isDir = isDir;
 		entry.size = size;
 		entry.lastWriteTime = lastWriteTime;
@@ -748,6 +831,12 @@ namespace uba
 		entry.size = size;
 		
 		out.size = size;
+
+		#if UBA_DEBUG_TRACK_MAPPING
+		entry.name = fileName;
+		m_debugLogger->Info(TC("Mapping created 0x%llx (%s) from view\n"), u64(entry.mapping.mh), entry.name.c_str());
+		#endif
+
 		return true;
 	}
 
@@ -1127,6 +1216,7 @@ namespace uba
 	,	m_writeFilesBottleneck(info.writeFilesBottleneck)
 	,	m_writeFilesFileMapMax(u64(info.writeFilesFileMapMaxMb)*1024*1024)
 	,	m_writeFilesNoBufferingMin(u64(info.writeFilesNoBufferingMinMb)*1024*1024)
+	,	m_dependencyCrawler(m_logger, *workManager)
 	{
 		if (info.useUniqueId)
 		{
@@ -1229,6 +1319,13 @@ namespace uba
 
 		if (info.traceOutputFile)
 			m_traceOutputFile.Append(info.traceOutputFile);
+
+		if (m_readObjFilesCompressed && !m_runningRemote)
+			m_dependencyCrawler.Init(
+				[this](const StringView& fileName, u32& outAttr) { return FileExists(m_logger, fileName.data, nullptr, &outAttr); }, // FileExists
+				[this](const StringView& path, const DependencyCrawler::FileFunc& fileFunc) {}, // TraverseFilesFunc
+				false);
+
 	}
 
 	bool Session::Create(const SessionCreateInfo& info)
@@ -1359,7 +1456,7 @@ namespace uba
 			CloseFileMapping(m_logger, kv.second.mappingHandle, TC("VirtualFile"));
 
 		for (auto& i : m_fileMappingTableLookup)
-			if (i.second.owned)
+			if (i.second.canBeFreed)
 				CloseFileMapping(m_logger, i.second.mapping, TC("FileMappingKeptFromOutput"));
 
 		UnmapViewOfFile(m_logger, m_fileMappingTableMem, FileMappingTableMemSize, TC("FileMappingTable"));
@@ -1886,6 +1983,42 @@ namespace uba
 			m_stats.Add(process.m_sessionStats);
 		}
 
+		SCOPED_FUTEX(process.m_usedFileMappingsLock, usedFileMappingsLock);
+		auto usedFileMappings = std::move(process.m_usedFileMappings); // Might be that crawler is still working
+		usedFileMappingsLock.Leave();
+
+		for (auto& fileNameKey : usedFileMappings)
+		{
+			SCOPED_FUTEX(m_fileMappingTableLookupLock, lookupLock);
+			auto findIt = m_fileMappingTableLookup.find(fileNameKey);
+			UBA_ASSERT(findIt != m_fileMappingTableLookup.end());
+			FileMappingEntry& entry = findIt->second;
+			lookupLock.Leave();
+
+			SCOPED_FUTEX(entry.lock, entryLock);
+
+			if (entry.usedCount < entry.usedCountBeforeFree)
+				++entry.usedCount;
+
+			if (--entry.refCount)
+				continue;
+
+			if (entry.usedCountBeforeFree == 255 || entry.usedCount < entry.usedCountBeforeFree)
+				continue;
+
+			UBA_ASSERT(entry.canBeFreed);
+
+			#if UBA_DEBUG_TRACK_MAPPING
+			m_debugLogger->Info(TC("Mapping freed 0x%llx (%s)\n"), u64(entry.mapping.mh), entry.name.c_str());
+			#endif
+
+			m_workManager->AddWork([mh = entry.mapping, this](const WorkContext&) { CloseFileMapping(m_logger, mh, TC("UsedFileMapping")); }, 1, TC("CloseFileMapping"));
+			entry.handled = false;
+			entry.mapping = {};
+		}
+
+
+
 		SCOPED_FUTEX(m_processesLock, lock);
 		m_deadProcesses.emplace_back(&process); // Here to prevent Process thread call trigger a delete of Process which causes a deadlock
 		auto& stats = m_applicationStats[applicationName.data];
@@ -2051,7 +2184,7 @@ namespace uba
 		v.shrink_to_fit();
 	}
 
-	bool Session::IsRarelyRead(ProcessImpl& process, const StringBufferBase& fileName) const
+	bool Session::IsRarelyRead(ProcessImpl& process, const StringView& fileName) const
 	{
 		return process.m_startInfo.rules->IsRarelyRead(fileName);
 	}
@@ -2077,16 +2210,27 @@ namespace uba
 		return fileName.EndsWith(TC(".h"));
 	}
 
-	bool Session::PrepareProcess(ProcessStartInfoHolder& startInfo, bool isChild, StringBufferBase& outRealApplication, const tchar*& outRealWorkingDir)
+	bool Session::PrepareProcess(ProcessImpl& process, bool isChild, StringBufferBase& outRealApplication, const tchar*& outRealWorkingDir)
 	{
+		ProcessStartInfoHolder& startInfo = process.m_startInfo;
 		if (StartsWith(startInfo.application, TC("ubacopy")))
 			return true;
-		if (IsAbsolutePath(startInfo.application))
-			return true;
-		if (!SearchPathForFile(m_logger, outRealApplication.Clear(), startInfo.application, startInfo.workingDir))
-			return false;
-		startInfo.applicationStr = outRealApplication.data;
-		startInfo.application = startInfo.applicationStr.c_str();
+
+		if (!IsAbsolutePath(startInfo.application))
+		{
+			if (!SearchPathForFile(m_logger, outRealApplication.Clear(), startInfo.application, startInfo.workingDir))
+				return false;
+			startInfo.applicationStr = outRealApplication.data;
+			startInfo.application = startInfo.applicationStr.c_str();
+		}
+
+		if (!isChild && !m_runningRemote && m_readObjFilesCompressed)
+		{
+			auto crawlerType = startInfo.rules->GetDependencyCrawlerType();
+			if (crawlerType == DependencyCrawlerType_MsvcLinker || crawlerType == DependencyCrawlerType_ClangLinker)
+				RunDependencyCrawler(process);
+		}
+
 		return true;
 	}
 
@@ -2236,57 +2380,19 @@ namespace uba
 		const StringBufferBase& fileName = msg.fileName;
 		const StringKey& fileNameKey = msg.fileNameKey;
 
+		if ((msg.access & ~FileAccess_Read) == 0)
+		{
+			TrackWorkScope tws(*m_workManager, AsView(TC("CreateFile")), ColorWork);
+			tws.AddHint(msg.fileName);
+			return CreateFileForRead(out, tws, fileName, fileNameKey, msg.process, *msg.process.m_startInfo.rules);
+		}
+
 		auto tableSizeGuard = MakeGuard([&]()
 			{
 				out.mappedFileTableSize = GetFileMappingSize();
 				out.directoryTableSize = GetDirectoryTableSize();
 			});
 
-		if ((msg.access & ~FileAccess_Read) == 0)
-		{
-			if (!IsWindows)
-			{
-				out.fileName.Append(fileName);
-				return true;
-			}
-		
-			if (fileName.EndsWith(TC(".dll")) || fileName.EndsWith(TC(".exe")))
-			{
-				UBA_ASSERTF(IsAbsolutePath(fileName.data), TC("Got bad filename from process %s"), fileName.data);
-				AddFileMapping(fileNameKey, fileName.data, TC("#"));
-				out.fileName.Append(TC("#"));
-				return true;
-			}
-			
-			if (m_allowMemoryMaps)
-			{
-				if (u64 alignment = GetMemoryMapAlignment(fileName))
-				{
-					MemoryMap map;
-					if (CreateMemoryMapFromFile(map, fileNameKey, fileName.data, false, alignment, TC("CreateFile")))
-					{
-						out.size = map.size;
-						out.fileName.Append(map.name);
-					}
-					else
-					{
-						out.fileName.Append(fileName);
-					}
-					return true;
-				}
-			}
-
-			if (!IsRarelyRead(msg.process, fileName))
-			{
-				AddFileMapping(fileNameKey, fileName.data, TC("#"));
-				out.fileName.Append(TC("#"));
-				return true;
-			}
-
-			out.fileName.Append(fileName);
-			return true;
-		}
-		
 		// if ((message.Access & FileAccess.Write) != 0)
 		m_storage.ReportFileWrite(fileNameKey, fileName.data);
 
@@ -2321,6 +2427,61 @@ namespace uba
 		return true;
 	}
 
+	bool Session::CreateFileForRead(CreateFileResponse& out, TrackWorkScope& tws, const StringView& fileName, const StringKey& fileNameKey, ProcessImpl& process, const ApplicationRules& rules)
+	{
+		auto tableSizeGuard = MakeGuard([&]()
+			{
+				out.mappedFileTableSize = GetFileMappingSize();
+				out.directoryTableSize = GetDirectoryTableSize();
+			});
+
+		if (!IsWindows)
+		{
+			out.fileName.Append(fileName);
+			return true;
+		}
+		
+		if (fileName.EndsWith(TC(".dll")) || fileName.EndsWith(TC(".exe")))
+		{
+			UBA_ASSERTF(IsAbsolutePath(fileName.data), TC("Got bad filename from process (%s)"), fileName.data);
+			AddFileMapping(fileNameKey, fileName.data, TC("#"));
+			out.fileName.Append(TC("#"));
+			return true;
+		}
+			
+		if (m_allowMemoryMaps)
+		{
+			u64 alignment = GetMemoryMapAlignment(fileName);
+			bool canBeCompressed = m_readObjFilesCompressed && g_globalRules.FileCanBeCompressed(fileName);
+			bool useMemoryMap = alignment != 0 || canBeCompressed;
+			if (useMemoryMap)
+			{
+				MemoryMap map;
+				bool canBeFreed = canBeCompressed;
+				if (CreateMemoryMapFromFile(map, fileNameKey, fileName.data, false, alignment, TC("CreateFile"), &process, canBeFreed))
+				{
+					out.size = map.size;
+					out.fileName.Append(map.name);
+				}
+				else
+				{
+					out.fileName.Append(fileName);
+				}
+				return true;
+			}
+		}
+
+		if (!IsRarelyRead(process, fileName))
+		{
+			AddFileMapping(fileNameKey, fileName.data, TC("#"));
+			out.fileName.Append(TC("#"));
+			return true;
+		}
+
+		out.fileName.Append(fileName);
+		return true;
+	}
+
 	void RemoveWrittenFile(ProcessImpl& process, const TString& name)
 	{
 		SCOPED_FUTEX(process.m_writtenFilesLock, writtenLock);
@@ -2826,9 +2987,12 @@ namespace uba
 				Storage::GetMappingString(name, file.mappingHandle, 0);
 				SCOPED_FUTEX(m_fileMappingTableLookupLock, lookupLock);
 				auto insres = m_fileMappingTableLookup.try_emplace(file.key);
+				UBA_ASSERT(insres.second);
 				FileMappingEntry& entry = insres.first->second;
 				lookupLock.Leave();
 				SCOPED_FUTEX(entry.lock, entryCs);
+				UBA_ASSERT(file.mappingHandle.IsValid());
+				UBA_ASSERT(!entry.mapping.IsValid());
 				entry.handled = true;
 				entry.mapping = file.mappingHandle;
 				entry.mappingOffset = 0;
@@ -2836,7 +3000,14 @@ namespace uba
 				entry.lastWriteTime = file.lastWriteTime;
 				entry.isDir = false;
 				entry.success = true;
-				entry.owned = true;
+				entry.canBeFreed = true;
+				entry.usedCount = 0;
+				entry.usedCountBeforeFree = g_globalRules.GetUsedCountBeforeFree(file.name);
+
+				#if UBA_DEBUG_TRACK_MAPPING
+				entry.name = file.name;
+				m_debugLogger->Info(TC("Mapping kept 0x%llx (%s) from detoured process (UsedCountBeforeFree: %u)\n"), u64(entry.mapping.mh), entry.name.c_str(), u32(entry.usedCountBeforeFree));
+				#endif
 
 				SCOPED_WRITE_LOCK(m_fileMappingTableMemLock, lock);
 				BinaryWriter writer(m_fileMappingTableMem, m_fileMappingTableSize);
@@ -2969,6 +3140,25 @@ namespace uba
 		return true;
 	}
 
+	bool Session::CheckRemapping(ProcessImpl& process, BinaryReader& reader, BinaryWriter& writer)
+	{
+		StringBuffer<> fileName;
+		reader.ReadString(fileName);
+		StringKey fileNameKey = reader.ReadStringKey();
+		UBA_ASSERT(fileNameKey != StringKeyZero);
+
+		#if UBA_DEBUG_TRACK_MAPPING
+		//m_debugLogger->Info(TC("Mapping check (%s)\n"), fileName.data);
+		#endif
+
+		MemoryMap out;
+		u64 alignment = GetMemoryMapAlignment(fileName);
+		if (!CreateMemoryMapFromFile(out, fileNameKey, fileName.data, false, alignment, TC("Remap"), &process, true))
+			return m_logger.Error(TC("Failed to remap %s"), fileName.data);
+		writer.WriteU32(GetFileMappingSize());
+		return true;
+	}
+
 	void Session::FileEntryAdded(StringKey fileNameKey, u64 lastWritten, u64 size)
 	{
 	}
@@ -3063,6 +3253,11 @@ namespace uba
 		Storage::GetMappingString(mappingName, virtualFile.mappingHandle, sourceOffset);
 		entry.success = true;
 
+		#if UBA_DEBUG_TRACK_MAPPING
+		entry.name = filePath.data;
+		m_debugLogger->Info(TC("Mapping created 0x%llx (%s) from virtual file\n"), u64(entry.mapping.mh), entry.name.c_str());
+		#endif
+
 		SCOPED_WRITE_LOCK(m_fileMappingTableMemLock, lock);
 		BinaryWriter writer(m_fileMappingTableMem, m_fileMappingTableSize);
 		writer.WriteStringKey(fileNameKey);
@@ -3662,6 +3857,65 @@ namespace uba
 	{
 	}
 
+	void Session::RunDependencyCrawler(ProcessImpl& process)
+	{
+		auto& startInfo = process.GetStartInfo();
+
+		auto crawlerType = startInfo.rules->GetDependencyCrawlerType();
+		if (crawlerType == DependencyCrawlerType_None)
+			return;
+
+		const tchar* at = TStrchr(startInfo.arguments, '@');
+		if (!at)
+			return;
+
+		auto CreateFileFunc = [this, ph = ProcessHandle(&process), rules = startInfo.rules](const StringView& fileName, const DependencyCrawler::AccessFileFunc& func)
+			{
+				CreateFileResponse out;
+				{
+					TrackWorkScope tws(*m_workManager, AsView(TC("CrawlFile")), ColorWork);
+					tws.AddHint(fileName);
+
+					if (!CreateFileForRead(out, tws, fileName, ToStringKey(fileName), *static_cast<ProcessImpl*>(ph.m_process), *rules))
+						return false;
+				}
+
+				if (!func)
+					return true;
+
+				if (out.fileName[0] == '^')
+				{
+					MappedView view = m_fileMappingBuffer.MapView(out.fileName, out.size, fileName.data);
+					if (view.memory)
+					{
+						bool res = func(view.memory, out.size);
+						m_fileMappingBuffer.UnmapView(view, fileName.data);
+						return res;
+					}
+					return m_logger.Warning(TC("Failed to open %s"), out.fileName.data);
+				}
+
+				if (out.fileName.Equals(TC("$d")))
+				{
+					// This can happen on apple targets.. crawler finds some includes that are not proper includes
+					return m_logger.Warning(TC("Trying to open directory %s as file"), fileName.data);
+				}
+
+				if (out.fileName.Equals(TC("#")))
+					out.fileName.Clear().Append(fileName);
+
+				FileAccessor fa(m_logger, out.fileName.data);
+				if (fa.OpenMemoryRead(0, false))
+					return func(fa.GetData(), fa.GetSize());
+				else
+					return m_logger.Warning(TC("Failed to open %s"), out.fileName.data);
+			};
+
+		auto DevirtualizePathFunc = [this, rootsHandle = startInfo.rootsHandle](StringBufferBase& inOut) { return DevirtualizePath(inOut, rootsHandle); };
+
+		m_dependencyCrawler.Add(at + 1, startInfo.workingDir, CreateFileFunc, DevirtualizePathFunc, startInfo.application, crawlerType, startInfo.rules->index);
+	}
+
 	void GenerateNameForProcess(StringBufferBase& out, const tchar* arguments, u32 counterSuffix)
 	{
 		const tchar* start = arguments;
diff --git a/Engine/Source/Programs/UnrealBuildAccelerator/Common/Private/UbaSessionClient.cpp b/Engine/Source/Programs/UnrealBuildAccelerator/Common/Private/UbaSessionClient.cpp
--- a/Engine/Source/Programs/UnrealBuildAccelerator/Common/Private/UbaSessionClient.cpp
+++ b/Engine/Source/Programs/UnrealBuildAccelerator/Common/Private/UbaSessionClient.cpp
@@ -44,7 +44,6 @@ namespace uba
 	,	m_waitToSendEvent(false)
 	,	m_loop(true)
 	,	m_allowSpawn(true)
-	,	m_dependencyCrawler(m_logger, m_client)
 	{
 		m_maxProcessCount = info.maxProcessCount;
 		m_dedicated = info.dedicated;
@@ -313,8 +312,9 @@ namespace uba
 		return WriteBinFile(out, destFile, newKey, keyStr, fileAttributes);
 	}
 
-	bool SessionClient::PrepareProcess(ProcessStartInfoHolder& startInfo, bool isChild, StringBufferBase& outRealApplication, const tchar*& outRealWorkingDir)
+	bool SessionClient::PrepareProcess(ProcessImpl& process, bool isChild, StringBufferBase& outRealApplication, const tchar*& outRealWorkingDir)
 	{
+		ProcessStartInfoHolder& startInfo = process.m_startInfo;
 		outRealWorkingDir = m_processWorkingDir.data;
 		if (StartsWith(startInfo.application, TC("ubacopy")))
 			return true;
@@ -515,26 +515,16 @@ namespace uba
 			if (!SendRootsHandle(rootsHandle))
 				return false;
 
-		if (!process.m_parentProcess)
+		if (!process.m_parentProcess && m_useDependencyCrawler)
 			RunDependencyCrawler(process);
 
 		return true;
 	}
 
-	bool SessionClient::CreateFile(CreateFileResponse& out, const CreateFileMessage& msg)
-	{
-		if ((msg.access & FileAccess_Write) != 0)
-			return Session::CreateFile(out, msg);
-
-		TrackWorkScope tws(m_client, AsView(TC("CreateFile")), ColorWork);
-		tws.AddHint(msg.fileName);
-		return CreateFileForRead(out, tws, msg.fileName, msg.fileNameKey, msg.process.m_id, *msg.process.m_startInfo.rules);
-	}
-
-	bool SessionClient::CreateFileForRead(CreateFileResponse& out, TrackWorkScope& tws, const StringView& fileName, const StringKey& fileNameKey, u32 processId, const ApplicationRules& rules)
+	bool SessionClient::CreateFileForRead(CreateFileResponse& out, TrackWorkScope& tws, const StringView& fileName, const StringKey& fileNameKey, ProcessImpl& process, const ApplicationRules& rules)
 	{
 		CasKey casKey;
-		if (!GetCasKeyForFile(casKey, processId, fileName, fileNameKey))
+		if (!GetCasKeyForFile(casKey, process.GetId(), fileName, fileNameKey))
 			return false;
 
 		// Not finding a file is a valid path. Some applications try with a path and if fails try another path
@@ -635,7 +625,7 @@ namespace uba
 						retryEntryGuard.Execute();
 
 					MemoryMap map;
-					if (!CreateMemoryMapFromFile(map, fileNameKey, newName.data, IsCompressed(newCasKey), memoryMapAlignment, fileName.data))
+					if (!CreateMemoryMapFromFile(map, fileNameKey, newName.data, IsCompressed(newCasKey), memoryMapAlignment, fileName.data, nullptr, false))
 					{
 						if (!shouldRetry)
 							return false;
@@ -2105,67 +2095,6 @@ namespace uba
 		return logLineCount;
 	}
 
-	void SessionClient::RunDependencyCrawler(ProcessImpl& process)
-	{
-		if (!m_useDependencyCrawler)
-			return;
-
-		auto& startInfo = process.GetStartInfo();
-
-		auto crawlerType = startInfo.rules->GetDependencyCrawlerType();
-		if (crawlerType == DependencyCrawlerType_None)
-			return;
-
-		const tchar* at = TStrchr(startInfo.arguments, '@');
-		if (!at)
-			return;
-
-		auto CreateFileFunc = [this, processId = process.GetId(), rules = startInfo.rules](const StringView& fileName, const DependencyCrawler::AccessFileFunc& func)
-			{
-				CreateFileResponse out;
-				{
-					TrackWorkScope tws(m_client, AsView(TC("CrawlFile")), ColorWork);
-					tws.AddHint(fileName);
-
-					if (!CreateFileForRead(out, tws, fileName, ToStringKey(fileName), processId, *rules))
-						return false;
-				}
-
-				if (out.fileName[0] == '^')
-				{
-					MappedView view = m_fileMappingBuffer.MapView(out.fileName, out.size, fileName.data);
-					if (view.memory)
-					{
-						bool res = func(view.memory, out.size);
-						m_fileMappingBuffer.UnmapView(view, fileName.data);
-						return res;
-					}
-					else
-						return m_logger.Warning(TC("Failed to open %s"), out.fileName.data);
-				}
-				else if (out.fileName.Equals(TC("$d")))
-				{
-					// This can happen on apple targets.. crawler finds some includes that are not proper includes
-					return m_logger.Warning(TC("Trying to open directory %s as file"), fileName.data);
-				}
-				else
-				{
-					FileAccessor fa(m_logger, out.fileName.data);
-					if (fa.OpenMemoryRead(0, false))
-						return func(fa.GetData(), fa.GetSize());
-					else
-						return m_logger.Warning(TC("Failed to open %s"), out.fileName.data);
-				}
-			};
-
-		auto DevirtualizePathFunc = [this, rootsHandle = startInfo.rootsHandle](StringBufferBase& inOut)
-			{
-				return DevirtualizePath(inOut, rootsHandle);
-			};
-
-		m_dependencyCrawler.Add(at + 1, startInfo.workingDir, CreateFileFunc, DevirtualizePathFunc, startInfo.application, crawlerType, startInfo.rules->index);
-	}
-
 	bool SessionClient::ParseDirectoryTable()
 	{
 		SCOPED_WRITE_LOCK(m_directoryTable.m_lookupLock, lock);
diff --git a/Engine/Source/Programs/UnrealBuildAccelerator/Common/Private/UbaSessionServer.cpp b/Engine/Source/Programs/UnrealBuildAccelerator/Common/Private/UbaSessionServer.cpp
--- a/Engine/Source/Programs/UnrealBuildAccelerator/Common/Private/UbaSessionServer.cpp
+++ b/Engine/Source/Programs/UnrealBuildAccelerator/Common/Private/UbaSessionServer.cpp
@@ -2058,9 +2058,9 @@ namespace uba
 		return str.Appendf(TC("%s"), findIt->second.GetStartInfo().GetDescription()).data;
 	}
 
-	bool SessionServer::PrepareProcess(ProcessStartInfoHolder& startInfo, bool isChild, StringBufferBase& outRealApplication, const tchar*& outRealWorkingDir)
+	bool SessionServer::PrepareProcess(ProcessImpl& process, bool isChild, StringBufferBase& outRealApplication, const tchar*& outRealWorkingDir)
 	{
-		if (!Session::PrepareProcess(startInfo, isChild, outRealApplication, outRealWorkingDir))
+		if (!Session::PrepareProcess(process, isChild, outRealApplication, outRealWorkingDir))
 			return false;
 
 		if (!m_memTotal || !m_allowWaitOnMem || isChild)
diff --git a/Engine/Source/Programs/UnrealBuildAccelerator/Common/Private/UbaStorage.cpp b/Engine/Source/Programs/UnrealBuildAccelerator/Common/Private/UbaStorage.cpp
--- a/Engine/Source/Programs/UnrealBuildAccelerator/Common/Private/UbaStorage.cpp
+++ b/Engine/Source/Programs/UnrealBuildAccelerator/Common/Private/UbaStorage.cpp
@@ -644,7 +644,7 @@ namespace uba
 		return false;
 	}
 
-	bool StorageImpl::DecompressMemoryToMemory(u8* compressedData, u64 compressedSize, u8* writeData, u64 decompressedSize, const tchar* readHint, const tchar* writeHint)
+	bool StorageImpl::DecompressMemoryToMemory(const u8* compressedData, u64 compressedSize, u8* writeData, u64 decompressedSize, const tchar* readHint, const tchar* writeHint)
 	{
 		UBA_ASSERTF(compressedData, TC("DecompressMemoryToMemory got readmem nullptr (%s)"), readHint);
 		UBA_ASSERTF(writeData, TC("DecompressMemoryToMemory got writemem nullptr (%s)"), writeHint);
@@ -658,7 +658,7 @@ namespace uba
 				Logger* logger;
 				const tchar* hint;
 				Atomic<u64> refCount;
-				u8* readPos = nullptr;
+				const u8* readPos = nullptr;
 				u8* writePos = nullptr;
 
 				Futex lock;
@@ -695,10 +695,10 @@ namespace uba
 								delete rec;
 							return;
 						}
-						u8* readPos = rec->readPos;
+						const u8* readPos = rec->readPos;
 						u8* writePos = rec->writePos;
-						u32 compressedBlockSize = ((u32*)readPos)[0];
-						u32 decompressedBlockSize = ((u32*)readPos)[1];
+						u32 compressedBlockSize = ((const u32*)readPos)[0];
+						u32 decompressedBlockSize = ((const u32*)readPos)[1];
 
 						if (decompressedBlockSize == 0 || decompressedBlockSize > rec->decompressedSize)
 						{
@@ -760,16 +760,16 @@ namespace uba
 		}
 		else
 		{
-			u8* readPos = compressedData;
+			const u8* readPos = compressedData;
 			u8* writePos = writeData;
 
 			u64 left = decompressedSize;
 			while (left)
 			{
-				u32 compressedBlockSize = ((u32*)readPos)[0];
+				u32 compressedBlockSize = ((const u32*)readPos)[0];
 				if (!compressedBlockSize)
 					break;
-				u32 decompressedBlockSize = ((u32*)readPos)[1];
+				u32 decompressedBlockSize = ((const u32*)readPos)[1];
 				if (decompressedBlockSize == 0 || decompressedBlockSize > left)
 					return m_logger.Error(TC("Decompressed block size %u is invalid. Decompressed file is %u (%s -> %s)"), decompressedBlockSize, decompressedSize, readHint, writeHint);
 				readPos += sizeof(u32) * 2;
@@ -2435,7 +2435,7 @@ namespace uba
 						}
 						else
 						{
-							if (!DecompressFileToMemory(CasKeyString(actualKey).str, readHandle, destinationMemory, decompressedSize, destination))
+							if (!DecompressFileToMemory(CasKeyString(actualKey).str, readHandle, destinationMemory, decompressedSize, destination, 0))
 								return false;
 						}
 						return true;
@@ -2675,6 +2675,7 @@ namespace uba
 		if (m_casLookup.empty())
 		{
 			logger.Info(TC("  Storage not loaded"));
+			logger.Info(TC(""));
 			return;
 		}
 		
@@ -2751,7 +2752,7 @@ namespace uba
 		return ToCasKey(hasher, storeCompressed);
 	}
 
-	bool StorageImpl::DecompressFileToMemory(const tchar* fileName, FileHandle fileHandle, u8* dest, u64 decompressedSize, const tchar* writeHint)
+	bool StorageImpl::DecompressFileToMemory(const tchar* fileName, FileHandle fileHandle, u8* dest, u64 decompressedSize, const tchar* writeHint, u64 fileStartOffset)
 	{
 		if (m_workManager && decompressedSize > BufferSlotSize*4) // Arbitrary size threshold. We want to at least catch the pch here
 		{
@@ -2767,7 +2768,8 @@ namespace uba
 				return m_logger.Error(TC("Failed to map view of file mapping for %s (%s)"), fileName, LastErrorToText().data);
 			auto udg = MakeGuard([&]() { UnmapViewOfFile(m_logger, fileData, compressedSize, fileName); });
 			
-			if (!DecompressMemoryToMemory(fileData + 8, compressedSize, dest, decompressedSize, fileName, writeHint))
+			u8* readPos = fileData + 8 + fileStartOffset;
+			if (!DecompressMemoryToMemory(readPos, compressedSize, dest, decompressedSize, fileName, writeHint))
 				return false;
 		}
 		else
diff --git a/Engine/Source/Programs/UnrealBuildAccelerator/Common/Public/UbaDependencyCrawler.h b/Engine/Source/Programs/UnrealBuildAccelerator/Common/Public/UbaDependencyCrawler.h
--- a/Engine/Source/Programs/UnrealBuildAccelerator/Common/Public/UbaDependencyCrawler.h
+++ b/Engine/Source/Programs/UnrealBuildAccelerator/Common/Public/UbaDependencyCrawler.h
@@ -30,7 +30,7 @@ namespace uba
 		using DevirtualizePathFunc = Function<bool(StringBufferBase& inOut)>;
 		using CreateFileFunc = Function<bool(const StringView& fileName, const AccessFileFunc& func)>;
 
-		inline void Init(FileExistsFunc&& fileExistsFunc, TraverseFilesFunc&& traverseFilesFunc);
+		inline void Init(FileExistsFunc&& fileExistsFunc, TraverseFilesFunc&& traverseFilesFunc, bool useBloomFilter = true);
 		inline bool Add(const tchar* rsp, const tchar* workDir, CreateFileFunc&& func, DevirtualizePathFunc&& devirtualizePathFunc, const tchar* app, DependencyCrawlerType type, u32 ruleIndex);
 
 		Logger& m_logger;
@@ -43,6 +43,7 @@ namespace uba
 		Futex m_pchLookupLock;
 		struct Pch { Futex lock; bool handled = false; UnorderedSet<StringKey> files; };
 		UnorderedMap<StringKey, Pch> m_pchLookup;
+		bool m_useBloomFilter = true;
 
 		struct IncludeRoot
 		{
@@ -101,10 +102,11 @@ namespace uba
 	{
 	}
 
-	void DependencyCrawler::Init(FileExistsFunc&& fileExistsFunc, TraverseFilesFunc&& traverseFilesFunc)
+	void DependencyCrawler::Init(FileExistsFunc&& fileExistsFunc, TraverseFilesFunc&& traverseFilesFunc, bool useBloomFilter)
 	{
 		m_fileExistsFunc = std::move(fileExistsFunc);
 		m_traverseFilesFunc = std::move(traverseFilesFunc);
+		m_useBloomFilter = useBloomFilter;
 	}
 
 	bool DependencyCrawler::Add(const tchar* rsp, const tchar* workDir, CreateFileFunc&& createFileFunc, DevirtualizePathFunc&& devirtualizePathFunc, const tchar* app, DependencyCrawlerType type, u32 ruleIndex)
@@ -157,6 +159,14 @@ namespace uba
 						}
 						instance.pch = &pch;
 					}
+					else if (instance.type == DependencyCrawlerType_MsvcLinker)
+					{
+						// If we end up here it means that we have compressed obj files enabled and want to decompress them in parallel
+						m_workManager.AddWork([this, ref, codeFile = cf.path](const WorkContext& context)
+							{
+								ref.instance.createFileFunc(codeFile, {});
+							}, 1, TC("Crawler"), ColorWork);
+					}
 					else
 					{
 						ParseCodeFile(instance, cf.path, instance.platform.empty(), instance.rsp);
@@ -192,7 +202,7 @@ namespace uba
 								codeFilePath.Resize(fixedPath.count).Append(file);
 								m_workManager.AddWork([this, ref, codeFile = TString(codeFilePath.data)](const WorkContext& context)
 									{
-										if (!ref.instance.createFileFunc(codeFile, [&](const void* data, u64 dataSize) { return true; }))
+										if (!ref.instance.createFileFunc(codeFile, {}))
 											m_logger.Warning(TC("Failed to open file %s from builtindir"), codeFile.c_str());
 									}, 1, TC("Crawler"), ColorWork);
 							});
@@ -246,12 +256,16 @@ namespace uba
 				fixedPath2.EnsureEndsWithSlash();
 
 				BloomFilter bloomFilter;
-				m_traverseFilesFunc(fixedPath2, [&](const StringView& file, bool isDirectory)
-					{
-						bloomFilter.Add(ToStringKey(file));
-					});
-				if (!bloomFilter.IsEmpty())
-					instance.includeRoots.push_back({TString(fixedPath2.data), bloomFilter});
+				if (m_useBloomFilter)
+				{
+					m_traverseFilesFunc(fixedPath2, [&](const StringView& file, bool isDirectory)
+						{
+							bloomFilter.Add(ToStringKey(file));
+						});
+					if (bloomFilter.IsEmpty())
+						return;
+				}
+				instance.includeRoots.push_back({TString(fixedPath2.data), bloomFilter});
 			};
 
 		auto IgnoreOption = [](const StringView& path) { UBA_CRAWLER_DEBUG_PRINT("IGNORED", path.data); };
@@ -486,7 +500,28 @@ namespace uba
 					else
 						IgnoreOption(sb);
 				});
-			}
+		}
+		else if (instance.type == DependencyCrawlerType_MsvcLinker)
+		{
+			StringBuffer<> prevArg;
+			ParseArguments((const char*)data, dataSize, [&](const char* arg, u32 argLen)
+				{
+					StringBuffer<> sb;
+					sb.Append(arg, argLen);
+
+					if (sb[0] == '/')
+					{
+						if (sb.StartsWith(TC("/LIBPATH")))
+							AddRoot(StringView(sb).Skip(9));
+						return;
+					}
+					else
+					{
+						if (sb.EndsWith(TC(".obj")))
+							AddCodeFile(sb, false, false);
+					}
+				});
+		}
 				
 		return true;
 	}
@@ -669,7 +704,7 @@ namespace uba
 
 								for (auto& root : instance.includeRoots)
 								{
-									if (root.bloomFilter.IsGuaranteedMiss(fileKey))
+									if (m_useBloomFilter && root.bloomFilter.IsGuaranteedMiss(fileKey))
 										continue;
 									if (HandleInclude(instance, root.path, include2, codeFile2, parseDefines))
 										return;
diff --git a/Engine/Source/Programs/UnrealBuildAccelerator/Common/Public/UbaProcess.h b/Engine/Source/Programs/UnrealBuildAccelerator/Common/Public/UbaProcess.h
--- a/Engine/Source/Programs/UnrealBuildAccelerator/Common/Public/UbaProcess.h
+++ b/Engine/Source/Programs/UnrealBuildAccelerator/Common/Public/UbaProcess.h
@@ -160,6 +160,9 @@ namespace uba
 		ProcessStats m_processStats;
 		KernelStats m_kernelStats;
 
+		Futex m_usedFileMappingsLock;
+		UnorderedSet<StringKey> m_usedFileMappings;
+
 		Thread m_messageThread;
 
 		bool m_extractExports = false;
diff --git a/Engine/Source/Programs/UnrealBuildAccelerator/Common/Public/UbaSession.h b/Engine/Source/Programs/UnrealBuildAccelerator/Common/Public/UbaSession.h
--- a/Engine/Source/Programs/UnrealBuildAccelerator/Common/Public/UbaSession.h
+++ b/Engine/Source/Programs/UnrealBuildAccelerator/Common/Public/UbaSession.h
@@ -3,6 +3,7 @@
 #pragma once
 
 #include "UbaBottleneck.h"
+#include "UbaDependencyCrawler.h"
 #include "UbaFile.h"
 #include "UbaDirectoryTable.h"
 #include "UbaFileMapping.h"
@@ -14,6 +15,8 @@
 #include "UbaTrace.h"
 #include "UbaThread.h"
 
+#define UBA_DEBUG_TRACK_MAPPING 0 // UBA_DEBUG_LOGGER
+
 namespace uba
 {
 	class Process;
@@ -120,7 +123,7 @@ namespace uba
 		struct BinaryModule { TString name; TString path; u32 fileAttributes = 0; bool isSystem = false; };
 		bool GetBinaryModules(Vector<BinaryModule>& out, const tchar* application);
 		void Free(Vector<BinaryModule>& v);
-		bool IsRarelyRead(ProcessImpl& process, const StringBufferBase& fileName) const;
+		bool IsRarelyRead(ProcessImpl& process, const StringView& fileName) const;
 		bool IsRarelyReadAfterWritten(ProcessImpl& process, const StringView& fileName) const;
 		bool IsKnownSystemFile(const tchar* applicationName);
 		bool ShouldWriteToDisk(const StringView& fileName);
@@ -128,13 +131,14 @@ namespace uba
 		u32 AddFileMapping(StringKey fileNameKey, const tchar* fileName, const tchar* newFileName, u64 fileSize = InvalidValue);
 		
 		struct MemoryMap { StringBuffer<128> name; u64 size = 0; };
-		bool CreateMemoryMapFromFile(MemoryMap& out, StringKey fileNameKey, const tchar* fileName, bool isCompressed, u64 alignment, const tchar* hint);
+		struct FileMappingEntry;
+		bool CreateMemoryMapFromFile(MemoryMap& out, StringKey fileNameKey, const tchar* fileName, bool isCompressedCas, u64 alignment, const tchar* hint, ProcessImpl* requestingProcess, bool canBeFreed);
 		bool CreateMemoryMapFromView(MemoryMap& out, StringKey fileNameKey, const tchar* fileName, const CasKey& casKey, u64 alignment);
 
 		bool RegisterCreateFileForWrite(StringKey fileNameKey, const StringView& fileName, bool registerRealFile, u64 fileSize = 0, u64 lastWriteTime = 0, bool invalidateStorage = true);
 		u32 RegisterDeleteFile(StringKey fileNameKey, const StringView& fileName);
 
-		virtual bool PrepareProcess(ProcessStartInfoHolder& startInfo, bool isChild, StringBufferBase& outRealApplication, const tchar*& outRealWorkingDir);
+		virtual bool PrepareProcess(ProcessImpl& process, bool isChild, StringBufferBase& outRealApplication, const tchar*& outRealWorkingDir);
 		virtual void* GetProcessEnvironmentVariables();
 		virtual void PrintSessionStats(Logger& logger);
 
@@ -143,6 +147,7 @@ namespace uba
 		virtual bool ProcessThreadStart(ProcessImpl& process);
 		virtual bool GetInitResponse(InitResponse& out, const InitMessage& msg);
 		virtual bool CreateFile(CreateFileResponse& out, const CreateFileMessage& msg);
+		virtual bool CreateFileForRead(CreateFileResponse& out, TrackWorkScope& tws, const StringView& fileName, const StringKey& fileNameKey, ProcessImpl& process, const ApplicationRules& rules);
 		virtual bool CloseFile(CloseFileResponse& out, const CloseFileMessage& msg);
 		virtual bool DeleteFile(DeleteFileResponse& out, const DeleteFileMessage& msg);
 		virtual bool CopyFile(CopyFileResponse& out, const CopyFileMessage& msg);
@@ -160,6 +165,7 @@ namespace uba
 		virtual bool SHGetKnownFolderPath(Process& process, BinaryReader& reader, BinaryWriter& writer);
 		virtual bool HostRun(BinaryReader& reader, BinaryWriter& writer);
 		virtual bool GetSymbols(const tchar* application, bool isArm, BinaryReader& reader, BinaryWriter& writer);
+		virtual bool CheckRemapping(ProcessImpl& process, BinaryReader& reader, BinaryWriter& writer);
 		virtual void FileEntryAdded(StringKey fileNameKey, u64 lastWritten, u64 size);
 		virtual bool FlushWrittenFiles(ProcessImpl& process);
 		virtual bool UpdateEnvironment(ProcessImpl& process, const StringView& reason, bool resetStats);
@@ -191,6 +197,8 @@ namespace uba
 		void TraceWrittenFile(u32 processId, const StringView& file, u64 size = 0);
 		virtual void TraceSessionUpdate();
 
+		void RunDependencyCrawler(ProcessImpl& process);
+
 		Storage& m_storage;
 		MutableLogger m_logger;
 
@@ -241,7 +249,25 @@ namespace uba
 		u8* m_fileMappingTableMem;
 		u32 m_fileMappingTableSize = 0;
 		Futex m_fileMappingTableLookupLock;
-		struct FileMappingEntry { Futex lock; FileMappingHandle mapping; u64 mappingOffset = 0; u64 size = 0; u64 lastWriteTime = 0; bool isDir = false; bool handled = false; bool success = false; bool owned = false; };
+		struct FileMappingEntry
+		{
+			Futex lock;
+			FileMappingHandle mapping;
+			u64 mappingOffset = 0;
+			u64 size = 0;
+			u64 lastWriteTime = 0;
+			u16 refCount = 0;
+			u8 usedCount = 0;
+			u8 usedCountBeforeFree = 0;
+			bool isDir = false;
+			bool handled = false;
+			bool success = false;
+			bool canBeFreed = false;
+
+			#if UBA_DEBUG_TRACK_MAPPING
+			TString name;
+			#endif
+		};
 		UnorderedMap<StringKey, FileMappingEntry> m_fileMappingTableLookup;
 
 		static constexpr u64 NameToHashMemSize = 48*1024*1024;
@@ -316,6 +342,8 @@ namespace uba
 
 		VolumeCache m_volumeCache;
 
+		DependencyCrawler m_dependencyCrawler;
+
 		friend class ProcessImpl;
 	};
 
diff --git a/Engine/Source/Programs/UnrealBuildAccelerator/Common/Public/UbaSessionClient.h b/Engine/Source/Programs/UnrealBuildAccelerator/Common/Public/UbaSessionClient.h
--- a/Engine/Source/Programs/UnrealBuildAccelerator/Common/Public/UbaSessionClient.h
+++ b/Engine/Source/Programs/UnrealBuildAccelerator/Common/Public/UbaSessionClient.h
@@ -2,7 +2,6 @@
 
 #pragma once
 
-#include "UbaDependencyCrawler.h"
 #include "UbaNetwork.h"
 #include "UbaSession.h"
 
@@ -54,10 +53,10 @@ namespace uba
 	private:
 		bool RetrieveCasFile(CasKey& outNewKey, u64& outSize, const CasKey& casKey, const tchar* hint, bool storeUncompressed, bool allowProxy = true);
 
-		virtual bool PrepareProcess(ProcessStartInfoHolder& startInfo, bool isChild, StringBufferBase& outRealApplication, const tchar*& outRealWorkingDir) override;
+		virtual bool PrepareProcess(ProcessImpl& process, bool isChild, StringBufferBase& outRealApplication, const tchar*& outRealWorkingDir) override;
 		virtual void* GetProcessEnvironmentVariables() override;
-		virtual bool ProcessThreadStart(ProcessImpl& process);
-		virtual bool CreateFile(CreateFileResponse& out, const CreateFileMessage& msg) override;
+		virtual bool ProcessThreadStart(ProcessImpl& process) override;
+		virtual bool CreateFileForRead(CreateFileResponse& out, TrackWorkScope& tws, const StringView& fileName, const StringKey& fileNameKey, ProcessImpl& process, const ApplicationRules& rules) override;
 		virtual bool DeleteFile(DeleteFileResponse& out, const DeleteFileMessage& msg) override;
 		virtual bool CopyFile(CopyFileResponse& out, const CopyFileMessage& msg) override;
 		virtual bool MoveFile(MoveFileResponse& out, const MoveFileMessage& msg) override;
@@ -83,7 +82,6 @@ namespace uba
 		struct InternalProcessStartInfo;
 		struct ModuleInfo;
 
-		bool CreateFileForRead(CreateFileResponse& out, TrackWorkScope& tws, const StringView& fileName, const StringKey& fileNameKey, u32 processId, const ApplicationRules& rules);
 		bool GetCasKeyForFile(CasKey& out, u32 processId, const StringView& fileName, const StringKey& fileNameKey);
 		bool ReadModules(List<ModuleInfo>& outModules, u32 processId, const tchar* application);
 		bool EnsureBinaryFile(StringBufferBase& out, StringBufferBase& outVirtual, u32 processId, const StringBufferBase& fileName, const StringKey& fileNameKey, const tchar* applicationDir, const u8* loaderPaths, u32 loaderPathsSize);
@@ -106,7 +104,6 @@ namespace uba
 		void SendLogFileToServer(ProcessImpl& pi);
 		void GetLogFileName(StringBufferBase& out, const tchar* logFile, const tchar* arguments, u32 processId);
 		u32 WriteLogLines(BinaryWriter& writer, ProcessImpl& process);
-		void RunDependencyCrawler(ProcessImpl& process);
 		bool ParseDirectoryTable();
 		bool EntryExists(const StringView& path, u32& outTableOffset);
 
@@ -183,7 +180,6 @@ namespace uba
 		UnorderedMap<CasKey, Vector<u8>> m_hostRunCache;
 		Futex m_hostRunCacheLock;
 
-		DependencyCrawler m_dependencyCrawler;
 		u32 m_dirtableParsedPosition = 0;
 		Futex m_dirVisitedLock;
 		struct DirVisitedEntry { Futex lock; bool handled = false; };
diff --git a/Engine/Source/Programs/UnrealBuildAccelerator/Common/Public/UbaSessionServer.h b/Engine/Source/Programs/UnrealBuildAccelerator/Common/Public/UbaSessionServer.h
--- a/Engine/Source/Programs/UnrealBuildAccelerator/Common/Public/UbaSessionServer.h
+++ b/Engine/Source/Programs/UnrealBuildAccelerator/Common/Public/UbaSessionServer.h
@@ -99,7 +99,7 @@ namespace uba
 		ProcessHandle GetProcess(u32 processId);
 		TString GetProcessDescription(u32 processId);
 
-		virtual bool PrepareProcess(ProcessStartInfoHolder& startInfo, bool isChild, StringBufferBase& outRealApplication, const tchar*& outRealWorkingDir) override final;
+		virtual bool PrepareProcess(ProcessImpl& process, bool isChild, StringBufferBase& outRealApplication, const tchar*& outRealWorkingDir) override final;
 		virtual bool CreateFile(CreateFileResponse& out, const CreateFileMessage& msg) override final;
 		virtual void FileEntryAdded(StringKey fileNameKey, u64 lastWritten, u64 size) override final;
 		virtual void PrintSessionStats(Logger& logger) override final;
diff --git a/Engine/Source/Programs/UnrealBuildAccelerator/Common/Public/UbaStorage.h b/Engine/Source/Programs/UnrealBuildAccelerator/Common/Public/UbaStorage.h
--- a/Engine/Source/Programs/UnrealBuildAccelerator/Common/Public/UbaStorage.h
+++ b/Engine/Source/Programs/UnrealBuildAccelerator/Common/Public/UbaStorage.h
@@ -32,8 +32,8 @@ namespace uba
 		virtual bool GetZone(StringBufferBase& out) = 0;
 		virtual bool HasProxy(u32 clientId) { return false; }
 
-		virtual bool DecompressFileToMemory(const tchar* fileName, FileHandle fileHandle, u8* dest, u64 decompressedSize, const tchar* writeHint) = 0;
-		virtual bool DecompressMemoryToMemory(u8* compressedData, u64 compressedSize, u8* writeData, u64 decompressedSize, const tchar* readHint, const tchar* writeHint) = 0;
+		virtual bool DecompressFileToMemory(const tchar* fileName, FileHandle fileHandle, u8* dest, u64 decompressedSize, const tchar* writeHint, u64 fileStartOffset) = 0;
+		virtual bool DecompressMemoryToMemory(const u8* compressedData, u64 compressedSize, u8* writeData, u64 decompressedSize, const tchar* readHint, const tchar* writeHint) = 0;
 		virtual bool CreateDirectory(const tchar* dir) = 0;
 		virtual bool DeleteCasForFile(const tchar* file) = 0;
 
@@ -130,7 +130,7 @@ namespace uba
 		virtual bool StoreCompressed() const final { return m_storeCompressed; }
 		virtual void PrintSummary(Logger& logger) override;
 
-		virtual bool DecompressFileToMemory(const tchar* fileName, FileHandle fileHandle, u8* dest, u64 decompressedSize, const tchar* writeHint) override;
+		virtual bool DecompressFileToMemory(const tchar* fileName, FileHandle fileHandle, u8* dest, u64 decompressedSize, const tchar* writeHint, u64 fileStartOffset) override;
 		virtual bool CreateDirectory(const tchar* dir) override;
 		virtual bool DeleteCasForFile(const tchar* file) override;
 		virtual bool RetrieveCasFile(RetrieveResult& out, const CasKey& casKey, const tchar* hint, FileMappingBuffer* mappingBuffer = nullptr, u64 memoryMapAlignment = 1, bool allowProxy = true) override;
@@ -173,7 +173,7 @@ namespace uba
 		bool AddCasFile(StringKey fileNameKey, const tchar* fileName, const CasKey& casKey, bool deferCreation);
 		void CasEntryAccessed(const CasKey& casKey);
 		virtual bool IsDisallowedPath(const tchar* fileName);
-		virtual bool DecompressMemoryToMemory(u8* compressedData, u64 compressedSize, u8* writeData, u64 decompressedSize, const tchar* readHint, const tchar* writeHint) override;
+		virtual bool DecompressMemoryToMemory(const u8* compressedData, u64 compressedSize, u8* writeData, u64 decompressedSize, const tchar* readHint, const tchar* writeHint) override;
 		bool DecompressMemoryToFile(u8* compressedData, FileAccessor& destination, u64 decompressedSize, bool useNoBuffering);
 
 		void CasEntryAccessed(CasEntry& entry);
diff --git a/Engine/Source/Programs/UnrealBuildAccelerator/Core/Private/UbaApplicationRules.cpp b/Engine/Source/Programs/UnrealBuildAccelerator/Core/Private/UbaApplicationRules.cpp
--- a/Engine/Source/Programs/UnrealBuildAccelerator/Core/Private/UbaApplicationRules.cpp
+++ b/Engine/Source/Programs/UnrealBuildAccelerator/Core/Private/UbaApplicationRules.cpp
@@ -11,9 +11,19 @@ namespace uba
 			return false;
 		return fileName.EndsWith(TC(".obj"))
 			|| fileName.EndsWith(TC(".o"))
+			|| fileName.EndsWith(TC(".pch"))
 			;
 	}
 
+	u8 GlobalRules::GetUsedCountBeforeFree(const StringView& fileName)
+	{
+		if (!fileName.EndsWith(TC(".obj")))
+			return 0;
+		if (fileName.EndsWith(TC(".h.obj")))
+			return 255;
+		return 2;
+	}
+
 	GlobalRules g_globalRules;
 
 
@@ -81,6 +91,7 @@ namespace uba
 			return fileName.EndsWith(TC(".obj"))
 				|| fileName.EndsWith(TC(".dep.json"))
 				|| fileName.EndsWith(TC(".sarif"))
+				|| fileName.EndsWith(TC(".pch"))
 				|| fileName.EndsWith(TC(".ifc"))
 				|| fileName.EndsWith(TC(".rc2.res")) // Not really an obj file.. 
 				;// || fileName.EndsWith(TC(".h.pch") // Not tested enough
@@ -127,6 +138,11 @@ namespace uba
 		{
 			return DependencyCrawlerType_MsvcCompiler;
 		}
+
+		virtual bool CanDependOnCompressedFiles()
+		{
+			return true; // .pch
+		}
 	};
 
 	class ApplicationRulesVcLink : public ApplicationRulesVC
@@ -202,6 +218,11 @@ namespace uba
 		{
 			return true;
 		}
+
+		virtual DependencyCrawlerType GetDependencyCrawlerType() const override
+		{
+			return DependencyCrawlerType_MsvcLinker;
+		}
 	};
 
 	class ApplicationRulesLinkExe : public ApplicationRulesVcLink
diff --git a/Engine/Source/Programs/UnrealBuildAccelerator/Core/Public/UbaApplicationRules.h b/Engine/Source/Programs/UnrealBuildAccelerator/Core/Public/UbaApplicationRules.h
--- a/Engine/Source/Programs/UnrealBuildAccelerator/Core/Public/UbaApplicationRules.h
+++ b/Engine/Source/Programs/UnrealBuildAccelerator/Core/Public/UbaApplicationRules.h
@@ -10,6 +10,7 @@ namespace uba
 	{
 	public:
 		bool FileCanBeCompressed(const StringView& fileName);
+		u8 GetUsedCountBeforeFree(const StringView& fileName);
 	};
 
 	extern GlobalRules g_globalRules;
@@ -20,6 +21,8 @@ namespace uba
 		DependencyCrawlerType_None,
 		DependencyCrawlerType_ClangCompiler,
 		DependencyCrawlerType_MsvcCompiler,
+		DependencyCrawlerType_ClangLinker,
+		DependencyCrawlerType_MsvcLinker,
 	};
 
 
diff --git a/Engine/Source/Programs/UnrealBuildAccelerator/Core/Public/UbaProtocol.h b/Engine/Source/Programs/UnrealBuildAccelerator/Core/Public/UbaProtocol.h
--- a/Engine/Source/Programs/UnrealBuildAccelerator/Core/Public/UbaProtocol.h
+++ b/Engine/Source/Programs/UnrealBuildAccelerator/Core/Public/UbaProtocol.h
@@ -39,6 +39,7 @@ namespace uba
 		UBA_PROCESS_MESSAGE(RpcCommunication) \
 		UBA_PROCESS_MESSAGE(HostRun) \
 		UBA_PROCESS_MESSAGE(ResolveCallstack) \
+		UBA_PROCESS_MESSAGE(CheckRemapping) \
 
 	enum MessageType : u8
 	{
@@ -48,7 +49,7 @@ namespace uba
 		#undef UBA_PROCESS_MESSAGE
 	};
 
-	inline constexpr u32 ProcessMessageVersion = 1340;
+	inline constexpr u32 ProcessMessageVersion = 1341;
 
 	inline constexpr u32 CommunicationMemSize = IsWindows ? 64*1024 : 64*1024*2; // Macos expands some commandlines to be crazy long
 
diff --git a/Engine/Source/Programs/UnrealBuildAccelerator/Core/Public/UbaSynchronization.h b/Engine/Source/Programs/UnrealBuildAccelerator/Core/Public/UbaSynchronization.h
--- a/Engine/Source/Programs/UnrealBuildAccelerator/Core/Public/UbaSynchronization.h
+++ b/Engine/Source/Programs/UnrealBuildAccelerator/Core/Public/UbaSynchronization.h
@@ -165,7 +165,7 @@ namespace uba
 	struct ScopeGuard
 	{
 		void Cancel() { m_called = true; }
-		auto Execute() { m_called = true; return m_lambda(); }
+		auto Execute() { bool called = m_called; m_called = true; if constexpr (!std::is_void_v<std::invoke_result_t<Lambda>>) return (called ? decltype(m_lambda()){} : m_lambda()); else if (!called) m_lambda(); }
 
 		ScopeGuard(Lambda lambda) : m_lambda(lambda) {}
 		ScopeGuard(ScopeGuard&& o) { m_lambda = std::move(o.m_lambda); o.m_called = true; }
diff --git a/Engine/Source/Programs/UnrealBuildAccelerator/Detours/Private/PosixOS/UbaDetoursFunctionsPosix.cpp b/Engine/Source/Programs/UnrealBuildAccelerator/Detours/Private/PosixOS/UbaDetoursFunctionsPosix.cpp
--- a/Engine/Source/Programs/UnrealBuildAccelerator/Detours/Private/PosixOS/UbaDetoursFunctionsPosix.cpp
+++ b/Engine/Source/Programs/UnrealBuildAccelerator/Detours/Private/PosixOS/UbaDetoursFunctionsPosix.cpp
@@ -329,7 +329,6 @@ int Shared_open(const char* funcName, const char* file, int flags, int mode, con
 
 	bool keepInMemory = false;
 
-	const char* realFileName = fileName.data;
 	u64 size = InvalidValue;
 	u32 closeId = 0;
 	u32 desiredAccess = GetFileAccessFlags(flags);
@@ -378,6 +377,7 @@ int Shared_open(const char* funcName, const char* file, int flags, int mode, con
 		}
 	}
 
+	const char* realFileName = fileName.data;
 
 	SCOPED_WRITE_LOCK(g_mappedFileTable.m_lookupLock, _);
 	auto insres = g_mappedFileTable.m_lookup.try_emplace(fileNameKey);
@@ -391,7 +391,7 @@ int Shared_open(const char* funcName, const char* file, int flags, int mode, con
 		if (!keepInMemory)
 		{
 			char newFileName[512];
-			Rpc_CreateFileW(realFileName, fileNameKey, u8(desiredAccess), newFileName, sizeof_array(newFileName), size, closeId, false);
+			Rpc_CreateFileW(fileName, fileNameKey, u8(desiredAccess), newFileName, sizeof_array(newFileName), size, closeId, false);
 			info.name = g_memoryBlock.Strdup(newFileName);
 			realFileName = info.name;
 		}
@@ -413,7 +413,7 @@ int Shared_open(const char* funcName, const char* file, int flags, int mode, con
 			{
 				info.deleted = false;
 				char newFileName[1024];
-				Rpc_CreateFileW(realFileName, fileNameKey, u8(desiredAccess), newFileName, sizeof_array(newFileName), size, closeId, false);
+				Rpc_CreateFileW(fileName, fileNameKey, u8(desiredAccess), newFileName, sizeof_array(newFileName), size, closeId, false);
 				info.name = g_memoryBlock.Strdup(newFileName);
 				realFileName = info.name;
 			}
@@ -1555,7 +1555,7 @@ UBA_EXPORT int UBA_WRAPPER(rename)(const char* oldpath, const char* newpath)
 				{
 					char temp[1024];
 					u64 size;
-					Rpc_CreateFileW(fixedNewPath.data, newKey, AccessFlag_Write, temp, sizeof_array(temp), size, fo.closeId, true);
+					Rpc_CreateFileW(fixedNewPath, newKey, AccessFlag_Write, temp, sizeof_array(temp), size, fo.closeId, true);
 				}
 
 				bool wasTempFile = fixedNewPath.StartsWith(g_systemTemp.data);
diff --git a/Engine/Source/Programs/UnrealBuildAccelerator/Detours/Private/UbaDetoursFileMappingTable.cpp b/Engine/Source/Programs/UnrealBuildAccelerator/Detours/Private/UbaDetoursFileMappingTable.cpp
--- a/Engine/Source/Programs/UnrealBuildAccelerator/Detours/Private/UbaDetoursFileMappingTable.cpp
+++ b/Engine/Source/Programs/UnrealBuildAccelerator/Detours/Private/UbaDetoursFileMappingTable.cpp
@@ -39,12 +39,17 @@ namespace uba
 			auto insres = m_lookup.try_emplace(g);
 			FileInfo& info = insres.first->second;
 			if (!insres.second)
+			{
+				if (info.name && info.name[0] == '^' && !mappedFileName.Equals(info.name)) // Mapped file has been re-mapped.
+				{
+					info.name = m_memoryBlock.Strdup(mappedFileName).data;
+					UBA_ASSERT(info.size == size);
+				}
 				continue;
+			}
 			info.fileNameKey = g;
 			info.name = m_memoryBlock.Strdup(mappedFileName).data;
-			//info.refCount = 0;
 			info.size = size;
-			//info.dwDesiredAccess = GENERIC_READ;
 		}
 		m_memPosition = tableSize;
 	}
@@ -66,7 +71,7 @@ namespace uba
 		sourceInfo.lastDesiredAccess = 0;
 	}
 
-	void Rpc_CreateFileW(const tchar* fileName, const StringKey& fileNameKey, u8 access, tchar* outNewName, u64 newNameCapacity, u64& outSize, u32& outCloseId, bool lock)
+	void Rpc_CreateFileW(const StringView& fileName, const StringKey& fileNameKey, u8 access, tchar* outNewName, u64 newNameCapacity, u64& outSize, u32& outCloseId, bool lock)
 	{
 		TimerScope ts(g_stats.createFile);
 		SCOPED_WRITE_LOCK(g_communicationLock, pcs);
@@ -83,7 +88,7 @@ namespace uba
 		u32 mappedFileTableSize = reader.ReadU32();
 		u32 directoryTableSize = u32(reader.ReadU32());
 		pcs.Leave();
-		DEBUG_LOG_PIPE(L"CreateFile", L"%ls (%ls)", (access == 0 ? L"ATTRIB" : ((access & AccessFlag_Write) ? L"WRITE" : L"READ")), fileName);
+		DEBUG_LOG_PIPE(L"CreateFile", L"%ls (%ls)", (access == 0 ? L"ATTRIB" : ((access & AccessFlag_Write) ? L"WRITE" : L"READ")), fileName.data);
 
 		if (lock)
 			g_mappedFileTable.Parse(mappedFileTableSize);
@@ -92,6 +97,20 @@ namespace uba
 		g_directoryTable.ParseDirectoryTable(directoryTableSize);
 	}
 
+	void Rpc_CheckRemapping(const StringView& fileName, const StringKey& fileNameKey)
+	{
+		SCOPED_WRITE_LOCK(g_communicationLock, pcs);
+		BinaryWriter writer;
+		writer.WriteByte(MessageType_CheckRemapping);
+		writer.WriteString(fileName);
+		writer.WriteStringKey(fileNameKey);
+		writer.Flush();
+		BinaryReader reader;
+		u32 mappedFileTableSize = reader.ReadU32();
+		pcs.Leave();
+		g_mappedFileTable.ParseNoLock(mappedFileTableSize);
+	}
+
 	u32 Rpc_UpdateDirectory(const StringKey& dirKey, const tchar* dirName, u64 dirNameLen, bool lockDirTable)
 	{
 		u32 directoryTableSize;
diff --git a/Engine/Source/Programs/UnrealBuildAccelerator/Detours/Private/UbaDetoursFileMappingTable.h b/Engine/Source/Programs/UnrealBuildAccelerator/Detours/Private/UbaDetoursFileMappingTable.h
--- a/Engine/Source/Programs/UnrealBuildAccelerator/Detours/Private/UbaDetoursFileMappingTable.h
+++ b/Engine/Source/Programs/UnrealBuildAccelerator/Detours/Private/UbaDetoursFileMappingTable.h
@@ -61,6 +61,7 @@ namespace uba
 		u32 refCount = 0;
 		u32 fileMapDesiredAccess = PAGE_READONLY;
 		u32 fileMapViewDesiredAccess = FILE_MAP_READ; // This is the parameter to MapViewOfFile
+		bool mappingChecked = false;
 		#endif
 	};
 
@@ -87,7 +88,8 @@ namespace uba
 
 	enum : u8 { AccessFlag_Read = 1, AccessFlag_Write = 2 };
 
-	void Rpc_CreateFileW(const tchar* fileName, const StringKey& fileNameKey, u8 accessFlags, tchar* outNewName, u64 newNameCapacity, u64& outSize, u32& outCloseId, bool lock);
+	void Rpc_CreateFileW(const StringView& fileName, const StringKey& fileNameKey, u8 accessFlags, tchar* outNewName, u64 newNameCapacity, u64& outSize, u32& outCloseId, bool lock);
+	void Rpc_CheckRemapping(const StringView& fileName, const StringKey& fileNameKey);
 	u32  Rpc_UpdateDirectory(const StringKey& dirKey, const tchar* dirName, u64 dirNameLen, bool lockDirTable = true);
 	void Rpc_UpdateCloseHandle(const tchar* handleName, u32 closeId, bool deleteOnClose, const tchar* newName, const FileMappingHandle& mappingHandle, u64 mappingWritten, bool success);
 	void Rpc_UpdateTables();
diff --git a/Engine/Source/Programs/UnrealBuildAccelerator/Detours/Private/Windows/UbaDetoursFileDecompressor.h b/Engine/Source/Programs/UnrealBuildAccelerator/Detours/Private/Windows/UbaDetoursFileDecompressor.h
deleted file mode 100644
diff --git a/Engine/Source/Programs/UnrealBuildAccelerator/Detours/Private/Windows/UbaDetoursFunctionsNtDll.inl b/Engine/Source/Programs/UnrealBuildAccelerator/Detours/Private/Windows/UbaDetoursFunctionsNtDll.inl
--- a/Engine/Source/Programs/UnrealBuildAccelerator/Detours/Private/Windows/UbaDetoursFunctionsNtDll.inl
+++ b/Engine/Source/Programs/UnrealBuildAccelerator/Detours/Private/Windows/UbaDetoursFunctionsNtDll.inl
@@ -541,7 +541,7 @@ NTSTATUS Detoured_NtSetInformationFile(HANDLE FileHandle, PIO_STATUS_BLOCK IoSta
 				u64 size;
 				StringBuffer<> fixedPath;
 				FixPath(fixedPath, newName.data);
-				Rpc_CreateFileW(fixedPath.data, newFileNameKey, AccessFlag_Write, temp, sizeof_array(temp), size, fo.closeId, true);
+				Rpc_CreateFileW(fixedPath, newFileNameKey, AccessFlag_Write, temp, sizeof_array(temp), size, fo.closeId, true);
 			}
 			DEBUG_LOG_DETOURED(L"NtSetInformationFile", L"File is set to be renamed on close (from %ls to %ls)", HandleToName(FileHandle), fo.newName.c_str());
 
@@ -897,9 +897,6 @@ NTSTATUS NTAPI Shared_NtCreateFile(bool IsCreateFunc, PHANDLE hFileHandle, ACCES
 	fileNameLower.MakeLower();
 	StringKey fileNameKey = ToStringKey(fileNameLower);
 
-	bool wasPreload = g_fileDecompressor.Wait(fileNameKey); // Wait for on-going preload if existing
-
-
 	if (createFileName)// && (keepInMemory || canDetour))
 	{
 		if (g_allowDirectoryCache)
@@ -958,7 +955,7 @@ NTSTATUS NTAPI Shared_NtCreateFile(bool IsCreateFunc, PHANDLE hFileHandle, ACCES
 
 				bool isWriteAttributes = (DesiredAccess & FILE_WRITE_ATTRIBUTES) != 0;
 
-				if (allowEarlyOut && !useContent && !isWriteAttributes && !wasPreload && !CouldBeCompressedFile(fileName))
+				if (allowEarlyOut && !useContent && !isWriteAttributes && !CouldBeCompressedFile(fileName))
 				{
 					auto dh = new DetouredHandle(HandleType_File);
 					dh->fileObject = new FileObject();
@@ -1083,7 +1080,7 @@ NTSTATUS NTAPI Shared_NtCreateFile(bool IsCreateFunc, PHANDLE hFileHandle, ACCES
 		{
 			u8 access = GetFileAccessFlags(DesiredAccess, CreateDisposition);
 			wchar_t newFileName[512];
-			Rpc_CreateFileW(lpFileName, fileNameKey, access, newFileName, sizeof_array(newFileName), size, closeId, false);
+			Rpc_CreateFileW(fileName, fileNameKey, access, newFileName, sizeof_array(newFileName), size, closeId, false);
 			info.name = g_memoryBlock.Strdup(newFileName);
 			lpFileName = info.name;
 		}
@@ -1108,7 +1105,7 @@ NTSTATUS NTAPI Shared_NtCreateFile(bool IsCreateFunc, PHANDLE hFileHandle, ACCES
 				info.deleted = false;
 				wchar_t newFileName[1024];
 				u8 access = GetFileAccessFlags(DesiredAccess, CreateDisposition);
-				Rpc_CreateFileW(lpFileName, fileNameKey, access, newFileName, sizeof_array(newFileName), size, closeId, false);
+				Rpc_CreateFileW(fileName, fileNameKey, access, newFileName, sizeof_array(newFileName), size, closeId, false);
 				info.name = g_memoryBlock.Strdup(newFileName);
 				//info.size = size; // TODO: Should this be set?
 				lpFileName = info.name;
@@ -1124,6 +1121,11 @@ NTSTATUS NTAPI Shared_NtCreateFile(bool IsCreateFunc, PHANDLE hFileHandle, ACCES
 		}
 		else
 		{
+			if (!info.mappingChecked && info.name[0] == '^' && !g_runningRemote && CouldBeCompressedFile(fileName))
+			{
+				Rpc_CheckRemapping(fileName, fileNameKey);
+				info.mappingChecked = true;
+			}
 			lpFileName = info.name;
 		}
 	}
@@ -1331,15 +1333,6 @@ NTSTATUS NTAPI Shared_NtCreateFile(bool IsCreateFunc, PHANDLE hFileHandle, ACCES
 	else
 		tempFileName = ToView(info.name);
 
-
-	if (g_fileDecompressor.IsDecompressed(fileNameKey, fileName, tempFileName, info))
-	{
-		*hFileHandle = CreateFileHandle();
-		TrackFileInput();
-		DEBUG_LOG_DETOURED(funcName, L"(DECOMPRESSED)%ls %llu (%ls) (%ls) -> Success", isWriteStr, uintptr_t(*hFileHandle), lpFileName, (fileName.data != lpFileName ? fileName.data : L""));
-		return STATUS_SUCCESS;
-	}
-
 	StringBuffer<> temp;
 	temp.Append(L"\\??\\");
 	temp.Append(tempFileName);
@@ -1554,7 +1547,7 @@ NTSTATUS NTAPI Detoured_NtClose(HANDLE handle)
 					fileNameKey = ToStringKeyLower(fixedName);
 
 				u64 size;
-				Rpc_CreateFileW(fixedName.data, fileNameKey, AccessFlag_Write, temp, sizeof_array(temp), size, fo->closeId, true);
+				Rpc_CreateFileW(fixedName, fileNameKey, AccessFlag_Write, temp, sizeof_array(temp), size, fo->closeId, true);
 			}
 
 			if (!fo->newName.empty())
@@ -1611,11 +1604,7 @@ NTSTATUS NTAPI Detoured_NtClose(HANDLE handle)
 		//UBA_ASSERTF(fo->newName.empty(), L"Got close of file that was renamed but had no closeId. Old: %ls New: %ls", fi.originalName, fo->newName.c_str());
 	}
 
-	auto fiRefCount = InterlockedDecrement(&fi.refCount);
-
-	// TODO: Can't free here..... Might still be open with CreateFileMapping or MapViewOfFile
-	if (!fiRefCount && fo->wasUsed && fi.freeFileMapOnClose)
-		g_fileDecompressor.ReturnMemory(fi);
+	InterlockedDecrement(&fi.refCount);
 
 	DEBUG_LOG_DETOURED(L"NtClose", L"%llu (%ls) -> %ls", uintptr_t(handle), HandleToName(handle), ToString(res).data);
 
diff --git a/Engine/Source/Programs/UnrealBuildAccelerator/Detours/Private/Windows/UbaDetoursFunctionsWin.cpp b/Engine/Source/Programs/UnrealBuildAccelerator/Detours/Private/Windows/UbaDetoursFunctionsWin.cpp
--- a/Engine/Source/Programs/UnrealBuildAccelerator/Detours/Private/Windows/UbaDetoursFunctionsWin.cpp
+++ b/Engine/Source/Programs/UnrealBuildAccelerator/Detours/Private/Windows/UbaDetoursFunctionsWin.cpp
@@ -45,7 +45,7 @@ DETOURED_FUNCTIONS_MEMORY_NON_WINE
 #include "UbaDetoursPayload.h"
 #include "UbaApplicationRules.h"
 #include "UbaDetoursShared.h"
-#include "UbaDetoursFileDecompressor.h"
+#include "UbaDetoursUtilsWin.h"
 
 #include "Shlwapi.h"
 #include <detours/detours.h>
@@ -228,14 +228,11 @@ inline ListDirectoryHandle& asListDirectoryHandle(HANDLE h) { return *(((ListDir
 ReaderWriterLock& g_loadedModulesLock = *new ReaderWriterLock();
 UnorderedMap<HMODULE, TString>& g_loadedModules = *new UnorderedMap<HMODULE, TString>();
 u64 g_memoryFileIndexCounter = ~u64(0) - 1000000; // I really hope this will not collide with anything
+bool g_filesCouldBeCompressed;
 
-FileDecompressor& g_fileDecompressor = *new FileDecompressor(); // Leak for shutdown
 bool CouldBeCompressedFile(const StringView& fileName)
 {
-	if (g_runningRemote)
-		return g_rules->CanDependOnCompressedFiles() && g_globalRules.FileCanBeCompressed(fileName);
-	else
-		return g_fileDecompressor.CouldBeCompressedFile(fileName);
+	return g_filesCouldBeCompressed && g_globalRules.FileCanBeCompressed(fileName);
 }
 
 
@@ -457,6 +454,11 @@ UBA_NOINLINE void UbaAssert(const tchar* text, const char* file, u32 line, const
 	SuppressDetourScope _;
 	static CriticalSection cs;
 	ScopedCriticalSection s(cs);
+
+	#if UBA_DEBUG_LOG_ENABLED
+	FlushDebugLog();
+	#endif
+
 	static auto& sb = *new StringBuffer<16*1024>();
 	WriteAssertInfo(sb.Clear(), text, file, line, expr);
 	Rpc_ResolveCallstack(sb, 3, context);
@@ -1115,17 +1117,11 @@ void Init(const DetoursPayload& payload, u64 startTime)
 	g_stats.attach.time += GetTime() - startTime;
 	g_stats.attach.count = 1;
 
-	if (payload.readObjFilesCompressed && g_rules->CanDependOnCompressedFiles())
-	{
-		bool preprocess = true;//g_rules->index == SpecialRulesIndex_LdExePlatform1;
-		g_fileDecompressor.Start(cmdLine, preprocess);
-	}
+	g_filesCouldBeCompressed = payload.readObjFilesCompressed && g_rules->CanDependOnCompressedFiles();
 }
 
 void Deinit(u64 startTime)
 {
-	g_fileDecompressor.Stop();
-
 	if (g_isRunningWine) // mt.exe etc fails if detaching is not done during shutdown
 	{
 		DetourTransactionBegin();
