[UBA]
* Changed work manager to be a reference in Session (since it can never be null)

diff --git a/Engine/Source/Programs/UnrealBuildAccelerator/Common/Private/UbaSession.cpp b/Engine/Source/Programs/UnrealBuildAccelerator/Common/Private/UbaSession.cpp
--- a/Engine/Source/Programs/UnrealBuildAccelerator/Common/Private/UbaSession.cpp
+++ b/Engine/Source/Programs/UnrealBuildAccelerator/Common/Private/UbaSession.cpp
@@ -745,7 +745,7 @@ namespace uba
 		//StringBuffer<> workName;
 		//u32 len = TStrlen(fileName);
 		//workName.Append(TC("MM:")).Append(len > 30 ? fileName + (len - 30) : fileName);
-		//TrackWorkScope tws(*m_workManager, workName.data);
+		//TrackWorkScope tws(m_workManager, workName.data);
 
 		SCOPED_FUTEX(m_fileMappingTableLookupLock, lookupLock);
 		auto insres = m_fileMappingTableLookup.try_emplace(fileNameKey);
@@ -1205,18 +1205,18 @@ namespace uba
 		return result;
 	}
 
-	Session::Session(const SessionCreateInfo& info, const tchar* logPrefix, bool runningRemote, WorkManager* workManager)
+	Session::Session(const SessionCreateInfo& info, const tchar* logPrefix, bool runningRemote, WorkManager& workManager)
 	:	m_storage(info.storage)
 	,	m_logger(info.logWriter, logPrefix)
 	,	m_workManager(workManager)
 	,	m_directoryTable(m_directoryTableMemory)
-	,	m_fileMappingBuffer(m_logger, workManager)
+	,	m_fileMappingBuffer(m_logger, &workManager)
 	,	m_processCommunicationAllocator(m_logger, TC("CommunicationAllocator"))
 	,	m_trace(info.logWriter)
 	,	m_writeFilesBottleneck(info.writeFilesBottleneck)
 	,	m_writeFilesFileMapMax(u64(info.writeFilesFileMapMaxMb)*1024*1024)
 	,	m_writeFilesNoBufferingMin(u64(info.writeFilesNoBufferingMinMb)*1024*1024)
-	,	m_dependencyCrawler(m_logger, *workManager)
+	,	m_dependencyCrawler(m_logger, workManager)
 	{
 		if (info.useUniqueId)
 		{
@@ -2012,7 +2012,7 @@ namespace uba
 			m_debugLogger->Info(TC("Mapping freed 0x%llx (%s)\n"), u64(entry.mapping.mh), entry.name.c_str());
 			#endif
 
-			m_workManager->AddWork([mh = entry.mapping, this](const WorkContext&) { CloseFileMapping(m_logger, mh, TC("UsedFileMapping")); }, 1, TC("CloseFileMapping"));
+			m_workManager.AddWork([mh = entry.mapping, this](const WorkContext&) { CloseFileMapping(m_logger, mh, TC("UsedFileMapping")); }, 1, TC("CloseFileMapping"));
 			entry.handled = false;
 			entry.mapping = {};
 		}
@@ -2382,7 +2382,7 @@ namespace uba
 
 		if ((msg.access & ~FileAccess_Read) == 0)
 		{
-			TrackWorkScope tws(*m_workManager, AsView(TC("CreateFile")), ColorWork);
+			TrackWorkScope tws(m_workManager, AsView(TC("CreateFile")), ColorWork);
 			tws.AddHint(msg.fileName);
 			return CreateFileForRead(out, tws, fileName, fileNameKey, msg.process, *msg.process.m_startInfo.rules);
 		}
@@ -2847,7 +2847,7 @@ namespace uba
 
 					if (false)
 					{
-						TrackWorkScope work(*m_workManager, AsView(TC("Flush")));
+						TrackWorkScope work(m_workManager, AsView(TC("Flush")));
 						work.AddHint(file.name);
 						FlushFileBuffers(file.mappingHandle.fh);
 					}
@@ -2872,7 +2872,7 @@ namespace uba
 						info.FileNameLength = u32(name2.count*sizeof(tchar));
 						info.Flags = FILE_RENAME_FLAG_POSIX_SEMANTICS;
 						memcpy(info.FileName, name2.data, info.FileNameLength+sizeof(tchar));
-						TrackWorkScope work(*m_workManager, AsView(TC("Rename")));
+						TrackWorkScope work(m_workManager, AsView(TC("Rename")));
 						work.AddHint(file.name);
 						success = SetFileInformationByHandle(file.mappingHandle.fh, FileRenameInfo, &info, sizeof(buffer));
 					}
@@ -2880,7 +2880,7 @@ namespace uba
 					if (success)
 					{
 						{
-							TrackWorkScope work(*m_workManager, AsView(TC("CloseFileMapping")));
+							TrackWorkScope work(m_workManager, AsView(TC("CloseFileMapping")));
 							work.AddHint(file.name);
 							CloseFileMapping(m_logger, file.mappingHandle, file.name.c_str());
 						}
@@ -2904,7 +2904,7 @@ namespace uba
 				if (storeCompressed)
 				{
 					Storage::WriteResult res;
-					CompressedObjFileHeader header { CalculateCasKey(mem, fileSize, true, m_workManager, file.name.c_str()) };
+					CompressedObjFileHeader header { CalculateCasKey(mem, fileSize, true, &m_workManager, file.name.c_str()) };
 
 					if (!m_storage.WriteCompressed(res, TC("MemoryMap"), InvalidFileHandle, mem, fileSize, file.name.c_str(), &header, sizeof(header), file.lastWriteTime))
 						return false;
@@ -2976,7 +2976,7 @@ namespace uba
 
 			if (shouldEvictFromMemory)
 			{
-				m_workManager->AddWork([mh = file.mappingHandle, n = file.name, this](const WorkContext&)
+				m_workManager.AddWork([mh = file.mappingHandle, n = file.name, this](const WorkContext&)
 					{
 						CloseFileMapping(m_logger, mh, n.c_str());
 					}, 1, TC("CloseFileMapping"));
@@ -3031,7 +3031,7 @@ namespace uba
 
 		Atomic<bool> success = true;
 		auto span = std::span(files, fileCount);
-		m_workManager->ParallelFor(fileCount - 1, span, [&](const WorkContext&, auto& it)
+		m_workManager.ParallelFor(fileCount - 1, span, [&](const WorkContext&, auto& it)
 			{
 				if (!writeFile(**it))
 					success = false;
@@ -3873,7 +3873,7 @@ namespace uba
 			{
 				CreateFileResponse out;
 				{
-					TrackWorkScope tws(*m_workManager, AsView(TC("CrawlFile")), ColorWork);
+					TrackWorkScope tws(m_workManager, AsView(TC("CrawlFile")), ColorWork);
 					tws.AddHint(fileName);
 
 					if (!CreateFileForRead(out, tws, fileName, ToStringKey(fileName), *static_cast<ProcessImpl*>(ph.m_process), *rules))
diff --git a/Engine/Source/Programs/UnrealBuildAccelerator/Common/Private/UbaSessionClient.cpp b/Engine/Source/Programs/UnrealBuildAccelerator/Common/Private/UbaSessionClient.cpp
--- a/Engine/Source/Programs/UnrealBuildAccelerator/Common/Private/UbaSessionClient.cpp
+++ b/Engine/Source/Programs/UnrealBuildAccelerator/Common/Private/UbaSessionClient.cpp
@@ -37,7 +37,7 @@ namespace uba
 	};
 
 	SessionClient::SessionClient(const SessionClientCreateInfo& info)
-	: Session(info, TC("UbaSessionClient"), true, &info.client)
+	: Session(info, TC("UbaSessionClient"), true, info.client)
 	,	m_client(info.client)
 	,	m_name(info.name.data)
 	,	m_terminationTime(~0ull)
diff --git a/Engine/Source/Programs/UnrealBuildAccelerator/Common/Private/UbaSessionServer.cpp b/Engine/Source/Programs/UnrealBuildAccelerator/Common/Private/UbaSessionServer.cpp
--- a/Engine/Source/Programs/UnrealBuildAccelerator/Common/Private/UbaSessionServer.cpp
+++ b/Engine/Source/Programs/UnrealBuildAccelerator/Common/Private/UbaSessionServer.cpp
@@ -125,7 +125,7 @@ namespace uba
 	}
 
 	SessionServer::SessionServer(const SessionServerCreateInfo& info, const u8* environment, u32 environmentSize)
-	:	Session(info, TC("UbaSessionServer"), false, &info.server)
+	:	Session(info, TC("UbaSessionServer"), false, info.server)
 	,	m_server(info.server)
 	,	m_maxRemoteProcessCount(~0u)
 	{
diff --git a/Engine/Source/Programs/UnrealBuildAccelerator/Common/Public/UbaSession.h b/Engine/Source/Programs/UnrealBuildAccelerator/Common/Public/UbaSession.h
--- a/Engine/Source/Programs/UnrealBuildAccelerator/Common/Public/UbaSession.h
+++ b/Engine/Source/Programs/UnrealBuildAccelerator/Common/Public/UbaSession.h
@@ -101,7 +101,7 @@ namespace uba
 		u32 GetDirectoryTableSize();
 
 	protected:
-		Session(const SessionCreateInfo& info, const tchar* logPrefix, bool runningRemote, WorkManager* workManager = nullptr);
+		Session(const SessionCreateInfo& info, const tchar* logPrefix, bool runningRemote, WorkManager& workManager);
 		bool Create(const SessionCreateInfo& info);
 
 		void ValidateStartInfo(const ProcessStartInfo& startInfo);
@@ -201,8 +201,7 @@ namespace uba
 
 		Storage& m_storage;
 		MutableLogger m_logger;
-
-		WorkManager* m_workManager;
+		WorkManager& m_workManager;
 
 		StringBuffer<32> m_id;
 		StringBuffer<MaxPath> m_rootDir;
