[UBA]
* Renamed CompressedObjFileHeader to CompressedFileHeader since it is now used for more than obj files

diff --git a/Engine/Source/Programs/UnrealBuildAccelerator/Cli/Private/UbaCli.cpp b/Engine/Source/Programs/UnrealBuildAccelerator/Cli/Private/UbaCli.cpp
--- a/Engine/Source/Programs/UnrealBuildAccelerator/Cli/Private/UbaCli.cpp
+++ b/Engine/Source/Programs/UnrealBuildAccelerator/Cli/Private/UbaCli.cpp
@@ -4,7 +4,7 @@
 #include "UbaCacheClient.h"
 #include "UbaCacheServer.h"
 #include "UbaClient.h"
-#include "UbaCompressedObjFileHeader.h"
+#include "UbaCompressedFileHeader.h"
 #include "UbaConfig.h"
 #include "UbaCoordinatorWrapper.h"
 #include "UbaFileAccessor.h"
@@ -653,7 +653,7 @@ namespace uba
 			StorageImpl storage(storageInfo);
 
 			Storage::WriteResult res;
-			CompressedObjFileHeader header { CalculateCasKey(mem, fileSize, true, &workManager, testCompress.c_str()) };
+			CompressedFileHeader header { CalculateCasKey(mem, fileSize, true, &workManager, testCompress.c_str()) };
 
 			StringBuffer<> dest;
 			dest.Append(storage.GetTempPath()).Append(TC("TestCompress.tmp"));
@@ -686,9 +686,9 @@ namespace uba
 
 			BinaryReader reader(mem, 0, fileSize);
 
-			auto& h = *(CompressedObjFileHeader*)mem;
+			auto& h = *(CompressedFileHeader*)mem;
 			if (h.IsValid())
-				reader.Skip(sizeof(CompressedObjFileHeader));
+				reader.Skip(sizeof(CompressedFileHeader));
 			u64 decompressedSize = reader.ReadU64();
 
 			if (testCompressOriginalSize && decompressedSize != testCompressOriginalSize)
@@ -764,9 +764,9 @@ namespace uba
 
 				CasKey key = CalculateCasKey(data, fileSize, false, nullptr, application.c_str());
 				CasKey uncompressedKey;
-				if (fileSize > sizeof(CompressedObjFileHeader))
+				if (fileSize > sizeof(CompressedFileHeader))
 				{
-					auto& hdr = *(CompressedObjFileHeader*)data;
+					auto& hdr = *(CompressedFileHeader*)data;
 					if (hdr.IsValid())
 						uncompressedKey = hdr.casKey;
 				}
diff --git a/Engine/Source/Programs/UnrealBuildAccelerator/Common/Private/UbaCacheClient.cpp b/Engine/Source/Programs/UnrealBuildAccelerator/Common/Private/UbaCacheClient.cpp
--- a/Engine/Source/Programs/UnrealBuildAccelerator/Common/Private/UbaCacheClient.cpp
+++ b/Engine/Source/Programs/UnrealBuildAccelerator/Common/Private/UbaCacheClient.cpp
@@ -4,7 +4,7 @@
 #include "UbaApplicationRules.h"
 #include "UbaCacheEntry.h"
 #include "UbaCompactTables.h"
-#include "UbaCompressedObjFileHeader.h"
+#include "UbaCompressedFileHeader.h"
 #include "UbaConfig.h"
 #include "UbaDirectoryIterator.h"
 #include "UbaFileAccessor.h"
@@ -352,8 +352,8 @@ namespace uba
 					u64 fileSize = fa.GetSize();
 					u8* fileMem = fa.GetData();
 
-					if (fileSize > sizeof(CompressedObjFileHeader) && ((CompressedObjFileHeader*)fileMem)->IsValid())
-						newKey = AsCompressed(((CompressedObjFileHeader*)fileMem)->casKey, false);
+					if (fileSize > sizeof(CompressedFileHeader) && ((CompressedFileHeader*)fileMem)->IsValid())
+						newKey = AsCompressed(((CompressedFileHeader*)fileMem)->casKey, false);
 					else
 						newKey = CalculateCasKey(fileMem, fileSize, false, nullptr, path.data);
 
diff --git a/Engine/Source/Programs/UnrealBuildAccelerator/Common/Private/UbaObjectFile.cpp b/Engine/Source/Programs/UnrealBuildAccelerator/Common/Private/UbaObjectFile.cpp
--- a/Engine/Source/Programs/UnrealBuildAccelerator/Common/Private/UbaObjectFile.cpp
+++ b/Engine/Source/Programs/UnrealBuildAccelerator/Common/Private/UbaObjectFile.cpp
@@ -2,7 +2,7 @@
 
 #include "UbaObjectFile.h"
 #include "UbaBinaryReaderWriter.h"
-#include "UbaCompressedObjFileHeader.h"
+#include "UbaCompressedFileHeader.h"
 #include "UbaFileAccessor.h"
 #include "UbaObjectFileCoff.h"
 #include "UbaObjectFileElf.h"
@@ -36,10 +36,10 @@ namespace uba
 		ObjectFile* objectFile = nullptr;
 
 		bool ownsData = false;
-		if (dataSize >= sizeof(CompressedObjFileHeader) && ((CompressedObjFileHeader*)data)->IsValid())
+		if (dataSize >= sizeof(CompressedFileHeader) && ((CompressedFileHeader*)data)->IsValid())
 		{
-			u64 decompressedSize = *(u64*)(data + sizeof(CompressedObjFileHeader));
-			u8* readPos = data + sizeof(CompressedObjFileHeader) + 8;
+			u64 decompressedSize = *(u64*)(data + sizeof(CompressedFileHeader));
+			u8* readPos = data + sizeof(CompressedFileHeader) + 8;
 
 			u8* decompressedData = (u8*)malloc(decompressedSize);
 			u8* writePos = decompressedData;
diff --git a/Engine/Source/Programs/UnrealBuildAccelerator/Common/Private/UbaSession.cpp b/Engine/Source/Programs/UnrealBuildAccelerator/Common/Private/UbaSession.cpp
--- a/Engine/Source/Programs/UnrealBuildAccelerator/Common/Private/UbaSession.cpp
+++ b/Engine/Source/Programs/UnrealBuildAccelerator/Common/Private/UbaSession.cpp
@@ -1,7 +1,7 @@
 // Copyright Epic Games, Inc. All Rights Reserved.
 
 #include "UbaSession.h"
-#include "UbaCompressedObjFileHeader.h"
+#include "UbaCompressedFileHeader.h"
 #include "UbaConfig.h"
 #include "UbaFileAccessor.h"
 #include "UbaObjectFile.h"
@@ -620,14 +620,14 @@ namespace uba
 				size = info.size;
 				lastWriteTime = info.lastWriteTime;
 
-				if (m_readObjFilesCompressed && info.size > sizeof(CompressedObjFileHeader) && g_globalRules.FileCanBeCompressed(ToView(fileName)))
+				if (m_readObjFilesCompressed && info.size > sizeof(CompressedFileHeader) && g_globalRules.FileCanBeCompressed(ToView(fileName)))
 				{
-					CompressedObjFileHeader header(CasKeyZero);
+					CompressedFileHeader header(CasKeyZero);
 					if (!ReadFile(m_logger, fileName, fileHandle, &header, sizeof(header)))
 						return m_logger.Error(TC("Failed to read header of compressed file %s (%s)"), fileName, LastErrorToText().data);
 					if (header.IsValid())
 					{
-						fileStartOffset = sizeof(CompressedObjFileHeader);
+						fileStartOffset = sizeof(CompressedFileHeader);
 						isCompressed = true;
 						if (!ReadFile(m_logger, fileName, fileHandle, &size, sizeof(u64)))
 							return m_logger.Error(TC("Failed to read first bytes from file %s (%s)"), fileName, LastErrorToText().data);
@@ -2904,7 +2904,7 @@ namespace uba
 				if (storeCompressed)
 				{
 					Storage::WriteResult res;
-					CompressedObjFileHeader header { CalculateCasKey(mem, fileSize, true, &m_workManager, file.name.c_str()) };
+					CompressedFileHeader header { CalculateCasKey(mem, fileSize, true, &m_workManager, file.name.c_str()) };
 
 					if (!m_storage.WriteCompressed(res, TC("MemoryMap"), InvalidFileHandle, mem, fileSize, file.name.c_str(), &header, sizeof(header), file.lastWriteTime))
 						return false;
diff --git a/Engine/Source/Programs/UnrealBuildAccelerator/Common/Private/UbaStorage.cpp b/Engine/Source/Programs/UnrealBuildAccelerator/Common/Private/UbaStorage.cpp
--- a/Engine/Source/Programs/UnrealBuildAccelerator/Common/Private/UbaStorage.cpp
+++ b/Engine/Source/Programs/UnrealBuildAccelerator/Common/Private/UbaStorage.cpp
@@ -3,7 +3,7 @@
 #include "UbaStorage.h"
 #include "UbaApplicationRules.h"
 #include "UbaBottleneck.h"
-#include "UbaCompressedObjFileHeader.h"
+#include "UbaCompressedFileHeader.h"
 #include "UbaConfig.h"
 #include "UbaFileAccessor.h"
 #include "UbaBinaryReaderWriter.h"
@@ -560,9 +560,9 @@ namespace uba
 			return m_logger.Error(TC("[WriteCasFileNoCheck] GetFileSize failed for %s (%s)"), fileName, LastErrorToText().data);
 
 		bool readIsCompressed = false;
-		if (readFileSize >= sizeof(CompressedObjFileHeader) && g_globalRules.FileCanBeCompressed(ToView(fileName)))
+		if (readFileSize >= sizeof(CompressedFileHeader) && g_globalRules.FileCanBeCompressed(ToView(fileName)))
 		{
-			CompressedObjFileHeader header(CasKeyZero);
+			CompressedFileHeader header(CasKeyZero);
 			if (!ReadFile(m_logger, fileName, readHandle, &header, sizeof(header)))
 				return m_logger.Error(TC("Failed to read header of compressed file %s (%s)"), fileName, LastErrorToText().data);
 			if (header.IsValid())
@@ -1945,17 +1945,17 @@ namespace uba
 				u8* fileMem = scope.MapView(TC("StoreCasFile"), externalMapping.handle, externalMapping.offset, externalMapping.size, fileName);
 				if (!fileMem)
 					return false;
-				if (fileSize > sizeof(CompressedObjFileHeader) && ((CompressedObjFileHeader*)fileMem)->IsValid())
-					fileEntry.casKey = ((CompressedObjFileHeader*)fileMem)->casKey;
+				if (fileSize > sizeof(CompressedFileHeader) && ((CompressedFileHeader*)fileMem)->IsValid())
+					fileEntry.casKey = ((CompressedFileHeader*)fileMem)->casKey;
 				else
 					fileEntry.casKey = CalculateCasKey(fileMem, externalMapping.size, m_storeCompressed);
 			}
 			else
 			{
 				bool handled = false;
-				if (fileSize >= sizeof(CompressedObjFileHeader) && g_globalRules.FileCanBeCompressed(ToView(fileName)))
+				if (fileSize >= sizeof(CompressedFileHeader) && g_globalRules.FileCanBeCompressed(ToView(fileName)))
 				{
-					CompressedObjFileHeader header(CasKeyZero);
+					CompressedFileHeader header(CasKeyZero);
 					if (!ReadFile(m_logger, fileName, fileHandle, &header, sizeof(header)))
 						return m_logger.Error(TC("Failed to read header of compressed file %s (%s)"), fileName, LastErrorToText().data);
 					if (header.IsValid())
@@ -2083,9 +2083,9 @@ namespace uba
 		if (casKeyOverride == CasKeyZero)
 		{
 			bool handled = false;
-			if (fileSize >= sizeof(CompressedObjFileHeader) && g_globalRules.FileCanBeCompressed(ToView(fileName)))
+			if (fileSize >= sizeof(CompressedFileHeader) && g_globalRules.FileCanBeCompressed(ToView(fileName)))
 			{
-				CompressedObjFileHeader header(CasKeyZero);
+				CompressedFileHeader header(CasKeyZero);
 				if (!ReadFile(m_logger, fileName, fileHandle, &header, sizeof(header)))
 					return m_logger.Error(TC("Failed to read header of compressed file %s (%s)"), fileName, LastErrorToText().data);
 				if (header.IsValid())
@@ -2470,13 +2470,13 @@ namespace uba
 						if (!GetFileSizeEx(compressedFileSize, readHandle))
 							return m_logger.Error(TC("[CopyOrLink] Failed to get file size of compressed file %s (%s)"), casFile.data, LastErrorToText().data);
 						
-					sizeOnDisk = compressedFileSize + sizeof(CompressedObjFileHeader);
+					sizeOnDisk = compressedFileSize + sizeof(CompressedFileHeader);
 
 					if (!destinationFile.CreateMemoryWrite(false, fileAttributes, sizeOnDisk, m_tempPath.data))
 						return false;
 					u8* writePos = destinationFile.GetData();
-					*(CompressedObjFileHeader*)writePos = CompressedObjFileHeader(casKey);
-					writePos += sizeof(CompressedObjFileHeader);
+					*(CompressedFileHeader*)writePos = CompressedFileHeader(casKey);
+					writePos += sizeof(CompressedFileHeader);
 
 					if (mappedView.memory)
 					{
diff --git a/Engine/Source/Programs/UnrealBuildAccelerator/Common/Private/UbaStorageUtils.cpp b/Engine/Source/Programs/UnrealBuildAccelerator/Common/Private/UbaStorageUtils.cpp
--- a/Engine/Source/Programs/UnrealBuildAccelerator/Common/Private/UbaStorageUtils.cpp
+++ b/Engine/Source/Programs/UnrealBuildAccelerator/Common/Private/UbaStorageUtils.cpp
@@ -1,7 +1,7 @@
 // Copyright Epic Games, Inc. All Rights Reserved.
 
 #include "UbaStorageUtils.h"
-#include "UbaCompressedObjFileHeader.h"
+#include "UbaCompressedFileHeader.h"
 #include "UbaStorage.h"
 #include "UbaFileAccessor.h"
 #include "UbaStats.h"
@@ -462,7 +462,7 @@ namespace uba
 			writer.WriteBool(false);
 			writer.WriteCasKey(casKey);
 			writer.WriteString(destination);
-			BinaryReader reader(slot + (writeCompressed ? sizeof(CompressedObjFileHeader) : 0), 0, SendMaxSize); // Create some space in front of reader to write obj file header there if destination is compressed
+			BinaryReader reader(slot + (writeCompressed ? sizeof(CompressedFileHeader) : 0), 0, SendMaxSize); // Create some space in front of reader to write obj file header there if destination is compressed
 			if (!msg.Send(reader))
 				return logger.Error(TC("Failed to send fetch begin message for cas %s (%s). Error: %u"), CasKeyString(casKey).str, destination, msg.GetError());
 			sizeOfFirstMessage = u32(reader.GetLeft());
@@ -490,7 +490,7 @@ namespace uba
 		}
 
 		bytesReceived = fileSize;
-		sizeOnDisk = writeCompressed ? (sizeof(CompressedObjFileHeader) + fileSize) : actualSize;
+		sizeOnDisk = writeCompressed ? (sizeof(CompressedFileHeader) + fileSize) : actualSize;
 
 		FileAccessor destinationFile(logger, destination);
 
@@ -545,9 +545,9 @@ namespace uba
 			u8* lastSource = readBuffer;
 			u64 lastResponseSize = responseSize;
 
-			lastSource -= sizeof(CompressedObjFileHeader);
-			lastResponseSize += sizeof(CompressedObjFileHeader);
-			new (lastSource) CompressedObjFileHeader{ casKey };
+			lastSource -= sizeof(CompressedFileHeader);
+			lastResponseSize += sizeof(CompressedFileHeader);
+			new (lastSource) CompressedFileHeader{ casKey };
 
 			auto writePrev = [&]() { return WriteDestination(lastSource, lastResponseSize); };
 
diff --git a/Engine/Source/Programs/UnrealBuildAccelerator/Core/Public/UbaCompressedFileHeader.h b/Engine/Source/Programs/UnrealBuildAccelerator/Core/Public/UbaCompressedFileHeader.h
new file mode 100644
--- /dev/null
+++ b/Engine/Source/Programs/UnrealBuildAccelerator/Core/Public/UbaCompressedFileHeader.h
@@ -0,0 +1,23 @@
+// Copyright Epic Games, Inc. All Rights Reserved.
+
+#pragma once
+
+#include "UbaHash.h"
+
+namespace uba
+{
+	struct CompressedFileHeader
+	{
+		inline static constexpr u8 Magic[] = { 'U', 'B', 'A', 0x01 };
+
+		CompressedFileHeader(const CasKey& key) : casKey(key)
+		{
+			memcpy(&magic, Magic, sizeof(magic));
+		}
+
+		bool IsValid() const { return memcmp(&magic, Magic, sizeof(Magic)) == 0; }
+
+		u8 magic[sizeof(Magic)];
+		CasKey casKey;
+	};
+}
diff --git a/Engine/Source/Programs/UnrealBuildAccelerator/Core/Public/UbaCompressedObjFileHeader.h b/Engine/Source/Programs/UnrealBuildAccelerator/Core/Public/UbaCompressedObjFileHeader.h
deleted file mode 100644
