[UBA]
* Renamed read/write obj files compressed to read/write intermediate files compressed

diff --git a/Engine/Source/Programs/UnrealBuildAccelerator/Cli/Private/UbaCli.cpp b/Engine/Source/Programs/UnrealBuildAccelerator/Cli/Private/UbaCli.cpp
--- a/Engine/Source/Programs/UnrealBuildAccelerator/Cli/Private/UbaCli.cpp
+++ b/Engine/Source/Programs/UnrealBuildAccelerator/Cli/Private/UbaCli.cpp
@@ -868,8 +868,8 @@ namespace uba
 		//info.shouldWriteToDisk = shouldWriteToDisk;
 		info.rootDir = g_rootDir.data;
 		//info.traceName.Append(TC("TESTTRACE"));
-		//info.storeObjFilesCompressed = false;
-		info.readObjFilesCompressed = true;
+		//info.storeIntermediateFilesCompressed = false;
+		info.readIntermediateFilesCompressed = true;
 		//info.extractObjFilesSymbols = true;
 		#if UBA_DEBUG_LOG_ENABLED
 		info.remoteLogEnabled = true;
diff --git a/Engine/Source/Programs/UnrealBuildAccelerator/Common/Private/UbaCacheClient.cpp b/Engine/Source/Programs/UnrealBuildAccelerator/Common/Private/UbaCacheClient.cpp
--- a/Engine/Source/Programs/UnrealBuildAccelerator/Common/Private/UbaCacheClient.cpp
+++ b/Engine/Source/Programs/UnrealBuildAccelerator/Common/Private/UbaCacheClient.cpp
@@ -875,7 +875,7 @@ namespace uba
 		else
 		{
 			DowngradedLogger logger(m_connected, m_logger.m_writer, TC("UbaCacheClientDownload"));
-			bool writeCompressed = m_session.ShouldStoreObjFilesCompressed() && g_globalRules.FileCanBeCompressed(path);
+			bool writeCompressed = m_session.ShouldStoreIntermediateFilesCompressed() && g_globalRules.FileCanBeCompressed(path);
 			if (!fetcher.RetrieveFile(logger, m_client, casKey, path.data, writeCompressed))
 				return logger.Error(TC("Failed to download cache output %s for %s"), path.data, info.GetDescription());
 		}
diff --git a/Engine/Source/Programs/UnrealBuildAccelerator/Common/Private/UbaProcess.cpp b/Engine/Source/Programs/UnrealBuildAccelerator/Common/Private/UbaProcess.cpp
--- a/Engine/Source/Programs/UnrealBuildAccelerator/Common/Private/UbaProcess.cpp
+++ b/Engine/Source/Programs/UnrealBuildAccelerator/Common/Private/UbaProcess.cpp
@@ -1722,7 +1722,7 @@ namespace uba
 			payload.trackInputs = m_startInfo.trackInputs;
 			payload.useCustomAllocator = allowCustomAllocator && m_startInfo.useCustomAllocator && m_startInfo.rules->AllowMiMalloc();
 			payload.isRunningWine = IsRunningWine();
-			payload.readObjFilesCompressed = m_session.m_readObjFilesCompressed;
+			payload.readIntermediateFilesCompressed = m_session.m_readIntermediateFilesCompressed;
 			payload.uiLanguage = m_startInfo.uiLanguage;
 			if (*m_startInfo.logFile)
 			{
diff --git a/Engine/Source/Programs/UnrealBuildAccelerator/Common/Private/UbaSession.cpp b/Engine/Source/Programs/UnrealBuildAccelerator/Common/Private/UbaSession.cpp
--- a/Engine/Source/Programs/UnrealBuildAccelerator/Common/Private/UbaSession.cpp
+++ b/Engine/Source/Programs/UnrealBuildAccelerator/Common/Private/UbaSession.cpp
@@ -227,7 +227,7 @@ namespace uba
 		table.GetValueAsBool(detailedTrace, TC("DetailedTrace"));
 		table.GetValueAsBool(traceChildProcesses, TC("TraceChildProcesses"));
 		table.GetValueAsBool(traceWrittenFiles, TC("TraceWrittenFiles"));
-		table.GetValueAsBool(storeObjFilesCompressed, TC("StoreObjFilesCompressed"));
+		table.GetValueAsBool(storeIntermediateFilesCompressed, TC("StoreIntermediateFilesCompressed"));
 		table.GetValueAsBool(allowLocalDetour, TC("AllowLocalDetour"));
 		table.GetValueAsBool(extractObjFilesSymbols, TC("ExtractObjFilesSymbols"));
 		table.GetValueAsBool(useFakeVolumeSerial, TC("UseFakeVolumeSerial"));
@@ -620,7 +620,7 @@ namespace uba
 				size = info.size;
 				lastWriteTime = info.lastWriteTime;
 
-				if (m_readObjFilesCompressed && info.size > sizeof(CompressedFileHeader) && g_globalRules.FileCanBeCompressed(ToView(fileName)))
+				if (m_readIntermediateFilesCompressed && info.size > sizeof(CompressedFileHeader) && g_globalRules.FileCanBeCompressed(ToView(fileName)))
 				{
 					CompressedFileHeader header(CasKeyZero);
 					if (!ReadFile(m_logger, fileName, fileHandle, &header, sizeof(header)))
@@ -1249,8 +1249,8 @@ namespace uba
 		m_shouldWriteToDisk = info.shouldWriteToDisk;
 		UBA_ASSERTF(m_shouldWriteToDisk || m_allowMemoryMaps, TC("Can't disable both should write to disk and allow memory maps"));
 
-		m_storeObjFilesCompressed = info.storeObjFilesCompressed && IsWindows; // Non-windows not implemented (yet)
-		m_readObjFilesCompressed = (m_storeObjFilesCompressed || (info.readObjFilesCompressed && IsWindows)) && !runningRemote; // with remote we decompress the files into memory
+		m_storeIntermediateFilesCompressed = info.storeIntermediateFilesCompressed && IsWindows; // Non-windows not implemented (yet)
+		m_readIntermediateFilesCompressed = (m_storeIntermediateFilesCompressed || (info.readIntermediateFilesCompressed && IsWindows)) && !runningRemote; // with remote we decompress the files into memory
 		m_allowLocalDetour = info.allowLocalDetour;
 
 		m_detailedTrace = info.detailedTrace;
@@ -1320,7 +1320,7 @@ namespace uba
 		if (info.traceOutputFile)
 			m_traceOutputFile.Append(info.traceOutputFile);
 
-		if (m_readObjFilesCompressed && !m_runningRemote)
+		if (m_readIntermediateFilesCompressed && !m_runningRemote)
 			m_dependencyCrawler.Init(
 				[this](const StringView& fileName, u32& outAttr) { return FileExists(m_logger, fileName.data, nullptr, &outAttr); }, // FileExists
 				[this](const StringView& path, const DependencyCrawler::FileFunc& fileFunc) {}, // TraverseFilesFunc
@@ -2224,7 +2224,7 @@ namespace uba
 			startInfo.application = startInfo.applicationStr.c_str();
 		}
 
-		if (!isChild && !m_runningRemote && m_readObjFilesCompressed)
+		if (!isChild && !m_runningRemote && m_readIntermediateFilesCompressed)
 		{
 			auto crawlerType = startInfo.rules->GetDependencyCrawlerType();
 			if (crawlerType == DependencyCrawlerType_MsvcLinker || crawlerType == DependencyCrawlerType_ClangLinker)
@@ -2452,7 +2452,7 @@ namespace uba
 		if (m_allowMemoryMaps)
 		{
 			u64 alignment = GetMemoryMapAlignment(fileName);
-			bool canBeCompressed = m_readObjFilesCompressed && g_globalRules.FileCanBeCompressed(fileName);
+			bool canBeCompressed = m_readIntermediateFilesCompressed && g_globalRules.FileCanBeCompressed(fileName);
 			bool useMemoryMap = alignment != 0 || canBeCompressed;
 			if (useMemoryMap)
 			{
@@ -2832,7 +2832,7 @@ namespace uba
 				auto& rules = *process.m_startInfo.rules;
 				RootsHandle rootsHandle = process.GetStartInfo().rootsHandle;
 
-				bool storeCompressed = m_storeObjFilesCompressed && g_globalRules.FileCanBeCompressed(file.name);
+				bool storeCompressed = m_storeIntermediateFilesCompressed && g_globalRules.FileCanBeCompressed(file.name);
 				bool shouldDevirtualize = false;
 				bool escapeSpaces = false;
 				if (!storeCompressed)
diff --git a/Engine/Source/Programs/UnrealBuildAccelerator/Common/Private/UbaSessionClient.cpp b/Engine/Source/Programs/UnrealBuildAccelerator/Common/Private/UbaSessionClient.cpp
--- a/Engine/Source/Programs/UnrealBuildAccelerator/Common/Private/UbaSessionClient.cpp
+++ b/Engine/Source/Programs/UnrealBuildAccelerator/Common/Private/UbaSessionClient.cpp
@@ -1325,7 +1325,7 @@ namespace uba
 
 		m_sessionId = reader.ReadU32();
 		m_uiLanguage = reader.ReadU32();
-		m_storeObjFilesCompressed = reader.ReadBool();
+		m_storeIntermediateFilesCompressed = reader.ReadBool();
 		m_detailedTrace = reader.ReadBool();
 		m_shouldSendLogToServer = reader.ReadBool();
 		m_shouldSendTraceToServer = reader.ReadBool();
diff --git a/Engine/Source/Programs/UnrealBuildAccelerator/Common/Private/UbaSessionServer.cpp b/Engine/Source/Programs/UnrealBuildAccelerator/Common/Private/UbaSessionServer.cpp
--- a/Engine/Source/Programs/UnrealBuildAccelerator/Common/Private/UbaSessionServer.cpp
+++ b/Engine/Source/Programs/UnrealBuildAccelerator/Common/Private/UbaSessionServer.cpp
@@ -803,7 +803,7 @@ namespace uba
 		writer.WriteBool(m_resetCas);
 		writer.WriteU32(sessionId);
 		writer.WriteU32(m_uiLanguage);
-		writer.WriteBool(m_storeObjFilesCompressed);
+		writer.WriteBool(m_storeIntermediateFilesCompressed);
 		writer.WriteBool(m_detailedTrace);
 		writer.WriteBool(m_remoteLogEnabled);
 		writer.WriteBool(m_remoteTraceEnabled);
@@ -1090,7 +1090,7 @@ namespace uba
 							return DevirtualizeDepsFile(rootsHandle, destData, sourceData, sourceSize, escapeSpaces, hint);
 						};
 				}
-				else if (m_storeObjFilesCompressed)
+				else if (m_storeIntermediateFilesCompressed)
 				{
 					writeCompressed = g_globalRules.FileCanBeCompressed(destination);
 				}
diff --git a/Engine/Source/Programs/UnrealBuildAccelerator/Common/Public/UbaSession.h b/Engine/Source/Programs/UnrealBuildAccelerator/Common/Public/UbaSession.h
--- a/Engine/Source/Programs/UnrealBuildAccelerator/Common/Public/UbaSession.h
+++ b/Engine/Source/Programs/UnrealBuildAccelerator/Common/Public/UbaSession.h
@@ -84,7 +84,7 @@ namespace uba
 		const ApplicationRules* GetRules(const ProcessStartInfo& si); // Get application rules used for process
 		const tchar* GetTempPath(); // Path for temp files used for current session
 		const tchar* GetRootDir(); // Root dir for session files
-		bool ShouldStoreObjFilesCompressed() const { return m_storeObjFilesCompressed; }
+		bool ShouldStoreIntermediateFilesCompressed() const { return m_storeIntermediateFilesCompressed; }
 		bool HasDetailedtrace() const { return m_detailedTrace; }
 
 		virtual ~Session();
@@ -225,8 +225,8 @@ namespace uba
 		bool m_traceChildProcesses;
 		bool m_traceWrittenFiles;
 		bool m_logToFile;
-		bool m_storeObjFilesCompressed;
-		bool m_readObjFilesCompressed;
+		bool m_storeIntermediateFilesCompressed;
+		bool m_readIntermediateFilesCompressed;
 		bool m_allowLocalDetour;
 
 		u64 m_keepOutputFileMemoryMapsThreshold;
diff --git a/Engine/Source/Programs/UnrealBuildAccelerator/Common/Public/UbaSessionCreateInfo.h b/Engine/Source/Programs/UnrealBuildAccelerator/Common/Public/UbaSessionCreateInfo.h
--- a/Engine/Source/Programs/UnrealBuildAccelerator/Common/Public/UbaSessionCreateInfo.h
+++ b/Engine/Source/Programs/UnrealBuildAccelerator/Common/Public/UbaSessionCreateInfo.h
@@ -35,8 +35,8 @@ namespace uba
 		bool detailedTrace = false;				// Enable detailed trace to include jobs, individual file I/O etc in trace dump
 		bool traceChildProcesses = false;		// Trace and visualize child processes so they can be seen in visualizer
 		bool traceWrittenFiles = false;			// Will add process output files to trace
-		bool storeObjFilesCompressed = false;	// Compiler will write .obj files compressed to disk and linker will decompress when reading
-		bool readObjFilesCompressed = false;	// Set to true to support reading compressed .obj files. This flag will not compress new obj files
+		bool storeIntermediateFilesCompressed = false;	// Compiler will write intermediate files (.obj, .pch etc) compressed to disk and will be decompressed when used
+		bool readIntermediateFilesCompressed = false;	// Set to true to support reading compressed .obj files. This flag will not compress new obj files
 		bool allowLocalDetour = true;			// Allow local processes to be detoured. If this is false it is up to outside logic to register all created/deleted files
 		bool extractObjFilesSymbols = false;	// Will extract import/export symbols and clear export table from obj file.
 		bool treatTempDirAsEmpty = true;		// If this is true, directory table will always see temp as empty. Use sub directories of temp instead
diff --git a/Engine/Source/Programs/UnrealBuildAccelerator/Core/Public/UbaDetoursPayload.h b/Engine/Source/Programs/UnrealBuildAccelerator/Core/Public/UbaDetoursPayload.h
--- a/Engine/Source/Programs/UnrealBuildAccelerator/Core/Public/UbaDetoursPayload.h
+++ b/Engine/Source/Programs/UnrealBuildAccelerator/Core/Public/UbaDetoursPayload.h
@@ -26,7 +26,7 @@ namespace uba
 		bool allowKeepFilesInMemory = IsWindows;
 		bool allowOutputFiles = IsWindows;
 		bool suppressLogging = false;
-		bool readObjFilesCompressed = false;
+		bool readIntermediateFilesCompressed = false;
 		int uiLanguage = 0;
 		StringBuffer<256> logFile;
 	};
diff --git a/Engine/Source/Programs/UnrealBuildAccelerator/Detours/Private/Windows/UbaDetoursFunctionsWin.cpp b/Engine/Source/Programs/UnrealBuildAccelerator/Detours/Private/Windows/UbaDetoursFunctionsWin.cpp
--- a/Engine/Source/Programs/UnrealBuildAccelerator/Detours/Private/Windows/UbaDetoursFunctionsWin.cpp
+++ b/Engine/Source/Programs/UnrealBuildAccelerator/Detours/Private/Windows/UbaDetoursFunctionsWin.cpp
@@ -1117,7 +1117,7 @@ void Init(const DetoursPayload& payload, u64 startTime)
 	g_stats.attach.time += GetTime() - startTime;
 	g_stats.attach.count = 1;
 
-	g_filesCouldBeCompressed = payload.readObjFilesCompressed && g_rules->CanDependOnCompressedFiles();
+	g_filesCouldBeCompressed = payload.readIntermediateFilesCompressed && g_rules->CanDependOnCompressedFiles();
 }
 
 void Deinit(u64 startTime)
diff --git a/Engine/Source/Programs/UnrealBuildAccelerator/Host/Private/UbaExports.cpp b/Engine/Source/Programs/UnrealBuildAccelerator/Host/Private/UbaExports.cpp
--- a/Engine/Source/Programs/UnrealBuildAccelerator/Host/Private/UbaExports.cpp
+++ b/Engine/Source/Programs/UnrealBuildAccelerator/Host/Private/UbaExports.cpp
@@ -411,7 +411,7 @@ extern "C"
 		return &process.GetStartInfo();
 	}
 
-	uba::SessionServerCreateInfo* SessionServerCreateInfo_Create(uba::StorageServer& storage, uba::NetworkServer& client, uba::LogWriter& writer, const uba::tchar* rootDir, const uba::tchar* traceOutputFile, bool disableCustomAllocator, bool launchVisualizer, bool resetCas, bool writeToDisk, bool detailedTrace, bool allowWaitOnMem, bool allowKillOnMem, bool storeObjFilesCompressed)
+	uba::SessionServerCreateInfo* SessionServerCreateInfo_Create(uba::StorageServer& storage, uba::NetworkServer& client, uba::LogWriter& writer, const uba::tchar* rootDir, const uba::tchar* traceOutputFile, bool disableCustomAllocator, bool launchVisualizer, bool resetCas, bool writeToDisk, bool detailedTrace, bool allowWaitOnMem, bool allowKillOnMem, bool storeIntermediateFilesCompressed)
 	{
 		auto info = new uba::SessionServerCreateInfo(storage, client, writer);
 		info->Apply(uba::GetConfig());
@@ -424,7 +424,7 @@ extern "C"
 		info->detailedTrace = detailedTrace;
 		info->allowWaitOnMem = allowWaitOnMem;
 		info->allowKillOnMem = allowKillOnMem;
-		info->storeObjFilesCompressed = storeObjFilesCompressed;
+		info->storeIntermediateFilesCompressed = storeIntermediateFilesCompressed;
 		//info->remoteTraceEnabled = true;
 		info->remoteLogEnabled = true;
 		return info;
diff --git a/Engine/Source/Programs/UnrealBuildAccelerator/Host/Public/UbaExports.h b/Engine/Source/Programs/UnrealBuildAccelerator/Host/Public/UbaExports.h
--- a/Engine/Source/Programs/UnrealBuildAccelerator/Host/Public/UbaExports.h
+++ b/Engine/Source/Programs/UnrealBuildAccelerator/Host/Public/UbaExports.h
@@ -101,7 +101,7 @@ extern "C"
 	using SessionServer_CustomServiceFunction = uba::u32(uba::ProcessHandle* handle, const void* recv, uba::u32 recvSize, void* send, uba::u32 sendCapacity, void* userData);
 
 	UBA_API uba::SessionServerCreateInfo* SessionServerCreateInfo_Create(uba::StorageServer& storage, uba::NetworkServer& client, uba::LogWriter& writer, const uba::tchar* rootDir, const uba::tchar* traceOutputFile,
-		bool disableCustomAllocator, bool launchVisualizer, bool resetCas, bool writeToDisk, bool detailedTrace, bool allowWaitOnMem = false, bool allowKillOnMem = false, bool storeObjFilesCompressed = false);
+		bool disableCustomAllocator, bool launchVisualizer, bool resetCas, bool writeToDisk, bool detailedTrace, bool allowWaitOnMem = false, bool allowKillOnMem = false, bool storeIntermediateFilesCompressed = false);
 	UBA_API void SessionServerCreateInfo_Destroy(uba::SessionServerCreateInfo* info);
 
 	UBA_API uba::SessionServer* SessionServer_Create(const uba::SessionServerCreateInfo& info, const uba::u8* environment = nullptr, uba::u32 environmentSize = 0);
