[UBA]
* Fixed compressed pch working with clang

diff --git a/Engine/Source/Programs/UnrealBuildAccelerator/Core/Private/UbaApplicationRules.cpp b/Engine/Source/Programs/UnrealBuildAccelerator/Core/Private/UbaApplicationRules.cpp
--- a/Engine/Source/Programs/UnrealBuildAccelerator/Core/Private/UbaApplicationRules.cpp
+++ b/Engine/Source/Programs/UnrealBuildAccelerator/Core/Private/UbaApplicationRules.cpp
@@ -12,6 +12,8 @@ namespace uba
 		return fileName.EndsWith(TC(".obj"))
 			|| fileName.EndsWith(TC(".o"))
 			|| fileName.EndsWith(TC(".pch"))
+			|| fileName.EndsWith(TC(".gch"))
+			|| fileName.EndsWith(TC(".ifc"))
 			;
 	}
 
@@ -94,7 +96,7 @@ namespace uba
 				|| fileName.EndsWith(TC(".pch"))
 				|| fileName.EndsWith(TC(".ifc"))
 				|| fileName.EndsWith(TC(".rc2.res")) // Not really an obj file.. 
-				;// || fileName.EndsWith(TC(".h.pch") // Not tested enough
+				;
 		}
 
 		virtual bool IsRarelyRead(const StringView& file) const override
@@ -216,7 +218,7 @@ namespace uba
 
 		virtual bool CanDependOnCompressedFiles() override
 		{
-			return true;
+			return true; // Yes, .obj files
 		}
 
 		virtual DependencyCrawlerType GetDependencyCrawlerType() const override
@@ -332,7 +334,9 @@ namespace uba
 				|| fileName.EndsWith(TC(".o.tmp")) // Clang writes to tmp file and then move
 				|| fileName.EndsWith(TC(".obj.tmp")) // Clang (verse) writes to tmp file and then move
 				|| fileName.EndsWith(TC(".obj")) // Clang with -fno-temp-file write obj file directly
-				;// || fileName.EndsWith(TC(".gch.tmp")); // Need to fix "has been modified since the precompiled header"
+				|| fileName.EndsWith(TC(".pch")) // Clang with -fno-temp-file write pch file directly
+				|| fileName.EndsWith(TC(".gch.tmp")) // 
+				;
 		}
 
 		virtual bool IsRarelyRead(const StringView& file) const override
@@ -380,6 +384,10 @@ namespace uba
 		{
 			return DependencyCrawlerType_MsvcCompiler;
 		}
+		virtual bool CanDependOnCompressedFiles() override
+		{
+			return true; // Yes, .pch
+		}
 	};
 
 	class ApplicationRulesLdLLdExe : public ApplicationRulesClang
