[UBA]
* Cosmetic fixes in visualizer

diff --git a/Engine/Source/Programs/UnrealBuildAccelerator/Visualizer/Private/UbaVisualizer.cpp b/Engine/Source/Programs/UnrealBuildAccelerator/Visualizer/Private/UbaVisualizer.cpp
--- a/Engine/Source/Programs/UnrealBuildAccelerator/Visualizer/Private/UbaVisualizer.cpp
+++ b/Engine/Source/Programs/UnrealBuildAccelerator/Visualizer/Private/UbaVisualizer.cpp
@@ -2183,15 +2183,15 @@ namespace uba
 			logger.SetColor(m_memColor).Info(L"  Mem: %ls/%ls", BytesToText(m_stats.memTotal - m_stats.memAvail).str, BytesToText(m_stats.memTotal).str);
 			if (m_stats.recvBytes || m_stats.sendBytes)
 			{
-				logger.SetColor(m_recvColor).Info(L"  Recv: %lsps", BytesToText(m_stats.recvBytesPerSecond*8).str);
-				logger.SetColor(m_sendColor).Info(L"  Send: %lsps", BytesToText(m_stats.sendBytesPerSecond*8).str);
+				logger.SetColor(m_recvColor).Info(L"  Recv: %lsbit/s", BytesToText(m_stats.recvBytesPerSecond*8).str);
+				logger.SetColor(m_sendColor).Info(L"  Send: %lsbit/s", BytesToText(m_stats.sendBytesPerSecond*8).str);
 			}
 			if (m_stats.ping)
 				logger.Info(L"  Ping: %ls", TimeToText(m_stats.ping, false, m_traceView.frequency).str);
 
 
 			for (auto& pair : m_stats.drives)
-				logger.SetColor(m_driveColor).Info(L"  %c: %u%% R:%sps W:%sps", pair.first, pair.second.busyPercent, BytesToText(pair.second.readPerSecond).str, BytesToText(pair.second.writePerSecond).str);
+				logger.SetColor(m_driveColor).Info(L"  %c: %u%% R:%s/s W:%s/s", pair.first, pair.second.busyPercent, BytesToText(pair.second.readPerSecond).str, BytesToText(pair.second.writePerSecond).str);
 
 			logger.AddSpace(3);
 			logger.DrawAtCursor();
