[UBA]
* fixed so kernelstats are written to the right place

diff --git a/Engine/Source/Programs/UnrealBuildAccelerator/Common/Private/UbaSession.cpp b/Engine/Source/Programs/UnrealBuildAccelerator/Common/Private/UbaSession.cpp
--- a/Engine/Source/Programs/UnrealBuildAccelerator/Common/Private/UbaSession.cpp
+++ b/Engine/Source/Programs/UnrealBuildAccelerator/Common/Private/UbaSession.cpp
@@ -3033,6 +3033,9 @@ namespace uba
 		auto span = std::span(files, fileCount);
 		m_workManager.ParallelFor(fileCount - 1, span, [&](const WorkContext&, auto& it)
 			{
+				KernelStatsScope ks(process.m_kernelStats);
+				StorageStatsScope ss(process.m_storageStats);
+				SessionStatsScope sessionStatsScope(process.m_sessionStats);
 				if (!writeFile(**it))
 					success = false;
 			}, TCV("WriteFilesToDisk"));
diff --git a/Engine/Source/Programs/UnrealBuildAccelerator/Common/Private/UbaStorage.cpp b/Engine/Source/Programs/UnrealBuildAccelerator/Common/Private/UbaStorage.cpp
--- a/Engine/Source/Programs/UnrealBuildAccelerator/Common/Private/UbaStorage.cpp
+++ b/Engine/Source/Programs/UnrealBuildAccelerator/Common/Private/UbaStorage.cpp
@@ -352,6 +352,9 @@ namespace uba
 			u8* uncompressedData = slot;
 			u8* compressBuffer = slot + BufferSlotHalfSize;
 
+			auto& memoryCompressTime = KernelStats::GetCurrent().memoryCompress;
+
+			TimerScope cts(stats.compressWrite);
 			while (left)
 			{
 				u64 uncompressedBlockSize = Min(left, maxUncompressedBlock);
@@ -375,10 +378,11 @@ namespace uba
 				u8* destBuf = compressBuffer;
 				OO_SINTa compressedBlockSize;
 				{
-					TimerScope cts(stats.compressWrite);
+					TimerScope kts(memoryCompressTime);
 					compressedBlockSize = OodleLZ_Compress((OodleLZ_Compressor)m_casCompressor, uncompressedData, (OO_SINTa)uncompressedBlockSize, destBuf + 8, (OodleLZ_CompressionLevel)m_casCompressionLevel, NULL, NULL, NULL, scratchMem, scratchSize);
 					if (compressedBlockSize == OODLELZ_FAILED)
 						return m_logger.Error(TC("Failed to compress %llu bytes for %s"), uncompressedBlockSize, from);
+					memoryCompressTime.bytes += compressedBlockSize;
 				}
 
 				*(u32*)destBuf =  u32(compressedBlockSize);
@@ -461,9 +465,14 @@ namespace uba
 			rec->events[i].Create(true);
 
 		StorageStats& stats = Stats();
+		TimerScope cts(stats.compressWrite);
+
+		auto& kernelStats = KernelStats::GetCurrent();
 
-		auto work = [this, rec, &stats](const WorkContext&)
+		auto work = [this, rec, &stats, &kernelStats](const WorkContext&)
 		{
+			KernelStatsScope kss(kernelStats);
+
 			u8* slot = m_bufferSlots.Pop();
 			auto _ = MakeGuard([&](){ m_bufferSlots.Push(slot); });
 			u8* compressSlotBuffer = slot + BufferSlotHalfSize;
@@ -483,7 +492,7 @@ namespace uba
 				{
 					void* scratchMem = slot;
 					u64 scratchSize = BufferSlotHalfSize;
-					TimerScope cts(stats.compressWrite);
+					TimerScope kts(kernelStats.memoryCompress);
 					compressedBlockSize = OodleLZ_Compress((OodleLZ_Compressor)m_casCompressor, uncompressedDataSlot, uncompressedBlockSize, compressSlotBuffer + 8, (OodleLZ_CompressionLevel)m_casCompressionLevel, NULL, NULL, NULL, scratchMem, scratchSize);
 					if (compressedBlockSize == OODLELZ_FAILED)
 					{
@@ -491,6 +500,7 @@ namespace uba
 						rec->error = true;
 						return;
 					}
+					kernelStats.memoryCompress.bytes += compressedBlockSize;
 				}
 				*(u32*)compressSlotBuffer =  u32(compressedBlockSize);
 				*(u32*)(compressSlotBuffer+4) =  u32(uncompressedBlockSize);
diff --git a/Engine/Source/Programs/UnrealBuildAccelerator/Common/Public/UbaTrace.h b/Engine/Source/Programs/UnrealBuildAccelerator/Common/Public/UbaTrace.h
--- a/Engine/Source/Programs/UnrealBuildAccelerator/Common/Public/UbaTrace.h
+++ b/Engine/Source/Programs/UnrealBuildAccelerator/Common/Public/UbaTrace.h
@@ -69,7 +69,7 @@ namespace uba
 		#undef UBA_TRACE_TYPE
 	};
 
-	static constexpr u32 TraceVersion = 40;
+	static constexpr u32 TraceVersion = 41;
 	static constexpr u32 TraceReadCompatibilityVersion = 6;
 
 	class Trace : public WorkTracker
diff --git a/Engine/Source/Programs/UnrealBuildAccelerator/Core/Public/UbaProcessStats.h b/Engine/Source/Programs/UnrealBuildAccelerator/Core/Public/UbaProcessStats.h
--- a/Engine/Source/Programs/UnrealBuildAccelerator/Core/Public/UbaProcessStats.h
+++ b/Engine/Source/Programs/UnrealBuildAccelerator/Core/Public/UbaProcessStats.h
@@ -133,6 +133,7 @@ namespace uba
 			UBA_KERNEL_STAT(ExtendedTimer, closeHandle, 0) \
 			UBA_KERNEL_STAT(ExtendedTimer, traverseDir, 27) \
 			UBA_KERNEL_STAT(ExtendedTimer, virtualAlloc, 30) \
+			UBA_KERNEL_STAT(TimeAndBytes, memoryCompress, 41) \
 
 		#define UBA_KERNEL_STAT(type, var, ver) type var;
 		UBA_KERNEL_STATS
