[UBA]
* Fixed bug related to getting size of obj file on remote detoured process running on wine (using NtQueryFullAttributesFile)

diff --git a/Engine/Source/Programs/UnrealBuildAccelerator/Common/Private/UbaSessionClient.cpp b/Engine/Source/Programs/UnrealBuildAccelerator/Common/Private/UbaSessionClient.cpp
--- a/Engine/Source/Programs/UnrealBuildAccelerator/Common/Private/UbaSessionClient.cpp
+++ b/Engine/Source/Programs/UnrealBuildAccelerator/Common/Private/UbaSessionClient.cpp
@@ -1329,6 +1329,7 @@ namespace uba
 		m_detailedTrace = reader.ReadBool();
 		m_shouldSendLogToServer = reader.ReadBool();
 		m_shouldSendTraceToServer = reader.ReadBool();
+		m_readIntermediateFilesCompressed = reader.ReadBool();
 
 		if (m_shouldSendTraceToServer)
 		{
diff --git a/Engine/Source/Programs/UnrealBuildAccelerator/Common/Private/UbaSessionServer.cpp b/Engine/Source/Programs/UnrealBuildAccelerator/Common/Private/UbaSessionServer.cpp
--- a/Engine/Source/Programs/UnrealBuildAccelerator/Common/Private/UbaSessionServer.cpp
+++ b/Engine/Source/Programs/UnrealBuildAccelerator/Common/Private/UbaSessionServer.cpp
@@ -807,6 +807,7 @@ namespace uba
 		writer.WriteBool(m_detailedTrace);
 		writer.WriteBool(m_remoteLogEnabled);
 		writer.WriteBool(m_remoteTraceEnabled);
+		writer.WriteBool(m_readIntermediateFilesCompressed);
 		WriteRemoteEnvironmentVariables(writer);
 
 		m_trace.SessionAdded(sessionId, connectionInfo.GetId(), name, info); // Must be inside lock for TraceSessionUpdate() to not include
diff --git a/Engine/Source/Programs/UnrealBuildAccelerator/Common/Public/UbaNetwork.h b/Engine/Source/Programs/UnrealBuildAccelerator/Common/Public/UbaNetwork.h
--- a/Engine/Source/Programs/UnrealBuildAccelerator/Common/Public/UbaNetwork.h
+++ b/Engine/Source/Programs/UnrealBuildAccelerator/Common/Public/UbaNetwork.h
@@ -16,7 +16,7 @@ namespace uba
 	
 	static constexpr u32 SystemNetworkVersion = 1339;
 	static constexpr u32 StorageNetworkVersion = 4;
-	static constexpr u32 SessionNetworkVersion = 44;
+	static constexpr u32 SessionNetworkVersion = 45;
 	static constexpr u32 CacheNetworkVersion = 5;
 	static constexpr u32 CacheBucketVersion = 3;
 
diff --git a/Engine/Source/Programs/UnrealBuildAccelerator/Detours/Private/UbaDetoursShared.cpp b/Engine/Source/Programs/UnrealBuildAccelerator/Detours/Private/UbaDetoursShared.cpp
--- a/Engine/Source/Programs/UnrealBuildAccelerator/Detours/Private/UbaDetoursShared.cpp
+++ b/Engine/Source/Programs/UnrealBuildAccelerator/Detours/Private/UbaDetoursShared.cpp
@@ -527,8 +527,17 @@ namespace uba
 						{
 							SCOPED_READ_LOCK(g_mappedFileTable.m_lookupLock, lock);
 							auto findIt = g_mappedFileTable.m_lookup.find(fileNameKey);
-							if (findIt != g_mappedFileTable.m_lookup.end() && !findIt->second.deleted)
+							if (findIt == g_mappedFileTable.m_lookup.end())
+							{
+								StringBuffer<> temp;
+								u32 closeId;
+								Rpc_CreateFileW(fileNameView, fileNameKey, AccessFlag_Read, temp.data, temp.capacity, fileSize, closeId, false);
+							}
+							else
+							{
+								UBA_ASSERT(!findIt->second.deleted);
 								fileSize = findIt->second.size;
+							}
 						}
 
 						outAttr.useCache = true;
diff --git a/Engine/Source/Programs/UnrealBuildAccelerator/Detours/Private/Windows/UbaDetoursFunctionsNtDll.inl b/Engine/Source/Programs/UnrealBuildAccelerator/Detours/Private/Windows/UbaDetoursFunctionsNtDll.inl
--- a/Engine/Source/Programs/UnrealBuildAccelerator/Detours/Private/Windows/UbaDetoursFunctionsNtDll.inl
+++ b/Engine/Source/Programs/UnrealBuildAccelerator/Detours/Private/Windows/UbaDetoursFunctionsNtDll.inl
@@ -485,7 +485,7 @@ NTSTATUS Detoured_NtQueryFullAttributesFile(POBJECT_ATTRIBUTES ObjectAttributes,
 		info.FileAttributes = data.dwFileAttributes;
 	}
 
-	DEBUG_LOG_DETOURED(L"NtQueryFullAttributesFile", L"(%.*s) -> %s", fileNameLen, fileName, ToString(res).data);
+	DEBUG_LOG_DETOURED(L"NtQueryFullAttributesFile", L"(%.*s) -> %s (Size: %llu)", fileNameLen, fileName, ToString(res).data, ToLargeInteger(attr.data.nFileSizeHigh, attr.data.nFileSizeLow).QuadPart);
 	return res;
 }
 
