[UBA]
* Added validation warning to code sending cache entry after seeing unexplainable crash on build farm

diff --git a/Engine/Source/Programs/UnrealBuildAccelerator/Common/Private/UbaCacheClient.cpp b/Engine/Source/Programs/UnrealBuildAccelerator/Common/Private/UbaCacheClient.cpp
--- a/Engine/Source/Programs/UnrealBuildAccelerator/Common/Private/UbaCacheClient.cpp
+++ b/Engine/Source/Programs/UnrealBuildAccelerator/Common/Private/UbaCacheClient.cpp
@@ -1176,9 +1176,14 @@ namespace uba
 		for (auto& kv : inputsStringToCasKey)
 			writer.Write7BitEncoded(kv.second);
 
-		if (logLinesSize)
-			if (writer.GetCapacityLeft() > logLinesSize + Get7BitEncodedCount(logLinesSize))
-				writer.WriteBytes(logLines, logLinesSize);
+		if (writer.GetPosition() >= SendMaxSize)
+		{
+			m_logger.Warning(TC("Something is wrong. Sending a cache entry that is too large. Output count: %llu, Input count: %llu LogLines size: %llu"), outputsStringToCasKey.size(), inputsStringToCasKey.size(), logLinesSize);
+			return false;
+		}
+
+		if (logLinesSize && logLinesSize < writer.GetCapacityLeft())
+			writer.WriteBytes(logLines, logLinesSize);
 
 		if (msg.Send(out))
 			return true;
