[UBA] Allow to save snapshot of current state of UBA trace periodically (CVar "r.UbaController.SaveTraceSnapshotInterval") and always save snapshot when processing an output file failed.

diff --git a/Engine/Source/Programs/UnrealBuildAccelerator/Common/Private/UbaSession.cpp b/Engine/Source/Programs/UnrealBuildAccelerator/Common/Private/UbaSession.cpp
--- a/Engine/Source/Programs/UnrealBuildAccelerator/Common/Private/UbaSession.cpp
+++ b/Engine/Source/Programs/UnrealBuildAccelerator/Common/Private/UbaSession.cpp
@@ -2124,6 +2124,11 @@ namespace uba
 		return m_trace.StopWrite(writeFile);
 	}
 
+	bool Session::SaveSnapshotOfTrace()
+	{
+		return m_trace.Write(m_traceOutputFile.data, true);
+	}
+
 	void Session::StartTraceThread()
 	{
 		m_traceThreadEvent.Create(true);
diff --git a/Engine/Source/Programs/UnrealBuildAccelerator/Common/Private/UbaTrace.cpp b/Engine/Source/Programs/UnrealBuildAccelerator/Common/Private/UbaTrace.cpp
--- a/Engine/Source/Programs/UnrealBuildAccelerator/Common/Private/UbaTrace.cpp
+++ b/Engine/Source/Programs/UnrealBuildAccelerator/Common/Private/UbaTrace.cpp
@@ -99,6 +99,40 @@ namespace uba
 		return true;
 	}
 
+	bool Trace::Write(const tchar* writeFileName, bool writeSummary)
+	{
+		if (!m_memoryBegin)
+			return true;
+
+		if (!writeFileName || !*writeFileName)
+			return true;
+
+		FileAccessor traceFile(m_logger, writeFileName);
+		if (!traceFile.CreateWrite(false, DefaultAttributes(), 0, nullptr))
+			return false;
+
+		ScopedFutex lock(m_memoryLock);
+		u64 fileSize = m_memoryPos;
+		if (!traceFile.Write(m_memoryBegin, fileSize))
+			return false;
+		lock.Leave();
+
+		if (writeSummary)
+		{
+			StackBinaryWriter<32> summaryWriter;
+			summaryWriter.WriteByte(TraceType_Summary);
+			summaryWriter.Write7BitEncoded(GetTime() - m_startTime);
+			if (!traceFile.Write(summaryWriter.GetData(), summaryWriter.GetPosition()))
+				return false;
+			fileSize += summaryWriter.GetPosition();
+		}
+
+		if (!traceFile.Close())
+			return false;
+		m_logger.Info(TC("Trace written to file %s with size %s"), writeFileName, BytesToText(fileSize).str);
+		return true;
+	}
+
 	bool Trace::StopWrite(const tchar* writeFileName)
 	{
 		if (!m_memoryBegin)
@@ -116,20 +150,7 @@ namespace uba
 			writer.Write7BitEncoded(GetTime() - m_startTime);
 		}
 
-		if (!writeFileName || !*writeFileName)
-			return true;
-
-		ScopedFutex lock(m_memoryLock);
-		FileAccessor traceFile(m_logger, writeFileName);
-		if (!traceFile.CreateWrite(false, DefaultAttributes(), 0, nullptr))
-			return false;
-		u64 fileSize = m_memoryPos;
-		if (!traceFile.Write(m_memoryBegin, fileSize))
-			return false;
-		if (!traceFile.Close())
-			return false;
-		m_logger.Info(TC("Trace file written to %s with size %s"), writeFileName, BytesToText(fileSize).str);
-		return true;
+		return Write(writeFileName, false); // Don't write summary, already included in memory
 	}
 
 	u32 Trace::TrackWorkStart(const StringView& desc, const Color& color)
diff --git a/Engine/Source/Programs/UnrealBuildAccelerator/Common/Public/UbaSession.h b/Engine/Source/Programs/UnrealBuildAccelerator/Common/Public/UbaSession.h
--- a/Engine/Source/Programs/UnrealBuildAccelerator/Common/Public/UbaSession.h
+++ b/Engine/Source/Programs/UnrealBuildAccelerator/Common/Public/UbaSession.h
@@ -87,6 +87,9 @@ namespace uba
 		bool ShouldStoreIntermediateFilesCompressed() const { return m_storeIntermediateFilesCompressed; }
 		bool HasDetailedtrace() const { return m_detailedTrace; }
 
+		// Will write current state of trace directly to disk
+		bool SaveSnapshotOfTrace();
+
 		virtual ~Session();
 
 		u32 CreateProcessId();
diff --git a/Engine/Source/Programs/UnrealBuildAccelerator/Common/Public/UbaTrace.h b/Engine/Source/Programs/UnrealBuildAccelerator/Common/Public/UbaTrace.h
--- a/Engine/Source/Programs/UnrealBuildAccelerator/Common/Public/UbaTrace.h
+++ b/Engine/Source/Programs/UnrealBuildAccelerator/Common/Public/UbaTrace.h
@@ -112,6 +112,10 @@ namespace uba
 		void CacheBeginWrite(u32 processId);
 		void CacheEndWrite(u32 processId, bool success, u64 bytesSent);
 
+		// Writes the current state of the trace to the specified output file without stopping anything
+		bool Write(const tchar* writeFileName, bool writeSummary);
+
+		// Stops the trace and writes it to the specified output file. This releases the internal memory buffer.
 		bool StopWrite(const tchar* writeFileName);
 
 		virtual u32 TrackWorkStart(const StringView& desc, const Color& color) final override;
diff --git a/Engine/Source/Programs/UnrealBuildAccelerator/Host/Private/UbaExports.cpp b/Engine/Source/Programs/UnrealBuildAccelerator/Host/Private/UbaExports.cpp
--- a/Engine/Source/Programs/UnrealBuildAccelerator/Host/Private/UbaExports.cpp
+++ b/Engine/Source/Programs/UnrealBuildAccelerator/Host/Private/UbaExports.cpp
@@ -558,6 +558,12 @@ extern "C"
 		server->RegisterCrossArchitectureMapping(from, to);
 	}
 
+	void SessionServer_SaveSnapshotOfTrace(uba::SessionServer* server)
+	{
+		if (server)
+			server->SaveSnapshotOfTrace();
+	}
+
 	void SessionServer_Destroy(uba::SessionServer* server)
 	{
 		if (server)
diff --git a/Engine/Source/Programs/UnrealBuildAccelerator/Host/Public/UbaExports.h b/Engine/Source/Programs/UnrealBuildAccelerator/Host/Public/UbaExports.h
--- a/Engine/Source/Programs/UnrealBuildAccelerator/Host/Public/UbaExports.h
+++ b/Engine/Source/Programs/UnrealBuildAccelerator/Host/Public/UbaExports.h
@@ -130,6 +130,7 @@ extern "C"
 	UBA_API void SessionServer_AddProcessBreadcrumbs(uba::SessionServer* server, uba::u32 processId, const uba::tchar* breadcrumbs, bool deleteOld = false);
 	UBA_API void SessionServer_RegisterCustomService(uba::SessionServer* server, SessionServer_CustomServiceFunction* function, void* userData = nullptr);
 	UBA_API void SessionServer_RegisterCrossArchitectureMapping(uba::SessionServer* server, const uba::tchar* from, const uba::tchar* to);
+	UBA_API void SessionServer_SaveSnapshotOfTrace(uba::SessionServer* server);
 	UBA_API void SessionServer_Destroy(uba::SessionServer* server);
 
 	// Scheduler
