[UBA]
* Fixed windows static analysis error

diff --git a/Engine/Source/Programs/UnrealBuildAccelerator/Core/Public/UbaProcessUtils.h b/Engine/Source/Programs/UnrealBuildAccelerator/Core/Public/UbaProcessUtils.h
--- a/Engine/Source/Programs/UnrealBuildAccelerator/Core/Public/UbaProcessUtils.h
+++ b/Engine/Source/Programs/UnrealBuildAccelerator/Core/Public/UbaProcessUtils.h
@@ -6,6 +6,9 @@
 
 #if PLATFORM_WINDOWS
 #include <malloc.h>
+#else
+#define _malloca alloca
+#define _freea
 #endif
 
 namespace uba
@@ -24,7 +27,8 @@ namespace uba
 		CharType currentChar = 0;
 		CharType lastChar = 0;
 		bool isBackslashOwned = false;
-		CharType* argWriteBuffer = static_cast<CharType*>( alloca(sizeof(CharType) * (argumentsLen + 1)) );
+		CharType* argWriteBuffer = static_cast<CharType*>( _malloca(sizeof(CharType) * (argumentsLen + 1)) );
+		auto g = MakeGuard([&]() { _freea(argWriteBuffer); });
 
 		for (const CharType* it = arguments; !isEnd; lastChar = currentChar, ++it)
 		{
