[UBA]
* Fixed bugs around output files being written/read by other processes within the same process heirarchy than the ones creating them.
* Changed so only root process writes output files to disk..

This code needs to be revisited because it is very fragile the way it works now... different application can have different rules around what is an "output file"

diff --git a/Engine/Source/Programs/UnrealBuildAccelerator/Common/Private/UbaProcess.cpp b/Engine/Source/Programs/UnrealBuildAccelerator/Common/Private/UbaProcess.cpp
--- a/Engine/Source/Programs/UnrealBuildAccelerator/Common/Private/UbaProcess.cpp
+++ b/Engine/Source/Programs/UnrealBuildAccelerator/Common/Private/UbaProcess.cpp
@@ -730,7 +730,7 @@ namespace uba
 			UBA_PROCESS_MESSAGES
 			#undef UBA_PROCESS_MESSAGE
 		}
-		return m_session.m_logger.Error(TC("Unknown message type %u"), messageType);
+		return CancelWithError(m_session.m_logger.Error(TC("Unknown message type %u"), messageType));
 	}
 
 	void ProcessImpl::LogLine(bool printInSession, TString&& line, LogEntryType logType)
@@ -769,7 +769,8 @@ namespace uba
 
 		auto rootsEntry = m_session.GetRootsEntry(m_startInfo.rootsHandle);
 		if (!rootsEntry)
-			return false;
+			return CancelWithError();
+
 		writer.WriteU16(u16(rootsEntry->memory.size()));
 		writer.WriteBytes(rootsEntry->memory.data(), rootsEntry->memory.size());
 		return true;
@@ -1086,10 +1087,10 @@ namespace uba
 		{
 			DuplicateHandle((HANDLE)m_nativeProcessHandle, nativeProcessHandle, GetCurrentProcess(), (HANDLE*)&process.m_nativeProcessHandle, 0, false, DUPLICATE_SAME_ACCESS);
 			if (!process.m_nativeProcessHandle || process.m_nativeProcessHandle == InvalidProcHandle)
-				return m_session.m_logger.Error(TC("Failed to duplicate handle for child process"));
+				return CancelWithError(m_session.m_logger.Error(TC("Failed to duplicate handle for child process")));
 			DuplicateHandle((HANDLE)m_nativeProcessHandle, nativeThreadHandle, GetCurrentProcess(), &process.m_nativeThreadHandle, 0, false, DUPLICATE_SAME_ACCESS);
 			if (!process.m_nativeThreadHandle || process.m_nativeThreadHandle == INVALID_HANDLE_VALUE)
-				return m_session.m_logger.Error(TC("Failed to duplicate handle for child thread"));
+				return CancelWithError(m_session.m_logger.Error(TC("Failed to duplicate handle for child thread")));
 			process.m_nativeProcessId = nativeProcessId;
 		}
 #else
@@ -1165,7 +1166,7 @@ namespace uba
 		LogEntryType entryType = isError ? LogEntryType_Error : LogEntryType_Info;
 
 		if (!m_session.LogLine(*this, line.c_str(), entryType))
-			return false;
+			return CancelWithError();
 
 		InternalLogLine(printInSession, std::move(line), entryType);
 		return true;
@@ -1210,7 +1211,7 @@ namespace uba
 
 		if (!IsCancelled())
 			if (m_startInfo.writeOutputFilesOnFail || m_startInfo.rules->IsExitCodeSuccess(m_nativeProcessExitCode))
-				m_messageSuccess = WriteFilesToDisk() && m_messageSuccess;
+				m_messageSuccess = WriteFilesToDisk(true) && m_messageSuccess;
 
 		if (m_parentProcess)
 		{
@@ -1225,7 +1226,7 @@ namespace uba
 
 	bool ProcessImpl::HandleFlushWrittenFiles(BinaryReader& reader, BinaryWriter& writer)
 	{
-		WriteFilesToDisk();
+		WriteFilesToDisk(false);
 		bool result = m_session.FlushWrittenFiles(*this);
 		writer.WriteBool(result);
 		return true;
@@ -1267,7 +1268,7 @@ namespace uba
 		kernelStats.Write(statsWriter);
 		BinaryReader statsReader(statsWriter.GetData(), 0, statsWriter.GetPosition());
 
-		WriteFilesToDisk();
+		WriteFilesToDisk(false);
 
 		bool newProcess = false;
 		m_exitCode = prevExitCode;
@@ -1345,6 +1346,32 @@ namespace uba
 		return true;
 	}
 
+	bool ProcessImpl::HandleTakeFileOwnership(BinaryReader& reader, BinaryWriter& writer)
+	{
+		StringKey fileNameKey = reader.ReadStringKey();
+		StringBuffer<> fileName;
+		reader.ReadString(fileName);
+
+		bool isOwner = false;
+		FileMappingHandle mappingHandle;
+		u64 mappingWritten = 0;
+
+		SCOPED_FUTEX(m_writtenFilesLock, lock);
+		auto findIt = m_writtenFiles.find(fileName.data);
+		if (findIt != m_writtenFiles.end())
+		{
+			WrittenFile& file = findIt->second;
+			isOwner = file.owner == this;
+			mappingHandle = file.mappingHandle;
+			mappingWritten = file.mappingWritten;
+			file.owner = this;
+		}
+		writer.WriteBool(isOwner);
+		writer.WriteU64(mappingHandle.ToU64());
+		writer.WriteU64(mappingWritten);
+		return true;
+	}
+
 	bool ProcessImpl::CreateTempFile(BinaryReader& reader, ProcHandle nativeProcessHandle, const tchar* application)
 	{
 		StringKey key = reader.ReadStringKey();
@@ -1397,14 +1424,16 @@ namespace uba
 		return true;
 	}
 
-	bool ProcessImpl::WriteFilesToDisk()
+	bool ProcessImpl::WriteFilesToDisk(bool isExiting)
 	{
 		Vector<WrittenFile*> files;
 		TimerScope ts(m_processStats.writeFiles);
 		SCOPED_FUTEX(m_writtenFilesLock, lock);
 		for (auto& kv : m_writtenFiles)
 		{
-			if (kv.second.owner != this)
+			// TODO: This requires some more logic. In some cases we have long living batch files as parent and then maybe we want to write out files to disk
+			// But for now we only let the root process write out all output files
+			if (m_parentProcess)
 				continue;
 			kv.second.owner = nullptr;
 			if (kv.second.mappingHandle.IsValid())
@@ -2233,4 +2262,11 @@ namespace uba
 			}
 		}
 	}
+
+	bool ProcessImpl::CancelWithError(bool dummy)
+	{
+		m_cancelEvent.Set();
+		return false;
+	}
+
 }
diff --git a/Engine/Source/Programs/UnrealBuildAccelerator/Common/Private/UbaSession.cpp b/Engine/Source/Programs/UnrealBuildAccelerator/Common/Private/UbaSession.cpp
--- a/Engine/Source/Programs/UnrealBuildAccelerator/Common/Private/UbaSession.cpp
+++ b/Engine/Source/Programs/UnrealBuildAccelerator/Common/Private/UbaSession.cpp
@@ -2992,7 +2992,8 @@ namespace uba
 				Storage::GetMappingString(name, file.mappingHandle, 0);
 				SCOPED_FUTEX(m_fileMappingTableLookupLock, lookupLock);
 				auto insres = m_fileMappingTableLookup.try_emplace(file.key);
-				UBA_ASSERT(insres.second);
+				if (!insres.second)
+					m_logger.Error(TC("Trying to write the same file twice (%s)"), file.name.c_str());
 				FileMappingEntry& entry = insres.first->second;
 				lookupLock.Leave();
 				SCOPED_FUTEX(entry.lock, entryCs);
diff --git a/Engine/Source/Programs/UnrealBuildAccelerator/Common/Public/UbaProcess.h b/Engine/Source/Programs/UnrealBuildAccelerator/Common/Public/UbaProcess.h
--- a/Engine/Source/Programs/UnrealBuildAccelerator/Common/Public/UbaProcess.h
+++ b/Engine/Source/Programs/UnrealBuildAccelerator/Common/Public/UbaProcess.h
@@ -85,7 +85,7 @@ namespace uba
 		void LogLine(bool printInSession, TString&& line, LogEntryType logType);
 		bool CreateTempFile(BinaryReader& reader, ProcHandle nativeProcessHandle, const tchar* application);
 		bool OpenTempFile(BinaryReader& reader, BinaryWriter& writer, const tchar* application);
-		bool WriteFilesToDisk();
+		bool WriteFilesToDisk(bool isExiting);
 
 		const tchar* InternalGetChildLogFile(StringBufferBase& temp);
 		
@@ -98,6 +98,7 @@ namespace uba
 		void InternalLogLine(bool printInSession, TString&& line, LogEntryType logType);
 		void ClearTempFiles();
 		void WaitForParent();
+		bool CancelWithError(bool dummy = true);
 
 		ProcessStartInfoHolder m_startInfo;
 		Session& m_session;
diff --git a/Engine/Source/Programs/UnrealBuildAccelerator/Core/Public/UbaProtocol.h b/Engine/Source/Programs/UnrealBuildAccelerator/Core/Public/UbaProtocol.h
--- a/Engine/Source/Programs/UnrealBuildAccelerator/Core/Public/UbaProtocol.h
+++ b/Engine/Source/Programs/UnrealBuildAccelerator/Core/Public/UbaProtocol.h
@@ -40,6 +40,7 @@ namespace uba
 		UBA_PROCESS_MESSAGE(HostRun) \
 		UBA_PROCESS_MESSAGE(ResolveCallstack) \
 		UBA_PROCESS_MESSAGE(CheckRemapping) \
+		UBA_PROCESS_MESSAGE(TakeFileOwnership) \
 
 	enum MessageType : u8
 	{
diff --git a/Engine/Source/Programs/UnrealBuildAccelerator/Detours/Private/Windows/UbaDetoursFunctionsKernelBase.inl b/Engine/Source/Programs/UnrealBuildAccelerator/Detours/Private/Windows/UbaDetoursFunctionsKernelBase.inl
--- a/Engine/Source/Programs/UnrealBuildAccelerator/Detours/Private/Windows/UbaDetoursFunctionsKernelBase.inl
+++ b/Engine/Source/Programs/UnrealBuildAccelerator/Detours/Private/Windows/UbaDetoursFunctionsKernelBase.inl
@@ -1597,7 +1597,7 @@ bool Shared_MoveFile(LPCWSTR lpExistingFileName, LPCWSTR lpNewFileName, DWORD dw
 		FileInfo& sourceInfo = it->second;
 		lock.Leave();
 
-		if (IsOutputFile(dest, true))
+		if (g_allowOutputFiles && g_rules->IsOutputFile(dest))
 		{
 			sourceInfo.deleted = true;
 			UBA_ASSERT(!sourceInfo.memoryFile->isLocalOnly);
diff --git a/Engine/Source/Programs/UnrealBuildAccelerator/Detours/Private/Windows/UbaDetoursFunctionsNtDll.inl b/Engine/Source/Programs/UnrealBuildAccelerator/Detours/Private/Windows/UbaDetoursFunctionsNtDll.inl
--- a/Engine/Source/Programs/UnrealBuildAccelerator/Detours/Private/Windows/UbaDetoursFunctionsNtDll.inl
+++ b/Engine/Source/Programs/UnrealBuildAccelerator/Detours/Private/Windows/UbaDetoursFunctionsNtDll.inl
@@ -856,8 +856,7 @@ NTSTATUS NTAPI Shared_NtCreateFile(bool IsCreateFunc, PHANDLE hFileHandle, ACCES
 
 	bool useContent = IsContentUse(DesiredAccess, CreateDisposition);
 	bool isWrite = IsWrite(DesiredAccess, CreateDisposition);
-	//bool keepInMemory = (KeepInMemory(fileName) && (useContent || NeedsSharedMemory(fileName.data))) || IsOutputFile(fileName, isWrite, isDeleteOnClose);
-	bool keepInMemory = KeepInMemory(fileName, isWrite) || IsOutputFile(fileName, isWrite, isDeleteOnClose);
+	bool keepInMemory = KeepInMemory(fileName, isWrite) || ((isWrite || isDeleteOnClose) && g_rules->IsOutputFile(fileName));
 
 #if UBA_DEBUG_LOG_ENABLED
 	const wchar_t* isWriteStr = isWrite ? L" WRITE" : L""; (void)isWriteStr;
@@ -1193,6 +1192,8 @@ NTSTATUS NTAPI Shared_NtCreateFile(bool IsCreateFunc, PHANDLE hFileHandle, ACCES
 
 	if (lpFileName[0] == '^') // It is a HANDLE from session process
 	{
+		UBA_ASSERT(!isWrite);
+
 		const wchar_t* handleStr = lpFileName + 1;
 		const wchar_t* handleStrEnd = wcschr(handleStr, '-');
 		if (!handleStrEnd)
@@ -1225,6 +1226,23 @@ NTSTATUS NTAPI Shared_NtCreateFile(bool IsCreateFunc, PHANDLE hFileHandle, ACCES
 		return STATUS_SUCCESS;
 	}
 
+	auto TakeFileOwnershipFromSession = [&](FileMappingHandle mappingHandle, u64 mappingHandleSize)
+		{
+			info.memoryFile = new MemoryFile(nullptr, false);
+			MemoryFile& mf = *info.memoryFile;
+			HANDLE newHandle = 0;
+			True_DuplicateHandle(g_hostProcess, mappingHandle.mh, GetCurrentProcess(), &newHandle, 0, FALSE, DUPLICATE_SAME_ACCESS);
+			UBA_ASSERTF(newHandle, L"DuplicateHandle failed when opening temp file %ls (%u)", fileName.data, GetLastError());
+			mf.writtenSize = mappingHandleSize;
+			mf.committedSize = mappingHandleSize;
+			mf.mappedSize = AlignUp(mf.committedSize, 4096);
+			TimerScope ts2(g_kernelStats.mapViewOfFile);
+			mf.baseAddress = (u8*)True_MapViewOfFile(newHandle, FILE_MAP_READ | FILE_MAP_ALL_ACCESS, 0, 0, mf.mappedSize);
+			UBA_ASSERTF(mf.baseAddress, L"MapViewOfFile failed when opening temp file %ls (%u)", fileName.data, GetLastError());
+			mf.reserveSize = FileTypeMaxSize(fileName, isSystemOrTempFile);
+			mf.mappingHandle.mh = newHandle;
+		};
+
 	if (keepInMemory || info.memoryFile)
 	{
 		#if UBA_DEBUG_LOG_ENABLED
@@ -1252,18 +1270,7 @@ NTSTATUS NTAPI Shared_NtCreateFile(bool IsCreateFunc, PHANDLE hFileHandle, ACCES
 					pcs.Leave();
 					if (mappingHandle.IsValid())
 					{
-						info.memoryFile = new MemoryFile(nullptr, false);
-						MemoryFile& mf = *info.memoryFile;
-						HANDLE newHandle = 0;
-						True_DuplicateHandle(g_hostProcess, mappingHandle.mh, GetCurrentProcess(), &newHandle, 0, FALSE, DUPLICATE_SAME_ACCESS);
-						UBA_ASSERTF(newHandle, L"DuplicateHandle failed when opening temp file %ls (%u)", fileName.data, GetLastError());
-						mf.writtenSize = mappingHandleSize;
-						mf.committedSize = mappingHandleSize;
-						TimerScope ts2(g_kernelStats.mapViewOfFile);
-						mf.baseAddress = (u8*)True_MapViewOfFile(newHandle, FILE_MAP_READ | FILE_MAP_ALL_ACCESS, 0, 0, mappingHandleSize);
-						UBA_ASSERTF(mf.baseAddress, L"MapViewOfFile failed when opening temp file %ls (%u)", fileName.data, GetLastError());
-						mf.reserveSize = FileTypeMaxSize(fileName, isSystemOrTempFile);
-						mf.mappingHandle.mh = newHandle;
+						TakeFileOwnershipFromSession(mappingHandle, mappingHandleSize);
 					}
 					else if (isWrite)
 					{
@@ -1278,17 +1285,54 @@ NTSTATUS NTAPI Shared_NtCreateFile(bool IsCreateFunc, PHANDLE hFileHandle, ACCES
 			}
 			else
 			{
+				bool isOutput = (isWrite || isDeleteOnClose) && g_rules->IsOutputFile(fileName);
 				bool isThrowAway = g_rules->IsThrowAway(fileName, g_runningRemote);
+
 				if (!isThrowAway && CreateDisposition == FILE_OPEN)
 				{
-					*hFileHandle = INVALID_HANDLE_VALUE;
-					DEBUG_LOG_DETOURED(funcName, L"ALREADYEXISTS (%ls) -> Error", fileName.data);
-					return STATUS_OBJECT_NAME_EXISTS;//STATUS_OBJECT_NAME_NOT_FOUND;
+					if (isOutput)
+					{
+						#if UBA_DEBUG_LOG_ENABLED
+						memoryType = L"TAKEOWNERSHIPWRITE";
+						#endif
+						TimerScope ts(g_stats.createFile);
+						SCOPED_WRITE_LOCK(g_communicationLock, pcs);
+						BinaryWriter writer;
+						writer.WriteByte(MessageType_TakeFileOwnership);
+						writer.WriteStringKey(fileNameKey);
+						writer.WriteString(fileName);
+						writer.Flush();
+						BinaryReader reader;
+						bool hasOwnership = reader.ReadBool();
+						FileMappingHandle mappingHandle = FileMappingHandle::FromU64(reader.ReadU64());
+						u64 mappingHandleSize = reader.ReadU64();
+						pcs.Leave();
+
+						if (!hasOwnership && mappingHandle.IsValid())
+						{
+							// This is a file that most likely was written by child process and is now used by parent (this) process
+							TakeFileOwnershipFromSession(mappingHandle, mappingHandleSize);
+						}
+					}
+
+					if (!info.memoryFile)
+					{
+						*hFileHandle = INVALID_HANDLE_VALUE;
+						DEBUG_LOG_DETOURED(funcName, L"ALREADYEXISTS (%ls) -> Error", fileName.data);
+						return STATUS_OBJECT_NAME_EXISTS;//STATUS_OBJECT_NAME_NOT_FOUND;
+					}
 				}
+				else
+				{
+					if (isOutput && (CreateDisposition & FILE_OPEN_IF) == 0)
+					{
+						UBA_ASSERTF(false, TC("Trying to open %s with openif. This is not supported"), fileName.data);
+					}
 
-				bool isLocal = !IsOutputFile(fileName, isWrite, isDeleteOnClose);
-				//UBA_ASSERTF(CreateDisposition != FILE_OPEN || Contains(fileName.data, L"vctip_"), TC("Unsupported disposition %u for file %s"), CreateDisposition, fileName.data);
-				info.memoryFile = new MemoryFile(isLocal, FileTypeMaxSize(fileName, isSystemOrTempFile), isThrowAway, fileName.data);
+					bool isLocal = !isOutput;
+					//UBA_ASSERTF(CreateDisposition != FILE_OPEN || Contains(fileName.data, L"vctip_"), TC("Unsupported disposition %u for file %s"), CreateDisposition, fileName.data);
+					info.memoryFile = new MemoryFile(isLocal, FileTypeMaxSize(fileName, isSystemOrTempFile), isThrowAway, fileName.data);
+				}
 			}
 
 			// TODO: Time should be in sync with host machine!
@@ -1317,16 +1361,34 @@ NTSTATUS NTAPI Shared_NtCreateFile(bool IsCreateFunc, PHANDLE hFileHandle, ACCES
 		DEBUG_LOG_DETOURED(funcName, L"(%s)%ls %llu (%ls) (%ls) -> Success", memoryType, isWriteStr, uintptr_t(*hFileHandle), lpFileName, (fileName.data != lpFileName ? fileName.data : L""));
 		return STATUS_SUCCESS;
 	}
-
-	lookupLock.Leave();
-
-	if (lpFileName[0] == ':')
+	else if (g_rules->IsOutputFile(fileName)) // It might be that a child process has written this file and then we need to try to take ownership of it in order to read it.
 	{
-		*hFileHandle = CreateFileHandle();
-		DEBUG_LOG_DETOURED(funcName, L"(PRELOAD)%ls %llu (%ls) (%ls) -> Success", isWriteStr, uintptr_t(*hFileHandle), lpFileName, (fileName.data != lpFileName ? fileName.data : L""));
-		return STATUS_SUCCESS;
+		TimerScope ts(g_stats.createFile);
+		SCOPED_WRITE_LOCK(g_communicationLock, pcs);
+		BinaryWriter writer;
+		writer.WriteByte(MessageType_TakeFileOwnership);
+		writer.WriteStringKey(fileNameKey);
+		writer.WriteString(fileName);
+		writer.Flush();
+		BinaryReader reader;
+		bool hasOwnership = reader.ReadBool();
+		FileMappingHandle mappingHandle = FileMappingHandle::FromU64(reader.ReadU64());
+		u64 mappingHandleSize = reader.ReadU64();
+		pcs.Leave();
+
+		if (!hasOwnership && mappingHandle.IsValid())
+		{
+			TakeFileOwnershipFromSession(mappingHandle, mappingHandleSize);
+			lookupLock.Leave();
+			*hFileHandle = CreateFileHandle();
+			TrackFileInput();
+			DEBUG_LOG_DETOURED(funcName, L"(TAKEOWNERSHIPREAD)%ls %llu (%ls) (%ls) -> Success", isWriteStr, uintptr_t(*hFileHandle), lpFileName, (fileName.data != lpFileName ? fileName.data : L""));
+			return STATUS_SUCCESS;
+		}
 	}
 
+	lookupLock.Leave();
+
 	StringView tempFileName;
 	if (lpFileName[0] == '#')
 		tempFileName = fileName;
@@ -1530,7 +1592,7 @@ NTSTATUS NTAPI Detoured_NtClose(HANDLE handle)
 		mappingWritten = mf->writtenSize;
 
 		u32 orginalNameLen = TStrlen(fi.originalName);
-		if (IsOutputFile(StringView(fi.originalName, orginalNameLen), IsWrite(fo->desiredAccess, 0), fo->deleteOnClose) && !g_rules->IsThrowAway(StringView(fi.originalName, orginalNameLen), g_runningRemote))
+		if ((fo->deleteOnClose || IsWrite(fo->desiredAccess, 0)) && g_rules->IsOutputFile({fi.originalName, orginalNameLen}) && !g_rules->IsThrowAway(StringView(fi.originalName, orginalNameLen), g_runningRemote))
 		{
 			// Need to report this file to host so it can be tracked in directory table
 			if (!mf->isReported)
diff --git a/Engine/Source/Programs/UnrealBuildAccelerator/Detours/Private/Windows/UbaDetoursFunctionsWin.cpp b/Engine/Source/Programs/UnrealBuildAccelerator/Detours/Private/Windows/UbaDetoursFunctionsWin.cpp
--- a/Engine/Source/Programs/UnrealBuildAccelerator/Detours/Private/Windows/UbaDetoursFunctionsWin.cpp
+++ b/Engine/Source/Programs/UnrealBuildAccelerator/Detours/Private/Windows/UbaDetoursFunctionsWin.cpp
@@ -509,8 +509,6 @@ const wchar_t* HandleToName(HANDLE handle)
 
 bool NeedsSharedMemory(const wchar_t* file) { return g_allowKeepFilesInMemory && g_rules->NeedsSharedMemory(file); }
 u64 FileTypeMaxSize(const StringBufferBase& file, bool isSystemOrTempFile) { return g_rules->FileTypeMaxSize(file, isSystemOrTempFile); }
-bool IsOutputFile(const StringView& fileName, bool isWrite, bool isDeleteOnClose = false) { return (isWrite || isDeleteOnClose) && g_allowOutputFiles && g_rules->IsOutputFile(fileName); }
-
 
 bool EnsureMapped(DetouredHandle& handle, DWORD dwFileOffsetHigh = 0, DWORD dwFileOffsetLow = 0, SIZE_T numberOfBytesToMap = 0, void* baseAddress = nullptr)
 {
