[UBA]
* Fixed so devirtualize path in dependency crawler is not outputting errors

diff --git a/Engine/Source/Programs/UnrealBuildAccelerator/Common/Private/UbaSession.cpp b/Engine/Source/Programs/UnrealBuildAccelerator/Common/Private/UbaSession.cpp
--- a/Engine/Source/Programs/UnrealBuildAccelerator/Common/Private/UbaSession.cpp
+++ b/Engine/Source/Programs/UnrealBuildAccelerator/Common/Private/UbaSession.cpp
@@ -1827,7 +1827,7 @@ namespace uba
 		return true;
 	}
 
-	bool Session::DevirtualizePath(StringBufferBase& inOut, RootsHandle rootsHandle)
+	bool Session::DevirtualizePath(StringBufferBase& inOut, RootsHandle rootsHandle, bool reportError)
 	{
 		if (!rootsHandle)
 			return true;
@@ -1840,7 +1840,7 @@ namespace uba
 			return true;
 		auto root = rootsEntry->roots.FindRoot(inOut);
 		if (!root)
-			return m_logger.Error(TC("Can't find root for path %s (Available roots: %s)"), inOut.data, rootsEntry->roots.GetAllRoots().c_str());
+			return reportError ? m_logger.Error(TC("Can't find root for path %s (Available roots: %s)"), inOut.data, rootsEntry->roots.GetAllRoots().c_str()) : false;
 
 		auto& path = rootsEntry->locals[root->index/PathsPerRoot];
 		StringBuffer<> temp;
@@ -3920,7 +3920,7 @@ namespace uba
 					return m_logger.Warning(TC("Failed to open %s"), out.fileName.data);
 			};
 
-		auto DevirtualizePathFunc = [this, rootsHandle = startInfo.rootsHandle](StringBufferBase& inOut) { return DevirtualizePath(inOut, rootsHandle); };
+		auto DevirtualizePathFunc = [this, rootsHandle = startInfo.rootsHandle](StringBufferBase& inOut) { return DevirtualizePath(inOut, rootsHandle, false); };
 
 		m_dependencyCrawler.Add(at + 1, startInfo.workingDir, CreateFileFunc, DevirtualizePathFunc, startInfo.application, crawlerType, startInfo.rules->index);
 	}
diff --git a/Engine/Source/Programs/UnrealBuildAccelerator/Common/Public/UbaSession.h b/Engine/Source/Programs/UnrealBuildAccelerator/Common/Public/UbaSession.h
--- a/Engine/Source/Programs/UnrealBuildAccelerator/Common/Public/UbaSession.h
+++ b/Engine/Source/Programs/UnrealBuildAccelerator/Common/Public/UbaSession.h
@@ -95,7 +95,7 @@ namespace uba
 		u32 CreateProcessId();
 
 		bool VirtualizePath(StringBufferBase& inOut, RootsHandle rootsHandle);
-		bool DevirtualizePath(StringBufferBase& inOut, RootsHandle rootsHandle);
+		bool DevirtualizePath(StringBufferBase& inOut, RootsHandle rootsHandle, bool reportError = true);
 		bool DevirtualizeString(TString& inOut, RootsHandle rootsHandle, bool allowPathsWithoutRoot, const tchar* hint = TC(""));
 		bool PopulateLocalToIndexRoots(RootPaths& out, RootsHandle rootsHandle);
 
