[UBA]
* Changed so written files are always populated for processes in process hierarchy. This is needed to be able to find files to transfer ownership of between processes when needed
* Fixed missing "output files" in clang application rules to make sure they are not downloaded/decompressed just to get the size (we should not care about the size of output files)
* Fixed linux compile error
* Fixed so child processes are waited on before output files are written since there are scenarios where child processes are never explicitly waited on by the detoured process
* Changed shared files inside process hierarchy is now in one struct

diff --git a/Engine/Source/Programs/UnrealBuildAccelerator/Common/Private/UbaLogger.cpp b/Engine/Source/Programs/UnrealBuildAccelerator/Common/Private/UbaLogger.cpp
--- a/Engine/Source/Programs/UnrealBuildAccelerator/Common/Private/UbaLogger.cpp
+++ b/Engine/Source/Programs/UnrealBuildAccelerator/Common/Private/UbaLogger.cpp
@@ -428,15 +428,15 @@ namespace uba
 			#if PLATFORM_WINDOWS
 			char buffer[512];
 			size_t destLen;
-			u32 maxLen = Min(strLen, 256u);
-			if (wcstombs_s(&destLen, buffer, sizeof(buffer), str, maxLen) != 0)
+			if (wcstombs_s(&destLen, buffer, sizeof(buffer), str, sizeof(buffer)-2) != 0)
 			{
 				strcpy_s(buffer, sizeof(buffer), "BAD_STRING\n");
 				destLen = 11;
 			}
 			else
 			{
-				buffer[destLen] = '\n';
+				buffer[destLen-1] = '\r'; // replace null terminator with \r
+				buffer[destLen++] = '\n';
 			}
 			m_file->Write(buffer, destLen);
 			#else
diff --git a/Engine/Source/Programs/UnrealBuildAccelerator/Common/Private/UbaProcess.cpp b/Engine/Source/Programs/UnrealBuildAccelerator/Common/Private/UbaProcess.cpp
--- a/Engine/Source/Programs/UnrealBuildAccelerator/Common/Private/UbaProcess.cpp
+++ b/Engine/Source/Programs/UnrealBuildAccelerator/Common/Private/UbaProcess.cpp
@@ -144,10 +144,7 @@ namespace uba
 	,	m_readEvent(*(m_comMemory.memory ? (new (m_comMemory.memory + sizeof(SharedEvent)*2) SharedEvent) : (SharedEvent*)nullptr))
 		#endif
 	,	m_detourEnabled(detourEnabled)
-	,	m_writtenFilesLock(parent ? parent->m_writtenFilesLock : *new Futex())
-	,	m_writtenFiles(parent ? parent->m_writtenFiles : *new UnorderedMap<TString, WrittenFile>())
-	,	m_tempFilesLock(parent ? parent->m_tempFilesLock : *new ReaderWriterLock())
-	,	m_tempFiles(parent ? parent->m_tempFiles : *new UnorderedMap<StringKey, WrittenFile>)
+	,	m_shared(parent ? parent->m_shared : *new Shared)
 	{
 		if (m_comMemory.memory)
 		{
@@ -182,15 +179,11 @@ namespace uba
 
 		if (!m_parentProcess)
 		{
-			for (auto& pair : m_writtenFiles)
+			for (auto& pair : m_shared.writtenFiles)
 				if (pair.second.mappingHandle.IsValid())
-					CloseFileMapping(m_session.m_logger, pair.second.mappingHandle, pair.first.c_str());
-			delete &m_writtenFiles;
-			delete &m_writtenFilesLock;
-
+					CloseFileMapping(m_session.m_logger, pair.second.mappingHandle, pair.second.name.c_str());
 			ClearTempFiles();
-			delete& m_tempFiles;
-			delete& m_tempFilesLock;
+			delete &m_shared;
 		}
 	}
 
@@ -540,6 +533,24 @@ namespace uba
 			if (exitCode == 0)
 				exitCode = InternalExitProcess(cancelled);
 
+			// For some reason a parent can exit before a child. I've seen it happen on ClangEditor win64 and also with some experimental features enabled on msvc
+			// We need to wait also because we need all written files to be added in shared file system
+			WaitForChildrenExit();
+
+			if (!cancelled)
+				if (m_startInfo.writeOutputFilesOnFail || m_startInfo.rules->IsExitCodeSuccess(m_nativeProcessExitCode))
+					if (!WriteFilesToDisk(true))
+						exitCode = UBA_EXIT_CODE(1);
+
+			if (m_parentProcess)
+			{
+				m_parentProcess->m_processStats.Add(m_processStats);
+				m_parentProcess->m_sessionStats.Add(m_sessionStats);
+				m_parentProcess->m_storageStats.Add(m_storageStats);
+				m_parentProcess->m_kernelStats.Add(m_kernelStats);
+			}
+
+
 			if (exitCode == 0 && !m_messageSuccess)
 				exitCode = UBA_EXIT_CODE(1);
 
@@ -561,8 +572,7 @@ namespace uba
 			m_logLines.clear();
 			m_trackedInputs.clear();
 
-			UBA_ASSERT(!m_parentProcess);
-			m_writtenFiles.clear();
+			m_shared.writtenFiles.clear();
 			ClearTempFiles();
 		}
 
@@ -589,18 +599,9 @@ namespace uba
 		KernelStats::GetGlobal().Add(m_kernelStats);
 
 		#if UBA_DEBUG_TRACK_PROCESS
-		m_session.m_debugLogger->Info(TC("ProcessExitedStart (%u)\n"), m_id);
+		m_session.m_debugLogger->Info(TC("ProcessExitedStart (%u)"), m_id);
 		#endif
 
-		// For some reason a parent can exit before a child. Need to figure out repro for this but I've seen it happen on ClangEditor win64
-		for (auto& child : m_childProcesses)
-		{
-			auto& childProcess = *(ProcessImpl*)child.m_process;
-			childProcess.m_waitForParent.Set();
-			while (!childProcess.m_hasExited)
-				Sleep(10);
-		}
-
 		UBA_ASSERT(!m_parentProcess || !m_parentProcess->m_hasExited);
 
 		m_hasExited = true;
@@ -641,7 +642,7 @@ namespace uba
 		m_session.ProcessExited(*this, m_processStats.wallTime);
 
 		#if UBA_DEBUG_TRACK_PROCESS
-		m_session.m_debugLogger->Info(TC("ProcessExitedDone  (%u)\n"), m_id);
+		m_session.m_debugLogger->Info(TC("ProcessExitedDone  (%u)"), m_id);
 		#endif
 	}
 
@@ -940,43 +941,43 @@ namespace uba
 		writer.WriteU32(m_session.GetDirectoryTableSize());
 		writer.WriteU32(m_session.GetFileMappingSize());
 
-#if PLATFORM_WINDOWS
-		if (!m_tempFilesModified)
+		if (m_tempFilesModified)
 		{
-			writer.WriteU32(0);
-			return true;
+			// This should be very rare (lld-link.exe uses it for mt.exe).. if we get too many temp files we'll need to rethink this design
+			SCOPED_READ_LOCK(m_shared.tempFilesLock, l);
+			writer.WriteU32(u32(m_shared.tempFiles.size()));
+			for (auto& kv : m_shared.tempFiles)
+			{
+				writer.WriteStringKey(kv.first);
+				writer.WriteU64(kv.second.mappingWritten);
+			}
 		}
+		else
+			writer.WriteU32(0);
 
-		// This should be very rare (lld-link.exe uses it for mt.exe).. if we get too many temp files we'll need to rethink this design
-		SCOPED_READ_LOCK(m_tempFilesLock, l);
-		writer.WriteU32(u32(m_tempFiles.size()));
-		for (auto& kv : m_tempFiles)
-		{
-			writer.WriteStringKey(kv.first);
-			writer.WriteU64(kv.second.mappingWritten);
-		}
-#endif
-		return true;
+		return HandleGetWrittenFiles(reader, writer);
 	}
 
-	bool ProcessImpl::HandleGetParentWrittenFiles(BinaryReader& reader, BinaryWriter& writer)
+	bool ProcessImpl::HandleGetWrittenFiles(BinaryReader& reader, BinaryWriter& writer)
 	{
-		// TODO: Recurse
-		UBA_ASSERT(m_parentProcess);
 		u32& writtenCount = *(u32*)writer.AllocWrite(4);
-		
-		//writer.WriteU32(u32(m_parentProcess->m_writtenFiles.size()));
 		u32 count = 0;
-		for (auto& kv : m_parentProcess->m_writtenFiles)
+		SCOPED_FUTEX_READ(m_shared.writtenFilesLock, lock);
+		for (auto& kv : m_shared.writtenFiles)
 		{
 			WrittenFile& wf = kv.second;
-			if (wf.mappingHandle.IsValid())
+			if (wf.owner == this)
 				continue;
+
 			u64 fileSize;
-			if (!FileExists(m_session.m_logger, wf.name.c_str(), &fileSize))
-				continue;
-			writer.WriteStringKey(wf.key);
-			writer.WriteString(wf.name);
+			if (wf.mappingHandle.IsValid())
+				fileSize = wf.mappingWritten;
+			else
+				if (!FileExists(m_session.m_logger, wf.backedName.c_str(), &fileSize))
+					continue;
+			writer.WriteStringKey(kv.first);
+			writer.WriteString(wf.backedName);
+			writer.WriteU64(wf.mappingHandle.ToU64());
 			writer.WriteU64(fileSize);
 			++count;
 		}
@@ -1052,7 +1053,7 @@ namespace uba
 		#endif
 
 		#if UBA_DEBUG_TRACK_PROCESS
-		m_session.m_debugLogger->Info(TC("CreateChildProcess (%u creating child %u at index %u) %s %s (%s)\n"), m_id, process.m_id, childProcessId-1, process.m_realApplication.c_str(), commandLineWithoutApplication.c_str(), info.logFile);
+		m_session.m_debugLogger->Info(TC("CreateChildProcess (%u creating child %u at index %u) %s %s (%s)"), m_id, process.m_id, childProcessId-1, process.m_realApplication.c_str(), commandLineWithoutApplication.c_str(), info.logFile);
 		#endif
 		return true;
 	}
@@ -1102,13 +1103,13 @@ namespace uba
 #endif
 		
 		#if UBA_DEBUG_TRACK_PROCESS
-		m_session.m_debugLogger->Info(TC("WaitForChildProcessReady (%u waiting for child %u at index %u)\n"), m_id, process.m_id, processId-1);
+		m_session.m_debugLogger->Info(TC("WaitForChildProcessReady (%u waiting for child %u at index %u)"), m_id, process.m_id, processId-1);
 		#endif
 
 		setWaitForParent.Execute();
 
 		#if UBA_DEBUG_TRACK_PROCESS
-		m_session.m_debugLogger->Info(TC("WaitForChildProcessReadyDone (%u waiting for child %u at index %u)\n"), m_id, process.m_id, processId-1);
+		m_session.m_debugLogger->Info(TC("WaitForChildProcessReadyDone (%u waiting for child %u at index %u)"), m_id, process.m_id, processId-1);
 		#endif
 
 		#if PLATFORM_WINDOWS
@@ -1209,18 +1210,6 @@ namespace uba
 		m_processStats.Add(stats);
 		m_kernelStats.Add(kernelStats);
 
-		if (!IsCancelled())
-			if (m_startInfo.writeOutputFilesOnFail || m_startInfo.rules->IsExitCodeSuccess(m_nativeProcessExitCode))
-				m_messageSuccess = WriteFilesToDisk(true) && m_messageSuccess;
-
-		if (m_parentProcess)
-		{
-			m_parentProcess->m_processStats.Add(m_processStats);
-			m_parentProcess->m_sessionStats.Add(m_sessionStats);
-			m_parentProcess->m_storageStats.Add(m_storageStats);
-			m_parentProcess->m_kernelStats.Add(m_kernelStats);
-		}
-
 		return false;
 	}
 
@@ -1290,7 +1279,7 @@ namespace uba
 		m_logLines.clear();
 		m_trackedInputs.clear();
 		m_trackedOutputs.clear();
-		m_writtenFiles.clear();
+		m_shared.writtenFiles.clear();
 
 		ClearTempFiles();
 		
@@ -1356,15 +1345,16 @@ namespace uba
 		FileMappingHandle mappingHandle;
 		u64 mappingWritten = 0;
 
-		SCOPED_FUTEX(m_writtenFilesLock, lock);
-		auto findIt = m_writtenFiles.find(fileName.data);
-		if (findIt != m_writtenFiles.end())
+		SCOPED_FUTEX(m_shared.writtenFilesLock, lock);
+		auto findIt = m_shared.writtenFiles.find(fileNameKey);
+		if (findIt != m_shared.writtenFiles.end())
 		{
 			WrittenFile& file = findIt->second;
 			isOwner = file.owner == this;
 			mappingHandle = file.mappingHandle;
 			mappingWritten = file.mappingWritten;
 			file.owner = this;
+			file.originalMappingHandle = {};
 		}
 		writer.WriteBool(isOwner);
 		writer.WriteU64(mappingHandle.ToU64());
@@ -1387,8 +1377,8 @@ namespace uba
 			return true;
 		}
 
-		SCOPED_WRITE_LOCK(m_tempFilesLock, tempLock);
-		auto insres = m_tempFiles.try_emplace(key, WrittenFile{ nullptr, StringKeyZero, fileName.data, newHandle, mappingHandleSize, mappingHandle });
+		SCOPED_WRITE_LOCK(m_shared.tempFilesLock, tempLock);
+		auto insres = m_shared.tempFiles.try_emplace(key, WrittenFile{ nullptr, StringKeyZero, fileName.data, fileName.data, newHandle, mappingHandleSize, mappingHandle });
 		if (insres.second)
 			return true;
 
@@ -1411,9 +1401,9 @@ namespace uba
 		u64 mappingHandle = 0;
 		u64 mappingWritten = 0;
 
-		SCOPED_READ_LOCK(m_tempFilesLock, lock);
-		auto findIt = m_tempFiles.find(fileKey);
-		if (findIt != m_tempFiles.end())
+		SCOPED_READ_LOCK(m_shared.tempFilesLock, lock);
+		auto findIt = m_shared.tempFiles.find(fileKey);
+		if (findIt != m_shared.tempFiles.end())
 		{
 			mappingHandle = findIt->second.mappingHandle.ToU64();
 			mappingWritten = findIt->second.mappingWritten;
@@ -1428,8 +1418,8 @@ namespace uba
 	{
 		Vector<WrittenFile*> files;
 		TimerScope ts(m_processStats.writeFiles);
-		SCOPED_FUTEX(m_writtenFilesLock, lock);
-		for (auto& kv : m_writtenFiles)
+		SCOPED_FUTEX(m_shared.writtenFilesLock, lock);
+		for (auto& kv : m_shared.writtenFiles)
 		{
 			// TODO: This requires some more logic. In some cases we have long living batch files as parent and then maybe we want to write out files to disk
 			// But for now we only let the root process write out all output files
@@ -1449,11 +1439,11 @@ namespace uba
 		if (m_startInfo.trackInputs)
 		{
 			u64 totalBytes = 0;
-			for (auto& kv : m_writtenFiles)
+			for (auto& kv : m_shared.writtenFiles)
 				totalBytes += GetStringWriteSize(kv.second.name.c_str(), kv.second.name.size());
 			m_trackedOutputs.resize(totalBytes);
 			BinaryWriter writer(m_trackedOutputs.data(), 0, totalBytes);
-			for (auto& kv : m_writtenFiles)
+			for (auto& kv : m_shared.writtenFiles)
 				writer.WriteString(kv.second.name);
 		}
 		return true;
@@ -1700,13 +1690,13 @@ namespace uba
 			m_nativeThreadHandle = processInfo.hThread;
 
 			#if UBA_DEBUG_TRACK_PROCESS
-			m_session.m_debugLogger->Info(TC("CreateRealProcess  (%u) %s %s\n"), m_id, m_realApplication.c_str(), m_startInfo.arguments);
+			m_session.m_debugLogger->Info(TC("CreateRealProcess  (%u) %s %s"), m_id, m_realApplication.c_str(), m_startInfo.arguments);
 			#endif
 		}
 		else
 		{
 			#if UBA_DEBUG_TRACK_PROCESS
-			m_session.m_debugLogger->Info(TC("WaitingForParentReady (%u)\n"), m_id);
+			m_session.m_debugLogger->Info(TC("WaitingForParentReady (%u)"), m_id);
 			#endif
 
 			WaitForParent();
@@ -1714,7 +1704,7 @@ namespace uba
 				return UBA_EXIT_CODE(7);
 
 			#if UBA_DEBUG_TRACK_PROCESS
-			m_session.m_debugLogger->Info(TC("WaitingForParentReadyDone (%u)\n"), m_id);
+			m_session.m_debugLogger->Info(TC("WaitingForParentReadyDone (%u)"), m_id);
 			#endif
 
 			m_extractExports = m_parentProcess->m_extractExports;
@@ -2037,13 +2027,13 @@ namespace uba
 			pipeGuard0.Cancel();
 
 			#if UBA_DEBUG_TRACK_PROCESS
-			m_session.m_debugLogger->Info(TC("CreateRealProcess  (%u) %s %.100s\n"), m_id, m_realApplication.c_str(), m_startInfo.arguments);
+			m_session.m_debugLogger->Info(TC("CreateRealProcess  (%u) %s %.100s"), m_id, m_realApplication.c_str(), m_startInfo.arguments);
 			#endif
 		}
 		else
 		{
 			#if UBA_DEBUG_TRACK_PROCESS
-			m_session.m_debugLogger->Info(TC("WaitingForParent (%u) %.100s\n"), m_id, m_realApplication.c_str());
+			m_session.m_debugLogger->Info(TC("WaitingForParent (%u) %.100s"), m_id, m_realApplication.c_str());
 			#endif
 
 			//logger.Info("Waiting for parent");
@@ -2110,7 +2100,7 @@ namespace uba
 						hadTimeout = true;
 						const tchar* gotMessage = m_gotExitMessage ? TC("Got") : TC("Did not get");
 						const tchar* isCancelledNewCheck = IsCancelled() ? TC("true") : TC("false");
-						logger.Info(TC("WaitForSingleObject timed out after 120 seconds waiting for process %s to exit (Exit code %u, %s ExitMessage and wrote %u files. Cancelled: %s. Runtime: %s). Will terminate and wait again"), m_startInfo.GetDescription(), m_nativeProcessExitCode, gotMessage, u32(m_writtenFiles.size()), isCancelledNewCheck, TimeToText(GetTime() - m_startTime).str);
+						logger.Info(TC("WaitForSingleObject timed out after 120 seconds waiting for process %s to exit (Exit code %u, %s ExitMessage and wrote %u files. Cancelled: %s. Runtime: %s). Will terminate and wait again"), m_startInfo.GetDescription(), m_nativeProcessExitCode, gotMessage, u32(m_shared.writtenFiles.size()), isCancelledNewCheck, TimeToText(GetTime() - m_startTime).str);
 						TerminateProcess((HANDLE)handle, m_nativeProcessExitCode);
 						continue;
 					}
@@ -2242,10 +2232,10 @@ namespace uba
 
 	void ProcessImpl::ClearTempFiles()
 	{
-		for (auto& pair : m_tempFiles)
+		for (auto& pair : m_shared.tempFiles)
 			if (pair.second.mappingHandle.IsValid())
 				CloseFileMapping(m_session.m_logger, pair.second.mappingHandle, pair.second.name.c_str());
-		m_tempFiles.clear();
+		m_shared.tempFiles.clear();
 	}
 
 	void ProcessImpl::WaitForParent()
@@ -2263,6 +2253,17 @@ namespace uba
 		}
 	}
 
+	void ProcessImpl::WaitForChildrenExit()
+	{
+		for (auto& child : m_childProcesses)
+		{
+			auto& childProcess = *(ProcessImpl*)child.m_process;
+			childProcess.m_waitForParent.Set();
+			while (!childProcess.m_hasExited)
+				Sleep(10);
+		}
+	}
+
 	bool ProcessImpl::CancelWithError(bool dummy)
 	{
 		m_cancelEvent.Set();
diff --git a/Engine/Source/Programs/UnrealBuildAccelerator/Common/Private/UbaSession.cpp b/Engine/Source/Programs/UnrealBuildAccelerator/Common/Private/UbaSession.cpp
--- a/Engine/Source/Programs/UnrealBuildAccelerator/Common/Private/UbaSession.cpp
+++ b/Engine/Source/Programs/UnrealBuildAccelerator/Common/Private/UbaSession.cpp
@@ -239,6 +239,224 @@ namespace uba
 
 	}
 
+	bool Session::WriteFileToDisk(ProcessImpl& process, WrittenFile& file)
+	{
+		bool shouldEvictFromMemory = IsRarelyReadAfterWritten(process, file.name) || file.mappingWritten > m_keepOutputFileMemoryMapsThreshold;
+
+		if (ShouldWriteToDisk(file.name))
+		{
+			auto& rules = *process.m_startInfo.rules;
+			RootsHandle rootsHandle = process.GetStartInfo().rootsHandle;
+
+			bool storeCompressed = m_storeIntermediateFilesCompressed && g_globalRules.FileCanBeCompressed(file.name);
+			bool shouldDevirtualize = false;
+			bool escapeSpaces = false;
+			if (!storeCompressed)
+				shouldDevirtualize = HasVfs(rootsHandle) && rules.ShouldDevirtualizeFile(file.name, escapeSpaces);
+
+#if UBA_ENABLE_ON_DISK_FILE_MAPPINGS
+			if (!storeCompressed && !shouldDevirtualize && file.mappingHandle.fh)
+			{
+				bool success = true;
+				SetFilePointerEx(file.mappingHandle.fh, ToLargeInteger(file.mappingWritten), NULL, FILE_BEGIN);
+				::SetEndOfFile(file.mappingHandle.fh);
+
+				if (false)
+				{
+					TrackWorkScope work(m_workManager, AsView(TC("Flush")));
+					work.AddHint(file.name);
+					FlushFileBuffers(file.mappingHandle.fh);
+				}
+
+				if (true)
+				{
+					FILE_DISPOSITION_INFO info;
+					info.DeleteFile = false;
+					if (!::SetFileInformationByHandle(file.mappingHandle.fh, FileDispositionInfo, &info, sizeof(info)))
+						m_logger.Warning(TC("Failed to remove delete-on-close"));
+				}
+
+				if (false)
+				{
+					StringBuffer<> name2;
+					name2.Append(TC("\\??\\")).Append(file.name);
+
+					u8 buffer[2048];
+					auto& info = *(FILE_RENAME_INFO*)buffer;;
+					info.ReplaceIfExists = true;
+					info.RootDirectory = 0;
+					info.FileNameLength = u32(name2.count*sizeof(tchar));
+					info.Flags = FILE_RENAME_FLAG_POSIX_SEMANTICS;
+					memcpy(info.FileName, name2.data, info.FileNameLength+sizeof(tchar));
+					TrackWorkScope work(m_workManager, AsView(TC("Rename")));
+					work.AddHint(file.name);
+					success = SetFileInformationByHandle(file.mappingHandle.fh, FileRenameInfo, &info, sizeof(buffer));
+				}
+
+				if (success)
+				{
+					{
+						TrackWorkScope work(m_workManager, AsView(TC("CloseFileMapping")));
+						work.AddHint(file.name);
+						CloseFileMapping(m_logger, file.mappingHandle, file.name.c_str());
+					}
+					file.mappingHandle = {};
+					file.originalMappingHandle = {};
+					m_storage.InvalidateCachedFileInfo(file.key);
+					return true;
+				}
+				m_logger.Warning(TC("SetFileInformationByHandle failed %s (%s)"), file.name.c_str(), LastErrorToText().data);
+			}
+#endif
+
+			#if UBA_DEBUG_TRACK_MAPPING
+			m_debugLogger->Info(TC("Writing written file with mapping 0x%llx for %s"), u64(file.mappingHandle.mh), file.name.c_str());
+			#endif
+
+			u64 fileSize = file.mappingWritten;
+			u8* mem = MapViewOfFile(m_logger, file.mappingHandle, FILE_MAP_READ, 0, fileSize);
+			if (!mem)
+				return m_logger.Error(TC("Failed to map view of filehandle for read %s (%s)"), file.name.c_str(), LastErrorToText().data);
+
+			//PrefetchVirtualMemory(mem, fileSize);
+
+			auto memClose = MakeGuard([&](){ UnmapViewOfFile(m_logger, mem, fileSize, file.name.c_str()); });
+
+			if (storeCompressed)
+			{
+				Storage::WriteResult res;
+				CompressedFileHeader header { CalculateCasKey(mem, fileSize, true, &m_workManager, file.name.c_str()) };
+
+				if (!m_storage.WriteCompressed(res, TC("MemoryMap"), InvalidFileHandle, mem, fileSize, file.name.c_str(), &header, sizeof(header), file.lastWriteTime))
+					return false;
+
+				shouldEvictFromMemory = false; // Can't evict without properly update filemappingtable.. the file on disk does now not match what was registered for write
+			}
+			else
+			{
+				FileAccessor destinationFile(m_logger, file.name.c_str());
+
+				if (shouldDevirtualize)
+				{
+					// Need to turn paths back into local paths
+
+					if (!destinationFile.CreateWrite(false, DefaultAttributes(), 0, m_tempPath.data))
+						return false;
+
+					MemoryBlock block(5*1024*1024);
+					if (!DevirtualizeDepsFile(rootsHandle, block, mem, fileSize, escapeSpaces, file.name.c_str()))
+						return false;
+
+					if (!destinationFile.Write(block.memory, block.writtenSize))
+						return false;
+				}
+				else
+				{
+					// Seems like best combo (for windows at least) is to use writes with overlap and max 16 at the same time.
+					// On one machine we get twice as fast without overlap if no bottleneck. On another machine (ntfs compression on) we get twice as slow without overlap
+					// Both machines behaves well with overlap AND bottleneck. Both machine are 128 logical core thread rippers.
+					bool useFileMapForWrite = fileSize && fileSize <= m_writeFilesFileMapMax; // ::CreateFileMappingW does not work for zero-length files
+					bool useOverlap = !useFileMapForWrite && fileSize >= m_writeFilesNoBufferingMin;//fileSize > 2*1024*1024;
+
+					u32 attributes = DefaultAttributes();
+					if (useOverlap)
+						attributes |= (FILE_FLAG_OVERLAPPED | FILE_FLAG_NO_BUFFERING);
+
+					if (useFileMapForWrite)
+					{
+						if (!destinationFile.CreateMemoryWrite(false, attributes, fileSize, m_tempPath.data))
+							return false;
+						MapMemoryCopy(destinationFile.GetData(), mem, fileSize);
+					}
+					else
+					{
+						if (!destinationFile.CreateWrite(false, attributes, fileSize, m_tempPath.data))
+							return false;
+						if (!destinationFile.Write(mem, fileSize, 0, true))
+							return false;
+					}
+				}
+
+				if (u64 time = file.lastWriteTime)
+					if (!SetFileLastWriteTime(destinationFile.GetHandle(), time))
+						return m_logger.Error(TC("Failed to set file time on filehandle for %s"), file.name.c_str());
+
+				if (!destinationFile.Close(file.lastWriteTime ? nullptr : &file.lastWriteTime))
+					return false;
+			}
+
+			// There are directory crawlers happening in parallel so we need to really make sure to invalidate this one since a crawler can actually
+			// hit this file with information from a query before it was written.. and then it will turn it back to "verified" using old info
+			m_storage.InvalidateCachedFileInfo(file.key);
+		}
+		else
+		{
+			// Delete existing file to make sure it is not picked up (since it is out of date)
+			uba::DeleteFileW(file.name.c_str());
+		}
+
+		FileMappingHandle mh = file.mappingHandle;
+		file.mappingHandle = {};
+		file.originalMappingHandle = {};
+
+		if (shouldEvictFromMemory)
+		{
+			#if UBA_DEBUG_TRACK_MAPPING
+			m_debugLogger->Info(TC("Mapping eviction queued 0x%llx (%s)"), u64(mh.mh), file.name.c_str());
+			#endif
+
+			m_workManager.AddWork([mh, n = file.name, this](const WorkContext&)
+				{
+					#if UBA_DEBUG_TRACK_MAPPING
+					m_debugLogger->Info(TC("Mapping evicted 0x%llx (%s)"), u64(mh.mh), n.c_str());
+					#endif
+
+					CloseFileMapping(m_logger, mh, n.c_str());
+				}, 1, TC("CloseFileMapping"));
+		}
+		else
+		{
+			StringBuffer<> name;
+			Storage::GetMappingString(name, mh, 0);
+			SCOPED_FUTEX(m_fileMappingTableLookupLock, lookupLock);
+			auto insres = m_fileMappingTableLookup.try_emplace(file.key);
+			FileMappingEntry& entry = insres.first->second;
+			lookupLock.Leave();
+
+			SCOPED_FUTEX(entry.lock, entryCs);
+			if (!insres.second && entry.canBeFreed) // It could be that this file has been read as input and that is fine.
+				m_logger.Error(TC("Trying to write the same file twice (%s)"), file.name.c_str());
+			UBA_ASSERT(mh.IsValid());
+			entry.handled = true;
+			entry.mapping = mh;
+			entry.mappingOffset = 0;
+			entry.size = file.mappingWritten;
+			entry.lastWriteTime = file.lastWriteTime;
+			entry.isDir = false;
+			entry.success = true;
+			entry.canBeFreed = true;
+			entry.usedCount = 0;
+			entry.usedCountBeforeFree = g_globalRules.GetUsedCountBeforeFree(file.name);
+
+			#if UBA_DEBUG_TRACK_MAPPING
+			entry.name = file.name;
+			m_debugLogger->Info(TC("Mapping kept 0x%llx (%s) from detoured process (UsedCountBeforeFree: %u)"), u64(mh.mh), entry.name.c_str(), u32(entry.usedCountBeforeFree));
+			#endif
+
+			SCOPED_WRITE_LOCK(m_fileMappingTableMemLock, lock);
+			BinaryWriter writer(m_fileMappingTableMem, m_fileMappingTableSize);
+			writer.WriteStringKey(file.key);
+			writer.WriteString(name);
+			writer.Write7BitEncoded(file.mappingWritten);
+			u32 newSize = (u32)writer.GetPosition();
+			m_fileMappingTableSize = (u32)newSize;
+		}
+
+		TraceWrittenFile(process.m_id, file.name, file.mappingWritten);
+
+		return true;
+	}
+
 	void Session::AddEnvironmentVariableNoLock(const tchar* key, const tchar* value)
 	{
 		m_environmentVariables.insert(m_environmentVariables.end(), key, key + TStrlen(key));
@@ -320,7 +538,7 @@ namespace uba
 					memoryWriter.WriteString(e.name, e.nameLen);
 
 					#if UBA_DEBUG_TRACK_DIR
-					m_debugLogger->Info(TC("    %s (Size: %llu, Attr: %u, Key: %s, Id: %llu)\n"), e.name, e.size, e.attributes, KeyToString(fileKey).data, e.id);
+					m_debugLogger->Info(TC("    %s (Size: %llu, Attr: %u, Key: %s, Id: %llu)"), e.name, e.size, e.attributes, KeyToString(fileKey).data, e.id);
 					#endif
 
 					u64 id = e.id;
@@ -357,7 +575,7 @@ namespace uba
 			if (!res)
 			{
 				#if UBA_DEBUG_TRACK_DIR
-				m_debugLogger->Info(TC("    FAILED (not existing?)\n"));
+				m_debugLogger->Info(TC("    FAILED (not existing?)"));
 				#endif
 
 				if (!IsWindows || dirPath.count > 3)
@@ -367,7 +585,7 @@ namespace uba
 		else
 		{
 			#if UBA_DEBUG_TRACK_DIR
-			m_debugLogger->Info(TC("    FORCED EMPTY\n"));
+			m_debugLogger->Info(TC("    FORCED EMPTY"));
 			#endif
 
 			#if PLATFORM_WINDOWS
@@ -729,7 +947,7 @@ namespace uba
 
 		#if UBA_DEBUG_TRACK_MAPPING
 		entry.name = fileName;
-		m_debugLogger->Info(TC("Mapping created 0x%llx (%s) from file (%s) - %s\n"), u64(entry.mapping.mh), entry.name.c_str(), hint, TimeToText(GetTime() - ts2.start).str);
+		m_debugLogger->Info(TC("Mapping created 0x%llx (%s) from file (%s) - %s"), u64(entry.mapping.mh), entry.name.c_str(), hint, TimeToText(GetTime() - ts2.start).str);
 		#endif
 
 		entry.isDir = isDir;
@@ -834,7 +1052,7 @@ namespace uba
 
 		#if UBA_DEBUG_TRACK_MAPPING
 		entry.name = fileName;
-		m_debugLogger->Info(TC("Mapping created 0x%llx (%s) from view\n"), u64(entry.mapping.mh), entry.name.c_str());
+		m_debugLogger->Info(TC("Mapping created 0x%llx (%s) from view"), u64(entry.mapping.mh), entry.name.c_str());
 		#endif
 
 		return true;
@@ -1017,7 +1235,7 @@ namespace uba
 		}
 
 		#if UBA_DEBUG_TRACK_DIR
-		m_debugLogger->Info(TC("TRACKADD    %s (Size: %llu, Attr: %u, Key: %s, Id: %llu)\n"), fileName.data, fileSize, attributes, KeyToString(fileNameKey).data, fileIndex);
+		m_debugLogger->Info(TC("TRACKADD    %s (Size: %llu, Attr: %u, Key: %s, Id: %llu)"), fileName.data, fileSize, attributes, KeyToString(fileNameKey).data, fileIndex);
 		#endif
 
 
@@ -1097,7 +1315,7 @@ namespace uba
 		}
 
 		#if UBA_DEBUG_TRACK_DIR
-		m_debugLogger->Info(TC("TRACKDEL    %s (Key: %s)\n"), fileName.data, KeyToString(fileNameKey).data);
+		m_debugLogger->Info(TC("TRACKDEL    %s (Key: %s)"), fileName.data, KeyToString(fileNameKey).data);
 		#endif
 
 		SCOPED_WRITE_LOCK(dirTable.m_memoryLock, memoryLock);
@@ -2009,7 +2227,7 @@ namespace uba
 			UBA_ASSERT(entry.canBeFreed);
 
 			#if UBA_DEBUG_TRACK_MAPPING
-			m_debugLogger->Info(TC("Mapping freed 0x%llx (%s)\n"), u64(entry.mapping.mh), entry.name.c_str());
+			m_debugLogger->Info(TC("Mapping freed 0x%llx (%s)"), u64(entry.mapping.mh), entry.name.c_str());
 			#endif
 
 			m_workManager.AddWork([mh = entry.mapping, this](const WorkContext&) { CloseFileMapping(m_logger, mh, TC("UsedFileMapping")); }, 1, TC("CloseFileMapping"));
@@ -2487,19 +2705,27 @@ namespace uba
 		return true;
 	}
 
-	void RemoveWrittenFile(ProcessImpl& process, const TString& name)
+	void Session::RemoveWrittenFile(ProcessImpl& process, const StringKey& fileKey)
 	{
-		SCOPED_FUTEX(process.m_writtenFilesLock, writtenLock);
-		auto& writtenFiles = process.m_writtenFiles;
-		auto findIt = writtenFiles.find(name);
+		SCOPED_FUTEX(process.m_shared.writtenFilesLock, writtenLock);
+		auto& writtenFiles = process.m_shared.writtenFiles;
+		auto findIt = writtenFiles.find(fileKey);
 		if (findIt == writtenFiles.end())
 			return;
+
 		FileMappingHandle h = findIt->second.mappingHandle;
+		TString name = findIt->second.name;
 		writtenFiles.erase(findIt);
 		writtenLock.Leave();
 
-		if (h.IsValid())
-			CloseFileMapping(process.m_session.GetLogger(), h, name.c_str());
+		if (!h.IsValid())
+			return;
+
+		#if UBA_DEBUG_TRACK_MAPPING
+		m_debugLogger->Info(TC("Removed %s with handle 0x%llx"), name.c_str(), h.mh);
+		#endif
+
+		CloseFileMapping(process.m_session.GetLogger(), h, name.c_str());
 	}
 
 	bool Session::CloseFile(CloseFileResponse& out, const CloseFileMessage& msg)
@@ -2509,7 +2735,7 @@ namespace uba
 		if (findIt == m_activeFiles.end())
 			return m_logger.Error(TC("This should not happen. Got unknown closeId %u - %s"), msg.closeId, msg.fileName.data);
 
-		ActiveFile file = findIt->second;
+		ActiveFile activeFile = findIt->second;
 		m_activeFiles.erase(msg.closeId);
 		lock.Leave();
 
@@ -2523,25 +2749,25 @@ namespace uba
 		}
 		if (msg.deleteOnClose)
 		{
-			RemoveWrittenFile(msg.process, file.name);
+			RemoveWrittenFile(msg.process, activeFile.nameKey);
 		}
 		else
 		{
-			StringKey key = file.nameKey;
-			const tchar* name = file.name.c_str();
-			const tchar* msgName = msg.fileName.data;
+			StringKey key = activeFile.nameKey;
+			StringView name = activeFile.name;
+			StringView msgName = msg.fileName;
 			if (!msg.newName.IsEmpty())
 			{
 				UBA_ASSERT(!msg.deleteOnClose);
-				RemoveWrittenFile(msg.process, file.name);
-				name = msg.newName.data;
+				RemoveWrittenFile(msg.process, key);
+				name = msg.newName;
 				key = msg.newNameKey;
 				if (!m_runningRemote)
-					msgName = msg.newName.data;
+					msgName = msg.newName;
 			}
 
-			SCOPED_FUTEX(msg.process.m_writtenFilesLock, writtenLock);
-			auto insres = msg.process.m_writtenFiles.try_emplace(name);
+			SCOPED_FUTEX(msg.process.m_shared.writtenFilesLock, writtenLock);
+			auto insres = msg.process.m_shared.writtenFiles.try_emplace(key);
 			WrittenFile& writtenFile = insres.first->second;
 
 			if (m_allowOutputFiles && writtenFile.owner != nullptr && writtenFile.owner != &msg.process)
@@ -2550,19 +2776,25 @@ namespace uba
 				//UBA_ASSERTF(false, TC("File %s changed owner.. should not happen (OldOwner: %s New owner: %s)"), name, writtenFile.owner->m_realApplication.c_str(), msg.process.m_realApplication.c_str());
 			}
 
-			writtenFile.owner = &msg.process;
 			writtenFile.attributes = msg.attributes;
 
 			bool addMapping = true;
-			if (!insres.second)
+			if (insres.second)
 			{
-				if (writtenFile.name != msgName)
+				writtenFile.name = name.ToString();
+				writtenFile.key = key;
+				writtenFile.backedName = msgName.ToString();
+				writtenFile.owner = &msg.process;
+			}
+			else
+			{
+				if (writtenFile.backedName != msgName.data)
 				{
 					UBA_ASSERT(!msg.mappingHandle.IsValid() && !writtenFile.mappingHandle.IsValid());
-					writtenFile.name = msgName;
+					writtenFile.backedName = msgName.ToString();
 				}
 
-				if (!msg.mappingHandle.IsValid() || msg.mappingHandle == writtenFile.originalMappingHandle)
+				if (!msg.mappingHandle.IsValid() || (msg.mappingHandle == writtenFile.originalMappingHandle && writtenFile.owner == &msg.process))
 				{
 					if (msg.mappingWritten)
 					{
@@ -2571,14 +2803,18 @@ namespace uba
 					}
 					addMapping = false;
 				}
-				else
+				else if (writtenFile.mappingHandle.IsValid())
 				{
-					if (writtenFile.mappingHandle.IsValid())
-					{
-						CloseFileMapping(m_logger, writtenFile.mappingHandle, msg.fileName.data);
-						writtenFile.mappingHandle = {};
-					}
+					#if UBA_DEBUG_TRACK_MAPPING
+					m_debugLogger->Info(TC("Closing old mapping 0x%llx for %s"), u64(writtenFile.mappingHandle.mh), writtenFile.name.c_str());
+					#endif
+
+					CloseFileMapping(m_logger, writtenFile.mappingHandle, msg.fileName.data);
+					writtenFile.mappingHandle = {};
+					writtenFile.originalMappingHandle = {};
 				}
+
+				writtenFile.owner = &msg.process;
 			}
 
 			ProcessStartInfo& startInfo = msg.process.m_startInfo;
@@ -2586,21 +2822,21 @@ namespace uba
 			if (!m_runningRemote && HasVfs(startInfo.rootsHandle)) // For posix we write the dependency file directly to disk so we need to update it if vfs is enabled
 			{
 				bool escapeSpaces;
-				if (!msg.mappingHandle.IsValid() && startInfo.rules->ShouldDevirtualizeFile(file.name, escapeSpaces))
+				if (!msg.mappingHandle.IsValid() && startInfo.rules->ShouldDevirtualizeFile(activeFile.name, escapeSpaces))
 				{
 					// TODO: On linux we don't use file mappings for outputs yet.. so we have to open the file and change it
-					FileAccessor readFile(m_logger, name);
+					FileAccessor readFile(m_logger, name.data);
 					if (!readFile.OpenMemoryRead())
 						return false;
 					void* mem = readFile.GetData();
 					fileSize = readFile.GetSize();
 					MemoryBlock block(5*1024*1024);
 					RootsHandle rootsHandle = startInfo.rootsHandle;
-					if (!DevirtualizeDepsFile(rootsHandle, block, mem, fileSize, escapeSpaces, name))
+					if (!DevirtualizeDepsFile(rootsHandle, block, mem, fileSize, escapeSpaces, name.data))
 						return false;
 					if (!readFile.Close())
 						return false;
-					FileAccessor writeFile(m_logger, name);
+					FileAccessor writeFile(m_logger, name.data);
 					if (!writeFile.CreateWrite())
 						return false;
 					if (!writeFile.Write(block.memory, block.writtenSize))
@@ -2613,17 +2849,19 @@ namespace uba
 
 			if (addMapping)
 			{
-				writtenFile.name = msgName;
 				FileMappingHandle mappingHandle;
 				if (msg.mappingHandle.IsValid())
-					if (!DuplicateFileMapping(m_logger, msg.process.m_nativeProcessHandle, msg.mappingHandle, GetCurrentProcessHandle(), mappingHandle, FILE_MAP_READ, false, DUPLICATE_SAME_ACCESS, msgName))
-						return m_logger.Error(TC("Failed to duplicate file mapping handle for %s"), name);
+					if (!DuplicateFileMapping(m_logger, msg.process.m_nativeProcessHandle, msg.mappingHandle, GetCurrentProcessHandle(), mappingHandle, 0, false, DUPLICATE_SAME_ACCESS, msgName.data))
+						return m_logger.Error(TC("Failed to duplicate file mapping handle for %s"), name.data);
 
-				writtenFile.key = key;
 				writtenFile.mappingHandle = mappingHandle;
 				writtenFile.mappingWritten = msg.mappingWritten;
 				writtenFile.originalMappingHandle = msg.mappingHandle;
 				writtenFile.lastWriteTime = GetSystemTimeAsFileTime();
+
+				#if UBA_DEBUG_TRACK_MAPPING
+				m_debugLogger->Info(TC("Adding written file with mapping 0x%llx (from 0x%llx) for %s"), u64(writtenFile.mappingHandle.mh), u64(msg.mappingHandle.mh), writtenFile.name.c_str());
+				#endif
 			}
 
 			if (writtenFile.mappingHandle.IsValid())
@@ -2633,23 +2871,23 @@ namespace uba
 				lastWriteTime = writtenFile.lastWriteTime;
 			}
 
-			if (msg.process.m_extractExports && startInfo.rules->ShouldExtractSymbols(file.name))
-				if (!ExtractSymbolsFromObjectFile(msg, name, fileSize))
+			if (msg.process.m_extractExports && startInfo.rules->ShouldExtractSymbols(activeFile.name))
+				if (!ExtractSymbolsFromObjectFile(msg, name.data, fileSize))
 					return false;
 		}
 
 		if (!msg.newName.IsEmpty())
 		{
-			RegisterDeleteFile(file.nameKey, file.name);
+			RegisterDeleteFile(activeFile.nameKey, activeFile.name);
 			if (RegisterCreateFileForWrite(msg.newNameKey, msg.newName, registerRealFile, fileSize, lastWriteTime) && registerRealFile)
 				TraceWrittenFile(msg.process.m_id, msg.newName, fileSize);
 		}
 		else if (msg.deleteOnClose)
-			RegisterDeleteFile(file.nameKey, file.name);
+			RegisterDeleteFile(activeFile.nameKey, activeFile.name);
 		else
 		{
-			if (RegisterCreateFileForWrite(file.nameKey, file.name, registerRealFile, fileSize, lastWriteTime) && registerRealFile)
-				TraceWrittenFile(msg.process.m_id, file.name, fileSize);
+			if (RegisterCreateFileForWrite(activeFile.nameKey, activeFile.name, registerRealFile, fileSize, lastWriteTime) && registerRealFile)
+				TraceWrittenFile(msg.process.m_id, activeFile.name, fileSize);
 		}
 
 		out.directoryTableSize = GetDirectoryTableSize();
@@ -2669,7 +2907,7 @@ namespace uba
 			m_outputFiles.erase(msg.fileName.data);
 		}
 
-		RemoveWrittenFile(msg.process, msg.fileName.data);
+		RemoveWrittenFile(msg.process, msg.fileNameKey);
 
 		out.result = uba::DeleteFileW(msg.fileName.data);
 		out.errorCode = GetLastError();
@@ -2697,9 +2935,10 @@ namespace uba
 		auto& process = msg.process;
 		bool isMoved = false;
 		{
-			SCOPED_FUTEX(process.m_writtenFilesLock, writtenLock);
-			auto findIt = process.m_writtenFiles.find(msg.fromName.data);
-			if (findIt != process.m_writtenFiles.end())
+			auto& fs = process.m_shared;
+			SCOPED_FUTEX(fs.writtenFilesLock, writtenLock);
+			auto findIt = fs.writtenFiles.find(msg.fromKey);
+			if (findIt != fs.writtenFiles.end())
 			{
 				auto& oldFile = findIt->second;
 				bool isMapping = oldFile.mappingHandle.IsValid();
@@ -2714,12 +2953,12 @@ namespace uba
 					isMoved = true;
 				}
 
-				auto insres = process.m_writtenFiles.try_emplace(msg.toName.data);
+				auto insres = fs.writtenFiles.try_emplace(msg.toKey);
 				UBA_ASSERT(insres.second);
 				WrittenFile& newFile = insres.first->second;
 				newFile = oldFile;
 				newFile.name = msg.toName.data;
-				process.m_writtenFiles.erase(findIt);
+				fs.writtenFiles.erase(findIt);
 
 				if (isMapping)
 				{
@@ -2828,208 +3067,6 @@ namespace uba
 		if (!fileCount)
 			return true;
 
-		auto writeFile = [&](WrittenFile& file)
-		{
-			bool shouldEvictFromMemory = IsRarelyReadAfterWritten(process, file.name) || file.mappingWritten > m_keepOutputFileMemoryMapsThreshold;
-
-			if (ShouldWriteToDisk(file.name))
-			{
-				auto& rules = *process.m_startInfo.rules;
-				RootsHandle rootsHandle = process.GetStartInfo().rootsHandle;
-
-				bool storeCompressed = m_storeIntermediateFilesCompressed && g_globalRules.FileCanBeCompressed(file.name);
-				bool shouldDevirtualize = false;
-				bool escapeSpaces = false;
-				if (!storeCompressed)
-					shouldDevirtualize = HasVfs(rootsHandle) && rules.ShouldDevirtualizeFile(file.name, escapeSpaces);
-
-#if UBA_ENABLE_ON_DISK_FILE_MAPPINGS
-				if (!storeCompressed && !shouldDevirtualize && file.mappingHandle.fh)
-				{
-					bool success = true;
-					SetFilePointerEx(file.mappingHandle.fh, ToLargeInteger(file.mappingWritten), NULL, FILE_BEGIN);
-					::SetEndOfFile(file.mappingHandle.fh);
-
-					if (false)
-					{
-						TrackWorkScope work(m_workManager, AsView(TC("Flush")));
-						work.AddHint(file.name);
-						FlushFileBuffers(file.mappingHandle.fh);
-					}
-
-					if (true)
-					{
-						FILE_DISPOSITION_INFO info;
-						info.DeleteFile = false;
-						if (!::SetFileInformationByHandle(file.mappingHandle.fh, FileDispositionInfo, &info, sizeof(info)))
-							m_logger.Warning(TC("Failed to remove delete-on-close"));
-					}
-
-					if (false)
-					{
-						StringBuffer<> name2;
-						name2.Append(TC("\\??\\")).Append(file.name);
-
-						u8 buffer[2048];
-						auto& info = *(FILE_RENAME_INFO*)buffer;;
-						info.ReplaceIfExists = true;
-						info.RootDirectory = 0;
-						info.FileNameLength = u32(name2.count*sizeof(tchar));
-						info.Flags = FILE_RENAME_FLAG_POSIX_SEMANTICS;
-						memcpy(info.FileName, name2.data, info.FileNameLength+sizeof(tchar));
-						TrackWorkScope work(m_workManager, AsView(TC("Rename")));
-						work.AddHint(file.name);
-						success = SetFileInformationByHandle(file.mappingHandle.fh, FileRenameInfo, &info, sizeof(buffer));
-					}
-
-					if (success)
-					{
-						{
-							TrackWorkScope work(m_workManager, AsView(TC("CloseFileMapping")));
-							work.AddHint(file.name);
-							CloseFileMapping(m_logger, file.mappingHandle, file.name.c_str());
-						}
-						file.mappingHandle = {};
-						m_storage.InvalidateCachedFileInfo(file.key);
-						return true;
-					}
-					m_logger.Warning(TC("SetFileInformationByHandle failed %s (%s)"), file.name.c_str(), LastErrorToText().data);
-				}
-#endif
-
-				u64 fileSize = file.mappingWritten;
-				u8* mem = MapViewOfFile(m_logger, file.mappingHandle, FILE_MAP_READ, 0, fileSize);
-				if (!mem)
-					return m_logger.Error(TC("Failed to map view of filehandle for read %s (%s)"), file.name.c_str(), LastErrorToText().data);
-
-				//PrefetchVirtualMemory(mem, fileSize);
-
-				auto memClose = MakeGuard([&](){ UnmapViewOfFile(m_logger, mem, fileSize, file.name.c_str()); });
-
-				if (storeCompressed)
-				{
-					Storage::WriteResult res;
-					CompressedFileHeader header { CalculateCasKey(mem, fileSize, true, &m_workManager, file.name.c_str()) };
-
-					if (!m_storage.WriteCompressed(res, TC("MemoryMap"), InvalidFileHandle, mem, fileSize, file.name.c_str(), &header, sizeof(header), file.lastWriteTime))
-						return false;
-
-					shouldEvictFromMemory = false; // Can't evict without properly update filemappingtable.. the file on disk does now not match what was registered for write
-				}
-				else
-				{
-					FileAccessor destinationFile(m_logger, file.name.c_str());
-
-					if (shouldDevirtualize)
-					{
-						// Need to turn paths back into local paths
-
-						if (!destinationFile.CreateWrite(false, DefaultAttributes(), 0, m_tempPath.data))
-							return false;
-
-						MemoryBlock block(5*1024*1024);
-						if (!DevirtualizeDepsFile(rootsHandle, block, mem, fileSize, escapeSpaces, file.name.c_str()))
-							return false;
-
-						if (!destinationFile.Write(block.memory, block.writtenSize))
-							return false;
-					}
-					else
-					{
-						// Seems like best combo (for windows at least) is to use writes with overlap and max 16 at the same time.
-						// On one machine we get twice as fast without overlap if no bottleneck. On another machine (ntfs compression on) we get twice as slow without overlap
-						// Both machines behaves well with overlap AND bottleneck. Both machine are 128 logical core thread rippers.
-						bool useFileMapForWrite = fileSize && fileSize <= m_writeFilesFileMapMax; // ::CreateFileMappingW does not work for zero-length files
-						bool useOverlap = !useFileMapForWrite && fileSize >= m_writeFilesNoBufferingMin;//fileSize > 2*1024*1024;
-
-						u32 attributes = DefaultAttributes();
-						if (useOverlap)
-							attributes |= (FILE_FLAG_OVERLAPPED | FILE_FLAG_NO_BUFFERING);
-
-						if (useFileMapForWrite)
-						{
-							if (!destinationFile.CreateMemoryWrite(false, attributes, fileSize, m_tempPath.data))
-								return false;
-							MapMemoryCopy(destinationFile.GetData(), mem, fileSize);
-						}
-						else
-						{
-							if (!destinationFile.CreateWrite(false, attributes, fileSize, m_tempPath.data))
-								return false;
-							if (!destinationFile.Write(mem, fileSize, 0, true))
-								return false;
-						}
-					}
-
-					if (u64 time = file.lastWriteTime)
-						if (!SetFileLastWriteTime(destinationFile.GetHandle(), time))
-							return m_logger.Error(TC("Failed to set file time on filehandle for %s"), file.name.c_str());
-
-					if (!destinationFile.Close(file.lastWriteTime ? nullptr : &file.lastWriteTime))
-						return false;
-				}
-
-				// There are directory crawlers happening in parallel so we need to really make sure to invalidate this one since a crawler can actually
-				// hit this file with information from a query before it was written.. and then it will turn it back to "verified" using old info
-				m_storage.InvalidateCachedFileInfo(file.key);
-			}
-			else
-			{
-				// Delete existing file to make sure it is not picked up (since it is out of date)
-				uba::DeleteFileW(file.name.c_str());
-			}
-
-			if (shouldEvictFromMemory)
-			{
-				m_workManager.AddWork([mh = file.mappingHandle, n = file.name, this](const WorkContext&)
-					{
-						CloseFileMapping(m_logger, mh, n.c_str());
-					}, 1, TC("CloseFileMapping"));
-			}
-			else
-			{
-				StringBuffer<> name;
-				Storage::GetMappingString(name, file.mappingHandle, 0);
-				SCOPED_FUTEX(m_fileMappingTableLookupLock, lookupLock);
-				auto insres = m_fileMappingTableLookup.try_emplace(file.key);
-				if (!insres.second)
-					m_logger.Error(TC("Trying to write the same file twice (%s)"), file.name.c_str());
-				FileMappingEntry& entry = insres.first->second;
-				lookupLock.Leave();
-				SCOPED_FUTEX(entry.lock, entryCs);
-				UBA_ASSERT(file.mappingHandle.IsValid());
-				UBA_ASSERT(!entry.mapping.IsValid());
-				entry.handled = true;
-				entry.mapping = file.mappingHandle;
-				entry.mappingOffset = 0;
-				entry.size = file.mappingWritten;
-				entry.lastWriteTime = file.lastWriteTime;
-				entry.isDir = false;
-				entry.success = true;
-				entry.canBeFreed = true;
-				entry.usedCount = 0;
-				entry.usedCountBeforeFree = g_globalRules.GetUsedCountBeforeFree(file.name);
-
-				#if UBA_DEBUG_TRACK_MAPPING
-				entry.name = file.name;
-				m_debugLogger->Info(TC("Mapping kept 0x%llx (%s) from detoured process (UsedCountBeforeFree: %u)\n"), u64(entry.mapping.mh), entry.name.c_str(), u32(entry.usedCountBeforeFree));
-				#endif
-
-				SCOPED_WRITE_LOCK(m_fileMappingTableMemLock, lock);
-				BinaryWriter writer(m_fileMappingTableMem, m_fileMappingTableSize);
-				writer.WriteStringKey(file.key);
-				writer.WriteString(name);
-				writer.Write7BitEncoded(file.mappingWritten);
-				u32 newSize = (u32)writer.GetPosition();
-				m_fileMappingTableSize = (u32)newSize;
-			}
-			file.mappingHandle = {};
-
-			TraceWrittenFile(process.m_id, file.name, file.mappingWritten);
-
-			return true;
-		};
-
 		// This is to not kill I/O when writing lots of pdb/dlls in parallel
 		#if PLATFORM_WINDOWS
 		BottleneckScope scope(m_writeFilesBottleneck, Stats().waitBottleneck);
@@ -3042,7 +3079,7 @@ namespace uba
 				KernelStatsScope ks(process.m_kernelStats);
 				StorageStatsScope ss(process.m_storageStats);
 				SessionStatsScope sessionStatsScope(process.m_sessionStats);
-				if (!writeFile(**it))
+				if (!WriteFileToDisk(process, **it))
 					success = false;
 			}, TCV("WriteFilesToDisk"));
 		return success;
@@ -3157,7 +3194,7 @@ namespace uba
 		UBA_ASSERT(fileNameKey != StringKeyZero);
 
 		#if UBA_DEBUG_TRACK_MAPPING
-		//m_debugLogger->Info(TC("Mapping check (%s)\n"), fileName.data);
+		//m_debugLogger->Info(TC("Mapping check (%s)"), fileName.data);
 		#endif
 
 		MemoryMap out;
@@ -3264,7 +3301,7 @@ namespace uba
 
 		#if UBA_DEBUG_TRACK_MAPPING
 		entry.name = filePath.data;
-		m_debugLogger->Info(TC("Mapping created 0x%llx (%s) from virtual file\n"), u64(entry.mapping.mh), entry.name.c_str());
+		m_debugLogger->Info(TC("Mapping created 0x%llx (%s) from virtual file"), u64(entry.mapping.mh), entry.name.c_str());
 		#endif
 
 		SCOPED_WRITE_LOCK(m_fileMappingTableMemLock, lock);
@@ -3768,7 +3805,7 @@ namespace uba
 
 		mg.Cancel();
 
-		auto insres = msg.process.m_writtenFiles.try_emplace(exportsFile.data);
+		auto insres = msg.process.m_shared.writtenFiles.try_emplace(symFileKey);
 		WrittenFile& writtenFile = insres.first->second;
 
 		UBA_ASSERT(writtenFile.owner == nullptr || writtenFile.owner == &msg.process);
@@ -3778,7 +3815,7 @@ namespace uba
 		writtenFile.mappingHandle = symHandle;
 		writtenFile.mappingWritten = memoryBlock.writtenSize;
 		writtenFile.lastWriteTime = lastWriteTime;
-		writtenFile.name = insres.first->first;
+		writtenFile.name = exportsFile.data;
 
 		return true;
 	}
diff --git a/Engine/Source/Programs/UnrealBuildAccelerator/Common/Private/UbaSessionClient.cpp b/Engine/Source/Programs/UnrealBuildAccelerator/Common/Private/UbaSessionClient.cpp
--- a/Engine/Source/Programs/UnrealBuildAccelerator/Common/Private/UbaSessionClient.cpp
+++ b/Engine/Source/Programs/UnrealBuildAccelerator/Common/Private/UbaSessionClient.cpp
@@ -665,63 +665,62 @@ namespace uba
 	bool SessionClient::SendFiles(ProcessImpl& process, Timer& sendFiles)
 	{
 		StorageStatsScope storageStatsScope(process.m_storageStats);
-		for (auto& pair : process.m_writtenFiles)
+		for (auto& pair : process.m_shared.writtenFiles)
 		{
+			WrittenFile& file = pair.second;
 			TimerScope timer(sendFiles);
 #ifdef _DEBUG
 			if (!pair.second.mappingHandle.IsValid())
 				m_logger.Warning(TC("%s is not using file mapping"), pair.first.c_str());
 #endif
-			bool keepMappingInMemory = IsWindows && !IsRarelyReadAfterWritten(process, pair.first);
-			bool compressed = process.m_startInfo.rules->SendFileCompressedFromClient(pair.first);
-			if (!SendFile(pair.second, pair.first.c_str(), process.GetId(), keepMappingInMemory, compressed))
+			bool keepMappingInMemory = IsWindows && !IsRarelyReadAfterWritten(process, file.name);
+			bool compressed = process.m_startInfo.rules->SendFileCompressedFromClient(file.name);
+			if (!SendFile(file, process.GetId(), keepMappingInMemory, compressed))
 				return false;
 		}
 		return true;
 	}
 
-	bool SessionClient::SendFile(WrittenFile& source, const tchar* destination, u32 processId, bool keepMappingInMemory, bool compressed)
+	bool SessionClient::SendFile(WrittenFile& file, u32 processId, bool keepMappingInMemory, bool compressed)
 	{
 		CasKey casKey;
 		{
 			TimerScope ts(m_stats.storageSend);
-			if (!m_storage.StoreCasFileClient(casKey, source.key, source.name.c_str(), source.mappingHandle, 0, source.mappingWritten, destination, keepMappingInMemory, compressed))
+			if (!m_storage.StoreCasFileClient(casKey, file.key, file.backedName.c_str(), file.mappingHandle, 0, file.mappingWritten, file.name.c_str(), keepMappingInMemory, compressed))
 				return false;
 		}
 		if (casKey == CasKeyZero)
-			return m_logger.Error(TC("Failed to store cas on server for local file %s (size %llu destination %s)"), source.name.c_str(), source.mappingWritten, destination);
+			return m_logger.Error(TC("Failed to store cas on server for local file %s (size %llu destination %s)"), file.backedName.c_str(), file.mappingWritten, file.name.c_str());
 	
-		CloseFileMapping(m_logger, source.mappingHandle, source.name.c_str());
-		source.mappingHandle = {};
+		CloseFileMapping(m_logger, file.mappingHandle, file.backedName.c_str());
+		file.mappingHandle = {};
 
 		StackBinaryReader<128> reader;
 		{
 			StackBinaryWriter<1024> writer;
 			NetworkMessage msg(m_client, ServiceId, SessionMessageType_SendFileToServer, writer);
 			writer.WriteU32(processId);
-			writer.WriteString(destination);
-			writer.WriteStringKey(source.key);
-			writer.WriteU32(source.attributes);
+			writer.WriteString(file.name);
+			writer.WriteStringKey(file.key);
+			writer.WriteU32(file.attributes);
 			writer.WriteCasKey(casKey);
 			if (!msg.Send(reader, Stats().sendFileMsg))
-				return m_logger.Error(TC("Failed to send file %s to server"), source.name.c_str());
+				return m_logger.Error(TC("Failed to send file %s to server"), file.backedName.c_str());
 		}
 
 		#if UBA_DEBUG_LOGGER
-		m_debugLogger->Info(TC("SENDFILE %s\n"), destination);
+		m_debugLogger->Info(TC("SENDFILE %s\n"), file.name.c_str());
 		#endif
 
 		if (!reader.ReadBool())
-			return m_logger.Error(TC("Server failed to copy cas %s to %s (local source %s)"), CasKeyString(casKey).str, destination, source.name.c_str());
+			return m_logger.Error(TC("Server failed to copy cas %s to %s (local source %s)"), CasKeyString(casKey).str, file.name.c_str(), file.backedName.c_str());
 		return true;
 	}
 
-	void RemoveWrittenFile(ProcessImpl& process, const TString& name);
-
 	bool SessionClient::DeleteFile(DeleteFileResponse& out, const DeleteFileMessage& msg)
 	{
 		// TODO: Deleting output files should also delete them on disk (for now they will leak until process shutdown)
-		RemoveWrittenFile(msg.process, msg.fileName.data);
+		RemoveWrittenFile(msg.process, msg.fileNameKey);
 
 		bool sendDelete = true;
 		if (msg.closeId != 0)
@@ -822,15 +821,16 @@ namespace uba
 		auto& process = msg.process;
 
 		{
-			SCOPED_FUTEX(process.m_writtenFilesLock, lock);
-			auto& writtenFiles = process.m_writtenFiles;
-			auto findIt = writtenFiles.find(fromName);
+			SCOPED_FUTEX(process.m_shared.writtenFilesLock, lock);
+			auto& writtenFiles = process.m_shared.writtenFiles;
+			auto findIt = writtenFiles.find(msg.fromKey);
 			if (findIt != writtenFiles.end())
 			{
-				auto insres = writtenFiles.try_emplace(toName);
+				auto insres = writtenFiles.try_emplace(msg.toKey);
 				UBA_ASSERTF(insres.second, TC("Moving written file %s to other written file %s. (%s)"), fromName, toName, process.m_startInfo.description);
 				insres.first->second = findIt->second;
 				insres.first->second.key = msg.toKey;
+				insres.first->second.name = toName;
 				insres.first->second.owner = &process;
 				writtenFiles.erase(findIt);
 			}
@@ -871,12 +871,10 @@ namespace uba
 
 	bool SessionClient::Chmod(ChmodResponse& out, const ChmodMessage& msg)
 	{
-		const tchar* fromName = msg.fileName.data;
-
 		{
-			SCOPED_FUTEX(msg.process.m_writtenFilesLock, lock);
-			auto& writtenFiles = msg.process.m_writtenFiles;
-			auto findIt = writtenFiles.find(fromName);
+			SCOPED_FUTEX(msg.process.m_shared.writtenFilesLock, lock);
+			auto& writtenFiles = msg.process.m_shared.writtenFiles;
+			auto findIt = writtenFiles.find(msg.fileNameKey);
 			if (findIt != writtenFiles.end())
 			{
 				bool executable = false;
@@ -1642,14 +1640,15 @@ namespace uba
 		if (!logFile || !*logFile)
 			return;
 		WrittenFile f;
-		f.name = logFile;
+		f.backedName = logFile;
 		f.attributes = DefaultAttributes();
 		StringBuffer<> dest;
 		if (const tchar* lastSlash = TStrrchr(logFile, PathSeparator))
 			logFile = lastSlash + 1;
 		dest.Append(TC("<log>")).Append(logFile);
+		f.name = dest.data;
 		f.key = ToStringKeyLower(dest);
-		SendFile(f, dest.data, pi.GetId(), false, true);
+		SendFile(f, pi.GetId(), false, true);
 		for (auto& child : pi.m_childProcesses)
 			SendLogFileToServer(*(ProcessImpl*)child.m_process);
 	}
@@ -2070,12 +2069,13 @@ namespace uba
 			if (StopTrace(ubaFile.data))
 			{
 				WrittenFile f;
-				f.name = ubaFile.data;
+				f.backedName = ubaFile.data;
 				f.attributes = DefaultAttributes();
 				StringBuffer<> dest(TC("<uba>"));
+				f.name = dest.data;
 				f.key = ToStringKeyLower(dest);
 
-				SendFile(f, dest.data, 0, false, true);
+				SendFile(f, 0, false, true);
 			}
 		}
 	}
@@ -2330,14 +2330,14 @@ namespace uba
 
 	bool SessionClient::FlushWrittenFiles(ProcessImpl& process)
 	{
-		SCOPED_FUTEX(process.m_writtenFilesLock, lock);
+		SCOPED_FUTEX(process.m_shared.writtenFilesLock, lock);
 		bool success = SendFiles(process, process.m_processStats.sendFiles);
 		{
 			SCOPED_FUTEX(m_outputFilesLock, lock2);
-			for (auto& kv : process.m_writtenFiles)
-				m_outputFiles.erase(kv.first);
+			for (auto& kv : process.m_shared.writtenFiles)
+				m_outputFiles.erase(kv.second.name);
 		}
-		process.m_writtenFiles.clear();
+		process.m_shared.writtenFiles.clear();
 
 		return success;
 	}
diff --git a/Engine/Source/Programs/UnrealBuildAccelerator/Common/Public/UbaProcess.h b/Engine/Source/Programs/UnrealBuildAccelerator/Common/Public/UbaProcess.h
--- a/Engine/Source/Programs/UnrealBuildAccelerator/Common/Public/UbaProcess.h
+++ b/Engine/Source/Programs/UnrealBuildAccelerator/Common/Public/UbaProcess.h
@@ -98,6 +98,7 @@ namespace uba
 		void InternalLogLine(bool printInSession, TString&& line, LogEntryType logType);
 		void ClearTempFiles();
 		void WaitForParent();
+		void WaitForChildrenExit();
 		bool CancelWithError(bool dummy = true);
 
 		ProcessStartInfoHolder m_startInfo;
@@ -130,6 +131,14 @@ namespace uba
 		HANDLE m_accountingJobObject = NULL;
 		#endif
 
+		struct Shared
+		{
+			Futex writtenFilesLock;
+			UnorderedMap<StringKey, WrittenFile> writtenFiles;
+			ReaderWriterLock tempFilesLock;
+			UnorderedMap<StringKey, WrittenFile> tempFiles;
+		};
+
 		u32 m_nativeProcessId = 0;
 		u32 m_nativeProcessExitCode = ~0u;
 		u32 m_exitCode = ~0u;
@@ -151,10 +160,7 @@ namespace uba
 		Vector<u8> m_trackedInputs;
 		Vector<u8> m_trackedOutputs;
 		Vector<ProcessHandle> m_childProcesses;
-		Futex& m_writtenFilesLock;
-		UnorderedMap<TString, WrittenFile>& m_writtenFiles;
-		ReaderWriterLock& m_tempFilesLock;
-		UnorderedMap<StringKey, WrittenFile>& m_tempFiles;
+		Shared& m_shared;
 		Atomic<u32> m_tempFilesModified;
 		SessionStats m_sessionStats;
 		StorageStats m_storageStats;
diff --git a/Engine/Source/Programs/UnrealBuildAccelerator/Common/Public/UbaSession.h b/Engine/Source/Programs/UnrealBuildAccelerator/Common/Public/UbaSession.h
--- a/Engine/Source/Programs/UnrealBuildAccelerator/Common/Public/UbaSession.h
+++ b/Engine/Source/Programs/UnrealBuildAccelerator/Common/Public/UbaSession.h
@@ -176,6 +176,8 @@ namespace uba
 
 		static constexpr CasKey CasKeyIsDirectory = { ~u64(0), ~u64(0), ~u32(0) };
 
+		void RemoveWrittenFile(ProcessImpl& process, const StringKey& fileKey);
+		bool WriteFileToDisk(ProcessImpl& process, WrittenFile& file);
 		void AddEnvironmentVariableNoLock(const tchar* key, const tchar* value);
 		bool WriteDirectoryEntriesInternal(DirectoryTable::Directory& dir, const StringKey& dirKey, StringView dirPath, bool isRefresh, u32& outTableOffset);
 		void WriteDirectoryEntriesRecursive(const StringKey& dirKey, StringView dirPath, u32& outTableOffset);
@@ -550,6 +552,7 @@ namespace uba
 		ProcessImpl* owner = nullptr;
 		StringKey key;
 		TString name;
+		TString backedName;
 		FileMappingHandle mappingHandle;
 		u64 mappingWritten = 0;
 		FileMappingHandle originalMappingHandle;
diff --git a/Engine/Source/Programs/UnrealBuildAccelerator/Common/Public/UbaSessionClient.h b/Engine/Source/Programs/UnrealBuildAccelerator/Common/Public/UbaSessionClient.h
--- a/Engine/Source/Programs/UnrealBuildAccelerator/Common/Public/UbaSessionClient.h
+++ b/Engine/Source/Programs/UnrealBuildAccelerator/Common/Public/UbaSessionClient.h
@@ -87,7 +87,7 @@ namespace uba
 		bool EnsureBinaryFile(StringBufferBase& out, StringBufferBase& outVirtual, u32 processId, const StringBufferBase& fileName, const StringKey& fileNameKey, const tchar* applicationDir, const u8* loaderPaths, u32 loaderPathsSize);
 		bool WriteBinFile(StringBufferBase& out, const StringView& binaryName, const CasKey& casKey, const KeyToString& applicationDir, u32 fileAttributes);
 		bool SendFiles(ProcessImpl& process, Timer& sendFiles);
-		bool SendFile(WrittenFile& source, const tchar* destination, u32 processId, bool keepMappingInMemory, bool compressed);
+		bool SendFile(WrittenFile& file, u32 processId, bool keepMappingInMemory, bool compressed);
 		bool SendUpdateDirectoryTable(StackBinaryReader<SendMaxSize>& reader); // Note, reader is sent in to save stack space.
 		bool UpdateDirectoryTableFromServer(StackBinaryReader<SendMaxSize>& reader);
 		bool SendUpdateNameToHashTable(StackBinaryReader<SendMaxSize>& reader);
diff --git a/Engine/Source/Programs/UnrealBuildAccelerator/Core/Private/UbaApplicationRules.cpp b/Engine/Source/Programs/UnrealBuildAccelerator/Core/Private/UbaApplicationRules.cpp
--- a/Engine/Source/Programs/UnrealBuildAccelerator/Core/Private/UbaApplicationRules.cpp
+++ b/Engine/Source/Programs/UnrealBuildAccelerator/Core/Private/UbaApplicationRules.cpp
@@ -331,11 +331,13 @@ namespace uba
 				|| fileName.EndsWith(TC(".h.d"))
 				|| fileName.EndsWith(TC(".cc.d"))
 				|| fileName.EndsWith(TC(".cpp.d"))
-				|| fileName.EndsWith(TC(".o.tmp")) // Clang writes to tmp file and then move
-				|| fileName.EndsWith(TC(".obj.tmp")) // Clang (verse) writes to tmp file and then move
-				|| fileName.EndsWith(TC(".obj")) // Clang with -fno-temp-file write obj file directly
-				|| fileName.EndsWith(TC(".pch")) // Clang with -fno-temp-file write pch file directly
-				|| fileName.EndsWith(TC(".gch.tmp")) // 
+				|| fileName.EndsWith(TC(".o"))
+				|| fileName.EndsWith(TC(".o.tmp"))
+				|| fileName.EndsWith(TC(".gch"))
+				|| fileName.EndsWith(TC(".gch.tmp"))
+				|| fileName.EndsWith(TC(".obj"))
+				|| fileName.EndsWith(TC(".obj.tmp"))
+				|| fileName.EndsWith(TC(".pch"))
 				;
 		}
 
diff --git a/Engine/Source/Programs/UnrealBuildAccelerator/Core/Public/UbaProcessUtils.h b/Engine/Source/Programs/UnrealBuildAccelerator/Core/Public/UbaProcessUtils.h
--- a/Engine/Source/Programs/UnrealBuildAccelerator/Core/Public/UbaProcessUtils.h
+++ b/Engine/Source/Programs/UnrealBuildAccelerator/Core/Public/UbaProcessUtils.h
@@ -7,8 +7,8 @@
 #if PLATFORM_WINDOWS
 #include <malloc.h>
 #else
-#define _malloca alloca
-#define _freea
+#define _malloca(x) alloca(x)
+#define _freea(x)
 #endif
 
 namespace uba
diff --git a/Engine/Source/Programs/UnrealBuildAccelerator/Core/Public/UbaProtocol.h b/Engine/Source/Programs/UnrealBuildAccelerator/Core/Public/UbaProtocol.h
--- a/Engine/Source/Programs/UnrealBuildAccelerator/Core/Public/UbaProtocol.h
+++ b/Engine/Source/Programs/UnrealBuildAccelerator/Core/Public/UbaProtocol.h
@@ -20,7 +20,7 @@ namespace uba
 		UBA_PROCESS_MESSAGE(RemoveDirectory) \
 		UBA_PROCESS_MESSAGE(ListDirectory) \
 		UBA_PROCESS_MESSAGE(UpdateTables) \
-		UBA_PROCESS_MESSAGE(GetParentWrittenFiles) \
+		UBA_PROCESS_MESSAGE(GetWrittenFiles) \
 		UBA_PROCESS_MESSAGE(CreateProcess) \
 		UBA_PROCESS_MESSAGE(StartProcess) \
 		UBA_PROCESS_MESSAGE(ExitChildProcess) \
@@ -50,7 +50,7 @@ namespace uba
 		#undef UBA_PROCESS_MESSAGE
 	};
 
-	inline constexpr u32 ProcessMessageVersion = 1341;
+	inline constexpr u32 ProcessMessageVersion = 1342;
 
 	inline constexpr u32 CommunicationMemSize = IsWindows ? 64*1024 : 64*1024*2; // Macos expands some commandlines to be crazy long
 
diff --git a/Engine/Source/Programs/UnrealBuildAccelerator/Detours/Private/PosixOS/UbaDetoursFunctionsPosix.cpp b/Engine/Source/Programs/UnrealBuildAccelerator/Detours/Private/PosixOS/UbaDetoursFunctionsPosix.cpp
--- a/Engine/Source/Programs/UnrealBuildAccelerator/Detours/Private/PosixOS/UbaDetoursFunctionsPosix.cpp
+++ b/Engine/Source/Programs/UnrealBuildAccelerator/Detours/Private/PosixOS/UbaDetoursFunctionsPosix.cpp
@@ -2640,8 +2640,8 @@ namespace uba
 		UBA_ASSERTF(dirTableMem != MAP_FAILED, "mmap for dirtable mem failed (%s)", strerror(errno));
 		g_directoryTable.Init(dirTableMem, directoryTableCount, directoryTableSize);
 
-		if (g_runningRemote && g_isChild)
-			Rpc_GetParentWrittenFiles();
+		if (g_isChild)
+			Rpc_GetWrittenFiles();
 
 		//pthread_atfork([]()
 		//{
diff --git a/Engine/Source/Programs/UnrealBuildAccelerator/Detours/Private/UbaDetoursFileMappingTable.cpp b/Engine/Source/Programs/UnrealBuildAccelerator/Detours/Private/UbaDetoursFileMappingTable.cpp
--- a/Engine/Source/Programs/UnrealBuildAccelerator/Detours/Private/UbaDetoursFileMappingTable.cpp
+++ b/Engine/Source/Programs/UnrealBuildAccelerator/Detours/Private/UbaDetoursFileMappingTable.cpp
@@ -174,6 +174,33 @@ namespace uba
 		g_directoryTable.ParseDirectoryTable(directoryTableSize);
 	}
 
+	void UpdateWrittenFiles(BinaryReader& reader)
+	{
+		u32 count = reader.ReadU32();
+		while (count--)
+		{
+			StringKey key = reader.ReadStringKey();
+			auto insres = g_mappedFileTable.m_lookup.try_emplace(key);
+			FileInfo& info = insres.first->second;
+			StringBuffer<> name;
+			reader.ReadString(name);
+			FileMappingHandle mappingHandle = FileMappingHandle::FromU64(reader.ReadU64());
+			u64 fileSize = reader.ReadU64();
+			if (!info.originalName)
+				info.originalName = g_mappedFileTable.m_memoryBlock.Strdup(name).data;
+			info.name = info.originalName;
+			info.size = fileSize;
+
+			if (mappingHandle.IsValid())
+			{
+				name.Clear().Append(':').AppendHex(mappingHandle.ToU64()).Append('-').AppendHex(0);
+				info.name = g_mappedFileTable.m_memoryBlock.Strdup(name).data;
+			}
+
+			DEBUG_LOG(TC("GOT WRITTEN FILE: %s (Mapping: %s Size: %llu)"), info.originalName, info.name, info.size);
+		}
+	}
+
 	void Rpc_UpdateTables()
 	{
 		u32 directoryTableSize;
@@ -211,33 +238,24 @@ namespace uba
 				}
 			}
 #endif
+			UpdateWrittenFiles(reader);
+
 			DEBUG_LOG_PIPE(L"UpdateTables", L"");
 		}
 		g_directoryTable.ParseDirectoryTable(directoryTableSize);
 		g_mappedFileTable.Parse(fileMappingTableSize);
 	}
 
-	void Rpc_GetParentWrittenFiles()
+	void Rpc_GetWrittenFiles()
 	{
 		TimerScope ts(g_stats.updateTables);
 		SCOPED_WRITE_LOCK(g_communicationLock, pcs);
 		BinaryWriter writer;
-		writer.WriteByte(MessageType_GetParentWrittenFiles);
+		writer.WriteByte(MessageType_GetWrittenFiles);
 		writer.Flush();
 		BinaryReader reader;
 
-		u32 count = reader.ReadU32();
-		while (count--)
-		{
-			StringKey key = reader.ReadStringKey();
-			auto insres = g_mappedFileTable.m_lookup.try_emplace(key);
-			FileInfo& info = insres.first->second;
-			StringBuffer<> name;
-			reader.ReadString(name);
-			u64 fileSize = reader.ReadU64();
-			info.name = g_mappedFileTable.m_memoryBlock.Strdup(name).data;
-			info.size = fileSize;
-		}
+		UpdateWrittenFiles(reader);
 	}
 
 	u32 Rpc_GetEntryOffset(const StringKey& entryNameKey, const tchar* entryName, u64 entryNameLen, bool checkIfDir)
diff --git a/Engine/Source/Programs/UnrealBuildAccelerator/Detours/Private/UbaDetoursFileMappingTable.h b/Engine/Source/Programs/UnrealBuildAccelerator/Detours/Private/UbaDetoursFileMappingTable.h
--- a/Engine/Source/Programs/UnrealBuildAccelerator/Detours/Private/UbaDetoursFileMappingTable.h
+++ b/Engine/Source/Programs/UnrealBuildAccelerator/Detours/Private/UbaDetoursFileMappingTable.h
@@ -95,7 +95,7 @@ namespace uba
 	void Rpc_UpdateTables();
 	u32  Rpc_GetEntryOffset(const StringKey& entryNameKey, const tchar* entryName, u64 entryNameLen, bool checkIfDir = false);
 	void Rpc_GetFullFileName(const tchar*& path, u64& pathLen, StringBufferBase& tempBuf, bool useVirtualName, const tchar* const* loaderPaths = nullptr);
-	void Rpc_GetParentWrittenFiles();
+	void Rpc_GetWrittenFiles();
 
 	struct DirHash
 	{
diff --git a/Engine/Source/Programs/UnrealBuildAccelerator/Detours/Private/UbaDetoursShared.cpp b/Engine/Source/Programs/UnrealBuildAccelerator/Detours/Private/UbaDetoursShared.cpp
--- a/Engine/Source/Programs/UnrealBuildAccelerator/Detours/Private/UbaDetoursShared.cpp
+++ b/Engine/Source/Programs/UnrealBuildAccelerator/Detours/Private/UbaDetoursShared.cpp
@@ -529,9 +529,14 @@ namespace uba
 							auto findIt = g_mappedFileTable.m_lookup.find(fileNameKey);
 							if (findIt == g_mappedFileTable.m_lookup.end())
 							{
-								StringBuffer<> temp;
-								u32 closeId;
-								Rpc_CreateFileW(fileNameView, fileNameKey, AccessFlag_Read, temp.data, temp.capacity, fileSize, closeId, false);
+								// If file is output file we accept wrong size because size is not supposed to be used anyway.
+								// We don't want to trigger unnecessary download/decompress of file
+								if (!g_rules->IsOutputFile(fileNameView))
+								{
+									StringBuffer<> temp;
+									u32 closeId;
+									Rpc_CreateFileW(fileNameView, fileNameKey, AccessFlag_Read, temp.data, temp.capacity, fileSize, closeId, false);
+								}
 							}
 							else
 							{
diff --git a/Engine/Source/Programs/UnrealBuildAccelerator/Detours/Private/Windows/UbaDetoursFunctionsKernelBase.inl b/Engine/Source/Programs/UnrealBuildAccelerator/Detours/Private/Windows/UbaDetoursFunctionsKernelBase.inl
--- a/Engine/Source/Programs/UnrealBuildAccelerator/Detours/Private/Windows/UbaDetoursFunctionsKernelBase.inl
+++ b/Engine/Source/Programs/UnrealBuildAccelerator/Detours/Private/Windows/UbaDetoursFunctionsKernelBase.inl
@@ -534,8 +534,11 @@ BOOL Detoured_GetVolumeInformationW(LPCWSTR lpRootPathName, LPWSTR lpVolumeNameB
 		if (lpVolumeSerialNumber)
 			*lpVolumeSerialNumber = lpRootPathName[0]; // Let's see if this works, LOL
 
+		if (lpMaximumComponentLength)
+			*lpMaximumComponentLength =  255; // TODO: Need to fix this
+
 		//UBA_ASSERT(!lpVolumeNameBuffer);
-		UBA_ASSERT(!lpMaximumComponentLength);
+		//UBA_ASSERT(!lpMaximumComponentLength);
 		UBA_ASSERT(!lpFileSystemFlags);
 
 		if (nFileSystemNameSize)
@@ -681,7 +684,7 @@ BOOL Detoured_WriteFile(HANDLE hFile, LPCVOID lpBuffer, DWORD nNumberOfBytesToWr
 
 			*lpNumberOfBytesWritten = nNumberOfBytesToWrite;
 			SetLastError(ERROR_SUCCESS);
-			DEBUG_LOG_DETOURED(L"WriteFile", L"(MEMORY) %llu (%ls) ToWrite: %u -> Success", uintptr_t(hFile), HandleToName(hFile), nNumberOfBytesToWrite);
+			DEBUG_LOG_DETOURED(L"WriteFile", L"(MEMORY)%s %llu (%ls) ToWrite: %u -> Success", (lpOverlapped ? L" OVERLAPPED" : L""), uintptr_t(hFile), HandleToName(hFile), nNumberOfBytesToWrite);
 			return TRUE;
 		}
 		UBA_ASSERTF(!fi.isFileMap, L"Trying to write to file %ls which is a filemap. This is not supported\n", HandleToName(hFile));
diff --git a/Engine/Source/Programs/UnrealBuildAccelerator/Detours/Private/Windows/UbaDetoursFunctionsNtDll.inl b/Engine/Source/Programs/UnrealBuildAccelerator/Detours/Private/Windows/UbaDetoursFunctionsNtDll.inl
--- a/Engine/Source/Programs/UnrealBuildAccelerator/Detours/Private/Windows/UbaDetoursFunctionsNtDll.inl
+++ b/Engine/Source/Programs/UnrealBuildAccelerator/Detours/Private/Windows/UbaDetoursFunctionsNtDll.inl
@@ -954,7 +954,9 @@ NTSTATUS NTAPI Shared_NtCreateFile(bool IsCreateFunc, PHANDLE hFileHandle, ACCES
 
 				bool isWriteAttributes = (DesiredAccess & FILE_WRITE_ATTRIBUTES) != 0;
 
-				if (allowEarlyOut && !useContent && !isWriteAttributes && !CouldBeCompressedFile(fileName))
+				// If file is an output file we still allllow this path and accept wrong (compressed) size.
+				// This is a bit hacky but we don't want to transfer and decompress file just to get the size
+				if (allowEarlyOut && !useContent && !isWriteAttributes && (!CouldBeCompressedFile(fileName) || g_rules->IsOutputFile(fileName)))
 				{
 					auto dh = new DetouredHandle(HandleType_File);
 					dh->fileObject = new FileObject();
@@ -1243,6 +1245,33 @@ NTSTATUS NTAPI Shared_NtCreateFile(bool IsCreateFunc, PHANDLE hFileHandle, ACCES
 			mf.mappingHandle.mh = newHandle;
 		};
 
+	if (lpFileName[0] == ':') // It is a HANDLE from session process. A written file that is writable
+	{
+		if (!info.memoryFile)
+		{
+			const wchar_t* handleStr = lpFileName + 1;
+			const wchar_t* handleStrEnd = wcschr(handleStr, '-');
+			if (!handleStrEnd)
+			{
+				UBA_ASSERT(handleStrEnd);
+				return STATUS_UNSUCCESSFUL;
+			}
+			auto mappingHandle = FileMappingHandle::FromU64(StringToValue(handleStr, handleStrEnd - handleStr));
+			const wchar_t* mappingOffsetStr = handleStrEnd + 1;
+			u64 mappingOffset = StringToValue(mappingOffsetStr, wcslen(mappingOffsetStr));
+			UBA_ASSERT(!mappingOffset);
+			UBA_ASSERT(info.size != InvalidValue);
+
+			TakeFileOwnershipFromSession(mappingHandle, info.size);
+		}
+
+		lookupLock.Leave();
+		*hFileHandle = CreateFileHandle();
+		TrackFileInput();
+		DEBUG_LOG_DETOURED(funcName, L"(WRITTENFILE)%ls %llu (%ls) (%ls) -> Success", isWriteStr, uintptr_t(*hFileHandle), lpFileName, (fileName.data != lpFileName ? fileName.data : L""));
+		return STATUS_SUCCESS;
+	}
+	
 	if (keepInMemory || info.memoryFile)
 	{
 		#if UBA_DEBUG_LOG_ENABLED
@@ -1285,36 +1314,10 @@ NTSTATUS NTAPI Shared_NtCreateFile(bool IsCreateFunc, PHANDLE hFileHandle, ACCES
 			}
 			else
 			{
-				bool isOutput = (isWrite || isDeleteOnClose) && g_rules->IsOutputFile(fileName);
 				bool isThrowAway = g_rules->IsThrowAway(fileName, g_runningRemote);
 
 				if (!isThrowAway && CreateDisposition == FILE_OPEN)
 				{
-					if (isOutput)
-					{
-						#if UBA_DEBUG_LOG_ENABLED
-						memoryType = L"TAKEOWNERSHIPWRITE";
-						#endif
-						TimerScope ts(g_stats.createFile);
-						SCOPED_WRITE_LOCK(g_communicationLock, pcs);
-						BinaryWriter writer;
-						writer.WriteByte(MessageType_TakeFileOwnership);
-						writer.WriteStringKey(fileNameKey);
-						writer.WriteString(fileName);
-						writer.Flush();
-						BinaryReader reader;
-						bool hasOwnership = reader.ReadBool();
-						FileMappingHandle mappingHandle = FileMappingHandle::FromU64(reader.ReadU64());
-						u64 mappingHandleSize = reader.ReadU64();
-						pcs.Leave();
-
-						if (!hasOwnership && mappingHandle.IsValid())
-						{
-							// This is a file that most likely was written by child process and is now used by parent (this) process
-							TakeFileOwnershipFromSession(mappingHandle, mappingHandleSize);
-						}
-					}
-
 					if (!info.memoryFile)
 					{
 						*hFileHandle = INVALID_HANDLE_VALUE;
@@ -1324,6 +1327,7 @@ NTSTATUS NTAPI Shared_NtCreateFile(bool IsCreateFunc, PHANDLE hFileHandle, ACCES
 				}
 				else
 				{
+					bool isOutput = (isWrite || isDeleteOnClose) && g_rules->IsOutputFile(fileName);
 					if (isOutput && (CreateDisposition & FILE_OPEN_IF) == 0)
 					{
 						UBA_ASSERTF(false, TC("Trying to open %s with openif. This is not supported"), fileName.data);
@@ -1361,31 +1365,6 @@ NTSTATUS NTAPI Shared_NtCreateFile(bool IsCreateFunc, PHANDLE hFileHandle, ACCES
 		DEBUG_LOG_DETOURED(funcName, L"(%s)%ls %llu (%ls) (%ls) -> Success", memoryType, isWriteStr, uintptr_t(*hFileHandle), lpFileName, (fileName.data != lpFileName ? fileName.data : L""));
 		return STATUS_SUCCESS;
 	}
-	else if (g_rules->IsOutputFile(fileName)) // It might be that a child process has written this file and then we need to try to take ownership of it in order to read it.
-	{
-		TimerScope ts(g_stats.createFile);
-		SCOPED_WRITE_LOCK(g_communicationLock, pcs);
-		BinaryWriter writer;
-		writer.WriteByte(MessageType_TakeFileOwnership);
-		writer.WriteStringKey(fileNameKey);
-		writer.WriteString(fileName);
-		writer.Flush();
-		BinaryReader reader;
-		bool hasOwnership = reader.ReadBool();
-		FileMappingHandle mappingHandle = FileMappingHandle::FromU64(reader.ReadU64());
-		u64 mappingHandleSize = reader.ReadU64();
-		pcs.Leave();
-
-		if (!hasOwnership && mappingHandle.IsValid())
-		{
-			TakeFileOwnershipFromSession(mappingHandle, mappingHandleSize);
-			lookupLock.Leave();
-			*hFileHandle = CreateFileHandle();
-			TrackFileInput();
-			DEBUG_LOG_DETOURED(funcName, L"(TAKEOWNERSHIPREAD)%ls %llu (%ls) (%ls) -> Success", isWriteStr, uintptr_t(*hFileHandle), lpFileName, (fileName.data != lpFileName ? fileName.data : L""));
-			return STATUS_SUCCESS;
-		}
-	}
 
 	lookupLock.Leave();
 
diff --git a/Engine/Source/Programs/UnrealBuildAccelerator/Detours/Private/Windows/UbaDetoursFunctionsWin.cpp b/Engine/Source/Programs/UnrealBuildAccelerator/Detours/Private/Windows/UbaDetoursFunctionsWin.cpp
--- a/Engine/Source/Programs/UnrealBuildAccelerator/Detours/Private/Windows/UbaDetoursFunctionsWin.cpp
+++ b/Engine/Source/Programs/UnrealBuildAccelerator/Detours/Private/Windows/UbaDetoursFunctionsWin.cpp
@@ -1109,8 +1109,8 @@ void Init(const DetoursPayload& payload, u64 startTime)
 		g_directoryTable.Init(directoryTableMem, directoryTableCount, directoryTableSize);
 	}
 
-	if (g_runningRemote && g_isChild)
-		Rpc_GetParentWrittenFiles();
+	if (g_isChild)
+		Rpc_GetWrittenFiles();
 
 	g_stats.attach.time += GetTime() - startTime;
 	g_stats.attach.count = 1;
