[UBA]
* Added handling of option /headerUnit:quote in dependency crawler

diff --git a/Engine/Source/Programs/UnrealBuildAccelerator/Common/Public/UbaDependencyCrawler.h b/Engine/Source/Programs/UnrealBuildAccelerator/Common/Public/UbaDependencyCrawler.h
--- a/Engine/Source/Programs/UnrealBuildAccelerator/Common/Public/UbaDependencyCrawler.h
+++ b/Engine/Source/Programs/UnrealBuildAccelerator/Common/Public/UbaDependencyCrawler.h
@@ -280,7 +280,8 @@ namespace uba
 				{ TCV("/imsvc"), true },
 				{ TCV("/experimental:log"), false },
 				{ TCV("/analyze:log"), false },
-				{ TCV("/sourceDependencies"), false},
+				{ TCV("/sourceDependencies"), false },
+				{ TCV("/headerUnit:quote"), false },
 			};
 
 			constexpr const tchar* clangClOptionsWithArg[] = 
