[UBA]
* Fixed missing assignment of file name key when getting written files
* Fixed shipping compile warning
* Added asserts to validate file name keys

diff --git a/Engine/Source/Programs/UnrealBuildAccelerator/Common/Private/UbaSession.cpp b/Engine/Source/Programs/UnrealBuildAccelerator/Common/Private/UbaSession.cpp
--- a/Engine/Source/Programs/UnrealBuildAccelerator/Common/Private/UbaSession.cpp
+++ b/Engine/Source/Programs/UnrealBuildAccelerator/Common/Private/UbaSession.cpp
@@ -2638,6 +2638,7 @@ namespace uba
 			out.fileName.Append(fileName);
 		}
 
+		UBA_ASSERT(fileNameKey != StringKeyZero);
 		SCOPED_FUTEX(m_activeFilesLock, lock);
 		u32 wantsOnCloseId = m_wantsOnCloseIdCounter++;
 		out.closeId = wantsOnCloseId;
@@ -2765,7 +2766,7 @@ namespace uba
 				if (!m_runningRemote)
 					msgName = msg.newName;
 			}
-
+			UBA_ASSERT(key != StringKeyZero);
 			SCOPED_FUTEX(msg.process.m_shared.writtenFilesLock, writtenLock);
 			auto insres = msg.process.m_shared.writtenFiles.try_emplace(key);
 			WrittenFile& writtenFile = insres.first->second;
@@ -2920,6 +2921,7 @@ namespace uba
 		out.fromName.Append(msg.fromName);
 		out.toName.Append(msg.toName);
 
+		UBA_ASSERT(msg.toKey != StringKeyZero);
 		SCOPED_FUTEX(m_activeFilesLock, lock);
 		u32 closeId = m_wantsOnCloseIdCounter++;
 		if (!m_activeFiles.try_emplace(closeId, ActiveFile{ msg.toName.data, msg.toKey }).second)
@@ -2953,6 +2955,7 @@ namespace uba
 					isMoved = true;
 				}
 
+				UBA_ASSERT(msg.toKey != StringKeyZero);
 				auto insres = fs.writtenFiles.try_emplace(msg.toKey);
 				UBA_ASSERT(insres.second);
 				WrittenFile& newFile = insres.first->second;
diff --git a/Engine/Source/Programs/UnrealBuildAccelerator/Common/Private/UbaSessionClient.cpp b/Engine/Source/Programs/UnrealBuildAccelerator/Common/Private/UbaSessionClient.cpp
--- a/Engine/Source/Programs/UnrealBuildAccelerator/Common/Private/UbaSessionClient.cpp
+++ b/Engine/Source/Programs/UnrealBuildAccelerator/Common/Private/UbaSessionClient.cpp
@@ -826,6 +826,7 @@ namespace uba
 			auto findIt = writtenFiles.find(msg.fromKey);
 			if (findIt != writtenFiles.end())
 			{
+				UBA_ASSERT(msg.toKey != StringKeyZero);
 				auto insres = writtenFiles.try_emplace(msg.toKey);
 				UBA_ASSERTF(insres.second, TC("Moving written file %s to other written file %s. (%s)"), fromName, toName, process.m_startInfo.description);
 				insres.first->second = findIt->second;
diff --git a/Engine/Source/Programs/UnrealBuildAccelerator/Detours/Private/UbaDetoursFileMappingTable.cpp b/Engine/Source/Programs/UnrealBuildAccelerator/Detours/Private/UbaDetoursFileMappingTable.cpp
--- a/Engine/Source/Programs/UnrealBuildAccelerator/Detours/Private/UbaDetoursFileMappingTable.cpp
+++ b/Engine/Source/Programs/UnrealBuildAccelerator/Detours/Private/UbaDetoursFileMappingTable.cpp
@@ -186,6 +186,7 @@ namespace uba
 			reader.ReadString(name);
 			FileMappingHandle mappingHandle = FileMappingHandle::FromU64(reader.ReadU64());
 			u64 fileSize = reader.ReadU64();
+			info.fileNameKey = key;
 			if (!info.originalName)
 				info.originalName = g_mappedFileTable.m_memoryBlock.Strdup(name).data;
 			info.name = info.originalName;
diff --git a/Engine/Source/Programs/UnrealBuildAccelerator/Detours/Private/Windows/UbaDetoursFunctionsNtDll.inl b/Engine/Source/Programs/UnrealBuildAccelerator/Detours/Private/Windows/UbaDetoursFunctionsNtDll.inl
--- a/Engine/Source/Programs/UnrealBuildAccelerator/Detours/Private/Windows/UbaDetoursFunctionsNtDll.inl
+++ b/Engine/Source/Programs/UnrealBuildAccelerator/Detours/Private/Windows/UbaDetoursFunctionsNtDll.inl
@@ -1258,8 +1258,7 @@ NTSTATUS NTAPI Shared_NtCreateFile(bool IsCreateFunc, PHANDLE hFileHandle, ACCES
 			}
 			auto mappingHandle = FileMappingHandle::FromU64(StringToValue(handleStr, handleStrEnd - handleStr));
 			const wchar_t* mappingOffsetStr = handleStrEnd + 1;
-			u64 mappingOffset = StringToValue(mappingOffsetStr, wcslen(mappingOffsetStr));
-			UBA_ASSERT(!mappingOffset);
+			UBA_ASSERT(!StringToValue(mappingOffsetStr, wcslen(mappingOffsetStr)));
 			UBA_ASSERT(info.size != InvalidValue);
 
 			TakeFileOwnershipFromSession(mappingHandle, info.size);
