[UBA]
* Fixed compile warning

diff --git a/Engine/Source/Programs/UnrealBuildAccelerator/Detours/Private/Windows/UbaDetoursFunctionsNtDll.inl b/Engine/Source/Programs/UnrealBuildAccelerator/Detours/Private/Windows/UbaDetoursFunctionsNtDll.inl
--- a/Engine/Source/Programs/UnrealBuildAccelerator/Detours/Private/Windows/UbaDetoursFunctionsNtDll.inl
+++ b/Engine/Source/Programs/UnrealBuildAccelerator/Detours/Private/Windows/UbaDetoursFunctionsNtDll.inl
@@ -1257,8 +1257,7 @@ NTSTATUS NTAPI Shared_NtCreateFile(bool IsCreateFunc, PHANDLE hFileHandle, ACCES
 				return STATUS_UNSUCCESSFUL;
 			}
 			auto mappingHandle = FileMappingHandle::FromU64(StringToValue(handleStr, handleStrEnd - handleStr));
-			const wchar_t* mappingOffsetStr = handleStrEnd + 1;
-			UBA_ASSERT(!StringToValue(mappingOffsetStr, wcslen(mappingOffsetStr)));
+			UBA_ASSERT(!StringToValue(handleStrEnd + 1, wcslen(handleStrEnd + 1)));
 			UBA_ASSERT(info.size != InvalidValue);
 
 			TakeFileOwnershipFromSession(mappingHandle, info.size);
