[UBA]
* Attempt at fixing VirtualAlloc failures on farm when building some clang targets

diff --git a/Engine/Source/Programs/UnrealBuildAccelerator/Detours/Private/Windows/UbaDetoursFunctionsNtDll.inl b/Engine/Source/Programs/UnrealBuildAccelerator/Detours/Private/Windows/UbaDetoursFunctionsNtDll.inl
--- a/Engine/Source/Programs/UnrealBuildAccelerator/Detours/Private/Windows/UbaDetoursFunctionsNtDll.inl
+++ b/Engine/Source/Programs/UnrealBuildAccelerator/Detours/Private/Windows/UbaDetoursFunctionsNtDll.inl
@@ -1236,8 +1236,8 @@ NTSTATUS NTAPI Shared_NtCreateFile(bool IsCreateFunc, PHANDLE hFileHandle, ACCES
 			True_DuplicateHandle(g_hostProcess, mappingHandle.mh, GetCurrentProcess(), &newHandle, 0, FALSE, DUPLICATE_SAME_ACCESS);
 			UBA_ASSERTF(newHandle, L"DuplicateHandle failed when opening temp file %ls (%u)", fileName.data, GetLastError());
 			mf.writtenSize = mappingHandleSize;
-			mf.committedSize = mappingHandleSize;
-			mf.mappedSize = AlignUp(mf.committedSize, 4096);
+			mf.committedSize = AlignUp(mappingHandleSize, g_pageSize);
+			mf.mappedSize = mf.committedSize;
 			TimerScope ts2(g_kernelStats.mapViewOfFile);
 			mf.baseAddress = (u8*)True_MapViewOfFile(newHandle, FILE_MAP_READ | FILE_MAP_ALL_ACCESS, 0, 0, mf.mappedSize);
 			UBA_ASSERTF(mf.baseAddress, L"MapViewOfFile failed when opening temp file %ls (%u)", fileName.data, GetLastError());
diff --git a/Engine/Source/Programs/UnrealBuildAccelerator/Detours/Private/Windows/UbaDetoursFunctionsWin.cpp b/Engine/Source/Programs/UnrealBuildAccelerator/Detours/Private/Windows/UbaDetoursFunctionsWin.cpp
--- a/Engine/Source/Programs/UnrealBuildAccelerator/Detours/Private/Windows/UbaDetoursFunctionsWin.cpp
+++ b/Engine/Source/Programs/UnrealBuildAccelerator/Detours/Private/Windows/UbaDetoursFunctionsWin.cpp
@@ -419,7 +419,7 @@ void MemoryFile::EnsureCommitted(const DetouredHandle& handle, u64 size)
 	{
 		u64 toCommit = Min(reserveSize, AlignUp(size - committedSize, g_pageSize));
 		if (!VirtualAlloc(baseAddress + committedSize, toCommit, MEM_COMMIT, PAGE_READWRITE))
-			FatalError(1347, L"Failed to ensure virtual memory for %ls. MappedSize: %llu, CommittedSize: %llu RequestedSize: %llu. (%u)", HandleToName(handle), mappedSize, committedSize, size, GetLastError());
+			FatalError(1347, L"Failed to ensure virtual memory for %ls trying to commit %llu at %llx. MappedSize: %llu, CommittedSize: %llu RequestedSize: %llu. (%u)", HandleToName(handle), toCommit, baseAddress + committedSize, mappedSize, committedSize, size, GetLastError());
 		committedSize += toCommit;
 	}
 }
