[UBA]
* Fixed compile error when compiling uba itself using clang

diff --git a/Engine/Source/Programs/UnrealBuildAccelerator/Common/Private/UbaFile.cpp b/Engine/Source/Programs/UnrealBuildAccelerator/Common/Private/UbaFile.cpp
--- a/Engine/Source/Programs/UnrealBuildAccelerator/Common/Private/UbaFile.cpp
+++ b/Engine/Source/Programs/UnrealBuildAccelerator/Common/Private/UbaFile.cpp
@@ -181,10 +181,9 @@ namespace uba
 
 		if (IsRunningWine())
 			return GetFileBasicInformation(out, logger, fileName, errorOnFail);
-		constexpr auto FileNetworkOpenInformation = (FILE_INFORMATION_CLASS)34;
 		IO_STATUS_BLOCK b;
 		FILE_NETWORK_OPEN_INFORMATION info;
-		NTSTATUS res = NtQueryInformationFile(asHANDLE(hFile), &b, &info, sizeof(info), FileNetworkOpenInformation);
+		NTSTATUS res = NtQueryInformationFile(asHANDLE(hFile), &b, &info, sizeof(info), (FILE_INFORMATION_CLASS)34); // FileNetworkOpenInformation
 		if (res != STATUS_SUCCESS)
 			return errorOnFail ? logger.Error(TC("GetFileBasicInformationByHandle (NtQueryInformationFile) failed on %s (0x%x)"), fileName, res) : false;
 		out.attributes = info.FileAttributes;
