[UBA]
* Fixed so local cmd.exe is always used on remote helpers. This fixes so cross compiled linux can be linked remotely on wine machines
* Disabled some logging about ReadConsole when not running in debug
* Fixed so GetWrittenFiles set the correct strings on originalname and name

diff --git a/Engine/Source/Programs/UnrealBuildAccelerator/Common/Private/UbaProcess.cpp b/Engine/Source/Programs/UnrealBuildAccelerator/Common/Private/UbaProcess.cpp
--- a/Engine/Source/Programs/UnrealBuildAccelerator/Common/Private/UbaProcess.cpp
+++ b/Engine/Source/Programs/UnrealBuildAccelerator/Common/Private/UbaProcess.cpp
@@ -976,6 +976,7 @@ namespace uba
 				if (!FileExists(m_session.m_logger, wf.backedName.c_str(), &fileSize))
 					continue;
 			writer.WriteStringKey(kv.first);
+			writer.WriteString(wf.name);
 			writer.WriteString(wf.backedName);
 			writer.WriteU64(wf.mappingHandle.ToU64());
 			writer.WriteU64(fileSize);
diff --git a/Engine/Source/Programs/UnrealBuildAccelerator/Common/Private/UbaSessionClient.cpp b/Engine/Source/Programs/UnrealBuildAccelerator/Common/Private/UbaSessionClient.cpp
--- a/Engine/Source/Programs/UnrealBuildAccelerator/Common/Private/UbaSessionClient.cpp
+++ b/Engine/Source/Programs/UnrealBuildAccelerator/Common/Private/UbaSessionClient.cpp
@@ -322,6 +322,9 @@ namespace uba
 		#if PLATFORM_LINUX
 		//if (EndsWith(startInfo.application, TStrlen(startInfo.application), "/sh")) // For some reason the local shell application on the helpers hang (need to investigate aws shell application)
 		//	return true;
+		#elif PLATFORM_WINDOWS
+		if (EndsWith(startInfo.application, TStrlen(startInfo.application), TC("system32\\cmd.exe"))) // We want to use local cmd.exe.. wine does not support kernel calls called by windows 11's cmd.exe for example
+			return true;
 		#endif
 
 		outRealApplication.Clear();
diff --git a/Engine/Source/Programs/UnrealBuildAccelerator/Detours/Private/UbaDetoursFileMappingTable.cpp b/Engine/Source/Programs/UnrealBuildAccelerator/Detours/Private/UbaDetoursFileMappingTable.cpp
--- a/Engine/Source/Programs/UnrealBuildAccelerator/Detours/Private/UbaDetoursFileMappingTable.cpp
+++ b/Engine/Source/Programs/UnrealBuildAccelerator/Detours/Private/UbaDetoursFileMappingTable.cpp
@@ -182,23 +182,27 @@ namespace uba
 			StringKey key = reader.ReadStringKey();
 			auto insres = g_mappedFileTable.m_lookup.try_emplace(key);
 			FileInfo& info = insres.first->second;
-			StringBuffer<> name;
-			reader.ReadString(name);
+
+			StringBuffer<> originalName;
+			reader.ReadString(originalName);
+			if (!info.originalName || !originalName.Equals(info.originalName))
+				info.originalName = g_mappedFileTable.m_memoryBlock.Strdup(originalName).data;
+
+			StringBufferBase& backedName = originalName.Clear();
+			reader.ReadString(backedName);
+
 			FileMappingHandle mappingHandle = FileMappingHandle::FromU64(reader.ReadU64());
 			u64 fileSize = reader.ReadU64();
 			info.fileNameKey = key;
-			if (!info.originalName)
-				info.originalName = g_mappedFileTable.m_memoryBlock.Strdup(name).data;
-			info.name = info.originalName;
 			info.size = fileSize;
 
 			if (mappingHandle.IsValid())
-			{
-				name.Clear().Append(':').AppendHex(mappingHandle.ToU64()).Append('-').AppendHex(0);
-				info.name = g_mappedFileTable.m_memoryBlock.Strdup(name).data;
-			}
+				backedName.Clear().Append(':').AppendHex(mappingHandle.ToU64()).Append('-').AppendHex(0);
+
+			if (!info.name || !backedName.Equals(info.name))
+				info.name = g_mappedFileTable.m_memoryBlock.Strdup(backedName).data;
 
-			DEBUG_LOG(TC("GOT WRITTEN FILE: %s (Mapping: %s Size: %llu)"), info.originalName, info.name, info.size);
+			DEBUG_LOG(TC("GOT WRITTEN FILE: %s (BackedFile: %s Size: %llu)"), info.originalName, info.name, info.size);
 		}
 	}
 
diff --git a/Engine/Source/Programs/UnrealBuildAccelerator/Detours/Private/Windows/UbaDetoursFunctionsKernelBase.inl b/Engine/Source/Programs/UnrealBuildAccelerator/Detours/Private/Windows/UbaDetoursFunctionsKernelBase.inl
--- a/Engine/Source/Programs/UnrealBuildAccelerator/Detours/Private/Windows/UbaDetoursFunctionsKernelBase.inl
+++ b/Engine/Source/Programs/UnrealBuildAccelerator/Detours/Private/Windows/UbaDetoursFunctionsKernelBase.inl
@@ -436,7 +436,9 @@ BOOL Detoured_WriteConsoleW(HANDLE hConsoleOutput, const VOID* lpBuffer, DWORD n
 BOOL Detoured_ReadConsoleW(HANDLE hConsoleInput, LPVOID lpBuffer, DWORD nNumberOfCharsToRead, LPDWORD lpNumberOfCharsRead, PCONSOLE_READCONSOLE_CONTROL pInputControl)
 {
 	DETOURED_CALL(ReadConsoleW);
-	Rpc_WriteLogf(L"WARNING Got call to ReadConsoleW.. this is not handled by Uba yet");
+	#if UBA_DEBUG
+	Rpc_WriteLogf(L"WARNING Got call to ReadConsoleW.. this is not handled by Uba yet"); // seems like wine's cmd.exe likes calling this
+	#endif
 	return 0;
 	//return True_ReadConsoleW(hConsoleInput, lpBuffer, nNumberOfCharsToRead, lpNumberOfCharsRead, pInputControl);
 }
@@ -2992,6 +2994,7 @@ BOOL Detoured_SetConsoleMode(HANDLE hConsoleHandle, DWORD mode)
 	else if (hConsoleHandle == g_stdHandle[1])
 	{
 		g_echoOn = (mode & ~503) != 0; // TODO: This might be wrong. Trying to figure out how echo off in batch files work in terms of win32 calls
+		DEBUG_LOG(L"Echo %s", (g_echoOn ? L"on" : L"off"));
 
 		SCOPED_WRITE_LOCK(g_communicationLock, pcs);
 		BinaryWriter writer;
