[UBA]
* Fixed bug where file shared in process hierarchy did not reset size if opened without FILE_OPEN/FILE_OPEN_IF disposition
* Fixed some trace file stats that were slightly wrong

diff --git a/Engine/Source/Programs/UnrealBuildAccelerator/Common/Private/UbaSession.cpp b/Engine/Source/Programs/UnrealBuildAccelerator/Common/Private/UbaSession.cpp
--- a/Engine/Source/Programs/UnrealBuildAccelerator/Common/Private/UbaSession.cpp
+++ b/Engine/Source/Programs/UnrealBuildAccelerator/Common/Private/UbaSession.cpp
@@ -243,7 +243,10 @@ namespace uba
 	{
 		bool shouldEvictFromMemory = IsRarelyReadAfterWritten(process, file.name) || file.mappingWritten > m_keepOutputFileMemoryMapsThreshold;
 
-		if (ShouldWriteToDisk(file.name))
+		u64 writtenSize = 0;
+		bool shouldWriteToDisk = ShouldWriteToDisk(file.name);
+
+		if (shouldWriteToDisk)
 		{
 			auto& rules = *process.m_startInfo.rules;
 			RootsHandle rootsHandle = process.GetStartInfo().rootsHandle;
@@ -331,6 +334,7 @@ namespace uba
 					return false;
 
 				shouldEvictFromMemory = false; // Can't evict without properly update filemappingtable.. the file on disk does now not match what was registered for write
+				writtenSize = res.size;
 			}
 			else
 			{
@@ -349,6 +353,7 @@ namespace uba
 
 					if (!destinationFile.Write(block.memory, block.writtenSize))
 						return false;
+					writtenSize = block.writtenSize;
 				}
 				else
 				{
@@ -383,6 +388,8 @@ namespace uba
 
 				if (!destinationFile.Close(file.lastWriteTime ? nullptr : &file.lastWriteTime))
 					return false;
+
+				writtenSize = fileSize;
 			}
 
 			// There are directory crawlers happening in parallel so we need to really make sure to invalidate this one since a crawler can actually
@@ -452,7 +459,8 @@ namespace uba
 			m_fileMappingTableSize = (u32)newSize;
 		}
 
-		TraceWrittenFile(process.m_id, file.name, file.mappingWritten);
+		if (shouldWriteToDisk)
+			TraceWrittenFile(process.m_id, file.name, writtenSize);
 
 		return true;
 	}
diff --git a/Engine/Source/Programs/UnrealBuildAccelerator/Detours/Private/Windows/UbaDetoursFunctionsNtDll.inl b/Engine/Source/Programs/UnrealBuildAccelerator/Detours/Private/Windows/UbaDetoursFunctionsNtDll.inl
--- a/Engine/Source/Programs/UnrealBuildAccelerator/Detours/Private/Windows/UbaDetoursFunctionsNtDll.inl
+++ b/Engine/Source/Programs/UnrealBuildAccelerator/Detours/Private/Windows/UbaDetoursFunctionsNtDll.inl
@@ -1263,6 +1263,9 @@ NTSTATUS NTAPI Shared_NtCreateFile(bool IsCreateFunc, PHANDLE hFileHandle, ACCES
 			TakeFileOwnershipFromSession(mappingHandle, info.size);
 		}
 
+		if (CreateDisposition != FILE_OPEN && CreateDisposition != FILE_OPEN_IF)
+			info.memoryFile->writtenSize = 0;
+
 		lookupLock.Leave();
 		*hFileHandle = CreateFileHandle();
 		TrackFileInput();
