[UBA]
* Fixed bug where getting callstack info from remote didn't send the correct info (and sometimes also crashed)

diff --git a/Engine/Source/Programs/UnrealBuildAccelerator/Common/Private/UbaLogger.cpp b/Engine/Source/Programs/UnrealBuildAccelerator/Common/Private/UbaLogger.cpp
--- a/Engine/Source/Programs/UnrealBuildAccelerator/Common/Private/UbaLogger.cpp
+++ b/Engine/Source/Programs/UnrealBuildAccelerator/Common/Private/UbaLogger.cpp
@@ -22,7 +22,7 @@ namespace uba
 {
 	CustomAssertHandler* g_assertHandler;
 
-	ANALYSIS_NORETURN void UbaAssert(const tchar* text, const char* file, u32 line, const char* expr, bool allowTerminate, u32 terminateCode, void* context)
+	ANALYSIS_NORETURN void UbaAssert(const tchar* text, const char* file, u32 line, const char* expr, bool allowTerminate, u32 terminateCode, void* context, u32 skipCallstackCount)
 	{
 		static ReaderWriterLock& assertLock = *new ReaderWriterLock(); // Leak to prevent asan annoyances during shutdown when asserts happen
 		SCOPED_WRITE_LOCK(assertLock, lock);
diff --git a/Engine/Source/Programs/UnrealBuildAccelerator/Common/Private/UbaSession.cpp b/Engine/Source/Programs/UnrealBuildAccelerator/Common/Private/UbaSession.cpp
--- a/Engine/Source/Programs/UnrealBuildAccelerator/Common/Private/UbaSession.cpp
+++ b/Engine/Source/Programs/UnrealBuildAccelerator/Common/Private/UbaSession.cpp
@@ -3191,8 +3191,11 @@ namespace uba
 
 		StringView searchPaths[3] = { detoursLibPath, alternativeLibPath, {} };
 
+		u64 size = reader.ReadU32();
+		BinaryReader reader2(reader.GetPositionData(), 0, size);
+
 		StringBuffer<16*1024> sb;
-		ParseCallstackInfo(sb, reader, application, searchPaths);
+		ParseCallstackInfo(sb, reader2, application, searchPaths);
 		writer.WriteString(sb);
 		return true;
 	}
diff --git a/Engine/Source/Programs/UnrealBuildAccelerator/Common/Private/UbaSessionClient.cpp b/Engine/Source/Programs/UnrealBuildAccelerator/Common/Private/UbaSessionClient.cpp
--- a/Engine/Source/Programs/UnrealBuildAccelerator/Common/Private/UbaSessionClient.cpp
+++ b/Engine/Source/Programs/UnrealBuildAccelerator/Common/Private/UbaSessionClient.cpp
@@ -2306,7 +2306,9 @@ namespace uba
 		NetworkMessage msg(m_client, ServiceId, SessionMessageType_GetSymbols, msgWriter);
 		msgWriter.WriteString(application);
 		msgWriter.WriteBool(isArm);
-		msgWriter.WriteBytes(reader.GetPositionData(), reader.GetLeft());
+		u32 size = reader.ReadU32();
+		msgWriter.WriteU32(size);
+		msgWriter.WriteBytes(reader.GetPositionData(), size);
 
 		BinaryReader responseReader(writer.GetData(), 0, writer.GetCapacityLeft());
 		if (!msg.Send(responseReader, m_stats.customMsg))
diff --git a/Engine/Source/Programs/UnrealBuildAccelerator/Core/Private/UbaPlatform.cpp b/Engine/Source/Programs/UnrealBuildAccelerator/Core/Private/UbaPlatform.cpp
--- a/Engine/Source/Programs/UnrealBuildAccelerator/Core/Private/UbaPlatform.cpp
+++ b/Engine/Source/Programs/UnrealBuildAccelerator/Core/Private/UbaPlatform.cpp
@@ -581,7 +581,7 @@ namespace uba
 		else
 			text.Appendf(L"ERROR: Unhandled Exception (Code: 0x%x)", exceptionCode);
 
-		UbaAssert(text.data, nullptr, 0, nullptr, false, exceptionCode, contextPtr);
+		UbaAssert(text.data, nullptr, 0, nullptr, false, exceptionCode, contextPtr, 0);
 		return EXCEPTION_CONTINUE_SEARCH;
 	}
 
@@ -591,7 +591,7 @@ namespace uba
 			return EXCEPTION_CONTINUE_SEARCH;
 		StringBuffer<> buf;
 		buf.Appendf(L"ERROR: Access violation reading %s", fileName);
-		UbaAssert(buf.data, __FILE__, __LINE__, "", true, code, info->ContextRecord);
+		UbaAssert(buf.data, __FILE__, __LINE__, "", true, code, info->ContextRecord, 0);
 		return EXCEPTION_EXECUTE_HANDLER;
 	}
 
@@ -619,7 +619,7 @@ namespace uba
 	static void segfault_sigaction(int signal)
 	{
 		const char* desc = signal == SIGSEGV ? "ERROR: Segmentation fault" : "ERROR: Unhandled signal";
-		uba::UbaAssert(desc, "", 0, "", true, -1, &signal);
+		uba::UbaAssert(desc, "", 0, "", true, -1, &signal, 0);
 	}
 #endif
 
diff --git a/Engine/Source/Programs/UnrealBuildAccelerator/Detours/Private/PosixOS/UbaDetoursFunctionsPosix.cpp b/Engine/Source/Programs/UnrealBuildAccelerator/Detours/Private/PosixOS/UbaDetoursFunctionsPosix.cpp
--- a/Engine/Source/Programs/UnrealBuildAccelerator/Detours/Private/PosixOS/UbaDetoursFunctionsPosix.cpp
+++ b/Engine/Source/Programs/UnrealBuildAccelerator/Detours/Private/PosixOS/UbaDetoursFunctionsPosix.cpp
@@ -1897,7 +1897,7 @@ void UnsupportedHostRun(char* const argv[], const char* msg)
 {
 	StringBuffer<4096> command;
 	FlattenArgs(command, argv);
-	UbaAssert(msg, __FILE__, __LINE__, command.data, true, 1999, nullptr);
+	UbaAssert(msg, __FILE__, __LINE__, command.data, true, 1999, nullptr, 0);
 }
 
 int shared_posix_spawn(pid_t* pid, const char* path, const posix_spawn_file_actions_t* file_actions, const posix_spawnattr_t* attrp, char* const argv[], char* const envp[])
@@ -1919,7 +1919,7 @@ int shared_posix_spawn(pid_t* pid, const char* path, const posix_spawn_file_acti
 	if (strstr(path, "xcode-select"))
 	{
 		if (strcmp(argv[1], "--print-path") != 0)
-			UbaAssert("xcode-select only supported with --print-path", __FILE__, __LINE__, "", true, 1999, nullptr);
+			UbaAssert("xcode-select only supported with --print-path", __FILE__, __LINE__, "", true, 1999, nullptr, 0);
 		if (!ExecuteHostRun(result, argv))
 			return -1;
 		return SpawnEcho(result.data, pid,  file_actions, attrp, envp);
@@ -2728,14 +2728,14 @@ namespace uba
 	}
 	#endif
 
-	ANALYSIS_NORETURN void UbaAssert(const tchar* text, const char* file, u32 line, const char* expr, bool allowTerminate, u32 terminateCode, void* context)
+	ANALYSIS_NORETURN void UbaAssert(const tchar* text, const char* file, u32 line, const char* expr, bool allowTerminate, u32 terminateCode, void* context, u32 skipCallstackCount)
 	{
 		SuppressDetourScope s;
 		static CriticalSection cs;
 		ScopedCriticalSection scs(cs);
 		static auto& sb =  *new StringBuffer<8*1024>;
 		WriteAssertInfo(sb, text, file, line, expr, context);
-		Rpc_ResolveCallstack(sb, 3, context);
+		Rpc_ResolveCallstack(sb, 3 + skipCallstackCount, context);
 		Rpc_WriteLog(sb.data, sb.count, true, true);
 
 		if (!allowTerminate)
diff --git a/Engine/Source/Programs/UnrealBuildAccelerator/Detours/Private/UbaDetoursShared.cpp b/Engine/Source/Programs/UnrealBuildAccelerator/Detours/Private/UbaDetoursShared.cpp
--- a/Engine/Source/Programs/UnrealBuildAccelerator/Detours/Private/UbaDetoursShared.cpp
+++ b/Engine/Source/Programs/UnrealBuildAccelerator/Detours/Private/UbaDetoursShared.cpp
@@ -343,8 +343,10 @@ namespace uba
 
 		BinaryWriter writer;
 		writer.WriteByte(MessageType_ResolveCallstack);
+		auto written = (u32*)writer.AllocWrite(4);
 		if (WriteCallstackInfo(writer, skipCallstackCount, context))
 		{
+			*written = writer.GetPosition() - 5;
 			writer.Flush();
 			BinaryReader reader;
 			reader.ReadString(out);
diff --git a/Engine/Source/Programs/UnrealBuildAccelerator/Detours/Private/Windows/UbaDetoursFunctionsKernelBase.inl b/Engine/Source/Programs/UnrealBuildAccelerator/Detours/Private/Windows/UbaDetoursFunctionsKernelBase.inl
--- a/Engine/Source/Programs/UnrealBuildAccelerator/Detours/Private/Windows/UbaDetoursFunctionsKernelBase.inl
+++ b/Engine/Source/Programs/UnrealBuildAccelerator/Detours/Private/Windows/UbaDetoursFunctionsKernelBase.inl
@@ -3055,7 +3055,7 @@ BOOL Detoured_CreateProcessW(LPCWSTR lpApplicationName, LPWSTR lpCommandLine, LP
 	{
 		if (g_runningRemote)
 		{
-			UbaAssert(L"Suppress debugger startup and try to report issue instead. This message is here to hopefully see callstack", "", 0, "", false, 0, nullptr);
+			UbaAssert(L"Unhandled exception/crash. Suppress debugger startup and try to report issue instead. This message is here to hopefully see callstack", "", 0, "", false, 0, nullptr, 1);
 			return false;
 		}
 		else
diff --git a/Engine/Source/Programs/UnrealBuildAccelerator/Detours/Private/Windows/UbaDetoursFunctionsWin.cpp b/Engine/Source/Programs/UnrealBuildAccelerator/Detours/Private/Windows/UbaDetoursFunctionsWin.cpp
--- a/Engine/Source/Programs/UnrealBuildAccelerator/Detours/Private/Windows/UbaDetoursFunctionsWin.cpp
+++ b/Engine/Source/Programs/UnrealBuildAccelerator/Detours/Private/Windows/UbaDetoursFunctionsWin.cpp
@@ -449,7 +449,7 @@ void ToInvestigate(const wchar_t* format, ...)
 #endif
 }
 
-UBA_NOINLINE void UbaAssert(const tchar* text, const char* file, u32 line, const char* expr, bool allowTerminate, u32 terminateCode, void* context)
+UBA_NOINLINE void UbaAssert(const tchar* text, const char* file, u32 line, const char* expr, bool allowTerminate, u32 terminateCode, void* context, u32 skipCallstackCount)
 {
 	SuppressDetourScope _;
 	static CriticalSection cs;
@@ -461,7 +461,7 @@ UBA_NOINLINE void UbaAssert(const tchar* text, const char* file, u32 line, const
 
 	static auto& sb = *new StringBuffer<16*1024>();
 	WriteAssertInfo(sb.Clear(), text, file, line, expr);
-	Rpc_ResolveCallstack(sb, 3, context);
+	Rpc_ResolveCallstack(sb, 3 + skipCallstackCount, context);
 	Rpc_WriteLog(sb.data, sb.count, true, true);
 
 	#if UBA_DEBUG_LOG_ENABLED
