[UBA]
* Added more info to error message

diff --git a/Engine/Source/Programs/UnrealBuildAccelerator/Detours/Private/Windows/UbaDetoursFunctionsNtDll.inl b/Engine/Source/Programs/UnrealBuildAccelerator/Detours/Private/Windows/UbaDetoursFunctionsNtDll.inl
--- a/Engine/Source/Programs/UnrealBuildAccelerator/Detours/Private/Windows/UbaDetoursFunctionsNtDll.inl
+++ b/Engine/Source/Programs/UnrealBuildAccelerator/Detours/Private/Windows/UbaDetoursFunctionsNtDll.inl
@@ -1211,8 +1211,8 @@ NTSTATUS NTAPI Shared_NtCreateFile(bool IsCreateFunc, PHANDLE hFileHandle, ACCES
 		info.isFileMap = true;
 		if (!True_DuplicateHandle(g_hostProcess, mappingHandle, GetCurrentProcess(), &info.trueFileMapHandle, 0, FALSE, DUPLICATE_SAME_ACCESS))
 		{
-			Rpc_WriteLogf(L"Can't duplicate handle 0x%llx (%ls) for file %ls", uintptr_t(mappingHandle), lpFileName, info.originalName);
-			UBA_ASSERTF(info.trueFileMapHandle, L"Can't duplicate handle 0x%llx (%ls) for file %ls", uintptr_t(mappingHandle), lpFileName, info.originalName);
+			Rpc_WriteLogf(L"Can't duplicate handle 0x%llx (%ls) for file %ls (Error %u)", uintptr_t(mappingHandle), lpFileName, info.originalName, GetLastError());
+			UBA_ASSERTF(info.trueFileMapHandle, L"Can't duplicate handle 0x%llx (%ls) for file %ls (Error %u)", uintptr_t(mappingHandle), lpFileName, info.originalName, GetLastError());
 			return STATUS_UNSUCCESSFUL;
 		}
 
