[UBA]
* Changed so .split file in ApplicationRulesStageSplitMergeApplication is set as output file. This increased my write throughput from ~300mb/s to 1.5gb/s since uba's write file logic is much more optimized

diff --git a/Engine/Source/Programs/UnrealBuildAccelerator/Core/Private/UbaApplicationRules.cpp b/Engine/Source/Programs/UnrealBuildAccelerator/Core/Private/UbaApplicationRules.cpp
--- a/Engine/Source/Programs/UnrealBuildAccelerator/Core/Private/UbaApplicationRules.cpp
+++ b/Engine/Source/Programs/UnrealBuildAccelerator/Core/Private/UbaApplicationRules.cpp
@@ -541,6 +541,11 @@ namespace uba
 	{
 		using Super = ApplicationRules;
 
+		virtual bool IsOutputFile(const StringView& fileName) const override
+		{
+			return fileName.EndsWith(TC(".split"));
+		}
+
 		virtual bool AllowStorageProxy(const StringView& file) const override
 		{
 			return !file.EndsWith(TC(".split"));
