[UBA]
* Missed file in last submit

diff --git a/Engine/Source/Programs/UnrealBuildAccelerator/Core/Public/UbaPlatform.h b/Engine/Source/Programs/UnrealBuildAccelerator/Core/Public/UbaPlatform.h
--- a/Engine/Source/Programs/UnrealBuildAccelerator/Core/Public/UbaPlatform.h
+++ b/Engine/Source/Programs/UnrealBuildAccelerator/Core/Public/UbaPlatform.h
@@ -69,8 +69,8 @@ namespace uba
 	#endif
 
 	#if UBA_DEBUG
-	#define UBA_ASSERT(x) do { if (x) break; uba::UbaAssert(TC(""), __FILE__, __LINE__, #x, true, 543221, nullptr); } while(false)
-	#define UBA_ASSERTF(x, fmt, ...) do { UBA_FMT_CHECK(fmt, ##__VA_ARGS__); if (x) break; uba::StringBuffer<> _buf; _buf.Appendf(fmt, ##__VA_ARGS__); uba::UbaAssert(_buf.data, __FILE__, __LINE__, #x, true, 543221, nullptr); } while(false)
+	#define UBA_ASSERT(x) do { if (x) break; uba::UbaAssert(TC(""), __FILE__, __LINE__, #x, true, 543221, nullptr, 0); } while(false)
+	#define UBA_ASSERTF(x, fmt, ...) do { UBA_FMT_CHECK(fmt, ##__VA_ARGS__); if (x) break; uba::StringBuffer<> _buf; _buf.Appendf(fmt, ##__VA_ARGS__); uba::UbaAssert(_buf.data, __FILE__, __LINE__, #x, true, 543221, nullptr, 0); } while(false)
 	#define UBA_ASSERT_MESSAGEBOX 0
 	#else
 	#define UBA_ASSERT(x) do { } while(false)
@@ -82,7 +82,7 @@ namespace uba
 
 	bool WriteCallstackInfo(BinaryWriter& out, u32 skipCallstack, void* context);
 	void WriteAssertInfo(StringBufferBase& out, const tchar* text, const char* file, u32 line, const char* expr, void* context = nullptr);
-	ANALYSIS_NORETURN void UbaAssert(const tchar* text, const char* file, u32 line, const char* expr, bool allowTerminate, u32 terminateCode, void* context);
+	ANALYSIS_NORETURN void UbaAssert(const tchar* text, const char* file, u32 line, const char* expr, bool allowTerminate, u32 terminateCode, void* context, u32 skipCallstackCount);
 	ANALYSIS_NORETURN void FatalError(u32 code, const tchar* format, ...);
 	using CustomAssertHandler = void(const tchar* text);
 	void SetCustomAssertHandler(CustomAssertHandler* handler);
