[UBA]
* Fixed bad assert when getting remapped files
* Fixed compile warning

diff --git a/Engine/Source/Programs/UnrealBuildAccelerator/Detours/Private/UbaDetoursFileMappingTable.cpp b/Engine/Source/Programs/UnrealBuildAccelerator/Detours/Private/UbaDetoursFileMappingTable.cpp
--- a/Engine/Source/Programs/UnrealBuildAccelerator/Detours/Private/UbaDetoursFileMappingTable.cpp
+++ b/Engine/Source/Programs/UnrealBuildAccelerator/Detours/Private/UbaDetoursFileMappingTable.cpp
@@ -42,8 +42,8 @@ namespace uba
 			{
 				if (info.name && info.name[0] == '^' && !mappedFileName.Equals(info.name)) // Mapped file has been re-mapped.
 				{
+					UBA_ASSERTF(!info.memoryFile, TC("Mapped file %s has changed mapping (%s to %s) while being in use"), info.originalName, info.name, mappedFileName.data);
 					info.name = m_memoryBlock.Strdup(mappedFileName).data;
-					UBA_ASSERT(info.size == size);
 				}
 				continue;
 			}
diff --git a/Engine/Source/Programs/UnrealBuildAccelerator/Detours/Private/UbaDetoursShared.cpp b/Engine/Source/Programs/UnrealBuildAccelerator/Detours/Private/UbaDetoursShared.cpp
--- a/Engine/Source/Programs/UnrealBuildAccelerator/Detours/Private/UbaDetoursShared.cpp
+++ b/Engine/Source/Programs/UnrealBuildAccelerator/Detours/Private/UbaDetoursShared.cpp
@@ -346,7 +346,7 @@ namespace uba
 		auto written = (u32*)writer.AllocWrite(4);
 		if (WriteCallstackInfo(writer, skipCallstackCount, context))
 		{
-			*written = writer.GetPosition() - 5;
+			*written = u32(writer.GetPosition()) - 5;
 			writer.Flush();
 			BinaryReader reader;
 			reader.ReadString(out);
