[UBA]
* Added dxcore.dll to known system files to make sure it is never transferred to remotes

diff --git a/Engine/Source/Programs/UnrealBuildAccelerator/Core/Public/UbaWinBinDependencyParser.h b/Engine/Source/Programs/UnrealBuildAccelerator/Core/Public/UbaWinBinDependencyParser.h
--- a/Engine/Source/Programs/UnrealBuildAccelerator/Core/Public/UbaWinBinDependencyParser.h
+++ b/Engine/Source/Programs/UnrealBuildAccelerator/Core/Public/UbaWinBinDependencyParser.h
@@ -28,6 +28,7 @@ namespace uba
 		L"dbghelp.dll",
 		L"dnsapi.dll",
 		L"dwmapi.dll",
+		L"dxcore.dll",
 		L"dxgi.dll",
 		L"fwpuclnt.dll",
 		L"gdi32.dll",
