[UBA]
* Hardened code around creating additional connections requested by clients. Now it does not hold on to a connection pointer anymore.

diff --git a/Engine/Source/Programs/UnrealBuildAccelerator/Common/Private/UbaNetworkServer.cpp b/Engine/Source/Programs/UnrealBuildAccelerator/Common/Private/UbaNetworkServer.cpp
--- a/Engine/Source/Programs/UnrealBuildAccelerator/Common/Private/UbaNetworkServer.cpp
+++ b/Engine/Source/Programs/UnrealBuildAccelerator/Common/Private/UbaNetworkServer.cpp
@@ -1231,24 +1231,37 @@ namespace uba
 					return true;
 				u32 toAdd = connectionCount - c.connectionCount;
 
-				auto connPtr = (NetworkServer::Connection*)connectionInfo.internalData;
+				u32 connectionId = connectionInfo.GetId();
 				SCOPED_FUTEX(m_addConnectionsLock, lock2);
 				for (u32 i = 0; i != toAdd; ++i)
 				{
-					m_addConnections.emplace_back([this, connPtr]()
+					m_addConnections.emplace_back([this, connectionId]()
 						{
-							auto& conn = *connPtr;
-							conn.m_backend.Connect(m_logger, conn.m_remoteSockAddr, [this, &conn](void* connection, const sockaddr& remoteSocketAddr, bool* timedOut)
+							Connection* conn = nullptr;
+							SCOPED_READ_LOCK(m_connectionsLock, lock);
+							for (auto& c : m_connections)
+								if (c.m_id == connectionId)
+									conn = &c;
+							if (!conn || conn->m_disconnected)
+								return 0;
+							auto& backend = conn->m_backend;
+							auto remoteSockAddr = conn->m_remoteSockAddr;
+							CryptoKey cryptoKey = InvalidCryptoKey;
+							if (conn->m_cryptoKey)
+							{
+								cryptoKey = Crypto::DuplicateKey(m_logger, conn->m_cryptoKey);
+								if (cryptoKey == InvalidCryptoKey)
+									return 0;
+							}
+							lock.Leave();
+
+							bool success = backend.Connect(m_logger, remoteSockAddr, [this, &backend, &cryptoKey](void* connection, const sockaddr& remoteSocketAddr, bool* timedOut)
 								{
-									CryptoKey cryptoKey = InvalidCryptoKey;
-									if (conn.m_cryptoKey)
-									{
-										cryptoKey = Crypto::DuplicateKey(m_logger, conn.m_cryptoKey);
-										if (cryptoKey == InvalidCryptoKey)
-											return false;
-									}
-									return AddConnection(conn.m_backend, connection, remoteSocketAddr, cryptoKey != InvalidCryptoKey, cryptoKey);
+									return AddConnection(backend, connection, remoteSocketAddr, cryptoKey != InvalidCryptoKey, cryptoKey);
 								}, nullptr);
+
+							if (!success)
+								Crypto::DestroyKey(cryptoKey);
 							return 0;
 						});
 				}
