[UBA]
* Moved initial stats information about mem/cpu from client update thread to connect message. This will make visualizer lines more correct in terms of when client connected

diff --git a/Engine/Source/Programs/UnrealBuildAccelerator/Common/Private/UbaSessionClient.cpp b/Engine/Source/Programs/UnrealBuildAccelerator/Common/Private/UbaSessionClient.cpp
--- a/Engine/Source/Programs/UnrealBuildAccelerator/Common/Private/UbaSessionClient.cpp
+++ b/Engine/Source/Programs/UnrealBuildAccelerator/Common/Private/UbaSessionClient.cpp
@@ -1259,6 +1259,20 @@ namespace uba
 		GetSystemInfo(info);
 		writer.WriteString(info);
 
+		u64 memAvail;
+		u64 memTotal;
+		GetMemoryInfo(memAvail, memTotal);
+		float cpuLoad = UpdateCpuLoad();
+
+		m_cpuUsage = cpuLoad;
+		m_memAvail = memAvail;
+		m_memTotal = memTotal;
+
+		writer.WriteU64(memAvail);
+		writer.WriteU64(memTotal);
+		writer.WriteU32(*(u32*)&cpuLoad);
+
+
 		StackBinaryReader<SendMaxSize> reader;
 
 		if (!msg.Send(reader, m_stats.connectMsg))
@@ -1704,13 +1718,8 @@ namespace uba
 
 		constexpr u64 waitTimeToSpawnAfterKillMs = 5 * 1000;
 
-		u64 memAvail;
-		u64 memTotal;
-		GetMemoryInfo(memAvail, memTotal);
-		m_memAvail = memAvail;
-		m_memTotal = memTotal;
-
-		SendPing(memAvail, memTotal);
+		u64 memAvail = m_memAvail;
+		u64 memTotal = m_memTotal;
 
 		u64 memRequiredToSpawn = u64(double(memTotal) * double(100 - m_memWaitLoadPercent) / 100.0);
 		u64 memRequiredFree = u64(double(memTotal) * double(100 - m_memKillLoadPercent) / 100.0);
diff --git a/Engine/Source/Programs/UnrealBuildAccelerator/Common/Private/UbaSessionServer.cpp b/Engine/Source/Programs/UnrealBuildAccelerator/Common/Private/UbaSessionServer.cpp
--- a/Engine/Source/Programs/UnrealBuildAccelerator/Common/Private/UbaSessionServer.cpp
+++ b/Engine/Source/Programs/UnrealBuildAccelerator/Common/Private/UbaSessionServer.cpp
@@ -780,6 +780,12 @@ namespace uba
 		StringBuffer<256> info;
 		reader.ReadString(info);
 
+		u64 memAvail = reader.ReadU64();
+		u64 memTotal = reader.ReadU64();
+		u32 cpuLoadValue = reader.ReadU32();
+		float cpuLoad = *(float*)&cpuLoadValue;
+
+
 		// I have no explanation for this. On linux we get a shutdown crash when running through UBT if session is allocated with normal new
 		// For now we will work around it by using aligned_alloc which seems to be working on all platforms
 		auto& session = *new (aligned_alloc(alignof(ClientSession), sizeof(ClientSession))) ClientSession();
@@ -791,6 +797,9 @@ namespace uba
 		session.processSlotCount = processSlotCount;
 		session.dedicated = dedicated;
 		session.isArm = isClientArm;
+		session.memAvail = memAvail;
+		session.memTotal = memTotal;
+		session.cpuLoad = cpuLoad;
 		m_availableRemoteSlotCount += processSlotCount;
 		++m_connectionCount;
 
@@ -811,6 +820,7 @@ namespace uba
 		WriteRemoteEnvironmentVariables(writer);
 
 		m_trace.SessionAdded(sessionId, connectionInfo.GetId(), name, info); // Must be inside lock for TraceSessionUpdate() to not include
+		m_trace.SessionUpdate(sessionId, 1, 0, 0, 0, memAvail, memTotal, cpuLoad);
 
 		lock.Leave();
 		return true;
diff --git a/Engine/Source/Programs/UnrealBuildAccelerator/Visualizer/Private/UbaVisualizer.cpp b/Engine/Source/Programs/UnrealBuildAccelerator/Visualizer/Private/UbaVisualizer.cpp
--- a/Engine/Source/Programs/UnrealBuildAccelerator/Visualizer/Private/UbaVisualizer.cpp
+++ b/Engine/Source/Programs/UnrealBuildAccelerator/Visualizer/Private/UbaVisualizer.cpp
@@ -1275,10 +1275,7 @@ namespace uba
 			{
 				if (posY + GraphHeight >= progressRect.top && posY + GraphHeight - 5 < progressRect.bottom)
 				{
-					bool isRemote = !isFirst;
 					int posX = int(m_scrollPosX) + progressRect.left;
-					if (isRemote)
-						posX -= 15; // TODO: This should be calculated instead.. remote graph stats are much older
 					int graphBaseY = posY + GraphHeight - 4;
 					double graphHeight = double(GraphHeight - 2);
 
