[UBA]
* Added so Config_Load returns pointer to config

diff --git a/Engine/Source/Programs/UnrealBuildAccelerator/Host/Private/UbaExports.cpp b/Engine/Source/Programs/UnrealBuildAccelerator/Host/Private/UbaExports.cpp
--- a/Engine/Source/Programs/UnrealBuildAccelerator/Host/Private/UbaExports.cpp
+++ b/Engine/Source/Programs/UnrealBuildAccelerator/Host/Private/UbaExports.cpp
@@ -140,10 +140,9 @@ extern "C"
 			delete writer;
 	}
 
-	bool Config_Load(const uba::tchar* configFile)
+	uba::Config* Config_Load(const uba::tchar* configFile)
 	{
-		uba::GetConfig(configFile);
-		return true;
+		return &uba::GetConfig(configFile);
 	}
 
 	uba::Config* Config_Create()
diff --git a/Engine/Source/Programs/UnrealBuildAccelerator/Host/Public/UbaExports.h b/Engine/Source/Programs/UnrealBuildAccelerator/Host/Public/UbaExports.h
--- a/Engine/Source/Programs/UnrealBuildAccelerator/Host/Public/UbaExports.h
+++ b/Engine/Source/Programs/UnrealBuildAccelerator/Host/Public/UbaExports.h
@@ -48,7 +48,7 @@ extern "C"
 	UBA_API void DestroyCallbackLogWriter(uba::LogWriter* writer);
 
 	// Config
-	UBA_API bool Config_Load(const uba::tchar* configFile);
+	UBA_API uba::Config* Config_Load(const uba::tchar* configFile);
 	UBA_API uba::Config* Config_Create();
 	UBA_API void Config_Destroy(uba::Config* config);
 	UBA_API uba::ConfigTable* Config_RootTable(uba::Config& config);
