[UBA]
* Added support for setting compressor/compression level through config

diff --git a/Engine/Source/Programs/UnrealBuildAccelerator/Common/Private/UbaStorage.cpp b/Engine/Source/Programs/UnrealBuildAccelerator/Common/Private/UbaStorage.cpp
--- a/Engine/Source/Programs/UnrealBuildAccelerator/Common/Private/UbaStorage.cpp
+++ b/Engine/Source/Programs/UnrealBuildAccelerator/Common/Private/UbaStorage.cpp
@@ -71,6 +71,14 @@ namespace uba
 #endif
 	}
 
+	StorageCreateInfo::StorageCreateInfo(const tchar* rootDir_, LogWriter& w)
+	:	writer(w)
+	,	rootDir(rootDir_)
+	,	casCompressor(DefaultCompressor)
+	,	casCompressionLevel(DefaultCompressionLevel)
+	{
+	}
+
 	void StorageCreateInfo::Apply(const Config& config)
 	{
 		const ConfigTable* tablePtr = config.GetTable(TC("Storage"));
@@ -81,6 +89,14 @@ namespace uba
 		table.GetValueAsBool(asyncUnmapViewOfFile, TC("AsyncUnmapViewOfFile"));
 		table.GetValueAsU64(casCapacityBytes, TC("CasCapacityBytes"));
 		table.GetValueAsBool(keepTransientDataMapped, TC("KeepTransientDataMapped"));
+
+		const tchar* compressor;
+		if (table.GetValueAsString(compressor, TC("Compressor")))
+			casCompressor = GetCompressor(compressor);
+
+		const tchar* compressionLevel;
+		if (table.GetValueAsString(compressionLevel, TC("CompressionLevel")))
+			casCompressionLevel = GetCompressionLevel(compressionLevel);
 	}
 
 	struct FileMappingScope
@@ -1258,8 +1274,8 @@ namespace uba
 
 		m_casDataBuffer.AddTransient(TC("CasData"), info.keepTransientDataMapped);
 
-		m_casCompressor = DefaultCompressor;
-		m_casCompressionLevel = DefaultCompressionLevel;
+		m_casCompressor = info.casCompressor;
+		m_casCompressionLevel = info.casCompressionLevel;
 
 		#if UBA_USE_MIMALLOC
 		OodleCore_Plugins_SetAllocators(Oodle_MallocAligned, Oodle_Free);
diff --git a/Engine/Source/Programs/UnrealBuildAccelerator/Common/Private/UbaStorageUtils.cpp b/Engine/Source/Programs/UnrealBuildAccelerator/Common/Private/UbaStorageUtils.cpp
--- a/Engine/Source/Programs/UnrealBuildAccelerator/Common/Private/UbaStorageUtils.cpp
+++ b/Engine/Source/Programs/UnrealBuildAccelerator/Common/Private/UbaStorageUtils.cpp
@@ -9,6 +9,41 @@
 
 namespace uba
 {
+	#define OODLE_COMPRESSORS \
+		OODLE_COMPRESSOR(Selkie) \
+		OODLE_COMPRESSOR(Mermaid) \
+		OODLE_COMPRESSOR(Kraken) \
+		OODLE_COMPRESSOR(Leviathan) \
+
+	#define OODLE_COMPRESSION_LEVELS \
+		OODLE_COMPRESSION_LEVEL(None) \
+		OODLE_COMPRESSION_LEVEL(SuperFast) \
+		OODLE_COMPRESSION_LEVEL(VeryFast) \
+		OODLE_COMPRESSION_LEVEL(Fast) \
+		OODLE_COMPRESSION_LEVEL(Normal) \
+		OODLE_COMPRESSION_LEVEL(Optimal1) \
+		OODLE_COMPRESSION_LEVEL(Optimal2) \
+		OODLE_COMPRESSION_LEVEL(Optimal3) \
+		OODLE_COMPRESSION_LEVEL(Optimal4) \
+		OODLE_COMPRESSION_LEVEL(Optimal5) \
+
+	u8 GetCompressor(const tchar* str)
+	{
+		#define OODLE_COMPRESSOR(x) if (Equals(str, TC(#x))) return u8(OodleLZ_Compressor_##x);
+		OODLE_COMPRESSORS
+		#undef OODLE_COMPRESSOR
+		return DefaultCompressor;
+	}
+
+	u8 GetCompressionLevel(const tchar* str)
+	{
+		#define OODLE_COMPRESSION_LEVEL(x) if (Equals(str, TC(#x))) return u8(OodleLZ_CompressionLevel_##x);
+		OODLE_COMPRESSION_LEVELS
+		#undef OODLE_COMPRESSION_LEVEL
+		return DefaultCompressionLevel;
+	}
+
+
 	CasKey CalculateCasKey(u8* fileMem, u64 fileSize, bool storeCompressed, WorkManager* workManager, const tchar* hint)
 	{
 		constexpr u32 MaxWorkItemsPerAction2 = 128; // Cap this to not starve other things
diff --git a/Engine/Source/Programs/UnrealBuildAccelerator/Common/Public/UbaStorage.h b/Engine/Source/Programs/UnrealBuildAccelerator/Common/Public/UbaStorage.h
--- a/Engine/Source/Programs/UnrealBuildAccelerator/Common/Public/UbaStorage.h
+++ b/Engine/Source/Programs/UnrealBuildAccelerator/Common/Public/UbaStorage.h
@@ -78,7 +78,7 @@ namespace uba
 
 	struct StorageCreateInfo
 	{
-		StorageCreateInfo(const tchar* rootDir_, LogWriter& w) : writer(w), rootDir(rootDir_) {}
+		StorageCreateInfo(const tchar* rootDir_, LogWriter& w);
 
 		void Apply(const Config& config);
 
@@ -93,6 +93,8 @@ namespace uba
 		bool keepTransientDataMapped = true;	// Will keep the transient data mapped in instead of doing map/unmap everytime we access it
 		WorkManager* workManager = nullptr;
 		MutexHandle exclusiveMutex = InvalidMutexHandle; // Set this to hand over pre-exclusive access to storage. Storage will release mutex on shutdown
+		u8 casCompressor = 0;
+		u8 casCompressionLevel = 0;
 	};
 
 	struct BufferSlots
diff --git a/Engine/Source/Programs/UnrealBuildAccelerator/Common/Public/UbaStorageUtils.h b/Engine/Source/Programs/UnrealBuildAccelerator/Common/Public/UbaStorageUtils.h
--- a/Engine/Source/Programs/UnrealBuildAccelerator/Common/Public/UbaStorageUtils.h
+++ b/Engine/Source/Programs/UnrealBuildAccelerator/Common/Public/UbaStorageUtils.h
@@ -15,6 +15,9 @@ namespace uba
 	static constexpr u8 DefaultCompressor = OodleLZ_Compressor_Kraken;
 	static constexpr u8 DefaultCompressionLevel = OodleLZ_CompressionLevel_SuperFast;
 
+	u8 GetCompressor(const tchar* str);
+	u8 GetCompressionLevel(const tchar* str);
+
 	CasKey CalculateCasKey(u8* fileMem, u64 fileSize, bool storeCompressed, WorkManager* workManager, const tchar* hint);
 	bool SendBatchMessages(Logger& logger, NetworkClient& client, u16 fetchId, u8* slot, u64 capacity, u64 left, u32 messageMaxSize, u32& readIndex, u32& responseSize, const Function<bool()>& runInWaitFunc = {}, const tchar* hint = TC(""), u32* outError = nullptr);
 	bool SendFile(Logger& logger, NetworkClient& client, const CasKey& casKey, const u8* sourceMem, u64 sourceSize, const tchar* hint);
