[UBA]
* Added toggle to disable visualization of drive stats

diff --git a/Engine/Source/Programs/UnrealBuildAccelerator/Visualizer/Private/UbaVisualizer.cpp b/Engine/Source/Programs/UnrealBuildAccelerator/Visualizer/Private/UbaVisualizer.cpp
--- a/Engine/Source/Programs/UnrealBuildAccelerator/Visualizer/Private/UbaVisualizer.cpp
+++ b/Engine/Source/Programs/UnrealBuildAccelerator/Visualizer/Private/UbaVisualizer.cpp
@@ -1270,7 +1270,7 @@ namespace uba
 				posY += m_sessionStepY;
 			}
 
-			bool showGraph = m_config.showNetworkStats || m_config.showCpuMemStats;
+			bool showGraph = m_config.showNetworkStats || m_config.showCpuMemStats || m_config.showDriveStats;
 			if (showGraph && hasUpdates)
 			{
 				if (posY + GraphHeight >= progressRect.top && posY + GraphHeight - 5 < progressRect.bottom)
@@ -1359,9 +1359,10 @@ namespace uba
 						drawGraph(session.memAvail, session.memTotal, 1.0, m_memPen, false);
 					}
 
-					for (auto& pair : session.drives)
-						if (pair.second.busyHighest)
-							drawGraph(pair.second.busyPercent, 0, -0.01, m_drivePen, false);
+					if (m_config.showDriveStats)
+						for (auto& pair : session.drives)
+							if (pair.second.busyHighest)
+								drawGraph(pair.second.busyPercent, 0, -0.01, m_drivePen, false);
 
 				}
 				posY += GraphHeight;
@@ -1900,18 +1901,26 @@ namespace uba
 					DrawTextW(hdc, text, strLen, &r, DT_SINGLELINE|DT_NOCLIP|DT_NOPREFIX);
 				};
 
-			if (m_config.showCpuMemStats)
+			if (m_config.showDriveStats)
 			{
 				SetActiveFont(m_defaultFont);
-				drawText(L"CPU", m_cpuColor);
-				drawText(L"MEM", m_memColor);
+				drawText(L"DRV", m_driveColor);
 			}
+
 			if (m_config.showNetworkStats)
 			{
 				SetActiveFont(m_defaultFont);
 				drawText(L"SND", m_sendColor);
 				drawText(L"RCV", m_recvColor);
 			}
+
+			if (m_config.showCpuMemStats)
+			{
+				SetActiveFont(m_defaultFont);
+				drawText(L"CPU", m_cpuColor);
+				drawText(L"MEM", m_memColor);
+			}
+
 			SetTextColor(hdc, m_textColor);
 		}
 
@@ -2888,7 +2897,7 @@ namespace uba
 				posY += m_sessionStepY;
 			}
 
-			bool showGraph = m_config.showNetworkStats || m_config.showCpuMemStats;
+			bool showGraph = m_config.showNetworkStats || m_config.showCpuMemStats || m_config.showDriveStats;
 			if (showGraph && !session.updates.empty())
 			{
 				if (pos.y >= posY && pos.y < posY + GraphHeight)
@@ -3334,7 +3343,7 @@ namespace uba
 
 	bool Visualizer::ActiveProcessesShouldFillHeight()
 	{
-		return !m_config.showDetailedData && !m_config.showTitleBars && !m_config.showCpuMemStats && !m_config.showNetworkStats && !m_config.showProcessBars;
+		return !m_config.showDetailedData && !m_config.showTitleBars && !m_config.showCpuMemStats && !m_config.showNetworkStats && !m_config.showDriveStats && !m_config.showProcessBars;
 	}
 
 	StringBuffer<128> Visualizer::GetWorldTime(u64 time)
diff --git a/Engine/Source/Programs/UnrealBuildAccelerator/Visualizer/Private/UbaVisualizer.h b/Engine/Source/Programs/UnrealBuildAccelerator/Visualizer/Private/UbaVisualizer.h
--- a/Engine/Source/Programs/UnrealBuildAccelerator/Visualizer/Private/UbaVisualizer.h
+++ b/Engine/Source/Programs/UnrealBuildAccelerator/Visualizer/Private/UbaVisualizer.h
@@ -16,8 +16,9 @@ namespace uba
 		UBA_VISUALIZER_FLAG(ActiveProcesses, false, L"active processes") \
 		UBA_VISUALIZER_FLAG(TitleBars, true, L"instance title bars") \
 		UBA_VISUALIZER_FLAG(DetailedData, false, L"detailed data (use -UbaDetailedTrace for even more)") \
-		UBA_VISUALIZER_FLAG(NetworkStats, true, L"network stats") \
 		UBA_VISUALIZER_FLAG(CpuMemStats, true, L"cpu/mem stats") \
+		UBA_VISUALIZER_FLAG(NetworkStats, true, L"network stats") \
+		UBA_VISUALIZER_FLAG(DriveStats, true, L"drive stats") \
 		UBA_VISUALIZER_FLAG(ProcessBars, true, L"process bars") \
 		UBA_VISUALIZER_FLAG(FinishedProcesses, true, L"finished process bars") \
 		UBA_VISUALIZER_FLAG(Timeline, true, L"timeline") \
