[UBA]
* Fixed bug in WriteToCache logic where loglines were never provided (and shouldn't have) from c# side causing uninitialized variables which sometimes caused crashes deeper into uba logic

diff --git a/Engine/Source/Programs/UnrealBuildAccelerator/Host/Private/UbaExports.cpp b/Engine/Source/Programs/UnrealBuildAccelerator/Host/Private/UbaExports.cpp
--- a/Engine/Source/Programs/UnrealBuildAccelerator/Host/Private/UbaExports.cpp
+++ b/Engine/Source/Programs/UnrealBuildAccelerator/Host/Private/UbaExports.cpp
@@ -772,29 +772,35 @@ extern "C"
 		return true;
 	}
 
-	bool CacheClient_WriteToCache(uba::CacheClient* cacheClient, uba::RootPaths* rootPaths, uba::u32 bucket, const uba::ProcessHandle* process, const uba::u8* inputs, uba::u32 inputsSize, const uba::u8* outputs, uba::u32 outputsSize)
+	void PopulateLogLines(uba::BinaryWriter& out, const uba::ProcessHandle& process)
 	{
 		using namespace uba;
-		StackBinaryWriter<16*1024> logLinesWriter;
-
-		auto& logLines = process->GetLogLines();
+		auto& logLines = process.GetLogLines();
 		for (auto& line : logLines)
 		{
-			if (logLinesWriter.GetCapacityLeft() < 1 + GetStringWriteSize(line.text.c_str(), line.text.size()))
+			if (out.GetCapacityLeft() < 1 + GetStringWriteSize(line.text.c_str(), line.text.size()))
 				break;
-			logLinesWriter.WriteString(line.text);
-			logLinesWriter.WriteByte(line.type);
+			out.WriteString(line.text);
+			out.WriteByte(line.type);
 		}
+	}
 
+	bool CacheClient_WriteToCache(uba::CacheClient* cacheClient, uba::RootPaths* rootPaths, uba::u32 bucket, const uba::ProcessHandle* process, const uba::u8* inputs, uba::u32 inputsSize, const uba::u8* outputs, uba::u32 outputsSize)
+	{
+		using namespace uba;
+		StackBinaryWriter<16*1024> logLinesWriter;
+		PopulateLogLines(logLinesWriter, *process);
 		CacheClientActiveScope ccas(cacheClient);
 		return cacheClient->WriteToCache(*rootPaths, bucket, process->GetStartInfo(), inputs, inputsSize, outputs, outputsSize, logLinesWriter.GetData(), logLinesWriter.GetPosition(), process->GetId());
 	}
 
-	bool CacheClient_WriteToCache2(uba::CacheClient* cacheClient, uba::u32 bucket, const uba::ProcessHandle* process, const uba::u8* inputs, uba::u32 inputsSize, const uba::u8* outputs, uba::u32 outputsSize, const uba::u8* logLines, uba::u32 logLinesSize)
+	bool CacheClient_WriteToCache2(uba::CacheClient* cacheClient, uba::u32 bucket, const uba::ProcessHandle* process, const uba::u8* inputs, uba::u32 inputsSize, const uba::u8* outputs, uba::u32 outputsSize)
 	{
 		using namespace uba;
+		StackBinaryWriter<16*1024> logLinesWriter;
+		PopulateLogLines(logLinesWriter, *process);
 		CacheClientActiveScope ccas(cacheClient);
-		return cacheClient->WriteToCache(bucket, process->GetStartInfo(), inputs, inputsSize, outputs, outputsSize, logLines, logLinesSize, process->GetId());
+		return cacheClient->WriteToCache(bucket, process->GetStartInfo(), inputs, inputsSize, outputs, outputsSize, logLinesWriter.GetData(), logLinesWriter.GetPosition(), process->GetId());
 	}
 
 	uba::u32 CacheClient_FetchFromCache(uba::CacheClient* cacheClient, uba::RootPaths* rootPaths, uba::u32 bucket, const uba::ProcessStartInfo& info)
diff --git a/Engine/Source/Programs/UnrealBuildAccelerator/Host/Public/UbaExports.h b/Engine/Source/Programs/UnrealBuildAccelerator/Host/Public/UbaExports.h
--- a/Engine/Source/Programs/UnrealBuildAccelerator/Host/Public/UbaExports.h
+++ b/Engine/Source/Programs/UnrealBuildAccelerator/Host/Public/UbaExports.h
@@ -147,7 +147,7 @@ extern "C"
 	// Cache
 	UBA_API uba::CacheClient* CacheClient_Create(uba::SessionServer* session, bool reportMissReason = false, const uba::tchar* crypto = nullptr, const uba::tchar* hint = nullptr);
 	UBA_API bool CacheClient_Connect(uba::CacheClient* cacheClient, const uba::tchar* host, int port);
-	UBA_API bool CacheClient_WriteToCache2(uba::CacheClient* cacheClient, uba::u32 bucket, const uba::ProcessHandle* process, const uba::u8* inputs, uba::u32 inputsSize, const uba::u8* outputs, uba::u32 outputsSize, const uba::u8* logLines, uba::u32 logLinesSize);
+	UBA_API bool CacheClient_WriteToCache2(uba::CacheClient* cacheClient, uba::u32 bucket, const uba::ProcessHandle* process, const uba::u8* inputs, uba::u32 inputsSize, const uba::u8* outputs, uba::u32 outputsSize);
 	UBA_API uba::CacheResult* CacheClient_FetchFromCache3(uba::CacheClient* cacheClient, uba::u64 rootsHandle, uba::u32 bucket, const uba::ProcessStartInfo& info);
 	UBA_API void CacheClient_RequestServerShutdown(uba::CacheClient* cacheClient, const uba::tchar* reason);
 	UBA_API void CacheClient_Destroy(uba::CacheClient* cacheClient);
