[UBA]
* Fixed some formatting errors in httpconnection where narrow char* was parsed as wide char* causing crashes when failing to connect to aws meta data server
* Removed some L"" in favor of TC("") in UbaTraceReader

diff --git a/Engine/Source/Programs/UnrealBuildAccelerator/Common/Private/UbaNetworkBackendTcp.cpp b/Engine/Source/Programs/UnrealBuildAccelerator/Common/Private/UbaNetworkBackendTcp.cpp
--- a/Engine/Source/Programs/UnrealBuildAccelerator/Common/Private/UbaNetworkBackendTcp.cpp
+++ b/Engine/Source/Programs/UnrealBuildAccelerator/Common/Private/UbaNetworkBackendTcp.cpp
@@ -1172,7 +1172,7 @@ namespace uba
 		while (totalBytesSent < requestLen) {
 			int bytesSent = send(m_socket, request + totalBytesSent, requestLen - totalBytesSent, 0);
 			if (bytesSent == -1)
-				return logger.Error(TC("HttpConnection: socket send error (%s)"), host);
+				return logger.Error(TC("HttpConnection: socket send error (%hs)"), host);
 			totalBytesSent += bytesSent;
 		}
 
@@ -1188,7 +1188,7 @@ namespace uba
 			readPos += bytesRead;
 
 		if (bytesRead == SOCKET_ERROR)
-			return logger.Error(TC("HttpConnection: socket recv error %s (%s %s)"), LastErrorToText(WSAGetLastError()).data, m_host, path);
+			return logger.Error(TC("HttpConnection: socket recv error %s (%hs %hs)"), LastErrorToText(WSAGetLastError()).data, m_host, path);
 
 		if (readPos == sizeof(buf))
 			return logger.Error(TC("HttpConnection: buffer overflow"));
diff --git a/Engine/Source/Programs/UnrealBuildAccelerator/Visualizer/Private/UbaTraceReader.cpp b/Engine/Source/Programs/UnrealBuildAccelerator/Visualizer/Private/UbaTraceReader.cpp
--- a/Engine/Source/Programs/UnrealBuildAccelerator/Visualizer/Private/UbaTraceReader.cpp
+++ b/Engine/Source/Programs/UnrealBuildAccelerator/Visualizer/Private/UbaTraceReader.cpp
@@ -100,7 +100,7 @@ namespace uba
 		u32 traceSize = reader.ReadU32(); (void)traceSize;
 		u32 version = reader.ReadU32();
 		if (version < TraceReadCompatibilityVersion || version > TraceVersion)
-			return m_logger.Error(L"Incompatible trace version (%u). Current executable supports version %u to %u.", version, TraceReadCompatibilityVersion, TraceVersion);
+			return m_logger.Error(TC("Incompatible trace version (%u). Current executable supports version %u to %u."), version, TraceReadCompatibilityVersion, TraceVersion);
 
 		out.version = version;
 		reader.ReadU32(); // ProcessId
@@ -175,7 +175,7 @@ namespace uba
 
 		out.finished = false;
 		out.sessions.emplace_back();
-		out.sessions.back().name = L"LOCAL";
+		out.sessions.back().name = TC("LOCAL");
 		bool changed;
 		return UpdateReadClient(out, client, changed);
 	}
@@ -234,7 +234,7 @@ namespace uba
 			if (!m_memoryHandle.IsValid())
 			{
 				if (!silentFail)
-					m_logger.Error(L"OpenFileMappingW - Failed to open file mapping %ls (%ls)", namedTrace, LastErrorToText().data);
+					m_logger.Error(TC("OpenFileMappingW - Failed to open file mapping %s (%s)"), namedTrace, LastErrorToText().data);
 				return false;
 			}
 			m_memoryBegin = MapViewOfFile(m_logger, m_memoryHandle, FILE_MAP_READ, 0, 0);
@@ -313,7 +313,7 @@ namespace uba
 			u64 traceSystemStartTimeUs = 0;
 
 			if (version < TraceReadCompatibilityVersion || version > TraceVersion)
-				return true;//m_logger.Error(L"Incompatible trace version (%u). Current executable supports version %u to %u.", version, TraceReadCompatibilityVersion, TraceVersion);
+				return true;//m_logger.Error(TC("Incompatible trace version (%u). Current executable supports version %u to %u."), version, TraceReadCompatibilityVersion, TraceVersion);
 
 			if (version >= 18)
 				traceSystemStartTimeUs = traceReader.Read7BitEncoded();
@@ -373,7 +373,7 @@ namespace uba
 			UBA_TRACE_TYPES
 			#undef UBA_TRACE_TYPE
 		};
-		OutputDebugString(StringBuffer<>().Appendf(L"TraceType: %s\r\n", traceTypeNames[traceType]).data);
+		OutputDebugString(StringBuffer<>().Appendf(TC("TraceType: %s\r\n"), traceTypeNames[traceType]).data);
 		#endif
 
 		switch (traceType)
@@ -407,7 +407,7 @@ namespace uba
 			}
 
 			StringBuffer<> fullName;
-			fullName.Append(sessionName).Append(L" (").Append(sessionInfo).Append(L")");
+			fullName.Append(sessionName).Append(TC(" (")).Append(sessionInfo).Append(TC(")"));
 
 			// Check if we can re-use existing session (same machine was disconnected and then reconnected)
 			u32 virtualSessionIndex = sessionIndex;
@@ -1305,7 +1305,7 @@ namespace uba
 		}
 		default:
 		{
-			return m_logger.Error(L"Unknown trace type found in stream. UbaVisualizer too old?");
+			return m_logger.Error(TC("Unknown trace type found in stream. UbaVisualizer too old?"));
 		}
 		}
 		return true;
@@ -1358,7 +1358,7 @@ namespace uba
 	bool TraceReader::SaveAs(const tchar* fileName)
 	{
 		if (!m_memoryBegin)
-			return m_logger.Warning(L"Can only save traces that are opened using -listen/-name.");
+			return m_logger.Warning(TC("Can only save traces that are opened using -listen/-name."));
 
 		FileAccessor file(m_logger, fileName);
 		if (!file.CreateWrite())
@@ -1388,7 +1388,7 @@ namespace uba
 			if (session.clientUid == clientUid)
 				return &session;
 		// First file can be retrieved here before session is connected... haven't figured out how this can happen but let's ignore that for now :-)
-		//UBA_ASSERTF(false, L"Failed to get session for clientUid %ls", GuidToString(clientUid).str);
+		//UBA_ASSERTF(false, TC("Failed to get session for clientUid %ls"), GuidToString(clientUid).str);
 		return nullptr;
 	}
 
