[UBA]
* Changed so cancelled actions are blue in visualizer

diff --git a/Engine/Source/Programs/UnrealBuildAccelerator/Visualizer/Private/UbaVisualizer.cpp b/Engine/Source/Programs/UnrealBuildAccelerator/Visualizer/Private/UbaVisualizer.cpp
--- a/Engine/Source/Programs/UnrealBuildAccelerator/Visualizer/Private/UbaVisualizer.cpp
+++ b/Engine/Source/Programs/UnrealBuildAccelerator/Visualizer/Private/UbaVisualizer.cpp
@@ -1975,7 +1975,12 @@ namespace uba
 			if (!process.returnedReason.empty())
 				logger.Info(L"  Returned: %7s", process.returnedReason.data());
 			if (hasExited && process.exitCode != 0)
-				logger.Info(L"  ExitCode: %7u", process.exitCode);
+			{
+				if (process.exitCode == ProcessCancelExitCode)
+					logger.Info(L"  ExitCode: Cancelled");
+				else
+					logger.Info(L"  ExitCode: %7u", process.exitCode);
+			}
 
 			const auto& breadcrumbs = process.breadcrumbs;
 			if (!breadcrumbs.empty())
@@ -2365,6 +2370,8 @@ namespace uba
 			brush = m_processBrushes[selected].inProgress;
 		else if (process.cacheFetch)
 			brush = m_processBrushes[selected].cacheFetch;
+		else if (process.exitCode == ProcessCancelExitCode)
+			brush = m_processBrushes[selected].returned;
 		else if (process.exitCode != 0)
 			brush = m_processBrushes[selected].error;
 
