[UBA]
* Added macro for detoured process rpc and replaced call sites in order to get logging in the log file

diff --git a/Engine/Source/Programs/UnrealBuildAccelerator/Detours/Private/UbaDetoursFileMappingTable.cpp b/Engine/Source/Programs/UnrealBuildAccelerator/Detours/Private/UbaDetoursFileMappingTable.cpp
--- a/Engine/Source/Programs/UnrealBuildAccelerator/Detours/Private/UbaDetoursFileMappingTable.cpp
+++ b/Engine/Source/Programs/UnrealBuildAccelerator/Detours/Private/UbaDetoursFileMappingTable.cpp
@@ -73,10 +73,7 @@ namespace uba
 
 	void Rpc_CreateFileW(const StringView& fileName, const StringKey& fileNameKey, u8 access, tchar* outNewName, u64 newNameCapacity, u64& outSize, u32& outCloseId, bool lock)
 	{
-		TimerScope ts(g_stats.createFile);
-		SCOPED_WRITE_LOCK(g_communicationLock, pcs);
-		BinaryWriter writer;
-		writer.WriteByte(MessageType_CreateFile);
+		RPC_MESSAGE(CreateFile, createFile)
 		writer.WriteString(fileName);
 		writer.WriteStringKey(fileNameKey);
 		writer.WriteByte(access);
@@ -99,9 +96,7 @@ namespace uba
 
 	void Rpc_CheckRemapping(const StringView& fileName, const StringKey& fileNameKey)
 	{
-		SCOPED_WRITE_LOCK(g_communicationLock, pcs);
-		BinaryWriter writer;
-		writer.WriteByte(MessageType_CheckRemapping);
+		RPC_MESSAGE(CheckRemapping, createFile)
 		writer.WriteString(fileName);
 		writer.WriteStringKey(fileNameKey);
 		writer.Flush();
@@ -116,10 +111,7 @@ namespace uba
 		u32 directoryTableSize;
 		u32 tableOffset;
 		{
-			TimerScope ts(g_stats.listDirectory);
-			SCOPED_WRITE_LOCK(g_communicationLock, pcs);
-			BinaryWriter writer;
-			writer.WriteByte(MessageType_ListDirectory);
+			RPC_MESSAGE(ListDirectory, listDirectory)
 			writer.WriteString(dirName, dirNameLen);
 			writer.WriteStringKey(dirKey);
 			writer.Flush();
@@ -141,10 +133,7 @@ namespace uba
 	{
 		u32 directoryTableSize;
 		{
-			TimerScope ts(g_stats.closeFile);
-			SCOPED_WRITE_LOCK(g_communicationLock, pcs);
-			BinaryWriter writer;
-			writer.WriteByte(MessageType_CloseFile);
+			RPC_MESSAGE(CloseFile, closeFile)
 			writer.WriteString(handleName);
 			writer.WriteU32(closeId);
 			//writer.WriteU32(attributes); // TODO THIS
@@ -211,10 +200,7 @@ namespace uba
 		u32 directoryTableSize;
 		u32 fileMappingTableSize;
 		{
-			TimerScope ts(g_stats.updateTables);
-			SCOPED_WRITE_LOCK(g_communicationLock, pcs);
-			BinaryWriter writer;
-			writer.WriteByte(MessageType_UpdateTables);
+			RPC_MESSAGE(UpdateTables, updateTables)
 			writer.Flush();
 			BinaryReader reader;
 			directoryTableSize = reader.ReadU32();
@@ -253,10 +239,7 @@ namespace uba
 
 	void Rpc_GetWrittenFiles()
 	{
-		TimerScope ts(g_stats.updateTables);
-		SCOPED_WRITE_LOCK(g_communicationLock, pcs);
-		BinaryWriter writer;
-		writer.WriteByte(MessageType_GetWrittenFiles);
+		RPC_MESSAGE(GetWrittenFiles, updateTables)
 		writer.Flush();
 		BinaryReader reader;
 
@@ -336,10 +319,7 @@ namespace uba
 		#endif
 
 		{
-			TimerScope ts(g_stats.getFullFileName);
-			SCOPED_WRITE_LOCK(g_communicationLock, pcs);
-			BinaryWriter writer;
-			writer.WriteByte(MessageType_GetFullFileName);
+			RPC_MESSAGE(GetFullFileName, getFullFileName)
 			writer.WriteString(path);
 			writer.WriteStringKey(fileNameKey);
 			u16& bytes = *(u16*)writer.AllocWrite(2);
diff --git a/Engine/Source/Programs/UnrealBuildAccelerator/Detours/Private/UbaDetoursShared.cpp b/Engine/Source/Programs/UnrealBuildAccelerator/Detours/Private/UbaDetoursShared.cpp
--- a/Engine/Source/Programs/UnrealBuildAccelerator/Detours/Private/UbaDetoursShared.cpp
+++ b/Engine/Source/Programs/UnrealBuildAccelerator/Detours/Private/UbaDetoursShared.cpp
@@ -304,10 +304,7 @@ namespace uba
 	{
 		DEBUG_LOG(TC("LOG  %.*s"), u32(textCharLength), text); // TODO: Investigate, deadlocks on non-windows
 		// DEBUG_LOG(TC("LOG [%7u] %.*s"), GetCurrentThreadId(), u32(textCharLength), text);
-		TimerScope ts(g_stats.log);
-		SCOPED_WRITE_LOCK(g_communicationLock, pcs);
-		BinaryWriter writer;
-		writer.WriteByte(MessageType_Log);
+		RPC_MESSAGE(Log, log)
 		writer.WriteBool(printInSession);
 		writer.WriteBool(isError);
 		writer.WriteString(text, textCharLength);
diff --git a/Engine/Source/Programs/UnrealBuildAccelerator/Detours/Private/UbaDetoursShared.h b/Engine/Source/Programs/UnrealBuildAccelerator/Detours/Private/UbaDetoursShared.h
--- a/Engine/Source/Programs/UnrealBuildAccelerator/Detours/Private/UbaDetoursShared.h
+++ b/Engine/Source/Programs/UnrealBuildAccelerator/Detours/Private/UbaDetoursShared.h
@@ -134,4 +134,11 @@ namespace uba
 
 	template<typename T> struct VariableMem { template<typename... Args> void Create(Args&&... args) { new (data) T(args...); }; u64 data[AlignUp(sizeof(T), sizeof(u64)) / sizeof(u64)]; };
 	#define VARIABLE_MEM(type, name) VariableMem<type> name##Mem; type& name = (type&)name##Mem.data;
+
+	#define RPC_MESSAGE(messageName, timerName) \
+		DEBUG_LOG(L"RPC_MESSAGE %hs", #messageName); \
+		TimerScope ts(g_stats.timerName); \
+		SCOPED_WRITE_LOCK(g_communicationLock, pcs); \
+		BinaryWriter writer; \
+		writer.WriteByte(MessageType_##messageName);
 }
diff --git a/Engine/Source/Programs/UnrealBuildAccelerator/Detours/Private/Windows/UbaDetoursFunctionsKernelBase.inl b/Engine/Source/Programs/UnrealBuildAccelerator/Detours/Private/Windows/UbaDetoursFunctionsKernelBase.inl
--- a/Engine/Source/Programs/UnrealBuildAccelerator/Detours/Private/Windows/UbaDetoursFunctionsKernelBase.inl
+++ b/Engine/Source/Programs/UnrealBuildAccelerator/Detours/Private/Windows/UbaDetoursFunctionsKernelBase.inl
@@ -29,10 +29,7 @@ DWORD Local_GetLongPathNameW(LPCWSTR lpszShortPath, LPWSTR lpszLongPath, DWORD c
 		u32 errorCode = 0;
 		StringBuffer<> longName;
 		{
-			TimerScope ts(g_stats.longPathName);
-			SCOPED_WRITE_LOCK(g_communicationLock, pcs);
-			BinaryWriter writer;
-			writer.WriteByte(MessageType_GetLongPathName);
+			RPC_MESSAGE(GetLongPathName, longPathName)
 			writer.WriteString(lpszShortPath);
 			writer.Flush();
 			BinaryReader reader;
@@ -238,10 +235,7 @@ BOOL Detoured_CreateDirectoryW(LPCWSTR lpPathName, LPSECURITY_ATTRIBUTES lpSecur
 	StringKey pathNameKey = ToStringKeyLower(pathName);
 
 	{
-		TimerScope ts(g_stats.createFile);
-		SCOPED_WRITE_LOCK(g_communicationLock, pcs);
-		BinaryWriter writer;
-		writer.WriteByte(MessageType_CreateDirectory);
+		RPC_MESSAGE(CreateDirectory, createFile)
 		writer.WriteStringKey(pathNameKey);
 		writer.WriteString(pathName);
 		writer.Flush();
@@ -279,10 +273,7 @@ BOOL Detoured_RemoveDirectoryW(LPCWSTR lpPathName)
 	StringKey pathNameKey = ToStringKeyLower(pathName);
 
 	{
-		TimerScope ts(g_stats.deleteFile);
-		SCOPED_WRITE_LOCK(g_communicationLock, pcs);
-		BinaryWriter writer;
-		writer.WriteByte(MessageType_RemoveDirectory);
+		RPC_MESSAGE(RemoveDirectory, deleteFile)
 		writer.WriteStringKey(pathNameKey);
 		writer.WriteString(pathName);
 		writer.Flush();
@@ -1417,10 +1408,7 @@ BOOL Detoured_CopyFileExW(LPCWSTR lpExistingFileName, LPCWSTR lpNewFileName, LPP
 	u32 lastError;
 	u32 directoryTableSize;
 	{
-		TimerScope ts(g_stats.copyFile);
-		SCOPED_WRITE_LOCK(g_communicationLock, pcs);
-		BinaryWriter writer;
-		writer.WriteByte(MessageType_CopyFile);
+		RPC_MESSAGE(CopyFile, copyFile)
 		writer.WriteStringKey(fromKey);
 		writer.WriteString(fromName);
 		writer.WriteStringKey(toKey);
@@ -1489,10 +1477,7 @@ BOOL Detoured_CreateHardLinkW(LPCWSTR lpFileName, LPCWSTR lpExistingFileName, LP
 	StringBuffer<> newToName;
 	u32 closeId;
 	{
-		TimerScope ts(g_stats.copyFile);
-		SCOPED_WRITE_LOCK(g_communicationLock, pcs);
-		BinaryWriter writer;
-		writer.WriteByte(MessageType_CopyFile);
+		RPC_MESSAGE(CopyFile, copyFile)
 		writer.WriteStringKey(fromKey);
 		writer.WriteString(fromName);
 		writer.WriteStringKey(toKey);
@@ -1557,10 +1542,7 @@ BOOL Detoured_DeleteFileW(LPCWSTR lpFileName)
 	u32 errorCode;
 	{
 		u32 closeId = 0;
-		TimerScope ts(g_stats.deleteFile);
-		SCOPED_WRITE_LOCK(g_communicationLock, pcs);
-		BinaryWriter writer;
-		writer.WriteByte(MessageType_DeleteFile);
+		RPC_MESSAGE(DeleteFile, deleteFile)
 		writer.WriteString(fixedName);
 		writer.WriteStringKey(fileNameKey);
 		writer.WriteU32(closeId);
@@ -1663,10 +1645,7 @@ bool Shared_MoveFile(LPCWSTR lpExistingFileName, LPCWSTR lpNewFileName, DWORD dw
 	u32 errorCode;
 	bool result;
 	{
-		TimerScope ts(g_stats.moveFile);
-		SCOPED_WRITE_LOCK(g_communicationLock, pcs);
-		BinaryWriter writer;
-		writer.WriteByte(MessageType_MoveFile);
+		RPC_MESSAGE(MoveFile, moveFile)
 		writer.WriteStringKey(sourceKey);
 		writer.WriteString(source);
 		writer.WriteStringKey(destKey);
@@ -2996,9 +2975,7 @@ BOOL Detoured_SetConsoleMode(HANDLE hConsoleHandle, DWORD mode)
 		g_echoOn = (mode & ~503) != 0; // TODO: This might be wrong. Trying to figure out how echo off in batch files work in terms of win32 calls
 		DEBUG_LOG(L"Echo %s", (g_echoOn ? L"on" : L"off"));
 
-		SCOPED_WRITE_LOCK(g_communicationLock, pcs);
-		BinaryWriter writer;
-		writer.WriteByte(MessageType_EchoOn);
+		RPC_MESSAGE(EchoOn, log)
 		writer.WriteBool(g_echoOn);
 		writer.Flush();
 		if (g_isDetachedProcess)
@@ -3129,10 +3106,7 @@ BOOL Detoured_CreateProcessW(LPCWSTR lpApplicationName, LPWSTR lpCommandLine, LP
 	u32 processId = 0;
 	char dll[1024];
 	{
-		TimerScope ts(g_stats.createProcess);
-		SCOPED_WRITE_LOCK(g_communicationLock, pcs);
-		BinaryWriter writer;
-		writer.WriteByte(MessageType_CreateProcess);
+		RPC_MESSAGE(CreateProcess, createProcess)
 		writer.WriteString(application);
 		writer.WriteLongString(commandLineWithoutApplication);
 		writer.WriteString(workingDir);
@@ -3248,11 +3222,8 @@ BOOL Detoured_CreateProcessW(LPCWSTR lpApplicationName, LPWSTR lpCommandLine, LP
 
 	if (isChild)
 	{
-		TimerScope ts(g_stats.createProcess);
+		RPC_MESSAGE(StartProcess, createProcess)
 		--g_stats.createProcess.count; // Don't want double count for one process
-		SCOPED_WRITE_LOCK(g_communicationLock, pcs);
-		BinaryWriter writer;
-		writer.WriteByte(MessageType_StartProcess);
 		writer.WriteU32(processId);
 		writer.WriteBool(res);
 		writer.WriteU32(lastError);
@@ -3343,9 +3314,19 @@ BOOL Detoured_TerminateProcess(HANDLE hProcess, UINT uExitCode)
 		SendExitMessage((DWORD)uExitCode, GetTime());
 		PostDeinit();
 	}
+	else
+	{
+		if (isDetouredHandle(hProcess))
+			hProcess = asDetouredHandle(hProcess).trueHandle;
+
+		if (GetProcessId(hProcess) == GetCurrentProcessId())
+		{
+			CloseCaches();
+			SendExitMessage((DWORD)uExitCode, GetTime());
+			PostDeinit();
+		}
+	}
 
-	if (isDetouredHandle(hProcess))
-		hProcess = asDetouredHandle(hProcess).trueHandle;
 	return True_TerminateProcess(hProcess, uExitCode);
 }
 
diff --git a/Engine/Source/Programs/UnrealBuildAccelerator/Detours/Private/Windows/UbaDetoursFunctionsNtDll.inl b/Engine/Source/Programs/UnrealBuildAccelerator/Detours/Private/Windows/UbaDetoursFunctionsNtDll.inl
--- a/Engine/Source/Programs/UnrealBuildAccelerator/Detours/Private/Windows/UbaDetoursFunctionsNtDll.inl
+++ b/Engine/Source/Programs/UnrealBuildAccelerator/Detours/Private/Windows/UbaDetoursFunctionsNtDll.inl
@@ -1287,10 +1287,7 @@ NTSTATUS NTAPI Shared_NtCreateFile(bool IsCreateFunc, PHANDLE hFileHandle, ACCES
 				memoryType = L"SHAREDMEMORY";
 				#endif
 				{
-					TimerScope ts(g_stats.openTempFile);
-					SCOPED_WRITE_LOCK(g_communicationLock, pcs);
-					BinaryWriter writer;
-					writer.WriteByte(MessageType_OpenTempFile);
+					RPC_MESSAGE(OpenTempFile, openTempFile)
 					writer.WriteStringKey(fileNameKey);
 					writer.WriteString(fileNameLower);
 					writer.Flush();
@@ -1619,10 +1616,7 @@ NTSTATUS NTAPI Detoured_NtClose(HANDLE handle)
 			StringBuffer<> fixedName;
 			FixPath(fixedName, path);
 
-			TimerScope ts(g_stats.createTempFile);
-			SCOPED_WRITE_LOCK(g_communicationLock, pcs);
-			BinaryWriter writer;
-			writer.WriteByte(MessageType_CreateTempFile);
+			RPC_MESSAGE(CreateTempFile, createTempFile)
 			writer.WriteStringKey(ToStringKeyLower(fixedName));
 			writer.WriteString(fixedName);
 			writer.WriteU64(mappingHandle.ToU64());
diff --git a/Engine/Source/Programs/UnrealBuildAccelerator/Detours/Private/Windows/UbaDetoursFunctionsShell32.inl b/Engine/Source/Programs/UnrealBuildAccelerator/Detours/Private/Windows/UbaDetoursFunctionsShell32.inl
--- a/Engine/Source/Programs/UnrealBuildAccelerator/Detours/Private/Windows/UbaDetoursFunctionsShell32.inl
+++ b/Engine/Source/Programs/UnrealBuildAccelerator/Detours/Private/Windows/UbaDetoursFunctionsShell32.inl
@@ -10,10 +10,7 @@ HRESULT Detoured_SHGetKnownFolderPath(REFKNOWNFOLDERID rfid, DWORD dwFlags, HAND
 	if (g_runningRemote)
 	{
 		UBA_ASSERT(hToken == NULL);
-		TimerScope ts(g_stats.getFullFileName);
-		SCOPED_WRITE_LOCK(g_communicationLock, pcs);
-		BinaryWriter writer;
-		writer.WriteByte(MessageType_SHGetKnownFolderPath);
+		RPC_MESSAGE(SHGetKnownFolderPath, getFullFileName)
 		writer.WriteBytes(&rfid, sizeof(KNOWNFOLDERID));
 		writer.WriteU32(dwFlags);
 		writer.Flush();
diff --git a/Engine/Source/Programs/UnrealBuildAccelerator/Detours/Private/Windows/UbaDetoursFunctionsWin.cpp b/Engine/Source/Programs/UnrealBuildAccelerator/Detours/Private/Windows/UbaDetoursFunctionsWin.cpp
--- a/Engine/Source/Programs/UnrealBuildAccelerator/Detours/Private/Windows/UbaDetoursFunctionsWin.cpp
+++ b/Engine/Source/Programs/UnrealBuildAccelerator/Detours/Private/Windows/UbaDetoursFunctionsWin.cpp
@@ -130,9 +130,7 @@ void SendInput()
 	u32 pos = 0;
 	while (left)
 	{
-		SCOPED_WRITE_LOCK(g_communicationLock, pcs);
-		BinaryWriter writer;
-		writer.WriteByte(MessageType_InputDependencies);
+		RPC_MESSAGE(InputDependencies, log)
 		writer.Write7BitEncoded(reserveSize);
 		reserveSize = 0;
 		u32 toWrite = Min(left, u32(writer.GetCapacityLeft() - sizeof(u32)));
@@ -565,10 +563,7 @@ LongPathMap& g_longPathNameCache = *new LongPathMap(g_memoryBlock);
 
 void Rpc_AllocFailed(const wchar_t* allocType, u32 error)
 {
-	TimerScope ts(g_stats.virtualAllocFailed);
-	SCOPED_WRITE_LOCK(g_communicationLock, pcs);
-	BinaryWriter writer;
-	writer.WriteByte(MessageType_VirtualAllocFailed);
+	RPC_MESSAGE(VirtualAllocFailed, virtualAllocFailed)
 	writer.WriteString(allocType);
 	writer.WriteU32(error);
 	writer.Flush();
@@ -639,9 +634,7 @@ void SendExitMessage(DWORD exitCode, u64 startTime)
 
 	g_stats.detoursMemory = g_memoryBlock.writtenSize; // + g_directoryTable.m_memorySize + g_mappedFileTable.m_memPosition;
 
-	SCOPED_WRITE_LOCK(g_communicationLock, pcs);
-	BinaryWriter writer;
-	writer.WriteByte(MessageType_Exit);
+	RPC_MESSAGE(Exit, log)
 	writer.WriteU32(exitCode);
 	writer.WriteString(g_logName);
 
@@ -986,10 +979,7 @@ void Init(const DetoursPayload& payload, u64 startTime)
 	u32 mappedFileTableCount;
 
 	{
-		TimerScope ts(g_stats.init);
-		SCOPED_WRITE_LOCK(g_communicationLock, pcs);
-		BinaryWriter writer;
-		writer.WriteByte(MessageType_Init);
+		RPC_MESSAGE(Init, init)
 		writer.Flush();
 		BinaryReader reader;
 
@@ -1165,10 +1155,7 @@ extern "C"
 
 	UBA_DETOURED_API u32 UbaSendCustomMessage(const void* send, u32 sendSize, void* recv, u32 recvCapacity)
 	{
-		//TimerScope ts(g_stats.init);
-		SCOPED_WRITE_LOCK(g_communicationLock, pcs);
-		BinaryWriter writer;
-		writer.WriteByte(MessageType_Custom);
+		RPC_MESSAGE(Custom, log)
 		writer.WriteU32(sendSize);
 		writer.WriteBytes(send, sendSize);
 		writer.Flush();
@@ -1181,9 +1168,7 @@ extern "C"
 
 	UBA_DETOURED_API bool UbaFlushWrittenFiles()
 	{
-		SCOPED_WRITE_LOCK(g_communicationLock, pcs);
-		BinaryWriter writer;
-		writer.WriteByte(MessageType_FlushWrittenFiles);
+		RPC_MESSAGE(FlushWrittenFiles, log)
 		writer.Flush();
 		BinaryReader reader;
 		return reader.ReadBool();
@@ -1192,9 +1177,7 @@ extern "C"
 	UBA_DETOURED_API bool UbaUpdateEnvironment(const wchar_t* reason, bool resetStats)
 	{
 		{
-			SCOPED_WRITE_LOCK(g_communicationLock, pcs);
-			BinaryWriter writer;
-			writer.WriteByte(MessageType_UpdateEnvironment);
+			RPC_MESSAGE(UpdateEnvironment, log)
 			writer.WriteString(reason ? reason : L"");
 			writer.WriteBool(resetStats);
 			writer.Flush();
@@ -1220,9 +1203,7 @@ extern "C"
 		*outArguments = 0;
 		bool newProcess;
 		{
-			SCOPED_WRITE_LOCK(g_communicationLock, pcs);
-			BinaryWriter writer;
-			writer.WriteByte(MessageType_GetNextProcess);
+			RPC_MESSAGE(GetNextProcess, log)
 			writer.WriteU32(prevExitCode);
 			g_stats.Write(writer);
 			g_kernelStats.Write(writer);
