[UBA]
* Implemented path in devirtualize logic handling backslash to double forward slash paths

diff --git a/Engine/Source/Programs/UnrealBuildAccelerator/Common/Private/UbaSession.cpp b/Engine/Source/Programs/UnrealBuildAccelerator/Common/Private/UbaSession.cpp
--- a/Engine/Source/Programs/UnrealBuildAccelerator/Common/Private/UbaSession.cpp
+++ b/Engine/Source/Programs/UnrealBuildAccelerator/Common/Private/UbaSession.cpp
@@ -2095,6 +2095,14 @@ namespace uba
 				hasRoot = true;
 				auto& path = rootsEntry->locals[(*str - RootPaths::RootStartByte)/PathsPerRoot];
 				newSize += path.size();
+				#if PLATFORM_WINDOWS
+				u32 rootIndex = (*str - RootPaths::RootStartByte);
+				u32 type = rootIndex % PathsPerRoot;
+				if (type == 2) // DoubleForward
+					for (tchar c : path)
+						newSize += c == '\\';
+
+				#endif
 			};
 
 		if (!rootsEntry->roots.NormalizeString<tchar>(m_logger, inOut.data(), inOut.size(), checkString, allowPathsWithoutRoot, hint))
@@ -2124,12 +2132,29 @@ namespace uba
 
 				#if PLATFORM_WINDOWS
 				u32 type = rootIndex % PathsPerRoot;
-				if (type != 1) // Not backslash
+				if (type == 1) // just backslash
+				{
+				}
+				else if (type == 0)
 				{
-					UBA_ASSERTF(type == 0, TC("Not root path type %u not implemented (%s)"), type, hint); // We don't like double backslash or escaped space
 					*newStringPos = 0;
 					Replace(start, '\\', '/');
 				}
+				else if (type == 2) // Double forward slash
+				{
+					*newStringPos = 0;
+					for (tchar* it=start; it!=newStringPos;++it)
+						if (*it == '\\')
+						{
+							*it = '/';
+							memmove(it+1, it, (newStringPos-it)*sizeof(tchar));
+							++newStringPos;
+						}
+				}
+				else
+				{
+					UBA_ASSERTF(false, TC("Not root path type %u not implemented (%s)"), type, hint); // We don't like double backslash or escaped space
+				}
 				#endif
 			};
 
