[UBA]
* Added support for logging all exceptions

diff --git a/Engine/Source/Programs/UnrealBuildAccelerator/Common/Private/UbaProcess.cpp b/Engine/Source/Programs/UnrealBuildAccelerator/Common/Private/UbaProcess.cpp
--- a/Engine/Source/Programs/UnrealBuildAccelerator/Common/Private/UbaProcess.cpp
+++ b/Engine/Source/Programs/UnrealBuildAccelerator/Common/Private/UbaProcess.cpp
@@ -1741,6 +1741,7 @@ namespace uba
 			payload.isChild = m_parentProcess != nullptr;
 			payload.trackInputs = m_startInfo.trackInputs;
 			payload.useCustomAllocator = allowCustomAllocator && m_startInfo.useCustomAllocator && m_startInfo.rules->AllowMiMalloc();
+			payload.reportAllExceptions = m_startInfo.reportAllExceptions || m_startInfo.rules->ReportAllExceptions();
 			payload.isRunningWine = IsRunningWine();
 			payload.readIntermediateFilesCompressed = m_session.m_readIntermediateFilesCompressed;
 			payload.uiLanguage = m_startInfo.uiLanguage;
diff --git a/Engine/Source/Programs/UnrealBuildAccelerator/Common/Public/UbaProcessStartInfo.h b/Engine/Source/Programs/UnrealBuildAccelerator/Common/Public/UbaProcessStartInfo.h
--- a/Engine/Source/Programs/UnrealBuildAccelerator/Common/Public/UbaProcessStartInfo.h
+++ b/Engine/Source/Programs/UnrealBuildAccelerator/Common/Public/UbaProcessStartInfo.h
@@ -29,6 +29,7 @@ namespace uba
 		bool writeOutputFilesOnFail = false;	// If set to true, output files will be written to disk regardless if process succeeds or not
 		const tchar* breadcrumbs = TC("");		// If not empty, write additional information to the UBA trace file
 		bool startSuspended = false;			// Start process suspended.. a bit internal atm and only supported on windows
+		bool reportAllExceptions = false;		// Report all SEH exceptions happening in the process regardless if they are handled or not
 
 		RootsHandle rootsHandle = 0;			// Handle if using vfs. Handle is generated through Session::RegisterRoots
 
diff --git a/Engine/Source/Programs/UnrealBuildAccelerator/Core/Private/UbaApplicationRules.cpp b/Engine/Source/Programs/UnrealBuildAccelerator/Core/Private/UbaApplicationRules.cpp
--- a/Engine/Source/Programs/UnrealBuildAccelerator/Core/Private/UbaApplicationRules.cpp
+++ b/Engine/Source/Programs/UnrealBuildAccelerator/Core/Private/UbaApplicationRules.cpp
@@ -570,6 +570,10 @@ namespace uba
 			return Super::IsThrowAway(fileName, isRunningRemote);
 		}
 
+		virtual bool ReportAllExceptions() const override
+		{
+			return true;
+		}
 	};
 
 	// ====
diff --git a/Engine/Source/Programs/UnrealBuildAccelerator/Core/Private/UbaPlatform.cpp b/Engine/Source/Programs/UnrealBuildAccelerator/Core/Private/UbaPlatform.cpp
--- a/Engine/Source/Programs/UnrealBuildAccelerator/Core/Private/UbaPlatform.cpp
+++ b/Engine/Source/Programs/UnrealBuildAccelerator/Core/Private/UbaPlatform.cpp
@@ -549,30 +549,30 @@ namespace uba
 #if PLATFORM_WINDOWS
 	void* g_startOfCurrentLibrary;
 	void* g_endOfCurrentLibrary;
+	bool g_reportAllExceptions;
 	LONG UbaExceptionHandler(_EXCEPTION_POINTERS* ExceptionInfo)
 	{
 		auto& record = *ExceptionInfo->ExceptionRecord;
 		u32 exceptionCode = record.ExceptionCode;
 
-		if (exceptionCode != EXCEPTION_STACK_OVERFLOW && exceptionCode != EXCEPTION_ACCESS_VIOLATION)
-			return EXCEPTION_CONTINUE_SEARCH;
-
-		if (!g_startOfCurrentLibrary)
+		if (!g_reportAllExceptions)
 		{
-			HMODULE currentModule;
-			GetModuleHandleExW(GET_MODULE_HANDLE_EX_FLAG_FROM_ADDRESS, (LPCWSTR)&g_startOfCurrentLibrary, &currentModule);
-			MODULEINFO mi;
-			GetModuleInformation(GetCurrentProcess(), currentModule, &mi, sizeof(mi));
-			g_startOfCurrentLibrary = mi.lpBaseOfDll;
-			g_endOfCurrentLibrary = (u8*)mi.lpBaseOfDll + mi.SizeOfImage;
-		}
+			if (exceptionCode != EXCEPTION_STACK_OVERFLOW && exceptionCode != EXCEPTION_ACCESS_VIOLATION)
+				return EXCEPTION_CONTINUE_SEARCH;
 
-		if (record.ExceptionAddress < g_startOfCurrentLibrary || record.ExceptionAddress >= g_endOfCurrentLibrary)
-			return EXCEPTION_CONTINUE_SEARCH;
+			if (!g_startOfCurrentLibrary)
+			{
+				HMODULE currentModule;
+				GetModuleHandleExW(GET_MODULE_HANDLE_EX_FLAG_FROM_ADDRESS, (LPCWSTR)&g_startOfCurrentLibrary, &currentModule);
+				MODULEINFO mi;
+				GetModuleInformation(GetCurrentProcess(), currentModule, &mi, sizeof(mi));
+				g_startOfCurrentLibrary = mi.lpBaseOfDll;
+				g_endOfCurrentLibrary = (u8*)mi.lpBaseOfDll + mi.SizeOfImage;
+			}
 
-		//CONTEXT context;
-		//RtlCaptureContext(&context);
-		//CONTEXT* contextPtr = &context;
+			if (record.ExceptionAddress < g_startOfCurrentLibrary || record.ExceptionAddress >= g_endOfCurrentLibrary)
+				return EXCEPTION_CONTINUE_SEARCH;
+		}
 		CONTEXT* contextPtr = ExceptionInfo->ContextRecord;
 
 		StringBuffer<> text;
diff --git a/Engine/Source/Programs/UnrealBuildAccelerator/Core/Public/UbaApplicationRules.h b/Engine/Source/Programs/UnrealBuildAccelerator/Core/Public/UbaApplicationRules.h
--- a/Engine/Source/Programs/UnrealBuildAccelerator/Core/Public/UbaApplicationRules.h
+++ b/Engine/Source/Programs/UnrealBuildAccelerator/Core/Public/UbaApplicationRules.h
@@ -185,6 +185,11 @@ namespace uba
 			return DependencyCrawlerType_None;
 		}
 
+		virtual bool ReportAllExceptions() const
+		{
+			return false;
+		}
+
 		u32 index = ~0u;
 	};
 
diff --git a/Engine/Source/Programs/UnrealBuildAccelerator/Core/Public/UbaDetoursPayload.h b/Engine/Source/Programs/UnrealBuildAccelerator/Core/Public/UbaDetoursPayload.h
--- a/Engine/Source/Programs/UnrealBuildAccelerator/Core/Public/UbaDetoursPayload.h
+++ b/Engine/Source/Programs/UnrealBuildAccelerator/Core/Public/UbaDetoursPayload.h
@@ -27,6 +27,7 @@ namespace uba
 		bool allowOutputFiles = IsWindows;
 		bool suppressLogging = false;
 		bool readIntermediateFilesCompressed = false;
+		bool reportAllExceptions = false;
 		int uiLanguage = 0;
 		StringBuffer<256> logFile;
 	};
diff --git a/Engine/Source/Programs/UnrealBuildAccelerator/Detours/Private/Windows/UbaDetoursFunctionsWin.cpp b/Engine/Source/Programs/UnrealBuildAccelerator/Detours/Private/Windows/UbaDetoursFunctionsWin.cpp
--- a/Engine/Source/Programs/UnrealBuildAccelerator/Detours/Private/Windows/UbaDetoursFunctionsWin.cpp
+++ b/Engine/Source/Programs/UnrealBuildAccelerator/Detours/Private/Windows/UbaDetoursFunctionsWin.cpp
@@ -838,6 +838,8 @@ int DetourDetachFunctions()
 	return 0;
 }
 
+extern bool g_reportAllExceptions;
+
 void PreInit(const DetoursPayload& payload)
 {
 	#if UBA_USE_MIMALLOC
@@ -847,6 +849,8 @@ void PreInit(const DetoursPayload& payload)
 
 	InitSharedVariables();
 
+	g_reportAllExceptions = payload.reportAllExceptions;
+
 	g_rulesIndex = payload.rulesIndex;
 	g_rules = GetApplicationRules()[payload.rulesIndex].rules;
 	g_useMiMalloc = payload.useCustomAllocator;
