[UBA]
* Added StringBuffer/StringView.EndsWith that takes in StringView

diff --git a/Engine/Source/Programs/UnrealBuildAccelerator/Common/Private/UbaSessionClient.cpp b/Engine/Source/Programs/UnrealBuildAccelerator/Common/Private/UbaSessionClient.cpp
--- a/Engine/Source/Programs/UnrealBuildAccelerator/Common/Private/UbaSessionClient.cpp
+++ b/Engine/Source/Programs/UnrealBuildAccelerator/Common/Private/UbaSessionClient.cpp
@@ -323,7 +323,7 @@ namespace uba
 		//if (EndsWith(startInfo.application, TStrlen(startInfo.application), "/sh")) // For some reason the local shell application on the helpers hang (need to investigate aws shell application)
 		//	return true;
 		#elif PLATFORM_WINDOWS
-		if (EndsWith(startInfo.application, TStrlen(startInfo.application), TC("system32\\cmd.exe"))) // We want to use local cmd.exe.. wine does not support kernel calls called by windows 11's cmd.exe for example
+		if (ToView(startInfo.application).EndsWith(TCV("system32\\cmd.exe"))) // We want to use local cmd.exe.. wine does not support kernel calls called by windows 11's cmd.exe for example
 			return true;
 		#endif
 
@@ -505,7 +505,7 @@ namespace uba
 		if (auto parent = process.m_parentProcess)
 		{
 			auto& si = parent->m_startInfo;
-			if (EndsWith(si.application, TStrlen(si.application), "/sh"))
+			if (ToView(si.application).EndsWith(TCV("/sh")))
 			{
 				FlushWrittenFiles(*parent);
 				StackBinaryReader<SendMaxSize> reader;
diff --git a/Engine/Source/Programs/UnrealBuildAccelerator/Core/Private/UbaApplicationRules.cpp b/Engine/Source/Programs/UnrealBuildAccelerator/Core/Private/UbaApplicationRules.cpp
--- a/Engine/Source/Programs/UnrealBuildAccelerator/Core/Private/UbaApplicationRules.cpp
+++ b/Engine/Source/Programs/UnrealBuildAccelerator/Core/Private/UbaApplicationRules.cpp
@@ -9,19 +9,19 @@ namespace uba
 	{
 		if (!IsWindows)
 			return false;
-		return fileName.EndsWith(TC(".obj"))
-			|| fileName.EndsWith(TC(".o"))
-			|| fileName.EndsWith(TC(".pch"))
-			|| fileName.EndsWith(TC(".gch"))
-			|| fileName.EndsWith(TC(".ifc"))
+		return fileName.EndsWith(TCV(".obj"))
+			|| fileName.EndsWith(TCV(".o"))
+			|| fileName.EndsWith(TCV(".pch"))
+			|| fileName.EndsWith(TCV(".gch"))
+			|| fileName.EndsWith(TCV(".ifc"))
 			;
 	}
 
 	u8 GlobalRules::GetUsedCountBeforeFree(const StringView& fileName)
 	{
-		if (!fileName.EndsWith(TC(".obj")))
+		if (!fileName.EndsWith(TCV(".obj")))
 			return 0;
-		if (fileName.EndsWith(TC(".h.obj")))
+		if (fileName.EndsWith(TCV(".h.obj")))
 			return 255;
 		return 2;
 	}
@@ -52,11 +52,11 @@ namespace uba
 
 		virtual u64 FileTypeMaxSize(const StringView& file, bool isSystemOrTempFile) const override
 		{
-			if (file.EndsWith(TC(".pdb")))
+			if (file.EndsWith(TCV(".pdb")))
 				return 14ull * 1024 * 1024 * 1024; // This is ridiculous
-			if (file.EndsWith(TC(".json")) || file.EndsWith(TC(".exp")) || file.EndsWith(TC(".sarif")) || file.EndsWith(TC(".res")))
+			if (file.EndsWith(TCV(".json")) || file.EndsWith(TCV(".exp")) || file.EndsWith(TCV(".sarif")) || file.EndsWith(TCV(".res")))
 				return 32 * 1024 * 1024;
-			if (file.EndsWith(TC(".ifc")) || file.EndsWith(TC(".obj")) || (isSystemOrTempFile && file.Contains(TC("_cl_")))) // There are _huge_ obj files when building with -stresstestunity
+			if (file.EndsWith(TCV(".ifc")) || file.EndsWith(TCV(".obj")) || (isSystemOrTempFile && file.Contains(TC("_cl_")))) // There are _huge_ obj files when building with -stresstestunity
 				return 1024 * 1024 * 1024;
 			return Super::FileTypeMaxSize(file, isSystemOrTempFile);
 		}
@@ -90,29 +90,29 @@ namespace uba
 	public:
 		virtual bool IsOutputFile(const StringView& fileName) const override
 		{
-			return fileName.EndsWith(TC(".obj"))
-				|| fileName.EndsWith(TC(".dep.json"))
-				|| fileName.EndsWith(TC(".sarif"))
-				|| fileName.EndsWith(TC(".pch"))
-				|| fileName.EndsWith(TC(".ifc"))
-				|| fileName.EndsWith(TC(".rc2.res")) // Not really an obj file.. 
+			return fileName.EndsWith(TCV(".obj"))
+				|| fileName.EndsWith(TCV(".dep.json"))
+				|| fileName.EndsWith(TCV(".sarif"))
+				|| fileName.EndsWith(TCV(".pch"))
+				|| fileName.EndsWith(TCV(".ifc"))
+				|| fileName.EndsWith(TCV(".rc2.res")) // Not really an obj file.. 
 				;
 		}
 
 		virtual bool IsRarelyRead(const StringView& file) const override
 		{
-			return file.EndsWith(TC(".cpp"))
-				|| (file.EndsWith(TC(".rsp")) && !file.EndsWith(TC("Shared.rsp")))
-				|| file.EndsWith(TC(".i")) // .i is preprocessed source code (for static analysis)
+			return file.EndsWith(TCV(".cpp"))
+				|| (file.EndsWith(TCV(".rsp")) && !file.EndsWith(TCV("Shared.rsp")))
+				|| file.EndsWith(TCV(".i")) // .i is preprocessed source code (for static analysis)
 				;
 		}
 
 		virtual bool IsRarelyReadAfterWritten(const StringView& fileName) const override
 		{
-			return fileName.EndsWith(TC(".dep.json"))
-				|| fileName.EndsWith(TC(".sarif"))
-				|| fileName.EndsWith(TC(".exe"))
-				|| fileName.EndsWith(TC(".dll"));
+			return fileName.EndsWith(TCV(".dep.json"))
+				|| fileName.EndsWith(TCV(".sarif"))
+				|| fileName.EndsWith(TCV(".exe"))
+				|| fileName.EndsWith(TCV(".dll"));
 		}
 
 		virtual bool NeedsSharedMemory(const tchar* file) const override
@@ -127,13 +127,13 @@ namespace uba
 
 		virtual bool ShouldExtractSymbols(const StringView& fileName) const override
 		{
-			return fileName.EndsWith(TC(".obj"));
+			return fileName.EndsWith(TCV(".obj"));
 		}
 
 		virtual bool ShouldDevirtualizeFile(const StringView& fileName, bool& outEscapeSpaces) const override
 		{
 			outEscapeSpaces = false;
-			return fileName.EndsWith(TC(".dep.json"));
+			return fileName.EndsWith(TCV(".dep.json"));
 		}
 
 		virtual DependencyCrawlerType GetDependencyCrawlerType() const override
@@ -153,12 +153,12 @@ namespace uba
 	public:
 		virtual bool IsOutputFile(const StringView& fileName) const override
 		{
-			return fileName.EndsWith(TC(".lib"))
-				|| fileName.EndsWith(TC(".exp"))
-				|| fileName.EndsWith(TC(".pdb"))
-				|| fileName.EndsWith(TC(".dll"))
-				|| fileName.EndsWith(TC(".exe"))
-				|| fileName.EndsWith(TC(".rc2.res")); // Not really an obj file.. 
+			return fileName.EndsWith(TCV(".lib"))
+				|| fileName.EndsWith(TCV(".exp"))
+				|| fileName.EndsWith(TCV(".pdb"))
+				|| fileName.EndsWith(TCV(".dll"))
+				|| fileName.EndsWith(TCV(".exe"))
+				|| fileName.EndsWith(TCV(".rc2.res")); // Not really an obj file.. 
 		}
 
 		virtual bool IsThrowAway(const StringView& fileName, bool isRunningRemote) const override
@@ -174,41 +174,41 @@ namespace uba
 		{
 			return Contains(file, TC("lnk{")) // This file is shared from link.exe to mt.exe and rc.exe so we need to put it shared memory
 			 	|| Contains(file, TC("\\_cl_")) // When link.exe is spawned by cl.exe we might use this which is in shared memory
-			 	|| EndsWith(file, TStrlen(file), TC(".manifest")); // lld-link.exe is using a different name for files shared with child processes
+			 	|| ToView(file).EndsWith(TCV(".manifest")); // lld-link.exe is using a different name for files shared with child processes
 		}
 
 		virtual u64 FileTypeMaxSize(const StringView& file, bool isSystemOrTempFile) const override
 		{
 			if (file.Contains(TC("lnk{")))
 				return 32 * 1024 * 1024;
-			if (file.EndsWith(TC(".lib"))) // This is import lib which is tiny but not sure if we output any full libs
+			if (file.EndsWith(TCV(".lib"))) // This is import lib which is tiny but not sure if we output any full libs
 				return 512 * 1024 * 1024;
-			if (file.EndsWith(TC(".dll")) || file.EndsWith(TC(".exe")))
+			if (file.EndsWith(TCV(".dll")) || file.EndsWith(TCV(".exe")))
 				return 4ull * 1024 * 1024 * 1024;
 			return Super::FileTypeMaxSize(file, isSystemOrTempFile);
 		}
 
 		virtual bool IsRarelyRead(const StringView& file) const override
 		{
-			return file.EndsWith(TC(".exp"))
-				|| file.EndsWith(TC(".dll.rsp"))
-				|| file.EndsWith(TC(".lib.rsp"))
-				|| file.EndsWith(TC(".ilk"))
-				|| file.EndsWith(TC(".pdb"));
+			return file.EndsWith(TCV(".exp"))
+				|| file.EndsWith(TCV(".dll.rsp"))
+				|| file.EndsWith(TCV(".lib.rsp"))
+				|| file.EndsWith(TCV(".ilk"))
+				|| file.EndsWith(TCV(".pdb"));
 		}
 
 		virtual bool AllowStorageProxy(const StringView& file) const override
 		{
-			if (file.EndsWith(TC(".obj")))
+			if (file.EndsWith(TCV(".obj")))
 				return false;
 			return Super::AllowStorageProxy(file);
 		}
 
 		virtual bool IsRarelyReadAfterWritten(const StringView& fileName) const override
 		{
-			return fileName.EndsWith(TC(".pdb"))
-				|| fileName.EndsWith(TC(".exe"))
-				|| fileName.EndsWith(TC(".dll"));
+			return fileName.EndsWith(TCV(".pdb"))
+				|| fileName.EndsWith(TCV(".exe"))
+				|| fileName.EndsWith(TCV(".dll"));
 		}
 
 		virtual bool IsCacheable() const override
@@ -251,7 +251,7 @@ namespace uba
 
 		virtual bool KeepInMemory(const StringView& fileName, const tchar* systemTemp, bool isRunningRemote, bool isWrite) const override
 		{
-			return fileName.EndsWith(TC(".manifest")) && fileName.Contains(systemTemp);
+			return fileName.EndsWith(TCV(".manifest")) && fileName.Contains(systemTemp);
 		}
 
 		virtual u64 FileTypeMaxSize(const StringView& file, bool isSystemOrTempFile) const override
@@ -276,7 +276,7 @@ namespace uba
 
 		virtual u64 FileTypeMaxSize(const StringView& file, bool isSystemOrTempFile) const override
 		{
-			if (file.Contains(TC(".pdb.tmp")) || file.EndsWith(TC(".pdb")))
+			if (file.Contains(TC(".pdb.tmp")) || file.EndsWith(TCV(".pdb")))
 				return 14ull * 1024 * 1024 * 1024; // This is ridiculous
 			return ApplicationRulesVcLink::FileTypeMaxSize(file, isSystemOrTempFile);
 		}
@@ -284,9 +284,9 @@ namespace uba
 		virtual bool IsOutputFile(const StringView& fileName) const override
 		{
 			return false//Super::IsOutputFile(fileName)
-				|| fileName.EndsWith(TC(".exe"))
-				|| fileName.EndsWith(TC(".dll"))
-				|| fileName.EndsWith(TC(".pdb"))
+				|| fileName.EndsWith(TCV(".exe"))
+				|| fileName.EndsWith(TCV(".dll"))
+				|| fileName.EndsWith(TCV(".pdb"))
 				|| fileName.Contains(TC(".exe.tmp"))
 				|| fileName.Contains(TC(".dll.tmp"))
 				|| fileName.Contains(TC(".pdb.tmp"));
@@ -307,9 +307,9 @@ namespace uba
 
 		virtual u64 FileTypeMaxSize(const StringView& file, bool isSystemOrTempFile) const override
 		{
-			if (file.EndsWith(TC(".obj")) || file.EndsWith(TC(".o")) || file.EndsWith(TC(".o.tmp")) || file.EndsWith(TC(".obj.tmp")))
+			if (file.EndsWith(TCV(".obj")) || file.EndsWith(TCV(".o")) || file.EndsWith(TCV(".o.tmp")) || file.EndsWith(TCV(".obj.tmp")))
 				return 512 * 1024 * 1024;
-			if (file.EndsWith(TC(".d")))
+			if (file.EndsWith(TCV(".d")))
 				return 32 * 1024 * 1024;
 			return Super::FileTypeMaxSize(file, isSystemOrTempFile);
 		}
@@ -327,29 +327,29 @@ namespace uba
 
 		virtual bool IsOutputFile(const StringView& fileName) const override
 		{
-			return fileName.EndsWith(TC(".c.d"))
-				|| fileName.EndsWith(TC(".h.d"))
-				|| fileName.EndsWith(TC(".cc.d"))
-				|| fileName.EndsWith(TC(".cpp.d"))
-				|| fileName.EndsWith(TC(".o"))
-				|| fileName.EndsWith(TC(".o.tmp"))
-				|| fileName.EndsWith(TC(".gch"))
-				|| fileName.EndsWith(TC(".gch.tmp"))
-				|| fileName.EndsWith(TC(".obj"))
-				|| fileName.EndsWith(TC(".obj.tmp"))
-				|| fileName.EndsWith(TC(".pch"))
+			return fileName.EndsWith(TCV(".c.d"))
+				|| fileName.EndsWith(TCV(".h.d"))
+				|| fileName.EndsWith(TCV(".cc.d"))
+				|| fileName.EndsWith(TCV(".cpp.d"))
+				|| fileName.EndsWith(TCV(".o"))
+				|| fileName.EndsWith(TCV(".o.tmp"))
+				|| fileName.EndsWith(TCV(".gch"))
+				|| fileName.EndsWith(TCV(".gch.tmp"))
+				|| fileName.EndsWith(TCV(".obj"))
+				|| fileName.EndsWith(TCV(".obj.tmp"))
+				|| fileName.EndsWith(TCV(".pch"))
 				;
 		}
 
 		virtual bool IsRarelyRead(const StringView& file) const override
 		{
-			return file.EndsWith(TC(".cpp"))
-				|| (file.EndsWith(TC(".rsp")) && !file.EndsWith(TC("Shared.rsp")));
+			return file.EndsWith(TCV(".cpp"))
+				|| (file.EndsWith(TCV(".rsp")) && !file.EndsWith(TCV("Shared.rsp")));
 		}
 
 		virtual bool IsRarelyReadAfterWritten(const StringView& fileName) const override
 		{
-			return fileName.EndsWith(TC(".d"));
+			return fileName.EndsWith(TCV(".d"));
 		}
 
 		virtual bool AllowMiMalloc() const override
@@ -364,13 +364,13 @@ namespace uba
 
 		virtual bool ShouldExtractSymbols(const StringView& fileName) const override
 		{
-			return fileName.EndsWith(TC(".obj")) || fileName.EndsWith(TC(".o"));
+			return fileName.EndsWith(TCV(".obj")) || fileName.EndsWith(TCV(".o"));
 		}
 
 		virtual bool ShouldDevirtualizeFile(const StringView& fileName, bool& outEscapeSpaces) const override
 		{
 			outEscapeSpaces = true;
-			return fileName.EndsWith(TC(".d"));
+			return fileName.EndsWith(TCV(".d"));
 		}
 
 		virtual DependencyCrawlerType GetDependencyCrawlerType() const override
@@ -434,7 +434,7 @@ namespace uba
 	{
 		virtual bool IsOutputFile(const StringView& fileName) const override
 		{
-			return false; // fileName.EndsWith(TC(".psym")); With psym as output file the BreakpadSymbolEncoder fails to output a .sym file
+			return false; // fileName.EndsWith(TCV(".psym")); With psym as output file the BreakpadSymbolEncoder fails to output a .sym file
 		}
 
 		virtual const tchar* const* LibrariesToPreload() const override
@@ -457,7 +457,7 @@ namespace uba
 
 		virtual bool IsThrowAway(const StringView& fileName, bool isRunningRemote) const override
 		{
-			return fileName.EndsWith(TC("-telemetry.json")) || Super::IsThrowAway(fileName, isRunningRemote);
+			return fileName.EndsWith(TCV("-telemetry.json")) || Super::IsThrowAway(fileName, isRunningRemote);
 		}
 
 		//virtual bool NeedsSharedMemory(const tchar* file)
@@ -472,7 +472,7 @@ namespace uba
 
 		//virtual bool IsOutputFile(const StringView& fileName) const override
 		//{
-		//	return fileName.EndsWith(TC(".self")) || Equals(file, TC("Symbols.map"));
+		//	return fileName.EndsWith(TCV(".self")) || Equals(file, TC("Symbols.map"));
 		//}
 		virtual bool KeepInMemory(const StringView& fileName, const tchar* systemTemp, bool isRunningRemote, bool isWrite) const override
 		{
@@ -502,7 +502,7 @@ namespace uba
 		virtual bool IsThrowAway(const StringView& fileName, bool isRunningRemote) const override
 		{
 			return Super::IsThrowAway(fileName, isRunningRemote)
-				|| (isRunningRemote && fileName.EndsWith(TC("-telemetry.json")));
+				|| (isRunningRemote && fileName.EndsWith(TCV("-telemetry.json")));
 		}
 
 		//virtual bool NeedsSharedMemory(const tchar* file)
@@ -523,7 +523,7 @@ namespace uba
 		virtual bool IsThrowAway(const StringView& fileName, bool isRunningRemote) const override
 		{
 			return Super::IsThrowAway(fileName, isRunningRemote)
-				|| (isRunningRemote && fileName.EndsWith(TC("-telemetry.json")));
+				|| (isRunningRemote && fileName.EndsWith(TCV("-telemetry.json")));
 		}
 
 		//virtual bool NeedsSharedMemory(const tchar* file)
@@ -543,29 +543,29 @@ namespace uba
 
 		virtual bool IsOutputFile(const StringView& fileName) const override
 		{
-			return fileName.EndsWith(TC(".split"));
+			return fileName.EndsWith(TCV(".split"));
 		}
 
 		virtual bool AllowStorageProxy(const StringView& file) const override
 		{
-			return !file.EndsWith(TC(".split"));
+			return !file.EndsWith(TCV(".split"));
 		}
 
 		virtual bool SendFileCompressedFromClient(const StringView& fileName) const override
 		{
-			return !fileName.EndsWith(TC(".split"));
+			return !fileName.EndsWith(TCV(".split"));
 		}
 
 		virtual bool KeepInMemory(const StringView& fileName, const tchar* systemTemp, bool isRunningRemote, bool isWrite) const override
 		{
-			if (isWrite && fileName.Contains(TC("\\Split\\")) && fileName.EndsWith(TC(".split")))
+			if (isWrite && fileName.Contains(TC("\\Split\\")) && fileName.EndsWith(TCV(".split")))
 				return true;
 			return false;
 		}
 
 		virtual bool IsThrowAway(const StringView& fileName, bool isRunningRemote) const override
 		{
-			if (fileName.Contains(TC("\\Split\\")) && fileName.EndsWith(TC(".split")))
+			if (fileName.Contains(TC("\\Split\\")) && fileName.EndsWith(TCV(".split")))
 				return true;
 			return Super::IsThrowAway(fileName, isRunningRemote);
 		}
@@ -590,17 +590,17 @@ namespace uba
 		virtual bool IsOutputFile(const StringView& fileName) const override
 		{
 			return fileName.Contains(TC(".generated.dummy"))
-				|| fileName.EndsWith(TC(".ispc.bc"))
-				|| fileName.EndsWith(TC(".ispc.txt"))
-				|| fileName.EndsWith(TC(".obj"))
-				|| fileName.EndsWith(TC(".o")); // Used when compiling for linux
+				|| fileName.EndsWith(TCV(".ispc.bc"))
+				|| fileName.EndsWith(TCV(".ispc.txt"))
+				|| fileName.EndsWith(TCV(".obj"))
+				|| fileName.EndsWith(TCV(".o")); // Used when compiling for linux
 		}
 
 		virtual u64 FileTypeMaxSize(const StringView& file, bool isSystemOrTempFile) const override
 		{
-			if (file.Contains(TC(".generated.dummy")) || file.EndsWith(TC(".ispc.txt")))
+			if (file.Contains(TC(".generated.dummy")) || file.EndsWith(TCV(".ispc.txt")))
 				return 32ull * 1024 * 1024;
-			if (file.EndsWith(TC(".obj")) || file.EndsWith(TC(".o")))
+			if (file.EndsWith(TCV(".obj")) || file.EndsWith(TCV(".o")))
 				return 128ull * 1024 * 1024;
 			return Super::FileTypeMaxSize(file, isSystemOrTempFile);
 		}
@@ -612,12 +612,12 @@ namespace uba
 
 		virtual bool ShouldExtractSymbols(const StringView& fileName) const override
 		{
-			return fileName.EndsWith(TC(".obj")) || fileName.EndsWith(TC(".o"));
+			return fileName.EndsWith(TCV(".obj")) || fileName.EndsWith(TCV(".o"));
 		}
 
 		virtual bool ShouldDevirtualizeFile(const StringView& fileName, bool& outEscapeSpaces) const override
 		{
-			return fileName.EndsWith(TC(".ispc.txt"));
+			return fileName.EndsWith(TCV(".ispc.txt"));
 		}
 	};
 
@@ -628,9 +628,9 @@ namespace uba
 			return false;
 			// TODO: These does not work when UnrealBuildTool creates these files multiple times in a row (building multiple targets)
 			// ... on output they get stored as file mappings.. and next execution of ubt opens them for write (writing file mappings not implemented right now)
-			//return fileName.EndsWith(TC(".modules"))
-			//	|| fileName.EndsWith(TC(".target"))
-			//	|| fileName.EndsWith(TC(".version"));
+			//return fileName.EndsWith(TCV(".modules"))
+			//	|| fileName.EndsWith(TCV(".target"))
+			//	|| fileName.EndsWith(TCV(".version"));
 		}
 	};
 
@@ -638,17 +638,17 @@ namespace uba
 	{
 		virtual bool IsOutputFile(const StringView& fileName) const override
 		{
-			return fileName.EndsWith(TC(".PVS-Studio.log"))
-				|| fileName.EndsWith(TC(".pvslog"))
-				|| fileName.EndsWith(TC(".stacktrace.txt"));
+			return fileName.EndsWith(TCV(".PVS-Studio.log"))
+				|| fileName.EndsWith(TCV(".pvslog"))
+				|| fileName.EndsWith(TCV(".stacktrace.txt"));
 		}
 		
 		virtual bool IsRarelyRead(const StringView& file) const override
 		{
-			return file.EndsWith(TC(".i"))
-				|| file.EndsWith(TC(".PVS-Studio.log"))
-				|| file.EndsWith(TC(".pvslog"))
-				|| file.EndsWith(TC(".stacktrace.txt"));
+			return file.EndsWith(TCV(".i"))
+				|| file.EndsWith(TCV(".PVS-Studio.log"))
+				|| file.EndsWith(TCV(".pvslog"))
+				|| file.EndsWith(TCV(".stacktrace.txt"));
 		}
 
 #if PLATFORM_WINDOWS
@@ -673,7 +673,7 @@ namespace uba
 
 		virtual bool IsOutputFile(const StringView& fileName) const override
 		{
-			return fileName.EndsWith(TC(".tmp")) || fileName.EndsWith(TC(".out"));
+			return fileName.EndsWith(TCV(".tmp")) || fileName.EndsWith(TCV(".out"));
 		}
 
 		virtual bool AllowDetach() const override
@@ -686,8 +686,8 @@ namespace uba
 	{
 		virtual bool IsOutputFile(const StringView& fileName) const override
 		{
-			return fileName.EndsWith(TC(".obj"))
-				|| fileName.EndsWith(TC(".exp"));
+			return fileName.EndsWith(TCV(".obj"))
+				|| fileName.EndsWith(TCV(".exp"));
 		}
 
 		virtual bool CanDependOnCompressedFiles() override
diff --git a/Engine/Source/Programs/UnrealBuildAccelerator/Core/Private/UbaStringBuffer.cpp b/Engine/Source/Programs/UnrealBuildAccelerator/Core/Private/UbaStringBuffer.cpp
--- a/Engine/Source/Programs/UnrealBuildAccelerator/Core/Private/UbaStringBuffer.cpp
+++ b/Engine/Source/Programs/UnrealBuildAccelerator/Core/Private/UbaStringBuffer.cpp
@@ -29,9 +29,8 @@ namespace uba
 		}
 	}
 
-	bool EndsWith(const tchar* str, u64 strLen, const tchar* value, bool ignoreCase)
+	bool EndsWith(const tchar* str, u64 strLen, const tchar* value, u64 valueLen, bool ignoreCase)
 	{
-		u64 valueLen = TStrlen(value);
 		if (strLen < valueLen)
 			return false;
 		if (!ignoreCase)
@@ -316,6 +315,16 @@ namespace uba
 		return str.count >= count && uba::StartsWith(data, str.data, ignoreCase);
 	}
 
+	bool StringBufferBase::EndsWith(const tchar* value, bool ignoreCase) const
+	{
+		return uba::EndsWith(data, count, value, TStrlen(value), ignoreCase);
+	}
+
+	bool StringBufferBase::EndsWith(const StringView& value, bool ignoreCase) const
+	{
+		return uba::EndsWith(data, count, value.data, value.count, ignoreCase);
+	}
+
 	bool StringBufferBase::Equals(const StringView& str, bool ignoreCase) const
 	{
 		return count == str.count && uba::Equals(data, str.data, count, ignoreCase);
@@ -412,6 +421,11 @@ namespace uba
 		#endif
 	}
 
+	bool StringView::EndsWith(const tchar* value, bool ignoreCase) const
+	{
+		return uba::EndsWith(data, count, value, TStrlen(value), ignoreCase);
+	}
+
 	StringView ToView(const tchar* s)
 	{
 		return StringView(s, TStrlen(s));
diff --git a/Engine/Source/Programs/UnrealBuildAccelerator/Core/Public/UbaStringBuffer.h b/Engine/Source/Programs/UnrealBuildAccelerator/Core/Public/UbaStringBuffer.h
--- a/Engine/Source/Programs/UnrealBuildAccelerator/Core/Public/UbaStringBuffer.h
+++ b/Engine/Source/Programs/UnrealBuildAccelerator/Core/Public/UbaStringBuffer.h
@@ -9,7 +9,7 @@ namespace uba
 	struct StringView;
 
 	bool StartsWith(const tchar* data, const tchar* str, bool ignoreCase = true);
-	bool EndsWith(const tchar* str, u64 strLen, const tchar* value, bool ignoreCase = true);
+	bool EndsWith(const tchar* str, u64 strLen, const tchar* value, u64 valueLen, bool ignoreCase);
 	bool Contains(const tchar* str, const tchar* sub, bool ignoreCase = true, const tchar** pos = nullptr);
 	bool Equals(const tchar* str1, const tchar* str2, bool ignoreCase = true);
 	bool Equals(const tchar* str1, const tchar* str2, u64 count, bool ignoreCase = true);
@@ -55,7 +55,8 @@ namespace uba
 		bool IsEmpty() const { return count == 0; }
 		bool StartsWith(const tchar* str, bool ignoreCase = true) const { return uba::StartsWith(data, str, ignoreCase); }
 		bool StartsWith(const StringView& str, bool ignoreCase = true) const;
-		bool EndsWith(const tchar* value, bool ignoreCase = true) const { return uba::EndsWith(data, count, value, ignoreCase); }
+		bool EndsWith(const tchar* value, bool ignoreCase = true) const;
+		bool EndsWith(const StringView& value, bool ignoreCase = true) const;
 		bool Contains(tchar c) const;
 		bool Contains(const tchar* str, bool ignoreCase = true, const tchar** pos = nullptr) const { return uba::Contains(data, str, ignoreCase, pos); }
 		bool Equals(const tchar* str, bool ignoreCase = true) const { return uba::Equals(data, str, ignoreCase); }
@@ -112,7 +113,8 @@ namespace uba
 
 		bool IsEmpty() const { return count == 0; }
 		bool StartsWith(const tchar* str, bool ignoreCase = true) const { return uba::StartsWith(data, str, ignoreCase); }
-		bool EndsWith(const tchar* value, bool ignoreCase = true) const { return uba::EndsWith(data, count, value, ignoreCase); }
+		bool EndsWith(const StringView& value, bool ignoreCase = true) const { return uba::EndsWith(data, count, value.data, value.count, ignoreCase); }
+		bool EndsWith(const tchar* value, bool ignoreCase = true) const;
 		bool Contains(const tchar* str, bool ignoreCase = true) const { return uba::Contains(data, str, ignoreCase); }
 		bool Equals(const tchar* str, bool ignoreCase = true) const { return uba::Equals(data, str, ignoreCase); }
 		bool Equals(const StringView& str, bool ignoreCase = true) const { return count == str.count && uba::Equals(data, str.data, count, ignoreCase); }
diff --git a/Engine/Source/Programs/UnrealBuildAccelerator/Core/Public/UbaWinBinDependencyParser.h b/Engine/Source/Programs/UnrealBuildAccelerator/Core/Public/UbaWinBinDependencyParser.h
--- a/Engine/Source/Programs/UnrealBuildAccelerator/Core/Public/UbaWinBinDependencyParser.h
+++ b/Engine/Source/Programs/UnrealBuildAccelerator/Core/Public/UbaWinBinDependencyParser.h
@@ -131,7 +131,7 @@ namespace uba
 
 	inline bool FindImports(const tchar* fileName, const Function<void(const tchar* import, bool isKnown, const char* const* importLoaderPaths)>& func, StringBufferBase& outError)
 	{
-		if (EndsWith(fileName, TStrlen(fileName), TC(".bat")))
+		if (ToView(fileName).EndsWith(TC(".bat")))
 			return true;
 		HANDLE fileHandle = True_CreateFileW(fileName, GENERIC_READ, FILE_SHARE_READ, 0, OPEN_EXISTING, FILE_ATTRIBUTE_NORMAL | FILE_FLAG_SEQUENTIAL_SCAN, 0);
 		if (fileHandle == INVALID_HANDLE_VALUE)
