[UBA]
* Added support for path hashes in cache system. It is now possible to register a path with a hash and then all files under that path will fold into one input path/hash. Useful for for example sdk that are strongly versioned and guarantees that inside is matching the hash

diff --git a/Engine/Source/Programs/UnrealBuildAccelerator/Common/Private/UbaCacheClient.cpp b/Engine/Source/Programs/UnrealBuildAccelerator/Common/Private/UbaCacheClient.cpp
--- a/Engine/Source/Programs/UnrealBuildAccelerator/Common/Private/UbaCacheClient.cpp
+++ b/Engine/Source/Programs/UnrealBuildAccelerator/Common/Private/UbaCacheClient.cpp
@@ -156,6 +156,12 @@ namespace uba
 
 	CacheClient::~CacheClient() = default;
 
+	bool CacheClient::RegisterPathHash(const tchar* path, const CasKey& hash)
+	{
+		m_pathHashes.emplace_back(TString(path), AsCompressed(hash, true));
+		return true;
+	}
+
 	bool CacheClient::WriteToCache(u32 bucketId, const ProcessStartInfo& info, const u8* inputs, u64 inputsSize, const u8* outputs, u64 outputsSize, const u8* logLines, u64 logLinesSize, u32 processId)
 	{
 		RootPaths rootPaths;
@@ -210,6 +216,9 @@ namespace uba
 
 		TString qualifiedPath;
 
+		Vector<bool> handledPathHashes;
+		handledPathHashes.resize(m_pathHashes.size());
+
 		UBA_TRACE_WRITE_HINT("TraverseInputsOutputs");
 
 		// Traverse all inputs and outputs. to create cache entry that we can send to server
@@ -264,6 +273,29 @@ namespace uba
 				continue;
 			}
 
+			// Handle path hashes.
+			if (!isOutput)
+			{
+				bool handled = false;
+				for (u32 i=0, e=u32(m_pathHashes.size()); i!=e; ++i)
+				{
+					auto& ph = m_pathHashes[i];
+					if (!path.StartsWith(StringView(ph.path), CaseInsensitiveFs))
+						continue;
+					if (handledPathHashes[i])
+					{
+						handled = true;
+						break;
+					}
+					handledPathHashes[i] = true;
+					path.Clear().Append(ph.path).Append(TCV("<PathHash>"));
+					casKey = ph.hash;
+					break;
+				}
+				if (handled)
+					continue;
+			}
+
 			if (m_useRoots)
 			{
 				// Find root for path in order to be able to normalize it.
@@ -585,7 +617,18 @@ namespace uba
 					u64 t = GetTime();
 					#endif
 
-					m_storage.StoreCasKey(localCasKey, fileNameKey, path.data, CasKeyZero);
+					if (path.EndsWith(TCV("<PathHash>")))
+					{
+						path.Resize(path.count - 10);
+						for (auto& ph : m_pathHashes)
+							if (path.Equals(ph.path) && cacheCasKey == ph.hash)
+							{
+								localCasKey = cacheCasKey;
+								break;
+							}
+					}
+					else
+						m_storage.StoreCasKey(localCasKey, fileNameKey, path.data, CasKeyZero);
 
 					#if UBA_TRACE_FETCH_CACHE
 					storeTime += GetTime() - t;
diff --git a/Engine/Source/Programs/UnrealBuildAccelerator/Common/Public/UbaCacheClient.h b/Engine/Source/Programs/UnrealBuildAccelerator/Common/Public/UbaCacheClient.h
--- a/Engine/Source/Programs/UnrealBuildAccelerator/Common/Public/UbaCacheClient.h
+++ b/Engine/Source/Programs/UnrealBuildAccelerator/Common/Public/UbaCacheClient.h
@@ -55,6 +55,8 @@ namespace uba
 		CacheClient(const CacheClientCreateInfo& info);
 		~CacheClient();
 
+		bool RegisterPathHash(const tchar* path, const CasKey& hash);
+
 		bool WriteToCache(u32 bucketId, const ProcessStartInfo& info, const u8* inputs, u64 inputsSize, const u8* outputs, u64 outputsSize, const u8* logLines, u64 logLinesSize, u32 processId = 0);
 		bool WriteToCache(const RootPaths& rootPaths, u32 bucketId, const ProcessStartInfo& info, const u8* inputs, u64 inputsSize, const u8* outputs, u64 outputsSize, const u8* logLines, u64 logLinesSize, u32 processId = 0);
 		bool FetchFromCache(CacheResult& outResult, RootsHandle rootsHandle, u32 bucketId, const ProcessStartInfo& info);
@@ -117,6 +119,9 @@ namespace uba
 		struct PreparedDir { Futex lock; Atomic<bool> done; };
 		UnorderedMap<StringKey, PreparedDir> m_directoryPreparser;
 
+		struct PathHash { TString path; CasKey hash; };
+		Vector<PathHash> m_pathHashes;
+
 		CacheClient(const CacheClient&) = delete;
 		CacheClient& operator=(const CacheClient&) = delete;
 	};
diff --git a/Engine/Source/Programs/UnrealBuildAccelerator/Host/Private/UbaExports.cpp b/Engine/Source/Programs/UnrealBuildAccelerator/Host/Private/UbaExports.cpp
--- a/Engine/Source/Programs/UnrealBuildAccelerator/Host/Private/UbaExports.cpp
+++ b/Engine/Source/Programs/UnrealBuildAccelerator/Host/Private/UbaExports.cpp
@@ -771,6 +771,15 @@ extern "C"
 		cacheClient->GetStorage().LoadCasTable();
 		return true;
 	}
+	
+	bool CacheClient_RegisterPathHash(uba::CacheClient* cacheClient, const uba::tchar* path, const uba::tchar* hashString)
+	{
+		using namespace uba;
+		CasKeyHasher hasher;
+		hasher.Update(hashString, TStrlen(hashString));
+		cacheClient->RegisterPathHash(path, ToCasKey(hasher, true));
+		return true;
+	}
 
 	void PopulateLogLines(uba::BinaryWriter& out, const uba::ProcessHandle& process)
 	{
diff --git a/Engine/Source/Programs/UnrealBuildAccelerator/Host/Public/UbaExports.h b/Engine/Source/Programs/UnrealBuildAccelerator/Host/Public/UbaExports.h
--- a/Engine/Source/Programs/UnrealBuildAccelerator/Host/Public/UbaExports.h
+++ b/Engine/Source/Programs/UnrealBuildAccelerator/Host/Public/UbaExports.h
@@ -147,6 +147,7 @@ extern "C"
 	// Cache
 	UBA_API uba::CacheClient* CacheClient_Create(uba::SessionServer* session, bool reportMissReason = false, const uba::tchar* crypto = nullptr, const uba::tchar* hint = nullptr);
 	UBA_API bool CacheClient_Connect(uba::CacheClient* cacheClient, const uba::tchar* host, int port);
+	UBA_API bool CacheClient_RegisterPathHash(uba::CacheClient* cacheClient, const uba::tchar* path, const uba::tchar* hashString);
 	UBA_API bool CacheClient_WriteToCache2(uba::CacheClient* cacheClient, uba::u32 bucket, const uba::ProcessHandle* process, const uba::u8* inputs, uba::u32 inputsSize, const uba::u8* outputs, uba::u32 outputsSize);
 	UBA_API uba::CacheResult* CacheClient_FetchFromCache3(uba::CacheClient* cacheClient, uba::u64 rootsHandle, uba::u32 bucket, const uba::ProcessStartInfo& info);
 	UBA_API void CacheClient_RequestServerShutdown(uba::CacheClient* cacheClient, const uba::tchar* reason);
diff --git a/Engine/Source/Programs/UnrealBuildAccelerator/Test/Private/UbaTestCache.cpp b/Engine/Source/Programs/UnrealBuildAccelerator/Test/Private/UbaTestCache.cpp
--- a/Engine/Source/Programs/UnrealBuildAccelerator/Test/Private/UbaTestCache.cpp
+++ b/Engine/Source/Programs/UnrealBuildAccelerator/Test/Private/UbaTestCache.cpp
@@ -568,6 +568,8 @@ namespace uba
 				return logger.Error(TC("Failed to connect"));
 			auto disconnectClient = MakeGuard([&]() { client.Disconnect(); });
 
+			//cacheClient.RegisterPathHash(workingDir.data, CasKey(1, 2, 3));
+
 			{
 				CacheResult result;
 				if (cacheClient.FetchFromCache(result, RootPaths(), 0, psi) || result.hit)
