[UBA]
* Changed all EndsWith(TC( to EndsWith(TCV(
* Changed all Equals(TC( to Equals(TCV(
* Changed all Append(TC( to Append(TCV(

diff --git a/Engine/Source/Programs/UnrealBuildAccelerator/Agent/Private/UbaAgent.cpp b/Engine/Source/Programs/UnrealBuildAccelerator/Agent/Private/UbaAgent.cpp
--- a/Engine/Source/Programs/UnrealBuildAccelerator/Agent/Private/UbaAgent.cpp
+++ b/Engine/Source/Programs/UnrealBuildAccelerator/Agent/Private/UbaAgent.cpp
@@ -171,7 +171,7 @@ namespace uba
 	bool DownloadBinaries(StorageClient& storageClient, CasKey* keys)
 	{
 		StringBuffer<256> binDir(g_rootDir);
-		binDir.Append(TC("\\binaries\\"));
+		binDir.Append(TCV("\\binaries\\"));
 		storageClient.CreateDirectory(binDir.data);
 		u32 index = 0;
 		for (auto file : g_ubaAgentBinaries)
@@ -209,7 +209,7 @@ namespace uba
 
 		StringBuffer<> args;
 		args.Append(g_rootDir).Append(TC("\\binaries\\") UBA_AGENT_EXECUTABLE);
-		args.Append(TC(" -relaunch=\"")).Append(currentDir).Append(TC("\""));
+		args.Append(TCV(" -relaunch=\"")).Append(currentDir).Append(TCV("\""));
 		args.Appendf(TC(" -waitid=%u"), GetCurrentProcessId());
 
 		for (int i = 1; i != argc; ++i)
@@ -404,19 +404,19 @@ namespace uba
 				name.Append(argv[i]);
 			}
 		
-			if (name.Equals(TC("-verbose")))
+			if (name.Equals(TCV("-verbose")))
 			{
 				verbose = true;
 			}
-			else if (name.Equals(TC("-relaunch")))
+			else if (name.Equals(TCV("-relaunch")))
 			{
 				relaunchPath.Append(value);
 			}
-			else if (name.Equals(TC("-waitid")))
+			else if (name.Equals(TCV("-waitid")))
 			{
 				value.Parse(waitProcessId);
 			}
-			else if (name.Equals(TC("-maxcpu")))
+			else if (name.Equals(TCV("-maxcpu")))
 			{
 				if (!ExpandEnvironmentVariables(value, PrintHelp))
 					return false;
@@ -427,7 +427,7 @@ namespace uba
 					maxProcessCount = defaultValue;
 				}
 			}
-			else if (name.Equals(TC("-mulcpu")))
+			else if (name.Equals(TCV("-mulcpu")))
 			{
 				if (!ExpandEnvironmentVariables(value, PrintHelp))
 					return false;
@@ -438,22 +438,22 @@ namespace uba
 					mulProcessValue = defaultValue;
 				}
 			}
-			else if (name.Equals(TC("-maxcon")) || name.Equals(TC("-maxtcp")))
+			else if (name.Equals(TCV("-maxcon")) || name.Equals(TCV("-maxtcp")))
 			{
 				if (!value.Parse(maxConnectionCount) || maxConnectionCount == 0)
 					return PrintHelp(TC("Invalid value for -maxcon"));
 			}
-			else if (name.Equals(TC("-maxworkers")))
+			else if (name.Equals(TCV("-maxworkers")))
 			{
 				if (!value.Parse(maxWorkerCount))
 					return PrintHelp(TC("Invalid value for -maxworkers"));
 			}
-			else if (name.Equals(TC("-capacity")))
+			else if (name.Equals(TCV("-capacity")))
 			{
 				if (!value.Parse(storageCapacityGb))
 					return PrintHelp(TC("Invalid value for -capacity"));
 			}
-			else if (name.Equals(TC("-config")))
+			else if (name.Equals(TCV("-config")))
 			{
 				if (value.IsEmpty())
 					return PrintHelp(TC("-config needs a value"));
@@ -461,7 +461,7 @@ namespace uba
 					return false;
 				configFile.Append(value);
 			}
-			else if (name.Equals(TC("-host")))
+			else if (name.Equals(TCV("-host")))
 			{
 				if (const tchar* portIndex = value.First(':'))
 				{
@@ -474,52 +474,52 @@ namespace uba
 					return PrintHelp(TC("-host needs a name/ip"));
 				host.Append(value);
 			}
-			else if (name.Equals(TC("-listen")))
+			else if (name.Equals(TCV("-listen")))
 			{
 				if (!value.IsEmpty())
 					if (!value.Parse(port))
 						return PrintHelp(TC("Invalid value for -capacity"));
 				useListen = true;
 			}
-			else if (name.Equals(TC("-listenTimeout")))
+			else if (name.Equals(TCV("-listenTimeout")))
 			{
 				if (value.IsEmpty())
 					return PrintHelp(TC("-listenTimeout needs a value"));
 				if (!value.Parse(listenTimeoutSec))
 					return PrintHelp(TC("Invalid value for -listenTimeout"));
 			}
-			else if (name.Equals(TC("-named")))
+			else if (name.Equals(TCV("-named")))
 			{
 				if (value.IsEmpty())
 					return PrintHelp(TC("-named needs a value"));
 				named.Append(value);
 			}
-			else if (name.Equals(TC("-log")))
+			else if (name.Equals(TCV("-log")))
 			{
 				logToFile = true;
 			}
-			else if (name.Equals(TC("-quiet")))
+			else if (name.Equals(TCV("-quiet")))
 			{
 				quiet = true;
 			}
-			else if (name.Equals(TC("-nocustomalloc")))
+			else if (name.Equals(TCV("-nocustomalloc")))
 			{
 				disableCustomAllocator = true;
 			}
-			else if (name.Equals(TC("-storeraw")))
+			else if (name.Equals(TCV("-storeraw")))
 			{
 				storeCompressed = false;
 			}
-			else if (name.Equals(TC("-sendraw")))
+			else if (name.Equals(TCV("-sendraw")))
 			{
 				sendCompressed = false;
 			}
-			else if (name.Equals(TC("-sendsize")))
+			else if (name.Equals(TCV("-sendsize")))
 			{
 				if (!value.Parse(sendSize))
 					return PrintHelp(TC("Invalid value for -sendsize"));
 			}
-			else if (name.Equals(TC("-dir")))
+			else if (name.Equals(TCV("-dir")))
 			{
 				if (value.IsEmpty())
 					return PrintHelp(TC("-dir needs a value"));
@@ -528,59 +528,59 @@ namespace uba
 				if ((g_rootDir.count = GetFullPathNameW(value.Replace('\\', PathSeparator).data, g_rootDir.capacity, g_rootDir.data, nullptr)) == 0)
 					return PrintHelp(StringBuffer<>().Appendf(TC("-dir has invalid path %s"), value.data).data);
 			}
-			else if (name.Equals(TC("-name")))
+			else if (name.Equals(TCV("-name")))
 			{
 				if (value.IsEmpty())
 					return PrintHelp(TC("-name needs a value"));
 				agentName.Append(value);
 			}
-			else if (name.Equals(TC("-nopoll")))
+			else if (name.Equals(TCV("-nopoll")))
 			{
 				poll = false;
 			}
-			else if (name.Equals(TC("-nostore")))
+			else if (name.Equals(TCV("-nostore")))
 			{
 				if (IsWindows) // Only supported on windows atm.
 					useStorage = false;
 			}
-			else if (name.Equals(TC("-nohandler")))
+			else if (name.Equals(TCV("-nohandler")))
 			{
 				useExceptionHandler = false;
 			}
-			else if (name.Equals(TC("-nodetoursdownload")))
+			else if (name.Equals(TCV("-nodetoursdownload")))
 			{
 				downloadDetoursLib = false;
 			}
-			else if (name.Equals(TC("-resetstore")))
+			else if (name.Equals(TCV("-resetstore")))
 			{
 				resetStore = true;
 			}
-			else if (name.Equals(TC("-binasversion")))
+			else if (name.Equals(TCV("-binasversion")))
 			{
 				useBinariesAsVersion = true;
 			}
-			else if (name.Equals(TC("-quic")))
+			else if (name.Equals(TCV("-quic")))
 			{
 				#if !UBA_USE_QUIC
 				return PrintHelp(TC("-quic not supported. Quic is not compiled into this binary"));
 				#endif
 				useQuic = true;
 			}
-			else if (name.Equals(TC("-maxidle")))
+			else if (name.Equals(TCV("-maxidle")))
 			{
 				if (!value.Parse(maxIdleSeconds))
 					return PrintHelp(TC("Invalid value for -maxidle"));
 			}
-			else if (name.Equals(TC("-proxyport")))
+			else if (name.Equals(TCV("-proxyport")))
 			{
 				if (!value.Parse(proxyPort))
 					return PrintHelp(TC("Invalid value for -proxyport"));
 			}
-			else if (name.Equals(TC("-summary")))
+			else if (name.Equals(TCV("-summary")))
 			{
 				printSummary = true;
 			}
-			else if (name.Equals(TC("-eventfile")))
+			else if (name.Equals(TCV("-eventfile")))
 			{
 				if (value.IsEmpty())
 					return PrintHelp(TC("-eventfile needs a value"));
@@ -589,25 +589,25 @@ namespace uba
 				if ((eventFile.count = GetFullPathNameW(value.Replace('\\', PathSeparator).data, eventFile.capacity, eventFile.data, nullptr)) == 0)
 					return PrintHelp(StringBuffer<>().Appendf(TC("-eventfile has invalid path %s"), value.data).data);
 			}
-			else if (name.Equals(TC("-killrandom")))
+			else if (name.Equals(TCV("-killrandom")))
 			{
 				killRandom = true;
 			}
-			else if (name.Equals(TC("-usecrawler")))
+			else if (name.Equals(TCV("-usecrawler")))
 			{
 				useCrawler = true;
 			}
-			else if (name.Equals(TC("-memwait")))
+			else if (name.Equals(TCV("-memwait")))
 			{
 				if (!value.Parse(memWaitLoadPercent) || memWaitLoadPercent > 100)
 					return PrintHelp(TC("Invalid value for -memwait"));
 			}
-			else if (name.Equals(TC("-memkill")))
+			else if (name.Equals(TCV("-memkill")))
 			{
 				if (!value.Parse(memKillLoadPercent) || memKillLoadPercent > 100)
 					return PrintHelp(TC("Invalid value for -memkill"));
 			}
-			else if (name.Equals(TC("-crypto")))
+			else if (name.Equals(TCV("-crypto")))
 			{
 				if (value.count != 32)
 					return PrintHelp(TC("Invalid number of characters in crypto string. Should be 32"));
@@ -615,39 +615,39 @@ namespace uba
 				((u64*)crypto)[1] = StringToValue(value.data + 16, 16);
 				hasCrypto = true;
 			}
-			else if (name.Equals(TC("-resendcas")))
+			else if (name.Equals(TCV("-resendcas")))
 			{
 				resendCas = true;
 			}
-			else if (name.Equals(TC("-populateCas")))
+			else if (name.Equals(TCV("-populateCas")))
 			{
 				if (value.IsEmpty())
 					return PrintHelp(TC("-populateCas needs a dir"));
 				populateCasDirs.push_back(value.data);
 			}
 			#if PLATFORM_MAC
-			else if (name.Equals(TC("-populateCasFromXcodeVersion")))
+			else if (name.Equals(TCV("-populateCasFromXcodeVersion")))
 			{
 				populateCasFromXcodeVersion.Append(value.data);
 			}
-			else if (name.Equals(TC("-populateCasFromAllXcodes")))
+			else if (name.Equals(TCV("-populateCasFromAllXcodes")))
 			{
 				populateCasFromAllXcodes = true;
 			}
 			#endif
-			else if (name.Equals(TC("-sentry")))
+			else if (name.Equals(TCV("-sentry")))
 			{
 				if (value.IsEmpty())
 					return PrintHelp(TC("-sentry needs a url value"));
 				sentryUrl.Append(value);
 			}
-			else if (name.Equals(TC("-zone")))
+			else if (name.Equals(TCV("-zone")))
 			{
 				if (value.IsEmpty())
 					return PrintHelp(TC("-zone needs a value"));
 				zone.Append(value);
 			}
-			else if (name.Equals(TC("-version")))
+			else if (name.Equals(TCV("-version")))
 			{
 				const tchar* dbgStr = TC("");
 				#if UBA_DEBUG
@@ -656,15 +656,15 @@ namespace uba
 				LoggerWithWriter(g_consoleLogWriter, TC("")).Info(TC("v%s%s (Network: %u, Storage: %u, Session: %u, Cache: %u)"), Version, dbgStr, SystemNetworkVersion, StorageNetworkVersion, SessionNetworkVersion, CacheNetworkVersion);
 				return true;
 			}
-			else if (name.Equals(TC("-noproxy")))
+			else if (name.Equals(TCV("-noproxy")))
 			{
 				allowProxy = false;
 			}
-			else if (name.Equals(TC("-proxyuselocalstorage")))
+			else if (name.Equals(TCV("-proxyuselocalstorage")))
 			{
 				proxyUseLocalStorage = true;
 			}
-			else if (name.Equals(TC("-description")))
+			else if (name.Equals(TCV("-description")))
 			{
 				if (value.IsEmpty())
 					return PrintHelp(TC("-description"));
@@ -672,11 +672,11 @@ namespace uba
 					value.Resize(value.count - 1);
 				description = value.data + (value[0] == '\"');
 			}
-			else if (name.Equals(TC("-noaws")))
+			else if (name.Equals(TCV("-noaws")))
 			{
 				couldBeAws = false;
 			}
-			else if (name.Equals(TC("-command")))
+			else if (name.Equals(TCV("-command")))
 			{
 				if (value.IsEmpty())
 					return PrintHelp(TC("-command needs a value"));
@@ -684,7 +684,7 @@ namespace uba
 				poll = false;
 				quiet = true;
 			}
-			else if (name.Equals(TC("-?")) || name.Equals(TC("-help")))
+			else if (name.Equals(TCV("-?")) || name.Equals(TCV("-help")))
 			{
 				return PrintHelp(TC(""));
 			}
@@ -746,7 +746,7 @@ namespace uba
 			sentry_options_set_release(options, release);
 			//sentry_options_set_debug(options, 1);
 			sentry_init(options);
-			extraInfo.Append(TC(", SentryEnabled"));
+			extraInfo.Append(TCV(", SentryEnabled"));
 		}
 		
 		auto sentryGuard = MakeGuard([&]() { if (!sentryUrl.IsEmpty()) sentry_close(); });
@@ -773,15 +773,15 @@ namespace uba
 			GetZone(zone);
 
 		if (IsRunningWine())
-			extraInfo.Append(TC(", Linux/Wine-")).Append(GetWineVersion());
+			extraInfo.Append(TCV(", Linux/Wine-")).Append(GetWineVersion());
 
 		if (useQuic)
-			extraInfo.Append(TC(", MsQuic"));
+			extraInfo.Append(TCV(", MsQuic"));
 		if (hasCrypto)
-			extraInfo.Append(TC(", Encrypted"));
+			extraInfo.Append(TCV(", Encrypted"));
 
 		if (!description.empty())
-			extraInfo.Append(TC(", ")).Append(description);
+			extraInfo.Append(TCV(", ")).Append(description);
 
 		//logger.Info(TC("\033[39m\n"));
 		const tchar* dbgStr = TC("");
@@ -1092,7 +1092,7 @@ namespace uba
 					nsci.workerCount = 192;
 					nsci.receiveTimeoutSeconds = 60;
 
-					proxyNetworkServerPrefix = StringBuffer<256>().Append(TC("UbaProxyServer (")).Append(GuidToString(client->GetUid()).str).Append(')').data;
+					proxyNetworkServerPrefix = StringBuffer<256>().Append(TCV("UbaProxyServer (")).Append(GuidToString(client->GetUid()).str).Append(')').data;
 					bool ctorSuccess = true;
 					auto proxyServer = new NetworkServer(ctorSuccess, nsci, proxyNetworkServerPrefix.c_str());
 					if (!ctorSuccess)
@@ -1173,7 +1173,7 @@ namespace uba
 			info.extraInfo = extraInfo.data;
 			info.deleteSessionsOlderThanSeconds = 1; // Delete all old sessions
 			//if (!awsInstanceId.IsEmpty())
-			//	info.name.Append(TC(" (")).Append(awsInstanceId).Append(TC(")"));
+			//	info.name.Append(TCV(" (")).Append(awsInstanceId).Append(TCV(")"));
 			info.rootDir = g_rootDir.data;
 			info.logToFile = logToFile;
 			info.disableCustomAllocator = disableCustomAllocator;
@@ -1273,7 +1273,7 @@ namespace uba
 				if (storageClient->PopulateCasFromDirs(populateCasDirs, maxProcessCount, [&]() { return isDisconnected.load(); }))
 					sessionClient->SetAllowSpawn(true);
 				else
-					terminationReason.Append(TC("Failed to prepopulate cas from local directories"));
+					terminationReason.Append(TCV("Failed to prepopulate cas from local directories"));
 			}
 
 			if (terminationReason.count)
diff --git a/Engine/Source/Programs/UnrealBuildAccelerator/CacheService/Private/UbaCacheService.cpp b/Engine/Source/Programs/UnrealBuildAccelerator/CacheService/Private/UbaCacheService.cpp
--- a/Engine/Source/Programs/UnrealBuildAccelerator/CacheService/Private/UbaCacheService.cpp
+++ b/Engine/Source/Programs/UnrealBuildAccelerator/CacheService/Private/UbaCacheService.cpp
@@ -132,7 +132,7 @@ namespace uba
 				name.Append(argv[i]);
 			}
 
-			if (name.Equals(TC("-port")))
+			if (name.Equals(TCV("-port")))
 			{
 				if (const tchar* portIndex = value.First(':'))
 				{
@@ -147,50 +147,50 @@ namespace uba
 						return PrintHelp(TC("Invalid value for -port"));
 				}
 			}
-			else if (name.Equals(TC("-dir")))
+			else if (name.Equals(TCV("-dir")))
 			{
 				if (value.IsEmpty())
 					return PrintHelp(TC("-dir needs a value"));
 				if ((g_rootDir.count = GetFullPathNameW(value.Replace('/', PathSeparator).data, g_rootDir.capacity, g_rootDir.data, nullptr)) == 0)
 					return PrintHelp(StringBuffer<>().Appendf(TC("-dir has invalid path %s"), g_rootDir.data).data);
 			}
-			else if (name.Equals(TC("-capacity")))
+			else if (name.Equals(TCV("-capacity")))
 			{
 				if (!value.Parse(storageCapacityGb))
 					return PrintHelp(TC("Invalid value for -capacity"));
 			}
-			else if (name.Equals(TC("-expiration")))
+			else if (name.Equals(TCV("-expiration")))
 			{
 				if (!value.Parse(expirationTimeSeconds))
 					return PrintHelp(TC("Invalid value for -expire"));
 			}
-			else if (name.Equals(TC("-http")))
+			else if (name.Equals(TCV("-http")))
 			{
 				if (!value.Parse(httpPort))
 					httpPort = 80;
 			}
-			else if (name.Equals(TC("-fullmaintenance")))
+			else if (name.Equals(TCV("-fullmaintenance")))
 			{
 				fullMaintenance = true;
 			}
-			else if (name.Equals(TC("-crash")))
+			else if (name.Equals(TCV("-crash")))
 			{
 				shouldCrash = true;
 			}
-			else if (name.Equals(TC("-nosignalhandler")))
+			else if (name.Equals(TCV("-nosignalhandler")))
 			{
 				signalHandlerEnabled = false;
 			}
-			else if (name.Equals(TC("-maxworkers")))
+			else if (name.Equals(TCV("-maxworkers")))
 			{
 				if (!value.Parse(maxWorkerCount))
 					return PrintHelp(TC("Invalid value for -maxworkers"));
 			}
-			else if (name.Equals(TC("-nomaintenance")))
+			else if (name.Equals(TCV("-nomaintenance")))
 			{
 				maintenanceEnabled = false;
 			}
-			else if (name.Equals(TC("-config")))
+			else if (name.Equals(TCV("-config")))
 			{
 				if (value.IsEmpty())
 					return PrintHelp(TC("-config needs a value"));
@@ -198,7 +198,7 @@ namespace uba
 					return false;
 				configFile = value.data;
 			}
-			else if (name.Equals(TC("-?")))
+			else if (name.Equals(TCV("-?")))
 			{
 				return PrintHelp(TC(""));
 			}
diff --git a/Engine/Source/Programs/UnrealBuildAccelerator/Cli/Private/UbaCli.cpp b/Engine/Source/Programs/UnrealBuildAccelerator/Cli/Private/UbaCli.cpp
--- a/Engine/Source/Programs/UnrealBuildAccelerator/Cli/Private/UbaCli.cpp
+++ b/Engine/Source/Programs/UnrealBuildAccelerator/Cli/Private/UbaCli.cpp
@@ -206,43 +206,43 @@ namespace uba
 
 		auto parseOption = [&](const StringView& name, StringBufferBase& value)
 			{
-				if (IsWindows && name.Equals(TC("-visualizer")))
+				if (IsWindows && name.Equals(TCV("-visualizer")))
 				{
 					launchVisualizer = true;
 				}
-				else if (name.Equals(TC("-crypto")))
+				else if (name.Equals(TCV("-crypto")))
 				{
 					if (value.IsEmpty())
-						value.Append(TC("0123456789abcdef0123456789abcdef"));
+						value.Append(TCV("0123456789abcdef0123456789abcdef"));
 					crypto = value.data;
 				}
-				else if (name.Equals(TC("-coordinator")))
+				else if (name.Equals(TCV("-coordinator")))
 				{
 					if (value.IsEmpty())
 						return PrintHelp(TC("-coordinator needs a value"));
 					coordinatorName = value.data;
 				}
-				else if (name.Equals(TC("-pool")))
+				else if (name.Equals(TCV("-pool")))
 				{
 					if (value.IsEmpty())
 						return PrintHelp(TC("-pool needs a value"));
 					coordinatorPool = value.data;
 				}
-				else if (name.Equals(TC("-maxcores")))
+				else if (name.Equals(TCV("-maxcores")))
 				{
 					if (value.IsEmpty())
 						return PrintHelp(TC("-maxcores needs a value"));
 					if (!value.Parse(coordinatorMaxCoreCount))
 						return PrintHelp(TC("Invalid value for -maxcores"));
 				}
-				else if (name.Equals(TC("-workdir")))
+				else if (name.Equals(TCV("-workdir")))
 				{
 					if (value.IsEmpty())
 						return PrintHelp(TC("-workdir needs a value"));
 					if ((workDir.count = GetFullPathNameW(value.data, workDir.capacity, workDir.data, nullptr)) == 0)
 						return PrintHelp(StringBuffer<>().Appendf(TC("-workdir has invalid path %s"), value.data).data);
 				}
-				else if (name.Equals(TC("-config")))
+				else if (name.Equals(TCV("-config")))
 				{
 					if (value.IsEmpty())
 						return PrintHelp(TC("-config needs a value"));
@@ -250,7 +250,7 @@ namespace uba
 						return false;
 					configFile = value.data;
 				}
-				else if (name.Equals(TC("-vfs")))
+				else if (name.Equals(TCV("-vfs")))
 				{
 					if (value.IsEmpty())
 						return PrintHelp(TC("-vfs needs a value"));
@@ -261,12 +261,12 @@ namespace uba
 					const tchar* localPath = semi + 1;
 					vfsEntries.push_back({TString(value.data), TString(localPath)});
 				}
-				else if (name.Equals(TC("-capacity")))
+				else if (name.Equals(TCV("-capacity")))
 				{
 					if (!value.Parse(storageCapacityGb))
 						return PrintHelp(TC("Invalid value for -capacity"));
 				}
-				else if (name.Equals(TC("-port")))
+				else if (name.Equals(TCV("-port")))
 				{
 					if (const tchar* portIndex = value.First(':'))
 					{
@@ -281,37 +281,37 @@ namespace uba
 							return PrintHelp(TC("Invalid value for -port"));
 					}
 				}
-				else if (name.Equals(TC("-log")))
+				else if (name.Equals(TCV("-log")))
 				{
 					logToFile = true;
 				}
-				else if (name.Equals(TC("-loop")))
+				else if (name.Equals(TCV("-loop")))
 				{
 					if (!value.Parse(loopCount))
 						return PrintHelp(TC("Invalid value for -loop"));
 				}
-				else if (name.Equals(TC("-quiet")))
+				else if (name.Equals(TCV("-quiet")))
 				{
 					quiet = true;
 				}
-				else if (name.Equals(TC("-nocustomalloc")))
+				else if (name.Equals(TCV("-nocustomalloc")))
 				{
 					disableCustomAllocator = true;
 				}
-				else if (name.Equals(TC("-maxcpu")))
+				else if (name.Equals(TCV("-maxcpu")))
 				{
 					if (!value.Parse(maxProcessCount))
 						return PrintHelp(TC("Invalid value for -maxcpu"));
 				}
-				else if (name.Equals(TC("-nostdout")))
+				else if (name.Equals(TCV("-nostdout")))
 				{
 					enableStdOut = false;
 				}
-				else if (name.Equals(TC("-checkcas")))
+				else if (name.Equals(TCV("-checkcas")))
 				{
 					checkCas = true;
 				}
-				else if (name.Equals(TC("-checkfiletable")))
+				else if (name.Equals(TCV("-checkfiletable")))
 				{
 					if (value.IsEmpty())
 						return PrintHelp(TC("-checkfiletable needs a value"));
@@ -320,51 +320,51 @@ namespace uba
 						return PrintHelp(StringBuffer<>().Appendf(TC("-checkfiletable has invalid path %s"), temp.data).data);
 					checkFileTable = temp.data;
 				}
-				else if (name.Equals(TC("-checkcas2")))
+				else if (name.Equals(TCV("-checkcas2")))
 				{
 					checkCas2 = true;
 				}
-				else if (name.Equals(TC("-checkaws")))
+				else if (name.Equals(TCV("-checkaws")))
 				{
 					checkAws = true;
 				}
-				else if (name.Equals(TC("-testcompress")))
+				else if (name.Equals(TCV("-testcompress")))
 				{
 					if (value.IsEmpty())
 						return PrintHelp(TC("-testCompress needs a value"));
 					testCompress = value.data;
 				}
-				else if (name.Equals(TC("-testdecompress")))
+				else if (name.Equals(TCV("-testdecompress")))
 				{
 					if (value.IsEmpty())
 					{
 						if (testCompress.empty())
 							return PrintHelp(TC("-testDecompress needs a value"));
-						value.Clear().Append(g_rootDir).EnsureEndsWithSlash().Append(TC("castemp")).EnsureEndsWithSlash().Append(TC("TestCompress.tmp"));
+						value.Clear().Append(g_rootDir).EnsureEndsWithSlash().Append(TCV("castemp")).EnsureEndsWithSlash().Append(TCV("TestCompress.tmp"));
 					}
 					testDecompress = value.data;
 				}
-				else if (name.Equals(TC("-deletecas")))
+				else if (name.Equals(TCV("-deletecas")))
 				{
 					deleteCas = true;
 				}
-				else if (name.Equals(TC("-addcas")))
+				else if (name.Equals(TCV("-addcas")))
 				{
 					addCas = value.data;
 				}
-				else if (name.Equals(TC("-getcas")))
+				else if (name.Equals(TCV("-getcas")))
 				{
 					getCas = true;
 				}
-				else if (name.Equals(TC("-summary")))
+				else if (name.Equals(TCV("-summary")))
 				{
 					printSummary = true;
 				}
-				else if (name.Equals(TC("-hackvfs")))
+				else if (name.Equals(TCV("-hackvfs")))
 				{
 					useHackVfs = true;
 				}
-				else if (name.Equals(TC("-cache")))
+				else if (name.Equals(TCV("-cache")))
 				{
 					if (value.IsEmpty())
 						return PrintHelp(TC("-cache needs a value"));
@@ -376,11 +376,11 @@ namespace uba
 					else
 						cacheHost.Append(value);
 				}
-				else if (name.Equals(TC("-populatecache")))
+				else if (name.Equals(TCV("-populatecache")))
 				{
 					populateCache = true;
 				}
-				else if (name.Equals(TC("-cachecommand")))
+				else if (name.Equals(TCV("-cachecommand")))
 				{
 					if (value.IsEmpty())
 						return PrintHelp(TC("-cachecommand needs a value"));
@@ -388,24 +388,24 @@ namespace uba
 					commandType = CommandType_None;
 					quiet = true;
 				}
-				else if (name.Equals(TC("-writecachesummary")))
+				else if (name.Equals(TCV("-writecachesummary")))
 				{
 					writeCacheSummary = true;
 					cacheFilterString = value.data;
 					commandType = CommandType_None;
 				}
-				else if (name.Equals(TC("-storeraw")))
+				else if (name.Equals(TCV("-storeraw")))
 				{
 					storeCompressed = false;
 				}
-				else if (name.Equals(TC("-dir")))
+				else if (name.Equals(TCV("-dir")))
 				{
 					if (value.IsEmpty())
 						return PrintHelp(TC("-dir needs a value"));
 					if ((g_rootDir.count = GetFullPathNameW(value.Replace('/', PathSeparator).data, g_rootDir.capacity, g_rootDir.data, nullptr)) == 0)
 						return PrintHelp(StringBuffer<>().Appendf(TC("-dir has invalid path %s"), g_rootDir.data).data);
 				}
-				else if (name.Equals(TC("-?")))
+				else if (name.Equals(TCV("-?")))
 				{
 					return PrintHelp(TC(""));
 				}
@@ -461,19 +461,19 @@ namespace uba
 			{
 				application = arg;
 			}
-			else if (name.Equals(TC("local")))
+			else if (name.Equals(TCV("local")))
 			{
 				commandType = CommandType_Local;
 			}
-			else if (name.Equals(TC("remote")))
+			else if (name.Equals(TCV("remote")))
 			{
 				commandType = CommandType_Remote;
 			}
-			else if (name.Equals(TC("native")))
+			else if (name.Equals(TCV("native")))
 			{
 				commandType = CommandType_Native;
 			}
-			else if (name.Equals(TC("agent")))
+			else if (name.Equals(TCV("agent")))
 			{
 				commandType = CommandType_Agent;
 			}
@@ -656,7 +656,7 @@ namespace uba
 			CompressedFileHeader header { CalculateCasKey(mem, fileSize, true, &workManager, testCompress.c_str()) };
 
 			StringBuffer<> dest;
-			dest.Append(storage.GetTempPath()).Append(TC("TestCompress.tmp"));
+			dest.Append(storage.GetTempPath()).Append(TCV("TestCompress.tmp"));
 			if (!storage.WriteCompressed(res, TC("MemoryMap"), InvalidFileHandle, mem, fileSize, dest.data, &header, sizeof(header), 0))
 				return false;
 			if (testDecompress.empty())
@@ -695,7 +695,7 @@ namespace uba
 				return logger.Error(TC("Compressed file %s has wrong decompressed size. (Is it compressed?)"), testDecompress.c_str());
 
 			StringBuffer<> dest;
-			dest.Append(storage.GetTempPath()).Append(TC("TestDecompress.tmp"));
+			dest.Append(storage.GetTempPath()).Append(TCV("TestDecompress.tmp"));
 			FileAccessor faDest(logger, dest.data);
 			if (!faDest.CreateMemoryWrite(false, DefaultAttributes(), decompressedSize))
 				return false;
@@ -817,7 +817,7 @@ namespace uba
 		if (logToFile)
 		{
 			logFile.count = GetFullPathNameW(g_rootDir.data, logFile.capacity, logFile.data, nullptr);
-			logFile.EnsureEndsWithSlash().Append(TC("DebugLog.log"));
+			logFile.EnsureEndsWithSlash().Append(TCV("DebugLog.log"));
 			logger.Info(TC("Logging to file: %s"), logFile.data);
 		}
 
@@ -850,7 +850,7 @@ namespace uba
 
 
 		bool isRemote = commandType == CommandType_Remote || commandType == CommandType_Agent;
-		bool useScheduler = EndsWith(application.c_str(), application.size(), TC(".yaml"));
+		bool useScheduler = StringView(application).EndsWith(TCV(".yaml"));
 
 		StorageServerCreateInfo storageInfo(networkServer, g_rootDir.data, logWriter);
 		storageInfo.casCapacityBytes = storageCapacity;
@@ -867,7 +867,7 @@ namespace uba
 		info.disableCustomAllocator = disableCustomAllocator;
 		//info.shouldWriteToDisk = shouldWriteToDisk;
 		info.rootDir = g_rootDir.data;
-		//info.traceName.Append(TC("TESTTRACE"));
+		//info.traceName.Append(TCV("TESTTRACE"));
 		//info.storeIntermediateFilesCompressed = false;
 		info.readIntermediateFilesCompressed = true;
 		//info.extractObjFilesSymbols = true;
@@ -912,7 +912,7 @@ namespace uba
 				StringBuffer<> tempFile(sessionServer.GetTempPath());
 				Guid guid;
 				CreateGuid(guid);
-				tempFile.Append(GuidToString(guid).str).Append(TC(".txt"));
+				tempFile.Append(GuidToString(guid).str).Append(TCV(".txt"));
 				if (!cacheClient->ExecuteCommand(logger, TC("content"), tempFile.data, cacheFilterString.data()))
 					return false;
 				logger.Info(TC("Cache status summary written to %s"), tempFile.data);
@@ -1118,11 +1118,11 @@ namespace uba
 					logger.BeginScope();
 					StringBuffer<128> extra;
 					if (ph.IsRemote())
-						extra.Append(TC(" [RemoteExecutor: ")).Append(ph.GetExecutingHost()).Append(']');
+						extra.Append(TCV(" [RemoteExecutor: ")).Append(ph.GetExecutingHost()).Append(']');
 					else if (ph.GetExecutionType() == ProcessExecutionType_Native)
-						extra.Append(TC(" (Not detoured)"));
+						extra.Append(TCV(" (Not detoured)"));
 					else if (ph.GetExecutionType() == ProcessExecutionType_FromCache)
-						extra.Append(TC(" (From cache)"));
+						extra.Append(TCV(" (From cache)"));
 					logger.Info(TC("[%u/%u] %s%s"), c, queued, desc, extra.data);
 					for (auto& line : ph.GetLogLines())
 						if (line.text != desc && !StartsWith(line.text.c_str(), TC("   Creating library")))
@@ -1183,7 +1183,7 @@ namespace uba
 		{
 			StringBuffer<512> sdkbin;
 			//sdkbin.count = GetEnvironmentVariable(TC("WindowsSdkVerBinPath"), sdkbin.data, sdkbin.capacity);
-			sdkbin.Append(TC(";C:\\Program Files (x86)\\Windows Kits\\10\\bin\\10.0.22621.0\\x64"));
+			sdkbin.Append(TCV(";C:\\Program Files (x86)\\Windows Kits\\10\\bin\\10.0.22621.0\\x64"));
 			if (sdkbin.count)
 			{
 				StringBuffer<4096> temp;
diff --git a/Engine/Source/Programs/UnrealBuildAccelerator/Cli/Private/UbaClient.h b/Engine/Source/Programs/UnrealBuildAccelerator/Cli/Private/UbaClient.h
--- a/Engine/Source/Programs/UnrealBuildAccelerator/Cli/Private/UbaClient.h
+++ b/Engine/Source/Programs/UnrealBuildAccelerator/Cli/Private/UbaClient.h
@@ -76,7 +76,7 @@ namespace uba
 			nsci.receiveTimeoutSeconds = 60;
 
 			StringBuffer<256> prefix;
-			prefix.Append(TC("UbaProxyServer (")).Append(GuidToString(networkClient->GetUid()).str).Append(')');
+			prefix.Append(TCV("UbaProxyServer (")).Append(GuidToString(networkClient->GetUid()).str).Append(')');
 			serverPrefix = prefix.data;
 			bool ctorSuccess = true;
 			proxyNetworkServer = new NetworkServer(ctorSuccess, nsci, serverPrefix.c_str());
diff --git a/Engine/Source/Programs/UnrealBuildAccelerator/Common/Private/UbaCacheClient.cpp b/Engine/Source/Programs/UnrealBuildAccelerator/Common/Private/UbaCacheClient.cpp
--- a/Engine/Source/Programs/UnrealBuildAccelerator/Common/Private/UbaCacheClient.cpp
+++ b/Engine/Source/Programs/UnrealBuildAccelerator/Common/Private/UbaCacheClient.cpp
@@ -243,7 +243,7 @@ namespace uba
 
 			// For .exe and .dll we sometimes get relative paths so we need to expand them to full
 			#if PLATFORM_WINDOWS
-			if (path[1] != ':' && (path.EndsWith(TC(".dll")) || path.EndsWith(TC(".exe"))))
+			if (path[1] != ':' && (path.EndsWith(TCV(".dll")) || path.EndsWith(TCV(".exe"))))
 			{
 				tchar temp[512];
 				bool res = SearchPathW(NULL, path.data, NULL, 512, temp, NULL);
@@ -928,7 +928,7 @@ namespace uba
 
 		// TODO: Remove when bug is found
 #if 0
-		if (path.EndsWith(TC(".obj")) || path.EndsWith(TC(".o")))
+		if (path.EndsWith(TCV(".obj")) || path.EndsWith(TCV(".o")))
 		{
 			if (fetcher.sizeOnDisk == 0)
 				return m_logger.Warning(TC("Received %s file that has size 0!"), path.data);
@@ -1485,17 +1485,17 @@ namespace uba
 	{
 		if (!m_useRoots)
 			return false;
-		if (path.EndsWith(TC(".json"))) // Contains absolute paths (dep files for msvc and vfsoverlay files for clang)
+		if (path.EndsWith(TCV(".json"))) // Contains absolute paths (dep files for msvc and vfsoverlay files for clang)
 			return true;
-		if (path.EndsWith(TC(".d"))) // Contains absolute paths (dep files for clang)
+		if (path.EndsWith(TCV(".d"))) // Contains absolute paths (dep files for clang)
 			return true;
-		if (path.EndsWith(TC(".tlh"))) // Contains absolute path in a comment
+		if (path.EndsWith(TCV(".tlh"))) // Contains absolute path in a comment
 			return true;
-		if (path.EndsWith(TC(".rsp"))) // Contains absolute paths in some cases
+		if (path.EndsWith(TCV(".rsp"))) // Contains absolute paths in some cases
 			return true;
-		if (path.EndsWith(TC(".bat"))) // Contains absolute paths in some cases
+		if (path.EndsWith(TCV(".bat"))) // Contains absolute paths in some cases
 			return true;
-		if (path.EndsWith(TC(".txt"))) // Contains absolute paths (ispc dependency file)
+		if (path.EndsWith(TCV(".txt"))) // Contains absolute paths (ispc dependency file)
 			return true;
 		return false;
 	}
diff --git a/Engine/Source/Programs/UnrealBuildAccelerator/Common/Private/UbaFile.cpp b/Engine/Source/Programs/UnrealBuildAccelerator/Common/Private/UbaFile.cpp
--- a/Engine/Source/Programs/UnrealBuildAccelerator/Common/Private/UbaFile.cpp
+++ b/Engine/Source/Programs/UnrealBuildAccelerator/Common/Private/UbaFile.cpp
@@ -992,7 +992,7 @@ namespace uba
 		#else
 		const tchar* platformStr = TC("Mac");
 		#endif
-		out.Resize(engineDirPos - out.data + 8).Append(TC("Binaries")).Append(PathSeparator).Append(platformStr).Append(PathSeparator).Append("UnrealBuildAccelerator").Append(PathSeparator);
+		out.Resize(engineDirPos - out.data + 8).Append(TCV("Binaries")).Append(PathSeparator).Append(platformStr).Append(PathSeparator).Append("UnrealBuildAccelerator").Append(PathSeparator);
 		if (IsWindows)
 			out.Append(isWindowsArm ? TC("arm64") : TC("x64")).Append(PathSeparator);
 		return true;
diff --git a/Engine/Source/Programs/UnrealBuildAccelerator/Common/Private/UbaLogger.cpp b/Engine/Source/Programs/UnrealBuildAccelerator/Common/Private/UbaLogger.cpp
--- a/Engine/Source/Programs/UnrealBuildAccelerator/Common/Private/UbaLogger.cpp
+++ b/Engine/Source/Programs/UnrealBuildAccelerator/Common/Private/UbaLogger.cpp
@@ -522,7 +522,7 @@ namespace uba
 			mod.handled = false;
 		}
 
-		out.Append(TC("\n CALLSTACK")).Append(isRunningWine ? TC(" (Wine)") : TC("")).Append(':');
+		out.Append(TCV("\n CALLSTACK")).Append(isRunningWine ? TC(" (Wine)") : TC("")).Append(':');
 
 		if (entries.empty())
 		{
@@ -566,7 +566,7 @@ namespace uba
 
 		for (auto& entry : entries)
 		{
-			out.Append(TC("\n   "));
+			out.Append(TCV("\n   "));
 			if (entry.moduleIndex == ~0u)
 			{
 				out.Appendf(TC("<Unknown>: 0x%llx"), entry.memoryOffset);
@@ -707,7 +707,7 @@ namespace uba
 			StringBuffer<512> fn;
 			fn.Append(ct->file);
 			StringBuffer<256> s;
-			s.Append(TC("  ")).AppendFileName(fn.data).Append(':').AppendValue(ct->line).Append(TC(" - ")).AppendValue(ct->count).Append(TC(" (")).Append(TimeToText(ct->time).str).Append(')');
+			s.Append(TCV("  ")).AppendFileName(fn.data).Append(':').AppendValue(ct->line).Append(TCV(" - ")).AppendValue(ct->count).Append(TCV(" (")).Append(TimeToText(ct->time).str).Append(')');
 			logger.Info(s.data);
 			ct->count = 0;
 			ct->time = 0;
diff --git a/Engine/Source/Programs/UnrealBuildAccelerator/Common/Private/UbaObjectFile.cpp b/Engine/Source/Programs/UnrealBuildAccelerator/Common/Private/UbaObjectFile.cpp
--- a/Engine/Source/Programs/UnrealBuildAccelerator/Common/Private/UbaObjectFile.cpp
+++ b/Engine/Source/Programs/UnrealBuildAccelerator/Common/Private/UbaObjectFile.cpp
@@ -229,11 +229,11 @@ namespace uba
 		MemoryBlock memoryBlock(16*1024*1024);
 
 		bool res;
-		if (platform.Equals(TC("win64")) || platform.Equals(TC("wingdk")) || platform.Equals(TC("xb1")) || platform.Equals(TC("xsx"))) 
+		if (platform.Equals(TCV("win64")) || platform.Equals(TCV("wingdk")) || platform.Equals(TCV("xb1")) || platform.Equals(TCV("xsx"))) 
 			res = ObjectFileCoff::CreateExtraFile(logger, platform, memoryBlock, allExternalImports, allInternalImports, allExports, includeExportsInFile);
-		else if (extraObjFilename.EndsWith(TC("dynlist")))
+		else if (extraObjFilename.EndsWith(TCV("dynlist")))
 			res = CreateDynamicListFile(logger, memoryBlock, allExternalImports, allInternalImports, allExports, includeExportsInFile);
-		else if (extraObjFilename.EndsWith(TC("emd")))
+		else if (extraObjFilename.EndsWith(TCV("emd")))
 			res = CreateEmdFile(logger, memoryBlock, moduleName, allExternalImports, allInternalImports, allExports, includeExportsInFile);
 		else
 			res = ObjectFileElf::CreateExtraFile(logger, platform, memoryBlock, allExternalImports, allInternalImports, allExports, includeExportsInFile);
diff --git a/Engine/Source/Programs/UnrealBuildAccelerator/Common/Private/UbaObjectFileElf.cpp b/Engine/Source/Programs/UnrealBuildAccelerator/Common/Private/UbaObjectFileElf.cpp
--- a/Engine/Source/Programs/UnrealBuildAccelerator/Common/Private/UbaObjectFileElf.cpp
+++ b/Engine/Source/Programs/UnrealBuildAccelerator/Common/Private/UbaObjectFileElf.cpp
@@ -350,7 +350,7 @@ namespace uba
 		header.e_ident[4] = 2;
 		header.e_ident[5] = 1;
 		header.e_ident[6] = 1;
-		header.e_ident[7] = platform.Equals(TC("PS4")) ? 9 : 0; // OS ABI identification
+		header.e_ident[7] = platform.Equals(TCV("PS4")) ? 9 : 0; // OS ABI identification
 		header.e_type = 1;
 		header.e_machine = EM_X86_64;
 		header.e_version = 1;
diff --git a/Engine/Source/Programs/UnrealBuildAccelerator/Common/Private/UbaProcess.cpp b/Engine/Source/Programs/UnrealBuildAccelerator/Common/Private/UbaProcess.cpp
--- a/Engine/Source/Programs/UnrealBuildAccelerator/Common/Private/UbaProcess.cpp
+++ b/Engine/Source/Programs/UnrealBuildAccelerator/Common/Private/UbaProcess.cpp
@@ -292,7 +292,7 @@ namespace uba
 
 			if (m_messageCount == 0) // This is bad.. bad binaries?
 			{
-				err.Append(TC("ERROR: Process did not start properly. "));
+				err.Append(TCV("ERROR: Process did not start properly. "));
 
 				ImageInfo imageInfo;
 				bool machineIsArm = IsArmBinary;
@@ -301,11 +301,11 @@ namespace uba
 				else if (!GetImageInfo(imageInfo, m_session.m_logger, m_realApplication.c_str(), true))
 					err.Appendf(TC("Failed to load %s"), m_realApplication.c_str());
 				else if (!imageInfo.is64Bit)
-					err.Append(TC("Doesn't seem to be a 64-bit executable"));
+					err.Append(TCV("Doesn't seem to be a 64-bit executable"));
 				else if (imageInfo.isDotnet)
-					err.Append(TC("Dotnet binary"));
+					err.Append(TCV("Dotnet binary"));
 				else if (!imageInfo.isArm64 && !imageInfo.isX64)
-					err.Append(TC("Unknown image architecture"));
+					err.Append(TCV("Unknown image architecture"));
 				else if (!machineIsArm && imageInfo.isArm64)
 					err.Appendf(TC("Machine is x64 and image is arm64"));
 
@@ -1518,7 +1518,7 @@ namespace uba
 			m_session.VirtualizePath(application, m_startInfo.rootsHandle);
 
 			TString commandLine;
-			if (application.EndsWith(TC(".bat")) && !IsRunningWine()) // If there are quotes around arguments used by batch file we need to quote the entire thing on windows
+			if (application.EndsWith(TCV(".bat")) && !IsRunningWine()) // If there are quotes around arguments used by batch file we need to quote the entire thing on windows
 				commandLine = TString(TC("\"\"")) + application.data + TC("\" ") + m_startInfo.arguments + TC("\"");
 			else
 				commandLine = TString(TC("\"")) + application.data + TC("\" ") + m_startInfo.arguments;
diff --git a/Engine/Source/Programs/UnrealBuildAccelerator/Common/Private/UbaRootPaths.cpp b/Engine/Source/Programs/UnrealBuildAccelerator/Common/Private/UbaRootPaths.cpp
--- a/Engine/Source/Programs/UnrealBuildAccelerator/Common/Private/UbaRootPaths.cpp
+++ b/Engine/Source/Programs/UnrealBuildAccelerator/Common/Private/UbaRootPaths.cpp
@@ -135,7 +135,7 @@ namespace uba
 			{
 				forwardSlash.Append('/');
 				backwardSlash.Append('\\');
-				doubleBackwardSlash.Append(TC("\\\\"));
+				doubleBackwardSlash.Append(TCV("\\\\"));
 				spaceEscapedBackwardSlash.Append('\\');
 			}
 			else if (c == '\\')
diff --git a/Engine/Source/Programs/UnrealBuildAccelerator/Common/Private/UbaScheduler.cpp b/Engine/Source/Programs/UnrealBuildAccelerator/Common/Private/UbaScheduler.cpp
--- a/Engine/Source/Programs/UnrealBuildAccelerator/Common/Private/UbaScheduler.cpp
+++ b/Engine/Source/Programs/UnrealBuildAccelerator/Common/Private/UbaScheduler.cpp
@@ -657,7 +657,7 @@ namespace uba
 				{
 					static u32 processId = 1; // TODO: This should be done in a better way.. or not at all?
 					GenerateNameForProcess(logFile, si.arguments, ++processId);
-					logFile.Append(TC(".log"));
+					logFile.Append(TCV(".log"));
 					si.logFile = logFile.data;
 				};
 				#endif
@@ -725,19 +725,19 @@ namespace uba
 				{
 				case InsideArray_None:
 				{
-					if (key.Equals(TC("environment")))
+					if (key.Equals(TCV("environment")))
 					{
 						#if PLATFORM_WINDOWS
 						SetEnvironmentVariable(TC("PATH"), valueStart);
 						#endif
 						return true;
 					}
-					if (key.Equals(TC("cacheroots")))
+					if (key.Equals(TCV("cacheroots")))
 					{
 						insideArray = InsideArray_CacheRoots;
 						return true;
 					}
-					if (key.Equals(TC("processes")))
+					if (key.Equals(TCV("processes")))
 					{
 						insideArray = InsideArray_Processes;
 						return true;
@@ -762,21 +762,21 @@ namespace uba
 							enqueueProcess();
 					}
 
-					if (key.Equals(TC("app")))
+					if (key.Equals(TCV("app")))
 						app = valueStart;
-					else if (key.Equals(TC("arg")))
+					else if (key.Equals(TCV("arg")))
 						arg = valueStart;
-					else if (key.Equals(TC("dir")))
+					else if (key.Equals(TCV("dir")))
 						dir = valueStart;
-					else if (key.Equals(TC("desc")))
+					else if (key.Equals(TCV("desc")))
 						desc = valueStart;
-					else if (key.Equals(TC("detour")))
+					else if (key.Equals(TCV("detour")))
 						allowDetour = !Equals(valueStart, TC("false"));
-					else if (key.Equals(TC("remote")))
+					else if (key.Equals(TCV("remote")))
 						allowRemote = !Equals(valueStart, TC("false"));
-					else if (key.Equals(TC("weight")))
+					else if (key.Equals(TCV("weight")))
 						StringBuffer<32>(valueStart).Parse(weight);
-					else if (key.Equals(TC("dep")))
+					else if (key.Equals(TCV("dep")))
 					{
 						const tchar* depStart = TStrchr(valueStart, '[');
 						if (!depStart)
diff --git a/Engine/Source/Programs/UnrealBuildAccelerator/Common/Private/UbaSession.cpp b/Engine/Source/Programs/UnrealBuildAccelerator/Common/Private/UbaSession.cpp
--- a/Engine/Source/Programs/UnrealBuildAccelerator/Common/Private/UbaSession.cpp
+++ b/Engine/Source/Programs/UnrealBuildAccelerator/Common/Private/UbaSession.cpp
@@ -282,7 +282,7 @@ namespace uba
 				if (false)
 				{
 					StringBuffer<> name2;
-					name2.Append(TC("\\??\\")).Append(file.name);
+					name2.Append(TCV("\\??\\")).Append(file.name);
 
 					u8 buffer[2048];
 					auto& info = *(FILE_RENAME_INFO*)buffer;;
@@ -504,9 +504,9 @@ namespace uba
 		m_debugLogger->BeginScope();
 		auto dg = MakeGuard([&]() { m_debugLogger->EndScope(); });
 		StringBuffer<> str;
-		str.Append(TC("TRACKDIR "));
+		str.Append(TCV("TRACKDIR "));
 		if (isRefresh)
-			str.Append(TC("(Refresh) "));
+			str.Append(TCV("(Refresh) "));
 		str.Append(dirPath).Append('\n');
 		m_debugLogger->Log(LogEntryType_Info, str);
 		#endif
@@ -870,9 +870,9 @@ namespace uba
 		if (isDir || size == 0)
 		{
 			if (isDir)
-				out.name.Append(TC("$d"));
+				out.name.Append(TCV("$d"));
 			else
-				out.name.Append(TC("$f"));
+				out.name.Append(TCV("$f"));
 		}
 		else
 		{
@@ -970,7 +970,7 @@ namespace uba
 	{
 		//StringBuffer<> workName;
 		//u32 len = TStrlen(fileName);
-		//workName.Append(TC("MM:")).Append(len > 30 ? fileName + (len - 30) : fileName);
+		//workName.Append(TCV("MM:")).Append(len > 30 ? fileName + (len - 30) : fileName);
 		//TrackWorkScope tws(m_workManager, workName.data);
 
 		SCOPED_FUTEX(m_fileMappingTableLookupLock, lookupLock);
@@ -1012,7 +1012,7 @@ namespace uba
 
 			if (size == 0)
 			{
-				out.name.Append(TC("$f"));
+				out.name.Append(TCV("$f"));
 			}
 			else
 			{
@@ -1454,7 +1454,7 @@ namespace uba
 		}
 		else
 		{
-			m_id.Append(TC("Debug"));
+			m_id.Append(TCV("Debug"));
 		}
 
 		UBA_ASSERTF(info.rootDir && *info.rootDir, TC("No root dir set when creating session"));
@@ -1489,7 +1489,7 @@ namespace uba
 		if (info.deleteSessionsOlderThanSeconds)
 		{
 			StringBuffer<> sessionsDir;
-			sessionsDir.Append(m_rootDir).Append(TC("sessions"));
+			sessionsDir.Append(m_rootDir).Append(TCV("sessions"));
 
 			u64 systemTimeAsFileTime = GetSystemTimeAsFileTime();
 
@@ -1509,10 +1509,10 @@ namespace uba
 				});
 		}
 
-		m_sessionDir.Append(m_rootDir).Append(TC("sessions")).Append(PathSeparator).Append(m_id).Append(PathSeparator);
-		m_sessionBinDir.Append(m_sessionDir).Append(TC("bin"));
-		m_sessionOutputDir.Append(m_sessionDir).Append(TC("output"));
-		m_sessionLogDir.Append(m_sessionDir).Append(TC("log"));
+		m_sessionDir.Append(m_rootDir).Append(TCV("sessions")).Append(PathSeparator).Append(m_id).Append(PathSeparator);
+		m_sessionBinDir.Append(m_sessionDir).Append(TCV("bin"));
+		m_sessionOutputDir.Append(m_sessionDir).Append(TCV("output"));
+		m_sessionLogDir.Append(m_sessionDir).Append(TCV("log"));
 
 		if (m_runningRemote)
 		{
@@ -1520,7 +1520,7 @@ namespace uba
 			m_storage.CreateDirectory(m_sessionOutputDir.data);
 		}
 
-		m_tempPath.Append(m_sessionDir).Append(TC("temp"));
+		m_tempPath.Append(m_sessionDir).Append(TCV("temp"));
 		m_storage.CreateDirectory(m_tempPath.data);
 		m_tempPath.EnsureEndsWithSlash();
 
@@ -1557,13 +1557,13 @@ namespace uba
 	bool Session::Create(const SessionCreateInfo& info)
 	{
 		#if UBA_DEBUG_LOGGER
-		m_debugLogger = StartDebugLogger(m_logger, StringBuffer<512>().Append(m_sessionDir).Append(TC("SessionDebug.log")).data);
+		m_debugLogger = StartDebugLogger(m_logger, StringBuffer<512>().Append(m_sessionDir).Append(TCV("SessionDebug.log")).data);
 		#endif
 
 		#if PLATFORM_WINDOWS
 		m_systemPath.count = GetEnvironmentVariableW(TC("SystemRoot"), m_systemPath.data, m_systemPath.capacity);
 		#else
-		m_systemPath.Append(TC("/nonexistingpath"));
+		m_systemPath.Append(TCV("/nonexistingpath"));
 		#endif
 
 		m_fileMappingTableHandle = uba::CreateMemoryMappingW(m_logger, PAGE_READWRITE, FileMappingTableMemSize, nullptr, TC("FileMappings"));
@@ -1635,10 +1635,10 @@ namespace uba
 			tchar fileName[512];
 			GetModuleFileNameW(currentModule, fileName, 512);
 			uba::StringBuffer<> launcherCmd;
-			launcherCmd.Append(TC("\""));
+			launcherCmd.Append(TCV("\""));
 			launcherCmd.AppendDir(fileName);
-			launcherCmd.Append(TC("\\UbaVisualizer.exe\""));
-			launcherCmd.Append(TC(" -named=")).Append(traceName);
+			launcherCmd.Append(TCV("\\UbaVisualizer.exe\""));
+			launcherCmd.Append(TCV(" -named=")).Append(traceName);
 			STARTUPINFOW si;
 			memset(&si, 0, sizeof(si));
 			PROCESS_INFORMATION pi;
@@ -1778,7 +1778,7 @@ namespace uba
 		{
 			logFile.Append(m_sessionLogDir);
 			GenerateNameForProcess(logFile, startInfo.arguments, processId);
-			logFile.Append(TC(".log"));
+			logFile.Append(TCV(".log"));
 			si.logFile = logFile.data;
 		}
 
@@ -1971,7 +1971,7 @@ namespace uba
 				return rules[i].rules;
 			}
 
-			if (!exeName.Equals(TC("dotnet.exe")))
+			if (!exeName.Equals(TCV("dotnet.exe")))
 				return rules[isDotnet].rules;
 			
 			isDotnet = true;
@@ -2362,7 +2362,7 @@ namespace uba
 
 		StringBuffer<256> buf;
 		if (!GetComputerNameW(buf))
-			buf.Append(TC("LOCAL"));
+			buf.Append(TCV("LOCAL"));
 		StringBuffer<> systemInfo;
 		GetSystemInfo(systemInfo);
 		m_trace.SessionAdded(0, {}, buf, systemInfo);
@@ -2463,7 +2463,7 @@ namespace uba
 	{
 		if (m_shouldWriteToDisk)
 			return true;
-		return fileName.EndsWith(TC(".h"));
+		return fileName.EndsWith(TCV(".h"));
 	}
 
 	bool Session::PrepareProcess(ProcessImpl& process, bool isChild, StringBufferBase& outRealApplication, const tchar*& outRealWorkingDir)
@@ -2501,21 +2501,21 @@ namespace uba
 		// Reason is because there is additional overhead to use memory mappings.
 		// Upside is that all things that are memory mapped can be stored compressed in cas storage so it saves space.
 
-		if (fileName.EndsWith(TC(".h")) || fileName.EndsWith(TC(".inl")) || fileName.EndsWith(TC(".gch")))
+		if (fileName.EndsWith(TCV(".h")) || fileName.EndsWith(TCV(".inl")) || fileName.EndsWith(TCV(".gch")))
 			return 4 * 1024; // clang seems to need 4k alignment? Is it a coincidence it works or what is happening inside the code? (msvc works with alignment 1byte here)
-		if (fileName.EndsWith(TC(".lib")))
+		if (fileName.EndsWith(TCV(".lib")))
 			return 4 * 1024;
 
 		if (runningRemote) // We store these compressed to save space
 		{
-			if (fileName.EndsWith(TC(".obj")) || fileName.EndsWith(TC(".o")))
+			if (fileName.EndsWith(TCV(".obj")) || fileName.EndsWith(TCV(".o")))
 				return 4 * 1024; // pch needs 64k alignment
-			if (fileName.EndsWith(TC(".pch")))
+			if (fileName.EndsWith(TCV(".pch")))
 				return 64 * 1024; // pch needs 64k alignment
 		}
 		else
 		{
-			if (fileName.EndsWith(TC(".h.obj")))
+			if (fileName.EndsWith(TCV(".h.obj")))
 				return 4 * 1024;
 		}
 		return 0;
@@ -2534,17 +2534,17 @@ namespace uba
 			varName.Append(env, TStrchr(env, '=') - env);
 			const tchar* varValue = env + varName.count + 1;
 
-			if (m_runningRemote && varName.Equals(TC("PATH")))
+			if (m_runningRemote && varName.Equals(TCV("PATH")))
 			{
 				AddEnvironmentVariableNoLock(TC("PATH"), TC("c:\\noenvironment"));
 				return;
 			}
-			if (varName.Equals(TC("TEMP")) || varName.Equals(TC("TMP")))
+			if (varName.Equals(TCV("TEMP")) || varName.Equals(TCV("TMP")))
 			{
 				AddEnvironmentVariableNoLock(varName.data, m_tempPath.data);
 				return;
 			}
-			if (varName.Equals(TC("_CL_")) || varName.Equals(TC("CL")))
+			if (varName.Equals(TCV("_CL_")) || varName.Equals(TCV("CL")))
 			{
 				return;
 			}
@@ -2698,11 +2698,11 @@ namespace uba
 			return true;
 		}
 		
-		if (fileName.EndsWith(TC(".dll")) || fileName.EndsWith(TC(".exe")))
+		if (fileName.EndsWith(TCV(".dll")) || fileName.EndsWith(TCV(".exe")))
 		{
 			UBA_ASSERTF(IsAbsolutePath(fileName.data), TC("Got bad filename from process (%s)"), fileName.data);
 			AddFileMapping(fileNameKey, fileName.data, TC("#"));
-			out.fileName.Append(TC("#"));
+			out.fileName.Append(TCV("#"));
 			return true;
 		}
 			
@@ -2731,7 +2731,7 @@ namespace uba
 		if (!IsRarelyRead(process, fileName))
 		{
 			AddFileMapping(fileNameKey, fileName.data, TC("#"));
-			out.fileName.Append(TC("#"));
+			out.fileName.Append(TCV("#"));
 			return true;
 		}
 
@@ -3392,7 +3392,7 @@ namespace uba
 
 		StringBuffer<128> cpuStr(TC("CPU"));
 		if (IsRunningArm())
-			cpuStr.Append(TC("[Arm]"));
+			cpuStr.Append(TCV("[Arm]"));
 		cpuStr.Append(':');
 		if (cpuGroupCount != 1)
 			cpuStr.AppendValue(cpuGroupCount).Append('x');
@@ -3451,7 +3451,7 @@ namespace uba
 
 		StringBuffer<128> zone;
 		if (m_storage.GetZone(zone))
-			out.Append(TC(" Zone:")).Append(zone);
+			out.Append(TCV(" Zone:")).Append(zone);
 
 		#if PLATFORM_WINDOWS
 		if (!IsRunningWine())
@@ -3468,23 +3468,23 @@ namespace uba
 			{
 				u32 lastAccessSettingsValue = value & 0xf;
 				if (lastAccessSettingsValue == 0 || lastAccessSettingsValue == 2)
-					out.Append(TC(" NtfsLastAccessEnabled"));
+					out.Append(TCV(" NtfsLastAccessEnabled"));
 			}
 			value = 0;
 			res = RegGetValueW(HKEY_LOCAL_MACHINE, fsKey, TC("NtfsDisable8dot3NameCreation"), RRF_RT_REG_DWORD, NULL, &value, &valueSize);
 			if (res == ERROR_SUCCESS)
 				if (value == 0)
-					out.Append(TC(" NtfsShortNamesEnabled"));
+					out.Append(TCV(" NtfsShortNamesEnabled"));
 		}
 		else
 		{
 			StringBuffer<> testDir;
-			testDir.Append(m_rootDir).Append(TC("UbaTestShortNames"));
+			testDir.Append(m_rootDir).Append(TCV("UbaTestShortNames"));
 			::RemoveDirectory(testDir.data);
 			wchar_t shortName[1024];
 			if (::CreateDirectoryW(testDir.data, NULL))
 				if (GetShortPathName(testDir.data, shortName, 1024) != 0 && !Contains(shortName, TC("UbaTestShortNames")))
-					out.Append(TC(" NtfsShortNamesEnabled"));
+					out.Append(TCV(" NtfsShortNamesEnabled"));
 		}
 		#endif
 
@@ -3492,7 +3492,7 @@ namespace uba
 			out.Append(m_extraInfo);
 
 		#if UBA_DEBUG
-		out.Append(TC(" - DEBUG"));
+		out.Append(TCV(" - DEBUG"));
 		#endif
 	}
 
@@ -3819,7 +3819,7 @@ namespace uba
 		const tchar* lastDot = TStrrchr(fileName, '.');
 		UBA_ASSERT(lastDot);
 		StringBuffer<> exportsFile;
-		exportsFile.Append(fileName, lastDot - fileName).Append(TC(".exi"));
+		exportsFile.Append(fileName, lastDot - fileName).Append(TCV(".exi"));
 
 		MemoryBlock memoryBlock(8*1024*1024);
 		if (!objectFile->WriteImportsAndExports(m_logger, memoryBlock))
@@ -3929,12 +3929,12 @@ namespace uba
 			size = InvalidValue;
 			FileBasicInformation info;
 			if (!GetFileBasicInformation(info, m_logger, file.data, false))
-				str.Append(TC(" (GetFileBasicInformation failed)")).Append(BytesToText(size).str).Append(TC(')'));
+				str.Append(TCV(" (GetFileBasicInformation failed)")).Append(BytesToText(size).str).Append(')');
 			else
 				size = info.size;
 		}
 		if (size != InvalidValue)
-			str.Append(TC( " (size: ")).Append(BytesToText(size).str).Append(TC(')'));
+			str.Append(TCV( " (size: ")).Append(BytesToText(size).str).Append(')');
 		m_trace.ProcessAddBreadcrumbs(processId, str, false);
 	}
 
@@ -3980,13 +3980,13 @@ namespace uba
 					return m_logger.Warning(TC("Failed to open %s"), out.fileName.data);
 				}
 
-				if (out.fileName.Equals(TC("$d")))
+				if (out.fileName.Equals(TCV("$d")))
 				{
 					// This can happen on apple targets.. crawler finds some includes that are not proper includes
 					return m_logger.Warning(TC("Trying to open directory %s as file"), fileName.data);
 				}
 
-				if (out.fileName.Equals(TC("#")))
+				if (out.fileName.Equals(TCV("#")))
 					out.fileName.Clear().Append(fileName);
 
 				FileAccessor fa(m_logger, out.fileName.data);
@@ -4030,7 +4030,7 @@ namespace uba
 		}
 
 		if (out.IsEmpty())
-			out.Append(TC("NoGoodName"));
+			out.Append(TCV("NoGoodName"));
 
 		if (counterSuffix)
 			out.Appendf(TC("_%03u"), counterSuffix);
diff --git a/Engine/Source/Programs/UnrealBuildAccelerator/Common/Private/UbaSessionClient.cpp b/Engine/Source/Programs/UnrealBuildAccelerator/Common/Private/UbaSessionClient.cpp
--- a/Engine/Source/Programs/UnrealBuildAccelerator/Common/Private/UbaSessionClient.cpp
+++ b/Engine/Source/Programs/UnrealBuildAccelerator/Common/Private/UbaSessionClient.cpp
@@ -64,7 +64,7 @@ namespace uba
 		if (m_name.IsEmpty())
 			GetComputerNameW(m_name);
 
-		m_processWorkingDir.Append(m_rootDir).Append(TC("empty"));// + TC("empty\\");
+		m_processWorkingDir.Append(m_rootDir).Append(TCV("empty"));// + TC("empty\\");
 		m_storage.CreateDirectory(m_processWorkingDir.data);
 		m_processWorkingDir.EnsureEndsWithSlash();
 
@@ -198,7 +198,7 @@ namespace uba
 	{
 		//StringBuffer<> workName;
 		//u32 len = TStrlen(hint);
-		//workName.Append(TC("RC:")).Append(len > 30 ? hint + (len - 30) : hint);
+		//workName.Append(TCV("RC:")).Append(len > 30 ? hint + (len - 30) : hint);
 		//TrackWorkScope tws(m_client, workName.data);
 
 		TimerScope s(m_stats.storageRetrieve);
@@ -435,7 +435,7 @@ namespace uba
 			{
 				StringBuffer<> localSystemModule;
 				localSystemModule.Append(m_systemPath).Append(moduleFile.data + serverSystemPathLen);
-				if (FileExists(m_logger, localSystemModule.data) && !localSystemModule.EndsWith(TC(".exe")))
+				if (FileExists(m_logger, localSystemModule.data) && !localSystemModule.EndsWith(TCV(".exe")))
 					continue;
 				moduleFile.Clear().Append(localSystemModule);
 			}
@@ -562,7 +562,7 @@ namespace uba
 
 			if (isDir)
 			{
-				newName.Append(TC("$d"));
+				newName.Append(TCV("$d"));
 			}
 			else if (casKey != CasKeyZero)
 			{
@@ -1663,7 +1663,7 @@ namespace uba
 		StringBuffer<> dest;
 		if (const tchar* lastSlash = TStrrchr(logFile, PathSeparator))
 			logFile = lastSlash + 1;
-		dest.Append(TC("<log>")).Append(logFile);
+		dest.Append(TCV("<log>")).Append(logFile);
 		f.name = dest.data;
 		f.key = ToStringKeyLower(dest);
 		SendFile(f, pi.GetId(), false, true);
@@ -1683,7 +1683,7 @@ namespace uba
 		else
 		{
 			GenerateNameForProcess(out, arguments, processId);
-			out.Append(TC(".log"));
+			out.Append(TCV(".log"));
 		}
 	}
 
@@ -2078,7 +2078,7 @@ namespace uba
 			m_trace.SessionSummary(0, writer.GetData(), writer.GetPosition());
 
 			StringBuffer<> ubaFile(m_sessionLogDir);
-			ubaFile.Append(TC("Trace.uba"));
+			ubaFile.Append(TCV("Trace.uba"));
 			if (StopTrace(ubaFile.data))
 			{
 				WrittenFile f;
@@ -2446,7 +2446,7 @@ namespace uba
 		u32 foundCount = 0;
 		dirTable.TraverseAllFilesNoLock([&](const DirectoryTable::EntryInformation& info, const StringBufferBase& path, u32 dirOffset)
 			{
-				if (!path.EndsWith(searchString.data))
+				if (!path.EndsWith(searchString))
 					return;
 				if (path[path.count - searchString.count - 1] != PathSeparator)
 					return;
@@ -2472,7 +2472,7 @@ namespace uba
 						if (entry.mapping.IsValid())
 							Storage::GetMappingString(mappingName, entry.mapping, entry.mappingOffset);
 						else
-							mappingName.Append(TC("Not valid"));
+							mappingName.Append(TCV("Not valid"));
 						logStr.Appendf(TC(" Success: %s Size: %u IsDir: %s Mapping name: %s Mapping offset: %u"), ToString(entry.success), entry.size, ToString(entry.isDir), mappingName.data, entry.mappingOffset);
 					}
 					else
diff --git a/Engine/Source/Programs/UnrealBuildAccelerator/Common/Private/UbaSessionServer.cpp b/Engine/Source/Programs/UnrealBuildAccelerator/Common/Private/UbaSessionServer.cpp
--- a/Engine/Source/Programs/UnrealBuildAccelerator/Common/Private/UbaSessionServer.cpp
+++ b/Engine/Source/Programs/UnrealBuildAccelerator/Common/Private/UbaSessionServer.cpp
@@ -112,7 +112,7 @@ namespace uba
 	bool GetCrossArchitectureDir(Logger& logger, StringBufferBase& dir, bool reportError)
 	{
 		// UBT has the path win-x64/native or win-arm64/native
-		bool isUbtPath = dir.EndsWith(TC("native"));
+		bool isUbtPath = dir.EndsWith(TCV("native"));
 		if (isUbtPath)
 			dir.Resize(dir.count - 7); // Remove native and slash
 		const tchar* archPath[2] = { TC("x64"), TC("arm64") };
@@ -120,7 +120,7 @@ namespace uba
 			return reportError ? logger.Error(TC("Module dir is not under supported folder (%s) to be able to run cross architecutres, can't figure out matching x64/arm64 folder"), dir.data) : false;
 		dir.Resize(dir.count - TStrlen(archPath[IsArmBinary])).Append(archPath[!IsArmBinary]);
 		if (isUbtPath)
-			dir.Append(PathSeparator).Append(TC("native"));
+			dir.Append(PathSeparator).Append(TCV("native"));
 		return true;
 	}
 
@@ -441,7 +441,7 @@ namespace uba
 		StringBuffer<> workingDir;
 		FixFileName(workingDir, workingDir_, nullptr);
 		if (workingDir[workingDir.count - 1] != '\\')
-			workingDir.Append(TC("\\"));
+			workingDir.Append(TCV("\\"));
 		StringBuffer<> fileName;
 		FixFileName(fileName, fileName_, workingDir.data);
 		StringKey fileNameKey = ToStringKey(fileName);
@@ -657,7 +657,7 @@ namespace uba
 		if (returnCount)
 		{
 			if (sessionName.IsEmpty())
-				sessionName.Append(TC("<can't find session>"));
+				sessionName.Append(TCV("<can't find session>"));
 
 			m_logger.Info(TC("Client session %s (%s) disconnected. Returned %u process(s) to queue"), sessionName.data, GuidToString(clientUid).str, returnCount);
 		}
@@ -1068,9 +1068,9 @@ namespace uba
 				if (session)
 					ubaPath.Append(session->name);
 				else
-					ubaPath.Append(TC("Connection")).AppendValue(connectionInfo.GetId());
+					ubaPath.Append(TCV("Connection")).AppendValue(connectionInfo.GetId());
 
-				ubaPath.Append(TC(".uba"));
+				ubaPath.Append(TCV(".uba"));
 				m_storage.CopyOrLink(casKey, ubaPath.data, attributes);
 				m_storage.DropCasFile(casKey, false, ubaPath.data);
 				return true;
@@ -1684,7 +1684,7 @@ namespace uba
 
 		auto WriteString = [&](const tchar* str, LogEntryType type = LogEntryType_Info) { writer.WriteByte(type); writer.WriteString(str); };
 
-		if (command.Equals(TC("status")))
+		if (command.Equals(TCV("status")))
 		{
 			u32 totalUsed = 0;
 			u32 totalSlots = 0;
@@ -1705,8 +1705,8 @@ namespace uba
 		}
 		else if (command.StartsWith(TC("abort")))
 		{
-			bool abortWithProxy = command.Equals(TC("abortproxy"));
-			bool abortUseProxy = command.Equals(TC("abortnonproxy"));
+			bool abortWithProxy = command.Equals(TCV("abortproxy"));
+			bool abortUseProxy = command.Equals(TCV("abortnonproxy"));
 			if (!abortWithProxy && !abortUseProxy)
 			{
 				abortWithProxy = true;
@@ -1728,7 +1728,7 @@ namespace uba
 			}
 			WriteString(StringBuffer<>().Appendf(TC("Aborting: %u remote sessions"), abortCount).data);
 		}
-		else if (command.Equals(TC("disableremote")))
+		else if (command.Equals(TCV("disableremote")))
 		{
 			DisableRemoteExecution();
 			WriteString(StringBuffer<>().Appendf(TC("Remote execution is disabled")).data);
@@ -2240,7 +2240,7 @@ namespace uba
 		{
 			StringBuffer<> varName;
 			varName.Append(it, TStrchr(it, '=') - it);
-			if (!varName.IsEmpty() && !varName.Equals(TC("CL")) && !varName.Equals(TC("_CL_")))
+			if (!varName.IsEmpty() && !varName.Equals(TCV("CL")) && !varName.Equals(TCV("_CL_")))
 				if (m_localEnvironmentVariables.find(varName.data) == m_localEnvironmentVariables.end())
 					writer.WriteString(it);
 		}
@@ -2334,7 +2334,7 @@ namespace uba
 		u32 foundCount = 0;
 		dirTable.TraverseAllFilesNoLock([&](const DirectoryTable::EntryInformation& info, const StringBufferBase& path, u32 dirOffset)
 			{
-				if (!path.EndsWith(searchString.data))
+				if (!path.EndsWith(searchString))
 					return;
 				if (path[path.count - searchString.count - 1] != PathSeparator)
 					return;
@@ -2361,7 +2361,7 @@ namespace uba
 							if (entry.mapping.IsValid())
 								Storage::GetMappingString(mappingName, entry.mapping, entry.mappingOffset);
 							else
-								mappingName.Append(TC("Not valid"));
+								mappingName.Append(TCV("Not valid"));
 							logStr.Appendf(TC(" Success: %s Size: %u IsDir: %s Mapping name: %s Mapping offset: %u"), ToString(entry.success), entry.size, ToString(entry.isDir), mappingName.data, entry.mappingOffset);
 						}
 						else
diff --git a/Engine/Source/Programs/UnrealBuildAccelerator/Common/Private/UbaStorage.cpp b/Engine/Source/Programs/UnrealBuildAccelerator/Common/Private/UbaStorage.cpp
--- a/Engine/Source/Programs/UnrealBuildAccelerator/Common/Private/UbaStorage.cpp
+++ b/Engine/Source/Programs/UnrealBuildAccelerator/Common/Private/UbaStorage.cpp
@@ -1222,7 +1222,7 @@ namespace uba
 	bool StorageImpl::DeleteIsRunningFile()
 	{
 		StringBuffer<256> isRunningName;
-		isRunningName.Append(m_rootDir).Append(TC(".isRunning"));
+		isRunningName.Append(m_rootDir).Append(TCV(".isRunning"));
 		if (DeleteFileW(isRunningName.data))
 			return true;
 		u32 lastError = GetLastError();
@@ -1264,13 +1264,13 @@ namespace uba
 
 		if (m_writeToDisk)
 		{
-			m_tempPath.Append(m_rootDir).Append(TC("castemp"));
+			m_tempPath.Append(m_rootDir).Append(TCV("castemp"));
 			CreateDirectory(m_tempPath.data);
 			DeleteAllFiles(m_logger, m_tempPath.data, false);
 			m_tempPath.EnsureEndsWithSlash();
 		}
 
-		m_rootDir.Append(TC("cas")).EnsureEndsWithSlash();
+		m_rootDir.Append(TCV("cas")).EnsureEndsWithSlash();
 
 		m_casDataBuffer.AddTransient(TC("CasData"), info.keepTransientDataMapped);
 
@@ -1364,7 +1364,7 @@ namespace uba
 		}
 
 		StringBuffer<> fileName;
-		fileName.Append(m_rootDir).Append(TC("casdb"));
+		fileName.Append(m_rootDir).Append(TCV("casdb"));
 
 		FileHandle fileHandle;
 		if (!OpenFileSequentialRead(m_logger, fileName.data, fileHandle, false))
@@ -1516,10 +1516,10 @@ namespace uba
 		u64 startTime = GetTime();
 
 		StringBuffer<256> fileName;
-		fileName.Append(m_rootDir).Append(TC("casdb"));
+		fileName.Append(m_rootDir).Append(TCV("casdb"));
 
 		StringBuffer<256> tempFileName;
-		tempFileName.Append(fileName).Append(TC(".tmp"));
+		tempFileName.Append(fileName).Append(TCV(".tmp"));
 
 		u32 deleteCount = 0;
 		{
diff --git a/Engine/Source/Programs/UnrealBuildAccelerator/Common/Private/UbaStorageClient.cpp b/Engine/Source/Programs/UnrealBuildAccelerator/Common/Private/UbaStorageClient.cpp
--- a/Engine/Source/Programs/UnrealBuildAccelerator/Common/Private/UbaStorageClient.cpp
+++ b/Engine/Source/Programs/UnrealBuildAccelerator/Common/Private/UbaStorageClient.cpp
@@ -385,7 +385,7 @@ namespace uba
 							m_logger.Warning(TC("Failed to create proxy server. This should never happen!"));
 							continue;
 						}
-						proxyHost.Append(TC("inprocess"));
+						proxyHost.Append(TCV("inprocess"));
 						isInProcessClient = true;
 					}
 					else
diff --git a/Engine/Source/Programs/UnrealBuildAccelerator/Common/Private/UbaTrace.cpp b/Engine/Source/Programs/UnrealBuildAccelerator/Common/Private/UbaTrace.cpp
--- a/Engine/Source/Programs/UnrealBuildAccelerator/Common/Private/UbaTrace.cpp
+++ b/Engine/Source/Programs/UnrealBuildAccelerator/Common/Private/UbaTrace.cpp
@@ -501,7 +501,7 @@ namespace uba
 	{
 		#if PLATFORM_WINDOWS
 		StringBuffer<245> channelMutex;
-		channelMutex.Append(TC("Uba")).Append(channelName).Append(TC("Channel"));
+		channelMutex.Append(TCV("Uba")).Append(channelName).Append(TCV("Channel"));
 		m_memHandle = uba::CreateMemoryMappingW(m_logger, PAGE_READWRITE, 256, channelName, TC("TraceChannel"));
 		if (!m_memHandle.IsValid())
 		{
@@ -524,7 +524,7 @@ namespace uba
 
 		auto mg = MakeGuard([&]() { UnmapViewOfFile(m_logger, m_mem, 256, channelMutex.data); m_mem = nullptr; });
 
-		channelMutex.Append(channelName).Append(TC("Mutex"));
+		channelMutex.Append(channelName).Append(TCV("Mutex"));
 		m_mutex = CreateMutexW(false, channelMutex.data);
 		if (!m_mutex)
 			return false;
diff --git a/Engine/Source/Programs/UnrealBuildAccelerator/Common/Public/UbaAWS.h b/Engine/Source/Programs/UnrealBuildAccelerator/Common/Public/UbaAWS.h
--- a/Engine/Source/Programs/UnrealBuildAccelerator/Common/Public/UbaAWS.h
+++ b/Engine/Source/Programs/UnrealBuildAccelerator/Common/Public/UbaAWS.h
@@ -36,7 +36,7 @@ namespace uba
 			u32 statusCode = 0;
 
 			StringBuffer<128> token;
-			token.Append(TC("X-aws-ec2-metadata-token: "));
+			token.Append(TCV("X-aws-ec2-metadata-token: "));
 			if (!http.Query(logger, "PUT", token, statusCode, g_imdsHost, g_imdsToken, "X-aws-ec2-metadata-token-ttl-seconds: 21600\r\n"))
 				return WriteIsNotAws(logger, rootDir);
 			token.Append("\r\n");
@@ -54,7 +54,7 @@ namespace uba
 			if (!http.Query(logger, "GET", instanceId, statusCode, g_imdsHost, g_imdsInstanceId, m_tokenString.c_str()))
 				return false;
 
-			outExtraInfo.Append(TC(", AWS: ")).Append(instanceId);
+			outExtraInfo.Append(TCV(", AWS: ")).Append(instanceId);
 
 			StringBuffer<32> instanceLifeCycle;
 			if (http.Query(logger, "GET", instanceLifeCycle, statusCode, g_imdsHost, g_imdsInstanceLifeCycle, m_tokenString.c_str()))
@@ -66,7 +66,7 @@ namespace uba
 			StringBuffer<32> autoscaling;
 			if (http.Query(logger, "GET", autoscaling, statusCode, g_imdsHost, g_imdsAutoScalingLifeCycleState, m_tokenString.c_str()) && statusCode == 200)
 			{
-				outExtraInfo.Append(m_isSpot ? '/' : ' ').Append(TC("autoscale"));
+				outExtraInfo.Append(m_isSpot ? '/' : ' ').Append(TCV("autoscale"));
 				m_isAutoscaling = true;
 			}
 
@@ -129,7 +129,7 @@ namespace uba
 				u32 statusCode = 0;
 				if (http.Query(logger, "GET", content, statusCode, g_imdsHost, g_imdsSpotInstanceAction, m_tokenString.c_str()) && statusCode == 200)
 				{
-					outReason.Append(TC("AWS spot instance interruption"));
+					outReason.Append(TCV("AWS spot instance interruption"));
 					return true;
 				}
 			}
@@ -148,7 +148,7 @@ namespace uba
 					if (!content.Contains(TC("InService"))) // AWS can return "InServiceI" as well?
 					{
 						//wprintf(L"AWSACTION: AUTOSCALE REBALANCING!!!! (%ls)\n", content.data);
-						outReason.Append(TC("AWS autoscale rebalancing"));
+						outReason.Append(TCV("AWS autoscale rebalancing"));
 						return true;
 					}
 				}
diff --git a/Engine/Source/Programs/UnrealBuildAccelerator/Common/Public/UbaCoordinatorWrapper.h b/Engine/Source/Programs/UnrealBuildAccelerator/Common/Public/UbaCoordinatorWrapper.h
--- a/Engine/Source/Programs/UnrealBuildAccelerator/Common/Public/UbaCoordinatorWrapper.h
+++ b/Engine/Source/Programs/UnrealBuildAccelerator/Common/Public/UbaCoordinatorWrapper.h
@@ -34,9 +34,9 @@ namespace uba
 			coordinatorBin.EnsureEndsWithSlash();
 
 			#if PLATFORM_WINDOWS
-			coordinatorBin.Append(TC("UbaCoordinator")).Append(coordinatorType).Append(TC(".dll"));
+			coordinatorBin.Append(TCV("UbaCoordinator")).Append(coordinatorType).Append(TCV(".dll"));
 			#else
-			coordinatorBin.Append(TC("libUbaCoordinator")).Append(coordinatorType).Append(TC(".so"));
+			coordinatorBin.Append(TCV("libUbaCoordinator")).Append(coordinatorType).Append(TCV(".so"));
 			#endif
 
 			HMODULE coordinatorModule = LoadLibrary(coordinatorBin.data);
diff --git a/Engine/Source/Programs/UnrealBuildAccelerator/Common/Public/UbaDependencyCrawler.h b/Engine/Source/Programs/UnrealBuildAccelerator/Common/Public/UbaDependencyCrawler.h
--- a/Engine/Source/Programs/UnrealBuildAccelerator/Common/Public/UbaDependencyCrawler.h
+++ b/Engine/Source/Programs/UnrealBuildAccelerator/Common/Public/UbaDependencyCrawler.h
@@ -150,7 +150,7 @@ namespace uba
 						SCOPED_FUTEX(pch.lock, lock2);
 						if (!pch.handled)
 						{
-							fixedFile.Append(TC(".dep.json"));
+							fixedFile.Append(TCV(".dep.json"));
 							instance.createFileFunc(fixedFile, [&](const void* data, u64 dataSize)
 							{
 								return ParsePch(pch, instance, data, dataSize);
@@ -348,7 +348,7 @@ namespace uba
 							{
 								handledFile.handled = true;
 
-								fixedFile.Append(TC(".pch"));
+								fixedFile.Append(TCV(".pch"));
 
 								u32 attributes = 0;
 								bool hasPch = m_fileExistsFunc(fixedFile, attributes);
@@ -481,7 +481,7 @@ namespace uba
 							AddRoot(sb);
 							break;
 						case ArgType_Code:
-							if (!sb.EndsWith(TC(".gch")))
+							if (!sb.EndsWith(TCV(".gch")))
 								AddCodeFile(sb, true, false);
 							break;
 						case ArgType_ResourceDir:
@@ -518,7 +518,7 @@ namespace uba
 					}
 					else
 					{
-						if (sb.EndsWith(TC(".obj")))
+						if (sb.EndsWith(TCV(".obj")))
 							AddCodeFile(sb, false, false);
 					}
 				});
@@ -661,7 +661,7 @@ namespace uba
 						{
 							if (strncmp(defineBegin, "UE_INLINE_GENERATED_CPP_BY_NAME", defineLen) == 0)
 							{
-								include.Append(argBegin, u32(argEnd - argBegin)).Append(TC(".gen.cpp"));
+								include.Append(argBegin, u32(argEnd - argBegin)).Append(TCV(".gen.cpp"));
 							}
 							else if (strncmp(defineBegin, "COMPILED_PLATFORM_HEADER", defineLen) == 0)
 							{
@@ -677,7 +677,7 @@ namespace uba
 						{
 							if (strncmp(defineBegin, "PER_MODULE_INLINE_FILE", defineLen) == 0)
 							{
-								include.Append(TC("HAL/PerModuleInline.inl")); // TODO: This is not correct.. need revisit
+								include.Append(TCV("HAL/PerModuleInline.inl")); // TODO: This is not correct.. need revisit
 							}
 						}
 					}
@@ -768,7 +768,7 @@ namespace uba
 
 	bool DependencyCrawler::HandleInclude(Instance& instance, const StringView& rootPath, const StringView& include, const StringView& codeFile, bool parseDefines)
 	{
-		if (include.EndsWith(TC(".ush")))
+		if (include.EndsWith(TCV(".ush")))
 			return true;
 		StringBuffer<> fullPath;
 		FixPath(include.data, rootPath.data, rootPath.count, fullPath);
diff --git a/Engine/Source/Programs/UnrealBuildAccelerator/Core/Private/UbaApplicationRules.cpp b/Engine/Source/Programs/UnrealBuildAccelerator/Core/Private/UbaApplicationRules.cpp
--- a/Engine/Source/Programs/UnrealBuildAccelerator/Core/Private/UbaApplicationRules.cpp
+++ b/Engine/Source/Programs/UnrealBuildAccelerator/Core/Private/UbaApplicationRules.cpp
@@ -56,7 +56,7 @@ namespace uba
 				return 14ull * 1024 * 1024 * 1024; // This is ridiculous
 			if (file.EndsWith(TCV(".json")) || file.EndsWith(TCV(".exp")) || file.EndsWith(TCV(".sarif")) || file.EndsWith(TCV(".res")))
 				return 32 * 1024 * 1024;
-			if (file.EndsWith(TCV(".ifc")) || file.EndsWith(TCV(".obj")) || (isSystemOrTempFile && file.Contains(TC("_cl_")))) // There are _huge_ obj files when building with -stresstestunity
+			if (file.EndsWith(TCV(".ifc")) || file.EndsWith(TCV(".obj")) || (isSystemOrTempFile && file.Contains(TCV("_cl_")))) // There are _huge_ obj files when building with -stresstestunity
 				return 1024 * 1024 * 1024;
 			return Super::FileTypeMaxSize(file, isSystemOrTempFile);
 		}
@@ -68,11 +68,11 @@ namespace uba
 		}
 		virtual bool IsThrowAway(const StringView& fileName, bool isRunningRemote) const override
 		{
-			return fileName.Contains(TC("vctip_")) || Super::IsThrowAway(fileName, isRunningRemote);
+			return fileName.Contains(TCV("vctip_")) || Super::IsThrowAway(fileName, isRunningRemote);
 		}
 		virtual bool KeepInMemory(const StringView& fileName, const tchar* systemTemp, bool isRunningRemote, bool isWrite) const override
 		{
-			if (fileName.Contains(TC("\\vctip_")))
+			if (fileName.Contains(TCV("\\vctip_")))
 				return true;
 			if (fileName.Contains(systemTemp))
 				return true;
@@ -163,7 +163,7 @@ namespace uba
 
 		virtual bool IsThrowAway(const StringView& fileName, bool isRunningRemote) const override
 		{
-			return fileName.Contains(TC(".sup.")); // .sup.lib/exp are throw-away files that we don't want created
+			return fileName.Contains(TCV(".sup.")); // .sup.lib/exp are throw-away files that we don't want created
 		}
 
 		virtual bool CanExist(const tchar* file) const override
@@ -179,7 +179,7 @@ namespace uba
 
 		virtual u64 FileTypeMaxSize(const StringView& file, bool isSystemOrTempFile) const override
 		{
-			if (file.Contains(TC("lnk{")))
+			if (file.Contains(TCV("lnk{")))
 				return 32 * 1024 * 1024;
 			if (file.EndsWith(TCV(".lib"))) // This is import lib which is tiny but not sure if we output any full libs
 				return 512 * 1024 * 1024;
@@ -256,7 +256,7 @@ namespace uba
 
 		virtual u64 FileTypeMaxSize(const StringView& file, bool isSystemOrTempFile) const override
 		{
-			if (file.Contains(TC(".pdb.tmp")))
+			if (file.Contains(TCV(".pdb.tmp")))
 				return 14ull * 1024 * 1024 * 1024; // This is ridiculous
 			return ApplicationRulesVcLink::FileTypeMaxSize(file, isSystemOrTempFile);
 		}
@@ -264,9 +264,9 @@ namespace uba
 		virtual bool IsOutputFile(const StringView& fileName) const override
 		{
 			return false//Super::IsOutputFile(fileName)
-				|| fileName.Contains(TC(".exe.tmp"))
-				|| fileName.Contains(TC(".dll.tmp"))
-				|| fileName.Contains(TC(".pdb.tmp"));
+				|| fileName.Contains(TCV(".exe.tmp"))
+				|| fileName.Contains(TCV(".dll.tmp"))
+				|| fileName.Contains(TCV(".pdb.tmp"));
 		}
 	};
 
@@ -276,7 +276,7 @@ namespace uba
 
 		virtual u64 FileTypeMaxSize(const StringView& file, bool isSystemOrTempFile) const override
 		{
-			if (file.Contains(TC(".pdb.tmp")) || file.EndsWith(TCV(".pdb")))
+			if (file.Contains(TCV(".pdb.tmp")) || file.EndsWith(TCV(".pdb")))
 				return 14ull * 1024 * 1024 * 1024; // This is ridiculous
 			return ApplicationRulesVcLink::FileTypeMaxSize(file, isSystemOrTempFile);
 		}
@@ -287,9 +287,9 @@ namespace uba
 				|| fileName.EndsWith(TCV(".exe"))
 				|| fileName.EndsWith(TCV(".dll"))
 				|| fileName.EndsWith(TCV(".pdb"))
-				|| fileName.Contains(TC(".exe.tmp"))
-				|| fileName.Contains(TC(".dll.tmp"))
-				|| fileName.Contains(TC(".pdb.tmp"));
+				|| fileName.Contains(TCV(".exe.tmp"))
+				|| fileName.Contains(TCV(".dll.tmp"))
+				|| fileName.Contains(TCV(".pdb.tmp"));
 		}
 	};
 
@@ -399,7 +399,7 @@ namespace uba
 
 		virtual bool IsOutputFile(const StringView& fileName) const override
 		{
-			return fileName.Contains(TC(".tmp")); // both .so.tmp and .tmp123456
+			return fileName.Contains(TCV(".tmp")); // both .so.tmp and .tmp123456
 		}
 
 		virtual u64 FileTypeMaxSize(const StringView& file, bool isSystemOrTempFile) const override
@@ -419,7 +419,7 @@ namespace uba
 
 		virtual bool IsOutputFile(const StringView& fileName) const override
 		{
-			return fileName.Contains(TC(".temp-stream-"));
+			return fileName.Contains(TCV(".temp-stream-"));
 		}
 
 		virtual u64 FileTypeMaxSize(const StringView& file, bool isSystemOrTempFile) const override
@@ -462,7 +462,7 @@ namespace uba
 
 		//virtual bool NeedsSharedMemory(const tchar* file)
 		//{
-		//	return fileName.Contains(TC("lto-llvm")); // Used by a clang based platform's link time optimization pass. Shared from lto process back to linker process
+		//	return fileName.Contains(TCV("lto-llvm")); // Used by a clang based platform's link time optimization pass. Shared from lto process back to linker process
 		//}
 	};
 
@@ -477,7 +477,7 @@ namespace uba
 		virtual bool KeepInMemory(const StringView& fileName, const tchar* systemTemp, bool isRunningRemote, bool isWrite) const override
 		{
 			return Super::KeepInMemory(fileName, systemTemp, isRunningRemote, isWrite)
-				|| fileName.Contains(TC("thinlto-"));// Used by a clang based platform's link time optimization pass. Shared from lto process back to linker process
+				|| fileName.Contains(TCV("thinlto-"));// Used by a clang based platform's link time optimization pass. Shared from lto process back to linker process
 		}
 		virtual bool NeedsSharedMemory(const tchar* file) const override
 		{
@@ -496,7 +496,7 @@ namespace uba
 
 		virtual bool IsOutputFile(const StringView& fileName) const override
 		{
-			return fileName.Contains(TC(".self")) || Super::IsOutputFile(fileName);
+			return fileName.Contains(TCV(".self")) || Super::IsOutputFile(fileName);
 		}
 
 		virtual bool IsThrowAway(const StringView& fileName, bool isRunningRemote) const override
@@ -507,7 +507,7 @@ namespace uba
 
 		//virtual bool NeedsSharedMemory(const tchar* file)
 		//{
-		//	return fileName.Contains(TC("lto-llvm")); // Used by a clang based platform's link time optimization pass. Shared from lto process back to linker process
+		//	return fileName.Contains(TCV("lto-llvm")); // Used by a clang based platform's link time optimization pass. Shared from lto process back to linker process
 		//}
 	};
 
@@ -517,7 +517,7 @@ namespace uba
 
 		virtual bool IsOutputFile(const StringView& fileName) const override
 		{
-			return fileName.Contains(TC(".self"));
+			return fileName.Contains(TCV(".self"));
 		}
 
 		virtual bool IsThrowAway(const StringView& fileName, bool isRunningRemote) const override
@@ -528,7 +528,7 @@ namespace uba
 
 		//virtual bool NeedsSharedMemory(const tchar* file)
 		//{
-		//	return fileName.Contains(TC("lto-llvm")); // Used by a clang based platform's link time optimization pass. Shared from lto process back to linker process
+		//	return fileName.Contains(TCV("lto-llvm")); // Used by a clang based platform's link time optimization pass. Shared from lto process back to linker process
 		//}
 
 		virtual bool CanDependOnCompressedFiles() override
@@ -558,14 +558,14 @@ namespace uba
 
 		virtual bool KeepInMemory(const StringView& fileName, const tchar* systemTemp, bool isRunningRemote, bool isWrite) const override
 		{
-			if (isWrite && fileName.Contains(TC("\\Split\\")) && fileName.EndsWith(TCV(".split")))
+			if (isWrite && fileName.Contains(TCV("\\Split\\")) && fileName.EndsWith(TCV(".split")))
 				return true;
 			return false;
 		}
 
 		virtual bool IsThrowAway(const StringView& fileName, bool isRunningRemote) const override
 		{
-			if (fileName.Contains(TC("\\Split\\")) && fileName.EndsWith(TCV(".split")))
+			if (fileName.Contains(TCV("\\Split\\")) && fileName.EndsWith(TCV(".split")))
 				return true;
 			return Super::IsThrowAway(fileName, isRunningRemote);
 		}
@@ -589,7 +589,7 @@ namespace uba
 
 		virtual bool IsOutputFile(const StringView& fileName) const override
 		{
-			return fileName.Contains(TC(".generated.dummy"))
+			return fileName.Contains(TCV(".generated.dummy"))
 				|| fileName.EndsWith(TCV(".ispc.bc"))
 				|| fileName.EndsWith(TCV(".ispc.txt"))
 				|| fileName.EndsWith(TCV(".obj"))
@@ -598,7 +598,7 @@ namespace uba
 
 		virtual u64 FileTypeMaxSize(const StringView& file, bool isSystemOrTempFile) const override
 		{
-			if (file.Contains(TC(".generated.dummy")) || file.EndsWith(TCV(".ispc.txt")))
+			if (file.Contains(TCV(".generated.dummy")) || file.EndsWith(TCV(".ispc.txt")))
 				return 32ull * 1024 * 1024;
 			if (file.EndsWith(TCV(".obj")) || file.EndsWith(TCV(".o")))
 				return 128ull * 1024 * 1024;
@@ -668,7 +668,7 @@ namespace uba
 	{
 		virtual bool IsRarelyRead(const StringView& file) const override
 		{
-			return file.Contains(TC(".uba."));
+			return file.Contains(TCV(".uba."));
 		}
 
 		virtual bool IsOutputFile(const StringView& fileName) const override
diff --git a/Engine/Source/Programs/UnrealBuildAccelerator/Core/Private/UbaPlatform.cpp b/Engine/Source/Programs/UnrealBuildAccelerator/Core/Private/UbaPlatform.cpp
--- a/Engine/Source/Programs/UnrealBuildAccelerator/Core/Private/UbaPlatform.cpp
+++ b/Engine/Source/Programs/UnrealBuildAccelerator/Core/Private/UbaPlatform.cpp
@@ -174,14 +174,14 @@ namespace uba
 		#endif
 
 		if (!context)
-			out.Append(TC("ASSERT: "));
+			out.Append(TCV("ASSERT: "));
 		out.Append(*text ? text : TC("Unknown"));
 		if (!*text && expr && *expr && strcmp(expr, "false") != 0)
 			out.Appendf(TC("\n EXPR: ")).Append(expr);
 		if (file && *file)
-			out.Append(TC("\n LOCATION: ")).Append(file).Append(':').AppendValue(line);
+			out.Append(TCV("\n LOCATION: ")).Append(file).Append(':').AppendValue(line);
 		if (signal)
-			out.Append(TC("\n SIGNAL: ")).AppendValue(signal);
+			out.Append(TCV("\n SIGNAL: ")).AppendValue(signal);
 	}
 
 	#if defined(_M_ARM64)
diff --git a/Engine/Source/Programs/UnrealBuildAccelerator/Core/Private/UbaStringBuffer.cpp b/Engine/Source/Programs/UnrealBuildAccelerator/Core/Private/UbaStringBuffer.cpp
--- a/Engine/Source/Programs/UnrealBuildAccelerator/Core/Private/UbaStringBuffer.cpp
+++ b/Engine/Source/Programs/UnrealBuildAccelerator/Core/Private/UbaStringBuffer.cpp
@@ -378,7 +378,7 @@ namespace uba
 		#else
 		out = strtoul(data + offset, nullptr, 10);
 		#endif
-		return out != 0 || Equals(TC("0"));
+		return out != 0 || Equals(TCV("0"));
 	}
 
 	bool StringBufferBase::Parse(u16& out, u64 offset)
@@ -401,7 +401,7 @@ namespace uba
 		#else
 		out = strtof(data, nullptr);
 		#endif
-		return out != 0 || Equals(TC("0"));
+		return out != 0 || Equals(TCV("0"));
 	}
 
 	u32 StringBufferBase::Parse(char* out, u64 outCapacity)
diff --git a/Engine/Source/Programs/UnrealBuildAccelerator/Core/Public/UbaStringBuffer.h b/Engine/Source/Programs/UnrealBuildAccelerator/Core/Public/UbaStringBuffer.h
--- a/Engine/Source/Programs/UnrealBuildAccelerator/Core/Public/UbaStringBuffer.h
+++ b/Engine/Source/Programs/UnrealBuildAccelerator/Core/Public/UbaStringBuffer.h
@@ -116,6 +116,7 @@ namespace uba
 		bool EndsWith(const StringView& value, bool ignoreCase = true) const { return uba::EndsWith(data, count, value.data, value.count, ignoreCase); }
 		bool EndsWith(const tchar* value, bool ignoreCase = true) const;
 		bool Contains(const tchar* str, bool ignoreCase = true) const { return uba::Contains(data, str, ignoreCase); }
+		bool Contains(const StringView& str, bool ignoreCase = true) const { return uba::Contains(data, str.data, ignoreCase); }
 		bool Equals(const tchar* str, bool ignoreCase = true) const { return uba::Equals(data, str, ignoreCase); }
 		bool Equals(const StringView& str, bool ignoreCase = true) const { return count == str.count && uba::Equals(data, str.data, count, ignoreCase); }
 		StringView GetFileName() const { const tchar* str = uba::GetFileName(data); return StringView(str, count - u32(str - data)); }
diff --git a/Engine/Source/Programs/UnrealBuildAccelerator/Core/Public/UbaWinBinDependencyParser.h b/Engine/Source/Programs/UnrealBuildAccelerator/Core/Public/UbaWinBinDependencyParser.h
--- a/Engine/Source/Programs/UnrealBuildAccelerator/Core/Public/UbaWinBinDependencyParser.h
+++ b/Engine/Source/Programs/UnrealBuildAccelerator/Core/Public/UbaWinBinDependencyParser.h
@@ -131,7 +131,7 @@ namespace uba
 
 	inline bool FindImports(const tchar* fileName, const Function<void(const tchar* import, bool isKnown, const char* const* importLoaderPaths)>& func, StringBufferBase& outError)
 	{
-		if (ToView(fileName).EndsWith(TC(".bat")))
+		if (ToView(fileName).EndsWith(TCV(".bat")))
 			return true;
 		HANDLE fileHandle = True_CreateFileW(fileName, GENERIC_READ, FILE_SHARE_READ, 0, OPEN_EXISTING, FILE_ATTRIBUTE_NORMAL | FILE_FLAG_SEQUENTIAL_SCAN, 0);
 		if (fileHandle == INVALID_HANDLE_VALUE)
diff --git a/Engine/Source/Programs/UnrealBuildAccelerator/Detours/Private/UbaDetoursShared.cpp b/Engine/Source/Programs/UnrealBuildAccelerator/Detours/Private/UbaDetoursShared.cpp
--- a/Engine/Source/Programs/UnrealBuildAccelerator/Detours/Private/UbaDetoursShared.cpp
+++ b/Engine/Source/Programs/UnrealBuildAccelerator/Detours/Private/UbaDetoursShared.cpp
@@ -95,7 +95,7 @@ namespace uba
 			t_a.Append(format, arg);
 			va_end(arg);
 		}
-		t_a.Append(TC("\n"));
+		t_a.Append(TCV("\n"));
 
 		u32 size__ = t_b_size;
 		StringBuffer<128> extra;
@@ -125,7 +125,7 @@ namespace uba
 			t_a.Append(format, arg);
 			va_end(arg);
 		}
-		t_a.Append(TC("\n"));
+		t_a.Append(TCV("\n"));
 
 		#if PLATFORM_WINDOWS
 		if (t_b_size)
@@ -290,7 +290,7 @@ namespace uba
 			TStrcpy_s(buffer, sizeof_array(buffer), format);
 		va_end(arg);
 		StringBuffer<2048> sb;
-		sb.Append(GetApplicationShortName()).Append(TC(" ERROR: ")).Append(buffer);
+		sb.Append(GetApplicationShortName()).Append(TCV(" ERROR: ")).Append(buffer);
 		Rpc_WriteLog(sb.data, sb.count, true, true);
 
 		#if PLATFORM_WINDOWS // Maybe all platforms should call exit()?
@@ -350,7 +350,7 @@ namespace uba
 		}
 		else
 		{
-			out.Append(TC("\n   Failed to resolve callstack\n"));
+			out.Append(TCV("\n   Failed to resolve callstack\n"));
 		}
 		// Note, we leave the lock even though we might not have it because we want to be able to report
 		g_communicationLock.Leave();
diff --git a/Engine/Source/Programs/UnrealBuildAccelerator/Detours/Private/Windows/UbaDetoursFunctionsKernelBase.inl b/Engine/Source/Programs/UnrealBuildAccelerator/Detours/Private/Windows/UbaDetoursFunctionsKernelBase.inl
--- a/Engine/Source/Programs/UnrealBuildAccelerator/Detours/Private/Windows/UbaDetoursFunctionsKernelBase.inl
+++ b/Engine/Source/Programs/UnrealBuildAccelerator/Detours/Private/Windows/UbaDetoursFunctionsKernelBase.inl
@@ -569,7 +569,7 @@ LPVOID Detoured_VirtualAlloc(LPVOID lpAddress, SIZE_T dwSize, DWORD flAllocation
 		if (error != ERROR_NOT_ENOUGH_MEMORY && error != ERROR_COMMITMENT_LIMIT)
 			return res;
 		StringBuffer<128> reason;
-		reason.Append(TC("VirtualAlloc ")).AppendValue(dwSize);
+		reason.Append(TCV("VirtualAlloc ")).AppendValue(dwSize);
 		Rpc_AllocFailed(reason.data, error);
 		++counter;
 
diff --git a/Engine/Source/Programs/UnrealBuildAccelerator/Detours/Private/Windows/UbaDetoursFunctionsNtDll.inl b/Engine/Source/Programs/UnrealBuildAccelerator/Detours/Private/Windows/UbaDetoursFunctionsNtDll.inl
--- a/Engine/Source/Programs/UnrealBuildAccelerator/Detours/Private/Windows/UbaDetoursFunctionsNtDll.inl
+++ b/Engine/Source/Programs/UnrealBuildAccelerator/Detours/Private/Windows/UbaDetoursFunctionsNtDll.inl
@@ -714,7 +714,7 @@ NTSTATUS NTAPI Local_NtCreateFile(bool IsCreateFunc, PHANDLE hFileHandle, ACCESS
 				return res;
 
 			StringView fileName(ObjectAttributes->ObjectName->Buffer, ObjectAttributes->ObjectName->Length / 2);
-			if (!fileName.EndsWith(TC(".pch")))
+			if (!fileName.EndsWith(TCV(".pch")))
 				return res;
 
 			#if UBA_DEBUG
diff --git a/Engine/Source/Programs/UnrealBuildAccelerator/Host/Private/UbaExports.cpp b/Engine/Source/Programs/UnrealBuildAccelerator/Host/Private/UbaExports.cpp
--- a/Engine/Source/Programs/UnrealBuildAccelerator/Host/Private/UbaExports.cpp
+++ b/Engine/Source/Programs/UnrealBuildAccelerator/Host/Private/UbaExports.cpp
@@ -39,7 +39,7 @@ namespace uba
 		if (prefixLen && strLen + prefixLen + 3 < strBuf.capacity) // TODO: Send prefix and prefixLen through callback
 		{
 			strBuf.Append(prefix, prefixLen);
-			strBuf.Append(TC(" - "), 3);
+			strBuf.Append(TCV(" - "));
 			strBuf.Append(str, strLen);
 			strLen += prefixLen + 3;
 			str = strBuf.data;
@@ -100,7 +100,7 @@ namespace uba
 		if (!fileName)
 		{
 			GetDirectoryOfCurrentModule(logger, temp);
-			temp.EnsureEndsWithSlash().Append(TC("UbaHost.toml"));
+			temp.EnsureEndsWithSlash().Append(TCV("UbaHost.toml"));
 			fileName = temp.data;
 		}
 
diff --git a/Engine/Source/Programs/UnrealBuildAccelerator/ObjTool/Private/UbaImportLibWriter.cpp b/Engine/Source/Programs/UnrealBuildAccelerator/ObjTool/Private/UbaImportLibWriter.cpp
--- a/Engine/Source/Programs/UnrealBuildAccelerator/ObjTool/Private/UbaImportLibWriter.cpp
+++ b/Engine/Source/Programs/UnrealBuildAccelerator/ObjTool/Private/UbaImportLibWriter.cpp
@@ -634,7 +634,7 @@ namespace uba
 			StringBuffer<> expFile;
 			expFile.Append(libFile);
 			if (const tchar* dot = expFile.Last('.'))
-				expFile.Resize(dot - expFile.data).Append(TC(".exp"));
+				expFile.Resize(dot - expFile.data).Append(TCV(".exp"));
 			FileAccessor faExp(logger, expFile.data);
 			if (!faExp.CreateWrite())
 				return false;
diff --git a/Engine/Source/Programs/UnrealBuildAccelerator/ObjTool/Private/UbaObjTool.cpp b/Engine/Source/Programs/UnrealBuildAccelerator/ObjTool/Private/UbaObjTool.cpp
--- a/Engine/Source/Programs/UnrealBuildAccelerator/ObjTool/Private/UbaObjTool.cpp
+++ b/Engine/Source/Programs/UnrealBuildAccelerator/ObjTool/Private/UbaObjTool.cpp
@@ -114,44 +114,44 @@ namespace uba
 
 				if (isImpLibRsp)
 				{
-					if (name.Equals(TC("/NOLOGO")))
+					if (name.Equals(TCV("/NOLOGO")))
 					{
 					}
-					else if (name.Equals(TC("/errorReport")))
+					else if (name.Equals(TCV("/errorReport")))
 					{
 					}
-					else if (name.Equals(TC("/MACHINE")))
+					else if (name.Equals(TCV("/MACHINE")))
 					{
-						//if (!value.Equals(TC("x64")))
+						//if (!value.Equals(TCV("x64")))
 						//{
 						//	logger.Error(TC("only x64 supported"));
 						//	return 0;
 						//}
 					}
-					else if (name.Equals(TC("/SUBSYSTEM")))
+					else if (name.Equals(TCV("/SUBSYSTEM")))
 					{
 					}
-					else if (name.Equals(TC("/DEF")))
+					else if (name.Equals(TCV("/DEF")))
 					{
 						writeImpLib = true;
 					}
-					else if (name.Equals(TC("/NAME")))
+					else if (name.Equals(TCV("/NAME")))
 					{
 						char buffer[256];
 						value.Parse(buffer, sizeof(buffer));
 						impLibName = buffer;
 					}
-					else if (name.Equals(TC("/OUT")))
+					else if (name.Equals(TCV("/OUT")))
 					{
 						impLibFile = value.data;
 					}
-					else if (name.Equals(TC("/IGNORE")))
+					else if (name.Equals(TCV("/IGNORE")))
 					{
 					}
-					else if (name.Equals(TC("/NODEFAULTLIB")))
+					else if (name.Equals(TCV("/NODEFAULTLIB")))
 					{
 					}
-					else if (name.Equals(TC("/LTCG")))
+					else if (name.Equals(TCV("/LTCG")))
 					{
 					}
 					else if (name.StartsWith(TC("/OPT:")))
@@ -182,26 +182,26 @@ namespace uba
 				{
 					moduleName = value.data;
 				}
-				else if (name.Equals(TC("-printsymbols")))
+				else if (name.Equals(TCV("-printsymbols")))
 				{
 					printSymbols = true;
 				}
-				else if (name.Equals(TC("-writeimplib")))
+				else if (name.Equals(TCV("-writeimplib")))
 				{
 					impLibFile = value.data;
 					writeImpLib = true;
 					allowLibInputs = true;
 				}
-				else if (name.Equals(TC("-stripexports")))
+				else if (name.Equals(TCV("-stripexports")))
 				{
 					stripExports = true;
 				}
-				else if (name.Equals(TC("/LIB")))
+				else if (name.Equals(TCV("/LIB")))
 				{
 					isImpLibRsp = true;
 					writeImpLib = true;
 				}
-				else if (name.Equals(TC("-?")))
+				else if (name.Equals(TCV("-?")))
 				{
 					return PrintHelp(TC(""));
 				}
@@ -344,7 +344,7 @@ namespace uba
 				auto& o = *it;
 				StringBuffer<> fixedPath;
 				FixPath(o.c_str(), currentDir.data, currentDir.count, fixedPath);
-				if (fixedPath.EndsWith(TC(".res")) || (fixedPath.EndsWith(TC(".lib")) && !allowLibInputs))
+				if (fixedPath.EndsWith(TCV(".res")) || (fixedPath.EndsWith(TCV(".lib")) && !allowLibInputs))
 					return;
 				if (ObjectFile* objectFile = ObjectFile::OpenAndParse(logger, fixedPath.data))
 				{
@@ -390,7 +390,7 @@ namespace uba
 				const tchar* lastDot = TStrrchr(fileName, '.');
 				UBA_ASSERT(lastDot);
 				StringBuffer<> exportsFile;
-				exportsFile.Append(fileName, lastDot - fileName).Append(TC(".exi"));
+				exportsFile.Append(fileName, lastDot - fileName).Append(TCV(".exi"));
 
 				if (!objectFile->WriteImportsAndExports(logger, exportsFile.data))
 					return false;
@@ -398,7 +398,7 @@ namespace uba
 #if 0
 				objectFile->StripExports(logger);
 				StringBuffer<> newFilename;
-				newFilename.Append(fileName, lastDot - fileName).Append(TC(".TEST")).Append(lastDot);
+				newFilename.Append(fileName, lastDot - fileName).Append(TCV(".TEST")).Append(lastDot);
 				FileAccessor fa(logger, newFilename.data);
 				fa.CreateWrite();
 				fa.Write(objectFile->GetData(), objectFile->GetDataSize());
diff --git a/Engine/Source/Programs/UnrealBuildAccelerator/Test/Private/UbaTestBasics.cpp b/Engine/Source/Programs/UnrealBuildAccelerator/Test/Private/UbaTestBasics.cpp
--- a/Engine/Source/Programs/UnrealBuildAccelerator/Test/Private/UbaTestBasics.cpp
+++ b/Engine/Source/Programs/UnrealBuildAccelerator/Test/Private/UbaTestBasics.cpp
@@ -148,7 +148,7 @@ namespace uba
 	bool TestFiles(LoggerWithWriter& logger, const StringBufferBase& rootDir)
 	{
 		StringBuffer<> testFileName(rootDir);
-		testFileName.Append(TC("UbaTestFile"));
+		testFileName.Append(TCV("UbaTestFile"));
 
 		FileAccessor fileHandle(logger, testFileName.data);
 		if (!fileHandle.CreateWrite())
@@ -219,7 +219,7 @@ namespace uba
 			return logger.Error(TC("UbaTestFile not found"));
 
 		StringBuffer<> testFileName2(rootDir);
-		testFileName2.Append(TC("UbaTestFile2"));
+		testFileName2.Append(TCV("UbaTestFile2"));
 
 		DeleteFileW(testFileName2.data);
 
@@ -267,7 +267,7 @@ namespace uba
 		u32 foundCount = 0;
 
 		StringBuffer<> testDir(rootDir);
-		testDir.Append(TC("TraverseDir")).EnsureEndsWithSlash();
+		testDir.Append(TCV("TraverseDir")).EnsureEndsWithSlash();
 		CHECK_TRUE(CreateDirectoryW(testDir.data));
 
 		CHECK_TRUE(TraverseDir(nullLogger, testDir, [&](const DirectoryEntry&) { ++foundCount; }, true));
@@ -275,7 +275,7 @@ namespace uba
 
 		bool isFile = false;
 		StringBuffer<> entry(testDir);
-		entry.Append(TC("Entry"));
+		entry.Append(TCV("Entry"));
 		FileAccessor fileHandle(logger, entry.data);
 		CHECK_TRUE(fileHandle.CreateWrite(false));
 		CHECK_TRUE(fileHandle.Close());
@@ -294,7 +294,7 @@ namespace uba
 	bool TestOverlappedIO(LoggerWithWriter& logger, const StringBufferBase& rootDir)
 	{
 		StringBuffer<> testFileName(rootDir);
-		testFileName.Append(TC("UbaTestFile"));
+		testFileName.Append(TCV("UbaTestFile"));
 
 		u64 left = 2*1024*1024;
 
@@ -467,7 +467,7 @@ namespace uba
 				else
 				{
 					temp.Append(str, strLen);
-					if (!temp.Equals(TC("foo")))
+					if (!temp.Equals(TCV("foo")))
 						success = false;
 				}
 			}, false, TC(""));
@@ -518,7 +518,7 @@ namespace uba
 		LoggerWithWriter logger(g_consoleLogWriter);
 		logger.Info(TC("IMPORT: %s"), import);
 		StringBuffer<> path(TC("[REDACTED_USER_HOME]/p4/fn/Engine/Binaries/Mac"));
-		path.EnsureEndsWithSlash().Append(TC(import));
+		path.EnsureEndsWithSlash().Append(import);
 		StringBuffer<> error;
 		FindImportsMac(path.data, LogImports, error);
 	}
@@ -530,7 +530,7 @@ namespace uba
 #if PLATFORM_WINDOWS
 		StringBuffer<> path;
 		GetDirectoryOfCurrentModule(logger, path);
-		path.EnsureEndsWithSlash().Append(TC("UbaTestApp.exe"));
+		path.EnsureEndsWithSlash().Append(TCV("UbaTestApp.exe"));
 		bool importKernel = false;
 		StringBuffer<> error;
 		FindImports(path.data, [&](const tchar* import, bool isKnown, const char* const* importLoaderPaths)
diff --git a/Engine/Source/Programs/UnrealBuildAccelerator/Test/Private/UbaTestCache.cpp b/Engine/Source/Programs/UnrealBuildAccelerator/Test/Private/UbaTestCache.cpp
--- a/Engine/Source/Programs/UnrealBuildAccelerator/Test/Private/UbaTestCache.cpp
+++ b/Engine/Source/Programs/UnrealBuildAccelerator/Test/Private/UbaTestCache.cpp
@@ -500,7 +500,7 @@ namespace uba
 		NetworkServer server(ctorSuccess, { logWriter });
 
 		StringBuffer<MaxPath> rootDir;
-		rootDir.Append(testRootDir).Append(TC("Uba"));
+		rootDir.Append(testRootDir).Append(TCV("Uba"));
 		if (!DeleteAllFiles(logger, rootDir.data))
 			return false;
 
@@ -527,7 +527,7 @@ namespace uba
 		auto disconnectServer = MakeGuard([&]() { server.DisconnectClients(); });
 
 		StringBuffer<MaxPath> workingDir;
-		workingDir.Append(testRootDir).Append(TC("WorkingDir"));
+		workingDir.Append(testRootDir).Append(TCV("WorkingDir"));
 		if (!DeleteAllFiles(logger, workingDir.data))
 			return false;
 		if (!storageServer.CreateDirectory(workingDir.data))
diff --git a/Engine/Source/Programs/UnrealBuildAccelerator/Test/Private/UbaTestSession.cpp b/Engine/Source/Programs/UnrealBuildAccelerator/Test/Private/UbaTestSession.cpp
--- a/Engine/Source/Programs/UnrealBuildAccelerator/Test/Private/UbaTestSession.cpp
+++ b/Engine/Source/Programs/UnrealBuildAccelerator/Test/Private/UbaTestSession.cpp
@@ -16,7 +16,7 @@ namespace uba
 		LogWriter& logWriter = logger.m_writer;
 
 		StringBuffer<MaxPath> rootDir;
-		rootDir.Append(testRootDir).Append(TC("Uba"));
+		rootDir.Append(testRootDir).Append(TCV("Uba"));
 		if (!DeleteAllFiles(logger, rootDir.data))
 			return false;
 
@@ -38,7 +38,7 @@ namespace uba
 		SessionServer session(sessionServerInfo);
 
 		StringBuffer<MaxPath> workingDir;
-		workingDir.Append(testRootDir).Append(TC("WorkingDir"));
+		workingDir.Append(testRootDir).Append(TCV("WorkingDir"));
 		if (!DeleteAllFiles(logger, workingDir.data))
 			return false;
 
@@ -61,7 +61,7 @@ namespace uba
 		NetworkClient client(ctorSuccess, { logWriter });
 
 		StringBuffer<MaxPath> rootDir;
-		rootDir.Append(testRootDir).Append(TC("Uba"));
+		rootDir.Append(testRootDir).Append(TCV("Uba"));
 		if (deleteAll && !DeleteAllFiles(logger, rootDir.data))
 			return false;
 
@@ -115,7 +115,7 @@ namespace uba
 		if (!client.FetchConfig(config))
 			return false;
 
-		rootDir.Append(TC("Client"));
+		rootDir.Append(TCV("Client"));
 		if (deleteAll && !DeleteAllFiles(logger, rootDir.data))
 			return false;
 
@@ -138,7 +138,7 @@ namespace uba
 		auto cg = MakeGuard([&]() { sessionClient.Stop(); disconnectGuard.Execute(); });
 
 		StringBuffer<MaxPath> workingDir;
-		workingDir.Append(testRootDir).Append(TC("WorkingDir"));
+		workingDir.Append(testRootDir).Append(TCV("WorkingDir"));
 		if (deleteAll && !DeleteAllFiles(logger, workingDir.data))
 			return false;
 		if (!storageServer.CreateDirectory(workingDir.data))
@@ -191,20 +191,20 @@ namespace uba
 
 		{
 			StringBuffer<MaxPath> dir;
-			dir.Append(workingDir).Append(TC("Dir1"));
+			dir.Append(workingDir).Append(TCV("Dir1"));
 			if (!CreateDirectoryW(dir.data))
 				return logger.Error(TC("Failed to create dir %s"), dir.data);
 
-			dir.Clear().Append(workingDir).Append(TC("Dir2"));
+			dir.Clear().Append(workingDir).Append(TCV("Dir2"));
 			if (!CreateDirectoryW(dir.data))
 				return logger.Error(TC("Failed to create dir %s"), dir.data);
-			dir.EnsureEndsWithSlash().Append(TC("Dir3"));
+			dir.EnsureEndsWithSlash().Append(TCV("Dir3"));
 			if (!CreateDirectoryW(dir.data))
 				return logger.Error(TC("Failed to create dir %s"), dir.data);
-			dir.EnsureEndsWithSlash().Append(TC("Dir4"));
+			dir.EnsureEndsWithSlash().Append(TCV("Dir4"));
 			if (!CreateDirectoryW(dir.data))
 				return logger.Error(TC("Failed to create dir %s"), dir.data);
-			dir.EnsureEndsWithSlash().Append(TC("Dir5"));
+			dir.EnsureEndsWithSlash().Append(TCV("Dir5"));
 			if (!CreateDirectoryW(dir.data))
 				return logger.Error(TC("Failed to create dir %s"), dir.data);
 		}
@@ -231,13 +231,13 @@ namespace uba
 
 		{
 			StringBuffer<MaxPath> fileW2;
-			fileW2.Append(workingDir).Append(TC("FileW2"));
+			fileW2.Append(workingDir).Append(TCV("FileW2"));
 			if (!FileExists(logger, fileW2.data))
 				return logger.Error(TC("Can't find file %s"), fileW2.data);
 		}
 		{
 			StringBuffer<MaxPath> fileWF;
-			fileWF.Append(workingDir).Append(TC("FileWF"));
+			fileWF.Append(workingDir).Append(TCV("FileWF"));
 			if (!FileExists(logger, fileWF.data))
 				return logger.Error(TC("Can't find file %s"), fileWF.data);
 		}
@@ -267,7 +267,7 @@ namespace uba
 	bool RunClang(LoggerWithWriter& logger, SessionServer& session, const tchar* workingDir, const RunProcessFunction& runProcess)
 	{
 		StringBuffer<MaxPath> sourceFile;
-		sourceFile.Append(workingDir).Append(TC("Code.cpp"));
+		sourceFile.Append(workingDir).Append(TCV("Code.cpp"));
 		FileAccessor codeFile(logger, sourceFile.data);
 		if (!codeFile.CreateWrite())
 			return false;
@@ -305,7 +305,7 @@ namespace uba
 		args.Append("-isysroot ");
 		args.Append(xcodeSDKPath.data).Append(' ');
 #endif
-		args.Append(TC("-o code Code.cpp"));
+		args.Append(TCV("-o code Code.cpp"));
 
 		processInfo.arguments = args.data;
 
@@ -399,7 +399,7 @@ namespace uba
 		return RunLocal(logger, testRootDir, [](LoggerWithWriter& logger, SessionServer& session, const tchar* workingDir, const RunProcessFunction& runProcess)
 			{
 				StringBuffer<> file;
-				file.Append(workingDir).Append(TC("TouchFile.h"));
+				file.Append(workingDir).Append(TCV("TouchFile.h"));
 				FileAccessor fr(logger, file.data);
 
 				CHECK_TRUE(fr.CreateWrite());
@@ -557,7 +557,7 @@ namespace uba
 		CHECK_TRUE(dirTable.EntryExists(ToView(workingDir), true) == DirectoryTable::Exists_Yes);
 
 		StringBuffer<MaxPath> sourceFile;
-		sourceFile.Append(workingDir).Append(TC("Code.cpp"));
+		sourceFile.Append(workingDir).Append(TCV("Code.cpp"));
 
 		CHECK_TRUE(GetAttributes(sourceFile) == INVALID_FILE_ATTRIBUTES);
 		FileAccessor codeFile(logger, sourceFile.data);
@@ -571,7 +571,7 @@ namespace uba
 		CHECK_TRUE(dirTable.EntryExists(sourceFile) == DirectoryTable::Exists_Yes);
 
 		StringBuffer<MaxPath> newDir;
-		newDir.Append(workingDir).Append(TC("NewDir"));
+		newDir.Append(workingDir).Append(TCV("NewDir"));
 		StringBuffer<MaxPath> newDirAndSlash(newDir);
 		newDirAndSlash.Append('/');
 
@@ -585,7 +585,7 @@ namespace uba
 		CHECK_TRUE(GetAttributes(newDirAndSlash) != INVALID_FILE_ATTRIBUTES);
 
 		StringBuffer<MaxPath> newDir2;
-		newDir2.Append(workingDir).Append(TC("NewDir2"));
+		newDir2.Append(workingDir).Append(TCV("NewDir2"));
 		CHECK_TRUE(CreateDirectoryW(newDir2.data));
 		CHECK_TRUE(GetAttributes(newDir2) == INVALID_FILE_ATTRIBUTES);
 		CHECK_TRUE(session.RefreshDirectory(workingDir))
diff --git a/Engine/Source/Programs/UnrealBuildAccelerator/Test/Private/UbaTestStdOut.cpp b/Engine/Source/Programs/UnrealBuildAccelerator/Test/Private/UbaTestStdOut.cpp
--- a/Engine/Source/Programs/UnrealBuildAccelerator/Test/Private/UbaTestStdOut.cpp
+++ b/Engine/Source/Programs/UnrealBuildAccelerator/Test/Private/UbaTestStdOut.cpp
@@ -17,7 +17,7 @@ namespace uba
 		NetworkServer networkServer(ctorSuccess, { logWriter });
 
 		StringBuffer<> rootDir;
-		rootDir.Append(testRootDir).Append(TC("Uba"));
+		rootDir.Append(testRootDir).Append(TCV("Uba"));
 		if (!DeleteAllFiles(logger, rootDir.data))
 			return false;
 
@@ -35,7 +35,7 @@ namespace uba
 		auto sg = MakeGuard([&]() { networkServer.DisconnectClients(); });
 
 		StringBuffer<> workingDir;
-		workingDir.Append(testRootDir).Append(TC("WorkingDir"));
+		workingDir.Append(testRootDir).Append(TCV("WorkingDir"));
 		if (!DeleteAllFiles(logger, workingDir.data))
 			return false;
 		if (!storageServer.CreateDirectory(workingDir.data))
@@ -102,9 +102,9 @@ namespace uba
 		if (!IsWindows)
 			return true;
 		StringBuffer<> args;
-		args.Append(TC("/c \""));
+		args.Append(TCV("/c \""));
 		GetTestAppPath(logger, args);
-		args.Append(TC(" -stdout=foo\""));
+		args.Append(TCV(" -stdout=foo\""));
 		return TestStdOut(logger, testRootDir, false, TC("c:\\windows\\system32\\cmd.exe"), args.data, TC("foo"));
 	}
 }
diff --git a/Engine/Source/Programs/UnrealBuildAccelerator/Test/Private/UbaTestStorage.cpp b/Engine/Source/Programs/UnrealBuildAccelerator/Test/Private/UbaTestStorage.cpp
--- a/Engine/Source/Programs/UnrealBuildAccelerator/Test/Private/UbaTestStorage.cpp
+++ b/Engine/Source/Programs/UnrealBuildAccelerator/Test/Private/UbaTestStorage.cpp
@@ -15,7 +15,7 @@ namespace uba
 		#endif
 
 		StringBuffer<> rootDir;
-		rootDir.Append(testRootDir).Append(TC("Uba"));
+		rootDir.Append(testRootDir).Append(TCV("Uba"));
 		StorageCreateInfo storageInfo(rootDir.data, logger.m_writer);
 		storageInfo.casCapacityBytes = 1024ull * 1024 * 1024;
 		StorageImpl storage(storageInfo);
@@ -34,7 +34,7 @@ namespace uba
 			return logger.Error(TC("Failed to find file %s"), detoursLib.data);
 
 		StringBuffer<> detoursLibCopy(detoursLib);
-		detoursLibCopy.Append(TC(".tmp"));
+		detoursLibCopy.Append(TCV(".tmp"));
 
 		auto deleteFile = MakeGuard([&]() { return DeleteFileW(detoursLibCopy.data); });
 
@@ -100,7 +100,7 @@ namespace uba
 		auto dg = MakeGuard([&]() { client.Disconnect(); });
 
 		StringBuffer<> rootDir;
-		rootDir.Append(testRootDir).Append(TC("Uba"));
+		rootDir.Append(testRootDir).Append(TCV("Uba"));
 		if (!DeleteAllFiles(logger, rootDir.data))
 			return false;
 
@@ -110,7 +110,7 @@ namespace uba
 
 		auto g = MakeGuard([&]() { server.DisconnectClients(); });
 
-		rootDir.Append(TC("Client"));
+		rootDir.Append(TCV("Client"));
 		if (!DeleteAllFiles(logger, rootDir.data))
 			return false;
 
@@ -132,7 +132,7 @@ namespace uba
 
 		StringBuffer<> fileName;
 		{
-			fileName.Append(rootDir).Append(TC("UbaTestFile"));
+			fileName.Append(rootDir).Append(TCV("UbaTestFile"));
 			FileAccessor fileHandle(logger, fileName.data);
 			if (!fileHandle.CreateWrite())
 				return logger.Error(TC("Failed to create file for write"));
@@ -148,7 +148,7 @@ namespace uba
 		if (!storageClient.StoreCasFileClient(key, ToStringKeyLower(fileName), fileName.data, FileMappingHandle{}, 0, 0, TC("UbaTestFile"), false, storeCompressed))
 			return logger.Error(TC("Failed to store file %s"), fileName.data);
 
-		fileName.Clear().Append(testRootDir).Append(TC("Uba")).Append(PathSeparator).Append(TC("UbaTestFile"));
+		fileName.Clear().Append(testRootDir).Append(TCV("Uba")).Append(PathSeparator).Append(TCV("UbaTestFile"));
 		if (!storageServer.CopyOrLink(key, fileName.data, DefaultAttributes()))
 			return logger.Error(TC("Failed to copy cas to file %s"), fileName.data);
 		return true;
@@ -166,7 +166,7 @@ namespace uba
 		auto dg = MakeGuard([&]() { client.Disconnect(); });
 
 		StringBuffer<> rootDir;
-		rootDir.Append(testRootDir).Append(TC("Uba"));
+		rootDir.Append(testRootDir).Append(TCV("Uba"));
 		if (!DeleteAllFiles(logger, rootDir.data))
 			return false;
 
@@ -177,7 +177,7 @@ namespace uba
 
 		auto g = MakeGuard([&]() { server.DisconnectClients(); });
 
-		rootDir.Append(TC("Client"));
+		rootDir.Append(TCV("Client"));
 		if (!DeleteAllFiles(logger, rootDir.data))
 			return false;
 
@@ -200,7 +200,7 @@ namespace uba
 		StringBuffer<> fileName;
 #if 1
 		{
-			fileName.Append(rootDir).Append(TC("UbaTestFile"));
+			fileName.Append(rootDir).Append(TCV("UbaTestFile"));
 			FileAccessor fileHandle(logger, fileName.data);
 			if (!fileHandle.CreateWrite())
 				return logger.Error(TC("Failed to create file for write"));
@@ -211,7 +211,7 @@ namespace uba
 				return false;
 		}
 #else
-		fileName.Append(TC("e:\\dev\\fn\\QAGame\\Saved\\StagedBuilds\\PS5_Temp\\Split\\qagame\\content\\paks\\qagame-ps5.ucas\\.copied"));
+		fileName.Append(TCV("e:\\dev\\fn\\QAGame\\Saved\\StagedBuilds\\PS5_Temp\\Split\\qagame\\content\\paks\\qagame-ps5.ucas\\.copied"));
 #endif
 
 		CasKey casKey;
diff --git a/Engine/Source/Programs/UnrealBuildAccelerator/Test/Private/UbaTestStress.cpp b/Engine/Source/Programs/UnrealBuildAccelerator/Test/Private/UbaTestStress.cpp
--- a/Engine/Source/Programs/UnrealBuildAccelerator/Test/Private/UbaTestStress.cpp
+++ b/Engine/Source/Programs/UnrealBuildAccelerator/Test/Private/UbaTestStress.cpp
@@ -20,7 +20,7 @@ namespace uba
 		NetworkServer networkServer(ctorSuccess, { logWriter });
 
 		StringBuffer<> rootDir;
-		rootDir.Append(testRootDir).Append(TC("Uba"));
+		rootDir.Append(testRootDir).Append(TCV("Uba"));
 		if (!DeleteAllFiles(logger, rootDir.data))
 			return false;
 
@@ -38,7 +38,7 @@ namespace uba
 		auto sg = MakeGuard([&]() { networkServer.DisconnectClients(); });
 
 		StringBuffer<> workingDir;
-		workingDir.Append(testRootDir).Append(TC("WorkingDir"));
+		workingDir.Append(testRootDir).Append(TCV("WorkingDir"));
 		if (!DeleteAllFiles(logger, workingDir.data))
 			return false;
 		if (!storageServer.CreateDirectory(workingDir.data))
@@ -65,7 +65,7 @@ namespace uba
 			return logger.Error(TC("uri/pool/oidc must be populated for this test to work"));
 
 			StringBuffer<512> coordinatorWorkDir(testRootDir);
-			coordinatorWorkDir.EnsureEndsWithSlash().Append(TC("Horde"));
+			coordinatorWorkDir.EnsureEndsWithSlash().Append(TCV("Horde"));
 			StringBuffer<512> binariesDir;
 			if (!GetDirectoryOfCurrentModule(logger, binariesDir))
 				return false;
@@ -94,7 +94,7 @@ namespace uba
 		while (true)
 		{
 			StringBuffer<512> fileName;
-			fileName.Append(workingDir).Append(TC("File_")).AppendValue(counter++).Append(TC(".in"));
+			fileName.Append(workingDir).Append(TCV("File_")).AppendValue(counter++).Append(TCV(".in"));
 			{
 				FileAccessor fa(logger, fileName.data);
 				fa.CreateWrite();
@@ -105,7 +105,7 @@ namespace uba
 			}
 			sessionServer.RegisterNewFile(fileName.data);
 			StringBuffer<512> args;
-			args.Append(TC("-file=")).Append(fileName);
+			args.Append(TCV("-file=")).Append(fileName);
 			ProcessStartInfo pi;
 			pi.application = testApp.data;
 			pi.arguments = args.data;
@@ -140,7 +140,7 @@ namespace uba
 							rec->error = true;
 						}
 						StringBuffer<256> outFile;
-						outFile.Append(rec->inFile.data, rec->inFile.count - 2).Append(TC("out"));
+						outFile.Append(rec->inFile.data, rec->inFile.count - 2).Append(TCV("out"));
 						if (!FileExists(logger, outFile.data))
 						{
 							logger.Error(TC("Couldn't find output file %s"), outFile.data);
diff --git a/Engine/Source/Programs/UnrealBuildAccelerator/Visualizer/Private/UbaTraceReader.cpp b/Engine/Source/Programs/UnrealBuildAccelerator/Visualizer/Private/UbaTraceReader.cpp
--- a/Engine/Source/Programs/UnrealBuildAccelerator/Visualizer/Private/UbaTraceReader.cpp
+++ b/Engine/Source/Programs/UnrealBuildAccelerator/Visualizer/Private/UbaTraceReader.cpp
@@ -407,7 +407,7 @@ namespace uba
 			}
 
 			StringBuffer<> fullName;
-			fullName.Append(sessionName).Append(TC(" (")).Append(sessionInfo).Append(TC(")"));
+			fullName.Append(sessionName).Append(TCV(" (")).Append(sessionInfo).Append(TCV(")"));
 
 			// Check if we can re-use existing session (same machine was disconnected and then reconnected)
 			u32 virtualSessionIndex = sessionIndex;
diff --git a/Engine/Source/Programs/UnrealBuildAccelerator/Visualizer/Private/UbaVisualizer.cpp b/Engine/Source/Programs/UnrealBuildAccelerator/Visualizer/Private/UbaVisualizer.cpp
--- a/Engine/Source/Programs/UnrealBuildAccelerator/Visualizer/Private/UbaVisualizer.cpp
+++ b/Engine/Source/Programs/UnrealBuildAccelerator/Visualizer/Private/UbaVisualizer.cpp
@@ -287,7 +287,7 @@ namespace uba
 						filter.Appendf(L"_%s%u", ownerInfo.id, ownerInfo.pid);
 				}
 
-				if (!traceName.Equals(m_newTraceName.data) && traceName.EndsWith(filter.data))
+				if (!traceName.Equals(m_newTraceName.data) && traceName.EndsWith(filter))
 				{
 					m_newTraceName.Clear().Append(traceName);
 					m_usingNamed = true;
@@ -967,14 +967,14 @@ namespace uba
 		if (m_traceView.sessions.empty() && !m_parentHwnd)
 		{
 			StringBuffer<> str[2];
-			str[1].Append(TC("Click here to open trace file"));
+			str[1].Append(TCV("Click here to open trace file"));
 			if (m_listenChannel.count)
 			{
 				str[0].Appendf(TC("Listening for new sessions on channel '%s'"), m_listenChannel.data);
-				str[1].Append(TC(" instead"));
+				str[1].Append(TCV(" instead"));
 			}
 			else
-				str[0].Append(TC("No trace active"));
+				str[0].Append(TCV("No trace active"));
 
 			SetActiveFont(m_popupFont);
 			DrawCenteredText(str, 2, false);
@@ -1798,7 +1798,7 @@ namespace uba
 		if (!m_filterString.empty())
 		{
 			StringBuffer<> str[2];
-			str[0].Append(TC("Box Filter: ")).Append(m_filterString);
+			str[0].Append(TCV("Box Filter: ")).Append(m_filterString);
 			str[1].Appendf(TC("Box Count: %u"), visibleBoxes);
 
 			SetActiveFont(m_popupFont);
diff --git a/Engine/Source/Programs/UnrealBuildAccelerator/Visualizer/Private/UbaVisualizerApp.cpp b/Engine/Source/Programs/UnrealBuildAccelerator/Visualizer/Private/UbaVisualizerApp.cpp
--- a/Engine/Source/Programs/UnrealBuildAccelerator/Visualizer/Private/UbaVisualizerApp.cpp
+++ b/Engine/Source/Programs/UnrealBuildAccelerator/Visualizer/Private/UbaVisualizerApp.cpp
@@ -54,7 +54,7 @@ struct MessageBoxLogWriter : public LogWriter
 		{
 			#if UBA_DEBUG
 			StringBuffer<> buf;
-			buf.Append(str, strLen).Append(TC("\r\n"));
+			buf.Append(str, strLen).Append(TCV("\r\n"));
 			OutputDebugStringW(buf.data);
 			#endif
 			return;
@@ -116,53 +116,53 @@ int WINAPI wWinMain(_In_ HINSTANCE hInstance, _In_opt_ HINSTANCE hPrevInstance,
 			name.Append(argv[i]);
 		}
 
-		if (name.Equals(TC("-help")))
+		if (name.Equals(TCV("-help")))
 		{
 			return PrintHelp();
 		}
-		else if (name.Equals(TC("-host")))
+		else if (name.Equals(TCV("-host")))
 		{
 			if (value.IsEmpty())
 				return PrintHelp(TC("-host needs a value"));
 			host.Append(value);
 		}
-		else if (name.Equals(TC("-named")))
+		else if (name.Equals(TCV("-named")))
 		{
 			if (value.IsEmpty())
 				return PrintHelp(TC("-named needs a value"));
 			named.Append(value);
 		}
-		else if (name.Equals(TC("-file")))
+		else if (name.Equals(TCV("-file")))
 		{
 			if (value.IsEmpty())
 				return PrintHelp(TC("-file needs a value"));
 			file.Append(value);
 		}
-		else if (name.Equals(TC("-port")))
+		else if (name.Equals(TCV("-port")))
 		{
 			if (!value.Parse(port))
 				return PrintHelp(TC("Invalid value for -port"));
 		}
-		else if (name.Equals(TC("-listen")))
+		else if (name.Equals(TCV("-listen")))
 		{
 			if (!value.IsEmpty())
 				channel.Append(value.data);
 			else
-				channel.Append(TC("Default"));
+				channel.Append(TCV("Default"));
 		}
-		else if (name.Equals(TC("-replay")))
+		else if (name.Equals(TCV("-replay")))
 		{
 			replay = 1;
 			if (!value.IsEmpty())
 				value.Parse(replay);
 		}
-		else if (name.Equals(TC("-config")))
+		else if (name.Equals(TCV("-config")))
 		{
 			if (value.IsEmpty())
 				return PrintHelp(TC("-config needs a value"));
 			configPath.Append(value);
 		}
-		else if (name.Equals(TC("-parent")))
+		else if (name.Equals(TCV("-parent")))
 		{
 			if (value.IsEmpty())
 				return PrintHelp(TC("-parent needs a value"));
@@ -178,17 +178,17 @@ int WINAPI wWinMain(_In_ HINSTANCE hInstance, _In_opt_ HINSTANCE hPrevInstance,
 			parent = StringToValue2(value.data, value.count);
 			//value.ParseHex(parent);
 		}
-		else if (name.Equals(TC("-nocopy")))
+		else if (name.Equals(TCV("-nocopy")))
 		{
 			copyAndLaunch = false;
 		}
-		else if (name.Equals(TC("-ownerPid")))
+		else if (name.Equals(TCV("-ownerPid")))
 		{
 			u32 ownerPid;
 			if (value.Parse(ownerPid))
 				const_cast<OwnerInfo&>(GetOwnerInfo()).pid = ownerPid;
 		}
-		else if (name.Equals(TC("-ownerId")))
+		else if (name.Equals(TCV("-ownerId")))
 		{
 			TStrcpy_s(const_cast<tchar*>(GetOwnerInfo().id), 260, value.data);
 		}
@@ -285,7 +285,7 @@ int WINAPI wWinMain(_In_ HINSTANCE hInstance, _In_opt_ HINSTANCE hPrevInstance,
 	LocalFree(argv);
 
 	if (host.IsEmpty() && named.IsEmpty() && file.IsEmpty() && !channel.count)
-		channel.Append(TC("Default")); // return PrintHelp(TC("No host/named/file provided. Add -host=<host> or -file=<file> or -named=<name>"));
+		channel.Append(TCV("Default")); // return PrintHelp(TC("No host/named/file provided. Add -host=<host> or -file=<file> or -named=<name>"));
 
 	bool showAllTraces = true;
 	if (!configPath.count)
