[UBA]
* Fixed unit test

diff --git a/Engine/Source/Programs/UnrealBuildAccelerator/Common/Private/UbaCacheClient.cpp b/Engine/Source/Programs/UnrealBuildAccelerator/Common/Private/UbaCacheClient.cpp
--- a/Engine/Source/Programs/UnrealBuildAccelerator/Common/Private/UbaCacheClient.cpp
+++ b/Engine/Source/Programs/UnrealBuildAccelerator/Common/Private/UbaCacheClient.cpp
@@ -1428,7 +1428,7 @@ namespace uba
 		}
 
 		// Add arguments list to key
-		auto hashString = [&](const tchar* str, u64 strLen, u32 rootPos) { hasher.Update(str, strLen*sizeof(tchar)); };
+		auto hashString = [&](const tchar* str, u64 strLenIncTerm, u32 rootPos) { hasher.Update(str, strLenIncTerm*sizeof(tchar)); };
 		if (!rootPaths.NormalizeString(m_logger, info.arguments, TStrlen(info.arguments), hashString, false, info.GetDescription(), TC(" calculating command line hash")))
 		{
 			if (report)
diff --git a/Engine/Source/Programs/UnrealBuildAccelerator/Test/Private/UbaTestBasics.cpp b/Engine/Source/Programs/UnrealBuildAccelerator/Test/Private/UbaTestBasics.cpp
--- a/Engine/Source/Programs/UnrealBuildAccelerator/Test/Private/UbaTestBasics.cpp
+++ b/Engine/Source/Programs/UnrealBuildAccelerator/Test/Private/UbaTestBasics.cpp
@@ -454,11 +454,11 @@ namespace uba
 		bool success = true;
 		StringBuffer<> temp;
 		u32 rootPos = ~0u;
-		bool res = paths.NormalizeString(logger, str, sizeof_array(str), [&](const tchar* str, u64 strLen, u32 rp)
+		bool res = paths.NormalizeString(logger, str, sizeof_array(str), [&](const tchar* str, u64 strLenIncTerm, u32 rp)
 			{
 				if (rp != ~0u)
 				{
-					if (strLen != 1)
+					if (strLenIncTerm != 1)
 						success = false;
 					if (str[0] != RootPaths::RootStartByte + PathsPerRoot + (IsWindows ? 1 : 0)) // Add one for windows because second entry is backslash
 						success = false;
@@ -466,7 +466,7 @@ namespace uba
 				}
 				else
 				{
-					temp.Append(str, strLen);
+					temp.Append(str, strLenIncTerm - 1);
 					if (!temp.Equals(TCV("foo")))
 						success = false;
 				}
@@ -486,7 +486,7 @@ namespace uba
 		#if PLATFORM_WINDOWS
 		const tchar str2[] = TC("file://e:/temp/");
 		bool foundPath = false;
-		res = paths.NormalizeString(logger, str2, sizeof_array(str2), [&](const tchar* str, u64 strLen, u32 rp)
+		res = paths.NormalizeString(logger, str2, sizeof_array(str2), [&](const tchar* str, u64 strLenIncTerm, u32 rp)
 			{
 				if (rp != ~0)
 				{
@@ -496,7 +496,7 @@ namespace uba
 				}
 				else
 				{
-					if (!((strLen == 1 && (!str[0] || str[0] == '/')) || (strLen == 6 && Equals(str, TC("file:/"), 6, false))))
+					if (!((strLenIncTerm == 1 && (!str[0] || str[0] == '/')) || (strLenIncTerm == 6 && Equals(str, TC("file:/"), 6, false))))
 						success = false;
 				}
 
