[UBA]
* Fixed bug with StringBuffer.StartsWith where size comparisson was inverted

diff --git a/Engine/Source/Programs/UnrealBuildAccelerator/Core/Private/UbaStringBuffer.cpp b/Engine/Source/Programs/UnrealBuildAccelerator/Core/Private/UbaStringBuffer.cpp
--- a/Engine/Source/Programs/UnrealBuildAccelerator/Core/Private/UbaStringBuffer.cpp
+++ b/Engine/Source/Programs/UnrealBuildAccelerator/Core/Private/UbaStringBuffer.cpp
@@ -312,7 +312,7 @@ namespace uba
 
 	bool StringBufferBase::StartsWith(const StringView& str, bool ignoreCase) const
 	{
-		return str.count >= count && uba::StartsWith(data, str.data, ignoreCase);
+		return count >= str.count && uba::StartsWith(data, str.data, ignoreCase);
 	}
 
 	bool StringBufferBase::EndsWith(const tchar* value, bool ignoreCase) const
