[UBA]
* Fixed so read event from detoured process is checked one last time after process stopped being active. There is a chance a process can both write last message and exit while session process goes from checking event to checking if process is alive

diff --git a/Engine/Source/Programs/UnrealBuildAccelerator/Common/Private/UbaProcess.cpp b/Engine/Source/Programs/UnrealBuildAccelerator/Common/Private/UbaProcess.cpp
--- a/Engine/Source/Programs/UnrealBuildAccelerator/Common/Private/UbaProcess.cpp
+++ b/Engine/Source/Programs/UnrealBuildAccelerator/Common/Private/UbaProcess.cpp
@@ -462,7 +462,8 @@ namespace uba
 			#endif
 
 			if (!IsActive())
-				return false;
+				return m_readEvent.IsSet(0); // Do one more check
+
 			if (IsCancelled())
 				return false;
 		}
