[UBA]
* Fixed compile error when compiling detouring lib on macos debug

diff --git a/Engine/Source/Programs/UnrealBuildAccelerator/Detours/Private/UbaDetoursShared.h b/Engine/Source/Programs/UnrealBuildAccelerator/Detours/Private/UbaDetoursShared.h
--- a/Engine/Source/Programs/UnrealBuildAccelerator/Detours/Private/UbaDetoursShared.h
+++ b/Engine/Source/Programs/UnrealBuildAccelerator/Detours/Private/UbaDetoursShared.h
@@ -136,7 +136,7 @@ namespace uba
 	#define VARIABLE_MEM(type, name) VariableMem<type> name##Mem; type& name = (type&)name##Mem.data;
 
 	#define RPC_MESSAGE(messageName, timerName) \
-		DEBUG_LOG(L"RPC_MESSAGE %hs", #messageName); \
+		DEBUG_LOG(L"RPC_MESSAGE %s", TC(#messageName)); \
 		TimerScope ts(g_stats.timerName); \
 		SCOPED_WRITE_LOCK(g_communicationLock, pcs); \
 		BinaryWriter writer; \
