[UBA]
* Added UbaCacheService -nosave which can be used to test maintenance without saving
* Changed TC( to TCV( where possible

diff --git a/Engine/Source/Programs/UnrealBuildAccelerator/CacheService/Private/UbaCacheService.cpp b/Engine/Source/Programs/UnrealBuildAccelerator/CacheService/Private/UbaCacheService.cpp
--- a/Engine/Source/Programs/UnrealBuildAccelerator/CacheService/Private/UbaCacheService.cpp
+++ b/Engine/Source/Programs/UnrealBuildAccelerator/CacheService/Private/UbaCacheService.cpp
@@ -111,6 +111,7 @@ namespace uba
 		bool storeCompressed = true;
 		bool fullMaintenance = false;
 		bool maintenanceEnabled = true;
+		bool allowSave = true;
 		bool signalHandlerEnabled = true;
 		bool shouldCrash = false;
 		u32 expirationTimeSeconds = DefaultExpiration;
@@ -190,6 +191,10 @@ namespace uba
 			{
 				maintenanceEnabled = false;
 			}
+			else if (name.Equals(TCV("-nosave")))
+			{
+				allowSave = false;
+			}
 			else if (name.Equals(TCV("-config")))
 			{
 				if (value.IsEmpty())
@@ -308,7 +313,7 @@ namespace uba
 			cacheServer.SetForceFullMaintenance();
 
 		if (maintenanceEnabled)
-			if (!cacheServer.RunMaintenance(true, ShouldExit))
+			if (!cacheServer.RunMaintenance(true, allowSave, ShouldExit))
 				return -1;
 
 		HttpServer httpServer(logWriter, networkBackend);
@@ -348,12 +353,13 @@ namespace uba
 			{
 				Sleep(1000);
 				if (maintenanceEnabled)
-					if (!cacheServer.RunMaintenance(false, ShouldExit))
+					if (!cacheServer.RunMaintenance(false, allowSave, ShouldExit))
 						break;
 			}
 		}
 
-		cacheServer.RunMaintenance(false, {});
+		if (maintenanceEnabled)
+			cacheServer.RunMaintenance(false, allowSave, {});
 
 		storageServer.DeleteIsRunningFile();
 
diff --git a/Engine/Source/Programs/UnrealBuildAccelerator/Common/Private/UbaCacheServer.cpp b/Engine/Source/Programs/UnrealBuildAccelerator/Common/Private/UbaCacheServer.cpp
--- a/Engine/Source/Programs/UnrealBuildAccelerator/Common/Private/UbaCacheServer.cpp
+++ b/Engine/Source/Programs/UnrealBuildAccelerator/Common/Private/UbaCacheServer.cpp
@@ -251,7 +251,7 @@ namespace uba
 		u64 startTime = GetTime();
 
 		StringBuffer<> fileName(m_rootDir);
-		fileName.EnsureEndsWithSlash().Append(TC("cachedb"));
+		fileName.EnsureEndsWithSlash().Append(TCV("cachedb"));
 
 		FileAccessor file(m_logger, fileName.data);
 		if (!file.OpenMemoryRead(0, false))
@@ -291,7 +291,7 @@ namespace uba
 		else
 		{
 			StringBuffer<MaxPath> bucketsDir(m_rootDir);
-			bucketsDir.EnsureEndsWithSlash().Append(TC("buckets"));
+			bucketsDir.EnsureEndsWithSlash().Append(TCV("buckets"));
 			TraverseDir(m_logger, bucketsDir, [&](const DirectoryEntry& e)
 				{
 					StringBuffer<128> keyName;
@@ -318,7 +318,7 @@ namespace uba
 					BinaryReader reader(bucketFile.GetData(), 0, bucketFile.GetSize());
 					u32 bucketVersion = reader.ReadU32();
 					LoadBucket(bucket, reader, bucketVersion, stats);
-				});
+				}, TCV("LoadBucket"));
 		}
 
 		u64 duration = GetTime() - startTime;
@@ -388,7 +388,7 @@ namespace uba
 		FileWriter(Logger& l, const tchar* fn)
 		:	logger(l)
 		,	fileName(fn)
-		,	tempFileName(StringBuffer<MaxPath>(fn).Append(TC(".tmp")).data)
+		,	tempFileName(StringBuffer<MaxPath>(fn).Append(TCV(".tmp")).data)
 		,	file(logger, tempFileName.c_str())
 		{
 			tempBuffer = (u8*)malloc(TempBufferSize);
@@ -475,7 +475,7 @@ namespace uba
 		u64 saveStart = GetTime();
 
 		StringBuffer<MaxPath> bucketsDir(m_rootDir);
-		bucketsDir.EnsureEndsWithSlash().Append(TC("buckets"));
+		bucketsDir.EnsureEndsWithSlash().Append(TCV("buckets"));
 		if (!m_storage.CreateDirectory(bucketsDir.data))
 			return false;
 		bucketsDir.EnsureEndsWithSlash();
@@ -554,7 +554,7 @@ namespace uba
 		if (!m_dbfileDirty)
 			return true;
 		StringBuffer<MaxPath> fileName(m_rootDir);
-		fileName.EnsureEndsWithSlash().Append(TC("cachedb"));
+		fileName.EnsureEndsWithSlash().Append(TCV("cachedb"));
 
 		FileWriter file(m_logger, fileName.data);
 		
@@ -583,7 +583,7 @@ namespace uba
 			}
 
 			StringBuffer<MaxPath> bucketsFile(m_rootDir);
-			bucketsFile.EnsureEndsWithSlash().Append(TC("buckets")).EnsureEndsWithSlash().AppendValue(it->first);
+			bucketsFile.EnsureEndsWithSlash().Append(TCV("buckets")).EnsureEndsWithSlash().AppendValue(it->first);
 			DeleteFileW(bucketsFile.data);
 			m_logger.Detail(TC("    Bucket %u was empty. Deleted"), bucket.index);
 			it = m_buckets.erase(it);
@@ -596,7 +596,7 @@ namespace uba
 		m_forceAllSteps = true;
 	}
 
-	bool CacheServer::RunMaintenance(bool force, const Function<bool()>& shouldExit)
+	bool CacheServer::RunMaintenance(bool force, bool allowSave, const Function<bool()>& shouldExit)
 	{
 		if (m_addsSinceMaintenance == 0 && !force && !m_bucketIsOverflowing)
 			return true;
@@ -649,7 +649,7 @@ namespace uba
 			}
 
 
-			if(!RunMaintenanceInternal(shouldExit))
+			if(!RunMaintenanceInternal(shouldExit, allowSave))
 				return false;
 		}
 		while (m_bucketIsOverflowing);
@@ -657,7 +657,7 @@ namespace uba
 		return true;
 	}
 
-	bool CacheServer::RunMaintenanceInternal(const Function<bool()>& shouldExit)
+	bool CacheServer::RunMaintenanceInternal(const Function<bool()>& shouldExit, bool allowSave)
 	{
 		//m_forceAllSteps = true;
 		bool forceAllSteps = m_forceAllSteps;
@@ -683,7 +683,7 @@ namespace uba
 				KernelStats().Write(writer);
 				trace.ProcessExited(0, 0, writer.GetData(), writer.GetPosition(), {});
 				StringBuffer<> traceFile(m_rootDir);
-				traceFile.Append(TC("UbaCacheServer.uba"));
+				traceFile.Append(TCV("UbaCacheServer.uba"));
 				trace.StopWrite(traceFile.data);
 			});
 
@@ -691,8 +691,6 @@ namespace uba
 		bool entriesAdded = addsSinceMaintenance != 0;
 		m_addsSinceMaintenance = 0;
 
-		bool allowSave = true;
-
 		u64 startTime = GetTime();
 
 		if (m_shouldWipe)
@@ -1992,13 +1990,13 @@ namespace uba
 				#endif
 
 				SCOPED_READ_LOCK(entries.lock, lock3);
-				writeLine(line.Clear().Append(TC("   SharedInputs:")));
+				writeLine(line.Clear().Append(TCV("   SharedInputs:")));
 				writePathsFromOffsets(bucket, entries.sharedInputCasKeyOffsets);
 				u32 index = 0;
 				for (auto& entry : entries.entries)
 				{
 					writeLine(line.Clear().Appendf(TC("  #%u (%s)"), index, TimeToText(GetFileTimeAsTime(GetSystemTimeAsFileTime() - (m_creationTime + entry.creationTime)), true).str));
-					writeLine(line.Clear().Append(TC("   InputRanges:")));
+					writeLine(line.Clear().Append(TCV("   InputRanges:")));
 					BinaryReader sharedReader(entries.sharedInputCasKeyOffsets);
 					BinaryReader rangeReader(entry.sharedInputCasKeyOffsetRanges);
 					u32 sharedIndex = 0;
@@ -2022,15 +2020,15 @@ namespace uba
 						writeLine(line.Clear().Appendf(TC("          %llu - %llu   (%llu - %llu)"), startIndex, sharedIndex-1, begin, end));
 					}
 
-					writeLine(line.Clear().Append(TC("   ExtraInputs:")));
+					writeLine(line.Clear().Append(TCV("   ExtraInputs:")));
 					writePathsFromOffsets(bucket, entry.extraInputCasKeyOffsets);
-					writeLine(line.Clear().Append(TC("   Outputs:")));
+					writeLine(line.Clear().Append(TCV("   Outputs:")));
 					writePathsFromOffsets(bucket, entry.outputCasKeyOffsets);
 					++index;
 				}
 			};
 
-		if (command.Equals(TC("content")))
+		if (command.Equals(TCV("content")))
 		{
 			writeLine(TCV("UbaCache server summary"));
 
@@ -2109,17 +2107,17 @@ namespace uba
 						u64 age = GetFileTimeAsTime(now - (m_creationTime + entry.creationTime));
 						writeLine(line.Clear().Appendf(TC("  #%u (%s ago)"), index, TimeToText(age, true).str));
 
-						writeLine(line.Clear().Append(TC("   Inputs:")));
+						writeLine(line.Clear().Append(TCV("   Inputs:")));
 						entries.Flatten(temp, entry);
 						writePathsFromOffsets(bucket, temp);
-						writeLine(line.Clear().Append(TC("   Outputs:")));
+						writeLine(line.Clear().Append(TCV("   Outputs:")));
 						writePathsFromOffsets(bucket, entry.outputCasKeyOffsets);
 						++index;
 					}
 				}
 			}
 		}
-		else if (command.Equals(TC("status")))
+		else if (command.Equals(TCV("status")))
 		{
 			writeLine(TCV("UbaCacheServer status"));
 			writeLine(line.Clear().Appendf(TC("  CreationTime: %s ago"), TimeToText(GetFileTimeAsTime(GetSystemTimeAsFileTime() - m_creationTime), true).str));
@@ -2138,9 +2136,19 @@ namespace uba
 					CasKey mostEntriesKey;
 					u64 lastUsed = 0;
 					u64 totalEntryCount = 0;
+					u64 expiredEntryCount = 0;
 					u64 totalKeySharedBytes = 0;
 					u64 totalCountHasOutputAsInput = 0;
 					u64 totalEntryBytes = 0;
+
+					u64 lastUseTimeLimit = 0;
+					if (m_expirationTimeSeconds)
+					{
+						u64 secondsSinceCreation = GetFileTimeAsSeconds(GetSystemTimeAsFileTime() - m_creationTime);
+						if (secondsSinceCreation > m_expirationTimeSeconds)
+							lastUseTimeLimit = GetSecondsAsFileTime(secondsSinceCreation - m_expirationTimeSeconds);
+					}
+
 					for (auto& kv2 : bucket.m_cacheEntryLookup)
 					{
 						CacheEntries& entries = kv2.second;
@@ -2153,7 +2161,12 @@ namespace uba
 							mostEntriesKey = kv2.first;
 						}
 						for (auto& entry : entries.entries)
+						{
 							lastUsed = Max(lastUsed, entry.lastUsedTime);
+
+							if (entry.creationTime < lastUseTimeLimit && entry.lastUsedTime < lastUseTimeLimit)
+								++expiredEntryCount;
+						}
 						totalEntryCount += entries.entries.size();
 						if (!entries.inputsThatAreOutputs.empty())
 							++totalCountHasOutputAsInput;
@@ -2172,6 +2185,7 @@ namespace uba
 					writeLine(line.Clear().Appendf(TC("      TotalEntries: %llu"), totalEntryCount));
 					writeLine(line.Clear().Appendf(TC("      TotalKeySharedEntry: %s"), BytesToText(totalKeySharedBytes).str));
 					writeLine(line.Clear().Appendf(TC("      LastEntryUsed: %s ago"), TimeToText(lastUsedTime, true).str));
+					writeLine(line.Clear().Appendf(TC("      ExpiredEntries: %llu"), expiredEntryCount));
 				}
 			}
 			u64 totalCasSize = 0;
@@ -2181,7 +2195,7 @@ namespace uba
 			writeLine(line.Clear().Appendf(TC("    Count: %llu"), totalCasCount));
 			writeLine(line.Clear().Appendf(TC("    Size: %s"), BytesToText(totalCasSize).str));
 		}
-		else if (command.Equals(TC("validate")))
+		else if (command.Equals(TCV("validate")))
 		{
 			writeLine(line.Clear().Appendf(TC("  Buckets:")));
 			{
@@ -2343,7 +2357,7 @@ namespace uba
 				}
 			}
 		}
-		else if (command.Equals(TC("updateentries")))
+		else if (command.Equals(TCV("updateentries")))
 		{
 			HashMap2<u32, u32> oldToNewCasKeyOffset;
 			MemoryBlock memory(64*1024);
@@ -2365,7 +2379,7 @@ namespace uba
 				bucket.needsSave = true;
 			}
 		}
-		else if (command.Equals(TC("pathtable")))
+		else if (command.Equals(TCV("pathtable")))
 		{
 			SCOPED_FUTEX(m_bucketsLock, bucketsLock);
 			for (auto& kv : m_buckets)
@@ -2384,10 +2398,10 @@ namespace uba
 					});
 
 				for (auto& path : sortedPaths)
-					writeLine(line.Clear().Append(TC("   ")).Append(path));
+					writeLine(line.Clear().Append(TCV("   ")).Append(path));
 			}
 		}
-		else if (command.Equals(TC("usedpaths")))
+		else if (command.Equals(TCV("usedpaths")))
 		{
 			SCOPED_FUTEX(m_bucketsLock, bucketsLock);
 			for (auto& kv : m_buckets)
@@ -2410,10 +2424,10 @@ namespace uba
 				}
 
 				for (auto& path : sortedPaths)
-					writeLine(line.Clear().Append(TC("   ")).Append(path));
+					writeLine(line.Clear().Append(TCV("   ")).Append(path));
 			}
 		}
-		else if (command.Equals(TC("usedsegments")))
+		else if (command.Equals(TCV("usedsegments")))
 		{
 			SCOPED_FUTEX(m_bucketsLock, bucketsLock);
 			for (auto& kv : m_buckets)
@@ -2449,13 +2463,13 @@ namespace uba
 				}
 				for (auto& kv2 : sortedSegments)
 				{
-					writeLine(line.Clear().Append(TC("  ")).AppendValue(1000000 - kv2.first));
+					writeLine(line.Clear().Append(TCV("  ")).AppendValue(1000000 - kv2.first));
 					for (auto& str : kv2.second)
-						writeLine(line.Clear().Append(TC("    ")).Append(str));
+						writeLine(line.Clear().Append(TCV("    ")).Append(str));
 				}
 			}
 		}
-		else if (command.Equals(TC("largestentry")))
+		else if (command.Equals(TCV("largestsharedinput")))
 		{
 			CacheEntries* largestEntries = nullptr;
 			Bucket* largestEntriesBucket = nullptr;
@@ -2475,23 +2489,48 @@ namespace uba
 					largestEntriesBucket = &bucket;
 				}
 			}
-			writeLine(line.Clear().Append(TC("Largest entry (based on shared inputs)")));
+			writeLine(line.Clear().Append(TCV("Largest entry (based on shared inputs)")));
+			if (largestEntries)
+				writeEntry(*largestEntriesBucket, *largestEntries);
+		}
+		else if (command.Equals(TCV("largestentry")))
+		{
+			CacheEntries* largestEntries = nullptr;
+			Bucket* largestEntriesBucket = nullptr;
+			u64 largestSize = 0;
+			SCOPED_FUTEX(m_bucketsLock, bucketsLock);
+			for (auto& kv : m_buckets)
+			{
+				Bucket& bucket = kv.second;
+				for (auto& kv2 : bucket.m_cacheEntryLookup)
+				{
+					CacheEntries& entries = kv2.second;
+					u64 totalSize = entries.GetTotalSize(CacheNetworkVersion, true);
+					SCOPED_READ_LOCK(entries.lock, lock3);
+					if (totalSize <= largestSize)
+						continue;
+					largestSize = totalSize;
+					largestEntries = &entries;
+					largestEntriesBucket = &bucket;
+				}
+			}
+			writeLine(line.Clear().Appendf(TC("Largest entry (based on total size of %s)"), BytesToText(largestSize).str));
 			if (largestEntries)
 				writeEntry(*largestEntriesBucket, *largestEntries);
 		}
-		else if (command.Equals(TC("obliterate")))
+		else if (command.Equals(TCV("obliterate")))
 		{
 			m_shouldWipe = true;
 			m_addsSinceMaintenance = 1;
 			writeLine(line.Clear().Appendf(TC("Cache server database obliteration queued!")));
 		}
-		else if (command.Equals(TC("maintenance")))
+		else if (command.Equals(TCV("maintenance")))
 		{
 			m_forceAllSteps = true;
 			m_addsSinceMaintenance = 1;
 			writeLine(line.Clear().Appendf(TC("Cache server maintenance queued!")));
 		}
-		else if (command.Equals(TC("save")))
+		else if (command.Equals(TCV("save")))
 		{
 			Save();
 		}
diff --git a/Engine/Source/Programs/UnrealBuildAccelerator/Common/Public/UbaCacheServer.h b/Engine/Source/Programs/UnrealBuildAccelerator/Common/Public/UbaCacheServer.h
--- a/Engine/Source/Programs/UnrealBuildAccelerator/Common/Public/UbaCacheServer.h
+++ b/Engine/Source/Programs/UnrealBuildAccelerator/Common/Public/UbaCacheServer.h
@@ -58,7 +58,7 @@ namespace uba
 
 		void SetForceFullMaintenance();
 
-		bool RunMaintenance(bool force, const Function<bool()>& shouldExit);
+		bool RunMaintenance(bool force, bool allowSave, const Function<bool()>& shouldExit);
 
 		bool ShouldShutdown();
 
@@ -66,7 +66,7 @@ namespace uba
 		struct Bucket;
 		struct LoadStats;
 
-		bool RunMaintenanceInternal(const Function<bool()>& shouldExit);
+		bool RunMaintenanceInternal(const Function<bool()>& shouldExit, bool allowSave);
 		bool LoadBucket(Bucket& bucket, BinaryReader& reader, u32 databaseVersion, LoadStats& outStats);
 		bool SaveBucket(Bucket& bucket, Vector<u8>& temp);
 		bool SaveNoLock();
