Unshelved from pending changelist '41013717':

[UBA]
* Fixed so execve never returns if successful. This fixed so macos linking (ld) works in uba

diff --git a/Engine/Source/Programs/UnrealBuildAccelerator/Detours/Private/PosixOS/UbaDetoursFunctionsPosix.cpp b/Engine/Source/Programs/UnrealBuildAccelerator/Detours/Private/PosixOS/UbaDetoursFunctionsPosix.cpp
--- a/Engine/Source/Programs/UnrealBuildAccelerator/Detours/Private/PosixOS/UbaDetoursFunctionsPosix.cpp
+++ b/Engine/Source/Programs/UnrealBuildAccelerator/Detours/Private/PosixOS/UbaDetoursFunctionsPosix.cpp
@@ -2290,7 +2290,10 @@ int Internal_execve(const char* pathname, char* const _Nullable argv[], char* co
 		writer.WriteU32(pid);
 	}
 
-	return WEXITSTATUS(status);
+	// This function should never return because process is replaced with new program
+	exit(WEXITSTATUS(status));
+
+	return -1; // should never end up here
 }
 
 UBA_EXPORT int UBA_WRAPPER(execve)(const char* pathname, char* const _Nullable argv[], char* const _Nullable envp[])
