[UBA]
* Added so posix stat check validity of paths containing ".." before fixing the path. This should fail but didn't:
/path/to/existing-dir/non-existing-dir/..

diff --git a/Engine/Source/Programs/UnrealBuildAccelerator/Detours/Private/PosixOS/UbaDetoursFunctionsPosix.cpp b/Engine/Source/Programs/UnrealBuildAccelerator/Detours/Private/PosixOS/UbaDetoursFunctionsPosix.cpp
--- a/Engine/Source/Programs/UnrealBuildAccelerator/Detours/Private/PosixOS/UbaDetoursFunctionsPosix.cpp
+++ b/Engine/Source/Programs/UnrealBuildAccelerator/Detours/Private/PosixOS/UbaDetoursFunctionsPosix.cpp
@@ -687,6 +687,22 @@ int Shared_stat(const char* funcName, const char* file, struct stat* attr, const
 
 	DEBUG_LOG_DETOURED(funcName, "%s (%s size: %llu id: %llu dev: %u)-> %i (%s)", file, realName, fileAttr.data.st_size, fileAttr.data.st_ino, fileAttr.data.st_dev, res, StrError(res, fileAttr.lastError));
 
+	if (res == 0) // If success and original path had ".." in path, we need to verify that the path that leads to the ".." actually exists.
+	{
+		// TODO: handle multiple ".." spread out in the path
+		const tchar* dotdot;
+		if (Contains(file, "..", false, &dotdot))
+		{
+			StringBuffer<> tempPath;
+			tempPath.Append(file, dotdot - file);
+			struct stat tempAttr;
+			int tempRes = Shared_stat("stat(dotdot)", tempPath.data, &tempAttr, trueStat);
+			if (tempRes != 0)
+				return tempRes;
+		}
+
+	}
+
 	errno = fileAttr.lastError;
 
 	if (res == 0) // This check is important.. dont write to attr unless success. (some programs rely on that)
