[UBA]
* Fixed unit test compile error

diff --git a/Engine/Source/Programs/UnrealBuildAccelerator/Test/Private/UbaTestCache.cpp b/Engine/Source/Programs/UnrealBuildAccelerator/Test/Private/UbaTestCache.cpp
--- a/Engine/Source/Programs/UnrealBuildAccelerator/Test/Private/UbaTestCache.cpp
+++ b/Engine/Source/Programs/UnrealBuildAccelerator/Test/Private/UbaTestCache.cpp
@@ -621,7 +621,7 @@ namespace uba
 			}
 		}
 
-		if (!cacheServer.RunMaintenance(true, []() { return false; }))
+		if (!cacheServer.RunMaintenance(true, true, []() { return false; }))
 			return false;
 
 		{
