[UBA]
* New attempt at fixing forking and executing new executables in same process on posix platforms

diff --git a/Engine/Source/Programs/UnrealBuildAccelerator/Detours/Private/PosixOS/UbaDetoursFunctionsPosix.cpp b/Engine/Source/Programs/UnrealBuildAccelerator/Detours/Private/PosixOS/UbaDetoursFunctionsPosix.cpp
--- a/Engine/Source/Programs/UnrealBuildAccelerator/Detours/Private/PosixOS/UbaDetoursFunctionsPosix.cpp
+++ b/Engine/Source/Programs/UnrealBuildAccelerator/Detours/Private/PosixOS/UbaDetoursFunctionsPosix.cpp
@@ -34,6 +34,7 @@ namespace uba
 	bool g_isDetouring;
 	bool g_isInitialized;
 	bool g_isCancelled;
+	pid_t g_pid;
 	void Deinit();
 }
 
@@ -2278,8 +2279,10 @@ int Internal_execve(const char* pathname, char* const _Nullable argv[], char* co
 		return TRUE_WRAPPER(execve)(pathname, argv, envp);
 	}
 
-	bool inVfork = t_inVfork != 0;
-	pid_t pid;
+	pid_t pid = getpid();
+	bool inVfork = g_pid != pid;
+	g_pid = pid;
+
 	int res = shared_posix_spawn(&pid, pathname, nullptr, nullptr, argv, envp);
 	if (inVfork)
 		t_inVfork = pid;
@@ -2306,20 +2309,21 @@ int Internal_execve(const char* pathname, char* const _Nullable argv[], char* co
 		writer.WriteU32(pid);
 	}
 
-	// This function should never return because process is replaced with new program
-	exit(WEXITSTATUS(status));
+	res = WEXITSTATUS(status);
+
+	if (inVfork)
+	{
+		TRUE_WRAPPER(_exit)(res);
+	}
+	else
+		TRUE_WRAPPER(exit)(res);
 
-	return -1; // should never end up here
+	return -1;
 }
 
 UBA_EXPORT int UBA_WRAPPER(execve)(const char* pathname, char* const _Nullable argv[], char* const _Nullable envp[])
 {
-	int res = Internal_execve(pathname, argv, envp);
-	if (t_inVfork)
-	{
-		TRUE_WRAPPER(_exit)(res); // execve should exit unless it really fails
-	}
-	return res;
+	return Internal_execve(pathname, argv, envp);
 }
 
 UBA_EXPORT int UBA_WRAPPER(execvp)(const char* file, char* const argv[])
@@ -2511,6 +2515,8 @@ namespace uba
 		g_fileHandlesMem.Create();
 		g_fileHandlesLockMem.Create();
 
+		g_pid = getpid();
+
 		SuppressDetourScope s;
 
 		g_systemTemp.Append(getenv("TMPDIR"));
