[UBA]
* Fixed so cmd.exe for windows and sh for linux is not part of cache key
* Fixed bug in download/decompress logic for cas files when downloading from cache client

diff --git a/Engine/Source/Programs/UnrealBuildAccelerator/Cli/Private/UbaCli.cpp b/Engine/Source/Programs/UnrealBuildAccelerator/Cli/Private/UbaCli.cpp
--- a/Engine/Source/Programs/UnrealBuildAccelerator/Cli/Private/UbaCli.cpp
+++ b/Engine/Source/Programs/UnrealBuildAccelerator/Cli/Private/UbaCli.cpp
@@ -605,7 +605,15 @@ namespace uba
 								success = false;
 							if (!client.storageClient->RetrieveCasFile(res, casKey, TC("")))
 								success = false;
-						}, 1, TC(""));
+
+							#if 0
+							StorageStats storageStats;
+							FileFetcher fetcher { client.storageClient->m_bufferSlots, storageStats };
+							bool destinationIsCompressed = false;
+							if (!fetcher.RetrieveFile(logger, *client.networkClient, casKey, TC("e:\\temp\\foo"), destinationIsCompressed))
+								success = false;
+							#endif
+						}, 1, TC("CheckCas2"));
 				});
 			workManager.FlushWork();
 			if (!success)
diff --git a/Engine/Source/Programs/UnrealBuildAccelerator/Common/Private/UbaCacheClient.cpp b/Engine/Source/Programs/UnrealBuildAccelerator/Common/Private/UbaCacheClient.cpp
--- a/Engine/Source/Programs/UnrealBuildAccelerator/Common/Private/UbaCacheClient.cpp
+++ b/Engine/Source/Programs/UnrealBuildAccelerator/Common/Private/UbaCacheClient.cpp
@@ -254,8 +254,12 @@ namespace uba
 					return false;
 				}
 			}
-			else
 			#endif
+
+			// Ignore cmd.exe or sh as input. It should always exist but can be different between windows versions
+			if (!isOutput && path.EndsWith(IsWindows ? TCV("\\cmd.exe") : TCV("/sh")))
+				continue;
+
 			if (ShouldNormalize(path)) // Paths can be absolute in rsp files so we need to normalize those paths
 			{
 				casKey = rootPaths.NormalizeAndHashFile(m_logger, path.data);
@@ -266,12 +270,6 @@ namespace uba
 				}
 				casKey = IsNormalized(casKey) ? AsCompressed(casKey, true) : CasKeyZero;
 			}
-			else if (path[path.count-1] == ':')
-			{
-				m_logger.Info(TC("GOT UNKNOWN RELATIVE PATH: %s (%s)"), path.data, info.GetDescription());
-				success = false;
-				continue;
-			}
 
 			// Handle path hashes.
 			if (!isOutput)
diff --git a/Engine/Source/Programs/UnrealBuildAccelerator/Common/Private/UbaStorageUtils.cpp b/Engine/Source/Programs/UnrealBuildAccelerator/Common/Private/UbaStorageUtils.cpp
--- a/Engine/Source/Programs/UnrealBuildAccelerator/Common/Private/UbaStorageUtils.cpp
+++ b/Engine/Source/Programs/UnrealBuildAccelerator/Common/Private/UbaStorageUtils.cpp
@@ -161,7 +161,8 @@ namespace uba
 			sendCount = sendCountCapacity;
 		else if (sendCount < sendCountCapacity && (left - sendCount * messageMaxSize) > 0)
 			++sendCount;
-
+		
+		UBA_ASSERT(sendCount);
 
 		u64 entriesMem[sizeof(Entry) * 8];
 		UBA_ASSERT(sizeof(Entry)*sendCount <= sizeof(entriesMem));
@@ -626,6 +627,10 @@ namespace uba
 			u64 leftCompressed = fileSize - responseSize;
 			do
 			{
+				// There is a risk the compressed file we download is larger than half buffer
+				u8* extraBuffer = nullptr;
+				auto ebg = MakeGuard([&]() { delete[] extraBuffer; });
+
 				// First read in a full decompressable block
 				bool isFirstInBlock = true;
 				u32 compressedSize = ~u32(0);
@@ -638,9 +643,30 @@ namespace uba
 					{
 						if (fetchId == u16(~0))
 							return logger.Error(TC("Cas content error (2). Server believes %s was only one segment but client sees more. UncompressedSize: %llu LeftUncompressed: %llu Size: %llu Left to read: %llu ResponseSize: %u. (%s)"), destination, actualSize, leftUncompressed, fileSize, left, responseSize, CasKeyString(casKey).str);
+						
+						u64 capacity = maxReadSize - u64(readPosition - readBuffer);
+						u64 writeCapacity = capacity;
+						u8* writeDest = readPosition;
+						if (capacity < sizeOfFirstMessage) // capacity left is less than the last message
+						{
+							UBA_ASSERT(!extraBuffer);
+							extraBuffer = new u8[sizeOfFirstMessage];
+							writeDest = extraBuffer;
+							writeCapacity = sizeOfFirstMessage;
+						}
+					
 						u32 error;
-						if (!SendBatchMessages(logger, client, fetchId, readPosition, maxReadSize - u32(readPosition - readBuffer), leftCompressed, sizeOfFirstMessage, readIndex, responseSize, tryWriteDecompressed, destination, &error))
+						if (!SendBatchMessages(logger, client, fetchId, writeDest, writeCapacity, leftCompressed, sizeOfFirstMessage, readIndex, responseSize, tryWriteDecompressed, destination, &error))
 							return logger.Error(TC("Failed to send batched messages to server while retrieving and decompressing cas %s to %s. Error: %u"), CasKeyString(casKey).str, destination, error);
+
+						if (extraBuffer)
+						{
+							memcpy(readPosition, extraBuffer, left);
+							memmove(extraBuffer, extraBuffer + left, u32(responseSize - left));
+							if (isFirstInBlock)
+								return logger.Error(TC("Make static analysis happy. This should not be possible to happen (%s)"), CasKeyString(casKey).str);
+						}
+
 						leftCompressed -= responseSize;
 					}
 					else
@@ -734,12 +760,29 @@ namespace uba
 				// Move overflow back to the beginning of the buffer and start the next block (if there is one)
 				readBuffer = slot;
 				maxReadSize = BufferSlotHalfSize;
-				UBA_ASSERTF(readPosition - overflow >= readBuffer, TC("ReadPosition - overflow is before beginning of buffer (overflow: %u) for file %s"), overflow, destination);
-				UBA_ASSERTF(readPosition <= readBuffer + BufferSlotHalfSize, TC("ReadPosition is outside readBuffer size (pos: %llu, overflow: %u) for file %s"), readPosition - readBuffer, overflow, destination);
-				memmove(readBuffer, readPosition - overflow, overflow);
+
+				if (extraBuffer)
+				{
+					memcpy(readBuffer, extraBuffer, overflow);
+					delete[] extraBuffer;
+					extraBuffer = nullptr;
+				}
+				else
+				{
+					// Move overflow back to the beginning of the buffer and start the next block (if there is one)
+					UBA_ASSERTF(readPosition - overflow >= readBuffer, TC("ReadPosition - overflow is before beginning of buffer (overflow: %u) for file %s"), overflow, destination);
+					UBA_ASSERTF(readPosition <= readBuffer + BufferSlotHalfSize, TC("ReadPosition is outside readBuffer size (pos: %llu, overflow: %u) for file %s"), readPosition - readBuffer, overflow, destination);
+					memmove(readBuffer, readPosition - overflow, overflow);
+				}
+
 				readPosition = readBuffer + overflow;
 				if (overflow)
-					responseSize = 0;
+				{
+					if (overflow < sizeof(u32) * 2) // Must always have the compressed and uncompressed size to be able to move on with logic above
+						sendSegmentMessage = true;
+					else
+						responseSize = 0;
+				}
 			} while (leftUncompressed);
 
 			if (!tryWriteDecompressed())
