[UBA]
* Improved trace info from crawler to know what is doing what and reading what files
* Added so -isysroot on mac is crawled

diff --git a/Engine/Source/Programs/UnrealBuildAccelerator/Common/Private/UbaSession.cpp b/Engine/Source/Programs/UnrealBuildAccelerator/Common/Private/UbaSession.cpp
--- a/Engine/Source/Programs/UnrealBuildAccelerator/Common/Private/UbaSession.cpp
+++ b/Engine/Source/Programs/UnrealBuildAccelerator/Common/Private/UbaSession.cpp
@@ -3954,14 +3954,13 @@ namespace uba
 		if (!at)
 			return;
 
-		auto CreateFileFunc = [this, ph = ProcessHandle(&process), rules = startInfo.rules](const StringView& fileName, const DependencyCrawler::AccessFileFunc& func)
+		auto CreateFileFunc = [this, ph = ProcessHandle(&process), rules = startInfo.rules](TrackWorkScope& tracker, const StringView& fileName, const DependencyCrawler::AccessFileFunc& func)
 			{
 				CreateFileResponse out;
 				{
-					TrackWorkScope tws(m_workManager, AsView(TC("CrawlFile")), ColorWork);
-					tws.AddHint(fileName);
+					tracker.AddHint(fileName);
 
-					if (!CreateFileForRead(out, tws, fileName, ToStringKey(fileName), *static_cast<ProcessImpl*>(ph.m_process), *rules))
+					if (!CreateFileForRead(out, tracker, fileName, ToStringKey(fileName), *static_cast<ProcessImpl*>(ph.m_process), *rules))
 						return false;
 				}
 
diff --git a/Engine/Source/Programs/UnrealBuildAccelerator/Common/Public/UbaDependencyCrawler.h b/Engine/Source/Programs/UnrealBuildAccelerator/Common/Public/UbaDependencyCrawler.h
--- a/Engine/Source/Programs/UnrealBuildAccelerator/Common/Public/UbaDependencyCrawler.h
+++ b/Engine/Source/Programs/UnrealBuildAccelerator/Common/Public/UbaDependencyCrawler.h
@@ -28,7 +28,7 @@ namespace uba
 
 		using AccessFileFunc = Function<bool(const void* data, u64 dataSize)>;
 		using DevirtualizePathFunc = Function<bool(StringBufferBase& inOut)>;
-		using CreateFileFunc = Function<bool(const StringView& fileName, const AccessFileFunc& func)>;
+		using CreateFileFunc = Function<bool(TrackWorkScope& tracker, const StringView& fileName, const AccessFileFunc& func)>;
 
 		inline void Init(FileExistsFunc&& fileExistsFunc, TraverseFilesFunc&& traverseFilesFunc, bool useBloomFilter = true);
 		inline bool Add(const tchar* rsp, const tchar* workDir, CreateFileFunc&& func, DevirtualizePathFunc&& devirtualizePathFunc, const tchar* app, DependencyCrawlerType type, u32 ruleIndex);
@@ -87,12 +87,13 @@ namespace uba
 		struct CodeFile { TString path; bool hasPch = false; };
 		using CodeFiles = List<CodeFile>;
 
-		inline bool ParseRsp(Instance& instance, const StringView& rsp, CodeFiles& outCodeFiles);
-		inline bool ParseRsp2(Instance& instance, const void* data, u64 dataSize, const tchar* rsp, CodeFiles& outCodeFiles);
+		inline bool ParseRsp(TrackWorkScope& tracker, Instance& instance, const StringView& rsp, CodeFiles& outCodeFiles);
+		inline bool ParseRsp2(TrackWorkScope& tracker, Instance& instance, const void* data, u64 dataSize, const tchar* rsp, CodeFiles& outCodeFiles);
 		inline bool ParsePch(Pch& pch, Instance& instance, const void* data, u64 dataSize);
-		inline bool ParseCodeFile(Instance& instance, const StringView& codeFile, bool parseDefines, const StringView& caller);
+		inline bool ParseCodeFile(TrackWorkScope& tracker, Instance& instance, const StringView& codeFile, bool parseDefines, const StringView& caller);
 		inline bool ParseCodeFile2(Instance& instance, const void* data, u64 dataSize, const StringView& codeFile, bool parseDefines);
-		inline bool HandleInclude(Instance& instance, const StringView& rootPath, const StringView& include, const StringView& codeFile, bool parseDefines);
+		inline bool HandleInclude(TrackWorkScope& tracker, Instance& instance, const StringView& rootPath, const StringView& include, const StringView& codeFile, bool parseDefines);
+		inline bool TraverseInclude(Instance& instance, const StringView& rootPathWithSlash);
 	};
 
 
@@ -129,7 +130,7 @@ namespace uba
 			{
 				auto& instance = ref.instance;
 				CodeFiles codeFiles;
-				if (!ParseRsp(instance, instance.rsp, codeFiles))
+				if (!ParseRsp(context.tracker, instance, instance.rsp, codeFiles))
 					return;
 
 				for (auto& cf : codeFiles)
@@ -151,7 +152,7 @@ namespace uba
 						if (!pch.handled)
 						{
 							fixedFile.Append(TCV(".dep.json"));
-							instance.createFileFunc(fixedFile, [&](const void* data, u64 dataSize)
+							instance.createFileFunc(context.tracker, fixedFile, [&](const void* data, u64 dataSize)
 							{
 								return ParsePch(pch, instance, data, dataSize);
 							});
@@ -164,12 +165,12 @@ namespace uba
 						// If we end up here it means that we have compressed obj files enabled and want to decompress them in parallel
 						m_workManager.AddWork([this, ref, codeFile = cf.path](const WorkContext& context)
 							{
-								ref.instance.createFileFunc(codeFile, {});
-							}, 1, TC("Crawler"), ColorWork);
+								ref.instance.createFileFunc(context.tracker, codeFile, {});
+							}, 1, TC("CrawlForDecomp"), ColorWork);
 					}
 					else
 					{
-						ParseCodeFile(instance, cf.path, instance.platform.empty(), instance.rsp);
+						ParseCodeFile(context.tracker, instance, cf.path, instance.platform.empty(), instance.rsp);
 						if (instance.platform.empty())
 						{
 							if (!instance.overriddenPlatformName.empty())
@@ -194,26 +195,16 @@ namespace uba
 						if (CaseInsensitiveFs)
 							fixedPath.MakeLower();
 						fixedPath.Append(PathSeparator);
-						StringBuffer<> codeFilePath(fixedPath);
-						m_traverseFilesFunc(fixedPath, [&](const StringView& file, bool isDirectory)
-							{
-								if (isDirectory)
-									return;
-								codeFilePath.Resize(fixedPath.count).Append(file);
-								m_workManager.AddWork([this, ref, codeFile = TString(codeFilePath.data)](const WorkContext& context)
-									{
-										if (!ref.instance.createFileFunc(codeFile, {}))
-											m_logger.Warning(TC("Failed to open file %s from builtindir"), codeFile.c_str());
-									}, 1, TC("Crawler"), ColorWork);
-							});
+
+						TraverseInclude(instance, fixedPath);
 					}
 				}
 
-			}, 1, TC("Crawler"), ColorWork);
+			}, 1, TC("CrawlRsp"), ColorWork);
 		return true;
 	}
 
-	bool DependencyCrawler::ParseRsp(Instance& instance, const StringView& rsp, CodeFiles& outCodeFiles)
+	bool DependencyCrawler::ParseRsp(TrackWorkScope& tracker, Instance& instance, const StringView& rsp, CodeFiles& outCodeFiles)
 	{
 		StringBuffer<> fixedPath;
 		FixPath(rsp.data, instance.workDir.data(), instance.workDir.size(), fixedPath);
@@ -221,12 +212,12 @@ namespace uba
 			m_logger.Info(TC("Failed to devirtualize path %s in %s"), fixedPath.data, rsp.data);
 		if (CaseInsensitiveFs)
 			fixedPath.MakeLower();
-		if (!instance.createFileFunc(fixedPath, [&](const void* data, u64 dataSize) { return ParseRsp2(instance, data, dataSize, rsp.data, outCodeFiles); }))
+		if (!instance.createFileFunc(tracker, fixedPath, [&](const void* data, u64 dataSize) { return ParseRsp2(tracker, instance, data, dataSize, rsp.data, outCodeFiles); }))
 			return m_logger.Warning(TC("Failed to parse rsp %s"), rsp.data);
 		return true;
 	}
 
-	bool DependencyCrawler::ParseRsp2(Instance& instance, const void* data, u64 dataSize, const tchar* rsp, CodeFiles& outCodeFiles)
+	bool DependencyCrawler::ParseRsp2(TrackWorkScope& tracker, Instance& instance, const void* data, u64 dataSize, const tchar* rsp, CodeFiles& outCodeFiles)
 	{
 		auto AddCodeFile = [&](const StringView& file, bool pushFront, bool hasPch)
 			{
@@ -385,7 +376,7 @@ namespace uba
 					}
 					else if (sb[0] == '@')
 					{
-						ParseRsp(instance, StringView(sb).Skip(1), outCodeFiles);
+						ParseRsp(tracker, instance, StringView(sb).Skip(1), outCodeFiles);
 					}
 					else
 					{
@@ -395,7 +386,7 @@ namespace uba
 		}
 		else if (instance.type == DependencyCrawlerType_ClangCompiler)
 		{
-			enum ArgType { ArgType_None, ArgType_Ignore, ArgType_RootPath, ArgType_ResourceDir, ArgType_Code };
+			enum ArgType { ArgType_None, ArgType_Ignore, ArgType_RootPath, ArgType_ISysRootPath, ArgType_ResourceDir, ArgType_Code };
 			struct OptionWithArg { const StringView name; ArgType type; };
 			constexpr const OptionWithArg dashOptionsWithArg[] = 
 			{
@@ -410,7 +401,7 @@ namespace uba
 				{ TCV("-arch"), ArgType_Ignore },
 				{ TCV("--sysroot"), ArgType_RootPath},
 				{ TCV("-isystem"), ArgType_RootPath },
-				{ TCV("-isysroot"), ArgType_RootPath },
+				{ TCV("-isysroot"), ArgType_ISysRootPath },
 				{ TCV("-internal-isystem"), ArgType_RootPath },
 				{ TCV("-I"), ArgType_RootPath },
 				{ TCV("-F"), ArgType_RootPath },
@@ -480,6 +471,12 @@ namespace uba
 						case ArgType_RootPath:
 							AddRoot(sb);
 							break;
+						case ArgType_ISysRootPath:
+							AddRoot(sb);
+							instance.builtinIncludesDir = sb.data;
+							//instance.builtinIncludesDir += TC("/usr/include/c++");
+							break;
+
 						case ArgType_Code:
 							if (!sb.EndsWith(TCV(".gch")))
 								AddCodeFile(sb, true, false);
@@ -495,7 +492,7 @@ namespace uba
 						return;
 					}
 					else if (sb[0] == '@')
-						ParseRsp(instance, StringView(sb).Skip(1), outCodeFiles);
+						ParseRsp(tracker, instance, StringView(sb).Skip(1), outCodeFiles);
 					else if (sb.Contains('/'))
 						AddCodeFile(sb, false, false);
 					else
@@ -575,9 +572,9 @@ namespace uba
 		return true;
 	}
 
-	bool DependencyCrawler::ParseCodeFile(Instance& instance, const StringView& codeFile, bool parseDefines, const StringView& caller)
+	bool DependencyCrawler::ParseCodeFile(TrackWorkScope& tracker, Instance& instance, const StringView& codeFile, bool parseDefines, const StringView& caller)
 	{
-		if (!instance.createFileFunc(codeFile, [&](const void* data, u64 dataSize) { return ParseCodeFile2(instance, data, dataSize, codeFile, parseDefines); }))
+		if (!instance.createFileFunc(tracker, codeFile, [&](const void* data, u64 dataSize) { return ParseCodeFile2(instance, data, dataSize, codeFile, parseDefines); }))
 			return m_logger.Warning(TC("Failed to parse code file %s found in %s"), codeFile.data, caller.data);
 		return true;
 	}
@@ -693,7 +690,7 @@ namespace uba
 								{
 									StringBuffer<> localDir;
 									if (const tchar* lastSeparator = TStrrchr(codeFile2.c_str(), PathSeparator))
-										if (HandleInclude(instance, localDir.Append(codeFile2.c_str(), lastSeparator - codeFile2.c_str() + 1), include2, codeFile2, parseDefines))
+										if (HandleInclude(context.tracker, instance, localDir.Append(codeFile2.c_str(), lastSeparator - codeFile2.c_str() + 1), include2, codeFile2, parseDefines))
 											return;
 								}
 
@@ -707,7 +704,7 @@ namespace uba
 								{
 									if (m_useBloomFilter && root.bloomFilter.IsGuaranteedMiss(fileKey))
 										continue;
-									if (HandleInclude(instance, root.path, include2, codeFile2, parseDefines))
+									if (HandleInclude(context.tracker, instance, root.path, include2, codeFile2, parseDefines))
 										return;
 								}
 
@@ -718,7 +715,7 @@ namespace uba
 							work({tws});
 						}
 						else
-							m_workManager.AddWork(work, 1, TC("Crawler"), ColorWork);
+							m_workManager.AddWork(work, 1, TC("CrawlIncludes"), ColorWork);
 					}
 				}
 				else if (parseDefines && strncmp(it, "define", 6) == 0)
@@ -766,7 +763,7 @@ namespace uba
 		return true;
 	}
 
-	bool DependencyCrawler::HandleInclude(Instance& instance, const StringView& rootPath, const StringView& include, const StringView& codeFile, bool parseDefines)
+	bool DependencyCrawler::HandleInclude(TrackWorkScope& tracker, Instance& instance, const StringView& rootPath, const StringView& include, const StringView& codeFile, bool parseDefines)
 	{
 		if (include.EndsWith(TCV(".ush")))
 			return true;
@@ -798,6 +795,43 @@ namespace uba
 				return true;
 		}
 
-		return ParseCodeFile(instance, fullPath, parseDefines, codeFile);
+		return ParseCodeFile(tracker, instance, fullPath, parseDefines, codeFile);
+	}
+
+	bool DependencyCrawler::TraverseInclude(Instance& instance, const StringView& rootPathWithSlash)
+	{
+		StringBuffer<> codeFilePath(rootPathWithSlash);
+		m_traverseFilesFunc(rootPathWithSlash, [&](const StringView& file, bool isDirectory)
+			{
+				codeFilePath.Resize(rootPathWithSlash.count).Append(file);
+				if (isDirectory)
+				{
+					#if PLATFORM_MAC
+					if (file.Equals(TCV("ruby"))) // Should probably be fixed in a nicer way but there are symlinks going back to previous dir in apple sdk
+						return;
+					if (file.Equals(TCV("Versions")) || file.Contains(TCV("_SwiftUI")) || file.Equals(TCV("PrivateFrameworks"))) // Directories we dont need to visit
+						return;
+					#endif
+					codeFilePath.Append(PathSeparator);
+					m_workManager.AddWork([this, ref = InstanceRef(instance), codeFile = TString(codeFilePath.data)](const WorkContext& context)
+						{
+							TraverseInclude(ref.instance, codeFile);
+						}, 1, TC("CrawlBiDir"), ColorWork);
+				}
+				else
+				{
+					#if PLATFORM_MAC
+					if (!file.EndsWith(TCV(".h")) && !file.EndsWith(TCV(".modulemap")) && TStrchr(file.data, '.') != 0)
+						return;
+					#endif
+
+					m_workManager.AddWork([this, ref = InstanceRef(instance), codeFile = TString(codeFilePath.data)](const WorkContext& context)
+						{
+							if (!ref.instance.createFileFunc(context.tracker, codeFile, {}))
+								m_logger.Warning(TC("Failed to open file %s from builtindir"), codeFile.c_str());
+						}, 1, TC("CrawlBiFile"), ColorWork);
+				}
+			});
+		return true;
 	}
 }
