[UBA]
* Fixed bug in dependency crawler where shared definitions is not always parsed
* Fixed better trace information for dependency crawler

diff --git a/Engine/Source/Programs/UnrealBuildAccelerator/Common/Public/UbaDependencyCrawler.h b/Engine/Source/Programs/UnrealBuildAccelerator/Common/Public/UbaDependencyCrawler.h
--- a/Engine/Source/Programs/UnrealBuildAccelerator/Common/Public/UbaDependencyCrawler.h
+++ b/Engine/Source/Programs/UnrealBuildAccelerator/Common/Public/UbaDependencyCrawler.h
@@ -91,7 +91,7 @@ namespace uba
 		inline bool ParseRsp2(TrackWorkScope& tracker, Instance& instance, const void* data, u64 dataSize, const tchar* rsp, CodeFiles& outCodeFiles);
 		inline bool ParsePch(Pch& pch, Instance& instance, const void* data, u64 dataSize);
 		inline bool ParseCodeFile(TrackWorkScope& tracker, Instance& instance, const StringView& codeFile, bool parseDefines, const StringView& caller);
-		inline bool ParseCodeFile2(Instance& instance, const void* data, u64 dataSize, const StringView& codeFile, bool parseDefines);
+		inline bool ParseCodeFile2(TrackWorkScope& tracker, Instance& instance, const void* data, u64 dataSize, const StringView& codeFile, bool parseDefines);
 		inline bool HandleInclude(TrackWorkScope& tracker, Instance& instance, const StringView& rootPath, const StringView& include, const StringView& codeFile, bool parseDefines);
 		inline bool TraverseInclude(Instance& instance, const StringView& rootPathWithSlash);
 	};
@@ -574,12 +574,12 @@ namespace uba
 
 	bool DependencyCrawler::ParseCodeFile(TrackWorkScope& tracker, Instance& instance, const StringView& codeFile, bool parseDefines, const StringView& caller)
 	{
-		if (!instance.createFileFunc(tracker, codeFile, [&](const void* data, u64 dataSize) { return ParseCodeFile2(instance, data, dataSize, codeFile, parseDefines); }))
+		if (!instance.createFileFunc(tracker, codeFile, [&](const void* data, u64 dataSize) { return ParseCodeFile2(tracker, instance, data, dataSize, codeFile, parseDefines); }))
 			return m_logger.Warning(TC("Failed to parse code file %s found in %s"), codeFile.data, caller.data);
 		return true;
 	}
 
-	bool DependencyCrawler::ParseCodeFile2(Instance& instance, const void* data, u64 dataSize, const StringView& codeFile, bool parseDefines)
+	bool DependencyCrawler::ParseCodeFile2(TrackWorkScope& tracker, Instance& instance, const void* data, u64 dataSize, const StringView& codeFile, bool parseDefines)
 	{
 		const char* it = (const char*)data;
 		const char* end = it + dataSize;
@@ -710,10 +710,7 @@ namespace uba
 
 							};
 						if (parseDefines)
-						{
-							TrackWorkScope tws;
-							work({tws});
-						}
+							work({tracker});
 						else
 							m_workManager.AddWork(work, 1, TC("CrawlIncludes"), ColorWork);
 					}
@@ -790,9 +787,12 @@ namespace uba
 				if (instance.pch->files.find(key) != instance.pch->files.end())
 					return true;
 			}
-			SCOPED_FUTEX(m_codeFilesHandledLock, lock);
-			if (!m_codeFilesHandled.try_emplace(key).second)
-				return true;
+			if (!parseDefines) // We must find UBT_COMPILED_PLATFORM before we start earlying out
+			{
+				SCOPED_FUTEX(m_codeFilesHandledLock, lock);
+				if (!m_codeFilesHandled.try_emplace(key).second)
+					return true;
+			}
 		}
 
 		return ParseCodeFile(tracker, instance, fullPath, parseDefines, codeFile);
