[UBA]
* Fixed so dependency crawler enter the matching header for .gch files since clang will parse all the headers anyway so we have to crawl them :-/
* Fixed non unity compile error

diff --git a/Engine/Source/Programs/UnrealBuildAccelerator/Common/Public/UbaDependencyCrawler.h b/Engine/Source/Programs/UnrealBuildAccelerator/Common/Public/UbaDependencyCrawler.h
--- a/Engine/Source/Programs/UnrealBuildAccelerator/Common/Public/UbaDependencyCrawler.h
+++ b/Engine/Source/Programs/UnrealBuildAccelerator/Common/Public/UbaDependencyCrawler.h
@@ -93,7 +93,7 @@ namespace uba
 		inline bool ParseCodeFile(TrackWorkScope& tracker, Instance& instance, const StringView& codeFile, bool parseDefines, const StringView& caller);
 		inline bool ParseCodeFile2(TrackWorkScope& tracker, Instance& instance, const void* data, u64 dataSize, const StringView& codeFile, bool parseDefines);
 		inline bool HandleInclude(TrackWorkScope& tracker, Instance& instance, const StringView& rootPath, const StringView& include, const StringView& codeFile, bool parseDefines);
-		inline bool TraverseInclude(Instance& instance, const StringView& rootPathWithSlash);
+		inline bool TraverseInclude(TrackWorkScope& tracker, Instance& instance, const StringView& rootPathWithSlash);
 	};
 
 
@@ -196,7 +196,7 @@ namespace uba
 							fixedPath.MakeLower();
 						fixedPath.Append(PathSeparator);
 
-						TraverseInclude(instance, fixedPath);
+						TraverseInclude(context.tracker, instance, fixedPath);
 					}
 				}
 
@@ -394,7 +394,7 @@ namespace uba
 				{ TCV("-x"), ArgType_Ignore },
 				{ TCV("-o"), ArgType_Ignore },
 				{ TCV("-include"), ArgType_Code },
-				{ TCV("-include-pch"), ArgType_Ignore },
+				{ TCV("-include-pch"), ArgType_Code },
 				{ TCV("-vctoolsdir"), ArgType_Ignore },
 				{ TCV("-Xclang"), ArgType_Ignore },
 				{ TCV("-target"), ArgType_Ignore },
@@ -478,8 +478,9 @@ namespace uba
 							break;
 
 						case ArgType_Code:
-							if (!sb.EndsWith(TCV(".gch")))
-								AddCodeFile(sb, true, false);
+							if (sb.EndsWith(TCV(".gch")) || sb.EndsWith(TCV(".pch")))
+								sb.Resize(sb.count - 4);
+							AddCodeFile(sb, true, false);
 							break;
 						case ArgType_ResourceDir:
 							instance.builtinIncludesDir = sb.data;
@@ -798,7 +799,7 @@ namespace uba
 		return ParseCodeFile(tracker, instance, fullPath, parseDefines, codeFile);
 	}
 
-	bool DependencyCrawler::TraverseInclude(Instance& instance, const StringView& rootPathWithSlash)
+	bool DependencyCrawler::TraverseInclude(TrackWorkScope& tracker, Instance& instance, const StringView& rootPathWithSlash)
 	{
 		StringBuffer<> codeFilePath(rootPathWithSlash);
 		m_traverseFilesFunc(rootPathWithSlash, [&](const StringView& file, bool isDirectory)
@@ -815,7 +816,7 @@ namespace uba
 					codeFilePath.Append(PathSeparator);
 					m_workManager.AddWork([this, ref = InstanceRef(instance), codeFile = TString(codeFilePath.data)](const WorkContext& context)
 						{
-							TraverseInclude(ref.instance, codeFile);
+							TraverseInclude(context.tracker, ref.instance, codeFile);
 						}, 1, TC("CrawlBiDir"), ColorWork);
 				}
 				else
diff --git a/Engine/Source/Programs/UnrealBuildAccelerator/Common/Public/UbaHashMap.h b/Engine/Source/Programs/UnrealBuildAccelerator/Common/Public/UbaHashMap.h
--- a/Engine/Source/Programs/UnrealBuildAccelerator/Common/Public/UbaHashMap.h
+++ b/Engine/Source/Programs/UnrealBuildAccelerator/Common/Public/UbaHashMap.h
@@ -3,6 +3,7 @@
 #pragma once
 
 #include "UbaMemory.h"
+#include "UbaPlatform.h"
 
 namespace uba
 {
