[UBA]
* Fixed so dependency crawler is finding used frameworks instead of just parsing all of them
* Changed error to warning when failing to load xcode
* Added warning when secondary connections fail to connect

diff --git a/Engine/Source/Programs/UnrealBuildAccelerator/Agent/Private/UbaAgent.cpp b/Engine/Source/Programs/UnrealBuildAccelerator/Agent/Private/UbaAgent.cpp
--- a/Engine/Source/Programs/UnrealBuildAccelerator/Agent/Private/UbaAgent.cpp
+++ b/Engine/Source/Programs/UnrealBuildAccelerator/Agent/Private/UbaAgent.cpp
@@ -894,7 +894,10 @@ namespace uba
 			});
 
 			if (xcodeDirectories.empty())
-				terminationReason.Append("Unable to populate from any Xcodes. Agent is unusable.");
+			{
+				//terminationReason.Append("Unable to populate from any Xcodes. Agent is unusable.");
+				logger.Warning(TC("Unable to populate from any Xcodes and host might not be able to share sdk files"));
+			}
 		}
 		// if we didn't want a single version, or all xcodes, then use active xcode (useful for user running their own agents)
 		else
@@ -1299,6 +1302,8 @@ namespace uba
 						bool timedOut = false;
 						if (client->Connect(*networkBackend, host.data, port, &timedOut))
 							++connectionCount;
+						else
+							logger.Warning(TC("Failed to connect secondary connection number %u"), connectionCount);
 					}
 				}
 
diff --git a/Engine/Source/Programs/UnrealBuildAccelerator/Common/Public/UbaDependencyCrawler.h b/Engine/Source/Programs/UnrealBuildAccelerator/Common/Public/UbaDependencyCrawler.h
--- a/Engine/Source/Programs/UnrealBuildAccelerator/Common/Public/UbaDependencyCrawler.h
+++ b/Engine/Source/Programs/UnrealBuildAccelerator/Common/Public/UbaDependencyCrawler.h
@@ -36,9 +36,13 @@ namespace uba
 		Logger& m_logger;
 		WorkManager& m_workManager;
 
-		Futex m_codeFilesHandledLock;
 		struct HandledFile { Futex lock; bool handled = false; };
-		UnorderedMap<StringKey, HandledFile> m_codeFilesHandled;
+		struct HandledFiles
+		{
+			Futex lookupLock;
+			UnorderedMap<StringKey, HandledFile> lookup;
+		};
+		HandledFiles m_handledFiles;
 
 		Futex m_pchLookupLock;
 		struct Pch { Futex lock; bool handled = false; UnorderedSet<StringKey> files; };
@@ -61,6 +65,7 @@ namespace uba
 			TString compiledPlatform;
 			TString overriddenPlatformName;
 			TString builtinIncludesDir;
+			TString frameworksDir;
 			bool platformIsExtension = false;
 			bool usePch = false;
 			CreateFileFunc createFileFunc;
@@ -90,9 +95,9 @@ namespace uba
 		inline bool ParseRsp(TrackWorkScope& tracker, Instance& instance, const StringView& rsp, CodeFiles& outCodeFiles);
 		inline bool ParseRsp2(TrackWorkScope& tracker, Instance& instance, const void* data, u64 dataSize, const tchar* rsp, CodeFiles& outCodeFiles);
 		inline bool ParsePch(Pch& pch, Instance& instance, const void* data, u64 dataSize);
-		inline bool ParseCodeFile(TrackWorkScope& tracker, Instance& instance, const StringView& codeFile, bool parseDefines, const StringView& caller);
-		inline bool ParseCodeFile2(TrackWorkScope& tracker, Instance& instance, const void* data, u64 dataSize, const StringView& codeFile, bool parseDefines);
-		inline bool HandleInclude(TrackWorkScope& tracker, Instance& instance, const StringView& rootPath, const StringView& include, const StringView& codeFile, bool parseDefines);
+		inline bool ParseCodeFile(TrackWorkScope& tracker, Instance& instance, HandledFiles& handledFiles, const StringView& codeFile, bool parseDefines, const StringView& caller);
+		inline bool ParseCodeFile2(TrackWorkScope& tracker, Instance& instance, HandledFiles& handledFiles, const void* data, u64 dataSize, const StringView& codeFile, bool parseDefines);
+		inline bool HandleInclude(TrackWorkScope& tracker, Instance& instance, HandledFiles& handledFiles, const StringView& rootPath, const StringView& include, const StringView& codeFile, bool parseDefines);
 		inline bool TraverseInclude(TrackWorkScope& tracker, Instance& instance, const StringView& rootPathWithSlash);
 	};
 
@@ -141,8 +146,9 @@ namespace uba
 						// So this code path is not implemented for clang.. if -fno-validate-ast-input-files-content is added we can implement this path
 						StringBuffer<> fixedFile;
 						FixPath(cf.path.c_str(), instance.workDir.data(), instance.workDir.size(), fixedFile);
-
-						StringKey pchKey = ToStringKeyLower(fixedFile);
+						if (CaseInsensitiveFs)
+							fixedFile.MakeLower();
+						StringKey pchKey = ToStringKey(fixedFile);
 
 						SCOPED_FUTEX(m_pchLookupLock, lock);
 						Pch& pch = m_pchLookup[pchKey];
@@ -170,7 +176,9 @@ namespace uba
 					}
 					else
 					{
-						ParseCodeFile(context.tracker, instance, cf.path, instance.platform.empty(), instance.rsp);
+						bool parseDefines = instance.platform.empty();
+						HandledFiles handledFiles;
+						ParseCodeFile(context.tracker, instance, parseDefines ? handledFiles : m_handledFiles, cf.path, parseDefines, instance.rsp);
 						if (instance.platform.empty())
 						{
 							if (!instance.overriddenPlatformName.empty())
@@ -329,8 +337,8 @@ namespace uba
 								fixedFile.MakeLower();
 
 							StringKey key = ToStringKey(fixedFile);
-							SCOPED_FUTEX(m_codeFilesHandledLock, lock);
-							auto insres = m_codeFilesHandled.try_emplace(key);
+							SCOPED_FUTEX(m_handledFiles.lookupLock, lock);
+							auto insres = m_handledFiles.lookup.try_emplace(key);
 							lock.Leave();
 
 							HandledFile& handledFile = insres.first->second;
@@ -473,8 +481,12 @@ namespace uba
 							break;
 						case ArgType_ISysRootPath:
 							AddRoot(sb);
+							#if PLATFORM_MAC
+							instance.frameworksDir = sb.data;
+							instance.frameworksDir += TC("/System/Library/Frameworks/");
 							instance.builtinIncludesDir = sb.data;
-							//instance.builtinIncludesDir += TC("/usr/include/c++");
+							instance.builtinIncludesDir += TC("/usr/include");
+							#endif
 							break;
 
 						case ArgType_Code:
@@ -573,14 +585,14 @@ namespace uba
 		return true;
 	}
 
-	bool DependencyCrawler::ParseCodeFile(TrackWorkScope& tracker, Instance& instance, const StringView& codeFile, bool parseDefines, const StringView& caller)
+	bool DependencyCrawler::ParseCodeFile(TrackWorkScope& tracker, Instance& instance, HandledFiles& handledFiles, const StringView& codeFile, bool parseDefines, const StringView& caller)
 	{
-		if (!instance.createFileFunc(tracker, codeFile, [&](const void* data, u64 dataSize) { return ParseCodeFile2(tracker, instance, data, dataSize, codeFile, parseDefines); }))
+		if (!instance.createFileFunc(tracker, codeFile, [&](const void* data, u64 dataSize) { return ParseCodeFile2(tracker, instance, handledFiles, data, dataSize, codeFile, parseDefines); }))
 			return m_logger.Warning(TC("Failed to parse code file %s found in %s"), codeFile.data, caller.data);
 		return true;
 	}
 
-	bool DependencyCrawler::ParseCodeFile2(TrackWorkScope& tracker, Instance& instance, const void* data, u64 dataSize, const StringView& codeFile, bool parseDefines)
+	bool DependencyCrawler::ParseCodeFile2(TrackWorkScope& tracker, Instance& instance, HandledFiles& handledFiles, const void* data, u64 dataSize, const StringView& codeFile, bool parseDefines)
 	{
 		const char* it = (const char*)data;
 		const char* end = it + dataSize;
@@ -594,9 +606,18 @@ namespace uba
 				while (*it == ' ' || *it == '\t')
 					++it;
 
-				if (strncmp(it, "include", 7) == 0)
+				bool isInclude = strncmp(it, "include", 7) == 0;
+				bool isImport = false;
+
+				#if PLATFORM_MAC
+				if (!isInclude)
+					isImport = strncmp(it, "import", 6) == 0;
+				#endif
+
+				if (isInclude || isImport)
 				{
-					it += 7;
+					it += isInclude ? 7 : 6;
+
 					while (*it == ' ' || *it == '\t')
 						++it;
 					StringBuffer<> include;
@@ -682,33 +703,72 @@ namespace uba
 
 					if (!include.IsEmpty())
 					{
-						include.MakeLower().FixPathSeparators();
+						if (CaseInsensitiveFs)
+							include.MakeLower();
+						include.FixPathSeparators();
+
 						auto work = 
-						[this, ref = InstanceRef(instance), isQuote, include2 = TString(include.data), codeFile2 = TString(codeFile.data), parseDefines](const WorkContext& context)
+						[this, ref = InstanceRef(instance), isQuote, isInclude, include2 = TString(include.data), codeFile2 = TString(codeFile.data), hf = &handledFiles, parseDefines](const WorkContext& context)
 							{
 								auto& instance = ref.instance;
-								if (isQuote)
+								auto& handledFiles = *hf;
+								if (isInclude)
 								{
-									StringBuffer<> localDir;
-									if (const tchar* lastSeparator = TStrrchr(codeFile2.c_str(), PathSeparator))
-										if (HandleInclude(context.tracker, instance, localDir.Append(codeFile2.c_str(), lastSeparator - codeFile2.c_str() + 1), include2, codeFile2, parseDefines))
+									if (isQuote)
+									{
+										StringBuffer<> localDir;
+										if (const tchar* lastSeparator = TStrrchr(codeFile2.c_str(), PathSeparator))
+											if (HandleInclude(context.tracker, instance, handledFiles, localDir.Append(codeFile2.c_str(), lastSeparator - codeFile2.c_str() + 1), include2, codeFile2, parseDefines))
+												return;
+									}
+
+									StringView keyView;
+									if (const tchar* firstSeparator = TStrchr(include2.c_str(), PathSeparator))
+										keyView = StringView(include2.c_str(), (firstSeparator - include2.c_str()));
+									else
+										keyView = StringView(include2);
+									StringKey fileKey =  ToStringKey(keyView);
+
+									for (auto& root : instance.includeRoots)
+									{
+										if (m_useBloomFilter && root.bloomFilter.IsGuaranteedMiss(fileKey))
+											continue;
+										if (HandleInclude(context.tracker, instance, handledFiles, root.path, include2, codeFile2, parseDefines))
 											return;
+									}
 								}
 
-								StringKey fileKey;
-								if (const tchar* firstSeparator = TStrchr(include2.c_str(), PathSeparator))
-									fileKey = ToStringKey(include2.c_str(), (firstSeparator - include2.c_str()));
-								else
-									fileKey = ToStringKey(include2);
-
-								for (auto& root : instance.includeRoots)
+								#if PLATFORM_MAC
+								if (!instance.frameworksDir.empty())
 								{
-									if (m_useBloomFilter && root.bloomFilter.IsGuaranteedMiss(fileKey))
-										continue;
-									if (HandleInclude(context.tracker, instance, root.path, include2, codeFile2, parseDefines))
-										return;
+									if (const tchar* firstSlash = TStrchr(include2.c_str(), '/'))
+									{
+										// TODO: Make a nicer solution for Frameworks in Frameworks. Also this is not the fastest doing HandleInclude on all these paths
+										const StringView frameworkRoots[] = { TCV(""), TCV("CoreServices.framework/Frameworks/"), TCV("ApplicationServices.framework/Frameworks/"), TCV("Carbon.framework/Frameworks/")  };
+
+										for (auto frameworkRoot : frameworkRoots)
+										{
+											StringBuffer<> tmp;
+											tmp.Append(instance.frameworksDir).Append(frameworkRoot).Append(include2.c_str(), firstSlash - include2.c_str()).Append(TCV(".framework/"));
+											u64 frameworkLen = tmp.count;
+											tmp.Append(TCV("Headers")).Append(firstSlash);
+											if (!HandleInclude(context.tracker, instance, handledFiles, {}, tmp, codeFile2, parseDefines))
+												continue;
+											tmp.Resize(frameworkLen).Append(TCV("Modules/module.modulemap"));
+											instance.devirtualizePathFunc(tmp);
+											if (CaseInsensitiveFs)
+												tmp.MakeLower();
+											StringKey moduleKey = ToStringKey(tmp);
+											SCOPED_FUTEX(handledFiles.lookupLock, lock);
+											if (!handledFiles.lookup.try_emplace(moduleKey).second)
+												return;
+											lock.Leave();
+											instance.createFileFunc(context.tracker, tmp, {});
+											return;
+										}
+									}
 								}
-
+								#endif
 							};
 						if (parseDefines)
 							work({tracker});
@@ -761,7 +821,7 @@ namespace uba
 		return true;
 	}
 
-	bool DependencyCrawler::HandleInclude(TrackWorkScope& tracker, Instance& instance, const StringView& rootPath, const StringView& include, const StringView& codeFile, bool parseDefines)
+	bool DependencyCrawler::HandleInclude(TrackWorkScope& tracker, Instance& instance, HandledFiles& handledFiles, const StringView& rootPath, const StringView& include, const StringView& codeFile, bool parseDefines)
 	{
 		if (include.EndsWith(TCV(".ush")))
 			return true;
@@ -780,23 +840,21 @@ namespace uba
 		if (IsDirectory(attributes))
 			return false;
 
+
+		StringKey key = ToStringKey(fullPath);
+		if (instance.pch)
 		{
-			StringKey key = ToStringKey(fullPath);
-			if (instance.pch)
-			{
-				SCOPED_FUTEX_READ(instance.pch->lock, lock);
-				if (instance.pch->files.find(key) != instance.pch->files.end())
-					return true;
-			}
-			if (!parseDefines) // We must find UBT_COMPILED_PLATFORM before we start earlying out
-			{
-				SCOPED_FUTEX(m_codeFilesHandledLock, lock);
-				if (!m_codeFilesHandled.try_emplace(key).second)
-					return true;
-			}
+			SCOPED_FUTEX_READ(instance.pch->lock, lock);
+			if (instance.pch->files.find(key) != instance.pch->files.end())
+				return true;
+		}
+		{
+			SCOPED_FUTEX(handledFiles.lookupLock, lock);
+			if (!handledFiles.lookup.try_emplace(key).second)
+				return true;
 		}
 
-		return ParseCodeFile(tracker, instance, fullPath, parseDefines, codeFile);
+		return ParseCodeFile(tracker, instance, handledFiles, fullPath, parseDefines, codeFile);
 	}
 
 	bool DependencyCrawler::TraverseInclude(TrackWorkScope& tracker, Instance& instance, const StringView& rootPathWithSlash)
@@ -807,12 +865,6 @@ namespace uba
 				codeFilePath.Resize(rootPathWithSlash.count).Append(file);
 				if (isDirectory)
 				{
-					#if PLATFORM_MAC
-					if (file.Equals(TCV("ruby"))) // Should probably be fixed in a nicer way but there are symlinks going back to previous dir in apple sdk
-						return;
-					if (file.Equals(TCV("Versions")) || file.Contains(TCV("_SwiftUI")) || file.Equals(TCV("PrivateFrameworks"))) // Directories we dont need to visit
-						return;
-					#endif
 					codeFilePath.Append(PathSeparator);
 					m_workManager.AddWork([this, ref = InstanceRef(instance), codeFile = TString(codeFilePath.data)](const WorkContext& context)
 						{
