[UBA}
* Fixed compile error

diff --git a/Engine/Source/Programs/UnrealBuildAccelerator/Common/Public/UbaDependencyCrawler.h b/Engine/Source/Programs/UnrealBuildAccelerator/Common/Public/UbaDependencyCrawler.h
--- a/Engine/Source/Programs/UnrealBuildAccelerator/Common/Public/UbaDependencyCrawler.h
+++ b/Engine/Source/Programs/UnrealBuildAccelerator/Common/Public/UbaDependencyCrawler.h
@@ -724,7 +724,7 @@ namespace uba
 
 									StringView keyView;
 									if (const tchar* firstSeparator = TStrchr(include2.c_str(), PathSeparator))
-										keyView = StringView(include2.c_str(), (firstSeparator - include2.c_str()));
+										keyView = StringView(include2.c_str(), u32(firstSeparator - include2.c_str()));
 									else
 										keyView = StringView(include2);
 									StringKey fileKey =  ToStringKey(keyView);
diff --git a/Engine/Source/Programs/UnrealBuildAccelerator/Test/Private/UbaTestDependencyCrawler.cpp b/Engine/Source/Programs/UnrealBuildAccelerator/Test/Private/UbaTestDependencyCrawler.cpp
--- a/Engine/Source/Programs/UnrealBuildAccelerator/Test/Private/UbaTestDependencyCrawler.cpp
+++ b/Engine/Source/Programs/UnrealBuildAccelerator/Test/Private/UbaTestDependencyCrawler.cpp
@@ -14,7 +14,7 @@ namespace uba
 
 		DependencyCrawler crawler(logger, workManager);
 
-		auto CreateFileFunc = [&](const StringView& fileName, const DependencyCrawler::AccessFileFunc& func)
+		auto CreateFileFunc = [&](TrackWorkScope& tracker, const StringView& fileName, const DependencyCrawler::AccessFileFunc& func)
 			{
 				FileAccessor fa(logger, fileName.data);
 				if (!fa.OpenMemoryRead())
@@ -47,10 +47,11 @@ namespace uba
 		//DependencyCrawlerType type = DependencyCrawlerType_MsvcCompiler;
 		//const tchar rsp[] = TC("e:\\dev\\fn\\Engine\\Intermediate\\Build\\Win64\\x64\\BlankProgram\\Development\\Core\\Module.Core.1.cpp.obj.rsp");
 		
-		DependencyCrawlerType type = DependencyCrawlerType_ClangCompiler;
+		DependencyCrawlerType type = DependencyCrawlerType_MsvcCompiler;
 		//const tchar rsp[] = TC("e:\\dev\\fn\\Engine\\Intermediate\\Build\\Linux\\x64\\BlankProgram\\Development\\Core\\Module.Core.1.cpp.o.rsp");
 		//const tchar rsp[] = TC("e:\\temp\\response-22b46d.txt");
-		const tchar rsp[] = TC("e:\\temp\\MiMalloc.c.o.rsp");
+		//const tchar rsp[] = TC("e:\\temp\\MiMalloc.c.o.rsp");
+		const tchar rsp[] = TC("e:\\dev\\fn\\Sandbox\\DevTools\\Clang\\Intermediate\\Build\\Win64\\x64\\ClangGame\\Debug\\InputCore\\Module.InputCore.cpp.obj.rsp");
 		//const tchar rsp[] = TC("e:\\dev\\fn\\engine\\plugins\\experimental\\geometrycollectionplugin\\intermediate\\build\\win64\\x64\\unrealeditor\\development\\geometrycollectionnodes\\module.geometrycollectionnodes.3.cpp.obj.rsp");
 		//const tchar rsp[] = TC("e:\\dev\\fn\\Engine\\Plugins\\Runtime\\MeshModelingToolset\\Intermediate\\Build\\Win64\\x64\\UnrealEditor\\Development\\ModelingComponents\\Module.ModelingComponents.3.cpp.obj.rsp");
 		//const tchar rsp[] = TC("e:\\dev\\fn\\Engine\\Intermediate\\Build\\Win64\\x64\\UnrealEditor\\Development\\AnimGraphRuntime\\Module.AnimGraphRuntime.3.cpp.obj.rsp");
