[UBA]
* Fixed bug in dependency crawler for msvc related to when there are multiple /FI

diff --git a/Engine/Source/Programs/UnrealBuildAccelerator/Common/Public/UbaDependencyCrawler.h b/Engine/Source/Programs/UnrealBuildAccelerator/Common/Public/UbaDependencyCrawler.h
--- a/Engine/Source/Programs/UnrealBuildAccelerator/Common/Public/UbaDependencyCrawler.h
+++ b/Engine/Source/Programs/UnrealBuildAccelerator/Common/Public/UbaDependencyCrawler.h
@@ -336,22 +336,29 @@ namespace uba
 							if (CaseInsensitiveFs)
 								fixedFile.MakeLower();
 
-							StringKey key = ToStringKey(fixedFile);
-							SCOPED_FUTEX(m_handledFiles.lookupLock, lock);
-							auto insres = m_handledFiles.lookup.try_emplace(key);
-							lock.Leave();
-
-							HandledFile& handledFile = insres.first->second;
-							SCOPED_FUTEX(handledFile.lock, fileLock);
-							if (!handledFile.handled)
+							if (fixedFile.Contains(TC("\\Definitions.")))
 							{
-								handledFile.handled = true;
+								AddCodeFile(fixedFile, true, false);
+							}
+							else
+							{
+								StringKey key = ToStringKey(fixedFile);
+								SCOPED_FUTEX(m_handledFiles.lookupLock, lock);
+								auto insres = m_handledFiles.lookup.try_emplace(key);
+								lock.Leave();
+
+								HandledFile& handledFile = insres.first->second;
+								SCOPED_FUTEX(handledFile.lock, fileLock);
+								if (!handledFile.handled)
+								{
+									handledFile.handled = true;
 
-								fixedFile.Append(TCV(".pch"));
+									fixedFile.Append(TCV(".pch"));
 
-								u32 attributes = 0;
-								bool hasPch = m_fileExistsFunc(fixedFile, attributes);
-								AddCodeFile(StringView(sb).Skip(3), true, hasPch);
+									u32 attributes = 0;
+									bool hasPch = m_fileExistsFunc(fixedFile, attributes);
+									AddCodeFile(StringView(sb).Skip(3), true, hasPch);
+								}
 							}
 
 							prevArg.Clear();
