[UBA]
* Cleaned up option parsing for UbaAgent so it is possible to use environment variables for more options

diff --git a/Engine/Source/Programs/UnrealBuildAccelerator/Agent/Private/UbaAgent.cpp b/Engine/Source/Programs/UnrealBuildAccelerator/Agent/Private/UbaAgent.cpp
--- a/Engine/Source/Programs/UnrealBuildAccelerator/Agent/Private/UbaAgent.cpp
+++ b/Engine/Source/Programs/UnrealBuildAccelerator/Agent/Private/UbaAgent.cpp
@@ -54,7 +54,7 @@ namespace uba
 	u32				DefaultProcessorCount = []() { return GetLogicalProcessorCount(); }();
 	u32				DefaultMaxConnectionCount = 4;
 
-	int PrintHelp(const tchar* message)
+	bool PrintHelp(const tchar* message)
 	{
 		LoggerWithWriter logger(g_consoleLogWriter, TC(""));
 		if (*message)
@@ -342,14 +342,14 @@ namespace uba
 		u32 maxConnectionCount = DefaultMaxConnectionCount;
 		u32 storageCapacityGb = DefaultCapacityGb;
 		StringBuffer<256> host;
-		StringBuffer<256> named;
+		TString named;
 		StringBuffer<512> relaunchPath;
 		StringBuffer<256> eventFile;
-		StringBuffer<256> configFile;
+		TString configFile;
 		TString command;
 		u16 port = DefaultPort;
 		u16 proxyPort = DefaultStorageProxyPort;
-		StringBuffer<> agentName;
+		TString agentName;
 		bool useListen = false;
 		bool logToFile = false;
 		bool storeCompressed = true;
@@ -371,7 +371,7 @@ namespace uba
 		bool downloadDetoursLib = true;
 		bool useExceptionHandler = true;
 		bool resendCas = false;
-		StringBuffer<512> sentryUrl;
+		TString sentryUrl;
 		StringBuffer<128> zone;
 		u32 maxIdleSeconds = ~0u;
 		u32 sendSize = SendDefaultSize;
@@ -394,6 +394,20 @@ namespace uba
 			StringBuffer<> name;
 			StringBuffer<> value;
 
+			auto parseValue = [&](auto& out)
+				{
+					if (value.IsEmpty())
+						return PrintHelp(StringBuffer().Appendf(TC("%s needs a value"), name.data).data);
+					if (!ExpandEnvironmentVariables(value, PrintHelp))
+						return false;
+					std::decay_t<decltype(out)> temp;
+					if (value.Parse(temp))
+						out = temp;
+					else
+						LoggerWithWriter(g_consoleLogWriter, TC("")).Warning(TC("Invalid value for %s, ignoring and will use default!"), name.data);
+					return true;
+				};
+
 			if (const tchar* equals = TStrchr(argv[i],'='))
 			{
 				name.Append(argv[i], equals - argv[i]);
@@ -418,48 +432,35 @@ namespace uba
 			}
 			else if (name.Equals(TCV("-maxcpu")))
 			{
-				if (!ExpandEnvironmentVariables(value, PrintHelp))
+				if (!parseValue(maxProcessCount))
 					return false;
-				u32 defaultValue = maxProcessCount;
-				if (!value.Parse(maxProcessCount))
-				{
-					LoggerWithWriter(g_consoleLogWriter, TC("")).Warning(TC("Invalid value for -maxcpu, ignoring!"));
-					maxProcessCount = defaultValue;
-				}
 			}
 			else if (name.Equals(TCV("-mulcpu")))
 			{
-				if (!ExpandEnvironmentVariables(value, PrintHelp))
+				if (!parseValue(mulProcessValue))
 					return false;
-				float defaultValue = mulProcessValue;
-				if (!value.Parse(mulProcessValue))
-				{
-					LoggerWithWriter(g_consoleLogWriter, TC("")).Warning(TC("Invalid value for -mulcpu, ignoring!"));
-					mulProcessValue = defaultValue;
-				}
 			}
 			else if (name.Equals(TCV("-maxcon")) || name.Equals(TCV("-maxtcp")))
 			{
-				if (!value.Parse(maxConnectionCount) || maxConnectionCount == 0)
+				if (!parseValue(maxConnectionCount))
+					return false;
+				if (maxConnectionCount == 0)
 					return PrintHelp(TC("Invalid value for -maxcon"));
 			}
 			else if (name.Equals(TCV("-maxworkers")))
 			{
-				if (!value.Parse(maxWorkerCount))
-					return PrintHelp(TC("Invalid value for -maxworkers"));
+				if (!parseValue(maxWorkerCount))
+					return false;
 			}
 			else if (name.Equals(TCV("-capacity")))
 			{
-				if (!value.Parse(storageCapacityGb))
-					return PrintHelp(TC("Invalid value for -capacity"));
+				if (!parseValue(storageCapacityGb))
+					return false;
 			}
 			else if (name.Equals(TCV("-config")))
 			{
-				if (value.IsEmpty())
-					return PrintHelp(TC("-config needs a value"));
-				if (!ExpandEnvironmentVariables(value, PrintHelp))
+				if (!parseValue(configFile))
 					return false;
-				configFile.Append(value);
 			}
 			else if (name.Equals(TCV("-host")))
 			{
@@ -476,23 +477,19 @@ namespace uba
 			}
 			else if (name.Equals(TCV("-listen")))
 			{
-				if (!value.IsEmpty())
-					if (!value.Parse(port))
-						return PrintHelp(TC("Invalid value for -capacity"));
+				if (!parseValue(port))
+					return false;
 				useListen = true;
 			}
 			else if (name.Equals(TCV("-listenTimeout")))
 			{
-				if (value.IsEmpty())
-					return PrintHelp(TC("-listenTimeout needs a value"));
-				if (!value.Parse(listenTimeoutSec))
-					return PrintHelp(TC("Invalid value for -listenTimeout"));
+				if (!parseValue(listenTimeoutSec))
+					return false;
 			}
 			else if (name.Equals(TCV("-named")))
 			{
-				if (value.IsEmpty())
-					return PrintHelp(TC("-named needs a value"));
-				named.Append(value);
+				if (!parseValue(named))
+					return false;
 			}
 			else if (name.Equals(TCV("-log")))
 			{
@@ -516,8 +513,8 @@ namespace uba
 			}
 			else if (name.Equals(TCV("-sendsize")))
 			{
-				if (!value.Parse(sendSize))
-					return PrintHelp(TC("Invalid value for -sendsize"));
+				if (!parseValue(sendSize))
+					return false;
 			}
 			else if (name.Equals(TCV("-dir")))
 			{
@@ -530,9 +527,8 @@ namespace uba
 			}
 			else if (name.Equals(TCV("-name")))
 			{
-				if (value.IsEmpty())
-					return PrintHelp(TC("-name needs a value"));
-				agentName.Append(value);
+				if (!parseValue(agentName))
+					return false;
 			}
 			else if (name.Equals(TCV("-nopoll")))
 			{
@@ -568,13 +564,13 @@ namespace uba
 			}
 			else if (name.Equals(TCV("-maxidle")))
 			{
-				if (!value.Parse(maxIdleSeconds))
-					return PrintHelp(TC("Invalid value for -maxidle"));
+				if (!parseValue(maxIdleSeconds))
+					return false;
 			}
 			else if (name.Equals(TCV("-proxyport")))
 			{
-				if (!value.Parse(proxyPort))
-					return PrintHelp(TC("Invalid value for -proxyport"));
+				if (!parseValue(proxyPort))
+					return false;
 			}
 			else if (name.Equals(TCV("-summary")))
 			{
@@ -599,12 +595,16 @@ namespace uba
 			}
 			else if (name.Equals(TCV("-memwait")))
 			{
-				if (!value.Parse(memWaitLoadPercent) || memWaitLoadPercent > 100)
+				if (!parseValue(memWaitLoadPercent))
+					return false;
+				if (memWaitLoadPercent > 100)
 					return PrintHelp(TC("Invalid value for -memwait"));
 			}
 			else if (name.Equals(TCV("-memkill")))
 			{
-				if (!value.Parse(memKillLoadPercent) || memKillLoadPercent > 100)
+				if (!parseValue(memKillLoadPercent))
+					return false;
+				if (memKillLoadPercent > 100)
 					return PrintHelp(TC("Invalid value for -memkill"));
 			}
 			else if (name.Equals(TCV("-crypto")))
@@ -621,14 +621,18 @@ namespace uba
 			}
 			else if (name.Equals(TCV("-populateCas")))
 			{
-				if (value.IsEmpty())
-					return PrintHelp(TC("-populateCas needs a dir"));
-				populateCasDirs.push_back(value.data);
+				TString temp;
+				if (!parseValue(temp))
+					return false;
+				populateCasDirs.push_back(temp);
 			}
 			#if PLATFORM_MAC
 			else if (name.Equals(TCV("-populateCasFromXcodeVersion")))
 			{
-				populateCasFromXcodeVersion.Append(value.data);
+				TString temp;
+				if (!parseValue(temp))
+					return false;
+				populateCasFromXcodeVersion.Append(temp);
 			}
 			else if (name.Equals(TCV("-populateCasFromAllXcodes")))
 			{
@@ -637,15 +641,13 @@ namespace uba
 			#endif
 			else if (name.Equals(TCV("-sentry")))
 			{
-				if (value.IsEmpty())
-					return PrintHelp(TC("-sentry needs a url value"));
-				sentryUrl.Append(value);
+				if (!parseValue(sentryUrl))
+					return false;
 			}
 			else if (name.Equals(TCV("-zone")))
 			{
-				if (value.IsEmpty())
-					return PrintHelp(TC("-zone needs a value"));
-				zone.Append(value);
+				if (!parseValue(zone))
+					return false;
 			}
 			else if (name.Equals(TCV("-version")))
 			{
@@ -678,9 +680,8 @@ namespace uba
 			}
 			else if (name.Equals(TCV("-command")))
 			{
-				if (value.IsEmpty())
-					return PrintHelp(TC("-command needs a value"));
-				command = value.data;
+				if (!parseValue(command))
+					return false;
 				poll = false;
 				quiet = true;
 			}
@@ -699,7 +700,7 @@ namespace uba
 		if (useExceptionHandler)
 			AddExceptionHandler();
 
-		if (!named.IsEmpty()) // We only run once with named connection
+		if (!named.empty()) // We only run once with named connection
 			poll = false;
 
 		maxProcessCount = u32(float(maxProcessCount) * mulProcessValue);
@@ -727,19 +728,19 @@ namespace uba
 			return LaunchReal(logger, relaunchPath, argc, argv);
 #endif // UBA_AUTO_UPDATE
 
-		if (host.IsEmpty() && named.IsEmpty() && !useListen)
+		if (host.IsEmpty() && named.empty() && !useListen)
 			return PrintHelp(TC("No host provided. Add -host=<host> (or use -listen)"));
 
 		StringBuffer<256> extraInfo;
 
 		#if defined(UBA_USE_SENTRY)
-		if (!sentryUrl.IsEmpty())
+		if (!sentryUrl.empty())
 		{
 			char release[128];
 			char url[512];
 			size_t urlLen;
 			sprintf_s(release, sizeof_array(release), "BoxAgent@%ls", Version);
-			wcstombs_s(&urlLen, url, sizeof_array(url), sentryUrl.data, sizeof_array(url) - 1);
+			wcstombs_s(&urlLen, url, sizeof_array(url), sentryUrl.c_str(), sizeof_array(url) - 1);
 			sentry_options_t* options = sentry_options_new();
 			sentry_options_set_dsn(options, url);
 			sentry_options_set_database_path(options, ".sentry-native");
@@ -749,7 +750,7 @@ namespace uba
 			extraInfo.Append(TCV(", SentryEnabled"));
 		}
 		
-		auto sentryGuard = MakeGuard([&]() { if (!sentryUrl.IsEmpty()) sentry_close(); });
+		auto sentryGuard = MakeGuard([&]() { if (!sentryUrl.empty()) sentry_close(); });
 		#endif
 
 
@@ -766,8 +767,12 @@ namespace uba
 		}
 		#endif
 
-		if (agentName.IsEmpty())
-			GetComputerNameW(agentName);
+		if (agentName.empty())
+		{
+			StringBuffer<128> temp;
+			if (GetComputerNameW(temp))
+				agentName = temp.ToString();
+		}
 
 		if (!zone.count)
 			GetZone(zone);
@@ -791,8 +796,8 @@ namespace uba
 		logger.Info(TC("UbaAgent v%s%s (Cpu%s: %u, MaxCon: %u, Dir: \"%s\", StoreCapacity: %uGb, Zone: %s%s)"), Version, dbgStr, (IsArmBinary ? TC("[Arm]") : TC("")), maxProcessCount, maxConnectionCount, g_rootDir.data, storageCapacityGb, zone.count ? zone.data : TC("none"), extraInfo.data);
 		
 		Config config;
-		if (!configFile.IsEmpty())
-			config.LoadFromFile(logger, configFile.data);
+		if (!configFile.empty())
+			config.LoadFromFile(logger, configFile.c_str());
 				
 		if (!eventFile.IsEmpty())
 			logger.Info(TC("  Will poll for external events in file %s"), eventFile.data);
diff --git a/Engine/Source/Programs/UnrealBuildAccelerator/Common/Private/UbaCacheClient.cpp b/Engine/Source/Programs/UnrealBuildAccelerator/Common/Private/UbaCacheClient.cpp
--- a/Engine/Source/Programs/UnrealBuildAccelerator/Common/Private/UbaCacheClient.cpp
+++ b/Engine/Source/Programs/UnrealBuildAccelerator/Common/Private/UbaCacheClient.cpp
@@ -37,7 +37,7 @@
 
 namespace uba
 {
-	struct DowngradedLogger : public LoggerWithWriter
+	struct CacheClient::DowngradedLogger : public LoggerWithWriter
 	{
 		DowngradedLogger(Atomic<bool>& c, LogWriter& writer, const tchar* prefix) : LoggerWithWriter(writer, prefix), connected(c) {}
 		virtual void Log(LogEntryType type, const tchar* str, u32 strLen) override { if (connected) LoggerWithWriter::Log(Max(type, LogEntryType_Info), str, strLen); }
diff --git a/Engine/Source/Programs/UnrealBuildAccelerator/Common/Public/UbaCacheClient.h b/Engine/Source/Programs/UnrealBuildAccelerator/Common/Public/UbaCacheClient.h
--- a/Engine/Source/Programs/UnrealBuildAccelerator/Common/Public/UbaCacheClient.h
+++ b/Engine/Source/Programs/UnrealBuildAccelerator/Common/Public/UbaCacheClient.h
@@ -124,5 +124,6 @@ namespace uba
 
 		CacheClient(const CacheClient&) = delete;
 		CacheClient& operator=(const CacheClient&) = delete;
+		struct DowngradedLogger;
 	};
 }
\ No newline at end of file
diff --git a/Engine/Source/Programs/UnrealBuildAccelerator/Core/Private/UbaStringBuffer.cpp b/Engine/Source/Programs/UnrealBuildAccelerator/Core/Private/UbaStringBuffer.cpp
--- a/Engine/Source/Programs/UnrealBuildAccelerator/Core/Private/UbaStringBuffer.cpp
+++ b/Engine/Source/Programs/UnrealBuildAccelerator/Core/Private/UbaStringBuffer.cpp
@@ -364,12 +364,12 @@ namespace uba
 		return *this;
 	}
 
-	bool StringBufferBase::Parse(u64& out)
+	bool StringBufferBase::Parse(u64& out) const
 	{
 		return uba::Parse(out, data, count);
 	}
 
-	bool StringBufferBase::Parse(u32& out, u64 offset)
+	bool StringBufferBase::Parse(u32& out, u64 offset) const
 	{
 		if (count <= offset)
 			return false;
@@ -381,7 +381,7 @@ namespace uba
 		return out != 0 || Equals(TCV("0"));
 	}
 
-	bool StringBufferBase::Parse(u16& out, u64 offset)
+	bool StringBufferBase::Parse(u16& out, u64 offset) const
 	{
 		u32 temp;
 		if (!Parse(temp, offset))
@@ -392,7 +392,7 @@ namespace uba
 		return true;
 	}
 
-	bool StringBufferBase::Parse(float& out)
+	bool StringBufferBase::Parse(float& out) const
 	{
 		if (!count)
 			return false;
@@ -404,7 +404,19 @@ namespace uba
 		return out != 0 || Equals(TCV("0"));
 	}
 
-	u32 StringBufferBase::Parse(char* out, u64 outCapacity)
+	bool StringBufferBase::Parse(TString& out, u64 offset) const
+	{
+		out.append(data + offset, count - offset);
+		return true;
+	}
+
+	bool StringBufferBase::Parse(StringBufferBase& out, u64 offset) const
+	{
+		out.Append(data + offset, count - offset);
+		return true;
+	}
+
+	u32 StringBufferBase::Parse(char* out, u64 outCapacity) const
 	{
 		#if PLATFORM_WINDOWS
 		size_t destLen;
diff --git a/Engine/Source/Programs/UnrealBuildAccelerator/Core/Public/UbaStringBuffer.h b/Engine/Source/Programs/UnrealBuildAccelerator/Core/Public/UbaStringBuffer.h
--- a/Engine/Source/Programs/UnrealBuildAccelerator/Core/Public/UbaStringBuffer.h
+++ b/Engine/Source/Programs/UnrealBuildAccelerator/Core/Public/UbaStringBuffer.h
@@ -70,11 +70,13 @@ namespace uba
 		StringBufferBase& FixPathSeparators();
 		StringBufferBase& MakeLower();
 
-		bool Parse(u64& out);
-		bool Parse(u32& out, u64 offset = 0);
-		bool Parse(u16& out, u64 offset = 0);
-		bool Parse(float& out);
-		u32 Parse(char* out, u64 outCapacity); // Note, return value contains length + null termination.. zero if failed
+		bool Parse(u64& out) const;
+		bool Parse(u32& out, u64 offset = 0) const;
+		bool Parse(u16& out, u64 offset = 0) const;
+		bool Parse(float& out) const;
+		bool Parse(TString& out, u64 offset = 0) const;
+		bool Parse(StringBufferBase& out, u64 offset = 0) const;
+		u32 Parse(char* out, u64 outCapacity) const; // Note, return value contains length + null termination.. zero if failed
 
 		TString ToString() const { return TString(data, data + count); }
 
diff --git a/Engine/Source/Programs/UnrealBuildAccelerator/Host/Private/UbaExports.cpp b/Engine/Source/Programs/UnrealBuildAccelerator/Host/Private/UbaExports.cpp
--- a/Engine/Source/Programs/UnrealBuildAccelerator/Host/Private/UbaExports.cpp
+++ b/Engine/Source/Programs/UnrealBuildAccelerator/Host/Private/UbaExports.cpp
@@ -108,9 +108,9 @@ namespace uba
 		return config;
 	}
 
-	struct DowngradedLogger : public LoggerWithWriter
+	struct ExportsDowngradedLogger : public LoggerWithWriter
 	{
-		DowngradedLogger(LogWriter& writer, const tchar* prefix) : LoggerWithWriter(writer, prefix) {}
+		ExportsDowngradedLogger(LogWriter& writer, const tchar* prefix) : LoggerWithWriter(writer, prefix) {}
 		virtual void Log(LogEntryType type, const tchar* str, u32 strLen) override { LoggerWithWriter::Log(Max(LogEntryType_Info, type), str, strLen); }
 	};
 
@@ -292,7 +292,7 @@ extern "C"
 		AWS aws;
 		if (!info.zone || !*info.zone)
 		{
-			DowngradedLogger downgradedLogger(info.writer, TC("AWS"));
+			ExportsDowngradedLogger downgradedLogger(info.writer, TC("AWS"));
 			if (aws.QueryAvailabilityZone(downgradedLogger, info.rootDir))
 				info.zone = aws.GetAvailabilityZone();
 		}
@@ -903,7 +903,7 @@ extern "C"
 	bool Uba_GetExclusiveAccess(const uba::tchar* path)
 	{
 		using namespace uba;
-		DowngradedLogger downgradedLogger(g_consoleLogWriter, TC("UbaGetExclusiveAccess"));
+		ExportsDowngradedLogger downgradedLogger(g_consoleLogWriter, TC("UbaGetExclusiveAccess"));
 		g_exclusiveMutex = StorageImpl::GetExclusiveAccess(downgradedLogger, ToView(path), false);
 		return g_exclusiveMutex != InvalidMutexHandle;
 	}
