[UBA]
* Fixed so client is kicked if trying to download disallowed cas files

diff --git a/Engine/Source/Programs/UnrealBuildAccelerator/Common/Private/UbaNetworkServer.cpp b/Engine/Source/Programs/UnrealBuildAccelerator/Common/Private/UbaNetworkServer.cpp
--- a/Engine/Source/Programs/UnrealBuildAccelerator/Common/Private/UbaNetworkServer.cpp
+++ b/Engine/Source/Programs/UnrealBuildAccelerator/Common/Private/UbaNetworkServer.cpp
@@ -515,7 +515,11 @@ namespace uba
 		else if (!rec.func({&connection}, {tws}, mi, reader, writer))
 		{
 			if (connection.SetShouldDisconnect())
+			{
+				#if UBA_DEBUG
 				server.m_logger.Error(TC("WORKER FUNCTION FAILED. id: %u, serviceid: %u type: %s, client: %u"), context.id, context.serviceId, rec.toString(context.messageType).data, connection.m_client->id);
+				#endif
+			}
 			size = ErrorSize;
 		}
 		else
diff --git a/Engine/Source/Programs/UnrealBuildAccelerator/Common/Private/UbaStorageServer.cpp b/Engine/Source/Programs/UnrealBuildAccelerator/Common/Private/UbaStorageServer.cpp
--- a/Engine/Source/Programs/UnrealBuildAccelerator/Common/Private/UbaStorageServer.cpp
+++ b/Engine/Source/Programs/UnrealBuildAccelerator/Common/Private/UbaStorageServer.cpp
@@ -507,8 +507,9 @@ namespace uba
 				if (casEntry->disallowed)
 				{
 					writer.WriteU16(0);
-					m_logger.Error(TC("Client is asking for cas content of file that is not allowed to be transferred. (%s)"), hint.data);
-					return true;
+					if (!connectionInfo.ShouldDisconnect())
+						m_logger.Warning(TC("Client %u is asking for cas content of file that is not allowed to be transferred. Client will be disconnected (%s)"), connectionInfo.GetId(), hint.data);
+					return false;
 				}
 				
 				StringBuffer<512> casFile;
diff --git a/Engine/Source/Programs/UnrealBuildAccelerator/Common/Public/UbaNetworkMessage.h b/Engine/Source/Programs/UnrealBuildAccelerator/Common/Public/UbaNetworkMessage.h
--- a/Engine/Source/Programs/UnrealBuildAccelerator/Common/Public/UbaNetworkMessage.h
+++ b/Engine/Source/Programs/UnrealBuildAccelerator/Common/Public/UbaNetworkMessage.h
@@ -43,4 +43,9 @@ namespace uba
 		NetworkMessage(const NetworkMessage&) = delete;
 		NetworkMessage(NetworkMessage&&) = delete;
 	};
+
+	enum NetworkMessageError
+	{
+		NetworkMessageError_Disconnecting = 1,
+	};
 }
