[UBA]
* Cleaned up aws code and made it cloud agnostic and added support for google cloud

diff --git a/Engine/Source/Programs/UnrealBuildAccelerator/Agent/Private/UbaAgent.cpp b/Engine/Source/Programs/UnrealBuildAccelerator/Agent/Private/UbaAgent.cpp
--- a/Engine/Source/Programs/UnrealBuildAccelerator/Agent/Private/UbaAgent.cpp
+++ b/Engine/Source/Programs/UnrealBuildAccelerator/Agent/Private/UbaAgent.cpp
@@ -106,7 +106,7 @@ namespace uba
 		logger.Info(TC("  -version                Prints the version for this executable."));
 		logger.Info(TC("  -noproxy                Does not allow this agent to be a storage proxy for other agents"));
 		logger.Info(TC("  -proxyuselocalstorage   Storage proxy will use local storage to see if files exist"));
-		logger.Info(TC("  -noaws                  Will not try to connect to aws meta data server (this can take time during first startup)"));
+		logger.Info(TC("  -nocloud                Will not try to connect to cloud meta data server (this can take time during first startup)"));
 		logger.Info(TC("  -killrandom             Kills random process and exit session"));
 		logger.Info(TC("  -memwait=<percent>      The amount of memory needed to spawn a process. Set this to 100 to disable. Defaults to 80%%"));
 		logger.Info(TC("  -memkill=<percent>      The amount of memory needed before processes starts to be killed. Set this to 100 to disable. Defaults to 90%%"));
@@ -360,7 +360,7 @@ namespace uba
 		bool poll = true;
 		bool allowProxy = true;
 		bool proxyUseLocalStorage = false;
-		bool couldBeAws = true;
+		bool couldBeCloud = true;
 		bool useStorage = true;
 		bool resetStore = false;
 		bool quiet = false;
@@ -674,9 +674,9 @@ namespace uba
 					value.Resize(value.count - 1);
 				description = value.data + (value[0] == '\"');
 			}
-			else if (name.Equals(TCV("-noaws")))
+			else if (name.Equals(TCV("-nocloud")))
 			{
-				couldBeAws = false;
+				couldBeCloud = false;
 			}
 			else if (name.Equals(TCV("-command")))
 			{
@@ -754,16 +754,16 @@ namespace uba
 		#endif
 
 
-		// Check if AWS
-		#if UBA_USE_AWS
-		AWS aws;
-		if (couldBeAws)
+		// Check if cloud (AWS/google cloud etc)
+		#if UBA_USE_CLOUD
+		Cloud cloud;
+		if (couldBeCloud)
 		{
 			DirectoryCache dirCache;
 			dirCache.CreateDirectory(logger, g_rootDir.data);
-			aws.QueryInformation(logger, extraInfo, g_rootDir.data);
+			cloud.QueryInformation(logger, extraInfo, g_rootDir.data);
 			if (zone.IsEmpty())
-				zone.Append(aws.GetAvailabilityZone());
+				zone.Append(cloud.GetAvailabilityZone());
 		}
 		#endif
 
@@ -990,12 +990,13 @@ namespace uba
 		bool relaunch = false;
 		u64 terminationTimeMs = 0;
 
-		#if UBA_USE_AWS
-		if (couldBeAws && terminationReason.IsEmpty())
-			aws.IsTerminating(logger, terminationReason, terminationTimeMs);
+		#if UBA_USE_CLOUD
+		if (couldBeCloud && terminationReason.IsEmpty())
+			cloud.IsTerminating(logger, terminationReason, terminationTimeMs);
 		#endif
 
 		bool isTerminating = !terminationReason.IsEmpty();
+		
 		do
 		{
 			// TODO: Revisit.. it could be that something stalls during startup so we can't have this timeout
@@ -1035,8 +1036,19 @@ namespace uba
 						return true;
 
 					u64 waitTime = GetTime() - startTime;
-					if (!poll && TimeToMs(waitTime) > listenTimeoutSec*1000)
-						return logger.Error(TC("Failed to get connection while listening for %s"), TimeToText(waitTime).str);
+					if (!poll)
+					{
+						if (TimeToMs(waitTime) > listenTimeoutSec*1000)
+							return logger.Error(TC("Failed to get connection while listening for %s"), TimeToText(waitTime).str);
+						continue;
+					}
+
+					#if UBA_USE_CLOUD
+					if (couldBeCloud && !isTerminating && cloud.IsTerminating(logger, terminationReason, terminationTimeMs))
+						isTerminating = true;
+					#endif
+					if (isTerminating)
+						return logger.Error(TC("Terminating.. (%s)"), terminationReason.data);
 				}
 			}
 			else
@@ -1053,8 +1065,19 @@ namespace uba
 					if (!timedOut)
 						return false;
 
-					if (!poll && !--retryCount)
-						return logger.Error(TC("Failed to connect to %s:%u (after %s)"), host.data, port, TimeToText(GetTime() - startTime).str);
+					if (!poll)
+					{
+						if (!--retryCount)
+							return logger.Error(TC("Failed to connect to %s:%u (after %s)"), host.data, port, TimeToText(GetTime() - startTime).str);
+						continue;
+					}
+
+					#if UBA_USE_CLOUD
+					if (couldBeCloud && !isTerminating && cloud.IsTerminating(logger, terminationReason, terminationTimeMs))
+						isTerminating = true;
+					#endif
+					if (isTerminating)
+						return logger.Error(TC("Terminating.. (%s)"), terminationReason.data);
 				}
 			}
 
@@ -1357,8 +1380,8 @@ namespace uba
 				{
 					if (IsTerminating(logger, eventFile.data, terminationReason, terminationTimeMs))
 						isTerminating = true;
-					#if UBA_USE_AWS
-					else if (couldBeAws && aws.IsTerminating(logger, terminationReason, terminationTimeMs))
+					#if UBA_USE_CLOUD
+					else if (couldBeCloud && cloud.IsTerminating(logger, terminationReason, terminationTimeMs))
 						isTerminating = true;
 					#endif
 
diff --git a/Engine/Source/Programs/UnrealBuildAccelerator/Cli/Private/UbaCli.cpp b/Engine/Source/Programs/UnrealBuildAccelerator/Cli/Private/UbaCli.cpp
--- a/Engine/Source/Programs/UnrealBuildAccelerator/Cli/Private/UbaCli.cpp
+++ b/Engine/Source/Programs/UnrealBuildAccelerator/Cli/Private/UbaCli.cpp
@@ -82,7 +82,7 @@ namespace uba
 		logger.Info(TC("   -vfs=<virtual>;<local>  Will convert virtual path to local under the hood. Can have multiple -vfs"));
 		logger.Info(TC("   -checkcas               Check so all cas entries are correct"));
 		logger.Info(TC("   -checkfiletable         Check so file table has correct cas stored"));
-		logger.Info(TC("   -checkaws               Check if we are inside aws and output information about aws"));
+		logger.Info(TC("   -checkcloud             Check if we are inside cloud and output information about cloud"));
 		logger.Info(TC("   -deletecas              Deletes the casdb"));
 		logger.Info(TC("   -getcas                 Will print hash of application"));
 		logger.Info(TC("   -summary                Print summary at the end of a session"));
@@ -163,7 +163,7 @@ namespace uba
 		bool quiet = false;
 		bool checkCas = false;
 		bool checkCas2 = false;
-		bool checkAws = false;
+		bool checkCloud = false;
 		bool getCas = false;
 		bool deleteCas = false;
 		bool enableStdOut = true;
@@ -324,9 +324,9 @@ namespace uba
 				{
 					checkCas2 = true;
 				}
-				else if (name.Equals(TCV("-checkaws")))
+				else if (name.Equals(TCV("-checkcloud")))
 				{
-					checkAws = true;
+					checkCloud = true;
 				}
 				else if (name.Equals(TCV("-testcompress")))
 				{
@@ -621,22 +621,24 @@ namespace uba
 			exit = true;
 		}
 
-#if UBA_USE_AWS
-		if (checkAws)
+#if UBA_USE_CLOUD
+		if (checkCloud)
 		{
-			AWS aws;
+			DirectoryCache dirCache;
+			dirCache.CreateDirectory(logger, g_rootDir.data);
+			Cloud cloud;
 			StringBuffer<> info;
-			if (aws.QueryInformation(logger, info, g_rootDir.data))
+			if (cloud.QueryInformation(logger, info, g_rootDir.data))
 			{
-				logger.Info(TC("We are inside AWS: %s (%s)"), info.data, aws.GetAvailabilityZone());
+				logger.Info(TC("We are inside cloud%s (%s)"), info.data, cloud.GetAvailabilityZone());
 				
 				StringBuffer<> reason;
 				u64 terminateTime;
-				if (aws.IsTerminating(logger, reason, terminateTime))
+				if (cloud.IsTerminating(logger, reason, terminateTime))
 					logger.Info(TC(".. and are being terminated: %s"), reason.data);
 			}
 			else
-				logger.Info(TC("Seems like we are not running inside aws."));
+				logger.Info(TC("Seems like we are not running inside cloud."));
 			exit = true;
 		}
 #endif
diff --git a/Engine/Source/Programs/UnrealBuildAccelerator/Common/Private/UbaFile.cpp b/Engine/Source/Programs/UnrealBuildAccelerator/Common/Private/UbaFile.cpp
--- a/Engine/Source/Programs/UnrealBuildAccelerator/Common/Private/UbaFile.cpp
+++ b/Engine/Source/Programs/UnrealBuildAccelerator/Common/Private/UbaFile.cpp
@@ -325,7 +325,12 @@ namespace uba
 				SetLastError(ERROR_FILE_NOT_FOUND);
 				return false;
 			}
-			UBA_ASSERTF(false, TC("FileExists error handling implemented"));
+			if (errno == ENOTDIR)
+			{
+				SetLastError(ERROR_PATH_NOT_FOUND);
+				return false;
+			}
+			UBA_ASSERTF(false, TC("FileExists error handling for %u is not implemented (%s)"), errno, strerror(errno));
 			return false;
 		}
 
@@ -557,6 +562,12 @@ namespace uba
 			return InvalidFileHandle;
 		}
 
+		if (errno == ENOTDIR)
+		{
+			SetLastError(ERROR_PATH_NOT_FOUND);
+			return InvalidFileHandle;
+		}
+
 		if (errno == EACCES)
 		{
 			SetLastError(ERROR_ACCESS_DENIED);
@@ -600,7 +611,7 @@ namespace uba
 			SetLastError(ERROR_ALREADY_EXISTS);
 			return false;
 		}
-		if (errno == ENOENT)
+		if (errno == ENOENT || errno == ENOTDIR)
 		{
 			SetLastError(ERROR_PATH_NOT_FOUND);
 			return false;
diff --git a/Engine/Source/Programs/UnrealBuildAccelerator/Common/Private/UbaNetworkBackendTcp.cpp b/Engine/Source/Programs/UnrealBuildAccelerator/Common/Private/UbaNetworkBackendTcp.cpp
--- a/Engine/Source/Programs/UnrealBuildAccelerator/Common/Private/UbaNetworkBackendTcp.cpp
+++ b/Engine/Source/Programs/UnrealBuildAccelerator/Common/Private/UbaNetworkBackendTcp.cpp
@@ -1151,9 +1151,11 @@ namespace uba
 		return true;
 	}
 
-	bool HttpConnection::Query(Logger& logger, const char* type, StringBufferBase& outResponse, u32& outStatusCode, const char* host, const char* path, const char* header)
+	bool HttpConnection::Query(Logger& logger, const char* type, StringBufferBase& outResponse, u32& outStatusCode, const char* host, const char* path, const char* header, u32 timeoutMs)
 	{
 		// TODO: Fix so we reuse socket connection for multiple queries
+		// Will need to change "Connection: close" and also know where end of message is
+
 		if (*m_host)// && _stricmp(m_host, host) != 0)
 		{
 			CloseSocket(logger, m_socket, TC("HttpQuery"));
@@ -1165,8 +1167,12 @@ namespace uba
 			if (!Connect(logger, host))
 				return false;
 
-		char request[512];
-		int requestLen = snprintf(request, 512, "%s /%s HTTP/1.1\r\nHost: %s\r\n%s\r\n", type, path, m_host, header);
+		char request[1024];
+		int requestLen = snprintf(request, sizeof_array(request), "%s /%s HTTP/1.1\r\nHost: %s\r\nUser-Agent: uba\r\nConnection: close\r\n%s\r\n", type, path, m_host, header);
+		if (requestLen >= sizeof_array(request))
+			return logger.Error(TC("STACK BUFFER TOO SMALL!"));
+
+		SetTimeout(logger, m_socket, timeoutMs);
 
 		int totalBytesSent = 0;
 		while (totalBytesSent < requestLen) {
@@ -1176,6 +1182,8 @@ namespace uba
 			totalBytesSent += bytesSent;
 		}
 
+		//logger.Warning(TC("REQUEST:\r\n%hs\r\n"), request);
+
 #if PLATFORM_WINDOWS
 #pragma warning(push)
 #pragma warning(disable:6386) // analyzer claims that buf can have buffer overrun.. but can't see how that can happen
@@ -1188,7 +1196,7 @@ namespace uba
 			readPos += bytesRead;
 
 		if (bytesRead == SOCKET_ERROR)
-			return logger.Error(TC("HttpConnection: socket recv error %s (%hs %hs)"), LastErrorToText(WSAGetLastError()).data, m_host, path);
+			return logger.Error(TC("HttpConnection: socket recv error after reading %u bytes - %s (%hs %hs)"), readPos, LastErrorToText(WSAGetLastError()).data, m_host, path);
 
 		if (readPos == sizeof(buf))
 			return logger.Error(TC("HttpConnection: buffer overflow"));
@@ -1199,18 +1207,26 @@ namespace uba
 #pragma warning(pop)
 #endif
 
+		//logger.Warning(TC("RESPONSE:\r\n%hs\r\n"), buf);
+
 		char* firstSpace = strchr(buf, ' '); // After version (where status code starts)
 		if (!firstSpace)
 			return logger.Error(TC("HttpConnection: first space not found (read %u)"), readPos);
+
 		char* secondSpace = strchr(firstSpace + 1, ' '); // after status code
 		if (!secondSpace)
 			return logger.Error(TC("HttpConnection: second space not found"));
+
 		*secondSpace = 0;
 		outStatusCode = strtoul(firstSpace + 1, nullptr, 10);
 
+		if (outStatusCode != 200)
+			return false;
+
 		char* bodyStart = strstr(secondSpace + 1, "\r\n\r\n");
 		if (!bodyStart)
 			return logger.Error(TC("HttpConnection: no body found"));
+
 		outResponse.Append(bodyStart + 4);
 		return true;
 	}
@@ -1219,4 +1235,4 @@ namespace uba
 	{
 		m_connectTimeOutMs = timeOutMs;
 	}
-}
\ No newline at end of file
+}
diff --git a/Engine/Source/Programs/UnrealBuildAccelerator/Common/Private/UbaNetworkClient.cpp b/Engine/Source/Programs/UnrealBuildAccelerator/Common/Private/UbaNetworkClient.cpp
--- a/Engine/Source/Programs/UnrealBuildAccelerator/Common/Private/UbaNetworkClient.cpp
+++ b/Engine/Source/Programs/UnrealBuildAccelerator/Common/Private/UbaNetworkClient.cpp
@@ -156,8 +156,8 @@ namespace uba
 		if (!backend.Send(m_logger, backendConnection, &m_uid, sizeof(m_uid), uidContext))
 			return false;
 
-		if (!rc.recvEvent.IsSet(~0u)) // This can not happen. Since both callbacks are using rc we can't leave this function until we know we are not in the callbacks
-			return m_logger.Error(TC("Timed out waiting for connection response from server"));
+		if (!rc.recvEvent.IsSet(5000)) // This can not happen. Since both callbacks are using rc we can't leave this function until we know we are not in the callbacks
+			return m_logger.Error(TC("Timed out waiting for connection response from server. Is this really an UbaNetworkServer we are connecting to?"));
 
 		m_isOrWasConnected.Set();
 
diff --git a/Engine/Source/Programs/UnrealBuildAccelerator/Common/Public/UbaAWS.h b/Engine/Source/Programs/UnrealBuildAccelerator/Common/Public/UbaAWS.h
--- a/Engine/Source/Programs/UnrealBuildAccelerator/Common/Public/UbaAWS.h
+++ b/Engine/Source/Programs/UnrealBuildAccelerator/Common/Public/UbaAWS.h
@@ -6,28 +6,106 @@
 #include "UbaNetworkBackend.h"
 #include "UbaStringBuffer.h"
 
-#define UBA_USE_AWS !PLATFORM_MAC
+#define UBA_USE_CLOUD !PLATFORM_MAC
 
-#if UBA_USE_AWS
+#if UBA_USE_CLOUD
 
 namespace uba
 {
-	class AWS
+	class Cloud
 	{
 	public:
+		struct Provider
+		{
+			const tchar* name;
+			const char* tokenCommand;
+			const char* tokenImds;
+			const char* tokenHeader;
+			const tchar* tokenPrefix;
+			const char* instanceIdImds;
+			const char* instanceLifeCycleImds;
+			const char* autoScalingLifeCycleStateImds;
+			const char* availabilityZoneImds;
+			const char* maintenanceEventImds;
+		};
+
+		static constexpr Provider g_providers[]
+		{
+			{
+				TC("AWS"),
+				"PUT",
+				"latest/api/token",
+				"X-aws-ec2-metadata-token-ttl-seconds: 21600\r\n",
+				TC("X-aws-ec2-metadata-token: "),
+				"latest/meta-data/instance-id",
+				"latest/meta-data/instance-life-cycle",
+				"latest/meta-data/autoscaling/target-lifecycle-state",
+				"latest/meta-data/placement/availability-zone",
+				nullptr, // maintenance
+			},
+			{
+				TC("GCP"),
+				"POST",
+				"computeMetadata/v1/instance/service-accounts/default/identity?audience=https://example.com",
+				"Accept: */*\r\nMetadata-Flavor: Google\r\nContent-Length: 0\r\n",
+				TC("Metadata-Flavor: Google\r\nAuthorization: Bearer "),
+				"computeMetadata/v1/instance/id",
+				nullptr, // Life cycle
+				nullptr, // Autoscaling
+				"computeMetadata/v1/instance/zone",
+				"computeMetadata/v1/instance/maintenance-event",
+			}
+		};
+
 
 		static constexpr char g_imdsHost[]						= "169.254.169.254";
-		static constexpr char g_imdsInstanceId[]				= "latest/meta-data/instance-id";
-		static constexpr char g_imdsInstanceLifeCycle[]			= "latest/meta-data/instance-life-cycle";
 		static constexpr char g_imdsAutoScalingLifeCycleState[] = "latest/meta-data/autoscaling/target-lifecycle-state";
-		static constexpr char g_imdsInstanceAvailabilityZone[]	= "latest/meta-data/placement/availability-zone";
 		static constexpr char g_imdsSpotInstanceAction[]		= "latest/meta-data/spot/instance-action";
-		static constexpr char g_imdsToken[]						= "latest/api/token";
 
+		bool QueryToken(Logger& logger, const tchar* rootDir)
+		{
+			for (auto& provider : g_providers)
+			{
+				if (IsNotCloud(logger, rootDir, provider.name))
+					continue;
+
+				HttpConnection http;
+				http.SetConnectTimeout(200);
+
+				u32 statusCode = 0;
+
+				//logger.Warning(TC("%hs TOKEN PUT START"), provider.name);
+				
+				StringBuffer<1024> token;
+				token.Append(provider.tokenPrefix);
+
+				if (!http.Query(logger, provider.tokenCommand, token, statusCode, g_imdsHost, provider.tokenImds, provider.tokenHeader, 250))
+				{
+					WriteIsNotCloud(logger, rootDir, provider.name);
+					continue;
+				}
+				token.Append("\r\n");
+
+				#if PLATFORM_WINDOWS
+				char tokenString[1024];
+				size_t tokenStringLen;
+				wcstombs_s(&tokenStringLen, tokenString, sizeof_array(tokenString), token.data, _TRUNCATE);
+				m_tokenString = tokenString;
+				#else
+				m_tokenString = token.data;
+				#endif
+
+				m_provider = &provider;
+
+				return true;
+			}
+
+			return false;
+		}
 
 		bool QueryInformation(Logger& logger, StringBufferBase& outExtraInfo, const tchar* rootDir)
 		{
-			if (IsNotAWS(logger, rootDir))
+			if (!QueryToken(logger, rootDir))
 				return false;
 
 			HttpConnection http;
@@ -35,50 +113,44 @@ namespace uba
 
 			u32 statusCode = 0;
 
-			StringBuffer<128> token;
-			token.Append(TCV("X-aws-ec2-metadata-token: "));
-			if (!http.Query(logger, "PUT", token, statusCode, g_imdsHost, g_imdsToken, "X-aws-ec2-metadata-token-ttl-seconds: 21600\r\n"))
-				return WriteIsNotAws(logger, rootDir);
-			token.Append("\r\n");
-
-			#if PLATFORM_WINDOWS
-			char tokenString[512];
-			size_t tokenStringLen;
-			wcstombs_s(&tokenStringLen, tokenString, sizeof_array(tokenString), token.data, _TRUNCATE);
-			m_tokenString = tokenString;
-			#else
-			m_tokenString = token.data;
-			#endif
-
-			StringBuffer<128> instanceId;
-			if (!http.Query(logger, "GET", instanceId, statusCode, g_imdsHost, g_imdsInstanceId, m_tokenString.c_str()))
+			//logger.Warning(TC("%hs TOKEN PUT START"), provider.name);
+				
+			StringBuffer<512> instanceId;
+			if (!http.Query(logger, "GET", instanceId, statusCode, g_imdsHost, m_provider->instanceIdImds, m_tokenString.c_str()))
 				return false;
 
-			outExtraInfo.Append(TCV(", AWS: ")).Append(instanceId);
+			outExtraInfo.Append(TCV(", ")).Append(m_provider->name).Append(TCV(": ")).Append(instanceId);
 
-			StringBuffer<32> instanceLifeCycle;
-			if (http.Query(logger, "GET", instanceLifeCycle, statusCode, g_imdsHost, g_imdsInstanceLifeCycle, m_tokenString.c_str()))
+			if (m_provider->instanceLifeCycleImds)
 			{
-				outExtraInfo.Append(' ').Append(instanceLifeCycle);
-				m_isSpot = instanceLifeCycle.Contains(TC("spot"));
+				StringBuffer<32> instanceLifeCycle;
+				if (http.Query(logger, "GET", instanceLifeCycle, statusCode, g_imdsHost, m_provider->instanceLifeCycleImds, m_tokenString.c_str()))
+				{
+					outExtraInfo.Append(' ').Append(instanceLifeCycle);
+					m_isSpot = instanceLifeCycle.Contains(TC("spot"));
+				}
 			}
 
-			StringBuffer<32> autoscaling;
-			if (http.Query(logger, "GET", autoscaling, statusCode, g_imdsHost, g_imdsAutoScalingLifeCycleState, m_tokenString.c_str()) && statusCode == 200)
+			if (m_provider->autoScalingLifeCycleStateImds)
 			{
-				outExtraInfo.Append(m_isSpot ? '/' : ' ').Append(TCV("autoscale"));
-				m_isAutoscaling = true;
+				StringBuffer<32> autoscaling;
+				if (http.Query(logger, "GET", autoscaling, statusCode, g_imdsHost, m_provider->autoScalingLifeCycleStateImds, m_tokenString.c_str()) && statusCode == 200)
+				{
+					outExtraInfo.Append(m_isSpot ? '/' : ' ').Append(TCV("autoscale"));
+					m_isAutoscaling = true;
+				}
 			}
 
 			if (!QueryAvailabilityZone(logger, nullptr))
 				return false;
-
 			return true;
 		}
 
 		bool QueryAvailabilityZone(Logger& logger, const tchar* rootDir)
 		{
-			if (rootDir && IsNotAWS(logger, rootDir))
+			if (rootDir && !QueryToken(logger, rootDir))
+				return false;
+			if (!m_provider->availabilityZoneImds)
 				return false;
 
 			HttpConnection http;
@@ -86,30 +158,31 @@ namespace uba
 
 			StringBuffer<128> availabilityZone;
 			u32 statusCode = 0;
-			if (!http.Query(logger, "GET", availabilityZone, statusCode, g_imdsHost, g_imdsInstanceAvailabilityZone, m_tokenString.c_str()))
-			{
-				if (rootDir)
-					WriteIsNotAws(logger, rootDir);
+			if (!http.Query(logger, "GET", availabilityZone, statusCode, g_imdsHost, m_provider->availabilityZoneImds, m_tokenString.c_str()))
 				return false;
-			}
-			m_availabilityZone = availabilityZone.data;
+
+			const tchar* zone = availabilityZone.data;
+			if (const tchar* lastSlash = TStrrchr(zone, '/'))
+				zone = lastSlash + 1;
+
+			m_availabilityZone = zone;
 			return true;
 		}
 
-		// Returns true if we _know_ we are not in AWS
-		bool IsNotAWS(Logger& logger, const tchar* rootDir)
+		// Returns true if we _know_ we are not in this cloud provider
+		bool IsNotCloud(Logger& logger, const tchar* rootDir, const tchar* provider)
 		{
 			StringBuffer<512> file;
-			file.Append(rootDir).EnsureEndsWithSlash().Append(".isNotAWS");
+			file.Append(rootDir).EnsureEndsWithSlash().Append(TCV(".isNot")).Append(provider);
 			if (FileExists(logger, file.data))
 				return true;
 			return false;
 		}
 
-		bool WriteIsNotAws(Logger& logger, const tchar* rootDir)
+		bool WriteIsNotCloud(Logger& logger, const tchar* rootDir, const tchar* provider)
 		{
 			StringBuffer<512> file;
-			file.Append(rootDir).EnsureEndsWithSlash().Append(".isNotAWS");
+			file.Append(rootDir).EnsureEndsWithSlash().Append(TCV(".isNot")).Append(provider);
 			FileAccessor f(logger, file.data);
 			if (!f.CreateWrite())
 				return false;
@@ -118,6 +191,13 @@ namespace uba
 			return false;
 		}
 
+		const tchar* GetProvider() const
+		{
+			if (!m_provider)
+				return TC("");
+			return m_provider->name;
+		}
+
 		bool IsTerminating(Logger& logger, StringBufferBase& outReason, u64& outTerminationTimeMs)
 		{
 			HttpConnection http;
@@ -153,6 +233,20 @@ namespace uba
 					}
 				}
 			}
+
+			if (m_provider && m_provider->maintenanceEventImds)
+			{
+				StringBuffer<1024> content;
+				u32 statusCode = 0;
+				if (http.Query(logger, "GET", content, statusCode, g_imdsHost, m_provider->maintenanceEventImds, m_tokenString.c_str()) && statusCode == 200)
+				{
+					if (!content.IsEmpty() && !content.Equals(TCV("NONE")))
+					{
+						outReason.Append(TCV("Google cloud instance interruption (")).Append(content).Append(')');
+						return true;
+					}
+				}
+			}
 			return false;
 		}
 
@@ -163,6 +257,7 @@ namespace uba
 
 		TString m_availabilityZone;
 		std::string m_tokenString;
+		const Provider* m_provider = nullptr;
 		bool m_isSpot = false;
 		bool m_isAutoscaling = false;
 	};
diff --git a/Engine/Source/Programs/UnrealBuildAccelerator/Common/Public/UbaNetworkBackend.h b/Engine/Source/Programs/UnrealBuildAccelerator/Common/Public/UbaNetworkBackend.h
--- a/Engine/Source/Programs/UnrealBuildAccelerator/Common/Public/UbaNetworkBackend.h
+++ b/Engine/Source/Programs/UnrealBuildAccelerator/Common/Public/UbaNetworkBackend.h
@@ -85,7 +85,7 @@ namespace uba
 	public:
 		HttpConnection();
 		~HttpConnection();
-		bool Query(Logger& logger, const char* type, StringBufferBase& outResponse, u32& outStatusCode, const char* host, const char* path, const char* header = "");
+		bool Query(Logger& logger, const char* type, StringBufferBase& outResponse, u32& outStatusCode, const char* host, const char* path, const char* header = "", u32 timeoutMs = 3000);
 
 		void SetConnectTimeout(u32 timeOutMs);
 
diff --git a/Engine/Source/Programs/UnrealBuildAccelerator/Host/Private/UbaExports.cpp b/Engine/Source/Programs/UnrealBuildAccelerator/Host/Private/UbaExports.cpp
--- a/Engine/Source/Programs/UnrealBuildAccelerator/Host/Private/UbaExports.cpp
+++ b/Engine/Source/Programs/UnrealBuildAccelerator/Host/Private/UbaExports.cpp
@@ -284,17 +284,17 @@ extern "C"
 	{
 		using namespace uba;
 		#ifndef __clang_analyzer__
-		#if UBA_USE_AWS
+		#if UBA_USE_CLOUD
 		StringBuffer<> fixedRootDir;
 		fixedRootDir.count = GetFullPathNameW(info.rootDir, fixedRootDir.capacity, fixedRootDir.data, NULL);
 		fixedRootDir.Replace('/', PathSeparator).EnsureEndsWithSlash();
 		info.rootDir = fixedRootDir.data;
-		AWS aws;
+		Cloud cloud;
 		if (!info.zone || !*info.zone)
 		{
-			ExportsDowngradedLogger downgradedLogger(info.writer, TC("AWS"));
-			if (aws.QueryAvailabilityZone(downgradedLogger, info.rootDir))
-				info.zone = aws.GetAvailabilityZone();
+			ExportsDowngradedLogger downgradedLogger(info.writer, TC("Cloud"));
+			if (cloud.QueryAvailabilityZone(downgradedLogger, info.rootDir))
+				info.zone = cloud.GetAvailabilityZone();
 		}
 		#endif
 
