[UBA]
* WIP Added support for requesting crash dumps from all clients. Can be triggered with UbaAgent -host=<host> -command=crashdump

diff --git a/Engine/Source/Programs/UnrealBuildAccelerator/Common/Private/UbaNetworkClient.cpp b/Engine/Source/Programs/UnrealBuildAccelerator/Common/Private/UbaNetworkClient.cpp
--- a/Engine/Source/Programs/UnrealBuildAccelerator/Common/Private/UbaNetworkClient.cpp
+++ b/Engine/Source/Programs/UnrealBuildAccelerator/Common/Private/UbaNetworkClient.cpp
@@ -256,7 +256,7 @@ namespace uba
 
 		constexpr u32 ErrorSize = 0xffffff - ReceiveHeaderSize; // ReceiveHeaderSize is removed from size in server send
 
-		if (messageSize == ErrorSize)
+		if (messageSize == ErrorSize || messageSize == 0xffffff)
 		{
 			msg->m_error = 1;
 			msg->Done();
diff --git a/Engine/Source/Programs/UnrealBuildAccelerator/Common/Private/UbaSessionClient.cpp b/Engine/Source/Programs/UnrealBuildAccelerator/Common/Private/UbaSessionClient.cpp
--- a/Engine/Source/Programs/UnrealBuildAccelerator/Common/Private/UbaSessionClient.cpp
+++ b/Engine/Source/Programs/UnrealBuildAccelerator/Common/Private/UbaSessionClient.cpp
@@ -12,6 +12,7 @@
 
 #if PLATFORM_WINDOWS
 #include <winerror.h>
+#include <tlhelp32.h>
 #endif
 
 namespace uba
@@ -1604,6 +1605,48 @@ namespace uba
 
 		if (reader.ReadBool()) // abort
 			abort();
+
+		if (reader.ReadBool()) // crashdump
+		{
+			struct Callstack { Vector<u8> data; Vector<u32> threadIds; };
+			UnorderedMap<CasKey, Callstack> callstacks;
+			StackBinaryWriter<SendMaxSize> stackWriter;
+
+			TraverseAllThreads([&](u32 tid, void* context)
+				{
+					stackWriter.Reset();
+					u32* infoSize = (u32*)stackWriter.AllocWrite(sizeof(u32));
+					WriteCallstackInfo(stackWriter, 0, context);
+					*infoSize = u32(stackWriter.GetPosition() - 4);
+
+					CasKeyHasher hasher;
+					hasher.Update(stackWriter.GetData(), stackWriter.GetPosition());
+
+					auto insres = callstacks.try_emplace(ToCasKey(hasher, false));
+					Callstack& cs = insres.first->second;
+					cs.threadIds.push_back(tid);
+					if (!insres.second)
+						return;
+					cs.data.resize(stackWriter.GetPosition());
+					memcpy(cs.data.data(), stackWriter.GetData(), cs.data.size());
+				});
+
+			for (auto& kv : callstacks)
+			{
+				Callstack& cs = kv.second;
+				BinaryReader stackReader(cs.data.data(), 0, cs.data.size());
+				stackWriter.Reset();
+				GetSymbols(UBA_AGENT_EXECUTABLE, IsArmBinary, stackReader, stackWriter);
+				BinaryReader resultReader(stackWriter.GetData(), 0, stackWriter.GetPosition());
+
+				StringBuffer<1024> tids(TC("Thread Ids: "));
+				for (u32 tid : cs.threadIds)
+					if (tids.count < tids.capacity - 5)
+						tids.AppendValue(tid).Append(' ');
+				TString infoString = resultReader.ReadString();
+				m_logger.Info(TC("%s%s"), tids.data, infoString.c_str());
+			}
+		}
 	}
 
 	void SessionClient::SendNotification(const StringView& text)
@@ -2322,7 +2365,7 @@ namespace uba
 		BinaryReader responseReader(writer.GetData(), 0, writer.GetCapacityLeft());
 		if (!msg.Send(responseReader, m_stats.customMsg))
 			return false;
-		writer.AllocWrite(responseReader.GetPosition());
+		writer.AllocWrite(responseReader.GetLeft());
 		return true;
 
 #if 0
diff --git a/Engine/Source/Programs/UnrealBuildAccelerator/Common/Private/UbaSessionServer.cpp b/Engine/Source/Programs/UnrealBuildAccelerator/Common/Private/UbaSessionServer.cpp
--- a/Engine/Source/Programs/UnrealBuildAccelerator/Common/Private/UbaSessionServer.cpp
+++ b/Engine/Source/Programs/UnrealBuildAccelerator/Common/Private/UbaSessionServer.cpp
@@ -1585,6 +1585,8 @@ namespace uba
 		session.memTotal = memTotal;
 		session.cpuLoad = *(float*)&cpuLoadValue;
 		writer.WriteBool(session.abort);
+		writer.WriteBool(session.crashdump);
+		session.crashdump = false;
 
 		return true;
 	}
@@ -1703,6 +1705,13 @@ namespace uba
 			}
 			WriteString(StringBuffer<>().Appendf(TC("Total remote slots used %u/%u"), totalUsed, totalSlots).data);
 		}
+		if (command.Equals(TCV("crashdump")))
+		{
+			WriteString(TC("Requesting crashdumps from all remotes on next ping"));
+			SCOPED_CRITICAL_SECTION(m_remoteProcessAndSessionLock, queueLock);
+			for (auto& s : m_clientSessions)
+				s->crashdump = true;
+		}
 		else if (command.StartsWith(TC("abort")))
 		{
 			bool abortWithProxy = command.Equals(TCV("abortproxy"));
diff --git a/Engine/Source/Programs/UnrealBuildAccelerator/Common/Private/UbaThread.cpp b/Engine/Source/Programs/UnrealBuildAccelerator/Common/Private/UbaThread.cpp
--- a/Engine/Source/Programs/UnrealBuildAccelerator/Common/Private/UbaThread.cpp
+++ b/Engine/Source/Programs/UnrealBuildAccelerator/Common/Private/UbaThread.cpp
@@ -3,6 +3,17 @@
 #include "UbaThread.h"
 #include "UbaPlatform.h"
 
+#if PLATFORM_WINDOWS
+#include <tlhelp32.h>
+#else
+#include <dirent.h>
+#include <execinfo.h>
+#include <sys/ptrace.h>
+#include <sys/user.h>
+#include <sys/uio.h>
+#define REG_SP regs.rsp
+#endif
+
 namespace uba
 {
 	bool AlternateThreadGroupAffinity(void* nativeThreadHandle)
@@ -155,4 +166,85 @@ namespace uba
 #endif
 	}
 
+	bool TraverseAllThreads(const Function<void(u32 tid, void* context)>& func)
+	{
+#if PLATFORM_WINDOWS
+
+		DWORD pid = GetCurrentProcessId();
+		DWORD tid = GetCurrentThreadId();
+		HANDLE hSnapshot = CreateToolhelp32Snapshot(TH32CS_SNAPTHREAD, 0);
+		if (hSnapshot == INVALID_HANDLE_VALUE)
+			return false;
+		auto sg = MakeGuard([&]() { CloseHandle(hSnapshot); });
+
+		UnorderedSet<CasKey> handledCallstacks;
+		THREADENTRY32 te32 = { sizeof(THREADENTRY32) };
+		if (!Thread32First(hSnapshot, &te32))
+			return false;
+
+		do
+		{
+			if (te32.th32OwnerProcessID != pid || te32.th32ThreadID == tid)
+				continue;
+
+			HANDLE hThread = OpenThread(THREAD_GET_CONTEXT | THREAD_SET_CONTEXT | THREAD_SUSPEND_RESUME | THREAD_QUERY_INFORMATION, FALSE, te32.th32ThreadID);
+			if (!hThread)
+				continue;
+			auto tg = MakeGuard([&]() { CloseHandle(hThread); });
+
+			if (SuspendThread(hThread) == -1)
+				continue;
+
+			CONTEXT ctx;
+			memset(&ctx, 0, sizeof(CONTEXT));
+			ctx.ContextFlags = CONTEXT_FULL;
+			GetThreadContext(hThread, &ctx);
+
+			func(te32.th32ThreadID, &ctx);
+
+			ResumeThread(hThread);
+		}
+		while (Thread32Next(hSnapshot, &te32));
+
+		return true;
+#else
+		pid_t currentTid = syscall(SYS_gettid);
+
+		DIR* dir = opendir("/proc/self/task");
+		if (!dir)
+			return false;
+
+		Vector<pid_t> tids;
+
+		struct dirent *entry;
+		while ((entry = readdir(dir)) != NULL)
+		{
+			char *endptr;
+			pid_t tid = (pid_t) strtol(entry->d_name, &endptr, 10);
+			if (*endptr == '\0' && errno != ERANGE && tid != currentTid)
+				tids.push_back(tid);
+
+		}
+		closedir(dir);
+
+		static Event s_ev(false);
+		static auto s_func = &func;
+
+		signal(SIGUSR1, [](int sig)
+			{
+				pid_t currentTid = syscall(SYS_gettid);
+				(*s_func)(currentTid, nullptr);
+				s_ev.Set();
+			});
+
+
+		for (pid_t tid : tids)
+		{
+			kill(tid, SIGUSR1);
+			s_ev.IsSet();
+		}
+
+		return true;
+#endif
+	}
 }
diff --git a/Engine/Source/Programs/UnrealBuildAccelerator/Common/Public/UbaSessionServer.h b/Engine/Source/Programs/UnrealBuildAccelerator/Common/Public/UbaSessionServer.h
--- a/Engine/Source/Programs/UnrealBuildAccelerator/Common/Public/UbaSessionServer.h
+++ b/Engine/Source/Programs/UnrealBuildAccelerator/Common/Public/UbaSessionServer.h
@@ -156,6 +156,7 @@ namespace uba
 			bool enabled = true;
 			bool dedicated = false;
 			bool abort = false;
+			bool crashdump = false;
 			bool hasNotification = false;
 			bool isArm = false;
 		};
diff --git a/Engine/Source/Programs/UnrealBuildAccelerator/Common/Public/UbaThread.h b/Engine/Source/Programs/UnrealBuildAccelerator/Common/Public/UbaThread.h
--- a/Engine/Source/Programs/UnrealBuildAccelerator/Common/Public/UbaThread.h
+++ b/Engine/Source/Programs/UnrealBuildAccelerator/Common/Public/UbaThread.h
@@ -42,4 +42,5 @@ namespace uba
 	};
 	bool SetThreadGroupAffinity(void* nativeThreadHandle, const GroupAffinity& affinity);
 	bool AlternateThreadGroupAffinity(void* nativeThreadHandle);
+	bool TraverseAllThreads(const Function<void(u32 tid, void* context)>& func);
 }
diff --git a/Engine/Source/Programs/UnrealBuildAccelerator/Test/Private/UbaTestAll.h b/Engine/Source/Programs/UnrealBuildAccelerator/Test/Private/UbaTestAll.h
--- a/Engine/Source/Programs/UnrealBuildAccelerator/Test/Private/UbaTestAll.h
+++ b/Engine/Source/Programs/UnrealBuildAccelerator/Test/Private/UbaTestAll.h
@@ -50,6 +50,7 @@ namespace uba
 		UBA_TEST(TestDetouredClang) \
 		UBA_TEST(TestFileMappingBuffer) \
 		UBA_TEST(TestRemoteDirectoryTable) \
+		UBA_TEST(TestThreads) \
 		UBA_TEST_EXTRA \
 
 #define UBA_POSIX_TESTS \
diff --git a/Engine/Source/Programs/UnrealBuildAccelerator/Test/Private/UbaTestBasics.cpp b/Engine/Source/Programs/UnrealBuildAccelerator/Test/Private/UbaTestBasics.cpp
--- a/Engine/Source/Programs/UnrealBuildAccelerator/Test/Private/UbaTestBasics.cpp
+++ b/Engine/Source/Programs/UnrealBuildAccelerator/Test/Private/UbaTestBasics.cpp
@@ -554,5 +554,49 @@ namespace uba
 		return true;
 	}
 
+	UBA_NOINLINE void TestFunctionForThread(Event& ev1, Event& ev2)
+	{
+		ev1.Set();
+		ev2.IsSet();
+	}
+
+	bool TestThreads(LoggerWithWriter& logger, const StringBufferBase& rootDir)
+	{
+		Event ev1(true);
+		Event ev2(true);
+		Thread t([&]()
+			{
+				TestFunctionForThread(ev1, ev2);
+				return true;
+			});
 
+		ev1.IsSet();
+
+		TraverseAllThreads([&](u32 tid, void* context)
+			{
+				static u8* writerMem = new u8[4096];
+				BinaryWriter writer(writerMem, 0, 4096);
+				WriteCallstackInfo(writer, 0, context);
+				BinaryReader reader(writerMem, 0, writer.GetPosition());
+				StringBuffer<16*1024> sb;
+				tchar executable[512] = TC("UbaTest");//{ 0 };
+				
+				StringView searchPaths[3];
+				StringBuffer<512> currentModuleDir;
+				LoggerWithWriter logger(g_nullLogWriter);
+				GetDirectoryOfCurrentModule(logger, currentModuleDir);
+				StringBuffer<512> alternativePath;
+				u32 searchPathIndex = 0;
+				if (GetAlternativeUbaPath(logger, alternativePath, currentModuleDir, IsWindows && IsArmBinary))
+					searchPaths[searchPathIndex++] = alternativePath;
+				searchPaths[searchPathIndex] = currentModuleDir;
+
+				ParseCallstackInfo(sb, reader, executable, searchPaths);
+				LoggerWithWriter(g_consoleLogWriter, TC("")).Info(TC("THREAD %u%s"), tid, sb.data);
+			});
+
+
+		ev2.Set();
+		return true;
+	}
 }
