[UBA]
* Another attempt at TraverseAllThreads but now by explicitly tracking them because MacOs has no way of traversing threads and get the thread ids needed to send signals

diff --git a/Engine/Source/Programs/UnrealBuildAccelerator/Common/Private/UbaThread.cpp b/Engine/Source/Programs/UnrealBuildAccelerator/Common/Private/UbaThread.cpp
--- a/Engine/Source/Programs/UnrealBuildAccelerator/Common/Private/UbaThread.cpp
+++ b/Engine/Source/Programs/UnrealBuildAccelerator/Common/Private/UbaThread.cpp
@@ -50,6 +50,8 @@ namespace uba
 #endif
 	}
 
+	Futex g_allThreadsLock;
+	Thread* g_firstThread;
 
 	Thread::Thread()
 	{
@@ -106,6 +108,10 @@ namespace uba
 		err = pthread_attr_destroy(&tattr);
 		UBA_ASSERT(err == 0); (void)err;
 #endif
+
+		SCOPED_FUTEX(g_allThreadsLock, lock);
+		m_next = g_firstThread;
+		g_firstThread = this;
 	}
 
 	bool Thread::Wait(u32 milliseconds, Event* wakeupEvent)
@@ -134,6 +140,17 @@ namespace uba
 		int res = pthread_join(*(pthread_t*)&m_handle, (void**)&ptr);
 		UBA_ASSERT(res == 0);
 #endif
+		{
+			SCOPED_FUTEX(g_allThreadsLock, lock);
+			if (m_next)
+				m_next->m_prev = m_prev;
+			if (m_prev)
+				m_prev->m_next = m_next;
+			else
+				g_firstThread = m_next;
+			m_next = nullptr;
+			m_prev = nullptr;
+		}
 
 		readLock.Leave();
 
@@ -169,7 +186,28 @@ namespace uba
 	bool TraverseAllThreads(const Function<void(u32 tid, void* context)>& func)
 	{
 #if PLATFORM_WINDOWS
+		DWORD currentThreadId = GetCurrentThreadId();
+		SCOPED_FUTEX(g_allThreadsLock, lock);
+		for (Thread* t=g_firstThread; t; t=t->m_next)
+		{
+			HANDLE hThread = t->m_handle;
+			if (currentThreadId == GetThreadId(hThread))
+				continue;
+
+			if (SuspendThread(hThread) == -1)
+				continue;
+
+			CONTEXT ctx;
+			memset(&ctx, 0, sizeof(CONTEXT));
+			ctx.ContextFlags = CONTEXT_FULL;
+			GetThreadContext(hThread, &ctx);
+
+			func(0, &ctx);
+
+			ResumeThread(hThread);
+		}
 
+		/*
 		DWORD pid = GetCurrentProcessId();
 		DWORD tid = GetCurrentThreadId();
 		HANDLE hSnapshot = CreateToolhelp32Snapshot(TH32CS_SNAPTHREAD, 0);
@@ -205,9 +243,31 @@ namespace uba
 			ResumeThread(hThread);
 		}
 		while (Thread32Next(hSnapshot, &te32));
-
+		*/
 		return true;
 #else
+		static Event s_ev(false);
+		static auto s_func = &func;
+
+		signal(SIGUSR1, [](int sig)
+			{
+				pid_t currentTid = syscall(SYS_gettid);
+				(*s_func)(currentTid, nullptr);
+				s_ev.Set();
+			});
+
+		pthread_t currentTid = pthread_self();
+		SCOPED_FUTEX(g_allThreadsLock, lock);
+		for (Thread* t=g_firstThread; t; t=t->m_next)
+		{
+			pthread_t tid = *(pthread_t*)&t->m_handle;
+			if (currentTid == tid)
+				continue;
+
+			pthread_kill(tid, SIGUSR1);
+			s_ev.IsSet();
+		}
+/*
 		pid_t currentTid = syscall(SYS_gettid);
 
 		DIR* dir = opendir("/proc/self/task");
@@ -243,7 +303,7 @@ namespace uba
 			kill(tid, SIGUSR1);
 			s_ev.IsSet();
 		}
-
+*/
 		return true;
 #endif
 	}
diff --git a/Engine/Source/Programs/UnrealBuildAccelerator/Common/Public/UbaThread.h b/Engine/Source/Programs/UnrealBuildAccelerator/Common/Public/UbaThread.h
--- a/Engine/Source/Programs/UnrealBuildAccelerator/Common/Public/UbaThread.h
+++ b/Engine/Source/Programs/UnrealBuildAccelerator/Common/Public/UbaThread.h
@@ -30,8 +30,13 @@ namespace uba
 		Event m_finished;
 		#endif
 
+		Thread* m_next = nullptr;
+		Thread* m_prev = nullptr;
+
 		Thread(const Thread&) = delete;
 		void operator=(const Thread&) = delete;
+
+		friend bool TraverseAllThreads(const Function<void(u32 tid, void* context)>& func);
 	};
 
 
