[UBA]
* Added define to disable thread tracking and disabled it

diff --git a/Engine/Source/Programs/UnrealBuildAccelerator/Common/Private/UbaThread.cpp b/Engine/Source/Programs/UnrealBuildAccelerator/Common/Private/UbaThread.cpp
--- a/Engine/Source/Programs/UnrealBuildAccelerator/Common/Private/UbaThread.cpp
+++ b/Engine/Source/Programs/UnrealBuildAccelerator/Common/Private/UbaThread.cpp
@@ -14,6 +14,8 @@
 #define REG_SP regs.rsp
 #endif
 
+#define UBA_TRACK_THREADS 0
+
 namespace uba
 {
 	bool AlternateThreadGroupAffinity(void* nativeThreadHandle)
@@ -109,9 +111,14 @@ namespace uba
 		UBA_ASSERT(err == 0); (void)err;
 #endif
 
+		#if UBA_TRACK_THREADS
 		SCOPED_FUTEX(g_allThreadsLock, lock);
+		printf("Adding THREAD %llx\n", (u64)*(pthread_t*)&m_handle);
 		m_next = g_firstThread;
+		if (m_next)
+			m_next->m_prev = this;
 		g_firstThread = this;
+		#endif
 	}
 
 	bool Thread::Wait(u32 milliseconds, Event* wakeupEvent)
@@ -120,6 +127,22 @@ namespace uba
 		if (!m_handle)
 			return true;
 
+		auto removeThread = [this]()
+			{
+				#if UBA_TRACK_THREADS
+				SCOPED_FUTEX(g_allThreadsLock, lock);
+				printf("REMOVING THREAD %llx\n", (u64)*(pthread_t*)&m_handle);
+				if (m_next)
+					m_next->m_prev = m_prev;
+				if (m_prev)
+					m_prev->m_next = m_next;
+				else if (g_firstThread == this)
+					g_firstThread = m_next;
+				m_next = nullptr;
+				m_prev = nullptr;
+				#endif
+			};
+
 #if PLATFORM_WINDOWS // Optimization, not needed in initial implementation
 		if (wakeupEvent)
 		{
@@ -133,24 +156,15 @@ namespace uba
 			if (WaitForSingleObject(m_handle, milliseconds) == WAIT_TIMEOUT)
 				return false;
 		}
+		removeThread();
 #else
 		if (!m_finished.IsSet(milliseconds))
 			return false;
+		removeThread();
 		int* ptr = 0;
 		int res = pthread_join(*(pthread_t*)&m_handle, (void**)&ptr);
 		UBA_ASSERT(res == 0);
 #endif
-		{
-			SCOPED_FUTEX(g_allThreadsLock, lock);
-			if (m_next)
-				m_next->m_prev = m_prev;
-			if (m_prev)
-				m_prev->m_next = m_next;
-			else
-				g_firstThread = m_next;
-			m_next = nullptr;
-			m_prev = nullptr;
-		}
 
 		readLock.Leave();
 
@@ -183,9 +197,11 @@ namespace uba
 #endif
 	}
 
-	bool TraverseAllThreads(const Function<void(u32 tid, void* context)>& func)
+	bool TraverseAllThreads(const TraverseThreadFunc& func)
 	{
 #if PLATFORM_WINDOWS
+
+#if UBA_TRACK_THREADS
 		DWORD currentThreadId = GetCurrentThreadId();
 		SCOPED_FUTEX(g_allThreadsLock, lock);
 		for (Thread* t=g_firstThread; t; t=t->m_next)
@@ -206,8 +222,7 @@ namespace uba
 
 			ResumeThread(hThread);
 		}
-
-		/*
+#else
 		DWORD pid = GetCurrentProcessId();
 		DWORD tid = GetCurrentThreadId();
 		HANDLE hSnapshot = CreateToolhelp32Snapshot(TH32CS_SNAPTHREAD, 0);
@@ -238,35 +253,60 @@ namespace uba
 			ctx.ContextFlags = CONTEXT_FULL;
 			GetThreadContext(hThread, &ctx);
 
-			func(te32.th32ThreadID, &ctx);
+			func(te32.th32ThreadID, &ctx, 0);
 
 			ResumeThread(hThread);
 		}
 		while (Thread32Next(hSnapshot, &te32));
-		*/
+#endif
 		return true;
 #else
+#if 0
 		static Event s_ev(false);
-		static auto s_func = &func;
+		static const TraverseThreadFunc* s_func;
+		static u32 s_skipCallstack;
 
-		signal(SIGUSR1, [](int sig)
+		s_func = &func;
+
+		#if PLATFORM_MAC
+		void* bt[100];
+		s_skipCallstack = backtrace(bt, 100) + 4;
+		//printf("SKIP: %u\n", s_skipCallstack);
+		#endif
+
+		struct sigaction sa;
+		memset(&sa, 0, sizeof(sa));
+		sa.sa_flags = SA_SIGINFO;
+		sa.sa_sigaction = [](int sig, siginfo_t* info, void* context)
 			{
-				pid_t currentTid = syscall(SYS_gettid);
-				(*s_func)(currentTid, nullptr);
+				printf("IN THREAD %llx\n", (u64)pthread_self());
+				//pid_t currentTid = 0;//syscall(SYS_gettid);
+				//(*s_func)(currentTid, nullptr, s_skipCallstack);
 				s_ev.Set();
-			});
+			};
+		sigaction(SIGUSR1, &sa, NULL);
+
 
 		pthread_t currentTid = pthread_self();
 		SCOPED_FUTEX(g_allThreadsLock, lock);
 		for (Thread* t=g_firstThread; t; t=t->m_next)
 		{
-			pthread_t tid = *(pthread_t*)&t->m_handle;
+			pthread_t tid = *(pthread_t*)&(t->m_handle);
 			if (currentTid == tid)
 				continue;
-
+			printf("HANDLING THREAD %llx\n", (u64)tid);
 			pthread_kill(tid, SIGUSR1);
 			s_ev.IsSet();
 		}
+		lock.Leave();
+/*		
+		for (auto tid : threads)
+		{
+			printf("HANDLING THREAD\n");
+			pthread_kill(tid, SIGUSR1);
+			s_ev.IsSet();
+		}
+*/
 /*
 		pid_t currentTid = syscall(SYS_gettid);
 
@@ -304,6 +344,7 @@ namespace uba
 			s_ev.IsSet();
 		}
 */
+#endif
 		return true;
 #endif
 	}
diff --git a/Engine/Source/Programs/UnrealBuildAccelerator/Common/Public/UbaThread.h b/Engine/Source/Programs/UnrealBuildAccelerator/Common/Public/UbaThread.h
--- a/Engine/Source/Programs/UnrealBuildAccelerator/Common/Public/UbaThread.h
+++ b/Engine/Source/Programs/UnrealBuildAccelerator/Common/Public/UbaThread.h
@@ -5,11 +5,13 @@
 #include "UbaBinaryReaderWriter.h"
 #include "UbaEvent.h"
 
+#define UBA_TRACK_THREADS 0
+
 namespace uba
 {
 	class Event;
 	struct GroupAffinity;
-
+	using TraverseThreadFunc = Function<void(u32 tid, void* context, u32 skipCallstack)>;
 
 	class Thread
 	{
@@ -30,13 +32,15 @@ namespace uba
 		Event m_finished;
 		#endif
 
+		#if UBA_TRACK_THREADS
 		Thread* m_next = nullptr;
 		Thread* m_prev = nullptr;
+		#endif
 
 		Thread(const Thread&) = delete;
 		void operator=(const Thread&) = delete;
 
-		friend bool TraverseAllThreads(const Function<void(u32 tid, void* context)>& func);
+		friend bool TraverseAllThreads(const TraverseThreadFunc& func);
 	};
 
 
@@ -47,5 +51,5 @@ namespace uba
 	};
 	bool SetThreadGroupAffinity(void* nativeThreadHandle, const GroupAffinity& affinity);
 	bool AlternateThreadGroupAffinity(void* nativeThreadHandle);
-	bool TraverseAllThreads(const Function<void(u32 tid, void* context)>& func);
+	bool TraverseAllThreads(const TraverseThreadFunc& func);
 }
