[UBA]
* Added instr-clang-cl.exe to application rules

diff --git a/Engine/Source/Programs/UnrealBuildAccelerator/Core/Private/UbaApplicationRules.cpp b/Engine/Source/Programs/UnrealBuildAccelerator/Core/Private/UbaApplicationRules.cpp
--- a/Engine/Source/Programs/UnrealBuildAccelerator/Core/Private/UbaApplicationRules.cpp
+++ b/Engine/Source/Programs/UnrealBuildAccelerator/Core/Private/UbaApplicationRules.cpp
@@ -749,6 +749,7 @@ namespace uba
 			{ ApplicationHash(TC("pvs-studio.exe")),			new ApplicationRulesPVSStudio() },
 			{ ApplicationHash(TC("ubaobjtool.exe")),			new ApplicationRulesUbaObjTool() },
 			{ ApplicationHash(TC("shadercompileworker.exe")),	new ApplicationRulesShaderCompileWorker() },
+			{ ApplicationHash(TC("instr-clang-cl.exe")),		new ApplicationRulesClangClExe() },
 			//{ L"MSBuild.dll"),				new ApplicationRules() },
 			//{ L"BreakpadSymbolEncoder.exe"),	new ApplicationRulesClang() },
 			//{ L"cmd.exe"),		new ApplicationRules() },
