[UBA]
* Split up WriteCallstackInfo in two functions
* Fixed compile error

diff --git a/Engine/Source/Programs/UnrealBuildAccelerator/Common/Private/UbaSessionClient.cpp b/Engine/Source/Programs/UnrealBuildAccelerator/Common/Private/UbaSessionClient.cpp
--- a/Engine/Source/Programs/UnrealBuildAccelerator/Common/Private/UbaSessionClient.cpp
+++ b/Engine/Source/Programs/UnrealBuildAccelerator/Common/Private/UbaSessionClient.cpp
@@ -1612,11 +1612,11 @@ namespace uba
 			UnorderedMap<CasKey, Callstack> callstacks;
 			StackBinaryWriter<SendMaxSize> stackWriter;
 
-			TraverseAllThreads([&](u32 tid, void* context)
+			TraverseAllThreads([&](u32 tid, void* context, u32 skipCallstack)
 				{
 					stackWriter.Reset();
 					u32* infoSize = (u32*)stackWriter.AllocWrite(sizeof(u32));
-					WriteCallstackInfo(stackWriter, 0, context);
+					WriteCallstackInfo(stackWriter, skipCallstack, context);
 					*infoSize = u32(stackWriter.GetPosition() - 4);
 
 					CasKeyHasher hasher;
diff --git a/Engine/Source/Programs/UnrealBuildAccelerator/Core/Private/UbaPlatform.cpp b/Engine/Source/Programs/UnrealBuildAccelerator/Core/Private/UbaPlatform.cpp
--- a/Engine/Source/Programs/UnrealBuildAccelerator/Core/Private/UbaPlatform.cpp
+++ b/Engine/Source/Programs/UnrealBuildAccelerator/Core/Private/UbaPlatform.cpp
@@ -195,8 +195,6 @@ namespace uba
 	UBA_NOINLINE bool WriteCallstackInfo(BinaryWriter& out, u32 skipCallstack, void* contextPtr)
 	{
 #if PLATFORM_WINDOWS
-		out.WriteBool(IsRunningWine());
-
 		const int MaxDepth = 16;// 62;
 		void* callers[MaxDepth];
 		u32 count = 0;//RtlCaptureStackBackTrace(0, MaxDepth, callers, NULL); // For some reason this does not work on wine
@@ -236,6 +234,33 @@ namespace uba
 				}
 			}
 		}
+		if (count < skipCallstack)
+			skipCallstack = count;
+
+		return WriteCallstackInfo(out, callers + skipCallstack, count - skipCallstack);
+#else
+		constexpr u32 maxCallers = 100;
+		void* callers[maxCallers];
+		int nptrs = backtrace(callers, maxCallers);
+
+		if (contextPtr) // Use this to know if we come from a signal
+			skipCallstack += 2;
+
+		if (nptrs <= skipCallstack)
+		{
+			if (!nptrs)
+				skipCallstack = 0;
+			else
+				skipCallstack = nptrs - 1;
+		}
+		return WriteCallstackInfo(out, callers + skipCallstack, nptrs - skipCallstack);
+#endif
+	}
+
+	bool WriteCallstackInfo(BinaryWriter& out, void** callstack, u32 callstackCount)
+	{
+#if PLATFORM_WINDOWS
+		out.WriteBool(IsRunningWine());
 
 		struct ModuleRec { u64 start; u64 size; HMODULE handle; u32 index; };
 		Map<u64, ModuleRec> moduleEndAddresses;
@@ -258,11 +283,11 @@ namespace uba
 
 		Vector<ModuleRec*> usedModules;
 
-		out.Write7BitEncoded(count < skipCallstack ? 0 : count - skipCallstack);
+		out.Write7BitEncoded(callstackCount);
 
-		for (u32 i = skipCallstack; i < count; i++)
+		for (u32 i = 0; i < callstackCount; i++)
 		{
-			auto addr = u64(callers[i]);
+			auto addr = u64(callstack[i]);
 
 			auto findIt = moduleEndAddresses.lower_bound(addr);
 			if (findIt != moduleEndAddresses.end() && addr >= findIt->second.start)
@@ -307,7 +332,7 @@ namespace uba
 		struct ModuleRec { u64 start; u64 size; TString name; u32 index; };
 		Map<u64, ModuleRec> moduleEndAddresses;
 
-#if PLATFORM_LINUX
+		#if PLATFORM_LINUX
 		dl_iterate_phdr([](struct dl_phdr_info* info, size_t size, void* data)
 			{
 				auto& moduleEndAddresses = *(Map<u64, ModuleRec>*)data;
@@ -330,7 +355,7 @@ namespace uba
 				rec.index = ~0u;
 				return 0;
 			}, &moduleEndAddresses);
-#else
+		#else
 		{
 			Map<u64, ModuleRec> modules;
 			for (u32 i = 0, e = _dyld_image_count(); i < e; i++)
@@ -358,30 +383,15 @@ namespace uba
 				prev = &kv.second;
 			}
 		}
-#endif
+		#endif
 
-		constexpr u32 maxCallers = 100;
-		void* buffer[maxCallers];
-		int nptrs = backtrace(buffer, maxCallers);
+		out.Write7BitEncoded(callstackCount);
 
 		Vector<ModuleRec*> usedModules;
 
-		if (contextPtr) // Use this to know if we come from a signal
-			skipCallstack += 2;
-
-		if (nptrs <= skipCallstack)
-		{
-			if (!nptrs)
-				skipCallstack = 0;
-			else
-				skipCallstack = nptrs - 1;
-		}
-
-		out.Write7BitEncoded(nptrs - skipCallstack);
-
-		for (int i = skipCallstack; i < nptrs; ++i)
+		for (int i = 0; i < callstackCount; ++i)
 		{
-			u64 addr = u64(buffer[i]);
+			u64 addr = u64(callstack[i]);
 			auto findIt = moduleEndAddresses.lower_bound(addr);
 			if (findIt != moduleEndAddresses.end() && addr >= findIt->second.start)
 			{
diff --git a/Engine/Source/Programs/UnrealBuildAccelerator/Core/Public/UbaPlatform.h b/Engine/Source/Programs/UnrealBuildAccelerator/Core/Public/UbaPlatform.h
--- a/Engine/Source/Programs/UnrealBuildAccelerator/Core/Public/UbaPlatform.h
+++ b/Engine/Source/Programs/UnrealBuildAccelerator/Core/Public/UbaPlatform.h
@@ -81,6 +81,7 @@ namespace uba
 	#define UBA_NOT_IMPLEMENTED(x) UBA_ASSERTF(false, TC("%s not implemented!"), #x);
 
 	bool WriteCallstackInfo(BinaryWriter& out, u32 skipCallstack, void* context);
+	bool WriteCallstackInfo(BinaryWriter& out, void** callstack, u32 callstackCount);
 	void WriteAssertInfo(StringBufferBase& out, const tchar* text, const char* file, u32 line, const char* expr, void* context = nullptr);
 	ANALYSIS_NORETURN void UbaAssert(const tchar* text, const char* file, u32 line, const char* expr, bool allowTerminate, u32 terminateCode, void* context, u32 skipCallstackCount);
 	ANALYSIS_NORETURN void FatalError(u32 code, const tchar* format, ...);
