[UBA]
* Compile error fix in unit tests

diff --git a/Engine/Source/Programs/UnrealBuildAccelerator/Test/Private/UbaTestBasics.cpp b/Engine/Source/Programs/UnrealBuildAccelerator/Test/Private/UbaTestBasics.cpp
--- a/Engine/Source/Programs/UnrealBuildAccelerator/Test/Private/UbaTestBasics.cpp
+++ b/Engine/Source/Programs/UnrealBuildAccelerator/Test/Private/UbaTestBasics.cpp
@@ -572,11 +572,11 @@ namespace uba
 
 		ev1.IsSet();
 
-		TraverseAllThreads([&](u32 tid, void* context)
+		TraverseAllThreads([&](u32 tid, void* context, u32 skipCallstack)
 			{
 				static u8* writerMem = new u8[4096];
 				BinaryWriter writer(writerMem, 0, 4096);
-				WriteCallstackInfo(writer, 0, context);
+				WriteCallstackInfo(writer, skipCallstack, context);
 				BinaryReader reader(writerMem, 0, writer.GetPosition());
 				StringBuffer<16*1024> sb;
 				tchar executable[512] = TC("UbaTest");//{ 0 };
