[UBA]
* Fixed so callstacks created in wine contains file and line. (Wine needs explicit call to SymSetOptions(SYMOPT_LOAD_LINES) for this to work)

diff --git a/Engine/Source/Programs/UnrealBuildAccelerator/Common/Private/UbaLogger.cpp b/Engine/Source/Programs/UnrealBuildAccelerator/Common/Private/UbaLogger.cpp
--- a/Engine/Source/Programs/UnrealBuildAccelerator/Common/Private/UbaLogger.cpp
+++ b/Engine/Source/Programs/UnrealBuildAccelerator/Common/Private/UbaLogger.cpp
@@ -544,6 +544,8 @@ namespace uba
 
 		if (SymInitializeW(processHandle, searchPathString.data, FALSE))
 		{
+			SymSetOptions(SYMOPT_LOAD_LINES);
+
 			for (auto& mod : modules)
 			{
 				if (mod.name.empty())
