[UBA]
* Another iteration of resolving callstacks on all threads in application. This implementation seems to work on all platforms but will most likely cause a crash on linux since it is using signal handlers

diff --git a/Engine/Source/Programs/UnrealBuildAccelerator/Common/Private/UbaLogger.cpp b/Engine/Source/Programs/UnrealBuildAccelerator/Common/Private/UbaLogger.cpp
--- a/Engine/Source/Programs/UnrealBuildAccelerator/Common/Private/UbaLogger.cpp
+++ b/Engine/Source/Programs/UnrealBuildAccelerator/Common/Private/UbaLogger.cpp
@@ -511,7 +511,7 @@ namespace uba
 			entry.memoryOffset = reader.Read7BitEncoded();
 		}
 
-		struct ModuleEntry { u64 start; u64 size; TString name; bool handled; };
+		struct ModuleEntry { u64 start; u64 size; TString name; bool handled; UnorderedMap<u64, TString> symbols; };
 		u64 moduleCount = reader.Read7BitEncoded();
 		Vector<ModuleEntry> modules(moduleCount);
 		for (auto& mod : modules)
@@ -613,78 +613,102 @@ namespace uba
 			out.Appendf(TC("%s: +0x%llx"), mod.name.c_str(), entry.memoryOffset);
 		}
 #else
-
-		for (auto& entry : entries)
+		for (auto& mod : modules)
 		{
-			if (entry.moduleIndex == ~0u)
-			{
-				out.Appendf(TC("\n   <Unknown>: 0x%llx"), entry.memoryOffset);
+			if (mod.handled)
 				continue;
-			}
-			auto& mod = modules[entry.moduleIndex];
-			
-			if (!mod.handled)
+			if (mod.name.empty())
+				mod.name = executable;
+			StringBuffer<128> name(mod.name);
+			if (auto lastSlash = name.Last(PathSeparator))
+				name.Clear().Append(lastSlash + 1);
+			for (auto it=searchPaths; it->count; ++it)
 			{
-				if (mod.name.empty())
-					mod.name = executable;
-				StringBuffer<128> name(mod.name);
-				if (auto lastSlash = name.Last(PathSeparator))
-					name.Clear().Append(lastSlash + 1);
-				for (auto it=searchPaths; it->count; ++it)
-				{
-					struct stat attr;
-					StringBuffer<512> buf;
-					buf.Clear().Append(*it).EnsureEndsWithSlash().Append(name);
-					if (stat(buf.data, &attr) == -1)
-						continue;
-					mod.name = buf.data;
-					break;
-				}
-				mod.handled  = true;
+				struct stat attr;
+				StringBuffer<512> buf;
+				buf.Clear().Append(*it).EnsureEndsWithSlash().Append(name);
+				if (stat(buf.data, &attr) == -1)
+					continue;
+				mod.name = buf.data;
+				break;
 			}
+			mod.handled  = true;
+		}
 
-			auto pg = MakeGuard([&]() { out.Appendf(TC("\n   %s: 0x%llx"), mod.name.c_str(), entry.memoryOffset); });
-
+		for (u64 i=0; i!=modules.size(); ++i)
+		{
+			auto& mod = modules[i];
 			StringBuffer<1024> cmd;
 			#if PLATFORM_LINUX
-			cmd.Appendf("addr2line 0x%x", entry.memoryOffset).Append(" -e ").Append(mod.name).Append(" 2>/dev/null");
+			cmd.Append("addr2line");
 			#else
-			cmd.Appendf("atos -o %s --offset 0x%llx", mod.name.c_str(), entry.memoryOffset).Append(" 2>/dev/null");
+			cmd.Appendf("atos -o %s --offset", mod.name.c_str());
+			#endif
+			Vector<u64> memoryOffsets;
+			for (auto& entry : entries)
+				if (entry.moduleIndex == i)
+				{
+					cmd.Appendf(" 0x%x", entry.memoryOffset);
+					memoryOffsets.push_back(entry.memoryOffset);
+				}
+			#if PLATFORM_LINUX
+			cmd.Append(" -f -C -p -e ").Append(mod.name);
 			#endif
+			cmd.Append(" 2>/dev/null");
 
 			fflush(stdout);
 
-			FILE* fp = popen(cmd.data, "r");
-			if (!fp)
-				continue;
-			auto cg = MakeGuard([&]() { pclose(fp); });
+			u32 memoryOffsetIndex = 0;
+			if (FILE* fp = popen(cmd.data, "r"))
+			{
+				auto cg = MakeGuard([&]() { pclose(fp); });
 
-			u32 index = 0;
-			u32 countBeforeCallstack = out.count;
-			errno = 0;
-			StringBuffer<1024> str;
-			if (!fgets(str.data, str.capacity, fp))
+				u32 index = 0;
+				u32 countBeforeCallstack = out.count;
+				errno = 0;
+				StringBuffer<1024> str;
+				while (true)
+				{
+					if (!fgets(str.data, str.capacity, fp))
+						break;
+					str.count = TStrlen(str.data);
+					if (str.EndsWith("\n"))
+						str.Resize(str.count - 1);
+					mod.symbols.try_emplace(memoryOffsets[memoryOffsetIndex++], str.ToString());
+				}
+			}
+		}
+
+		u32 skipIndex = 0; // remove signal handler logic
+		for (u64 i=0; i!=entries.size() && !skipIndex; ++i)
+			if (entries[i].moduleIndex != ~0u)
+				if (modules[entries[i].moduleIndex].symbols[entries[i].memoryOffset].find("__restore_rt") != -1)
+					skipIndex = i + 1;
+
+		for (auto& entry : entries)
+		{
+			if (skipIndex > 0)
+			{
+				--skipIndex;
 				continue;
-			str.count = TStrlen(str.data);
-			if (str.StartsWith("addr2line"))
+			}
+			if (entry.moduleIndex == ~0u)
+			{
+				out.Appendf(TC("\n   <Unknown>: 0x%llx"), entry.memoryOffset);
 				continue;
-			if (str.EndsWith("\n"))
-				str.Resize(str.count - 1);
+			}
+			auto& mod = modules[entry.moduleIndex];
+		
+			auto pg = MakeGuard([&]() { out.Appendf(TC("\n   %s: 0x%llx"), mod.name.c_str(), entry.memoryOffset); });
+
+			StringBuffer<1024> str(mod.symbols[entry.memoryOffset]);
 			if (!str.count || str[0] == ':' || str[0] == '?')
 				continue;
 			char* fileName = str.data;
-			if (char* lastSlash = strrchr(fileName, '/'))
-				fileName = lastSlash + 1;
-			if (strncmp(fileName, "function.h", 10) == 0 || strncmp(fileName, "invoke.h:", 9) == 0)
-				continue;
-
 			pg.Cancel();
 			if (out.capacity - out.count < strlen(fileName) + 5)
 				break;
 			out.Appendf(TC("\n   %s"), fileName);
-
-			int errnoCopy = errno;
-			int err = ferror(fp);
 		}
 #endif
 	}
diff --git a/Engine/Source/Programs/UnrealBuildAccelerator/Common/Private/UbaThread.cpp b/Engine/Source/Programs/UnrealBuildAccelerator/Common/Private/UbaThread.cpp
--- a/Engine/Source/Programs/UnrealBuildAccelerator/Common/Private/UbaThread.cpp
+++ b/Engine/Source/Programs/UnrealBuildAccelerator/Common/Private/UbaThread.cpp
@@ -5,13 +5,15 @@
 
 #if PLATFORM_WINDOWS
 #include <tlhelp32.h>
-#else
+#elif PLATFORM_LINUX
 #include <dirent.h>
 #include <execinfo.h>
-#include <sys/ptrace.h>
 #include <sys/user.h>
 #include <sys/uio.h>
 #define REG_SP regs.rsp
+#else
+#include <mach/mach.h>
+#include <mach/thread_act.h>
 #endif
 
 #define UBA_TRACK_THREADS 0
@@ -209,17 +211,13 @@ namespace uba
 			HANDLE hThread = t->m_handle;
 			if (currentThreadId == GetThreadId(hThread))
 				continue;
-
 			if (SuspendThread(hThread) == -1)
 				continue;
-
 			CONTEXT ctx;
 			memset(&ctx, 0, sizeof(CONTEXT));
 			ctx.ContextFlags = CONTEXT_FULL;
 			GetThreadContext(hThread, &ctx);
-
 			func(0, &ctx);
-
 			ResumeThread(hThread);
 		}
 #else
@@ -229,123 +227,235 @@ namespace uba
 		if (hSnapshot == INVALID_HANDLE_VALUE)
 			return false;
 		auto sg = MakeGuard([&]() { CloseHandle(hSnapshot); });
-
 		UnorderedSet<CasKey> handledCallstacks;
 		THREADENTRY32 te32 = { sizeof(THREADENTRY32) };
 		if (!Thread32First(hSnapshot, &te32))
 			return false;
-
 		do
 		{
 			if (te32.th32OwnerProcessID != pid || te32.th32ThreadID == tid)
 				continue;
-
 			HANDLE hThread = OpenThread(THREAD_GET_CONTEXT | THREAD_SET_CONTEXT | THREAD_SUSPEND_RESUME | THREAD_QUERY_INFORMATION, FALSE, te32.th32ThreadID);
 			if (!hThread)
 				continue;
 			auto tg = MakeGuard([&]() { CloseHandle(hThread); });
-
 			if (SuspendThread(hThread) == -1)
 				continue;
-
 			CONTEXT ctx;
 			memset(&ctx, 0, sizeof(CONTEXT));
 			ctx.ContextFlags = CONTEXT_FULL;
 			GetThreadContext(hThread, &ctx);
-
-			func(te32.th32ThreadID, &ctx, 0);
-
+			void* callstack[100];
+			u32 callstackCount = GetCallstack(callstack, 100, 1, &ctx);
+			func(te32.th32ThreadID, callstack, callstackCount);
 			ResumeThread(hThread);
 		}
 		while (Thread32Next(hSnapshot, &te32));
 #endif
 		return true;
-#else
-#if 0
+#elif PLATFORM_LINUX
+
+		// TODO: None of these approaches work.
+		// signal path will break if thread is in certain system calls.. and ptrace does not work either :-/
+
 		static Event s_ev(false);
 		static const TraverseThreadFunc* s_func;
-		static u32 s_skipCallstack;
-
 		s_func = &func;
 
-		#if PLATFORM_MAC
-		void* bt[100];
-		s_skipCallstack = backtrace(bt, 100) + 4;
-		//printf("SKIP: %u\n", s_skipCallstack);
-		#endif
-
 		struct sigaction sa;
 		memset(&sa, 0, sizeof(sa));
 		sa.sa_flags = SA_SIGINFO;
 		sa.sa_sigaction = [](int sig, siginfo_t* info, void* context)
 			{
-				printf("IN THREAD %llx\n", (u64)pthread_self());
-				//pid_t currentTid = 0;//syscall(SYS_gettid);
-				//(*s_func)(currentTid, nullptr, s_skipCallstack);
+				void* callstack[100];
+				u32 callstackCount = GetCallstack(callstack, 100, 1, context);
+				(*s_func)(syscall(SYS_gettid), callstack, callstackCount);
 				s_ev.Set();
 			};
 		sigaction(SIGUSR1, &sa, NULL);
 
+		DIR* dir = opendir("/proc/self/task");
+		if (!dir)
+			return false;
 
-		pthread_t currentTid = pthread_self();
-		SCOPED_FUTEX(g_allThreadsLock, lock);
-		for (Thread* t=g_firstThread; t; t=t->m_next)
+		pid_t currentTid = syscall(SYS_gettid);
+		struct dirent *entry;
+		while ((entry = readdir(dir)) != NULL)
 		{
-			pthread_t tid = *(pthread_t*)&(t->m_handle);
-			if (currentTid == tid)
+			pid_t tid = (pid_t)atoi(entry->d_name);
+			if (tid <= 0 || tid == currentTid)
 				continue;
-			printf("HANDLING THREAD %llx\n", (u64)tid);
-			pthread_kill(tid, SIGUSR1);
+			kill(tid, SIGUSR1);
 			s_ev.IsSet();
 		}
-		lock.Leave();
-/*		
-		for (auto tid : threads)
+
+		#if 0
+		pid_t parentPid = getpid();
+
+		char message[100] = "Original message in parent.";
+ 
+		// We need to fork to be able to use ptrace on threads
+		pid_t child = fork();
+		if (child < 0)
 		{
-			printf("HANDLING THREAD\n");
-			pthread_kill(tid, SIGUSR1);
-			s_ev.IsSet();
+			perror("fork");
+			return false;
 		}
-*/
-/*
-		pid_t currentTid = syscall(SYS_gettid);
+		
+		if (child == 0) // The child
+		{
+			StringBuffer<128>().Appendf("/proc/%u/task", parentPid);
 
-		DIR* dir = opendir("/proc/self/task");
-		if (!dir)
-			return false;
+			auto g = MakeGuard([] { exit(EXIT_SUCCESS); });
+			DIR* dir = opendir(StringBuffer<128>().Appendf("/proc/%u/task", parentPid).data);
+			if (!dir)
+				return false;
 
-		Vector<pid_t> tids;
+			struct dirent *entry;
+			while ((entry = readdir(dir)) != NULL)
+			{
+				pid_t tid = (pid_t)atoi(entry->d_name);
+				if (tid <= 0)
+					continue;
 
-		struct dirent *entry;
-		while ((entry = readdir(dir)) != NULL)
-		{
-			char *endptr;
-			pid_t tid = (pid_t) strtol(entry->d_name, &endptr, 10);
-			if (*endptr == '\0' && errno != ERANGE && tid != currentTid)
-				tids.push_back(tid);
+				if (ptrace(PTRACE_ATTACH, tid, NULL, NULL) == -1)
+				{
+					perror("ptrace attach");
+					continue;
+				}
 
-		}
-		closedir(dir);
+				int status;
+				waitpid(tid, &status, 0); // Wait until the thread stops.
 
-		static Event s_ev(false);
-		static auto s_func = &func;
 
-		signal(SIGUSR1, [](int sig)
+				struct user_regs_struct regs;
+				if (ptrace(PTRACE_GETREGS, tid, NULL, &regs) == -1)
+				{
+					perror("ptrace getregs");
+					return false;
+				}
+				
+				printf("Callstack for thread %d:\n", tid);
+				printf("RIP: %llx\n", regs.rip);
+
+				// Start from the current base pointer.
+				unsigned long long fp = regs.rbp;
+				int frame = 0;
+
+				#define MAX_FRAMES 64
+				while (fp && frame < MAX_FRAMES) {
+					// Read the saved frame pointer (first word) and return address (second word)
+					unsigned long long next_fp, ret_addr;
+					errno = 0;
+					next_fp = ptrace(PTRACE_PEEKDATA, tid, (void *)fp, NULL);
+					if (errno != 0) break;
+					ret_addr = ptrace(PTRACE_PEEKDATA, tid, (void *)(fp + sizeof(unsigned long long)), NULL);
+					if (errno != 0) break;
+					printf("  Frame %d: ret_addr = %llx (fp = %llx)\n", frame, ret_addr, fp);
+					fp = next_fp;
+					frame++;
+				}
+				printf("\n");
+
+
+
+
+
+
+
+
+				if (ptrace(PTRACE_DETACH, tid, NULL, NULL) == -1) {
+					perror("ptrace detach");
+					return false;
+				}
+
+				const char* new_msg = "Hello from child!";
+
+				struct iovec local[1];
+				local[0].iov_base = (void *)new_msg;
+				local[0].iov_len  = strlen(new_msg) + 1;  // Include the null terminator
+
+				// Set up the remote iovec.
+				// The parent's view of the address is the same as the child's copy after fork.
+				struct iovec remote[1];
+				remote[0].iov_base = message;
+				remote[0].iov_len  = strlen(new_msg) + 1;
+
+				// Use process_vm_writev to write the new message into the parent's memory.
+				ssize_t nwritten = process_vm_writev(getppid(), local, 1, remote, 1, 0);
+			}
+
+			closedir(dir);
+
+		}
+		else // The parent
+		{
+			while (message[0] == 'O')
 			{
-				pid_t currentTid = syscall(SYS_gettid);
-				(*s_func)(currentTid, nullptr);
-				s_ev.Set();
-			});
+				Sleep(500);
+			}
+		}
+		#endif
 
+#else
+		task_t task;
+		kern_return_t kr = task_for_pid(mach_task_self(), getpid(), &task);
+		if (kr != KERN_SUCCESS)
+			return false;
 
-		for (pid_t tid : tids)
+		thread_act_array_t threads;
+		mach_msg_type_number_t thread_count;
+		kr = task_threads(task, &threads, &thread_count);
+		if (kr != KERN_SUCCESS)
+			return false;
+		auto tsg = MakeGuard([&] { vm_deallocate(mach_task_self(), (vm_address_t)threads, thread_count * sizeof(thread_t)); });
+
+		for (int i = 0; i < thread_count; i++)
 		{
-			kill(tid, SIGUSR1);
-			s_ev.IsSet();
+			if (threads[i] == mach_thread_self())
+				continue;
+
+			auto thread = threads[i];
+			thread_suspend(thread);
+			auto g = MakeGuard([&] { thread_resume(thread); });
+
+			kern_return_t kr;
+			uint64_t pc = 0, fp = 0;
+
+			#if defined(__x86_64__)
+				x86_thread_state64_t state;
+				mach_msg_type_number_t count = x86_THREAD_STATE64_COUNT;
+				kr = thread_get_state(thread, x86_THREAD_STATE64, (thread_state_t)&state, &count);
+				if (kr != KERN_SUCCESS)
+					return false;
+				pc = state.__rip;
+				fp = state.__rbp;
+			#elif defined(__arm64__)
+				arm_thread_state64_t state;
+				mach_msg_type_number_t count = ARM_THREAD_STATE64_COUNT;
+				kr = thread_get_state(thread, ARM_THREAD_STATE64, (thread_state_t)&state, &count);
+				if (kr != KERN_SUCCESS)
+					return false;
+				pc = state.__pc;
+				fp = state.__fp;
+			#else
+				#error "Unsupported architecture"
+			#endif
+		
+			void* callstack[100];
+			u32 callstackCount = 0;
+			for (int i = 0; i < 32 && fp; i++) {
+				uint64_t *stack = (uint64_t *)fp;
+				uint64_t return_addr = stack[1];
+				if (i > 0)
+					callstack[callstackCount++] = (void*)return_addr;
+				fp = stack[0];
+			}
+			g.Execute();
+
+			func(0, callstack, callstackCount);
 		}
-*/
 #endif
 		return true;
-#endif
 	}
 }
diff --git a/Engine/Source/Programs/UnrealBuildAccelerator/Common/Public/UbaThread.h b/Engine/Source/Programs/UnrealBuildAccelerator/Common/Public/UbaThread.h
--- a/Engine/Source/Programs/UnrealBuildAccelerator/Common/Public/UbaThread.h
+++ b/Engine/Source/Programs/UnrealBuildAccelerator/Common/Public/UbaThread.h
@@ -11,7 +11,7 @@ namespace uba
 {
 	class Event;
 	struct GroupAffinity;
-	using TraverseThreadFunc = Function<void(u32 tid, void* context, u32 skipCallstack)>;
+	using TraverseThreadFunc = Function<void(u32 tid, void** callstack, u32 callstackCount)>;
 
 	class Thread
 	{
diff --git a/Engine/Source/Programs/UnrealBuildAccelerator/Core/Private/UbaPlatform.cpp b/Engine/Source/Programs/UnrealBuildAccelerator/Core/Private/UbaPlatform.cpp
--- a/Engine/Source/Programs/UnrealBuildAccelerator/Core/Private/UbaPlatform.cpp
+++ b/Engine/Source/Programs/UnrealBuildAccelerator/Core/Private/UbaPlatform.cpp
@@ -192,10 +192,10 @@ namespace uba
 		#define GET_SP(context) context.Rsp
 	#endif
 
-	UBA_NOINLINE bool WriteCallstackInfo(BinaryWriter& out, u32 skipCallstack, void* contextPtr)
+	UBA_NOINLINE u32 GetCallstack(void** outCallstack, u32 callstackCapacity, u32 skipCallstack, void* contextPtr)
 	{
 #if PLATFORM_WINDOWS
-		const int MaxDepth = 16;// 62;
+		constexpr int MaxDepth = 16;// 62;
 		void* callers[MaxDepth];
 		u32 count = 0;//RtlCaptureStackBackTrace(0, MaxDepth, callers, NULL); // For some reason this does not work on wine
 		if (!count)
@@ -234,27 +234,31 @@ namespace uba
 				}
 			}
 		}
-		if (count < skipCallstack)
-			skipCallstack = count;
-
-		return WriteCallstackInfo(out, callers + skipCallstack, count - skipCallstack);
 #else
-		constexpr u32 maxCallers = 100;
+		constexpr u32 maxCallers = 64;
 		void* callers[maxCallers];
-		int nptrs = backtrace(callers, maxCallers);
-
+		u32 count = backtrace(callers, maxCallers);
 		if (contextPtr) // Use this to know if we come from a signal
 			skipCallstack += 2;
-
-		if (nptrs <= skipCallstack)
+#endif
+		if (count < skipCallstack)
+			skipCallstack = count;
+		u32 returnCount = count - skipCallstack;
+		if (callstackCapacity < returnCount)
 		{
-			if (!nptrs)
-				skipCallstack = 0;
-			else
-				skipCallstack = nptrs - 1;
+			skipCallstack += returnCount - callstackCapacity;
+			returnCount = count - skipCallstack;
 		}
-		return WriteCallstackInfo(out, callers + skipCallstack, nptrs - skipCallstack);
-#endif
+
+		memcpy(outCallstack, callers + skipCallstack, returnCount*sizeof(void*));
+		return returnCount;
+	}
+
+	UBA_NOINLINE bool WriteCallstackInfo(BinaryWriter& out, u32 skipCallstack, void* contextPtr)
+	{
+		void* callstack[100];
+		u32 callstackCount = GetCallstack(callstack, 100, skipCallstack, contextPtr);
+		return WriteCallstackInfo(out, callstack, callstackCount);
 	}
 
 	bool WriteCallstackInfo(BinaryWriter& out, void** callstack, u32 callstackCount)
diff --git a/Engine/Source/Programs/UnrealBuildAccelerator/Core/Public/UbaPlatform.h b/Engine/Source/Programs/UnrealBuildAccelerator/Core/Public/UbaPlatform.h
--- a/Engine/Source/Programs/UnrealBuildAccelerator/Core/Public/UbaPlatform.h
+++ b/Engine/Source/Programs/UnrealBuildAccelerator/Core/Public/UbaPlatform.h
@@ -80,6 +80,7 @@ namespace uba
 
 	#define UBA_NOT_IMPLEMENTED(x) UBA_ASSERTF(false, TC("%s not implemented!"), #x);
 
+	u32 GetCallstack(void** outCallstack, u32 callstackCapacity, u32 skipCallstack, void* contextPtr);
 	bool WriteCallstackInfo(BinaryWriter& out, u32 skipCallstack, void* context);
 	bool WriteCallstackInfo(BinaryWriter& out, void** callstack, u32 callstackCount);
 	void WriteAssertInfo(StringBufferBase& out, const tchar* text, const char* file, u32 line, const char* expr, void* context = nullptr);
diff --git a/Engine/Source/Programs/UnrealBuildAccelerator/Test/Private/UbaTestBasics.cpp b/Engine/Source/Programs/UnrealBuildAccelerator/Test/Private/UbaTestBasics.cpp
--- a/Engine/Source/Programs/UnrealBuildAccelerator/Test/Private/UbaTestBasics.cpp
+++ b/Engine/Source/Programs/UnrealBuildAccelerator/Test/Private/UbaTestBasics.cpp
@@ -554,29 +554,37 @@ namespace uba
 		return true;
 	}
 
-	UBA_NOINLINE void TestFunctionForThread(Event& ev1, Event& ev2)
+	UBA_NOINLINE void TestFunctionForThread(Event& ev1, Event& traverseDone)
 	{
 		ev1.Set();
-		ev2.IsSet();
+		traverseDone.IsSet();
 	}
 
 	bool TestThreads(LoggerWithWriter& logger, const StringBufferBase& rootDir)
 	{
+		Event traverseDone(true);
+
 		Event ev1(true);
-		Event ev2(true);
-		Thread t([&]()
+		Thread t1([&]()
 			{
-				TestFunctionForThread(ev1, ev2);
+				TestFunctionForThread(ev1, traverseDone);
 				return true;
 			});
-
 		ev1.IsSet();
 
-		TraverseAllThreads([&](u32 tid, void* context, u32 skipCallstack)
+		Event ev2(true);
+		Thread t2([&]()
+			{
+				TestFunctionForThread(ev2, traverseDone);
+				return true;
+			});
+		ev2.IsSet();
+
+		TraverseAllThreads([&](u32 tid, void** callstack, u32 callstackCount)
 			{
 				static u8* writerMem = new u8[4096];
 				BinaryWriter writer(writerMem, 0, 4096);
-				WriteCallstackInfo(writer, skipCallstack, context);
+				WriteCallstackInfo(writer, callstack, callstackCount);
 				BinaryReader reader(writerMem, 0, writer.GetPosition());
 				StringBuffer<16*1024> sb;
 				tchar executable[512] = TC("UbaTest");//{ 0 };
@@ -596,7 +604,7 @@ namespace uba
 			});
 
 
-		ev2.Set();
+		traverseDone.Set();
 		return true;
 	}
 }
