[UBA]
* Fixed so crashdump command outputs the correct callstacks on connected agents

diff --git a/Engine/Source/Programs/UnrealBuildAccelerator/Common/Private/UbaSessionClient.cpp b/Engine/Source/Programs/UnrealBuildAccelerator/Common/Private/UbaSessionClient.cpp
--- a/Engine/Source/Programs/UnrealBuildAccelerator/Common/Private/UbaSessionClient.cpp
+++ b/Engine/Source/Programs/UnrealBuildAccelerator/Common/Private/UbaSessionClient.cpp
@@ -1612,11 +1612,11 @@ namespace uba
 			UnorderedMap<CasKey, Callstack> callstacks;
 			StackBinaryWriter<SendMaxSize> stackWriter;
 
-			TraverseAllThreads([&](u32 tid, void* context, u32 skipCallstack)
+			TraverseAllThreads([&](u32 tid, void** callstack, u32 callstackCount)
 				{
 					stackWriter.Reset();
 					u32* infoSize = (u32*)stackWriter.AllocWrite(sizeof(u32));
-					WriteCallstackInfo(stackWriter, skipCallstack, context);
+					WriteCallstackInfo(stackWriter, callstack, callstackCount);
 					*infoSize = u32(stackWriter.GetPosition() - 4);
 
 					CasKeyHasher hasher;
diff --git a/Engine/Source/Programs/UnrealBuildAccelerator/Common/Private/UbaThread.cpp b/Engine/Source/Programs/UnrealBuildAccelerator/Common/Private/UbaThread.cpp
--- a/Engine/Source/Programs/UnrealBuildAccelerator/Common/Private/UbaThread.cpp
+++ b/Engine/Source/Programs/UnrealBuildAccelerator/Common/Private/UbaThread.cpp
@@ -447,13 +447,15 @@ namespace uba
 			for (int i = 0; i < 32 && fp; i++) {
 				uint64_t *stack = (uint64_t *)fp;
 				uint64_t return_addr = stack[1];
+				fp = stack[0];
+				if (!fp)
+					break;
 				if (i > 0)
 					callstack[callstackCount++] = (void*)return_addr;
-				fp = stack[0];
 			}
 			g.Execute();
 
-			func(0, callstack, callstackCount);
+			func(i+1, callstack, callstackCount);
 		}
 #endif
 		return true;
