[UBA]
* Added special handling for disallowed cas. If a client hit these it should disconnect and stop everything

diff --git a/Engine/Source/Programs/UnrealBuildAccelerator/Common/Private/UbaStorageClient.cpp b/Engine/Source/Programs/UnrealBuildAccelerator/Common/Private/UbaStorageClient.cpp
--- a/Engine/Source/Programs/UnrealBuildAccelerator/Common/Private/UbaStorageClient.cpp
+++ b/Engine/Source/Programs/UnrealBuildAccelerator/Common/Private/UbaStorageClient.cpp
@@ -365,6 +365,11 @@ namespace uba
 				fetchId = reader.ReadU16();
 				if (fetchId == 0)
 					return m_logger.Error(TC("Failed to fetch cas %s (%s)"), casFile.data, hint);
+				if (fetchId == u16(~0))
+				{
+					m_client.Disconnect();
+					return m_logger.Error(TC("Disallowed cas %s (%s)"), casFile.data, hint);
+				}
 
 				fileSize = reader.Read7BitEncoded();
 
diff --git a/Engine/Source/Programs/UnrealBuildAccelerator/Common/Private/UbaStorageProxy.cpp b/Engine/Source/Programs/UnrealBuildAccelerator/Common/Private/UbaStorageProxy.cpp
--- a/Engine/Source/Programs/UnrealBuildAccelerator/Common/Private/UbaStorageProxy.cpp
+++ b/Engine/Source/Programs/UnrealBuildAccelerator/Common/Private/UbaStorageProxy.cpp
@@ -288,6 +288,14 @@ namespace uba
 			m_logger.Error(TC("FetchBegin failed for cas file %s (%s)."), CasKeyString(file.casKey).str, hint.data);
 			return;
 		}
+		if (fetchId == u64(~0))
+		{
+			file.disallowed = true;
+			error = true;
+			m_logger.Error(TC("Disallowed download of cas file %s (%s)."), CasKeyString(file.casKey).str, hint.data);
+			return;
+		}
+
 		u64 fileSize = m.reader.Read7BitEncoded();
 		file.size = fileSize;
 
@@ -483,8 +491,16 @@ namespace uba
 	{
 		if (file.error)
 		{
-			writer.WriteU16(0);
-			return false;
+			if (file.disallowed)
+			{
+				writer.WriteU16(u16(~0));
+				return true;
+			}
+			else
+			{
+				writer.WriteU16(0);
+				return false;
+			}
 		}
 		u16 fetchId = u16(~0);
 
diff --git a/Engine/Source/Programs/UnrealBuildAccelerator/Common/Private/UbaStorageServer.cpp b/Engine/Source/Programs/UnrealBuildAccelerator/Common/Private/UbaStorageServer.cpp
--- a/Engine/Source/Programs/UnrealBuildAccelerator/Common/Private/UbaStorageServer.cpp
+++ b/Engine/Source/Programs/UnrealBuildAccelerator/Common/Private/UbaStorageServer.cpp
@@ -506,10 +506,9 @@ namespace uba
 
 				if (casEntry->disallowed)
 				{
-					writer.WriteU16(0);
-					if (!connectionInfo.ShouldDisconnect())
-						m_logger.Warning(TC("Client %u is asking for cas content of file that is not allowed to be transferred. Client will be disconnected (%s)"), connectionInfo.GetId(), hint.data);
-					return false;
+					writer.WriteU16(u16(~0));
+					m_logger.Warning(TC("Client %u is asking for cas content of file that is not allowed to be transferred (%s)"), hint.data);
+					return true;
 				}
 				
 				StringBuffer<512> casFile;
diff --git a/Engine/Source/Programs/UnrealBuildAccelerator/Common/Public/UbaStorageProxy.h b/Engine/Source/Programs/UnrealBuildAccelerator/Common/Public/UbaStorageProxy.h
--- a/Engine/Source/Programs/UnrealBuildAccelerator/Common/Public/UbaStorageProxy.h
+++ b/Engine/Source/Programs/UnrealBuildAccelerator/Common/Public/UbaStorageProxy.h
@@ -87,6 +87,7 @@ namespace uba
 			bool storeCompressed = false;
 			bool sendEnd = false;
 			bool error = false;
+			bool disallowed = false;
 			bool available = false;
 			BeginMessage* beginMessage = nullptr;
 			Vector<SegmentMessage*> segmentMessages;
