[UBA]
* Changed Disallowed cas fetch to be ~u16(0) - 1 to not collide with single message cas

diff --git a/Engine/Source/Programs/UnrealBuildAccelerator/Common/Private/UbaStorageClient.cpp b/Engine/Source/Programs/UnrealBuildAccelerator/Common/Private/UbaStorageClient.cpp
--- a/Engine/Source/Programs/UnrealBuildAccelerator/Common/Private/UbaStorageClient.cpp
+++ b/Engine/Source/Programs/UnrealBuildAccelerator/Common/Private/UbaStorageClient.cpp
@@ -365,7 +365,7 @@ namespace uba
 				fetchId = reader.ReadU16();
 				if (fetchId == 0)
 					return m_logger.Error(TC("Failed to fetch cas %s (%s)"), casFile.data, hint);
-				if (fetchId == u16(~0))
+				if (fetchId == FetchCasIdDisallowed)
 				{
 					m_client.Disconnect();
 					return m_logger.Error(TC("Disallowed cas %s (%s)"), casFile.data, hint);
@@ -530,7 +530,7 @@ namespace uba
 				{
 					if (sendSegmentMessage)
 					{
-						if (fetchId == u16(~0))
+						if (fetchId == FetchCasIdDone)
 							return m_logger.Error(TC("Cas content error. Server believes %s was only one segment but client sees more. Size: %llu Left to read: %llu ResponseSize: %u. (%s)"), hint, fileSize, left, responseSize, casFile.data);
 						readBuffer = slot;
 						u32 error;
@@ -601,7 +601,7 @@ namespace uba
 						{
 							if (sendSegmentMessage)
 							{
-								if (fetchId == u16(~0))
+								if (fetchId == FetchCasIdDone)
 									return m_logger.Error(TC("Cas content error (2). Server believes %s was only one segment but client sees more. UncompressedSize: %llu LeftUncompressed: %llu Size: %llu Left to read: %llu ResponseSize: %u. (%s)"), hint, actualSize, leftUncompressed, fileSize, left, responseSize, casFile.data);
 
 								u64 capacity = maxReadSize - u64(readPosition - readBuffer);
diff --git a/Engine/Source/Programs/UnrealBuildAccelerator/Common/Private/UbaStorageProxy.cpp b/Engine/Source/Programs/UnrealBuildAccelerator/Common/Private/UbaStorageProxy.cpp
--- a/Engine/Source/Programs/UnrealBuildAccelerator/Common/Private/UbaStorageProxy.cpp
+++ b/Engine/Source/Programs/UnrealBuildAccelerator/Common/Private/UbaStorageProxy.cpp
@@ -288,7 +288,7 @@ namespace uba
 			m_logger.Error(TC("FetchBegin failed for cas file %s (%s)."), CasKeyString(file.casKey).str, hint.data);
 			return;
 		}
-		if (fetchId == u64(~0))
+		if (fetchId == FetchCasIdDisallowed)
 		{
 			file.disallowed = true;
 			error = true;
@@ -493,7 +493,7 @@ namespace uba
 		{
 			if (file.disallowed)
 			{
-				writer.WriteU16(u16(~0));
+				writer.WriteU16(FetchCasIdDisallowed);
 				return true;
 			}
 			else
@@ -502,7 +502,7 @@ namespace uba
 				return false;
 			}
 		}
-		u16 fetchId = u16(~0);
+		u16 fetchId = FetchCasIdDone;
 
 		u64 headerSize = sizeof(u16) + Get7BitEncodedCount(file.size) + sizeof(u8);
 		u64 fetchedSize = Min(file.size, m_client.GetMessageMaxSize() - m_client.GetMessageReceiveHeaderSize() - headerSize);
diff --git a/Engine/Source/Programs/UnrealBuildAccelerator/Common/Private/UbaStorageServer.cpp b/Engine/Source/Programs/UnrealBuildAccelerator/Common/Private/UbaStorageServer.cpp
--- a/Engine/Source/Programs/UnrealBuildAccelerator/Common/Private/UbaStorageServer.cpp
+++ b/Engine/Source/Programs/UnrealBuildAccelerator/Common/Private/UbaStorageServer.cpp
@@ -506,7 +506,7 @@ namespace uba
 
 				if (casEntry->disallowed)
 				{
-					writer.WriteU16(u16(~0));
+					writer.WriteU16(FetchCasIdDisallowed);
 					m_logger.Warning(TC("Client %u is asking for cas content of file that is not allowed to be transferred (%s)"), hint.data);
 					return true;
 				}
@@ -649,7 +649,7 @@ namespace uba
 
 				if (!left)
 				{
-					*fetchId = u16(~0);
+					*fetchId = FetchCasIdDone;
 					u64 sendCasTime = GetTime() - start;
 					stats.sendCas += Timer{sendCasTime, 1};
 					return true;
@@ -1040,7 +1040,7 @@ namespace uba
 		SCOPED_FUTEX(m_availableIdsLock, lock);
 		if (m_availableIds.empty())
 		{
-			UBA_ASSERT(m_availableIdsHigh < u16(~0) - 1);
+			UBA_ASSERT(m_availableIdsHigh < u16(~0) - 2);
 			return m_availableIdsHigh++;
 		}
 		u16 storeId = m_availableIds.back();
diff --git a/Engine/Source/Programs/UnrealBuildAccelerator/Common/Public/UbaNetwork.h b/Engine/Source/Programs/UnrealBuildAccelerator/Common/Public/UbaNetwork.h
--- a/Engine/Source/Programs/UnrealBuildAccelerator/Common/Public/UbaNetwork.h
+++ b/Engine/Source/Programs/UnrealBuildAccelerator/Common/Public/UbaNetwork.h
@@ -127,6 +127,9 @@ namespace uba
 		}
 	}
 
+	constexpr u16 FetchCasIdDone = u16(~0);
+	constexpr u16 FetchCasIdDisallowed = u16(~0) - 1;
+
 	// Response types for SessionMessageType_ProcessAvailable
 	enum SessionProcessAvailableResponse : u32
 	{
