[UBA]
* Fixed infinite loop issue in UbaAgent

diff --git a/Engine/Source/Programs/UnrealBuildAccelerator/Agent/Private/UbaAgent.cpp b/Engine/Source/Programs/UnrealBuildAccelerator/Agent/Private/UbaAgent.cpp
--- a/Engine/Source/Programs/UnrealBuildAccelerator/Agent/Private/UbaAgent.cpp
+++ b/Engine/Source/Programs/UnrealBuildAccelerator/Agent/Private/UbaAgent.cpp
@@ -1315,6 +1315,7 @@ namespace uba
 
 			while (!ShouldExit())
 			{
+				u32 sessionWaitTimeout = 5*1000;
 				if (useListen)
 				{
 					if (connectionCount != targetConnectionCount)
@@ -1325,17 +1326,20 @@ namespace uba
 				}
 				else
 				{
-					while (connectionCount < targetConnectionCount)
+					if (connectionCount < targetConnectionCount && client->IsConnected())
 					{
 						bool timedOut = false;
 						if (client->Connect(*networkBackend, host.data, port, &timedOut))
+						{
 							++connectionCount;
+							sessionWaitTimeout = 0;
+						}
 						else
 							logger.Warning(TC("Failed to connect secondary connection number %u"), connectionCount);
 					}
 				}
 
-				if (sessionClient->Wait(5*1000, &wakeupSessionWait))
+				if (sessionClient->Wait(sessionWaitTimeout, &wakeupSessionWait))
 				{
 					// We got version mismatch and have the cas keys for the needed Agent/Detours binaries
 					if (keys[0] != CasKeyZero)
