[UBA]
* Fixed printf mismatch

diff --git a/Engine/Source/Programs/UnrealBuildAccelerator/Common/Private/UbaStorageServer.cpp b/Engine/Source/Programs/UnrealBuildAccelerator/Common/Private/UbaStorageServer.cpp
--- a/Engine/Source/Programs/UnrealBuildAccelerator/Common/Private/UbaStorageServer.cpp
+++ b/Engine/Source/Programs/UnrealBuildAccelerator/Common/Private/UbaStorageServer.cpp
@@ -507,7 +507,7 @@ namespace uba
 				if (casEntry->disallowed)
 				{
 					writer.WriteU16(FetchCasIdDisallowed);
-					m_logger.Warning(TC("Client %u is asking for cas content of file that is not allowed to be transferred (%s)"), hint.data);
+					m_logger.Warning(TC("Client %u is asking for cas content of file that is not allowed to be transferred (%s)"), connectionInfo.GetId(), hint.data);
 					return true;
 				}
 				
