[UBA]
* Added so casEntry is set as disallowed if server says it is disallowed

diff --git a/Engine/Source/Programs/UnrealBuildAccelerator/Common/Private/UbaStorageClient.cpp b/Engine/Source/Programs/UnrealBuildAccelerator/Common/Private/UbaStorageClient.cpp
--- a/Engine/Source/Programs/UnrealBuildAccelerator/Common/Private/UbaStorageClient.cpp
+++ b/Engine/Source/Programs/UnrealBuildAccelerator/Common/Private/UbaStorageClient.cpp
@@ -253,6 +253,9 @@ namespace uba
 			if (casEntry->verified && casEntry->exists)
 				return true;
 
+			if (casEntry->disallowed)
+				return false;
+
 			casEntry->dropped = false;  // In case this comes from a retry where previous cas was dropped
 			casEntry->verified = true;
 		}
@@ -367,6 +370,7 @@ namespace uba
 					return m_logger.Error(TC("Failed to fetch cas %s (%s)"), casFile.data, hint);
 				if (fetchId == FetchCasIdDisallowed)
 				{
+					casEntry->disallowed = true;
 					m_client.Disconnect();
 					return m_logger.Error(TC("Disallowed cas %s (%s)"), casFile.data, hint);
 				}
