[UBA]
* Added more info to assert about thing that should not happen but seem to be able to happen

diff --git a/Engine/Source/Programs/UnrealBuildAccelerator/Common/Private/UbaStorageProxy.cpp b/Engine/Source/Programs/UnrealBuildAccelerator/Common/Private/UbaStorageProxy.cpp
--- a/Engine/Source/Programs/UnrealBuildAccelerator/Common/Private/UbaStorageProxy.cpp
+++ b/Engine/Source/Programs/UnrealBuildAccelerator/Common/Private/UbaStorageProxy.cpp
@@ -251,7 +251,9 @@ namespace uba
 				file.memory = memory;
 				auto f = file.beginMessage;
 				file.beginMessage = nullptr;
-				UBA_ASSERT(f);
+				UBA_ASSERTF(f, TC("No begin message connected to %s. Should not happen"), m.hint);
+				if (!f)
+					return;
 				file.error = error;
 				fileLock.Leave();
 
