[UBA]
* Fixed so dependency crawler stops if process is cancelled

diff --git a/Engine/Source/Programs/UnrealBuildAccelerator/Common/Private/UbaNetworkClient.cpp b/Engine/Source/Programs/UnrealBuildAccelerator/Common/Private/UbaNetworkClient.cpp
--- a/Engine/Source/Programs/UnrealBuildAccelerator/Common/Private/UbaNetworkClient.cpp
+++ b/Engine/Source/Programs/UnrealBuildAccelerator/Common/Private/UbaNetworkClient.cpp
@@ -305,9 +305,11 @@ namespace uba
 
 	void NetworkClient::Disconnect()
 	{
-		m_isDisconnecting = true;
 		{
 			SCOPED_READ_LOCK(m_connectionsLock, lock);
+			if (m_isDisconnecting)
+				return;
+			m_isDisconnecting = true;
 			for (auto& c : m_connections)
 			{
 				OnDisconnected(c, 0);
diff --git a/Engine/Source/Programs/UnrealBuildAccelerator/Common/Private/UbaSession.cpp b/Engine/Source/Programs/UnrealBuildAccelerator/Common/Private/UbaSession.cpp
--- a/Engine/Source/Programs/UnrealBuildAccelerator/Common/Private/UbaSession.cpp
+++ b/Engine/Source/Programs/UnrealBuildAccelerator/Common/Private/UbaSession.cpp
@@ -3956,11 +3956,15 @@ namespace uba
 
 		auto CreateFileFunc = [this, ph = ProcessHandle(&process), rules = startInfo.rules](TrackWorkScope& tracker, const StringView& fileName, const DependencyCrawler::AccessFileFunc& func)
 			{
+				auto& process = *static_cast<ProcessImpl*>(ph.m_process);
+				if (process.IsCancelled())
+					return false;
+
 				CreateFileResponse out;
 				{
 					tracker.AddHint(fileName);
 
-					if (!CreateFileForRead(out, tracker, fileName, ToStringKey(fileName), *static_cast<ProcessImpl*>(ph.m_process), *rules))
+					if (!CreateFileForRead(out, tracker, fileName, ToStringKey(fileName), process, *rules))
 						return false;
 				}
 
diff --git a/Engine/Source/Programs/UnrealBuildAccelerator/Common/Private/UbaStorageClient.cpp b/Engine/Source/Programs/UnrealBuildAccelerator/Common/Private/UbaStorageClient.cpp
--- a/Engine/Source/Programs/UnrealBuildAccelerator/Common/Private/UbaStorageClient.cpp
+++ b/Engine/Source/Programs/UnrealBuildAccelerator/Common/Private/UbaStorageClient.cpp
@@ -371,6 +371,8 @@ namespace uba
 				if (fetchId == FetchCasIdDisallowed)
 				{
 					casEntry->disallowed = true;
+					if (proxy)
+						proxy->client.Disconnect();
 					m_client.Disconnect();
 					return m_logger.Error(TC("Disallowed cas %s (%s)"), casFile.data, hint);
 				}
diff --git a/Engine/Source/Programs/UnrealBuildAccelerator/Common/Private/UbaStorageProxy.cpp b/Engine/Source/Programs/UnrealBuildAccelerator/Common/Private/UbaStorageProxy.cpp
--- a/Engine/Source/Programs/UnrealBuildAccelerator/Common/Private/UbaStorageProxy.cpp
+++ b/Engine/Source/Programs/UnrealBuildAccelerator/Common/Private/UbaStorageProxy.cpp
@@ -244,6 +244,7 @@ namespace uba
 	{
 		auto& file = m.file;
 		u8* memory = nullptr;
+		StringView hint(m.hint);
 
 		auto sendResponses = MakeGuard([&]()
 			{
@@ -251,7 +252,7 @@ namespace uba
 				file.memory = memory;
 				auto f = file.beginMessage;
 				file.beginMessage = nullptr;
-				UBA_ASSERTF(f, TC("No begin message connected to %s. Should not happen"), m.hint);
+				UBA_ASSERTF(f, TC("No begin message connected to %s. Should not happen"), hint.data);
 				if (!f)
 					return;
 				file.error = error;
@@ -281,8 +282,6 @@ namespace uba
 		if (error || !m.message.ProcessAsyncResults(m.reader))
 			return;
 
-		StringView hint(m.hint);
-
 		u16 fetchId = m.reader.ReadU16();
 		if (fetchId == 0)
 		{
