[UBA]
* Fixed so done function for async messages are always called
* Fixed bug in storage proxy where message could be deleted twice if disconnected when sending async message

diff --git a/Engine/Source/Programs/UnrealBuildAccelerator/Common/Private/UbaNetworkClient.cpp b/Engine/Source/Programs/UnrealBuildAccelerator/Common/Private/UbaNetworkClient.cpp
--- a/Engine/Source/Programs/UnrealBuildAccelerator/Common/Private/UbaNetworkClient.cpp
+++ b/Engine/Source/Programs/UnrealBuildAccelerator/Common/Private/UbaNetworkClient.cpp
@@ -488,6 +488,8 @@ namespace uba
 				message.m_error = 12; // should never happen
 			else
 				message.m_error = 6;
+			if (async)
+				message.Done(false);
 			return false;
 		}
 
@@ -542,6 +544,8 @@ namespace uba
 					if (!connection.connected)
 					{
 						message.m_error = 7;
+						if (async)
+							message.Done(false);
 						return false;
 					}
 
diff --git a/Engine/Source/Programs/UnrealBuildAccelerator/Common/Private/UbaStorageProxy.cpp b/Engine/Source/Programs/UnrealBuildAccelerator/Common/Private/UbaStorageProxy.cpp
--- a/Engine/Source/Programs/UnrealBuildAccelerator/Common/Private/UbaStorageProxy.cpp
+++ b/Engine/Source/Programs/UnrealBuildAccelerator/Common/Private/UbaStorageProxy.cpp
@@ -225,12 +225,7 @@ namespace uba
 				}, file.beginMessage);
 
 			if (!res)
-			{
-				delete file.beginMessage;
-				file.beginMessage = nullptr;
-				file.error = true;
 				return false;
-			}
 		}
 
 		auto& deferredResponse = file.beginMessage->deferredResponses.emplace_back();
