[uba] visualizer additions:
- show world time in addition to time offset on mouseover of timeline view
- add a graph showing active process count over time across all active machines

diff --git a/Engine/Source/Programs/UnrealBuildAccelerator/Core/Public/UbaTimer.h b/Engine/Source/Programs/UnrealBuildAccelerator/Core/Public/UbaTimer.h
--- a/Engine/Source/Programs/UnrealBuildAccelerator/Core/Public/UbaTimer.h
+++ b/Engine/Source/Programs/UnrealBuildAccelerator/Core/Public/UbaTimer.h
@@ -72,6 +72,7 @@ namespace uba
 
 	struct TimeToText
 	{
+		TimeToText(float seconds, bool allowMinutes = false);
 		TimeToText(u64 time, bool allowMinutes = false);
 		TimeToText(u64 time, bool allowMinutes, u64 frequency);
 		operator const tchar*() const { return str; };
diff --git a/Engine/Source/Programs/UnrealBuildAccelerator/Visualizer/Private/UbaTraceReader.cpp b/Engine/Source/Programs/UnrealBuildAccelerator/Visualizer/Private/UbaTraceReader.cpp
--- a/Engine/Source/Programs/UnrealBuildAccelerator/Visualizer/Private/UbaTraceReader.cpp
+++ b/Engine/Source/Programs/UnrealBuildAccelerator/Visualizer/Private/UbaTraceReader.cpp
@@ -1449,6 +1449,17 @@ namespace uba
 
 		TraceView::Process& process = session.processors[active.processorIndex].processes[active.processIndex];
 
+		// active process counts array has one entry for every 100ms of the duration of the trace
+		u32 endCountIndex = (u32)round(TimeToMs(time) / 100);
+		u32 startCountIndex = (u32)round(TimeToMs(process.start) / 100);
+		
+		out.activeProcessCounts.resize(endCountIndex, 0u);
+		for (u32 c = startCountIndex; c < endCountIndex; ++c)
+		{
+			u16& count = out.activeProcessCounts[c];
+			count++;
+			out.maxActiveProcessCount = Max(out.maxActiveProcessCount, count);
+		}
 		process.stop = time;
 		process.bitmapDirty = true;
 		return &process;
diff --git a/Engine/Source/Programs/UnrealBuildAccelerator/Visualizer/Private/UbaTraceReader.h b/Engine/Source/Programs/UnrealBuildAccelerator/Visualizer/Private/UbaTraceReader.h
--- a/Engine/Source/Programs/UnrealBuildAccelerator/Visualizer/Private/UbaTraceReader.h
+++ b/Engine/Source/Programs/UnrealBuildAccelerator/Visualizer/Private/UbaTraceReader.h
@@ -163,6 +163,7 @@ namespace uba
 		Vector<Session> sessions;
 		Vector<WorkTrack> workTracks;
 		Vector<tchar*> strings;
+		Vector<u16> activeProcessCounts;
 		Map<u64, StatusUpdate> statusMap;
 		Map<u32, CacheWrite> cacheWrites;
 		u64 realStartTime = 0;
@@ -176,6 +177,7 @@ namespace uba
 		u32 progressProcessesTotal = 0;
 		u32 progressProcessesDone = 0;
 		u32 progressErrorCount = 0;
+		u16 maxActiveProcessCount = 0;
 		bool remoteExecutionDisabled = false;
 		bool finished = true;
 	};
diff --git a/Engine/Source/Programs/UnrealBuildAccelerator/Visualizer/Private/UbaVisualizer.cpp b/Engine/Source/Programs/UnrealBuildAccelerator/Visualizer/Private/UbaVisualizer.cpp
--- a/Engine/Source/Programs/UnrealBuildAccelerator/Visualizer/Private/UbaVisualizer.cpp
+++ b/Engine/Source/Programs/UnrealBuildAccelerator/Visualizer/Private/UbaVisualizer.cpp
@@ -417,6 +417,7 @@ namespace uba
 			m_processSelected = false;
 			m_sessionSelectedIndex = ~0u;
 			m_statsSelected = false;
+			m_activeProcessGraphSelected = false;
 			m_buttonSelected = ~0u;
 			m_timelineSelected = 0;
 			m_fetchedFilesSelected = ~0u;
@@ -485,6 +486,7 @@ namespace uba
 			m_cpuColor = RGB(170, 170, 0);
 			m_memColor = RGB(170, 0, 255);
 			m_driveColor = RGB(170, 65, 55);
+			m_activeProcColor = RGB(0, 170, 170);
 		}
 		else
 		{
@@ -521,6 +523,7 @@ namespace uba
 			m_cpuColor = RGB(200, 130, 0); // Orange
 			m_memColor = RGB(170, 0, 255); // Purple
 			m_driveColor = RGB(255, 115, 96);
+			m_activeProcColor = RGB(0, 170, 170);
 		}
 
 		m_textPen = CreatePen(PS_SOLID, 1, m_textColor);
@@ -529,6 +532,7 @@ namespace uba
 		m_cpuPen = CreatePen(PS_SOLID, 1, m_cpuColor);
 		m_memPen = CreatePen(PS_SOLID, 1, m_memColor);
 		m_drivePen = CreatePen(PS_SOLID, 1, m_driveColor);
+		m_activeProcPen = CreatePen(PS_SOLID, 1, m_activeProcColor);
 	}
 
 	void Visualizer::ThreadLoop()
@@ -713,6 +717,7 @@ namespace uba
 			m_processSelected = false;
 			m_sessionSelectedIndex = ~0u;
 			m_statsSelected = false;
+			m_activeProcessGraphSelected = false;
 			m_buttonSelected = ~0u;
 			m_timelineSelected = 0;
 			m_fetchedFilesSelected = ~0u;
@@ -1105,6 +1110,67 @@ namespace uba
 			SetTextColor(hdc, m_textColor);
 			posY += 3;
 		}
+
+		bool showProcessGraph = m_config.showActiveProcessGraph;
+		if (showProcessGraph)
+		{
+			if (posY + GraphHeight >= progressRect.top && posY + GraphHeight - 5 < progressRect.bottom)
+			{
+				int posX = int(m_scrollPosX) + progressRect.left;
+				int graphBaseY = posY + GraphHeight - 4;
+				double graphHeight = double(GraphHeight - 2);
+
+				Vector<POINT> line;
+
+				SelectObject(hdc, m_activeProcPen);
+				bool isFirstUpdate = true;
+				bool isFirstDraw = true;
+				u64 prevValue = 0;
+				int prevX = 0;
+				int prevY = 0;
+				u64 e = m_traceView.activeProcessCounts.size();
+				line.clear();
+				for (u64 i = 0; i != e; ++i)
+				{
+					auto value = m_traceView.activeProcessCounts[i];
+					int x = posX + int(float(i/10.f) * scaleX);
+					int y = graphBaseY;
+
+					isFirstUpdate = i == 0;
+
+					double div = double(m_traceView.maxActiveProcessCount);
+					y -= int(double(value) * graphHeight / div);
+
+					if (prevValue <= 0)
+						prevY = y;
+
+					if (x > clientRect.left)
+					{
+						if (isFirstUpdate)
+						{
+							line.push_back({ x, y }); //MoveToEx(hdc, x, y, NULL);
+							isFirstDraw = false;
+						}
+						else
+						{
+							if (isFirstDraw)
+								line.push_back({ prevX, prevY }); // MoveToEx(hdc, prevX, prevY, NULL);
+							line.push_back({ x, y }); // LineTo(hdc, x, y);
+							isFirstDraw = false;
+						}
+					}
+					if (x > clientRect.right)
+						break;
+					isFirstUpdate = false;
+					prevX = x;
+					prevY = y;
+					prevValue = value;
+				}
+				if (line.size() > 1)
+					Polyline(hdc, line.data(), int(line.size()));
+			}
+			posY += GraphHeight;
+		}
 		
 		if (m_config.showActiveProcesses && !m_trace.m_activeProcesses.empty())
 		{
@@ -1854,10 +1920,12 @@ namespace uba
 					b.Append('0');
 				b.AppendValue(milliseconds);
 
+				StringBuffer<128> wt = GetWorldTime(timelineSelected);
+
 				SetActiveFont(m_popupFont);
 				DrawTextLogger logger(m_hwnd, hdc, m_popupFont.height, m_tooltipBackgroundBrush);
 
-				logger.Info(L"%s", b.data);
+				logger.Info(L"%s (%s)", b.data, wt.data);
 				logger.DrawAtPos(left + 4, timelineTop - 20);
 			}
 		}
@@ -2207,6 +2275,15 @@ namespace uba
 			logger.AddSpace(3);
 			logger.DrawAtCursor();
 		}
+		else if (m_activeProcessGraphSelected)
+		{
+			SetActiveFont(m_popupFont);
+			DrawTextLogger logger(m_hwnd, hdc, m_popupFont.height, m_tooltipBackgroundBrush);
+			logger.AddSpace(3);
+			logger.AddSpace(3);
+			logger.SetColor(m_activeProcColor).Info(L"  Active Processes: %u", m_activeProcessCount);
+			logger.DrawAtCursor();
+		}
 		else if (m_buttonSelected != ~0u)
 		{
 			const wchar_t* tooltip[] =
@@ -2833,6 +2910,31 @@ namespace uba
 
 		outResult.section = 1;
 
+		bool showProcessGraph = m_config.showActiveProcessGraph;
+		if (showProcessGraph)
+		{
+			if (posY + GraphHeight >= progressRect.top && posY + GraphHeight - 5 < progressRect.bottom)
+			{
+				POINT cursorPos;
+				GetCursorPos(&cursorPos);
+				ScreenToClient(m_hwnd, &cursorPos);
+
+				if (cursorPos.x < m_progressRectLeft)
+				{
+					return;
+				}
+
+				outResult.activeProcessGraphSelected = true;
+
+				float timeScale = (m_horizontalScaleValue * m_zoomValue) * 5.0f;
+				float startOffset = -(m_scrollPosX / timeScale);
+				// active process counts array has one entry for every 100ms so convert to an index on that time scale
+				u32 timeOffset100Ms = (u32)round(startOffset + float(cursorPos.x - m_progressRectLeft) / timeScale);
+				outResult.activeProcessCount = m_traceView.activeProcessCounts[timeOffset100Ms];
+			}
+			posY += GraphHeight;
+		}
+
 		TraceView::ProcessLocation& outLocation = outResult.processLocation;
 
 		if (m_config.showActiveProcesses && !m_trace.m_activeProcesses.empty())
@@ -3250,6 +3352,8 @@ namespace uba
 			res.sessionSelectedIndex == m_sessionSelectedIndex &&
 			res.statsSelected == m_statsSelected && memcmp(&res.stats, &m_stats, sizeof(Stats)) == 0 &&
 			res.buttonSelected == m_buttonSelected && res.timelineSelected == m_timelineSelected &&
+			res.activeProcessGraphSelected == m_activeProcessGraphSelected && 
+			res.activeProcessCount == m_activeProcessCount && 
 			res.fetchedFilesSelected == m_fetchedFilesSelected &&
 			res.workSelected == m_workSelected && res.workTrack == m_workTrack && res.workIndex == m_workIndex &&
 			res.hyperLink == m_hyperLinkSelected)
@@ -3259,6 +3363,8 @@ namespace uba
 		m_sessionSelectedIndex = res.sessionSelectedIndex;
 		m_statsSelected = res.statsSelected;
 		m_stats = res.stats;
+		m_activeProcessGraphSelected = res.activeProcessGraphSelected;
+		m_activeProcessCount = res.activeProcessCount;
 		m_buttonSelected = res.buttonSelected;
 		m_timelineSelected = res.timelineSelected;
 		m_fetchedFilesSelected = res.fetchedFilesSelected;
@@ -3354,9 +3460,14 @@ namespace uba
 	}
 
 	StringBuffer<128> Visualizer::GetWorldTime(u64 time)
+	{
+		return GetWorldTime(TimeToS(time));
+	}
+
+	StringBuffer<128> Visualizer::GetWorldTime(float seconds)
 	{
 		time_t rawTime = m_traceView.traceSystemStartTimeUs / 1'000'000;
-		rawTime += TimeToMs(time)/1000;
+		rawTime += (time_t)seconds;
 		struct tm timeinfo;
 		localtime_s(&timeinfo, &rawTime);
 		StringBuffer<128> buffer;
@@ -3425,7 +3536,7 @@ namespace uba
 			{
 				m_trace.ReadFile(m_traceView, m_fileName.data, m_replay != 0);
 				m_traceView.finished = m_replay == 0;
-				PostNewTitle(GetTitlePrefix(title).Appendf(L"%s (v%u) - %s", m_fileName.data, m_traceView.version, GetWorldTime(0).data));
+				PostNewTitle(GetTitlePrefix(title).Appendf(L"%s (v%u) - %s", m_fileName.data, m_traceView.version, GetWorldTime((u64)0).data));
 			}
 			else if (m_usingNamed)
 			{
diff --git a/Engine/Source/Programs/UnrealBuildAccelerator/Visualizer/Private/UbaVisualizer.h b/Engine/Source/Programs/UnrealBuildAccelerator/Visualizer/Private/UbaVisualizer.h
--- a/Engine/Source/Programs/UnrealBuildAccelerator/Visualizer/Private/UbaVisualizer.h
+++ b/Engine/Source/Programs/UnrealBuildAccelerator/Visualizer/Private/UbaVisualizer.h
@@ -18,6 +18,7 @@ namespace uba
 		UBA_VISUALIZER_FLAG(DetailedData, false, L"detailed data (use -UbaDetailedTrace for even more)") \
 		UBA_VISUALIZER_FLAG(CpuMemStats, true, L"cpu/mem stats") \
 		UBA_VISUALIZER_FLAG(NetworkStats, true, L"network stats") \
+		UBA_VISUALIZER_FLAG(ActiveProcessGraph, false, L"graph of active processes over time") \
 		UBA_VISUALIZER_FLAG(DriveStats, true, L"drive stats") \
 		UBA_VISUALIZER_FLAG(ProcessBars, true, L"process bars") \
 		UBA_VISUALIZER_FLAG(FinishedProcesses, true, L"finished process bars") \
@@ -101,6 +102,7 @@ namespace uba
 			u64 memTotal = 0;
 			u64 recvBytes = 0;
 			u64 sendBytes = 0;
+			u64 procActive = 0;
 			float cpuLoad = 0;
 			struct Drive
 			{
@@ -124,6 +126,8 @@ namespace uba
 			bool workSelected = false;
 			u32 workTrack = ~0u;
 			u32 workIndex = ~0u;
+			bool activeProcessGraphSelected = false;
+			u16 activeProcessCount = 0u;
 			TString hyperLink;
 		};
 		u64 GetPlayTime();
@@ -165,6 +169,7 @@ namespace uba
 		void UpdateTheme();
 		bool ActiveProcessesShouldFillHeight();
 		StringBuffer<128> GetWorldTime(u64 time);
+		StringBuffer<128> GetWorldTime(float seconds);
 
 		StringBuffer<256> m_namedTrace;
 		StringBuffer<256> m_fileName;
@@ -198,6 +203,7 @@ namespace uba
 		COLORREF m_cpuColor = {};
 		COLORREF m_memColor = {};
 		COLORREF m_driveColor = {};
+		COLORREF m_activeProcColor = {};
 		HBRUSH m_backgroundBrush = 0;
 		HBRUSH m_tooltipBackgroundBrush = 0;
 		HPEN m_textPen = 0;
@@ -207,6 +213,7 @@ namespace uba
 		HPEN m_cpuPen = 0;
 		HPEN m_memPen = 0;
 		HPEN m_drivePen = 0;
+		HPEN m_activeProcPen = 0;
 		HPEN m_processUpdatePen = 0;
 		HPEN m_checkboxPen = 0;
 		int m_boxHeight = 12;
@@ -262,6 +269,8 @@ namespace uba
 		bool m_processSelected = false;
 		u32 m_sessionSelectedIndex = ~0u;
 		bool m_statsSelected = false;
+		bool m_activeProcessGraphSelected = false;
+		u64 m_activeProcessCount = 0u;
 		Stats m_stats;
 		u32 m_buttonSelected = ~0u;
 		float m_timelineSelected = 0;
