[UBA]
* Added TraverseAllCallstacks function and added so it is called when shutting down times out
* Changed existing code to use TraverseAllCallstacks
* Added description to TraverseAllThreads

diff --git a/Engine/Source/Programs/UnrealBuildAccelerator/Common/Private/UbaLogger.cpp b/Engine/Source/Programs/UnrealBuildAccelerator/Common/Private/UbaLogger.cpp
--- a/Engine/Source/Programs/UnrealBuildAccelerator/Common/Private/UbaLogger.cpp
+++ b/Engine/Source/Programs/UnrealBuildAccelerator/Common/Private/UbaLogger.cpp
@@ -5,6 +5,7 @@
 #include "UbaFileAccessor.h"
 #include "UbaPlatform.h"
 #include "UbaStringBuffer.h"
+#include "UbaThread.h"
 #include "UbaTimer.h"
 
 #if PLATFORM_WINDOWS
@@ -22,18 +23,8 @@ namespace uba
 {
 	CustomAssertHandler* g_assertHandler;
 
-	ANALYSIS_NORETURN void UbaAssert(const tchar* text, const char* file, u32 line, const char* expr, bool allowTerminate, u32 terminateCode, void* context, u32 skipCallstackCount)
+	void ParseCallstack(StringBufferBase& out, BinaryReader& reader)
 	{
-		static ReaderWriterLock& assertLock = *new ReaderWriterLock(); // Leak to prevent asan annoyances during shutdown when asserts happen
-		SCOPED_WRITE_LOCK(assertLock, lock);
-
-		static u8* writerMem = new u8[4096];
-		BinaryWriter writer(writerMem, 0, 4096);
-		WriteCallstackInfo(writer, 2, context);
-
-		static auto& sb = *new StringBuffer<16*1024>();
-		WriteAssertInfo(sb.Clear(), text, file, line, expr, context);
-
 		StringView searchPaths[3];
 		StringBuffer<512> currentModuleDir;
 		LoggerWithWriter logger(g_nullLogWriter);
@@ -49,8 +40,23 @@ namespace uba
 		readlink("/proc/self/exe", executable, sizeof_array(executable));
 		#endif
 
+		ParseCallstackInfo(out, reader, executable, searchPaths);
+	}
+
+	ANALYSIS_NORETURN void UbaAssert(const tchar* text, const char* file, u32 line, const char* expr, bool allowTerminate, u32 terminateCode, void* context, u32 skipCallstackCount)
+	{
+		static ReaderWriterLock& assertLock = *new ReaderWriterLock(); // Leak to prevent asan annoyances during shutdown when asserts happen
+		SCOPED_WRITE_LOCK(assertLock, lock);
+
+		static u8* writerMem = new u8[4096];
+		BinaryWriter writer(writerMem, 0, 4096);
+		WriteCallstackInfo(writer, 2, context);
+
+		static auto& sb = *new StringBuffer<16*1024>();
+		WriteAssertInfo(sb.Clear(), text, file, line, expr, context);
+
 		BinaryReader reader(writerMem, 0, writer.GetPosition());
-		ParseCallstackInfo(sb, reader, executable, searchPaths);
+		ParseCallstack(sb, reader);
 
 		if (g_assertHandler)
 		{
@@ -522,11 +528,11 @@ namespace uba
 			mod.handled = false;
 		}
 
-		out.Append(TCV("\n CALLSTACK")).Append(isRunningWine ? TC(" (Wine)") : TC("")).Append(':');
+		out.Append(TCV("\n CALLSTACK")).Append(isRunningWine ? TCV(" (Wine)") : TCV("")).Append(':');
 
 		if (entries.empty())
 		{
-			out.Appendf(TC("\n   <No entries available>"));
+			out.Append(TCV("\n   <No entries available>"));
 			return;
 		}
 
@@ -571,7 +577,7 @@ namespace uba
 			out.Append(TCV("\n   "));
 			if (entry.moduleIndex == ~0u)
 			{
-				out.Appendf(TC("<Unknown>: 0x%llx"), entry.memoryOffset);
+				out.Append(TCV("<Unknown>"));
 				continue;
 			}
 			auto& mod = modules[entry.moduleIndex];
@@ -601,10 +607,10 @@ namespace uba
 					if (auto lastSlash = TStrrchr(fileName, '\\'))
 						fileName = lastSlash + 1;
 					if (gotSymbol)
-						out.Appendf(TC(" ("));
+						out.Append(TCV(" ("));
 					out.Appendf(TC("%s:%u"), fileName, line.LineNumber);
 					if (gotSymbol)
-						out.Appendf(TC(")"));
+						out.Append(TCV(")"));
 					gotLine = true;
 				}
 				if (gotSymbol || gotLine)
@@ -740,4 +746,52 @@ namespace uba
 		}
 	#endif
 	}
+
+	void TraverseAllCallstacks(const Function<void(const CallstackInfo&)>& func)
+	{
+		UnorderedMap<CasKey, CallstackInfo> callstacks;
+
+		TraverseAllThreads([&](u32 tid, void** callstack, u32 callstackCount, const tchar* desc)
+			{
+				StackBinaryWriter<4096> stackWriter;
+				WriteCallstackInfo(stackWriter, callstack, callstackCount);
+
+				CasKeyHasher hasher;
+				hasher.Update(stackWriter.GetData(), stackWriter.GetPosition());
+
+				auto insres = callstacks.try_emplace(ToCasKey(hasher, false));
+				CallstackInfo& cs = insres.first->second;
+				cs.threadIds.push_back(tid);
+				if (!insres.second)
+					return;
+				if (desc)
+					cs.desc = desc;
+				cs.data.resize(stackWriter.GetPosition());
+				memcpy(cs.data.data(), stackWriter.GetData(), cs.data.size());
+			});
+		for (auto& kv : callstacks)
+		{
+			CallstackInfo& cs = kv.second;
+			StringBuffer<1024> temp;
+			if (!cs.desc.empty())
+				temp.Append(cs.desc).Append(TCV(" - "));
+			temp.Append(TCV("Thread Ids: "));
+			for (u32 tid : cs.threadIds)
+				if (temp.count < temp.capacity - 5)
+					temp.AppendValue(tid).Append(' ');
+			cs.desc = temp.data;
+			func(cs);
+		}
+	}
+
+	void PrintAllCallstacks(Logger& logger)
+	{
+		TraverseAllCallstacks([&](const CallstackInfo& cs)
+			{
+				BinaryReader stackReader(cs.data.data(), 0, cs.data.size());
+				static auto& sb = *new StringBuffer<16*1024>();
+				ParseCallstack(sb.Clear(), stackReader);
+				logger.Info(TC("%s%s"), cs.desc.c_str(), sb.data);
+			});
+	}
 }
diff --git a/Engine/Source/Programs/UnrealBuildAccelerator/Common/Private/UbaNetworkServer.cpp b/Engine/Source/Programs/UnrealBuildAccelerator/Common/Private/UbaNetworkServer.cpp
--- a/Engine/Source/Programs/UnrealBuildAccelerator/Common/Private/UbaNetworkServer.cpp
+++ b/Engine/Source/Programs/UnrealBuildAccelerator/Common/Private/UbaNetworkServer.cpp
@@ -157,12 +157,18 @@ namespace uba
 			{
 				if (TimeToMs(GetTime() - startTimer) > 3000)
 				{
-					m_server.m_logger.Error(TC("Connection has waited 3 seconds to stop... something is stuck (Active worker count: %u). Need someone to attach a debugger to process to get callstacks"), m_activeWorkerCount.load());
+					m_server.m_logger.Error(TC("Connection has waited 3 seconds to stop... something is stuck (Active worker count: %u)"), m_activeWorkerCount.load());
+					PrintAllCallstacks(m_server.m_logger);
 					return false;
 				}
 				Sleep(1);
 			}
 
+			static bool printOnce = true;
+			if (printOnce)
+				PrintAllCallstacks(m_server.m_logger);
+			printOnce = false;
+
 			if (!m_disconnectCallbackCalled.IsSet(60000)) // This should never time out!
 				m_server.m_logger.Warning(TC("This should never happen!! Unknown consequences"));
 			return true;
diff --git a/Engine/Source/Programs/UnrealBuildAccelerator/Common/Private/UbaSessionClient.cpp b/Engine/Source/Programs/UnrealBuildAccelerator/Common/Private/UbaSessionClient.cpp
--- a/Engine/Source/Programs/UnrealBuildAccelerator/Common/Private/UbaSessionClient.cpp
+++ b/Engine/Source/Programs/UnrealBuildAccelerator/Common/Private/UbaSessionClient.cpp
@@ -1608,44 +1608,15 @@ namespace uba
 
 		if (reader.ReadBool()) // crashdump
 		{
-			struct Callstack { Vector<u8> data; Vector<u32> threadIds; };
-			UnorderedMap<CasKey, Callstack> callstacks;
-			StackBinaryWriter<SendMaxSize> stackWriter;
-
-			TraverseAllThreads([&](u32 tid, void** callstack, u32 callstackCount)
+			TraverseAllCallstacks([&](const CallstackInfo& cs)
 				{
-					stackWriter.Reset();
-					u32* infoSize = (u32*)stackWriter.AllocWrite(sizeof(u32));
-					WriteCallstackInfo(stackWriter, callstack, callstackCount);
-					*infoSize = u32(stackWriter.GetPosition() - 4);
-
-					CasKeyHasher hasher;
-					hasher.Update(stackWriter.GetData(), stackWriter.GetPosition());
-
-					auto insres = callstacks.try_emplace(ToCasKey(hasher, false));
-					Callstack& cs = insres.first->second;
-					cs.threadIds.push_back(tid);
-					if (!insres.second)
-						return;
-					cs.data.resize(stackWriter.GetPosition());
-					memcpy(cs.data.data(), stackWriter.GetData(), cs.data.size());
+					BinaryReader stackReader(cs.data.data(), 0, cs.data.size());
+					StackBinaryWriter<SendMaxSize> stackWriter;
+					GetSymbols(UBA_AGENT_EXECUTABLE, IsArmBinary, stackReader, stackWriter);
+					BinaryReader resultReader(stackWriter.GetData(), 0, stackWriter.GetPosition());
+					TString infoString = resultReader.ReadString();
+					m_logger.Info(TC("%s%s"), cs.desc.c_str(), infoString.c_str());
 				});
-
-			for (auto& kv : callstacks)
-			{
-				Callstack& cs = kv.second;
-				BinaryReader stackReader(cs.data.data(), 0, cs.data.size());
-				stackWriter.Reset();
-				GetSymbols(UBA_AGENT_EXECUTABLE, IsArmBinary, stackReader, stackWriter);
-				BinaryReader resultReader(stackWriter.GetData(), 0, stackWriter.GetPosition());
-
-				StringBuffer<1024> tids(TC("Thread Ids: "));
-				for (u32 tid : cs.threadIds)
-					if (tids.count < tids.capacity - 5)
-						tids.AppendValue(tid).Append(' ');
-				TString infoString = resultReader.ReadString();
-				m_logger.Info(TC("%s%s"), tids.data, infoString.c_str());
-			}
 		}
 	}
 
diff --git a/Engine/Source/Programs/UnrealBuildAccelerator/Common/Private/UbaThread.cpp b/Engine/Source/Programs/UnrealBuildAccelerator/Common/Private/UbaThread.cpp
--- a/Engine/Source/Programs/UnrealBuildAccelerator/Common/Private/UbaThread.cpp
+++ b/Engine/Source/Programs/UnrealBuildAccelerator/Common/Private/UbaThread.cpp
@@ -241,13 +241,16 @@ namespace uba
 			auto tg = MakeGuard([&]() { CloseHandle(hThread); });
 			if (SuspendThread(hThread) == -1)
 				continue;
+			PWSTR threadDesc = nullptr;
+			GetThreadDescription(hThread, &threadDesc);
+			LocalFree(threadDesc);
 			CONTEXT ctx;
 			memset(&ctx, 0, sizeof(CONTEXT));
 			ctx.ContextFlags = CONTEXT_FULL;
 			GetThreadContext(hThread, &ctx);
 			void* callstack[100];
 			u32 callstackCount = GetCallstack(callstack, 100, 1, &ctx);
-			func(te32.th32ThreadID, callstack, callstackCount);
+			func(te32.th32ThreadID, callstack, callstackCount, threadDesc);
 			ResumeThread(hThread);
 		}
 		while (Thread32Next(hSnapshot, &te32));
@@ -269,7 +272,7 @@ namespace uba
 			{
 				void* callstack[100];
 				u32 callstackCount = GetCallstack(callstack, 100, 1, context);
-				(*s_func)(syscall(SYS_gettid), callstack, callstackCount);
+				(*s_func)(syscall(SYS_gettid), callstack, callstackCount, nullptr);
 				s_ev.Set();
 			};
 		sigaction(SIGUSR1, &sa, NULL);
@@ -455,7 +458,7 @@ namespace uba
 			}
 			g.Execute();
 
-			func(i+1, callstack, callstackCount);
+			func(i+1, callstack, callstackCount, nullptr);
 		}
 #endif
 		return true;
diff --git a/Engine/Source/Programs/UnrealBuildAccelerator/Common/Public/UbaLogger.h b/Engine/Source/Programs/UnrealBuildAccelerator/Common/Public/UbaLogger.h
--- a/Engine/Source/Programs/UnrealBuildAccelerator/Common/Public/UbaLogger.h
+++ b/Engine/Source/Programs/UnrealBuildAccelerator/Common/Public/UbaLogger.h
@@ -3,7 +3,7 @@
 #pragma once
 
 #include "UbaLogWriter.h"
-#include "UbaSynchronization.h"
+#include "UbaMemory.h"
 #include <stdarg.h>
 
 #define UBA_DEBUG_LOGGER 0
@@ -95,4 +95,8 @@ namespace uba
 
 	void ParseCallstackInfo(StringBufferBase& out, BinaryReader& reader, const tchar* executable, const StringView* searchPaths);
 	void PrintContentionSummary(class Logger& logger);
+
+	struct CallstackInfo { Vector<u8> data; Vector<u32> threadIds; TString desc; };
+	void TraverseAllCallstacks(const Function<void(const CallstackInfo&)>& func);
+	void PrintAllCallstacks(Logger& logger);
 }
diff --git a/Engine/Source/Programs/UnrealBuildAccelerator/Common/Public/UbaThread.h b/Engine/Source/Programs/UnrealBuildAccelerator/Common/Public/UbaThread.h
--- a/Engine/Source/Programs/UnrealBuildAccelerator/Common/Public/UbaThread.h
+++ b/Engine/Source/Programs/UnrealBuildAccelerator/Common/Public/UbaThread.h
@@ -11,7 +11,7 @@ namespace uba
 {
 	class Event;
 	struct GroupAffinity;
-	using TraverseThreadFunc = Function<void(u32 tid, void** callstack, u32 callstackCount)>;
+	using TraverseThreadFunc = Function<void(u32 tid, void** callstack, u32 callstackCount, const tchar* description)>;
 
 	class Thread
 	{
diff --git a/Engine/Source/Programs/UnrealBuildAccelerator/Test/Private/UbaTestBasics.cpp b/Engine/Source/Programs/UnrealBuildAccelerator/Test/Private/UbaTestBasics.cpp
--- a/Engine/Source/Programs/UnrealBuildAccelerator/Test/Private/UbaTestBasics.cpp
+++ b/Engine/Source/Programs/UnrealBuildAccelerator/Test/Private/UbaTestBasics.cpp
@@ -580,7 +580,7 @@ namespace uba
 			});
 		ev2.IsSet();
 
-		TraverseAllThreads([&](u32 tid, void** callstack, u32 callstackCount)
+		TraverseAllThreads([&](u32 tid, void** callstack, u32 callstackCount, const tchar* desc)
 			{
 				static u8* writerMem = new u8[4096];
 				BinaryWriter writer(writerMem, 0, 4096);
