[UBA]
* Fixed bug in previous submit where session client would not get data with size written in the beginning
* Removed debug code

diff --git a/Engine/Source/Programs/UnrealBuildAccelerator/Common/Private/UbaLogger.cpp b/Engine/Source/Programs/UnrealBuildAccelerator/Common/Private/UbaLogger.cpp
--- a/Engine/Source/Programs/UnrealBuildAccelerator/Common/Private/UbaLogger.cpp
+++ b/Engine/Source/Programs/UnrealBuildAccelerator/Common/Private/UbaLogger.cpp
@@ -754,7 +754,9 @@ namespace uba
 		TraverseAllThreads([&](u32 tid, void** callstack, u32 callstackCount, const tchar* desc)
 			{
 				StackBinaryWriter<4096> stackWriter;
+				u32* stackSize = (u32*)stackWriter.AllocWrite(sizeof(u32));
 				WriteCallstackInfo(stackWriter, callstack, callstackCount);
+				*stackSize = u32(stackWriter.GetPosition() - 4);
 
 				CasKeyHasher hasher;
 				hasher.Update(stackWriter.GetData(), stackWriter.GetPosition());
@@ -789,6 +791,7 @@ namespace uba
 		TraverseAllCallstacks([&](const CallstackInfo& cs)
 			{
 				BinaryReader stackReader(cs.data.data(), 0, cs.data.size());
+				stackReader.Skip(4); // Skip the size
 				static auto& sb = *new StringBuffer<16*1024>();
 				ParseCallstack(sb.Clear(), stackReader);
 				logger.Info(TC("%s%s"), cs.desc.c_str(), sb.data);
diff --git a/Engine/Source/Programs/UnrealBuildAccelerator/Common/Private/UbaNetworkServer.cpp b/Engine/Source/Programs/UnrealBuildAccelerator/Common/Private/UbaNetworkServer.cpp
--- a/Engine/Source/Programs/UnrealBuildAccelerator/Common/Private/UbaNetworkServer.cpp
+++ b/Engine/Source/Programs/UnrealBuildAccelerator/Common/Private/UbaNetworkServer.cpp
@@ -164,11 +164,6 @@ namespace uba
 				Sleep(1);
 			}
 
-			static bool printOnce = true;
-			if (printOnce)
-				PrintAllCallstacks(m_server.m_logger);
-			printOnce = false;
-
 			if (!m_disconnectCallbackCalled.IsSet(60000)) // This should never time out!
 				m_server.m_logger.Warning(TC("This should never happen!! Unknown consequences"));
 			return true;
