[UBA]
* Fixed crashes related to active process count in visualizer

diff --git a/Engine/Source/Programs/UnrealBuildAccelerator/Visualizer/Private/UbaTraceReader.cpp b/Engine/Source/Programs/UnrealBuildAccelerator/Visualizer/Private/UbaTraceReader.cpp
--- a/Engine/Source/Programs/UnrealBuildAccelerator/Visualizer/Private/UbaTraceReader.cpp
+++ b/Engine/Source/Programs/UnrealBuildAccelerator/Visualizer/Private/UbaTraceReader.cpp
@@ -40,6 +40,8 @@ namespace uba
 		strings.clear();
 		statusMap.clear();
 		cacheWrites.clear();
+		activeProcessCounts.clear();
+		maxActiveProcessCount = 0;
 		startTime = 0;
 		totalProcessActiveCount = 0;
 		totalProcessExitedCount = 0;
diff --git a/Engine/Source/Programs/UnrealBuildAccelerator/Visualizer/Private/UbaVisualizer.cpp b/Engine/Source/Programs/UnrealBuildAccelerator/Visualizer/Private/UbaVisualizer.cpp
--- a/Engine/Source/Programs/UnrealBuildAccelerator/Visualizer/Private/UbaVisualizer.cpp
+++ b/Engine/Source/Programs/UnrealBuildAccelerator/Visualizer/Private/UbaVisualizer.cpp
@@ -2915,22 +2915,15 @@ namespace uba
 		{
 			if (posY + GraphHeight >= progressRect.top && posY + GraphHeight - 5 < progressRect.bottom)
 			{
-				POINT cursorPos;
-				GetCursorPos(&cursorPos);
-				ScreenToClient(m_hwnd, &cursorPos);
-
-				if (cursorPos.x < m_progressRectLeft)
-				{
-					return;
-				}
-
-				outResult.activeProcessGraphSelected = true;
-
 				float timeScale = (m_horizontalScaleValue * m_zoomValue) * 5.0f;
 				float startOffset = -(m_scrollPosX / timeScale);
 				// active process counts array has one entry for every 100ms so convert to an index on that time scale
-				u32 timeOffset100Ms = (u32)round(startOffset + float(cursorPos.x - m_progressRectLeft) / timeScale);
-				outResult.activeProcessCount = m_traceView.activeProcessCounts[timeOffset100Ms];
+				u32 timeOffset100Ms = (u32)round(startOffset + float(pos.x - m_progressRectLeft) / timeScale);
+				if (timeOffset100Ms < m_traceView.activeProcessCounts.size())
+				{
+					outResult.activeProcessGraphSelected = true;
+					outResult.activeProcessCount = m_traceView.activeProcessCounts[timeOffset100Ms];
+				}
 			}
 			posY += GraphHeight;
 		}
