[UBA]
* Fixed so disallowed cas log entry is made before disconnect (since disconnect mute logging)

diff --git a/Engine/Source/Programs/UnrealBuildAccelerator/Common/Private/UbaStorageClient.cpp b/Engine/Source/Programs/UnrealBuildAccelerator/Common/Private/UbaStorageClient.cpp
--- a/Engine/Source/Programs/UnrealBuildAccelerator/Common/Private/UbaStorageClient.cpp
+++ b/Engine/Source/Programs/UnrealBuildAccelerator/Common/Private/UbaStorageClient.cpp
@@ -370,11 +370,13 @@ namespace uba
 					return m_logger.Error(TC("Failed to fetch cas %s (%s)"), casFile.data, hint);
 				if (fetchId == FetchCasIdDisallowed)
 				{
+					m_logger.Error(TC("Disallowed cas %s (%s)"), casFile.data, hint); // Log first since Disconnect mute log
 					casEntry->disallowed = true;
+					casEntryLock.Execute();
 					if (proxy)
 						proxy->client.Disconnect();
 					m_client.Disconnect();
-					return m_logger.Error(TC("Disallowed cas %s (%s)"), casFile.data, hint);
+					return false;
 				}
 
 				fileSize = reader.Read7BitEncoded();
