[UBA]
* Fixed so SendPing is earlying out if fails to send message

diff --git a/Engine/Source/Programs/UnrealBuildAccelerator/Common/Private/UbaSessionClient.cpp b/Engine/Source/Programs/UnrealBuildAccelerator/Common/Private/UbaSessionClient.cpp
--- a/Engine/Source/Programs/UnrealBuildAccelerator/Common/Private/UbaSessionClient.cpp
+++ b/Engine/Source/Programs/UnrealBuildAccelerator/Common/Private/UbaSessionClient.cpp
@@ -1592,7 +1592,10 @@ namespace uba
 		StackBinaryReader<32> reader;
 		time = GetTime();
 		if (!msg.Send(reader, m_stats.pingMsg))
+		{
 			m_loop = false;
+			return;
+		}
 		u64 newTime = GetTime();
 		u64 lastPing = newTime - time;
 		m_lastPing = lastPing;
