[UBA]
* Fixed so active process count graph works when playing

diff --git a/Engine/Source/Programs/UnrealBuildAccelerator/Visualizer/Private/UbaTraceReader.cpp b/Engine/Source/Programs/UnrealBuildAccelerator/Visualizer/Private/UbaTraceReader.cpp
--- a/Engine/Source/Programs/UnrealBuildAccelerator/Visualizer/Private/UbaTraceReader.cpp
+++ b/Engine/Source/Programs/UnrealBuildAccelerator/Visualizer/Private/UbaTraceReader.cpp
@@ -1420,7 +1420,11 @@ namespace uba
 		auto& process = processor->processes.back();
 
 		m_activeProcesses.try_emplace(id, TraceView::ProcessLocation{sessionIndex, processorIndex, u32(processor->processes.size() - 1)});
-			
+		
+		u16 activeCount = u16(m_activeProcesses.size());
+		out.activeProcessCounts.emplace_back(time, activeCount);
+		out.maxActiveProcessCount = Max(out.maxActiveProcessCount, activeCount);
+
 		++session.processActiveCount;
 		++out.totalProcessActiveCount;
 
@@ -1451,17 +1455,8 @@ namespace uba
 
 		TraceView::Process& process = session.processors[active.processorIndex].processes[active.processIndex];
 
-		// active process counts array has one entry for every 100ms of the duration of the trace
-		u32 endCountIndex = (u32)round(TimeToMs(time) / 100);
-		u32 startCountIndex = (u32)round(TimeToMs(process.start) / 100);
-		
-		out.activeProcessCounts.resize(endCountIndex, 0u);
-		for (u32 c = startCountIndex; c < endCountIndex; ++c)
-		{
-			u16& count = out.activeProcessCounts[c];
-			count++;
-			out.maxActiveProcessCount = Max(out.maxActiveProcessCount, count);
-		}
+		out.activeProcessCounts.emplace_back(time, u16(m_activeProcesses.size()));
+
 		process.stop = time;
 		process.bitmapDirty = true;
 		return &process;
diff --git a/Engine/Source/Programs/UnrealBuildAccelerator/Visualizer/Private/UbaTraceReader.h b/Engine/Source/Programs/UnrealBuildAccelerator/Visualizer/Private/UbaTraceReader.h
--- a/Engine/Source/Programs/UnrealBuildAccelerator/Visualizer/Private/UbaTraceReader.h
+++ b/Engine/Source/Programs/UnrealBuildAccelerator/Visualizer/Private/UbaTraceReader.h
@@ -160,10 +160,16 @@ namespace uba
 		const Session& GetSession(const ProcessLocation& loc);
 		void Clear();
 
+		struct ActiveProcessCount
+		{
+			u64 time;
+			u16 count;
+		};
+
 		Vector<Session> sessions;
 		Vector<WorkTrack> workTracks;
 		Vector<tchar*> strings;
-		Vector<u16> activeProcessCounts;
+		Vector<ActiveProcessCount> activeProcessCounts;
 		Map<u64, StatusUpdate> statusMap;
 		Map<u32, CacheWrite> cacheWrites;
 		u64 realStartTime = 0;
diff --git a/Engine/Source/Programs/UnrealBuildAccelerator/Visualizer/Private/UbaVisualizer.cpp b/Engine/Source/Programs/UnrealBuildAccelerator/Visualizer/Private/UbaVisualizer.cpp
--- a/Engine/Source/Programs/UnrealBuildAccelerator/Visualizer/Private/UbaVisualizer.cpp
+++ b/Engine/Source/Programs/UnrealBuildAccelerator/Visualizer/Private/UbaVisualizer.cpp
@@ -1111,12 +1111,10 @@ namespace uba
 			posY += 3;
 		}
 
-		bool showProcessGraph = m_config.showActiveProcessGraph;
-		if (showProcessGraph)
+		if (m_config.showActiveProcessGraph)
 		{
 			if (posY + GraphHeight >= progressRect.top && posY + GraphHeight - 5 < progressRect.bottom)
 			{
-				int posX = int(m_scrollPosX) + progressRect.left;
 				int graphBaseY = posY + GraphHeight - 4;
 				double graphHeight = double(GraphHeight - 2);
 
@@ -1128,18 +1126,33 @@ namespace uba
 				u64 prevValue = 0;
 				int prevX = 0;
 				int prevY = 0;
-				u64 e = m_traceView.activeProcessCounts.size();
-				line.clear();
-				for (u64 i = 0; i != e; ++i)
+
+				float timeScale = (m_horizontalScaleValue*m_zoomValue)*50.0f;
+				float startOffset = ((m_scrollPosX/timeScale) - float(int(m_scrollPosX/timeScale))) * timeScale;
+				float time = -startOffset/timeScale;
+				float playTimeS = TimeToS(playTime);
+				u32 activeProcessCountIndex = 0;
+				u16 count = 0;
+				while (time < playTimeS)
 				{
-					auto value = m_traceView.activeProcessCounts[i];
-					int x = posX + int(float(i/10.f) * scaleX);
-					int y = graphBaseY;
+					int posX = progressRect.left + int(startOffset + time*timeScale);
+					if (posX >= clientRect.right)
+						break;
 
-					isFirstUpdate = i == 0;
+					while (activeProcessCountIndex < m_traceView.activeProcessCounts.size())
+					{
+						auto& entry = m_traceView.activeProcessCounts[activeProcessCountIndex];
+						count = entry.count;
+						if (float(TimeToMs(entry.time))/1000.0f > time)
+							break;
+						++activeProcessCountIndex;
+					}
+
+					int x = posX;
+					int y = graphBaseY;
 
 					double div = double(m_traceView.maxActiveProcessCount);
-					y -= int(double(value) * graphHeight / div);
+					y -= int(double(count) * graphHeight / div);
 
 					if (prevValue <= 0)
 						prevY = y;
@@ -1148,14 +1161,14 @@ namespace uba
 					{
 						if (isFirstUpdate)
 						{
-							line.push_back({ x, y }); //MoveToEx(hdc, x, y, NULL);
+							line.push_back({ x, y });
 							isFirstDraw = false;
 						}
 						else
 						{
 							if (isFirstDraw)
-								line.push_back({ prevX, prevY }); // MoveToEx(hdc, prevX, prevY, NULL);
-							line.push_back({ x, y }); // LineTo(hdc, x, y);
+								line.push_back({ prevX, prevY });
+							line.push_back({ x, y });
 							isFirstDraw = false;
 						}
 					}
@@ -1164,14 +1177,16 @@ namespace uba
 					isFirstUpdate = false;
 					prevX = x;
 					prevY = y;
-					prevValue = value;
+					prevValue = count;
+
+					time += 0.25f;
 				}
 				if (line.size() > 1)
 					Polyline(hdc, line.data(), int(line.size()));
 			}
 			posY += GraphHeight;
 		}
-		
+
 		if (m_config.showActiveProcesses && !m_trace.m_activeProcesses.empty())
 		{
 			auto drawBox = [&](u64 start, u64 stop, int height, bool selected, bool inProgress)
@@ -2910,19 +2925,27 @@ namespace uba
 
 		outResult.section = 1;
 
-		bool showProcessGraph = m_config.showActiveProcessGraph;
-		if (showProcessGraph)
+		if (m_config.showActiveProcessGraph)
 		{
-			if (posY + GraphHeight >= progressRect.top && posY + GraphHeight - 5 < progressRect.bottom)
+			if (pos.y > posY && pos.y < posY + GraphHeight)
 			{
-				float timeScale = (m_horizontalScaleValue * m_zoomValue) * 5.0f;
+				float timeScale = (m_horizontalScaleValue * m_zoomValue)*50.0f;
 				float startOffset = -(m_scrollPosX / timeScale);
-				// active process counts array has one entry for every 100ms so convert to an index on that time scale
-				u32 timeOffset100Ms = (u32)round(startOffset + float(pos.x - m_progressRectLeft) / timeScale);
-				if (timeOffset100Ms < m_traceView.activeProcessCounts.size())
+				u64 selectedTimeMs = u64(1000.0f*(startOffset + float(pos.x - m_progressRectLeft) / timeScale));
+
+				u64 lastTime = m_traceView.activeProcessCounts.empty() ? 0 : m_traceView.activeProcessCounts[m_traceView.activeProcessCounts.size()-1].time;
+				if (selectedTimeMs < TimeToMs(lastTime))
 				{
+					u16 count = 0;
+					for (auto& entry : m_traceView.activeProcessCounts)
+					{
+						count = entry.count;
+						if (TimeToMs(entry.time) > selectedTimeMs)
+							break;
+					}
+					
 					outResult.activeProcessGraphSelected = true;
-					outResult.activeProcessCount = m_traceView.activeProcessCounts[timeOffset100Ms];
+					outResult.activeProcessCount = count;
 				}
 			}
 			posY += GraphHeight;
