[UBA]
* Changed so disconnect caused by disallowed cas does not flush work since itself might be inside a job

diff --git a/Engine/Source/Programs/UnrealBuildAccelerator/Common/Private/UbaNetworkClient.cpp b/Engine/Source/Programs/UnrealBuildAccelerator/Common/Private/UbaNetworkClient.cpp
--- a/Engine/Source/Programs/UnrealBuildAccelerator/Common/Private/UbaNetworkClient.cpp
+++ b/Engine/Source/Programs/UnrealBuildAccelerator/Common/Private/UbaNetworkClient.cpp
@@ -303,8 +303,10 @@ namespace uba
 		return true;
 	}
 
-	void NetworkClient::Disconnect()
+	void NetworkClient::Disconnect(bool flushWork)
 	{
+		auto fg = MakeGuard([&]() { if (flushWork) FlushWork(); });
+
 		{
 			SCOPED_READ_LOCK(m_connectionsLock, lock);
 			if (m_isDisconnecting)
@@ -324,8 +326,6 @@ namespace uba
 			m_connections.clear();
 			m_connectionsIt = m_connections.end();
 		}
-
-		FlushWork();
 	}
 
 	bool NetworkClient::StartListen(NetworkBackend& backend, u16 port, const tchar* ip)
diff --git a/Engine/Source/Programs/UnrealBuildAccelerator/Common/Private/UbaStorageClient.cpp b/Engine/Source/Programs/UnrealBuildAccelerator/Common/Private/UbaStorageClient.cpp
--- a/Engine/Source/Programs/UnrealBuildAccelerator/Common/Private/UbaStorageClient.cpp
+++ b/Engine/Source/Programs/UnrealBuildAccelerator/Common/Private/UbaStorageClient.cpp
@@ -375,7 +375,7 @@ namespace uba
 					casEntryLock.Execute();
 					if (proxy)
 						proxy->client.Disconnect();
-					m_client.Disconnect();
+					m_client.Disconnect(false); // Can't flush since we might be in a job
 					return false;
 				}
 
diff --git a/Engine/Source/Programs/UnrealBuildAccelerator/Common/Public/UbaNetworkClient.h b/Engine/Source/Programs/UnrealBuildAccelerator/Common/Public/UbaNetworkClient.h
--- a/Engine/Source/Programs/UnrealBuildAccelerator/Common/Public/UbaNetworkClient.h
+++ b/Engine/Source/Programs/UnrealBuildAccelerator/Common/Public/UbaNetworkClient.h
@@ -38,7 +38,7 @@ namespace uba
 		~NetworkClient();
 
 		bool Connect(NetworkBackend& backend, const tchar* ip, u16 port = DefaultPort, bool* timedOut = nullptr);
-		void Disconnect();
+		void Disconnect(bool flushWork = true);
 
 		bool StartListen(NetworkBackend& backend, u16 port = DefaultPort, const tchar* ip = TC("0.0.0.0"));
 		bool SetConnectionCount(u32 count);
