[UBA]
* Added so work manager flush can time out

diff --git a/Engine/Source/Programs/UnrealBuildAccelerator/Common/Private/UbaNetworkClient.cpp b/Engine/Source/Programs/UnrealBuildAccelerator/Common/Private/UbaNetworkClient.cpp
--- a/Engine/Source/Programs/UnrealBuildAccelerator/Common/Private/UbaNetworkClient.cpp
+++ b/Engine/Source/Programs/UnrealBuildAccelerator/Common/Private/UbaNetworkClient.cpp
@@ -305,7 +305,15 @@ namespace uba
 
 	void NetworkClient::Disconnect(bool flushWork)
 	{
-		auto fg = MakeGuard([&]() { if (flushWork) FlushWork(); });
+		auto fg = MakeGuard([&]()
+			{
+				if (!flushWork)
+					return;
+				if (FlushWork(30*1000))
+					return;
+				m_logger.Error(TC("NetworkClient has waited 30 seconds for all work to finish... something is stuck"));
+				PrintAllCallstacks(m_logger);
+			});
 
 		{
 			SCOPED_READ_LOCK(m_connectionsLock, lock);
diff --git a/Engine/Source/Programs/UnrealBuildAccelerator/Common/Private/UbaWorkManager.cpp b/Engine/Source/Programs/UnrealBuildAccelerator/Common/Private/UbaWorkManager.cpp
--- a/Engine/Source/Programs/UnrealBuildAccelerator/Common/Private/UbaWorkManager.cpp
+++ b/Engine/Source/Programs/UnrealBuildAccelerator/Common/Private/UbaWorkManager.cpp
@@ -183,8 +183,11 @@ namespace uba
 		}
 	}
 
-	void WorkManagerImpl::FlushWork()
+	bool WorkManagerImpl::FlushWork(u32 timeoutMs)
 	{
+		u64 startTime = GetTime();
+		auto hasTimedOut = [&]() { return timeoutMs ? (timeoutMs < TimeToMs(GetTime() - startTime)) : false; };
+
 		while (true)
 		{
 			SCOPED_FUTEX_READ(m_workLock, lock);
@@ -192,10 +195,18 @@ namespace uba
 			lock.Leave();
 			if (workEmpty)
 				break;
+			if (hasTimedOut())
+				return false;
 			Sleep(5);
 		}
+
 		while (m_activeWorkerCount)
+		{
+			if (hasTimedOut())
+				return false;
 			Sleep(5);
+		}
+		return true;
 	}
 
 	WorkManagerImpl::Worker* WorkManagerImpl::PopWorkerNoLock()
diff --git a/Engine/Source/Programs/UnrealBuildAccelerator/Common/Public/UbaWorkManager.h b/Engine/Source/Programs/UnrealBuildAccelerator/Common/Public/UbaWorkManager.h
--- a/Engine/Source/Programs/UnrealBuildAccelerator/Common/Public/UbaWorkManager.h
+++ b/Engine/Source/Programs/UnrealBuildAccelerator/Common/Public/UbaWorkManager.h
@@ -62,7 +62,7 @@ namespace uba
 		virtual void AddWork(const WorkFunction& work, u32 count, const tchar* desc, const Color& color = ColorWork, bool highPriority = false) override;
 		virtual u32 GetWorkerCount() override;
 		virtual void DoWork(u32 count = 1) override;
-		void FlushWork();
+		bool FlushWork(u32 timeoutMs = 0);
 
 	private:
 		struct Worker;
