[UBA]
* Bumped obj file size for clang-cl.exe to 1gb.. they are bigger than 512mb when running with stress test unity on clang 19

diff --git a/Engine/Source/Programs/UnrealBuildAccelerator/Core/Private/UbaApplicationRules.cpp b/Engine/Source/Programs/UnrealBuildAccelerator/Core/Private/UbaApplicationRules.cpp
--- a/Engine/Source/Programs/UnrealBuildAccelerator/Core/Private/UbaApplicationRules.cpp
+++ b/Engine/Source/Programs/UnrealBuildAccelerator/Core/Private/UbaApplicationRules.cpp
@@ -308,7 +308,7 @@ namespace uba
 		virtual u64 FileTypeMaxSize(const StringView& file, bool isSystemOrTempFile) const override
 		{
 			if (file.EndsWith(TCV(".obj")) || file.EndsWith(TCV(".o")) || file.EndsWith(TCV(".o.tmp")) || file.EndsWith(TCV(".obj.tmp")))
-				return 512 * 1024 * 1024;
+				return 1024 * 1024 * 1024; // There are huge obj files when building with -stresstestunity
 			if (file.EndsWith(TCV(".d")))
 				return 32 * 1024 * 1024;
 			return Super::FileTypeMaxSize(file, isSystemOrTempFile);
