[UBA]
* Fixed so UbaAgent exits early if it fails to start listening

diff --git a/Engine/Source/Programs/UnrealBuildAccelerator/Agent/Private/UbaAgent.cpp b/Engine/Source/Programs/UnrealBuildAccelerator/Agent/Private/UbaAgent.cpp
--- a/Engine/Source/Programs/UnrealBuildAccelerator/Agent/Private/UbaAgent.cpp
+++ b/Engine/Source/Programs/UnrealBuildAccelerator/Agent/Private/UbaAgent.cpp
@@ -1028,7 +1028,9 @@ namespace uba
 
 			if (useListen)
 			{
-				client->StartListen(*networkBackend, port);
+				if (!client->StartListen(*networkBackend, port))
+					return logger.Error(TC("Failed to get start listening on port %u"), port);
+
 				u64 startTime = GetTime();
 				while (!client->IsOrWasConnected(200))
 				{
diff --git a/Engine/Source/Programs/UnrealBuildAccelerator/Common/Private/UbaNetworkClient.cpp b/Engine/Source/Programs/UnrealBuildAccelerator/Common/Private/UbaNetworkClient.cpp
--- a/Engine/Source/Programs/UnrealBuildAccelerator/Common/Private/UbaNetworkClient.cpp
+++ b/Engine/Source/Programs/UnrealBuildAccelerator/Common/Private/UbaNetworkClient.cpp
@@ -338,12 +338,10 @@ namespace uba
 
 	bool NetworkClient::StartListen(NetworkBackend& backend, u16 port, const tchar* ip)
 	{
-		backend.StartListen(m_logger, port, ip, [&](void* connection, const sockaddr& remoteSockAddr)
+		return backend.StartListen(m_logger, port, ip, [&](void* connection, const sockaddr& remoteSockAddr)
 			{
 				return AddConnection(backend, connection, nullptr);
 			});
-
-		return true;
 	}
 
 	bool NetworkClient::SetConnectionCount(u32 count)
