[UBA]
* Fixed shutdown deadlock when cancelling processes at certain point in time.
* Added variable for enabling sending symbols to client when resolving callstacks

diff --git a/Engine/Source/Programs/UnrealBuildAccelerator/Common/Private/UbaProcess.cpp b/Engine/Source/Programs/UnrealBuildAccelerator/Common/Private/UbaProcess.cpp
--- a/Engine/Source/Programs/UnrealBuildAccelerator/Common/Private/UbaProcess.cpp
+++ b/Engine/Source/Programs/UnrealBuildAccelerator/Common/Private/UbaProcess.cpp
@@ -523,7 +523,7 @@ namespace uba
 				}
 
 				#if !PLATFORM_WINDOWS
-				while (PollStdPipes(outReader, errReader))
+				while (PollStdPipes(outReader, errReader, 500) && !IsCancelled())
 					continue;
 				#endif
 			}
@@ -603,7 +603,7 @@ namespace uba
 		m_session.m_debugLogger->Info(TC("ProcessExitedStart (%u)"), m_id);
 		#endif
 
-		UBA_ASSERT(!m_parentProcess || !m_parentProcess->m_hasExited);
+		UBA_ASSERT(m_cancelled || !m_parentProcess || !m_parentProcess->m_hasExited);
 
 		m_hasExited = true;
 
@@ -1853,14 +1853,13 @@ namespace uba
 			arguments2.push_back(nullptr);
 			auto argsArray = arguments2.data();
 
-			short flags = POSIX_SPAWN_SETPGROUP;
-
 			posix_spawnattr_t attr;
 			int res = posix_spawnattr_init(&attr);
 			UBA_ASSERTF(res == 0, TC("posix_spawnattr_init (%s)"), strerror(errno));
 			auto attrGuard = MakeGuard([&]() { posix_spawnattr_destroy(&attr); });
 
-			res = posix_spawnattr_setflags(&attr, flags);
+			// We set process group because we want to make sure that all processes get killed when ctrl-c is pressed
+			res = posix_spawnattr_setflags(&attr, POSIX_SPAWN_SETPGROUP);
 			UBA_ASSERTF(res == 0, TC("posix_spawnattr_setflags (%s)"), strerror(errno));
 			res = posix_spawnattr_setpgroup(&attr, getpgrp());
 			UBA_ASSERTF(res == 0, TC("posix_spawnattr_setpgroup (%s)"), strerror(errno));
@@ -2244,10 +2243,8 @@ namespace uba
 	void ProcessImpl::WaitForParent()
 	{
 		u64 startTime = GetTime();
-		while (!m_waitForParent.IsSet(500))
+		while (!m_waitForParent.IsSet(500) && !IsCancelled())
 		{
-			if (IsCancelled())
-				break;
 			if (TimeToMs(GetTime() - startTime) > 120 * 1000) // 
 			{
 				startTime = GetTime();
@@ -2262,7 +2259,7 @@ namespace uba
 		{
 			auto& childProcess = *(ProcessImpl*)child.m_process;
 			childProcess.m_waitForParent.Set();
-			while (!childProcess.m_hasExited)
+			while (!childProcess.m_hasExited && !IsCancelled())
 				Sleep(10);
 		}
 	}
diff --git a/Engine/Source/Programs/UnrealBuildAccelerator/Common/Private/UbaSessionClient.cpp b/Engine/Source/Programs/UnrealBuildAccelerator/Common/Private/UbaSessionClient.cpp
--- a/Engine/Source/Programs/UnrealBuildAccelerator/Common/Private/UbaSessionClient.cpp
+++ b/Engine/Source/Programs/UnrealBuildAccelerator/Common/Private/UbaSessionClient.cpp
@@ -2340,23 +2340,24 @@ namespace uba
 		if (!msg.Send(responseReader, m_stats.customMsg))
 			return false;
 		writer.AllocWrite(responseReader.GetLeft());
-		return true;
 
-#if 0
-		CasKey detoursSymbolsKey = msgReader.ReadCasKey();
-		if (detoursSymbolsKey == CasKeyZero)
-			return true;
-		StringBuffer<128> symbolsFile(UBA_DETOURS_LIBRARY);
-		#if PLATFORM_WINDOWS
-		symbolsFile.Resize(symbolsFile.count - 3).Append("pdb");
-		#else
-		symbolsFile.Resize(symbolsFile.count - 2).Append("debug");
-		#endif
-		Storage::RetrieveResult result;
-		StringBuffer<> throwaway;
-		if (m_storage.RetrieveCasFile(result, AsCompressed(detoursSymbolsKey, false), symbolsFile.data))
-			WriteBinFile(throwaway, symbolsFile.data, detoursSymbolsKey, KeyToString(StringKeyZero), DefaultAttributes());
-#endif
+		if constexpr (DownloadDebugSymbols)
+		{
+			CasKey detoursSymbolsKey = responseReader.ReadCasKey();
+			if (detoursSymbolsKey == CasKeyZero)
+				return true;
+			StringBuffer<128> symbolsFile(UBA_DETOURS_LIBRARY);
+			#if PLATFORM_WINDOWS
+			symbolsFile.Resize(symbolsFile.count - 3).Append("pdb");
+			#else
+			symbolsFile.Resize(symbolsFile.count - 2).Append("debug");
+			#endif
+			Storage::RetrieveResult result;
+			StringBuffer<> throwaway;
+			if (m_storage.RetrieveCasFile(result, AsCompressed(detoursSymbolsKey, false), symbolsFile.data))
+				WriteBinFile(throwaway, symbolsFile, detoursSymbolsKey, KeyToString(StringKeyZero), DefaultAttributes());
+		}
+		return true;
 	}
 
 
diff --git a/Engine/Source/Programs/UnrealBuildAccelerator/Common/Private/UbaSessionServer.cpp b/Engine/Source/Programs/UnrealBuildAccelerator/Common/Private/UbaSessionServer.cpp
--- a/Engine/Source/Programs/UnrealBuildAccelerator/Common/Private/UbaSessionServer.cpp
+++ b/Engine/Source/Programs/UnrealBuildAccelerator/Common/Private/UbaSessionServer.cpp
@@ -2428,24 +2428,25 @@ namespace uba
 		}
 
 		GetSymbols(application.c_str(), isClientArm, reader, writer);
-		return true;
 
-		#if 0
-		bool isClientArm = reader.ReadBool();
-		CasKey detoursSymbolsKey;
-		StringBuffer<> dir;
-		if (GetDirectoryOfCurrentModule(m_logger, dir))
+		if constexpr (DownloadDebugSymbols)
 		{
-			bool deferCreation = true;
-			auto ChangeToSymbolExtension = [](StringBufferBase& str) -> StringBufferBase& { IsWindows ? str.Resize(str.count - 3).Append("pdb") : str.Resize(str.count - 2).Append("debug"); return str; };
-			if (!m_storage.StoreCasFile(detoursSymbolsKey, ChangeToSymbolExtension(dir).data, CasKeyZero, deferCreation) || detoursSymbolsKey == CasKeyZero)
+			CasKey detoursSymbolsKey;
+			StringBuffer<> dir;
+			if (GetDirectoryOfCurrentModule(m_logger, dir))
 			{
-				StringBuffer<> dir2;
-				if (GetAlternativeUbaPath(m_logger, dir2, dir, IsWindows && isClientArm))
-					m_storage.StoreCasFile(detoursSymbolsKey, ChangeToSymbolExtension(dir2.Append(UBA_DETOURS_LIBRARY)).data, CasKeyZero, deferCreation);
+				bool deferCreation = true;
+				auto ChangeToSymbolExtension = [](StringBufferBase& str) -> StringBufferBase& { IsWindows ? str.Resize(str.count - 3).Append("pdb") : str.Resize(str.count - 2).Append("debug"); return str; };
+				if (!m_storage.StoreCasFile(detoursSymbolsKey, ChangeToSymbolExtension(dir).data, CasKeyZero, deferCreation) || detoursSymbolsKey == CasKeyZero)
+				{
+					StringBuffer<> dir2;
+					if (GetAlternativeUbaPath(m_logger, dir2, dir, IsWindows && isClientArm))
+						m_storage.StoreCasFile(detoursSymbolsKey, ChangeToSymbolExtension(dir2.Append(UBA_DETOURS_LIBRARY)).data, CasKeyZero, deferCreation);
+				}
 			}
+			writer.WriteCasKey(detoursSymbolsKey);
 		}
-		writer.WriteCasKey(detoursSymbolsKey);
-		#endif
+
+		return true;
 	}
 }
diff --git a/Engine/Source/Programs/UnrealBuildAccelerator/Common/Private/UbaThread.cpp b/Engine/Source/Programs/UnrealBuildAccelerator/Common/Private/UbaThread.cpp
--- a/Engine/Source/Programs/UnrealBuildAccelerator/Common/Private/UbaThread.cpp
+++ b/Engine/Source/Programs/UnrealBuildAccelerator/Common/Private/UbaThread.cpp
@@ -458,7 +458,7 @@ namespace uba
 			}
 			g.Execute();
 
-			func(i+1, callstack, callstackCount, nullptr);
+			func(thread, callstack, callstackCount, nullptr);
 		}
 #endif
 		return true;
diff --git a/Engine/Source/Programs/UnrealBuildAccelerator/Common/Public/UbaNetwork.h b/Engine/Source/Programs/UnrealBuildAccelerator/Common/Public/UbaNetwork.h
--- a/Engine/Source/Programs/UnrealBuildAccelerator/Common/Public/UbaNetwork.h
+++ b/Engine/Source/Programs/UnrealBuildAccelerator/Common/Public/UbaNetwork.h
@@ -130,6 +130,8 @@ namespace uba
 	constexpr u16 FetchCasIdDone = u16(~0);
 	constexpr u16 FetchCasIdDisallowed = u16(~0) - 1;
 
+	constexpr bool DownloadDebugSymbols = false;
+
 	// Response types for SessionMessageType_ProcessAvailable
 	enum SessionProcessAvailableResponse : u32
 	{
