[UBA]
* Fixed compile error on windows debug config

diff --git a/Engine/Source/Programs/UnrealBuildAccelerator/Common/Private/UbaProcess.cpp b/Engine/Source/Programs/UnrealBuildAccelerator/Common/Private/UbaProcess.cpp
--- a/Engine/Source/Programs/UnrealBuildAccelerator/Common/Private/UbaProcess.cpp
+++ b/Engine/Source/Programs/UnrealBuildAccelerator/Common/Private/UbaProcess.cpp
@@ -603,7 +603,7 @@ namespace uba
 		m_session.m_debugLogger->Info(TC("ProcessExitedStart (%u)"), m_id);
 		#endif
 
-		UBA_ASSERT(m_cancelled || !m_parentProcess || !m_parentProcess->m_hasExited);
+		UBA_ASSERT(IsCancelled() || !m_parentProcess || !m_parentProcess->m_hasExited);
 
 		m_hasExited = true;
 
