[UBA]
* Added more coverage of detouring registry access when running in debug.

diff --git a/Engine/Source/Programs/UnrealBuildAccelerator/Detours/Private/Windows/UbaDetoursFunctionsKernelBase.inl b/Engine/Source/Programs/UnrealBuildAccelerator/Detours/Private/Windows/UbaDetoursFunctionsKernelBase.inl
--- a/Engine/Source/Programs/UnrealBuildAccelerator/Detours/Private/Windows/UbaDetoursFunctionsKernelBase.inl
+++ b/Engine/Source/Programs/UnrealBuildAccelerator/Detours/Private/Windows/UbaDetoursFunctionsKernelBase.inl
@@ -4060,6 +4060,16 @@ BOOL Detoured_FreeLibrary(HMODULE hModule)
 	return res;
 }
 
+LSTATUS Detoured_RegOpenKeyW(HKEY hKey, LPCWSTR lpSubKey, PHKEY phkResult)
+{
+	DETOURED_CALL(RegOpenKeyW);
+	//UBA_NOT_IMPLEMENTED();
+	SuppressCreateFileDetourScope cfs;
+	auto res = True_RegOpenKeyW(hKey, lpSubKey, phkResult);
+	DEBUG_LOG_TRUE(L"RegOpenKeyW", L"(%ls) -> %ls", lpSubKey, ToString(res == ERROR_SUCCESS));
+	return res;
+}
+
 LSTATUS Detoured_RegOpenKeyExW(HKEY hKey, LPCWSTR lpSubKey, DWORD ulOptions, REGSAM samDesired, PHKEY phkResult)
 {
 	DETOURED_CALL(RegOpenKeyExW);
@@ -4070,6 +4080,16 @@ LSTATUS Detoured_RegOpenKeyExW(HKEY hKey, LPCWSTR lpSubKey, DWORD ulOptions, REG
 	return res;
 }
 
+LSTATUS Detoured_RegCreateKeyExW(HKEY hKey, LPCWSTR lpSubKey, DWORD Reserved, LPWSTR lpClass, DWORD dwOptions, REGSAM samDesired, const LPSECURITY_ATTRIBUTES lpSecurityAttributes, PHKEY phkResult, LPDWORD lpdwDisposition)
+{
+	DETOURED_CALL(RegOpenKeyExW);
+	//UBA_NOT_IMPLEMENTED();
+	SuppressCreateFileDetourScope cfs;
+	auto res = True_RegCreateKeyExW(hKey, lpSubKey, Reserved, lpClass, dwOptions, samDesired, lpSecurityAttributes, phkResult, lpdwDisposition);
+	DEBUG_LOG_TRUE(L"RegCreateKeyExW", L"(%ls) -> %ls", lpSubKey, ToString(res == ERROR_SUCCESS));
+	return res;
+}
+
 LSTATUS Detoured_RegOpenKeyExA(HKEY hKey, LPCSTR lpSubKey, DWORD ulOptions, REGSAM samDesired, PHKEY phkResult)
 {
 	DETOURED_CALL(RegOpenKeyExA);
diff --git a/Engine/Source/Programs/UnrealBuildAccelerator/Detours/Private/Windows/UbaDetoursFunctionsWin.h b/Engine/Source/Programs/UnrealBuildAccelerator/Detours/Private/Windows/UbaDetoursFunctionsWin.h
--- a/Engine/Source/Programs/UnrealBuildAccelerator/Detours/Private/Windows/UbaDetoursFunctionsWin.h
+++ b/Engine/Source/Programs/UnrealBuildAccelerator/Detours/Private/Windows/UbaDetoursFunctionsWin.h
@@ -231,7 +231,9 @@
 #define DETOURED_FUNCTIONS_KERNELBASE_DEBUG \
 	DETOURED_FUNCTION(GetCommandLineA) \
 	DETOURED_FUNCTION(FreeLibrary) \
+	DETOURED_FUNCTION(RegOpenKeyW) \
 	DETOURED_FUNCTION(RegOpenKeyExW) \
+	DETOURED_FUNCTION(RegCreateKeyExW) \
 	DETOURED_FUNCTION_X64(SetLastError) \
 	DETOURED_FUNCTION_X64(GetLastError) \
 	DETOURED_FUNCTION(RegOpenKeyExA) \
