[UBA]
* Fixed bug where a file that only changed attributes (through chmod) was not properly updated in directory table

diff --git a/Engine/Source/Programs/UnrealBuildAccelerator/Common/Private/UbaSession.cpp b/Engine/Source/Programs/UnrealBuildAccelerator/Common/Private/UbaSession.cpp
--- a/Engine/Source/Programs/UnrealBuildAccelerator/Common/Private/UbaSession.cpp
+++ b/Engine/Source/Programs/UnrealBuildAccelerator/Common/Private/UbaSession.cpp
@@ -1189,7 +1189,7 @@ namespace uba
 		if (!insres.second && insres.first->second != ~u32(0))
 		{
 			BinaryReader reader(m_directoryTableMem + insres.first->second);
-			u32 oldAttr = reader.ReadFileAttributes();(void)oldAttr;
+			u32 oldAttr = reader.ReadFileAttributes();
 
 			if (isDirectory) // Ignore updating directories.. they should always be the same regardless
 			{
@@ -1206,7 +1206,7 @@ namespace uba
 				UBA_ASSERT(oldFileIndex == fileIndex); // Checking so it is really the same file
 				#endif
 				u64 oldSize = reader.ReadFileSize();
-				if (oldSize == fileSize)
+				if (oldSize == fileSize && oldAttr == attributes) // Only attributes could change from a chmod
 					return true;
 				// TODO: Somehow this can happen and I have no idea how. last written time should be set on close file so it shouldnt be possible.
 				//else
