[UBA]
* Removed some code that was used to figure out paths from xcrun call. It does not work anymore and is not needed it seems

diff --git a/Engine/Source/Programs/UnrealBuildAccelerator/Detours/Private/PosixOS/UbaDetoursFunctionsPosix.cpp b/Engine/Source/Programs/UnrealBuildAccelerator/Detours/Private/PosixOS/UbaDetoursFunctionsPosix.cpp
--- a/Engine/Source/Programs/UnrealBuildAccelerator/Detours/Private/PosixOS/UbaDetoursFunctionsPosix.cpp
+++ b/Engine/Source/Programs/UnrealBuildAccelerator/Detours/Private/PosixOS/UbaDetoursFunctionsPosix.cpp
@@ -1959,6 +1959,8 @@ int shared_posix_spawn(pid_t* pid, const char* path, const posix_spawn_file_acti
 		const char* argv2[] = { argv[0], argv[1], argv[2], "--find", argv[3], 0 };
 		if (!ExecuteHostRun(result, argv2))
 			return -1;
+
+#if 0 // This does not seem to be needed anymore. If we run into this again then we should instead use realpath query to host since below logic does not work with newer sdk
 		// This will return the trampoline metal which we don't want
 		if (const char* usrbin = strstr(result.data, "/bin/metal"))
 		{
@@ -1970,6 +1972,7 @@ int shared_posix_spawn(pid_t* pid, const char* path, const posix_spawn_file_acti
 				result.Append("ios");
 			result.Append("/bin/").Append(argv[3]);
 		}
+#endif
 
 		path = result.data;
 
@@ -2784,7 +2787,7 @@ extern "C"
 	UBA_EXPORT bool UbaRequestNextProcess(u32 prevExitCode, char* outArguments, u32 outArgumentsCapacity)
 	{
 		#if UBA_DEBUG_LOG_ENABLED
-		//FlushDebugLog();
+		FlushDebugLog();
 		#endif
 
 		*outArguments = 0;
@@ -2814,14 +2817,14 @@ extern "C"
 			g_kernelStats = {};
 			g_stats = {};
 
-			//#if UBA_DEBUG_LOG_ENABLED
-			//SuppressCreateFileDetourScope scope;
-			//HANDLE debugFile = (HANDLE)g_debugFile;
-			//g_debugFile = InvalidFileHandle;
-			//CloseHandle(debugFile);
-			//debugFile = CreateFileW(g_logName.data, GENERIC_WRITE, FILE_SHARE_READ, NULL, CREATE_ALWAYS, FILE_ATTRIBUTE_NORMAL, NULL);
-			//g_debugFile = (FileHandle)(u64)debugFile;
-			//#endif
+			#if UBA_DEBUG_LOG_ENABLED
+			SuppressDetourScope scope;
+			int debugFile = g_debugFile;
+			g_debugFile = InvalidFileHandle;
+			close(debugFile);
+			debugFile = open(g_logName.data, O_CREAT | O_TRUNC | O_WRONLY, S_IRUSR | S_IWUSR);
+			g_debugFile = (FileHandle)debugFile;
+			#endif
 		}
 
 		Rpc_UpdateTables();
