[UBA]
* Changed TString(str.data) to str.ToString()

diff --git a/Engine/Source/Programs/UnrealBuildAccelerator/Agent/Private/UbaAgent.cpp b/Engine/Source/Programs/UnrealBuildAccelerator/Agent/Private/UbaAgent.cpp
--- a/Engine/Source/Programs/UnrealBuildAccelerator/Agent/Private/UbaAgent.cpp
+++ b/Engine/Source/Programs/UnrealBuildAccelerator/Agent/Private/UbaAgent.cpp
@@ -974,7 +974,7 @@ namespace uba
 					entryType = LogEntryType_Error;
 				}
 				SCOPED_FUTEX(logLinesLock, lock);
-				logLines[logLinesIndex].push_back({ TString(name.data), entryType });
+				logLines[logLinesIndex].push_back({ name.ToString(), entryType });
 			}
 
 			logLinesAvailable.Set();
diff --git a/Engine/Source/Programs/UnrealBuildAccelerator/Cli/Private/UbaCli.cpp b/Engine/Source/Programs/UnrealBuildAccelerator/Cli/Private/UbaCli.cpp
--- a/Engine/Source/Programs/UnrealBuildAccelerator/Cli/Private/UbaCli.cpp
+++ b/Engine/Source/Programs/UnrealBuildAccelerator/Cli/Private/UbaCli.cpp
@@ -259,7 +259,7 @@ namespace uba
 						return PrintHelp(TC("-vfs needs a semicolon between virtual and local path"));
 					*(tchar*)semi = 0;
 					const tchar* localPath = semi + 1;
-					vfsEntries.push_back({TString(value.data), TString(localPath)});
+					vfsEntries.push_back({value.ToString(), TString(localPath)});
 				}
 				else if (name.Equals(TCV("-capacity")))
 				{
diff --git a/Engine/Source/Programs/UnrealBuildAccelerator/Common/Private/UbaCacheClient.cpp b/Engine/Source/Programs/UnrealBuildAccelerator/Common/Private/UbaCacheClient.cpp
--- a/Engine/Source/Programs/UnrealBuildAccelerator/Common/Private/UbaCacheClient.cpp
+++ b/Engine/Source/Programs/UnrealBuildAccelerator/Common/Private/UbaCacheClient.cpp
@@ -641,7 +641,7 @@ namespace uba
 
 				if (!outIsMatch)
 					if (m_reportMissReason && path.count) // if empty this has already been reported
-						misses.push_back({TString(path.data), entryIndex, cacheCasKey, localCasKey });
+						misses.push_back({path.ToString(), entryIndex, cacheCasKey, localCasKey });
 				return true;
 			};
 
diff --git a/Engine/Source/Programs/UnrealBuildAccelerator/Common/Private/UbaProcess.cpp b/Engine/Source/Programs/UnrealBuildAccelerator/Common/Private/UbaProcess.cpp
--- a/Engine/Source/Programs/UnrealBuildAccelerator/Common/Private/UbaProcess.cpp
+++ b/Engine/Source/Programs/UnrealBuildAccelerator/Common/Private/UbaProcess.cpp
@@ -689,7 +689,7 @@ namespace uba
 		{
 			m_exitCode = GetLastError();
 			temp.Clear().Appendf(TC("Failed to copy %s to %s (%s)"), fromName.data, toName.data, LastErrorToText(m_exitCode).data);
-			m_logLines.push_back({ TString(temp.data), LogEntryType_Error });
+			m_logLines.push_back({ temp.ToString(), LogEntryType_Error });
 			return true;
 		}
 		
diff --git a/Engine/Source/Programs/UnrealBuildAccelerator/Common/Private/UbaStorage.cpp b/Engine/Source/Programs/UnrealBuildAccelerator/Common/Private/UbaStorage.cpp
--- a/Engine/Source/Programs/UnrealBuildAccelerator/Common/Private/UbaStorage.cpp
+++ b/Engine/Source/Programs/UnrealBuildAccelerator/Common/Private/UbaStorage.cpp
@@ -1023,7 +1023,7 @@ namespace uba
 
 				if (allowParallel && m_workManager)
 				{
-					m_workManager->AddWork([&traverseCasFileDir, p = TString(fullPath.data)](const WorkContext&) { traverseCasFileDir(p); }, 1, TC("TraverseCasFiles"));
+					m_workManager->AddWork([&traverseCasFileDir, p = fullPath.ToString()](const WorkContext&) { traverseCasFileDir(p); }, 1, TC("TraverseCasFiles"));
 				}
 				else
 				{
@@ -1643,7 +1643,7 @@ namespace uba
 		TraverseAllCasFiles(casRoot.data, [&](const StringBufferBase& fullPath, const DirectoryEntry& e)
 			{
 				++entryCount;
-				workManager.AddWork([&, filePath = TString(fullPath.data), name = TString(e.name), lastWritten = e.lastWritten](const WorkContext&)
+				workManager.AddWork([&, filePath = fullPath.ToString(), name = TString(e.name), lastWritten = e.lastWritten](const WorkContext&)
 					{
 						StringBuffer<> timeStr;
 						writeTimeAgo(timeStr, lastWritten);
@@ -1743,7 +1743,7 @@ namespace uba
 					u64 lastWritten = e.lastWritten;
 					u64 size = e.size;
 
-					workManager.AddWork([&, p = TString(path.data), lastWritten, size](const WorkContext&)
+					workManager.AddWork([&, p = path.ToString(), lastWritten, size](const WorkContext&)
 						{
 							++foundFiles;
 							StringKey key = ToStringKey(p.data(), p.size());
diff --git a/Engine/Source/Programs/UnrealBuildAccelerator/Common/Private/UbaStorageClient.cpp b/Engine/Source/Programs/UnrealBuildAccelerator/Common/Private/UbaStorageClient.cpp
--- a/Engine/Source/Programs/UnrealBuildAccelerator/Common/Private/UbaStorageClient.cpp
+++ b/Engine/Source/Programs/UnrealBuildAccelerator/Common/Private/UbaStorageClient.cpp
@@ -982,7 +982,7 @@ namespace uba
 					if (!seenIds.insert(e.id).second)
 						return;
 					lock.Leave();
-					workManager.AddWork([&, filePath = TString(fullPath.data)](const WorkContext& context)
+					workManager.AddWork([&, filePath = fullPath.ToString()](const WorkContext& context)
 						{
 							PopulateCasFromDirsRecursive(filePath.c_str(), workManager, seenIds, seenIdsLock, shouldExit);
 						}, 1, TC("PopulateCasFromDirsRecursive"), ColorWork);
@@ -1015,7 +1015,7 @@ namespace uba
 					return;
 				}
 
-				workManager.AddWork([&, fe = &fileEntry, lw = e.lastWritten, s = e.size, filePath = TString(fullPath.data)](const WorkContext& context)
+				workManager.AddWork([&, fe = &fileEntry, lw = e.lastWritten, s = e.size, filePath = fullPath.ToString()](const WorkContext& context)
 					{
 						auto feLockLeave = MakeGuard([fe]() { fe->lock.Leave(); });
 
diff --git a/Engine/Source/Programs/UnrealBuildAccelerator/Common/Public/UbaDependencyCrawler.h b/Engine/Source/Programs/UnrealBuildAccelerator/Common/Public/UbaDependencyCrawler.h
--- a/Engine/Source/Programs/UnrealBuildAccelerator/Common/Public/UbaDependencyCrawler.h
+++ b/Engine/Source/Programs/UnrealBuildAccelerator/Common/Public/UbaDependencyCrawler.h
@@ -238,9 +238,9 @@ namespace uba
 				if (CaseInsensitiveFs)
 					fixedPath2.MakeLower();
 				if (pushFront)
-					outCodeFiles.push_front({TString(fixedPath2.data), hasPch});
+					outCodeFiles.push_front({fixedPath2.ToString(), hasPch});
 				else
-					outCodeFiles.push_back({TString(fixedPath2.data), hasPch});
+					outCodeFiles.push_back({fixedPath2.ToString(), hasPch});
 			};
 
 		auto AddRoot = [&](const StringView& path)
@@ -264,7 +264,7 @@ namespace uba
 					if (bloomFilter.IsEmpty())
 						return;
 				}
-				instance.includeRoots.push_back({TString(fixedPath2.data), bloomFilter});
+				instance.includeRoots.push_back({fixedPath2.ToString(), bloomFilter});
 			};
 
 		auto IgnoreOption = [](const StringView& path) { UBA_CRAWLER_DEBUG_PRINT("IGNORED", path.data); };
@@ -715,7 +715,7 @@ namespace uba
 						include.FixPathSeparators();
 
 						auto work = 
-						[this, ref = InstanceRef(instance), isQuote, isInclude, include2 = TString(include.data), codeFile2 = TString(codeFile.data), hf = &handledFiles, parseDefines](const WorkContext& context)
+						[this, ref = InstanceRef(instance), isQuote, isInclude, include2 = include.ToString(), codeFile2 = codeFile.ToString(), hf = &handledFiles, parseDefines](const WorkContext& context)
 							{
 								auto& instance = ref.instance;
 								auto& handledFiles = *hf;
@@ -873,7 +873,7 @@ namespace uba
 				if (isDirectory)
 				{
 					codeFilePath.Append(PathSeparator);
-					m_workManager.AddWork([this, ref = InstanceRef(instance), codeFile = TString(codeFilePath.data)](const WorkContext& context)
+					m_workManager.AddWork([this, ref = InstanceRef(instance), codeFile = codeFilePath.ToString()](const WorkContext& context)
 						{
 							TraverseInclude(context.tracker, ref.instance, codeFile);
 						}, 1, TC("CrawlBiDir"), ColorWork);
@@ -885,7 +885,7 @@ namespace uba
 						return;
 					#endif
 
-					m_workManager.AddWork([this, ref = InstanceRef(instance), codeFile = TString(codeFilePath.data)](const WorkContext& context)
+					m_workManager.AddWork([this, ref = InstanceRef(instance), codeFile = codeFilePath.ToString()](const WorkContext& context)
 						{
 							if (!ref.instance.createFileFunc(context.tracker, codeFile, {}))
 								m_logger.Warning(TC("Failed to open file %s from builtindir"), codeFile.c_str());
