[UBA]
* Fixed so FileAccessor with temporary file on linux can handle temporary file being on a different device

diff --git a/Engine/Source/Programs/UnrealBuildAccelerator/Common/Private/UbaFileAccessor.cpp b/Engine/Source/Programs/UnrealBuildAccelerator/Common/Private/UbaFileAccessor.cpp
--- a/Engine/Source/Programs/UnrealBuildAccelerator/Common/Private/UbaFileAccessor.cpp
+++ b/Engine/Source/Programs/UnrealBuildAccelerator/Common/Private/UbaFileAccessor.cpp
@@ -70,6 +70,9 @@ namespace uba
 		StringBuffer<> tempFile;
 		if (tempPath)
 		{
+			#if PLATFORM_LINUX
+			allowRead = true; // Needed if rename does not work
+			#endif
 			m_tempFileIndex = g_tempFileCounter++;
 			realFileName = tempFile.Append(tempPath).Append("Temp_").AppendValue(m_tempFileIndex).data;
 		}
@@ -505,11 +508,38 @@ namespace uba
 					u64 left = m_size;
 					while (left)
 					{
-						u64 written = sendfile(targetFd, sourceFd, NULL, left);
+						auto written = copy_file_range(sourceFd, NULL, targetFd, NULL, left, 0);
 						if (written == -1)
-							return m_logger.Error(TC("Failed to do sendfile from temporary %s (%i) to file %s (%i) (%s)"), realFileName, targetFd, m_fileName, sourceFd, strerror(errno));
+							break;
 						left -= written;
 					}
+					if (left)
+					{
+						left = m_size;
+						if (lseek(sourceFd, 0, SEEK_SET) == -1)
+							return m_logger.Error(TC("Failed to do lseek to beginning of sourceFd (%s)"), strerror(errno));
+						if (lseek(targetFd, 0, SEEK_SET) == -1)
+							return m_logger.Error(TC("Failed to do lseek to beginning of targetFd (%s)"), strerror(errno));
+						while (left)
+						{
+							u8 buffer[256*1024];
+							u64 toCopy = Min(left, u64(sizeof_array(buffer)));
+							auto bytesRead = read(sourceFd, buffer, toCopy);
+							if (bytesRead <= 0)
+								return m_logger.Error(TC("Failed to do read from file %s (%s)"), realFileName, strerror(errno));
+							u8* readPos = buffer;
+							u64 toWrite = bytesRead;
+							while (toWrite)
+							{
+								auto written = write(targetFd, readPos, toWrite);
+								if (written == -1)
+									return m_logger.Error(TC("Failed to write to file %s (%s)"), m_fileName, strerror(errno));
+								toWrite -= written;
+								readPos += written;
+							}
+							left -= bytesRead;
+						}
+					}
 					#endif
 
 					remove(realFileName); // Remove real file now when we have copied it over
diff --git a/Engine/Source/Programs/UnrealBuildAccelerator/Common/Private/UbaStorage.cpp b/Engine/Source/Programs/UnrealBuildAccelerator/Common/Private/UbaStorage.cpp
--- a/Engine/Source/Programs/UnrealBuildAccelerator/Common/Private/UbaStorage.cpp
+++ b/Engine/Source/Programs/UnrealBuildAccelerator/Common/Private/UbaStorage.cpp
@@ -2478,14 +2478,14 @@ namespace uba
 					if (!decompressToMemory(slot))
 						return false;
 
-					if (!destinationFile.CreateWrite(false, DefaultAttributes(), 0, m_tempPath.data))
-						return false;
-
 					MemoryBlock block(5*1024*1024);
 
 					if (!formattingFunc(block, slot, decompressedSize, destination))
 						return false;
 
+					if (!destinationFile.CreateWrite(false, DefaultAttributes(), block.writtenSize, m_tempPath.data))
+						return false;
+						
 					if (!destinationFile.Write(block.memory, block.writtenSize))
 						return false;
 				}
diff --git a/Engine/Source/Programs/UnrealBuildAccelerator/Detours/Private/PosixOS/UbaDetoursFunctionsPosix.cpp b/Engine/Source/Programs/UnrealBuildAccelerator/Detours/Private/PosixOS/UbaDetoursFunctionsPosix.cpp
--- a/Engine/Source/Programs/UnrealBuildAccelerator/Detours/Private/PosixOS/UbaDetoursFunctionsPosix.cpp
+++ b/Engine/Source/Programs/UnrealBuildAccelerator/Detours/Private/PosixOS/UbaDetoursFunctionsPosix.cpp
@@ -116,11 +116,11 @@ using namespace uba;
 	DETOURED_FUNCTION_LINUX \
 	DETOURED_FUNCTION_MACOS \
 
-	//DETOURED_FUNCTION(write)
-	//DETOURED_FUNCTION(fwrite)
 
 #if UBA_DEBUG && PLATFORM_LINUX
 #define DETOURED_FUNCTION_DEBUG \
+	DETOURED_FUNCTION(write) \
+	//DETOURED_FUNCTION(fwrite)
 
 #else
 #define DETOURED_FUNCTION_DEBUG
@@ -1091,17 +1091,19 @@ UBA_EXPORT int UBA_WRAPPER(open)(const char* file, int flags, ...)
 	return Shared_open("open", file, flags, mode, [](const char* realFile, int flags, int mode) { return TRUE_WRAPPER(open)(realFile, flags, mode); });
 }
 
-//UBA_EXPORT ssize_t UBA_WRAPPER(write)(int fd, const void* buf, size_t count)
-//{
-//	UBA_INIT_DETOUR(write, fd, buf, count);
-//	if (isatty(fd)) // stdout and stderr
-//	{
-//		Shared_WriteConsole((const char*)buf, count, fd == fileno(stderr));
-//		return count;
-//	}
-//	//DEBUG_LOG_TRUE("write", "(%i size: %llu)", fd, count);
-//	return TRUE_WRAPPER(write)(fd, buf, count);
-//}
+#if UBA_DEBUG && PLATFORM_LINUX
+UBA_EXPORT ssize_t UBA_WRAPPER(write)(int fd, const void* buf, size_t count)
+{
+	UBA_INIT_DETOUR(write, fd, buf, count);
+	//if (isatty(fd)) // stdout and stderr
+	//{
+	//	Shared_WriteConsole((const char*)buf, count, fd == fileno(stderr));
+	//	return count;
+	//}
+	DEBUG_LOG_TRUE("write", "(%i size: %llu)", fd, count);
+	return TRUE_WRAPPER(write)(fd, buf, count);
+}
+#endif
 
 //UBA_EXPORT int fcntl(int __fd, int __cmd, ...)
 //{
@@ -2744,8 +2746,12 @@ namespace uba
 	void WriteDebug(const char* str, u32 strLen)
 	{
 		int t = errno;
-		//TRUE_WRAPPER(write)(g_debugFile, str, strLen);
+
+		#if UBA_DEBUG && PLATFORM_LINUX
+		TRUE_WRAPPER(write)(g_debugFile, str, strLen);
+		#else
 		write(g_debugFile, str, strLen);
+		#endif
 		//fsync(g_debugFile);
 		errno = t;
 	}
@@ -2830,4 +2836,4 @@ extern "C"
 		Rpc_UpdateTables();
 		return newProcess;
 	}
-}
\ No newline at end of file
+}
