[UBA]
* Added support for detoured popen on linux with some todos
* Added unit tests
* Fixed some handling of path '/'

diff --git a/Engine/Source/Programs/UnrealBuildAccelerator/Detours/Private/PosixOS/UbaDetoursFunctionsPosix.cpp b/Engine/Source/Programs/UnrealBuildAccelerator/Detours/Private/PosixOS/UbaDetoursFunctionsPosix.cpp
--- a/Engine/Source/Programs/UnrealBuildAccelerator/Detours/Private/PosixOS/UbaDetoursFunctionsPosix.cpp
+++ b/Engine/Source/Programs/UnrealBuildAccelerator/Detours/Private/PosixOS/UbaDetoursFunctionsPosix.cpp
@@ -109,6 +109,8 @@ using namespace uba;
 	DETOURED_FUNCTION(fork) \
 	DETOURED_FUNCTION(vfork) \
 	DETOURED_FUNCTION(popen) \
+	DETOURED_FUNCTION(fgets) \
+	DETOURED_FUNCTION(pclose) \
 	DETOURED_FUNCTION(exit) \
 	DETOURED_FUNCTION(_exit) \
 	DETOURED_FUNCTION(_Exit) \
@@ -120,7 +122,6 @@ using namespace uba;
 #if UBA_DEBUG && PLATFORM_LINUX
 #define DETOURED_FUNCTION_DEBUG \
 	DETOURED_FUNCTION(write) \
-	//DETOURED_FUNCTION(fwrite)
 
 #else
 #define DETOURED_FUNCTION_DEBUG
@@ -812,7 +813,7 @@ UBA_EXPORT int UBA_WRAPPER(mkdir)(const char* path, mode_t mode)
 	UBA_INIT_DETOUR(mkdir, path, mode);
 
 	StringBuffer<> pathName;
-	if (!FixPath(pathName, path) || !CanDetour2(pathName))
+	if (!FixPath(pathName, path) || !CanDetour2(pathName) || (pathName.count == 1 && pathName[0] == '/'))
 	{
 		int res = TRUE_WRAPPER(mkdir)(path, mode);
 		DEBUG_LOG_TRUE("mkdir", "%s -> %i", path, res);
@@ -1245,6 +1246,8 @@ UBA_EXPORT DIR* UBA_WRAPPER(opendir)(const char* name)
 	}
 
 	StringBuffer<> forHash(dirName);
+	if (forHash.count == 1)
+		forHash.Resize(0);
 	if (CaseInsensitiveFs)
 		forHash.MakeLower();
 	DirHash hash(forHash);
@@ -1857,7 +1860,7 @@ void FlattenArgs(StringBufferBase& out, const char* const argv[])
 	}
 }
 
-bool ExecuteHostRun(StringBufferBase& out, const char* const* argv)
+bool ExecuteHostRun(StringBufferBase& out, const char* const* argv, bool removeLineFeed = true)
 {
 	{
 		StringBuffer<4096> command;
@@ -1879,7 +1882,7 @@ bool ExecuteHostRun(StringBufferBase& out, const char* const* argv)
 	bool success = reader.ReadBool();
 	reader.ReadString(out);
 	
-	if (out.count && out.data[out.count-1] == '\n')
+	if (removeLineFeed && out.count && out.data[out.count-1] == '\n')
 		out.Resize(out.count-1);
 
 	if (!success)
@@ -2395,11 +2398,61 @@ UBA_EXPORT pid_t UBA_WRAPPER(vfork)(void)
 	return pid;
 }
 
+// TODO: Allow multiple popen.. hacky!
+FILE* g_activePopen; 
+StringBuffer<> g_activePopenResult;
+u32 g_activePopenReadPos = 0;
+
 UBA_EXPORT FILE* UBA_WRAPPER(popen)(const char* command, const char* type)
 {
 	UBA_INIT_DETOUR(popen, command, type);
-	DEBUG_LOG_TRUE("popen", "");
-	return TRUE_WRAPPER(popen)(command, type);
+	DEBUG_LOG_TRUE("popen", "%s", command);
+
+	UBA_ASSERT(!g_activePopen);
+	const char* argv[] = { command, nullptr };
+	if (!ExecuteHostRun(g_activePopenResult, argv, false))
+		return nullptr;
+	g_activePopenReadPos = 0;
+	g_activePopen = (FILE*)1337;
+	return g_activePopen;
+	//return TRUE_WRAPPER(popen)(command, type);
+}
+
+UBA_EXPORT char* UBA_WRAPPER(fgets)(char* str, int count, FILE* stream)
+{
+	UBA_INIT_DETOUR(fgets, str, count, stream);
+	DEBUG_LOG_TRUE("fgets", "(%p)", stream);
+	if (stream == g_activePopen)
+	{
+		u32 toWrite;
+		if (const char* endl = g_activePopenResult.First('\n', g_activePopenReadPos))
+		{
+			u32 lineLen = u32(endl - g_activePopenResult.data) - g_activePopenReadPos + 1;
+			toWrite = Min(u32(count - 2), lineLen);
+		}
+		else
+		{
+			toWrite = Min(u32(count - 2), g_activePopenResult.count - g_activePopenReadPos);
+		}
+		memcpy(str, g_activePopenResult.data + g_activePopenReadPos, toWrite);
+		str[toWrite] = 0;
+		g_activePopenReadPos += toWrite;
+		return str;
+	}
+	return TRUE_WRAPPER(fgets)(str, count, stream);
+}
+
+UBA_EXPORT int UBA_WRAPPER(pclose)(FILE* stream)
+{
+	UBA_INIT_DETOUR(pclose, stream);
+	DEBUG_LOG_TRUE("pclose", "%p", stream);
+	if (stream == g_activePopen)
+	{
+		g_activePopenResult.Clear();
+		g_activePopen = nullptr;
+		return 0;
+	}
+	return TRUE_WRAPPER(pclose)(stream);
 }
 
 UBA_EXPORT void UBA_WRAPPER(exit)(int status)
diff --git a/Engine/Source/Programs/UnrealBuildAccelerator/Test/Private/UbaTestAll.h b/Engine/Source/Programs/UnrealBuildAccelerator/Test/Private/UbaTestAll.h
--- a/Engine/Source/Programs/UnrealBuildAccelerator/Test/Private/UbaTestAll.h
+++ b/Engine/Source/Programs/UnrealBuildAccelerator/Test/Private/UbaTestAll.h
@@ -80,6 +80,7 @@ namespace uba
 		UBA_NONMAC_TESTS \
 		UBA_POSIX_TESTS \
 		UBA_TEST(TestDetouredTouch) \
+		UBA_TEST(TestDetouredPopen) \
 
 
 #define UBA_MAC_TESTS \
diff --git a/Engine/Source/Programs/UnrealBuildAccelerator/Test/Private/UbaTestSession.cpp b/Engine/Source/Programs/UnrealBuildAccelerator/Test/Private/UbaTestSession.cpp
--- a/Engine/Source/Programs/UnrealBuildAccelerator/Test/Private/UbaTestSession.cpp
+++ b/Engine/Source/Programs/UnrealBuildAccelerator/Test/Private/UbaTestSession.cpp
@@ -430,6 +430,37 @@ namespace uba
 			});
 	}
 
+	bool TestDetouredPopen(LoggerWithWriter& logger, const StringBufferBase& testRootDir)
+	{
+		return RunLocal(logger, testRootDir, [](LoggerWithWriter& logger, SessionServer& session, const tchar* workingDir, const RunProcessFunction& runProcess)
+			{
+				StringBuffer<> testApp;
+				GetTestAppPath(logger, testApp);
+
+				ProcessStartInfo processInfo;
+				processInfo.application = testApp.data;
+				processInfo.workingDir = workingDir;
+				processInfo.arguments = "-popen";
+				processInfo.logLineFunc = [](void* userData, const tchar* line, u32 length, LogEntryType type)
+					{
+						LoggerWithWriter(g_consoleLogWriter, TC("")).Info(line);
+					};
+
+				ProcessHandle process = runProcess(processInfo);
+				if (!process.WaitForExit(100000))
+					return logger.Error(TC("UbaTestApp did not exit in 10 seconds"));
+				u32 exitCode = process.GetExitCode();
+
+				if (exitCode != 0)
+				{
+					for (auto& logLine : process.GetLogLines())
+						logger.Error(logLine.text.c_str());
+					return logger.Error(TC("UbaTestApp returned exit code %u"), exitCode);
+				}
+				return true;
+			});
+	}
+
 	const tchar* GetSystemApplication()
 	{
 		#if PLATFORM_WINDOWS
diff --git a/Engine/Source/Programs/UnrealBuildAccelerator/TestApp/Private/PosixOS/UbaTestAppPosix.cpp b/Engine/Source/Programs/UnrealBuildAccelerator/TestApp/Private/PosixOS/UbaTestAppPosix.cpp
--- a/Engine/Source/Programs/UnrealBuildAccelerator/TestApp/Private/PosixOS/UbaTestAppPosix.cpp
+++ b/Engine/Source/Programs/UnrealBuildAccelerator/TestApp/Private/PosixOS/UbaTestAppPosix.cpp
@@ -238,6 +238,22 @@ int main(int argc, char* argv[])
 			return 255;
 		return attr.st_mode;
 	}
+	else if (strcmp(argv[1], "-popen") == 0)
+	{
+		FILE* file = popen("xdg-user-dir DOCUMENTS", "r");
+		if (!file)
+			return -3;
+		char docPath[256];
+		if (!fgets(docPath, 256, file))
+			return -4;
+		auto docLen = strlen(docPath) - 1;
+		if (docLen <= 0)
+			return -5;
+		if (docPath[docLen] != '\n')
+			return -6;
+		pclose(file);
+		return 0;
+	}
 	else if (strncmp(argv[1], "-file=", 6) == 0)
 	{
 		void* detoursHandle = dlopen(UBA_DETOURS_LIBRARY, RTLD_LAZY);
