[UBA]
* Changed Timed out waiting for connection response to server from error to warning
* Added ToFalse() to logger and stringbuilder

diff --git a/Engine/Source/Programs/UnrealBuildAccelerator/Common/Private/UbaLogger.cpp b/Engine/Source/Programs/UnrealBuildAccelerator/Common/Private/UbaLogger.cpp
--- a/Engine/Source/Programs/UnrealBuildAccelerator/Common/Private/UbaLogger.cpp
+++ b/Engine/Source/Programs/UnrealBuildAccelerator/Common/Private/UbaLogger.cpp
@@ -122,12 +122,13 @@ namespace uba
 #endif
 	}
 
-	void Logger::Info(const tchar* format, ...)
+	Logger& Logger::Info(const tchar* format, ...)
 	{
 		va_list arg;
 		va_start(arg, format);
 		LogArg(LogEntryType_Info, format, arg);
 		va_end(arg);
+		return *this;
 	}
 
 	void Logger::Detail(const tchar* format, ...)
diff --git a/Engine/Source/Programs/UnrealBuildAccelerator/Common/Private/UbaNetworkClient.cpp b/Engine/Source/Programs/UnrealBuildAccelerator/Common/Private/UbaNetworkClient.cpp
--- a/Engine/Source/Programs/UnrealBuildAccelerator/Common/Private/UbaNetworkClient.cpp
+++ b/Engine/Source/Programs/UnrealBuildAccelerator/Common/Private/UbaNetworkClient.cpp
@@ -157,7 +157,7 @@ namespace uba
 			return false;
 
 		if (!rc.recvEvent.IsSet(5000)) // This can not happen. Since both callbacks are using rc we can't leave this function until we know we are not in the callbacks
-			return m_logger.Error(TC("Timed out waiting for connection response from server. Is this really an UbaNetworkServer we are connecting to?"));
+			return m_logger.Info(TC("Timed out waiting for connection response from server.")).ToFalse();
 
 		m_isOrWasConnected.Set();
 
diff --git a/Engine/Source/Programs/UnrealBuildAccelerator/Common/Public/UbaLogger.h b/Engine/Source/Programs/UnrealBuildAccelerator/Common/Public/UbaLogger.h
--- a/Engine/Source/Programs/UnrealBuildAccelerator/Common/Public/UbaLogger.h
+++ b/Engine/Source/Programs/UnrealBuildAccelerator/Common/Public/UbaLogger.h
@@ -26,7 +26,7 @@ namespace uba
 		Logger() {}
 		bool Error(const tchar* format, ...);
 		bool Warning(const tchar* format, ...);
-		void Info(const tchar* format, ...);
+		Logger& Info(const tchar* format, ...);
 		void Detail(const tchar* format, ...);
 		void Debug(const tchar* format, ...);
 		void Logf(LogEntryType type, const tchar* format, ...);
@@ -37,6 +37,8 @@ namespace uba
 		virtual void EndScope() = 0;
 		virtual void Log(LogEntryType type, const tchar* str, u32 strLen) = 0;
 		virtual ~Logger() {}
+
+		bool ToFalse() { return false; }
 	};
 
 	class LoggerWithWriter : public Logger
diff --git a/Engine/Source/Programs/UnrealBuildAccelerator/Core/Public/UbaStringBuffer.h b/Engine/Source/Programs/UnrealBuildAccelerator/Core/Public/UbaStringBuffer.h
--- a/Engine/Source/Programs/UnrealBuildAccelerator/Core/Public/UbaStringBuffer.h
+++ b/Engine/Source/Programs/UnrealBuildAccelerator/Core/Public/UbaStringBuffer.h
@@ -79,6 +79,7 @@ namespace uba
 		u32 Parse(char* out, u64 outCapacity) const; // Note, return value contains length + null termination.. zero if failed
 
 		TString ToString() const { return TString(data, data + count); }
+		bool ToFalse() const { return false; }
 
 		u32 count;
 		u32 capacity;
