[UBA]
* Added thread description to posix threads

diff --git a/Engine/Source/Programs/UnrealBuildAccelerator/Common/Private/UbaThread.cpp b/Engine/Source/Programs/UnrealBuildAccelerator/Common/Private/UbaThread.cpp
--- a/Engine/Source/Programs/UnrealBuildAccelerator/Common/Private/UbaThread.cpp
+++ b/Engine/Source/Programs/UnrealBuildAccelerator/Common/Private/UbaThread.cpp
@@ -108,6 +108,8 @@ namespace uba
 				return (void*)(uintptr_t)res;
 			}, this);
 		UBA_ASSERT(err == 0); (void)err;
+		if (description)
+			pthread_setname_np(pth, description);
 
 		err = pthread_attr_destroy(&tattr);
 		UBA_ASSERT(err == 0); (void)err;
