[UBA]
* Cleaned up code to use StringView instead of StringBuffer

diff --git a/Engine/Source/Programs/UnrealBuildAccelerator/Common/Private/UbaLogger.cpp b/Engine/Source/Programs/UnrealBuildAccelerator/Common/Private/UbaLogger.cpp
--- a/Engine/Source/Programs/UnrealBuildAccelerator/Common/Private/UbaLogger.cpp
+++ b/Engine/Source/Programs/UnrealBuildAccelerator/Common/Private/UbaLogger.cpp
@@ -626,9 +626,7 @@ namespace uba
 				continue;
 			if (mod.name.empty())
 				mod.name = executable;
-			StringBuffer<128> name(mod.name);
-			if (auto lastSlash = name.Last(PathSeparator))
-				name.Clear().Append(lastSlash + 1);
+			StringView name = StringView(mod.name).GetFileName();
 			for (auto it=searchPaths; it->count; ++it)
 			{
 				struct stat attr;
