[UBA]
* Added StringView::GetPath

diff --git a/Engine/Source/Programs/UnrealBuildAccelerator/Core/Private/UbaStringBuffer.cpp b/Engine/Source/Programs/UnrealBuildAccelerator/Core/Private/UbaStringBuffer.cpp
--- a/Engine/Source/Programs/UnrealBuildAccelerator/Core/Private/UbaStringBuffer.cpp
+++ b/Engine/Source/Programs/UnrealBuildAccelerator/Core/Private/UbaStringBuffer.cpp
@@ -438,6 +438,13 @@ namespace uba
 		return uba::EndsWith(data, count, value, TStrlen(value), ignoreCase);
 	}
 
+	StringView StringView::GetPath() const
+	{
+		if (const tchar* lastSeparator = TStrrchr(data, PathSeparator))
+			return StringView(data, u32(lastSeparator - data));
+		return {};
+	}
+
 	StringView ToView(const tchar* s)
 	{
 		return StringView(s, TStrlen(s));
diff --git a/Engine/Source/Programs/UnrealBuildAccelerator/Core/Public/UbaStringBuffer.h b/Engine/Source/Programs/UnrealBuildAccelerator/Core/Public/UbaStringBuffer.h
--- a/Engine/Source/Programs/UnrealBuildAccelerator/Core/Public/UbaStringBuffer.h
+++ b/Engine/Source/Programs/UnrealBuildAccelerator/Core/Public/UbaStringBuffer.h
@@ -123,6 +123,7 @@ namespace uba
 		bool Equals(const tchar* str, bool ignoreCase = true) const { return uba::Equals(data, str, ignoreCase); }
 		bool Equals(const StringView& str, bool ignoreCase = true) const { return count == str.count && uba::Equals(data, str.data, count, ignoreCase); }
 		StringView GetFileName() const { const tchar* str = uba::GetFileName(data); return StringView(str, count - u32(str - data)); }
+		StringView GetPath() const;
 		StringView Skip(u32 skipCount) { return StringView(data + skipCount, count - skipCount); }
 
 		TString ToString() const { return TString(data, data + count); }
