[UBA]
* Changed so UBA_API for import adds __attribute__((weak)) which makes it possible to delay load uba host dll

diff --git a/Engine/Source/Programs/UnrealBuildAccelerator/Core/Public/UbaBase.h b/Engine/Source/Programs/UnrealBuildAccelerator/Core/Public/UbaBase.h
--- a/Engine/Source/Programs/UnrealBuildAccelerator/Core/Public/UbaBase.h
+++ b/Engine/Source/Programs/UnrealBuildAccelerator/Core/Public/UbaBase.h
@@ -45,7 +45,9 @@ namespace uba
 #if !defined(UBA_API)
 	#if PLATFORM_WINDOWS
 		#define UBA_API __declspec(dllimport)
+	#elif PLATFORM_LINUX
+		#define UBA_API __attribute__((weak))
 	#else
-		#define UBA_API
+		#define UBA_API 
 	#endif
 #endif
