[UBA]
* Added Rpc_GetFullFileName2 which returns both virtual and real path of file

diff --git a/Engine/Source/Programs/UnrealBuildAccelerator/Detours/Private/UbaDetoursFileMappingTable.cpp b/Engine/Source/Programs/UnrealBuildAccelerator/Detours/Private/UbaDetoursFileMappingTable.cpp
--- a/Engine/Source/Programs/UnrealBuildAccelerator/Detours/Private/UbaDetoursFileMappingTable.cpp
+++ b/Engine/Source/Programs/UnrealBuildAccelerator/Detours/Private/UbaDetoursFileMappingTable.cpp
@@ -359,6 +359,43 @@ namespace uba
 		pathLen = tempBuf.count;
 	}
 
+	void Rpc_GetFullFileName2(const tchar* path, StringBufferBase& outReal, StringBufferBase& outVirtual, const tchar* const* loaderPaths)
+	{
+		StringKey fileNameKey;
+		StringBuffer<> temp2;
+		if (IsAbsolutePath(path))
+		{
+			FixPath(temp2, path);
+			path = temp2.data;
+			fileNameKey = CaseInsensitiveFs ? ToStringKeyLower(temp2) : ToStringKey(temp2);
+		}
+
+		u32 mappedFileTableSize;
+
+		{
+			RPC_MESSAGE(GetFullFileName, getFullFileName)
+			writer.WriteString(path);
+			writer.WriteStringKey(fileNameKey);
+			u16& bytes = *(u16*)writer.AllocWrite(2);
+			auto pos = writer.GetPosition();
+			if (loaderPaths)
+				for (auto i=loaderPaths; *i; ++i)
+					writer.WriteString(*i);
+			bytes = u16(writer.GetPosition() - pos);
+			writer.Flush();
+			BinaryReader reader;
+			reader.ReadString(outReal);
+			reader.ReadString(outVirtual);
+			mappedFileTableSize = reader.ReadU32();
+			DEBUG_LOG_PIPE(TC("GetFileName"), TC("(%ls)"), tempBuf);
+		}
+
+		#if UBA_DEBUG
+		DEBUG_LOG_DETOURED(TC("Rpc_GetFullFileName"), TC("%s -> %s (%s)"), path, outReal.data, outVirtual.data);
+		#endif
+		g_mappedFileTable.Parse(mappedFileTableSize);
+	}
+
 	DirHash::DirHash(const StringView& str)
 	{
 		CHECK_PATH(str);
diff --git a/Engine/Source/Programs/UnrealBuildAccelerator/Detours/Private/UbaDetoursFileMappingTable.h b/Engine/Source/Programs/UnrealBuildAccelerator/Detours/Private/UbaDetoursFileMappingTable.h
--- a/Engine/Source/Programs/UnrealBuildAccelerator/Detours/Private/UbaDetoursFileMappingTable.h
+++ b/Engine/Source/Programs/UnrealBuildAccelerator/Detours/Private/UbaDetoursFileMappingTable.h
@@ -95,6 +95,7 @@ namespace uba
 	void Rpc_UpdateTables();
 	u32  Rpc_GetEntryOffset(const StringKey& entryNameKey, const tchar* entryName, u64 entryNameLen, bool checkIfDir = false);
 	void Rpc_GetFullFileName(const tchar*& path, u64& pathLen, StringBufferBase& tempBuf, bool useVirtualName, const tchar* const* loaderPaths = nullptr);
+	void Rpc_GetFullFileName2(const tchar* path, StringBufferBase& outReal, StringBufferBase& outVirtual, const tchar* const* loaderPaths = nullptr);
 	void Rpc_GetWrittenFiles();
 
 	struct DirHash
