[UBA]
* Minor logging description improvement

diff --git a/Engine/Source/Programs/UnrealBuildAccelerator/Common/Private/UbaStorage.cpp b/Engine/Source/Programs/UnrealBuildAccelerator/Common/Private/UbaStorage.cpp
--- a/Engine/Source/Programs/UnrealBuildAccelerator/Common/Private/UbaStorage.cpp
+++ b/Engine/Source/Programs/UnrealBuildAccelerator/Common/Private/UbaStorage.cpp
@@ -1072,7 +1072,7 @@ namespace uba
 							BinaryReader reader(fa.GetData(), 0, fa.GetSize());
 							if (reader.GetLeft() < 12)
 							{
-								m_logger.Detail(TC("Corrupt cas. Is %llu, must be at least 12 bytes (%s)"), reader.GetLeft(), e.name);
+								m_logger.Detail(TC("Corrupt cas. Is %llu bytes, must be at least 12 (%s)"), reader.GetLeft(), e.name);
 								deleteFile = true;
 							}
 							else
