[UBA]
* Implemented support for loading executables and dlls remotely on linux. This is needed for ShaderCompileWorker which is a modular build loading lots of dlls
* Changed signature of FindImports

diff --git a/Engine/Source/Programs/UnrealBuildAccelerator/Common/Private/UbaSession.cpp b/Engine/Source/Programs/UnrealBuildAccelerator/Common/Private/UbaSession.cpp
--- a/Engine/Source/Programs/UnrealBuildAccelerator/Common/Private/UbaSession.cpp
+++ b/Engine/Source/Programs/UnrealBuildAccelerator/Common/Private/UbaSession.cpp
@@ -1417,7 +1417,7 @@ namespace uba
 		out.push_back({ library, temp3.data, attr, isSystem });
 
 		StringBuffer<> errorStr;
-		FindImports(applicationName, [&](const tchar* importName, bool isKnown, const char* const* importLoaderPaths)
+		FindImports(applicationName, ToView(applicationDir), [&](const tchar* importName, bool isKnown, const char* const* importLoaderPaths)
 			{
 				if (result && !isKnown)
 					result = CopyImports(out, importName, applicationDir, applicationDirEnd, handledImports, importLoaderPaths);
diff --git a/Engine/Source/Programs/UnrealBuildAccelerator/Common/Private/UbaSessionServer.cpp b/Engine/Source/Programs/UnrealBuildAccelerator/Common/Private/UbaSessionServer.cpp
--- a/Engine/Source/Programs/UnrealBuildAccelerator/Common/Private/UbaSessionServer.cpp
+++ b/Engine/Source/Programs/UnrealBuildAccelerator/Common/Private/UbaSessionServer.cpp
@@ -879,7 +879,9 @@ namespace uba
 			for (auto& loaderPath : loaderPaths)
 			{
 				StringBuffer<> fullPath;
-				fullPath.Append(applicationDir).EnsureEndsWithSlash().Append(loaderPath).EnsureEndsWithSlash().Append(fileName);
+				if (loaderPath[0] != '/')
+					fullPath.Append(applicationDir).EnsureEndsWithSlash();
+				fullPath.Append(loaderPath).EnsureEndsWithSlash().Append(fileName);
 				if (GetFileAttributesW(fullPath.data) == INVALID_FILE_ATTRIBUTES)
 					continue;
 				FixPath(fullPath.data, nullptr, 0, absoluteFile);
@@ -889,6 +891,15 @@ namespace uba
 					return false;
 				break;
 			}
+
+			if (casKey == CasKeyZero)
+			{
+				m_logger.Warning(TC("Failed to find file %s"), fileName.data);
+				for (auto& loaderPath : loaderPaths)
+				{
+					m_logger.Warning(TC("   LoaderPath %s"), loaderPath.c_str());
+				}
+			}
 		}
 		else if (SearchPathForFile(m_logger, absoluteFile, fileName.data, applicationDir.data))
 		{
diff --git a/Engine/Source/Programs/UnrealBuildAccelerator/Core/Public/UbaLinuxBinDependencyParser.h b/Engine/Source/Programs/UnrealBuildAccelerator/Core/Public/UbaLinuxBinDependencyParser.h
--- a/Engine/Source/Programs/UnrealBuildAccelerator/Core/Public/UbaLinuxBinDependencyParser.h
+++ b/Engine/Source/Programs/UnrealBuildAccelerator/Core/Public/UbaLinuxBinDependencyParser.h
@@ -2,21 +2,240 @@
 
 #pragma once
 
+#if PLATFORM_LINUX
+
 #include "UbaStringBuffer.h"
+#include <elf.h>
 
 #if !defined(UBA_IS_DETOURED_INCLUDE)
 #define TRUE_WRAPPER(func) func
 #endif
 
+#define UBA_FINDIMPORTS_DEBUG 0
+
 namespace uba
 {
+	constexpr StringView g_systemFiles[] =
+	{
+		TCV("libstdc++.so"),
+		TCV("libpthread.so"),
+		TCV("ld-linux-x86-64.so"),
+		TCV("librt.so"),
+		TCV("libdl.so"),
+		TCV("libc.so"),
+		TCV("libdbus-1.so"),
+		TCV("libgcc_s.so"),
+		TCV("libm.so"),
+		TCV("libX11.so"),
+		TCV("libXext.so"),
+		TCV("libXcursor.so"),
+		TCV("libXi.so"),
+		TCV("libXfixes.so"),
+		TCV("libXrandr.so"),
+		TCV("libXss.so"),
+		TCV("libudev.so"),
+	};
+
 	inline bool IsKnownSystemFile(const tchar* fileName)
 	{
+		StringView file = ToView(fileName);
+		bool isSystem = false;
+		for (auto systemFile : g_systemFiles)
+			if (file.StartsWith(systemFile.data))
+				return true;
 		return false;
 	}
 
-	inline bool FindImports(const tchar* fileName, const Function<void(const tchar* import, bool isKnown, const char* const* loaderPaths)>& func, StringBufferBase& outError)
+	bool FindImports(const tchar* fileName, const StringView& originalPath, const Function<void(const tchar* import, bool isKnown, const char* const* loaderPaths)>& func, StringBufferBase& outError)
 	{
+		#if UBA_FINDIMPORTS_DEBUG
+		printf("FINDIMPORTS: %s\n", fileName);
+		#endif
+
+		int fd = TRUE_WRAPPER(open)(fileName, O_RDONLY);
+		if (fd == -1)
+			return outError.Appendf("Open failed for file %s", fileName).ToFalse();
+
+		auto closeFileHandle = MakeGuard([&]() { TRUE_WRAPPER(close)(fd); });
+		struct stat sb;
+		if (TRUE_WRAPPER(fstat)(fd, &sb) == -1)
+			return outError.Appendf("Stat failed for file %s", fileName).ToFalse();
+
+		u32 size = sb.st_size;
+
+		auto mem = (u8*)mmap(nullptr, size, PROT_READ, MAP_PRIVATE, fd, 0);
+		if (mem == MAP_FAILED)
+			return outError.Appendf("Mmap failed for file %s", fileName).ToFalse();
+		auto unmap = MakeGuard([&]() { munmap(mem, size); });
+
+		auto ehdr = (Elf64_Ehdr*)mem;
+		if (memcmp(ehdr->e_ident, ELFMAG, SELFMAG) != 0 || ehdr->e_ident[EI_CLASS] != ELFCLASS64)
+			return outError.Append("Not a valid 64-bit ELF file").ToFalse();
+
+		Elf64_Phdr* phdrs = (Elf64_Phdr*)(mem + ehdr->e_phoff);
+		Elf64_Shdr* shdrs = (Elf64_Shdr *)(mem + ehdr->e_shoff);
+		Elf64_Dyn* dynamic = NULL;
+		size_t dyn_count = 0;
+
+		auto vaddr_to_offset = [](Elf64_Addr vaddr, Elf64_Phdr *phdrs, int phnum) -> Elf64_Off {
+			for (int i = 0; i < phnum; i++) {
+				if (phdrs[i].p_type != PT_LOAD) continue;
+				Elf64_Addr start = phdrs[i].p_vaddr;
+				Elf64_Addr end = start + phdrs[i].p_memsz;
+				if (vaddr >= start && vaddr < end) {
+					return phdrs[i].p_offset + (vaddr - start);
+				}
+			}
+			return 0;
+		};
+
+		for (int i = 0; i < ehdr->e_phnum; ++i)
+		{
+			if (phdrs[i].p_type != PT_DYNAMIC)
+				continue;
+			Elf64_Off dyn_offset = vaddr_to_offset(phdrs[i].p_vaddr, phdrs, ehdr->e_phnum);
+			if (dyn_offset >= sb.st_size || dyn_offset + sizeof(Elf64_Dyn) > sb.st_size)
+				return outError.Appendf("Dynamic offset out of file bounds!").ToFalse();
+			dynamic = (Elf64_Dyn*)(mem + dyn_offset);			
+			dyn_count = phdrs[i].p_filesz / sizeof(Elf64_Dyn);
+			break;
+		}
+
+		if (!dynamic)
+			return outError.Appendf("No PT_DYNAMIC segment found.").ToFalse();
+		if ((u8*)dynamic < mem || (u8*)dynamic >= mem + sb.st_size)
+			return outError.Appendf("Dynamic segment is out of file bounds!").ToFalse();
+
+		Elf64_Addr strtab_addr = 0;
+		Elf64_Xword strsz = 0;
+		Elf64_Xword needed[512];
+		u32 needed_count = 0;
+
+		for (size_t i = 0; i < dyn_count; ++i)
+		{
+			auto& dyn = dynamic[i];
+			if (dyn.d_tag == DT_NULL)
+				break;
+			if (dyn.d_tag == DT_STRTAB)
+				strtab_addr = dyn.d_un.d_ptr;
+			else if (dyn.d_tag == DT_STRSZ)
+				strsz = dyn.d_un.d_val;
+			else if (dyn.d_tag == DT_NEEDED && needed_count < sizeof_array(needed))
+				needed[needed_count++] = dyn.d_un.d_val;
+		}
+
+		char* strtab = NULL;
+
+		// Step 3: If DT_STRTAB is missing, fallback to section header lookup
+		if (!strtab_addr && ehdr->e_shoff != 0) {
+			for (int i = 0; i < ehdr->e_shnum; i++) {
+				if (shdrs[i].sh_type == SHT_STRTAB && shdrs[i].sh_flags == 0) { // Likely .dynstr
+					strtab = (char *)mem + shdrs[i].sh_offset;
+					strsz = shdrs[i].sh_size;
+					break;
+				}
+			}
+		}
+
+		// If we have strtab_addr but not strtab, resolve via vaddr to file offset
+		if (strtab_addr && !strtab)
+		{
+			Elf64_Off str_off = vaddr_to_offset(strtab_addr, phdrs, ehdr->e_phnum);
+			strtab = (char *)mem + str_off;
+		}
+
+		if (!strtab)
+			return outError.Appendf("Failed to find string table.").ToFalse();
+
+		const char* loaderPaths[256];
+		u32 loaderPathsCount = 0;
+
+		List<TString> fixedLoaderPaths;
+
+		#if UBA_FINDIMPORTS_DEBUG
+		printf("ELF: %s\n", fileName);
+		#endif
+
+		#if UBA_FINDIMPORTS_DEBUG
+		printf("RPATH / RUNPATH:\n");
+		#endif
+
+		
+		//StringView origin(fileName, strrchr(fileName, '/') - fileName);
+
+		for (Elf64_Dyn *dyn = dynamic; dyn->d_tag != DT_NULL; dyn++)
+		{
+			if (dyn->d_tag != DT_RPATH && dyn->d_tag != DT_RUNPATH)
+				continue;
+
+			#if UBA_FINDIMPORTS_DEBUG
+			if (dyn->d_tag == DT_RPATH)
+				printf("  DT_RPATH: %s\n", strtab + dyn->d_un.d_val);
+			else if (dyn->d_tag == DT_RUNPATH)
+				printf("  DT_RUNPATH: %s\n", strtab + dyn->d_un.d_val);
+			#endif
+
+			char* begin = strtab + dyn->d_un.d_val;
+			char* end = begin + strlen(begin);
+			char* pathBegin = begin;
+			while (pathBegin < end)
+			{
+				char* pathEnd = strchr(pathBegin, ':');
+				if (!pathEnd)
+					pathEnd = end;
+				StringView loaderPath(pathBegin, pathEnd - pathBegin);
+				if (loaderPath.StartsWith("${ORIGIN}"))
+					loaderPaths[loaderPathsCount++] = fixedLoaderPaths.emplace_back(originalPath.ToString() + loaderPath.Skip(9).ToString()).c_str();
+				else
+					loaderPaths[loaderPathsCount++] = fixedLoaderPaths.emplace_back(loaderPath.ToString()).c_str();
+
+				#if UBA_FINDIMPORTS_DEBUG
+				printf("  SEARCH_PATH: %s\n", loaderPaths[loaderPathsCount-1]);
+				#endif
+
+				pathBegin = pathEnd + 1;
+			}
+		}
+		loaderPaths[loaderPathsCount] = nullptr;
+
+		#if UBA_FINDIMPORTS_DEBUG
+		printf("DT_NEEDED:\n");
+		#endif
+		
+		#if 0
+		for (Elf64_Dyn *dyn = dynamic; dyn->d_tag != DT_NULL; dyn++) {
+			if (dyn->d_tag != DT_NEEDED)
+				continue;
+			printf("  %s\n", strtab + dyn->d_un.d_val);
+			func(strtab + dyn->d_un.d_val, false, loaderPaths);
+		}
+		#endif
+
+		for (int i = 0; i < needed_count; i++) {
+			if (needed[i] < strsz || strsz == 0)
+			{
+				StringView file = ToView(strtab + needed[i]);
+				bool isSystem = false;
+				for (auto systemFile : g_systemFiles)
+					isSystem |= file.StartsWith(systemFile.data);
+				if (isSystem)
+					continue;
+
+				#if UBA_FINDIMPORTS_DEBUG
+				printf("  %s\n", strtab + needed[i]);
+				#endif
+
+				func(file.data, false, loaderPaths);
+			}
+			else
+			{
+				#if UBA_FINDIMPORTS_DEBUG
+				printf("  [offset %lu out of bounds]\n", needed[i]);
+				#endif
+			}
+		}		
 		return true;
 	}
-}
\ No newline at end of file
+}
+
+#endif
diff --git a/Engine/Source/Programs/UnrealBuildAccelerator/Core/Public/UbaMacBinDependencyParser.h b/Engine/Source/Programs/UnrealBuildAccelerator/Core/Public/UbaMacBinDependencyParser.h
--- a/Engine/Source/Programs/UnrealBuildAccelerator/Core/Public/UbaMacBinDependencyParser.h
+++ b/Engine/Source/Programs/UnrealBuildAccelerator/Core/Public/UbaMacBinDependencyParser.h
@@ -17,7 +17,7 @@ namespace uba
 		return false;
 	}
 
-	inline bool FindImports(const tchar* fileName, const Function<void(const tchar* import, bool isKnown, const char* const* loaderPaths)>& func, StringBufferBase& outError)
+	inline bool FindImports(const tchar* fileName, const StringView& originalPath, const Function<void(const tchar* import, bool isKnown, const char* const* loaderPaths)>& func, StringBufferBase& outError)
 	{
 		int fd = TRUE_WRAPPER(open)(fileName, O_RDONLY);
 		if (fd == -1)
diff --git a/Engine/Source/Programs/UnrealBuildAccelerator/Core/Public/UbaWinBinDependencyParser.h b/Engine/Source/Programs/UnrealBuildAccelerator/Core/Public/UbaWinBinDependencyParser.h
--- a/Engine/Source/Programs/UnrealBuildAccelerator/Core/Public/UbaWinBinDependencyParser.h
+++ b/Engine/Source/Programs/UnrealBuildAccelerator/Core/Public/UbaWinBinDependencyParser.h
@@ -129,7 +129,7 @@ namespace uba
 		return true;
 	}
 
-	inline bool FindImports(const tchar* fileName, const Function<void(const tchar* import, bool isKnown, const char* const* importLoaderPaths)>& func, StringBufferBase& outError)
+	inline bool FindImports(const tchar* fileName, const StringView& originalPath, const Function<void(const tchar* import, bool isKnown, const char* const* importLoaderPaths)>& func, StringBufferBase& outError)
 	{
 		if (ToView(fileName).EndsWith(TCV(".bat")))
 			return true;
diff --git a/Engine/Source/Programs/UnrealBuildAccelerator/Detours/Private/PosixOS/UbaDetoursFunctionsPosix.cpp b/Engine/Source/Programs/UnrealBuildAccelerator/Detours/Private/PosixOS/UbaDetoursFunctionsPosix.cpp
--- a/Engine/Source/Programs/UnrealBuildAccelerator/Detours/Private/PosixOS/UbaDetoursFunctionsPosix.cpp
+++ b/Engine/Source/Programs/UnrealBuildAccelerator/Detours/Private/PosixOS/UbaDetoursFunctionsPosix.cpp
@@ -307,6 +307,9 @@ bool CanDetour2(const StringView& file)
 	return g_isDetouring && !t_disallowDetour
 		&& !file.StartsWith("/dev/")
 		&& !file.StartsWith("/etc/")
+	#if PLATFORM_LINUX
+		&& !file.StartsWith("/sys/") // Don't know if this is needed for macos.. but is needed for linux
+	#endif
 		&& !file.StartsWith(g_systemTemp.data)
 		;
 }
@@ -320,12 +323,24 @@ int Shared_open(const char* funcName, const char* file, int flags, int mode, con
 	FixPath(fileName, file);
 	DevirtualizePath(fileName);
 
-	if (!CanDetour2(fileName))
+	if (!CanDetour2(fileName) || fileName.Equals("/"))
+	{
+		int res = trueOpen(file, flags, mode);
+		DEBUG_LOG_TRUE(funcName, "NODETOUR (%s) -> %i", file, res);
+		return res;
+	}
+
+	#if PLATFORM_LINUX
+	if (fileName.StartsWith("/proc/"))
 	{
+		if (fileName.StartsWith("/proc/self/cmdline"))
+			DEBUG_LOG("TODO!!! /proc/self/cmdline");
 		int res = trueOpen(file, flags, mode);
 		DEBUG_LOG_TRUE(funcName, "NODETOUR (%s) -> %i", file, res);
 		return res;
 	}
+	#endif
+
 
 	StringKey fileNameKey = ToFilenameKey(fileName);
 
@@ -1753,6 +1768,7 @@ UBA_EXPORT ssize_t UBA_WRAPPER(pread)(int __fd, void * __buf, size_t __nbyte, of
 
 UBA_EXPORT ssize_t UBA_WRAPPER(read)(int fd, void *buf, size_t nbyte)
 {
+	UBA_INIT_DETOUR(read, fd, buf, nbyte)
 	// char filePath[PATH_MAX];
 	// if (fcntl(fd, F_GETPATH, filePath) != -1)
 	// {
@@ -2209,12 +2225,14 @@ Set<TString> g_handledLibraries;
 
 void Shared_LoadLibrary(const char*& path, const char* const* loaderPaths, StringBufferBase& tempBuf)
 {
-#if PLATFORM_MAC
-	u64 nameLen = 0;
-	Rpc_GetFullFileName(path, nameLen, tempBuf, false, loaderPaths);
+	StringBuffer<512> virtualPath;
+	Rpc_GetFullFileName2(path, tempBuf, virtualPath, loaderPaths);
+	path = tempBuf.data;
+
+	StringView originalPath = StringView(virtualPath).GetPath();
 
 	StringBuffer<> error;
-	FindImports(path, [&](const tchar* import, bool isKnown, const char* const* importLoaderPaths)
+	FindImports(path, originalPath, [&](const tchar* import, bool isKnown, const char* const* importLoaderPaths)
 	{
 		if (!g_handledLibraries.insert(import).second)
 			return;
@@ -2223,18 +2241,16 @@ void Shared_LoadLibrary(const char*& path, const char* const* loaderPaths, Strin
 	}, error);
 	if (error.count)
 		DEBUG_LOG(error.data);
-#endif
 }
 
 UBA_EXPORT void* UBA_WRAPPER(dlopen)(const char* path, int mode)
 {
 	UBA_INIT_DETOUR(dlopen, path, mode);
 
-#if PLATFORM_MAC
 	StringBuffer<> tempBuf;
-	if (g_runningRemote)
+	if (g_runningRemote && path && *path)
 	{
-		auto originalPath = path;
+	#if PLATFORM_MAC
 		if (StartsWith(path, "@rpath/"))
 		{
 			path += 7;
@@ -2247,8 +2263,16 @@ UBA_EXPORT void* UBA_WRAPPER(dlopen)(const char* path, int mode)
 			u64 nameLen = 0;
 			Rpc_GetFullFileName(path, nameLen, tempBuf, false);
 		}
+	#else
+		if (!IsKnownSystemFile(path))
+		{
+			const char* loaderPaths[] = { "/", 0 };
+			if (g_handledLibraries.insert(path).second)
+				Shared_LoadLibrary(path, loaderPaths, tempBuf);
+		}
+	#endif
 	}
-#endif
+
 	void* res = TRUE_WRAPPER(dlopen)(path, mode);
 	DEBUG_LOG_TRUE("dlopen", "%s (%i) -> 0x%x", path, mode, res);
 	return res;
@@ -2400,17 +2424,19 @@ UBA_EXPORT pid_t UBA_WRAPPER(vfork)(void)
 
 // TODO: Allow multiple popen.. hacky!
 FILE* g_activePopen; 
-StringBuffer<> g_activePopenResult;
+StringBuffer<64*1024>* g_activePopenResult;
 u32 g_activePopenReadPos = 0;
 
 UBA_EXPORT FILE* UBA_WRAPPER(popen)(const char* command, const char* type)
 {
 	UBA_INIT_DETOUR(popen, command, type);
-	DEBUG_LOG_TRUE("popen", "%s", command);
+	DEBUG_LOG_DETOURED("popen", "%s", command);
 
 	UBA_ASSERT(!g_activePopen);
 	const char* argv[] = { command, nullptr };
-	if (!ExecuteHostRun(g_activePopenResult, argv, false))
+	g_activePopenResult = new StringBuffer<64*1024>();
+
+	if (!ExecuteHostRun(*g_activePopenResult, argv, false))
 		return nullptr;
 	g_activePopenReadPos = 0;
 	g_activePopen = (FILE*)1337;
@@ -2425,33 +2451,41 @@ UBA_EXPORT char* UBA_WRAPPER(fgets)(char* str, int count, FILE* stream)
 	if (stream == g_activePopen)
 	{
 		u32 toWrite;
-		if (const char* endl = g_activePopenResult.First('\n', g_activePopenReadPos))
+		if (const char* endl = g_activePopenResult->First('\n', g_activePopenReadPos))
 		{
-			u32 lineLen = u32(endl - g_activePopenResult.data) - g_activePopenReadPos + 1;
+			u32 lineLen = u32(endl - g_activePopenResult->data) - g_activePopenReadPos + 1;
 			toWrite = Min(u32(count - 2), lineLen);
 		}
 		else
 		{
-			toWrite = Min(u32(count - 2), g_activePopenResult.count - g_activePopenReadPos);
+			toWrite = Min(u32(count - 2), g_activePopenResult->count - g_activePopenReadPos);
 		}
-		memcpy(str, g_activePopenResult.data + g_activePopenReadPos, toWrite);
+		memcpy(str, g_activePopenResult->data + g_activePopenReadPos, toWrite);
 		str[toWrite] = 0;
 		g_activePopenReadPos += toWrite;
+
+		DEBUG_LOG_DETOURED("fgets", "%s", str);
 		return str;
 	}
-	return TRUE_WRAPPER(fgets)(str, count, stream);
+
+	auto res = TRUE_WRAPPER(fgets)(str, count, stream);
+	DEBUG_LOG_TRUE("fgets", "%s", str);
+	return res;
 }
 
 UBA_EXPORT int UBA_WRAPPER(pclose)(FILE* stream)
 {
 	UBA_INIT_DETOUR(pclose, stream);
-	DEBUG_LOG_TRUE("pclose", "%p", stream);
 	if (stream == g_activePopen)
 	{
-		g_activePopenResult.Clear();
+		DEBUG_LOG_DETOURED("pclose", "%p", stream);
+		delete g_activePopenResult;
+		g_activePopenResult = nullptr;
 		g_activePopen = nullptr;
 		return 0;
 	}
+
+	DEBUG_LOG_TRUE("pclose", "%p", stream);
 	return TRUE_WRAPPER(pclose)(stream);
 }
 
diff --git a/Engine/Source/Programs/UnrealBuildAccelerator/Detours/Private/Windows/UbaDetoursFunctionsKernelBase.inl b/Engine/Source/Programs/UnrealBuildAccelerator/Detours/Private/Windows/UbaDetoursFunctionsKernelBase.inl
--- a/Engine/Source/Programs/UnrealBuildAccelerator/Detours/Private/Windows/UbaDetoursFunctionsKernelBase.inl
+++ b/Engine/Source/Programs/UnrealBuildAccelerator/Detours/Private/Windows/UbaDetoursFunctionsKernelBase.inl
@@ -2713,7 +2713,7 @@ HMODULE Recursive_LoadLibraryExW(LPCWSTR lpLibFileName, LPCWSTR originalName, DW
 	{
 		SuppressCreateFileDetourScope cfs;
 		StringBuffer<256> error;
-		if (!FindImports(lpLibFileName, [&](const wchar_t* import, bool isKnown, const char* const* importLoaderPaths)
+		if (!FindImports(lpLibFileName, {}, [&](const wchar_t* import, bool isKnown, const char* const* importLoaderPaths)
 			{
 				if (!GetModuleHandleW(import))
 					importedModules.emplace_back(import, isKnown);
diff --git a/Engine/Source/Programs/UnrealBuildAccelerator/Host/Private/UbaExports.cpp b/Engine/Source/Programs/UnrealBuildAccelerator/Host/Private/UbaExports.cpp
--- a/Engine/Source/Programs/UnrealBuildAccelerator/Host/Private/UbaExports.cpp
+++ b/Engine/Source/Programs/UnrealBuildAccelerator/Host/Private/UbaExports.cpp
@@ -896,7 +896,9 @@ extern "C"
 	{
 #if PLATFORM_WINDOWS
 		uba::StringBuffer<> errors;
-		uba::FindImports(binary, [&](const uba::tchar* importName, bool isKnown, const char* const* importLoaderPaths) { func(importName, userData); }, errors);
+		uba::FindImports(binary, {}, [&](const uba::tchar* importName, bool isKnown, const char* const* importLoaderPaths) { func(importName, userData); }, errors);
+#else
+		UBA_ASSERT(false);
 #endif
 	}
 
