[UBA]
* Added UbaCli.exe -listimports

diff --git a/Engine/Source/Programs/UnrealBuildAccelerator/Cli/Private/UbaCli.cpp b/Engine/Source/Programs/UnrealBuildAccelerator/Cli/Private/UbaCli.cpp
--- a/Engine/Source/Programs/UnrealBuildAccelerator/Cli/Private/UbaCli.cpp
+++ b/Engine/Source/Programs/UnrealBuildAccelerator/Cli/Private/UbaCli.cpp
@@ -32,6 +32,8 @@
 namespace uba
 {
 
+	bool FindImports(const tchar* fileName, const StringView& originalPath, const Function<void(const tchar* import, bool isKnown, const char* const* loaderPaths)>& func, StringBufferBase& outError);
+
 	const tchar*	Version = GetVersionString();
 	u32				DefaultCapacityGb = 20;
 	const tchar*	DefaultRootDir = []() {
@@ -85,6 +87,7 @@ namespace uba
 		logger.Info(TC("   -checkcloud             Check if we are inside cloud and output information about cloud"));
 		logger.Info(TC("   -deletecas              Deletes the casdb"));
 		logger.Info(TC("   -getcas                 Will print hash of application"));
+		logger.Info(TC("   -listimports            Will print explicit imports of binary"));
 		logger.Info(TC("   -summary                Print summary at the end of a session"));
 		logger.Info(TC("   -nocustomalloc          Disable custom allocator for processes. If you see odd crashes this can be tested"));
 		logger.Info(TC("   -nostdout               Disable stdout from process."));
@@ -165,6 +168,7 @@ namespace uba
 		bool checkCas2 = false;
 		bool checkCloud = false;
 		bool getCas = false;
+		bool listImports = false;
 		bool deleteCas = false;
 		bool enableStdOut = true;
 		bool printSummary = false;
@@ -356,6 +360,10 @@ namespace uba
 				{
 					getCas = true;
 				}
+				else if (name.Equals(TCV("-listimports")))
+				{
+					listImports = true;
+				}
 				else if (name.Equals(TCV("-summary")))
 				{
 					printSummary = true;
@@ -804,6 +812,18 @@ namespace uba
 					logger.Info(TC("  CasKey (uncompressed): %s"), CasKeyString(uncompressedKey).str);
 				return true;
 			}
+
+			if (listImports)
+			{
+				logger.Info(TC("Imports:"));
+				StringBuffer<> error;
+				if (!FindImports(application.c_str(), StringView(application).GetPath(), [&](const tchar* import, bool isKnown, const char* const* loaderPaths)
+				{
+					logger.Info(TC("  %s"), import);
+				}, error))
+					return logger.Error(TC("%s"), error.data);
+				return true;
+			}
 		}
 
 		const tchar* dbgStr = TC("");
