[UBA]
* Fixed macos compile error in UbaThread

diff --git a/Engine/Source/Programs/UnrealBuildAccelerator/Common/Private/UbaThread.cpp b/Engine/Source/Programs/UnrealBuildAccelerator/Common/Private/UbaThread.cpp
--- a/Engine/Source/Programs/UnrealBuildAccelerator/Common/Private/UbaThread.cpp
+++ b/Engine/Source/Programs/UnrealBuildAccelerator/Common/Private/UbaThread.cpp
@@ -108,8 +108,13 @@ namespace uba
 				return (void*)(uintptr_t)res;
 			}, this);
 		UBA_ASSERT(err == 0); (void)err;
-		if (description)
-			pthread_setname_np(pth, description);
+		if (!description)
+			description = "UbaUnknown";
+		#if PLATFORM_MACOS
+		pthread_setname_np(description);
+		#else
+		pthread_setname_np(pth, description);
+		#endif
 
 		err = pthread_attr_destroy(&tattr);
 		UBA_ASSERT(err == 0); (void)err;
