[UBA]
* Improved network backend logging when errors occur

diff --git a/Engine/Source/Programs/UnrealBuildAccelerator/Common/Private/UbaNetworkBackendMemory.cpp b/Engine/Source/Programs/UnrealBuildAccelerator/Common/Private/UbaNetworkBackendMemory.cpp
--- a/Engine/Source/Programs/UnrealBuildAccelerator/Common/Private/UbaNetworkBackendMemory.cpp
+++ b/Engine/Source/Programs/UnrealBuildAccelerator/Common/Private/UbaNetworkBackendMemory.cpp
@@ -49,7 +49,7 @@ namespace uba
 			cb(rc.disconnectContext, m_connection->uid, connection);
 	}
 
-	bool NetworkBackendMemory::Send(Logger& logger, void* connection, const void* data, u32 dataSize, SendContext& sendContext)
+	bool NetworkBackendMemory::Send(Logger& logger, void* connection, const void* data, u32 dataSize, SendContext& sendContext, const tchar* sendHint)
 	{
 		UBA_ASSERT(m_connection);
 		sendContext.isUsed = true;
diff --git a/Engine/Source/Programs/UnrealBuildAccelerator/Common/Private/UbaNetworkBackendQuic.cpp b/Engine/Source/Programs/UnrealBuildAccelerator/Common/Private/UbaNetworkBackendQuic.cpp
--- a/Engine/Source/Programs/UnrealBuildAccelerator/Common/Private/UbaNetworkBackendQuic.cpp
+++ b/Engine/Source/Programs/UnrealBuildAccelerator/Common/Private/UbaNetworkBackendQuic.cpp
@@ -64,7 +64,7 @@ namespace uba
 			conn.Release();
 	}
 
-	bool NetworkBackendQuic::Send(Logger& logger, void* connection, const void* data, u32 dataSize, SendContext& sendContext)
+	bool NetworkBackendQuic::Send(Logger& logger, void* connection, const void* data, u32 dataSize, SendContext& sendContext, const tchar* sendHint)
 	{
 		SendContext* sc = &sendContext;
 		UBA_ASSERT(!sc->isUsed);
diff --git a/Engine/Source/Programs/UnrealBuildAccelerator/Common/Private/UbaNetworkBackendTcp.cpp b/Engine/Source/Programs/UnrealBuildAccelerator/Common/Private/UbaNetworkBackendTcp.cpp
--- a/Engine/Source/Programs/UnrealBuildAccelerator/Common/Private/UbaNetworkBackendTcp.cpp
+++ b/Engine/Source/Programs/UnrealBuildAccelerator/Common/Private/UbaNetworkBackendTcp.cpp
@@ -43,6 +43,8 @@
 #define WSAPoll poll
 #endif
 
+#define UBA_LOG_SOCKET_ERRORS UBA_DEBUG
+
 namespace uba
 {
 	constexpr u32 MaxHeaderSize = 24;
@@ -98,7 +100,7 @@ namespace uba
 	bool SetRecvBuf(Logger& logger, SOCKET socket, u32 windowSize);
 	bool SetSendBuf(Logger& logger, SOCKET socket, u32 windowSize);
 	bool DisableNagle(Logger& logger, SOCKET socket);
-	bool SendSocket(Logger& logger, SOCKET socket, const void* b, u64 bufferLen);
+	bool SendSocket(Logger& logger, SOCKET socket, const void* b, u64 bufferLen, const tchar* hint);
 
 	struct RecvCache
 	{
@@ -199,7 +201,7 @@ namespace uba
 		ShutdownSocket(conn.logger, conn.socket, TC("Shutdown"));
 	}
 
-	bool NetworkBackendTcp::Send(Logger& logger, void* connection, const void* data, u32 dataSize, SendContext& sendContext)
+	bool NetworkBackendTcp::Send(Logger& logger, void* connection, const void* data, u32 dataSize, SendContext& sendContext, const tchar* sendHint)
 	{
 		auto& conn = *(Connection*)connection;
 		sendContext.isUsed = true;
@@ -207,7 +209,7 @@ namespace uba
 		SCOPED_FUTEX(conn.sendLock, lock);
 		if (conn.socket == INVALID_SOCKET)
 			return false;
-		bool res = SendSocket(logger, conn.socket, data, dataSize);
+		bool res = SendSocket(logger, conn.socket, data, dataSize, sendHint);
 		lock.Leave();
 
 		sendContext.isFinished = true;
@@ -833,7 +835,7 @@ namespace uba
 		return true;
 	}
 
-	bool SendSocket(Logger& logger, SOCKET socket, const void* b, u64 bufferLen)
+	bool SendSocket(Logger& logger, SOCKET socket, const void* b, u64 bufferLen, const tchar* hint)
 	{
 		char* buffer = (char*)b;
 		u64 left = bufferLen;
@@ -854,24 +856,21 @@ namespace uba
 					int res = WSAPoll(&p, 1, int(timeoutMs));
 					if (!res)
 					{
-						logger.Info(TC("WSAPoll returned timeout for connection trying to send data"));
+						logger.Info(TC("SendSocket - WSAPoll returned 10 minutes timeout for socket %i (%s)"), socket, hint);
 						return false;
 					}
 					if (res == SOCKET_ERROR)
 					{
-	#if UBA_DEBUG && PLATFORM_WINDOWS
-						// When cancelling all kinds of errors can happen..
-						lastError = WSAGetLastError();
-						if (lastError != WSAEINTR && lastError != WSAESHUTDOWN && lastError != WSAECONNABORTED && lastError != WSAECONNRESET) // Interrupted by cancel
-							logger.Warning(TC("WSAPoll returned an error for connection"));
-	#endif
+						#if UBA_LOG_SOCKET_ERRORS
+						logger.Info(TC("SendSocket - WSAPoll returned an error for socket %i: %s (%s)"), socket, LastErrorToText(WSAGetLastError()).data, hint);
+						#endif
 						return false;
 					}
 					continue;
 				}
-				//#if UBA_DEBUG
-				//logger.Warning(TC("ERROR sending socket (error: %s)"), LastErrorToText(WSAGetLastError()).data);
-				//#endif
+				#if UBA_LOG_SOCKET_ERRORS
+				logger.Info(TC("SendSocket - send returned an error for socket %i: %s (%s)"), socket, LastErrorToText(WSAGetLastError()).data, hint);
+				#endif
 				return false;
 			}
 
@@ -948,12 +947,16 @@ namespace uba
 			}
 
 			if (read == 0)
+			{
+				#if UBA_LOG_SOCKET_ERRORS
+				logger.Info(TC("RecvSocket - recv gracefully closed by peer for socket %i and connection %s (%s%s)"), socket, GuidToString(connection).str, hint1, hint2);
+				#endif
 				return false;
+			}
 
 			UBA_ASSERT(read == SOCKET_ERROR);
 
-			int lastError = WSAGetLastError();
-			if (lastError == WSAEWOULDBLOCK)
+			if (WSAGetLastError() == WSAEWOULDBLOCK)
 			{
 				WSAPOLLFD p;
 				p.fd = socket;
@@ -962,31 +965,26 @@ namespace uba
 				int res = WSAPoll(&p, 1, int(timeoutMs));
 				if (!res)
 				{
-					logger.Info(TC("WSAPoll returned timeout for connection %s after %s (%s%s)"), GuidToString(connection).str, TimeToText(MsToTime(timeoutMs)).str, hint1, hint2);
+					logger.Info(TC("RecvSocket - WSAPoll returned timeout for socket %i and connection %s after %s (%s%s)"), socket, GuidToString(connection).str, TimeToText(MsToTime(timeoutMs)).str, hint1, hint2);
 					return false;
 				}
 				if (res == SOCKET_ERROR)
 				{
-					#if UBA_DEBUG && PLATFORM_WINDOWS
-					// When cancelling all kinds of errors can happen..
-					lastError = WSAGetLastError();
-					if (lastError != WSAEINTR && lastError != WSAESHUTDOWN && lastError != WSAECONNABORTED && lastError != WSAECONNRESET) // Interrupted by cancel
-						logger.Warning(TC("WSAPoll returned an error for connection %s: %s (%s%s)"), GuidToString(connection).str, LastErrorToText(lastError).data, hint1, hint2);
+					#if UBA_LOG_SOCKET_ERRORS
+					logger.Info(TC("RecvSocket - WSAPoll returned an error for socket %i and connection %s: %s (%s%s)"), socket, GuidToString(connection).str, LastErrorToText(WSAGetLastError()).data, hint1, hint2);
 					#endif
 					return false;
 				}
 				continue;
 			}
-#if PLATFORM_WINDOWS
-#if UBA_DEBUG
-			// When cancelling all kinds of errors can happen..
-			if (lastError != WSAEINTR && lastError != WSAESHUTDOWN && lastError != WSAECONNABORTED && lastError != WSAECONNRESET) // Interrupted by cancel
-				logger.Warning(TC("ERROR receiving socket: %s (%s%s)"), LastErrorToText(lastError).data, hint1, hint2);
-#endif
-#else
+			#if !PLATFORM_WINDOWS
 			if (!isFirstCall && errno != ECONNRESET)
-				logger.Error(TC("ERROR receiving socket %i after %u bytes (%s%s) (%s)"), socket, bufferLen, hint1, hint2, strerror(errno));
-#endif
+				return logger.Error(TC("RecvSocket - error receiving socket %i after %u bytes (%s%s) (%s)"), socket, bufferLen, hint1, hint2, strerror(errno));
+			#endif
+
+			#if UBA_LOG_SOCKET_ERRORS
+			logger.Info(TC("RecvSocket - read returned an error for socket %i and connection %s: %s (%s%s)"), socket, GuidToString(connection).str, LastErrorToText(WSAGetLastError()).data, hint1, hint2);
+			#endif
 			return false;
 		}
 		return true;
diff --git a/Engine/Source/Programs/UnrealBuildAccelerator/Common/Private/UbaNetworkClient.cpp b/Engine/Source/Programs/UnrealBuildAccelerator/Common/Private/UbaNetworkClient.cpp
--- a/Engine/Source/Programs/UnrealBuildAccelerator/Common/Private/UbaNetworkClient.cpp
+++ b/Engine/Source/Programs/UnrealBuildAccelerator/Common/Private/UbaNetworkClient.cpp
@@ -142,18 +142,18 @@ namespace uba
 				return false;
 
 			NetworkBackend::SendContext handskakeContext;
-			if (!backend.Send(m_logger, backendConnection, encryptedBuffer, sizeof(EncryptionHandshakeString), handskakeContext))
+			if (!backend.Send(m_logger, backendConnection, encryptedBuffer, sizeof(EncryptionHandshakeString), handskakeContext, TC("CryptoHandshake")))
 				return false;
 		}
 
 
 		u32 version = SystemNetworkVersion;
 		NetworkBackend::SendContext versionContext;
-		if (!backend.Send(m_logger, backendConnection, &version, sizeof(version), versionContext))
+		if (!backend.Send(m_logger, backendConnection, &version, sizeof(version), versionContext, TC("Version")))
 			return false;
 
 		NetworkBackend::SendContext uidContext;
-		if (!backend.Send(m_logger, backendConnection, &m_uid, sizeof(m_uid), uidContext))
+		if (!backend.Send(m_logger, backendConnection, &m_uid, sizeof(m_uid), uidContext, TC("Uid")))
 			return false;
 
 		if (!rc.recvEvent.IsSet(5000)) // This can not happen. Since both callbacks are using rc we can't leave this function until we know we are not in the callbacks
@@ -616,7 +616,7 @@ namespace uba
 
 		{
 			TimerScope ts(m_sendTimer);
-			if (!connection.backend->Send(m_logger, connection.backendConnection, data, sendSize, message.m_sendContext))
+			if (!connection.backend->Send(m_logger, connection.backendConnection, data, sendSize, message.m_sendContext, TC("Message")))
 			{
 				message.m_error = 9;
 				OnDisconnected(connection, 9);
diff --git a/Engine/Source/Programs/UnrealBuildAccelerator/Common/Private/UbaNetworkServer.cpp b/Engine/Source/Programs/UnrealBuildAccelerator/Common/Private/UbaNetworkServer.cpp
--- a/Engine/Source/Programs/UnrealBuildAccelerator/Common/Private/UbaNetworkServer.cpp
+++ b/Engine/Source/Programs/UnrealBuildAccelerator/Common/Private/UbaNetworkServer.cpp
@@ -175,7 +175,7 @@ namespace uba
 			*data = value;
 			*(Guid*)(data+1) = m_server.m_uid;
 			NetworkBackend::SendContext context(NetworkBackend::SendFlags_Async);
-			return m_backend.Send(m_server.m_logger, m_backendConnection, data, 1 + sizeof(Guid), context);
+			return m_backend.Send(m_server.m_logger, m_backendConnection, data, 1 + sizeof(Guid), context, TC("UidResponse"));
 		}
 
 		static bool ReceiveHandshakeHeader(void* context, const Guid& connectionUid, u8* headerData, void*& outBodyContext, u8*& outBodyData, u32& outBodySize)
@@ -371,11 +371,11 @@ namespace uba
 			return true;
 		}
 
-		void Send(const void* data, u32 bytes)
+		void Send(const void* data, u32 bytes, const tchar* sendHint)
 		{
 			TimerScope ts(m_sendTimer);
 			NetworkBackend::SendContext context;
-			if (!m_backend.Send(m_server.m_logger, m_backendConnection, data, bytes, context))
+			if (!m_backend.Send(m_server.m_logger, m_backendConnection, data, bytes, context, sendHint))
 				SetShouldDisconnect();
 		}
 
@@ -552,7 +552,7 @@ namespace uba
 			//if (size == ErrorSize)
 			//	UBA_ASSERT(false);
 
-			connection.Send(writer.GetData(), size == ErrorSize ? HeaderSize : size);
+			connection.Send(writer.GetData(), size == ErrorSize ? HeaderSize : size, TC("MessageResponse"));
 		}
 			
 		connection.Release();
@@ -1081,7 +1081,7 @@ namespace uba
 		idAndSizePtr[0] = info.messageId >> 8;
 		*(u32*)(idAndSizePtr + 1) = bodySize | u32(info.messageId << 24);
 
-		connection.Send(writer.GetData(), u32(writer.GetPosition()));
+		connection.Send(writer.GetData(), u32(writer.GetPosition()), TC("MessageResponse"));
 		return true;
 	}
 
diff --git a/Engine/Source/Programs/UnrealBuildAccelerator/Common/Public/UbaNetworkBackend.h b/Engine/Source/Programs/UnrealBuildAccelerator/Common/Public/UbaNetworkBackend.h
--- a/Engine/Source/Programs/UnrealBuildAccelerator/Common/Public/UbaNetworkBackend.h
+++ b/Engine/Source/Programs/UnrealBuildAccelerator/Common/Public/UbaNetworkBackend.h
@@ -36,7 +36,7 @@ namespace uba
 		struct SendContext;
 
 		// Send data to connection. If sendContext is set it means that backend expects sendContext to stay around in memory until send is done (or rather, response is received)
-		virtual bool Send(Logger& logger, void* connection, const void* data, u32 dataSize, SendContext& sendContext) = 0;
+		virtual bool Send(Logger& logger, void* connection, const void* data, u32 dataSize, SendContext& sendContext, const tchar* sendHint) = 0;
 		using DataSentCallback = void(void* context, u32 bytes);
 		virtual void SetDataSentCallback(void* connection, void* context, DataSentCallback* callback) = 0;
 
diff --git a/Engine/Source/Programs/UnrealBuildAccelerator/Common/Public/UbaNetworkBackendMemory.h b/Engine/Source/Programs/UnrealBuildAccelerator/Common/Public/UbaNetworkBackendMemory.h
--- a/Engine/Source/Programs/UnrealBuildAccelerator/Common/Public/UbaNetworkBackendMemory.h
+++ b/Engine/Source/Programs/UnrealBuildAccelerator/Common/Public/UbaNetworkBackendMemory.h
@@ -14,7 +14,7 @@ namespace uba
 		NetworkBackendMemory(LogWriter& writer, const tchar* prefix = TC("NetworkBackendMemory"));
 		virtual ~NetworkBackendMemory();
 		virtual void Shutdown(void* connection) override;
-		virtual bool Send(Logger& logger, void* connection, const void* data, u32 dataSize, SendContext& sendContext) override;
+		virtual bool Send(Logger& logger, void* connection, const void* data, u32 dataSize, SendContext& sendContext, const tchar* sendHint) override;
 		virtual void SetDataSentCallback(void* connection, void* context, DataSentCallback* callback) override;
 		virtual void SetRecvCallbacks(void* connection, void* context, u32 headerSize, RecvHeaderCallback* h, RecvBodyCallback* b, const tchar* recvHint) override;
 		virtual void SetRecvTimeout(void* connection, u32 timeoutMs) override;
diff --git a/Engine/Source/Programs/UnrealBuildAccelerator/Common/Public/UbaNetworkBackendQuic.h b/Engine/Source/Programs/UnrealBuildAccelerator/Common/Public/UbaNetworkBackendQuic.h
--- a/Engine/Source/Programs/UnrealBuildAccelerator/Common/Public/UbaNetworkBackendQuic.h
+++ b/Engine/Source/Programs/UnrealBuildAccelerator/Common/Public/UbaNetworkBackendQuic.h
@@ -31,7 +31,7 @@ namespace uba
 		virtual ~NetworkBackendQuic();
 		virtual void Shutdown(void* connection) override;
 		virtual void Close(void* connection) override;
-		virtual bool Send(Logger& logger, void* connection, const void* data, u32 dataSize, SendContext& sendContext) override;
+		virtual bool Send(Logger& logger, void* connection, const void* data, u32 dataSize, SendContext& sendContext, const tchar* sendHint) override;
 		virtual void SetDataSentCallback(void* connection, void* context, DataSentCallback* callback) override;
 		virtual void SetRecvCallbacks(void* connection, void* context, u32 headerSize, RecvHeaderCallback* h, RecvBodyCallback* b, const wchar_t* recvHint) override;
 		virtual void SetRecvTimeout(void* connection, u32 timeoutMs) override;
diff --git a/Engine/Source/Programs/UnrealBuildAccelerator/Common/Public/UbaNetworkBackendTcp.h b/Engine/Source/Programs/UnrealBuildAccelerator/Common/Public/UbaNetworkBackendTcp.h
--- a/Engine/Source/Programs/UnrealBuildAccelerator/Common/Public/UbaNetworkBackendTcp.h
+++ b/Engine/Source/Programs/UnrealBuildAccelerator/Common/Public/UbaNetworkBackendTcp.h
@@ -25,7 +25,7 @@ namespace uba
 		NetworkBackendTcp(const NetworkBackendTcpCreateInfo& info, const tchar* prefix = TC("NetworkBackendTcp"));
 		virtual ~NetworkBackendTcp();
 		virtual void Shutdown(void* connection) override;
-		virtual bool Send(Logger& logger, void* connection, const void* data, u32 dataSize, SendContext& sendContext) override;
+		virtual bool Send(Logger& logger, void* connection, const void* data, u32 dataSize, SendContext& sendContext, const tchar* sendHint) override;
 		virtual void SetDataSentCallback(void* connection, void* context, DataSentCallback* callback) override;
 		virtual void SetRecvCallbacks(void* connection, void* context, u32 headerSize, RecvHeaderCallback* h, RecvBodyCallback* b, const tchar* recvHint) override;
 		virtual void SetRecvTimeout(void* connection, u32 timeoutMs) override;
