[UBA]
* Fixed cache server and test compile errors

diff --git a/Engine/Source/Programs/UnrealBuildAccelerator/CacheService/Private/UbaHttpServer.cpp b/Engine/Source/Programs/UnrealBuildAccelerator/CacheService/Private/UbaHttpServer.cpp
--- a/Engine/Source/Programs/UnrealBuildAccelerator/CacheService/Private/UbaHttpServer.cpp
+++ b/Engine/Source/Programs/UnrealBuildAccelerator/CacheService/Private/UbaHttpServer.cpp
@@ -31,7 +31,7 @@ namespace uba
 		char response[1024];
 		auto sendAndExit = MakeGuard([&]()
 			{
-				NetworkBackend::SendContext sendContext; c.server.m_backend.Send(c.server.m_logger, c.connection, response, u32(strlen(response)), sendContext);
+				NetworkBackend::SendContext sendContext; c.server.m_backend.Send(c.server.m_logger, c.connection, response, u32(strlen(response)), sendContext, TC("HttpServer"));
 				delete[] bodyData;
 			});
 
diff --git a/Engine/Source/Programs/UnrealBuildAccelerator/Test/Private/UbaTestBasics.cpp b/Engine/Source/Programs/UnrealBuildAccelerator/Test/Private/UbaTestBasics.cpp
--- a/Engine/Source/Programs/UnrealBuildAccelerator/Test/Private/UbaTestBasics.cpp
+++ b/Engine/Source/Programs/UnrealBuildAccelerator/Test/Private/UbaTestBasics.cpp
@@ -533,7 +533,7 @@ namespace uba
 		path.EnsureEndsWithSlash().Append(TCV("UbaTestApp.exe"));
 		bool importKernel = false;
 		StringBuffer<> error;
-		FindImports(path.data, [&](const tchar* import, bool isKnown, const char* const* importLoaderPaths)
+		FindImports(path.data, {}, [&](const tchar* import, bool isKnown, const char* const* importLoaderPaths)
 		{
 			importKernel |= isKnown && Contains(import, TC("KERNEL32.dll"));
 		}, error);
diff --git a/Engine/Source/Programs/UnrealBuildAccelerator/Test/Private/UbaTestNetwork.cpp b/Engine/Source/Programs/UnrealBuildAccelerator/Test/Private/UbaTestNetwork.cpp
--- a/Engine/Source/Programs/UnrealBuildAccelerator/Test/Private/UbaTestNetwork.cpp
+++ b/Engine/Source/Programs/UnrealBuildAccelerator/Test/Private/UbaTestNetwork.cpp
@@ -43,7 +43,7 @@ namespace uba
 				tcp.SetRecvCallbacks(connection, &tcp, 1, [](void* context, const Guid& connectionUid, u8* headerData, void*& outBodyContext, u8*& outBodyData, u32& outBodySize) { return true; }, nullptr, TC(""));
 				u8 b = 42;
 				NetworkBackend::SendContext sc;
-				tcp.Send(logger, connection, &b, 1, sc);
+				tcp.Send(logger, connection, &b, 1, sc, TC(""));
 				return true;
 			}, port);
 
diff --git a/Engine/Source/Programs/UnrealBuildAccelerator/Test/Private/UbaTestSession.cpp b/Engine/Source/Programs/UnrealBuildAccelerator/Test/Private/UbaTestSession.cpp
--- a/Engine/Source/Programs/UnrealBuildAccelerator/Test/Private/UbaTestSession.cpp
+++ b/Engine/Source/Programs/UnrealBuildAccelerator/Test/Private/UbaTestSession.cpp
@@ -432,6 +432,7 @@ namespace uba
 
 	bool TestDetouredPopen(LoggerWithWriter& logger, const StringBufferBase& testRootDir)
 	{
+		#if PLATFORM_LINUX
 		return RunLocal(logger, testRootDir, [](LoggerWithWriter& logger, SessionServer& session, const tchar* workingDir, const RunProcessFunction& runProcess)
 			{
 				StringBuffer<> testApp;
@@ -459,6 +460,9 @@ namespace uba
 				}
 				return true;
 			});
+		#else
+		return true;
+		#endif
 	}
 
 	const tchar* GetSystemApplication()
