[UBA]
* Added so failing to open for write because of ERROR_ACCESS_DENIED will retry and also log out reason if privilege to get that info

diff --git a/Engine/Source/Programs/UnrealBuildAccelerator/Common/Private/UbaFileAccessor.cpp b/Engine/Source/Programs/UnrealBuildAccelerator/Common/Private/UbaFileAccessor.cpp
--- a/Engine/Source/Programs/UnrealBuildAccelerator/Common/Private/UbaFileAccessor.cpp
+++ b/Engine/Source/Programs/UnrealBuildAccelerator/Common/Private/UbaFileAccessor.cpp
@@ -99,7 +99,7 @@ namespace uba
 			}
 			u32 lastError = GetLastError();
 			#if PLATFORM_WINDOWS
-			if (lastError == ERROR_SHARING_VIOLATION || lastError == ERROR_USER_MAPPED_FILE)
+			if (lastError == ERROR_SHARING_VIOLATION || lastError == ERROR_USER_MAPPED_FILE || lastError == ERROR_ACCESS_DENIED)
 			{
 				// Since we have unmap of files in jobs it might be that the file is not unmapped yet, therefor we should retry
 				if (retryCount == 1)
