[UBA]
* Added linux version to UbaAgent description

diff --git a/Engine/Source/Programs/UnrealBuildAccelerator/Agent/Private/UbaAgent.cpp b/Engine/Source/Programs/UnrealBuildAccelerator/Agent/Private/UbaAgent.cpp
--- a/Engine/Source/Programs/UnrealBuildAccelerator/Agent/Private/UbaAgent.cpp
+++ b/Engine/Source/Programs/UnrealBuildAccelerator/Agent/Private/UbaAgent.cpp
@@ -777,8 +777,19 @@ namespace uba
 		if (!zone.count)
 			GetZone(zone);
 
+		#if PLATFORM_WINDOWS
 		if (IsRunningWine())
 			extraInfo.Append(TCV(", Linux/Wine-")).Append(GetWineVersion());
+		#elif PLATFORM_LINUX
+		if (FILE* osRelease = fopen("/etc/os-release", "r"))
+		{
+			char buffer[256];
+			if (fgets(buffer, sizeof_array(buffer), osRelease))
+				if (StartsWith(buffer, "PRETTY_NAME=\""))
+					extraInfo.Append(TCV(", ")).Append(buffer + 13, TStrlen(buffer) - 14);
+			fclose(osRelease);
+		}
+		#endif
 
 		if (useQuic)
 			extraInfo.Append(TCV(", MsQuic"));
