[UBA]
* Added macos version to agent description

diff --git a/Engine/Source/Programs/UnrealBuildAccelerator/Agent/Private/UbaAgent.cpp b/Engine/Source/Programs/UnrealBuildAccelerator/Agent/Private/UbaAgent.cpp
--- a/Engine/Source/Programs/UnrealBuildAccelerator/Agent/Private/UbaAgent.cpp
+++ b/Engine/Source/Programs/UnrealBuildAccelerator/Agent/Private/UbaAgent.cpp
@@ -789,6 +789,14 @@ namespace uba
 					extraInfo.Append(TCV(", ")).Append(buffer + 13, TStrlen(buffer) - 14);
 			fclose(osRelease);
 		}
+		#elif PLATFORM_MAC
+		if (FILE* osRelease = popen("sw_vers -productVersion", "r"))
+		{
+			char buffer[256];
+			if (fgets(buffer, sizeof_array(buffer), osRelease))
+				extraInfo.Append(TCV(", MacOS-")).Append(buffer, TStrlen(buffer) - 1);
+			pclose(osRelease);
+		}
 		#endif
 
 		if (useQuic)
