[UBA]
* added ifdef PLATFORM_LINUX around code that broke mac loading of dylib dependencies. Will revisit to understand why this went wrong for mac

diff --git a/Engine/Source/Programs/UnrealBuildAccelerator/Common/Private/UbaSessionServer.cpp b/Engine/Source/Programs/UnrealBuildAccelerator/Common/Private/UbaSessionServer.cpp
--- a/Engine/Source/Programs/UnrealBuildAccelerator/Common/Private/UbaSessionServer.cpp
+++ b/Engine/Source/Programs/UnrealBuildAccelerator/Common/Private/UbaSessionServer.cpp
@@ -879,7 +879,10 @@ namespace uba
 			for (auto& loaderPath : loaderPaths)
 			{
 				StringBuffer<> fullPath;
-				if (loaderPath[0] != '/')
+
+				#if PLATFORM_LINUX
+				if (loaderPath[0] != '/') // TODO: Revisit this.. should be done in a less hacky way.
+				#endif
 					fullPath.Append(applicationDir).EnsureEndsWithSlash();
 				fullPath.Append(loaderPath).EnsureEndsWithSlash().Append(fileName);
 				if (GetFileAttributesW(fullPath.data) == INVALID_FILE_ATTRIBUTES)
@@ -892,6 +895,7 @@ namespace uba
 				break;
 			}
 
+			#if 0
 			if (casKey == CasKeyZero)
 			{
 				m_logger.Warning(TC("Failed to find file %s"), fileName.data);
@@ -900,6 +904,7 @@ namespace uba
 					m_logger.Warning(TC("   LoaderPath %s"), loaderPath.c_str());
 				}
 			}
+			#endif
 		}
 		else if (SearchPathForFile(m_logger, absoluteFile, fileName.data, applicationDir.data))
 		{
