[UBA]
* Changed so WSAPoll for socket send times out after 40 seconds, then report if it timed out before trying again with a 10 minutes timeout. This way we will know if client send times out before server recv

diff --git a/Engine/Source/Programs/UnrealBuildAccelerator/Common/Private/UbaNetworkBackendTcp.cpp b/Engine/Source/Programs/UnrealBuildAccelerator/Common/Private/UbaNetworkBackendTcp.cpp
--- a/Engine/Source/Programs/UnrealBuildAccelerator/Common/Private/UbaNetworkBackendTcp.cpp
+++ b/Engine/Source/Programs/UnrealBuildAccelerator/Common/Private/UbaNetworkBackendTcp.cpp
@@ -847,24 +847,34 @@ namespace uba
 				u32 lastError = WSAGetLastError();
 				if (lastError == WSAEWOULDBLOCK)
 				{
-					u64 timeoutMs = 10 * 60 * 1000; // 10 minutes send timeout
-
-					WSAPOLLFD p;
-					p.fd = socket;
-					p.revents = 0;
-					p.events = POLLWRNORM;
-					int res = WSAPoll(&p, 1, int(timeoutMs));
-					if (!res)
-					{
-						logger.Info(TC("SendSocket - WSAPoll returned 10 minutes timeout for socket %i (%s)"), socket, hint);
-						return false;
-					}
-					if (res == SOCKET_ERROR)
+					u64 timeoutMs = 40 * 1000; // 40 seconds timeout to begin with
+
+					for (u32 i=0; i!=2; ++i)
 					{
-						#if UBA_LOG_SOCKET_ERRORS
-						logger.Info(TC("SendSocket - WSAPoll returned an error for socket %i: %s (%s)"), socket, LastErrorToText(WSAGetLastError()).data, hint);
-						#endif
-						return false;
+						WSAPOLLFD p;
+						p.fd = socket;
+						p.revents = 0;
+						p.events = POLLWRNORM;
+						int res = WSAPoll(&p, 1, int(timeoutMs));
+						if (!res)
+						{
+							if (i == 0)
+							{
+								logger.Info(TC("SendSocket - WSAPoll returned timeout after 40 seconds for socket %i (%s)"), socket, hint);
+								timeoutMs = 10 * 60 * 1000; // 10 minutes send timeout
+								continue;
+							}
+							logger.Info(TC("SendSocket - WSAPoll returned timeout after 10 minutes for socket %i (%s)"), socket, hint);
+							return false;
+						}
+						if (res == SOCKET_ERROR)
+						{
+							#if UBA_LOG_SOCKET_ERRORS
+							logger.Info(TC("SendSocket - WSAPoll returned an error for socket %i: %s (%s)"), socket, LastErrorToText(WSAGetLastError()).data, hint);
+							#endif
+							return false;
+						}
+						break;
 					}
 					continue;
 				}
diff --git a/Engine/Source/Programs/UnrealBuildAccelerator/Common/Private/UbaNetworkServer.cpp b/Engine/Source/Programs/UnrealBuildAccelerator/Common/Private/UbaNetworkServer.cpp
--- a/Engine/Source/Programs/UnrealBuildAccelerator/Common/Private/UbaNetworkServer.cpp
+++ b/Engine/Source/Programs/UnrealBuildAccelerator/Common/Private/UbaNetworkServer.cpp
@@ -144,8 +144,11 @@ namespace uba
 			if (m_disconnectCalled.fetch_add(1) != 0)
 				return;
 			SetShouldDisconnect();
-			if (--m_activeWorkerCount == 0) // Will disconnect in send if there are active workers
+			int activeWorkerCount = --m_activeWorkerCount;
+			if (!activeWorkerCount) // Will disconnect in send if there are active workers
 				TestDisconnect();
+			else
+				m_server.m_logger.Detail(TC("Connection %u disconnected while it has %u active workers"), m_id, activeWorkerCount);
 		}
 
 		bool Stop()
@@ -157,7 +160,7 @@ namespace uba
 			{
 				if (TimeToMs(GetTime() - startTimer) > 3000)
 				{
-					m_server.m_logger.Error(TC("Connection has waited 3 seconds to stop... something is stuck (Active worker count: %u)"), m_activeWorkerCount.load());
+					m_server.m_logger.Error(TC("Connection %u has waited 3 seconds to stop... something is stuck (Active worker count: %u)"), m_id, m_activeWorkerCount.load());
 					PrintAllCallstacks(m_server.m_logger);
 					return false;
 				}
