[UBA]
* Fixed so files in rsp file is not devirtualized twice (causing warning messages)

diff --git a/Engine/Source/Programs/UnrealBuildAccelerator/Common/Public/UbaDependencyCrawler.h b/Engine/Source/Programs/UnrealBuildAccelerator/Common/Public/UbaDependencyCrawler.h
--- a/Engine/Source/Programs/UnrealBuildAccelerator/Common/Public/UbaDependencyCrawler.h
+++ b/Engine/Source/Programs/UnrealBuildAccelerator/Common/Public/UbaDependencyCrawler.h
@@ -227,20 +227,24 @@ namespace uba
 
 	bool DependencyCrawler::ParseRsp2(TrackWorkScope& tracker, Instance& instance, const void* data, u64 dataSize, const tchar* rsp, CodeFiles& outCodeFiles)
 	{
-		auto AddCodeFile = [&](const StringView& file, bool pushFront, bool hasPch)
+		auto AddCodeFile = [&](StringView file, bool pushFront, bool hasPch, bool fixPath = true)
 			{
 				UBA_CRAWLER_DEBUG_PRINT("CODEFILE", file.data);
 				StringBuffer<> fixedPath2;
-				FixPath(file.data, instance.workDir.data(), instance.workDir.size(), fixedPath2);
-				if (!instance.devirtualizePathFunc(fixedPath2))
-					m_logger.Info(TC("Failed to devirtualize path %s in %s"), fixedPath2.data, rsp);
+				if (fixPath)
+				{
+					FixPath(file.data, instance.workDir.data(), instance.workDir.size(), fixedPath2);
+					if (!instance.devirtualizePathFunc(fixedPath2))
+						m_logger.Info(TC("Failed to devirtualize path %s in %s"), fixedPath2.data, rsp);
 
-				if (CaseInsensitiveFs)
-					fixedPath2.MakeLower();
+					if (CaseInsensitiveFs)
+						fixedPath2.MakeLower();
+					file = fixedPath2;
+				}
 				if (pushFront)
-					outCodeFiles.push_front({fixedPath2.ToString(), hasPch});
+					outCodeFiles.push_front({file.ToString(), hasPch});
 				else
-					outCodeFiles.push_back({fixedPath2.ToString(), hasPch});
+					outCodeFiles.push_back({file.ToString(), hasPch});
 			};
 
 		auto AddRoot = [&](const StringView& path)
@@ -338,7 +342,7 @@ namespace uba
 
 							if (fixedFile.Contains(TC("\\Definitions.")))
 							{
-								AddCodeFile(fixedFile, true, false);
+								AddCodeFile(fixedFile, true, false, false);
 							}
 							else
 							{
@@ -357,7 +361,8 @@ namespace uba
 
 									u32 attributes = 0;
 									bool hasPch = m_fileExistsFunc(fixedFile, attributes);
-									AddCodeFile(StringView(sb).Skip(3), true, hasPch);
+									fixedFile.Resize(fixedFile.count - 4);
+									AddCodeFile(fixedFile, true, hasPch, false);
 								}
 							}
 
