[UBA]
* Changed so client sends ping messages once it is connected

diff --git a/Engine/Source/Programs/UnrealBuildAccelerator/Common/Private/UbaSessionClient.cpp b/Engine/Source/Programs/UnrealBuildAccelerator/Common/Private/UbaSessionClient.cpp
--- a/Engine/Source/Programs/UnrealBuildAccelerator/Common/Private/UbaSessionClient.cpp
+++ b/Engine/Source/Programs/UnrealBuildAccelerator/Common/Private/UbaSessionClient.cpp
@@ -141,6 +141,7 @@ namespace uba
 	{
 		m_client.RegisterOnDisconnected([this]()
 			{
+				m_connected = false;
 				m_loop = false;
 			});
 
@@ -2557,7 +2558,7 @@ namespace uba
 
 	void SessionClient::TraceSessionUpdate()
 	{
-		if (m_memAvail)
+		if (m_connected)
 			SendPing(m_memAvail, m_memTotal);
 
 		if (!m_trace.IsWriting())
