[UBA]
* Moved unmapping of memory mapping creation to a job (these are used for decompression of pch files)

diff --git a/Engine/Source/Programs/UnrealBuildAccelerator/Common/Private/UbaSession.cpp b/Engine/Source/Programs/UnrealBuildAccelerator/Common/Private/UbaSession.cpp
--- a/Engine/Source/Programs/UnrealBuildAccelerator/Common/Private/UbaSession.cpp
+++ b/Engine/Source/Programs/UnrealBuildAccelerator/Common/Private/UbaSession.cpp
@@ -886,7 +886,7 @@ namespace uba
 			MappedView mappedView;
 
 			auto ownedGuard1 = MakeGuard([&]() { CloseFileMapping(m_logger, mapping, fileName); });
-			auto ownedGuard2 = MakeGuard([&]() { UnmapViewOfFile(m_logger, mappingMemory, size, fileName); });
+			auto ownedGuard2 = MakeGuard([&]() { m_workManager.AddWork([this, mappingMemory, size, fn = TString(fileName)](const WorkContext&) { UnmapViewOfFile(m_logger, mappingMemory, size, fn.c_str()); }, 1, TC("UnmapViewOfFile")); });
 			auto viewGuard = MakeGuard([&]() { m_fileMappingBuffer.UnmapView(mappedView, fileName); });
 
 			if (canBeFreed)
