[UBA]
* Added so helper prints out all callstacks if ping message takes more than 20 seconds.

diff --git a/Engine/Source/Programs/UnrealBuildAccelerator/Common/Private/UbaSessionClient.cpp b/Engine/Source/Programs/UnrealBuildAccelerator/Common/Private/UbaSessionClient.cpp
--- a/Engine/Source/Programs/UnrealBuildAccelerator/Common/Private/UbaSessionClient.cpp
+++ b/Engine/Source/Programs/UnrealBuildAccelerator/Common/Private/UbaSessionClient.cpp
@@ -1591,12 +1591,28 @@ namespace uba
 		writer.WriteU64(memTotal);
 		writer.WriteU32(*(u32*)&cpuLoad);
 		StackBinaryReader<32> reader;
+
+		Event gotResponse(true);
+		auto doneFunc = [](bool error, void* userData) { ((Event*)userData)->Set(); };
+
 		time = GetTime();
-		if (!msg.Send(reader, m_stats.pingMsg))
+		
+		if (!msg.SendAsync(reader, doneFunc, &gotResponse))
 		{
 			m_loop = false;
 			return;
 		}
+
+		#if PLATFORM_WINDOWS || PLATFORM_MAC // Linux will crash with print all callstacks
+		if (!gotResponse.IsSet(20*1000)) // We accept 20 seconds time out until we write out callstacks
+		{
+			LoggerWithWriter logger(g_consoleLogWriter);
+			PrintAllCallstacks(logger);
+		}
+		#endif
+		
+		gotResponse.IsSet(); // Wait forever (until message trigger error or finishes)
+
 		u64 newTime = GetTime();
 		u64 lastPing = newTime - time;
 		m_lastPing = lastPing;
