[UBA]
* Fixed so visualizer network communication is happening on different thread than windows message pump

diff --git a/Engine/Source/Programs/UnrealBuildAccelerator/Visualizer/Private/UbaTraceReader.cpp b/Engine/Source/Programs/UnrealBuildAccelerator/Visualizer/Private/UbaTraceReader.cpp
--- a/Engine/Source/Programs/UnrealBuildAccelerator/Visualizer/Private/UbaTraceReader.cpp
+++ b/Engine/Source/Programs/UnrealBuildAccelerator/Visualizer/Private/UbaTraceReader.cpp
@@ -187,6 +187,31 @@ namespace uba
 		outChanged = false;
 		if (!m_memoryHandle.IsValid())
 			return true;
+
+		if (!client.IsConnected() && !out.finished)
+		{
+			for (auto& session : out.sessions)
+				for (auto& processor : session.processors)
+				{
+					auto& process = processor.processes.back();
+					if (process.stop == ~u64(0))
+					{
+						process.stop = GetTime() - out.startTime;
+						process.exitCode = 0; // This is wrong but since we didn't get the final result we can't tell if it was success or not
+						process.bitmapDirty = true;
+					}
+				}
+			out.finished = true;
+			return false;
+		}
+
+		SCOPED_FUTEX_READ(m_memoryFutex, lock);
+		outChanged = !m_activeProcesses.empty() || m_memoryPos != m_memoryEnd;
+		return ReadMemory(out, false, ~u64(0));
+	}
+
+	bool TraceReader::UpdateReceiveClient(NetworkClient& client)
+	{
 		while (true)
 		{
 			u32 pos = u32(m_memoryEnd - m_memoryBegin);
@@ -196,31 +221,18 @@ namespace uba
 			writer.WriteU32(pos);
 			StackBinaryReader<SendDefaultSize> reader;
 			if (!msg.Send(reader))
-			{
-				for (auto& session : out.sessions)
-					for (auto& processor : session.processors)
-					{
-						auto& process = processor.processes.back();
-						if (process.stop == ~u64(0))
-						{
-							process.stop = GetTime() - out.startTime;
-							process.exitCode = 0; // This is wrong but since we didn't get the final result we can't tell if it was success or not
-							process.bitmapDirty = true;
-						}
-					}
-				out.finished = true;
 				return false;
-			}
-			u64 remotePos = reader.ReadU32();
+			u64 remotePos = reader.ReadU32();(void)remotePos;
 			u32 left = u32(reader.GetLeft());
+
+			SCOPED_FUTEX(m_memoryFutex, lock);
 			reader.ReadBytes(m_memoryEnd, left);
 			pos += left;
 			m_memoryEnd += left;
 			if (remotePos == pos)
 				break;
 		}
-		outChanged = !m_activeProcesses.empty() || m_memoryPos != m_memoryEnd;
-		return ReadMemory(out, false, ~u64(0));
+		return true;
 	}
 
 	bool TraceReader::StartReadNamed(TraceView& out, const tchar* namedTrace, bool silentFail, bool replay)
diff --git a/Engine/Source/Programs/UnrealBuildAccelerator/Visualizer/Private/UbaTraceReader.h b/Engine/Source/Programs/UnrealBuildAccelerator/Visualizer/Private/UbaTraceReader.h
--- a/Engine/Source/Programs/UnrealBuildAccelerator/Visualizer/Private/UbaTraceReader.h
+++ b/Engine/Source/Programs/UnrealBuildAccelerator/Visualizer/Private/UbaTraceReader.h
@@ -204,6 +204,7 @@ namespace uba
 		// Use for network
 		bool StartReadClient(TraceView& out, NetworkClient& client);
 		bool UpdateReadClient(TraceView& out, NetworkClient& client, bool& outChanged);
+		bool UpdateReceiveClient(NetworkClient& client);
 
 		// Use for local
 		bool StartReadNamed(TraceView& out, const tchar* namedTrace, bool silentFail = false, bool replay = false);
@@ -241,5 +242,6 @@ namespace uba
 		u8* m_memoryPos = nullptr;
 		u8* m_memoryEnd = nullptr;
 		HANDLE m_hostProcess = NULL;
+		Futex m_memoryFutex;
 	};
 }
\ No newline at end of file
diff --git a/Engine/Source/Programs/UnrealBuildAccelerator/Visualizer/Private/UbaVisualizer.cpp b/Engine/Source/Programs/UnrealBuildAccelerator/Visualizer/Private/UbaVisualizer.cpp
--- a/Engine/Source/Programs/UnrealBuildAccelerator/Visualizer/Private/UbaVisualizer.cpp
+++ b/Engine/Source/Programs/UnrealBuildAccelerator/Visualizer/Private/UbaVisualizer.cpp
@@ -350,8 +350,10 @@ namespace uba
 			PostNewTitle(GetTitlePrefix(title).Appendf(L"Connected to %s:%u", host, port));
 			PostNewTrace(0, false);
 
-			while (m_hwnd && m_client->IsConnected() && !m_clientDisconnect.IsSet(1000))
-				;
+			while (m_hwnd && m_client->IsConnected() && !m_clientDisconnect.IsSet(100))
+			{
+				m_trace.UpdateReceiveClient(*m_client);
+			}
 
 			PostNewTitle(GetTitlePrefix(title).Appendf(L"Disconnected..."));
 
@@ -368,9 +370,9 @@ namespace uba
 	{
 		m_looping = true;
 		m_autoScroll = false;
+		m_fileName.Append(fileName);
 		if (!StartHwndThread())
 			return true;
-		m_fileName.Append(fileName);
 		PostNewTrace(replay, 0);
 		return true;
 	}
@@ -972,15 +974,21 @@ namespace uba
 		if (m_traceView.sessions.empty() && !m_parentHwnd)
 		{
 			StringBuffer<> str[2];
-			str[1].Append(TCV("Click here to open trace file"));
-			if (m_listenChannel.count)
+			if (!m_fileName.IsEmpty())
 			{
-				str[0].Appendf(TC("Listening for new sessions on channel '%s'"), m_listenChannel.data);
-				str[1].Append(TCV(" instead"));
+				str[0].Append(TCV("Loading trace file"));
 			}
 			else
-				str[0].Append(TCV("No trace active"));
-
+			{
+				str[1].Append(TCV("Click here to open trace file"));
+				if (m_listenChannel.count)
+				{
+					str[0].Appendf(TC("Listening for new sessions on channel '%s'"), m_listenChannel.data);
+					str[1].Append(TCV(" instead"));
+				}
+				else
+					str[0].Append(TCV("No trace active"));
+			}
 			SetActiveFont(m_popupFont);
 			DrawCenteredText(str, 2, false);
 			return;
