[UBA]
* Fixed bug where long strings did not track character width.. which caused win64 visualizer not being able to show breadcrumbs from mac/linux

diff --git a/Engine/Source/Programs/UnrealBuildAccelerator/Common/Public/UbaTrace.h b/Engine/Source/Programs/UnrealBuildAccelerator/Common/Public/UbaTrace.h
--- a/Engine/Source/Programs/UnrealBuildAccelerator/Common/Public/UbaTrace.h
+++ b/Engine/Source/Programs/UnrealBuildAccelerator/Common/Public/UbaTrace.h
@@ -69,7 +69,7 @@ namespace uba
 		#undef UBA_TRACE_TYPE
 	};
 
-	static constexpr u32 TraceVersion = 41;
+	static constexpr u32 TraceVersion = 42;
 	static constexpr u32 TraceReadCompatibilityVersion = 6;
 
 	class Trace : public WorkTracker
diff --git a/Engine/Source/Programs/UnrealBuildAccelerator/Core/Private/UbaPlatform.cpp b/Engine/Source/Programs/UnrealBuildAccelerator/Core/Private/UbaPlatform.cpp
--- a/Engine/Source/Programs/UnrealBuildAccelerator/Core/Private/UbaPlatform.cpp
+++ b/Engine/Source/Programs/UnrealBuildAccelerator/Core/Private/UbaPlatform.cpp
@@ -710,12 +710,12 @@ namespace uba
 	{
 		if (str.count < threshold)
 		{
-			WriteBool(true);
+			WriteByte(0);
 			WriteString(str);
 			return;
 		}
-			
-		WriteBool(false);
+
+		WriteByte(sizeof(tchar));
 		u32 uncompressedSize = str.count * sizeof(tchar);
 		u64 needed = (u64)OodleLZ_GetCompressedBufferSizeNeeded(OodleLZ_Compressor_Kraken, uncompressedSize);
 		Vector<u8> mem;
@@ -729,17 +729,31 @@ namespace uba
 
 	TString BinaryReader::ReadLongString()
 	{
-		if (ReadBool())
+		u8 s = ReadByte();
+		if (!s)
 			return ReadString();
 
 		u64 stringLength = Read7BitEncoded();
-		u64 uncompressedSize = stringLength * sizeof(tchar);
+		u64 uncompressedSize = stringLength * s;
 		u64 compressedSize = Read7BitEncoded();
-		TString str;
-		str.resize(stringLength);
-		OO_SINTa decompLen = OodleLZ_Decompress(GetPositionData(), (OO_SINTa)compressedSize, str.data(), (OO_SINTa)uncompressedSize);
-		UBA_ASSERT(decompLen == (OO_SINTa)uncompressedSize);(void)decompLen;
 		Skip(compressedSize);
-		return str;
+
+		if (s == sizeof(tchar))
+		{
+			TString str;
+			str.resize(stringLength);
+			OO_SINTa decompLen = OodleLZ_Decompress(GetPositionData(), (OO_SINTa)compressedSize, str.data(), (OO_SINTa)uncompressedSize);
+			UBA_ASSERT(decompLen == (OO_SINTa)uncompressedSize);(void)decompLen;
+			return str;
+		}
+		else
+		{
+			UBA_ASSERT(s == 1);
+			std::string temp;
+			temp.resize(stringLength);
+			OO_SINTa decompLen = OodleLZ_Decompress(GetPositionData(), (OO_SINTa)compressedSize, temp.data(), (OO_SINTa)uncompressedSize);
+			UBA_ASSERT(decompLen == (OO_SINTa)uncompressedSize);(void)decompLen;
+			return TString(temp.begin(), temp.end());
+		}
 	}
 }
diff --git a/Engine/Source/Programs/UnrealBuildAccelerator/Visualizer/Private/UbaTraceReader.cpp b/Engine/Source/Programs/UnrealBuildAccelerator/Visualizer/Private/UbaTraceReader.cpp
--- a/Engine/Source/Programs/UnrealBuildAccelerator/Visualizer/Private/UbaTraceReader.cpp
+++ b/Engine/Source/Programs/UnrealBuildAccelerator/Visualizer/Private/UbaTraceReader.cpp
@@ -624,6 +624,17 @@ namespace uba
 			{
 				if (out.version < 38)
 					breadcrumbs = reader.ReadString();
+				else if (out.version < 42)
+				{
+					if (ReadBool())
+						breadcrumbs = reader.ReadString();
+					else
+					{
+						Read7BitEncoded();
+						Skip(Read7BitEncoded());
+						breadcrumbs = TC("Upgrade your visualizer");
+					}
+				}
 				else
 					breadcrumbs = reader.ReadLongString();
 			}
