[UBA]
* Fixed bug from previous change with long strings when producing trace from macos and read it on win64

diff --git a/Engine/Source/Programs/UnrealBuildAccelerator/Core/Private/UbaPlatform.cpp b/Engine/Source/Programs/UnrealBuildAccelerator/Core/Private/UbaPlatform.cpp
--- a/Engine/Source/Programs/UnrealBuildAccelerator/Core/Private/UbaPlatform.cpp
+++ b/Engine/Source/Programs/UnrealBuildAccelerator/Core/Private/UbaPlatform.cpp
@@ -736,13 +736,15 @@ namespace uba
 		u64 stringLength = Read7BitEncoded();
 		u64 uncompressedSize = stringLength * s;
 		u64 compressedSize = Read7BitEncoded();
+
+		const u8* data = GetPositionData();
 		Skip(compressedSize);
 
 		if (s == sizeof(tchar))
 		{
 			TString str;
 			str.resize(stringLength);
-			OO_SINTa decompLen = OodleLZ_Decompress(GetPositionData(), (OO_SINTa)compressedSize, str.data(), (OO_SINTa)uncompressedSize);
+			OO_SINTa decompLen = OodleLZ_Decompress(data, (OO_SINTa)compressedSize, str.data(), (OO_SINTa)uncompressedSize);
 			UBA_ASSERT(decompLen == (OO_SINTa)uncompressedSize);(void)decompLen;
 			return str;
 		}
@@ -751,7 +753,7 @@ namespace uba
 			UBA_ASSERT(s == 1);
 			std::string temp;
 			temp.resize(stringLength);
-			OO_SINTa decompLen = OodleLZ_Decompress(GetPositionData(), (OO_SINTa)compressedSize, temp.data(), (OO_SINTa)uncompressedSize);
+			OO_SINTa decompLen = OodleLZ_Decompress(data, (OO_SINTa)compressedSize, temp.data(), (OO_SINTa)uncompressedSize);
 			UBA_ASSERT(decompLen == (OO_SINTa)uncompressedSize);(void)decompLen;
 			return TString(temp.begin(), temp.end());
 		}
diff --git a/Engine/Source/Programs/UnrealBuildAccelerator/Visualizer/Private/UbaTraceReader.cpp b/Engine/Source/Programs/UnrealBuildAccelerator/Visualizer/Private/UbaTraceReader.cpp
--- a/Engine/Source/Programs/UnrealBuildAccelerator/Visualizer/Private/UbaTraceReader.cpp
+++ b/Engine/Source/Programs/UnrealBuildAccelerator/Visualizer/Private/UbaTraceReader.cpp
@@ -626,12 +626,12 @@ namespace uba
 					breadcrumbs = reader.ReadString();
 				else if (out.version < 42)
 				{
-					if (ReadBool())
+					if (reader.ReadBool())
 						breadcrumbs = reader.ReadString();
 					else
 					{
-						Read7BitEncoded();
-						Skip(Read7BitEncoded());
+						reader.Read7BitEncoded();
+						reader.Skip(reader.Read7BitEncoded());
 						breadcrumbs = TC("Upgrade your visualizer");
 					}
 				}
@@ -769,6 +769,17 @@ namespace uba
 			{
 				if (out.version < 38)
 					breadcrumbs = reader.ReadString();
+				else if (out.version < 42)
+				{
+					if (reader.ReadBool())
+						breadcrumbs = reader.ReadString();
+					else
+					{
+						reader.Read7BitEncoded();
+						reader.Skip(reader.Read7BitEncoded());
+						breadcrumbs = TC("Upgrade your visualizer");
+					}
+				}
 				else
 					breadcrumbs = reader.ReadLongString();
 			}
