[UBA]
* Change d network error to warning when getting ECONNRESET on recv calls other than first

diff --git a/Engine/Source/Programs/UnrealBuildAccelerator/Common/Private/UbaNetworkBackendTcp.cpp b/Engine/Source/Programs/UnrealBuildAccelerator/Common/Private/UbaNetworkBackendTcp.cpp
--- a/Engine/Source/Programs/UnrealBuildAccelerator/Common/Private/UbaNetworkBackendTcp.cpp
+++ b/Engine/Source/Programs/UnrealBuildAccelerator/Common/Private/UbaNetworkBackendTcp.cpp
@@ -989,7 +989,7 @@ namespace uba
 			}
 			#if !PLATFORM_WINDOWS
 			if (!isFirstCall && errno != ECONNRESET)
-				return logger.Error(TC("RecvSocket - error receiving socket %i after %u bytes (%s%s) (%s)"), socket, bufferLen, hint1, hint2, strerror(errno));
+				return logger.Warning(TC("RecvSocket - error receiving socket %i after %u bytes (%s%s) (%s)"), socket, bufferLen, hint1, hint2, strerror(errno));
 			#endif
 
 			#if UBA_LOG_SOCKET_ERRORS
