[UBA]
* Cosmetic fix in visualizer when printing out session info lines with extra description containing hyperlink

diff --git a/Engine/Source/Programs/UnrealBuildAccelerator/Visualizer/Private/UbaTraceReader.cpp b/Engine/Source/Programs/UnrealBuildAccelerator/Visualizer/Private/UbaTraceReader.cpp
--- a/Engine/Source/Programs/UnrealBuildAccelerator/Visualizer/Private/UbaTraceReader.cpp
+++ b/Engine/Source/Programs/UnrealBuildAccelerator/Visualizer/Private/UbaTraceReader.cpp
@@ -416,6 +416,10 @@ namespace uba
 				hyperlink.assign(hyperlinkPos, hyperlinkEnd);
 				if (*hyperlinkEnd == ' ')
 					++hyperlinkEnd;
+				while (hyperlinkPos[-1] == ' ')
+					--hyperlinkPos;
+				if (hyperlinkPos[-1] == ',')
+					--hyperlinkPos;
 				memmove(const_cast<tchar*>(hyperlinkPos), hyperlinkEnd, (sessionInfo.data + sessionInfo.count - hyperlinkEnd)*sizeof(tchar));
 				sessionInfo.Resize(sessionInfo.count - (hyperlinkEnd - hyperlinkPos));
 			}
