[UBA]
* Changed lock type when networking root handles to see if this reduce weird timeouts on mac

diff --git a/Engine/Source/Programs/UnrealBuildAccelerator/Common/Private/UbaNetworkBackendTcp.cpp b/Engine/Source/Programs/UnrealBuildAccelerator/Common/Private/UbaNetworkBackendTcp.cpp
--- a/Engine/Source/Programs/UnrealBuildAccelerator/Common/Private/UbaNetworkBackendTcp.cpp
+++ b/Engine/Source/Programs/UnrealBuildAccelerator/Common/Private/UbaNetworkBackendTcp.cpp
@@ -987,9 +987,10 @@ namespace uba
 				}
 				continue;
 			}
+
 			#if !PLATFORM_WINDOWS
 			if (!isFirstCall && errno != ECONNRESET)
-				return logger.Warning(TC("RecvSocket - error receiving socket %i after %u bytes (%s%s) (%s)"), socket, bufferLen, hint1, hint2, strerror(errno));
+				return logger.Warning(TC("RecvSocket - recv error on socket %i and connection %s: %s (%s%s)"), socket, GuidToString(connection).str, strerror(errno), hint1, hint2);
 			#endif
 
 			#if UBA_LOG_SOCKET_ERRORS
diff --git a/Engine/Source/Programs/UnrealBuildAccelerator/Common/Private/UbaSession.cpp b/Engine/Source/Programs/UnrealBuildAccelerator/Common/Private/UbaSession.cpp
--- a/Engine/Source/Programs/UnrealBuildAccelerator/Common/Private/UbaSession.cpp
+++ b/Engine/Source/Programs/UnrealBuildAccelerator/Common/Private/UbaSession.cpp
@@ -1343,13 +1343,19 @@ namespace uba
 		CasKey key = ToCasKey(CasKeyHasher().Update(rootsData, rootsDataSize), false);
 		RootsHandle rootsHandle = WithVfs(key.a, false);
 
-		SCOPED_WRITE_LOCK(m_rootsLookupLock, lock);
-		auto insres = m_rootsLookup.try_emplace(rootsHandle);
-		RootsEntry& entry = insres.first->second;
-		if (insres.second)
+		SCOPED_FUTEX(m_rootsLookupLock, rootsLock);
+		RootsEntry& entry = m_rootsLookup.try_emplace(rootsHandle).first->second;
+		rootsLock.Leave();
+
+		SCOPED_FUTEX(entry.lock, entryLock);
+		UBA_ASSERT(!entry.handled || memcmp(entry.memory.data(), rootsData, rootsDataSize) == 0);
+
+		if (!entry.handled)
+		{
 			PopulateRootsEntry(entry, rootsData, rootsDataSize);
-		else
-			UBA_ASSERT(memcmp(entry.memory.data(), rootsData, rootsDataSize) == 0);
+			entry.handled = true;
+		}
+
 		return WithVfs(rootsHandle, !entry.locals.empty());
 	}
 
@@ -3761,12 +3767,16 @@ namespace uba
 
 	const Session::RootsEntry* Session::GetRootsEntry(RootsHandle rootsHandle)
 	{
-		SCOPED_READ_LOCK(m_rootsLookupLock, lock);
+		SCOPED_FUTEX_READ(m_rootsLookupLock, rootsLock);
 		auto findIt = m_rootsLookup.find(WithVfs(rootsHandle, false));
-		if (findIt != m_rootsLookup.end())
-			return &findIt->second;
-		m_logger.Error(TC("Can't find entry from roots handle %llu"), rootsHandle);
-		return nullptr;
+		if (findIt == m_rootsLookup.end())
+			return m_logger.Error(TC("Can't find entry from roots handle %llu"), rootsHandle) ? nullptr : nullptr;
+
+		RootsEntry& entry = findIt->second;
+		rootsLock.Leave();
+
+		UBA_ASSERT(entry.handled);
+		return &entry;
 	}
 
 	void Session::PopulateRootsEntry(RootsEntry& entry, const void* rootsData, uba::u64 rootsDataSize)
diff --git a/Engine/Source/Programs/UnrealBuildAccelerator/Common/Private/UbaSessionClient.cpp b/Engine/Source/Programs/UnrealBuildAccelerator/Common/Private/UbaSessionClient.cpp
--- a/Engine/Source/Programs/UnrealBuildAccelerator/Common/Private/UbaSessionClient.cpp
+++ b/Engine/Source/Programs/UnrealBuildAccelerator/Common/Private/UbaSessionClient.cpp
@@ -1607,6 +1607,7 @@ namespace uba
 		if (!gotResponse.IsSet(20*1000)) // We accept 20 seconds time out until we write out callstacks
 		{
 			LoggerWithWriter logger(g_consoleLogWriter);
+			logger.Info(TC("Took more than 20 seconds to send data. Printing callstacks"));
 			PrintAllCallstacks(logger);
 		}
 		#endif
@@ -1651,17 +1652,23 @@ namespace uba
 
 	bool SessionClient::SendRootsHandle(RootsHandle rootsHandle)
 	{
-		SCOPED_WRITE_LOCK(m_rootsLookupLock, rootsLock);
-		auto insres2 = m_rootsLookup.try_emplace(WithVfs(rootsHandle, false));
-		if (!insres2.second)
+		SCOPED_FUTEX(m_rootsLookupLock, rootsLock);
+		auto insres = m_rootsLookup.try_emplace(WithVfs(rootsHandle, false));
+		RootsEntry& entry = insres.first->second;
+		rootsLock.Leave();
+
+		SCOPED_FUTEX(entry.lock, entryLock);
+		if (entry.handled)
 			return true;
+		entry.handled = true;
+
 		StackBinaryWriter<128> writer;
 		NetworkMessage msg(m_client, ServiceId, SessionMessageType_GetRoots, writer);
 		writer.WriteU64(rootsHandle);
 		StackBinaryReader<8*1024> reader;
 		if (!msg.Send(reader, m_stats.getApplicationMsg))
 			return false;
-		PopulateRootsEntry(insres2.first->second, reader.GetPositionData(), reader.GetLeft());
+		PopulateRootsEntry(entry, reader.GetPositionData(), reader.GetLeft());
 		return true;
 	}
 
@@ -2577,6 +2584,8 @@ namespace uba
 		if (m_connected)
 			SendPing(m_memAvail, m_memTotal);
 
+		// TODO: There should be some sort of log entry if we have a process that has done no progress for 30 minutes
+
 		if (!m_trace.IsWriting())
 			return;
 
diff --git a/Engine/Source/Programs/UnrealBuildAccelerator/Common/Public/UbaSession.h b/Engine/Source/Programs/UnrealBuildAccelerator/Common/Public/UbaSession.h
--- a/Engine/Source/Programs/UnrealBuildAccelerator/Common/Public/UbaSession.h
+++ b/Engine/Source/Programs/UnrealBuildAccelerator/Common/Public/UbaSession.h
@@ -191,7 +191,7 @@ namespace uba
 
 		bool HasVfs(RootsHandle handle) const;
 		RootsHandle WithVfs(RootsHandle key, bool vfs) const;
-		struct RootsEntry { Vector<u8> memory; RootPaths roots; Vector<TString> locals; Vector<TString> vfs; };
+		struct RootsEntry { Vector<u8> memory; RootPaths roots; Vector<TString> locals; Vector<TString> vfs; Futex lock; bool handled = false; };
 		const RootsEntry* GetRootsEntry(RootsHandle rootsHandle);
 		void PopulateRootsEntry(RootsEntry& entry, const void* rootsData, uba::u64 rootsDataSize);
 
@@ -301,7 +301,7 @@ namespace uba
 		struct VirtualSourceFile { FileMappingHandle mappingHandle; u64 size; };
 		UnorderedMap<StringKey, VirtualSourceFile> m_virtualSourceFiles;
 
-		ReaderWriterLock m_rootsLookupLock;
+		Futex m_rootsLookupLock;
 		UnorderedMap<RootsHandle, RootsEntry> m_rootsLookup;
 
 		u32 m_wantsOnCloseIdCounter = 1;
