[UBA]
* Fixed spelling in visualizer

diff --git a/Engine/Source/Programs/UnrealBuildAccelerator/Visualizer/Private/UbaVisualizer.cpp b/Engine/Source/Programs/UnrealBuildAccelerator/Visualizer/Private/UbaVisualizer.cpp
--- a/Engine/Source/Programs/UnrealBuildAccelerator/Visualizer/Private/UbaVisualizer.cpp
+++ b/Engine/Source/Programs/UnrealBuildAccelerator/Visualizer/Private/UbaVisualizer.cpp
@@ -2285,8 +2285,8 @@ namespace uba
 			logger.SetColor(m_memColor).Info(L"  Mem: %ls/%ls", BytesToText(m_stats.memTotal - m_stats.memAvail).str, BytesToText(m_stats.memTotal).str);
 			if (m_stats.recvBytes || m_stats.sendBytes)
 			{
-				logger.SetColor(m_recvColor).Info(L"  Recv: %lsbit/s", BytesToText(m_stats.recvBytesPerSecond*8).str);
-				logger.SetColor(m_sendColor).Info(L"  Send: %lsbit/s", BytesToText(m_stats.sendBytesPerSecond*8).str);
+				logger.SetColor(m_recvColor).Info(L"  Recv: %lsit/s", BytesToText(m_stats.recvBytesPerSecond*8).str);
+				logger.SetColor(m_sendColor).Info(L"  Send: %lsit/s", BytesToText(m_stats.sendBytesPerSecond*8).str);
 			}
 			if (m_stats.ping)
 				logger.Info(L"  Ping: %ls", TimeToText(m_stats.ping, false, m_traceView.frequency).str);
