[UBA]
* Added checks for GetThreadContext

diff --git a/Engine/Source/Programs/UnrealBuildAccelerator/Common/Private/UbaLogger.cpp b/Engine/Source/Programs/UnrealBuildAccelerator/Common/Private/UbaLogger.cpp
--- a/Engine/Source/Programs/UnrealBuildAccelerator/Common/Private/UbaLogger.cpp
+++ b/Engine/Source/Programs/UnrealBuildAccelerator/Common/Private/UbaLogger.cpp
@@ -746,7 +746,7 @@ namespace uba
 	#endif
 	}
 
-	void TraverseAllCallstacks(const Function<void(const CallstackInfo&)>& func)
+	void TraverseAllCallstacks(const Function<void(const CallstackInfo&)>& func, const TraverseThreadErrorFunc& errorFunc)
 	{
 		UnorderedMap<CasKey, CallstackInfo> callstacks;
 
@@ -769,7 +769,7 @@ namespace uba
 					cs.desc = desc;
 				cs.data.resize(stackWriter.GetPosition());
 				memcpy(cs.data.data(), stackWriter.GetData(), cs.data.size());
-			});
+			}, errorFunc);
 		for (auto& kv : callstacks)
 		{
 			CallstackInfo& cs = kv.second;
@@ -794,6 +794,10 @@ namespace uba
 				static auto& sb = *new StringBuffer<16*1024>();
 				ParseCallstack(sb.Clear(), stackReader);
 				logger.Info(TC("%s%s"), cs.desc.c_str(), sb.data);
+			},
+			[&](const StringView& error)
+			{
+				logger.Info(error.data);
 			});
 	}
 }
diff --git a/Engine/Source/Programs/UnrealBuildAccelerator/Common/Private/UbaSessionClient.cpp b/Engine/Source/Programs/UnrealBuildAccelerator/Common/Private/UbaSessionClient.cpp
--- a/Engine/Source/Programs/UnrealBuildAccelerator/Common/Private/UbaSessionClient.cpp
+++ b/Engine/Source/Programs/UnrealBuildAccelerator/Common/Private/UbaSessionClient.cpp
@@ -1637,6 +1637,10 @@ namespace uba
 					BinaryReader resultReader(stackWriter.GetData(), 0, stackWriter.GetPosition());
 					TString infoString = resultReader.ReadString();
 					m_logger.Info(TC("%s%s"), cs.desc.c_str(), infoString.c_str());
+				},
+				[&](const StringView& error)
+				{
+					m_logger.Info(error.data);
 				});
 		}
 	}
diff --git a/Engine/Source/Programs/UnrealBuildAccelerator/Common/Private/UbaThread.cpp b/Engine/Source/Programs/UnrealBuildAccelerator/Common/Private/UbaThread.cpp
--- a/Engine/Source/Programs/UnrealBuildAccelerator/Common/Private/UbaThread.cpp
+++ b/Engine/Source/Programs/UnrealBuildAccelerator/Common/Private/UbaThread.cpp
@@ -206,9 +206,11 @@ namespace uba
 #endif
 	}
 
-	bool TraverseAllThreads(const TraverseThreadFunc& func)
+	bool TraverseAllThreads(const TraverseThreadFunc& func, const TraverseThreadErrorFunc& errorFunc)
 	{
 #if PLATFORM_WINDOWS
+		StringBuffer<256> error;
+		auto reportError = [&](HANDLE hThread, const tchar* call) { errorFunc(error.Clear().Appendf(TC("%s failed for thread %llu (%s)"), call, u64(hThread), LastErrorToText().data)); };
 
 #if UBA_TRACK_THREADS
 		DWORD currentThreadId = GetCurrentThreadId();
@@ -219,13 +221,20 @@ namespace uba
 			if (currentThreadId == GetThreadId(hThread))
 				continue;
 			if (SuspendThread(hThread) == -1)
+			{
+				reportError(hThread, TC("SuspendThread"));
 				continue;
+			}
+			auto rtg = MakeGuard([&]() { ResumeThread(hThread); });
 			CONTEXT ctx;
 			memset(&ctx, 0, sizeof(CONTEXT));
 			ctx.ContextFlags = CONTEXT_FULL;
-			GetThreadContext(hThread, &ctx);
+			if (!GetThreadContext(hThread, &ctx))
+			{
+				reportError(hThread, TC("GetThreadContext"));
+				continue;
+			}
 			func(0, &ctx);
-			ResumeThread(hThread);
 		}
 #else
 		DWORD pid = GetCurrentProcessId();
@@ -244,21 +253,31 @@ namespace uba
 				continue;
 			HANDLE hThread = OpenThread(THREAD_GET_CONTEXT | THREAD_SET_CONTEXT | THREAD_SUSPEND_RESUME | THREAD_QUERY_INFORMATION, FALSE, te32.th32ThreadID);
 			if (!hThread)
+			{
+				reportError(hThread, TC("OpenThread"));
 				continue;
+			}
 			auto tg = MakeGuard([&]() { CloseHandle(hThread); });
 			if (SuspendThread(hThread) == -1)
+			{
+				reportError(hThread, TC("SuspendThread"));
 				continue;
+			}
+			auto rtg = MakeGuard([&]() { ResumeThread(hThread); });
 			PWSTR threadDesc = nullptr;
 			GetThreadDescription(hThread, &threadDesc);
 			LocalFree(threadDesc);
 			CONTEXT ctx;
 			memset(&ctx, 0, sizeof(CONTEXT));
 			ctx.ContextFlags = CONTEXT_FULL;
-			GetThreadContext(hThread, &ctx);
+			if (!GetThreadContext(hThread, &ctx))
+			{
+				reportError(hThread, TC("GetThreadContext"));
+				continue;
+			}
 			void* callstack[100];
 			u32 callstackCount = GetCallstack(callstack, 100, 1, &ctx);
 			func(te32.th32ThreadID, callstack, callstackCount, threadDesc);
-			ResumeThread(hThread);
 		}
 		while (Thread32Next(hSnapshot, &te32));
 #endif
diff --git a/Engine/Source/Programs/UnrealBuildAccelerator/Common/Public/UbaLogger.h b/Engine/Source/Programs/UnrealBuildAccelerator/Common/Public/UbaLogger.h
--- a/Engine/Source/Programs/UnrealBuildAccelerator/Common/Public/UbaLogger.h
+++ b/Engine/Source/Programs/UnrealBuildAccelerator/Common/Public/UbaLogger.h
@@ -13,6 +13,7 @@ namespace uba
 	class StringBufferBase;
 	struct BinaryReader;
 	struct StringView;
+	using TraverseThreadErrorFunc = Function<void(const StringView& error)>;
 
 	struct LogEntry
 	{
@@ -99,6 +100,6 @@ namespace uba
 	void PrintContentionSummary(class Logger& logger);
 
 	struct CallstackInfo { Vector<u8> data; Vector<u32> threadIds; TString desc; };
-	void TraverseAllCallstacks(const Function<void(const CallstackInfo&)>& func);
+	void TraverseAllCallstacks(const Function<void(const CallstackInfo&)>& func, const TraverseThreadErrorFunc& errorFunc);
 	void PrintAllCallstacks(Logger& logger);
 }
diff --git a/Engine/Source/Programs/UnrealBuildAccelerator/Common/Public/UbaThread.h b/Engine/Source/Programs/UnrealBuildAccelerator/Common/Public/UbaThread.h
--- a/Engine/Source/Programs/UnrealBuildAccelerator/Common/Public/UbaThread.h
+++ b/Engine/Source/Programs/UnrealBuildAccelerator/Common/Public/UbaThread.h
@@ -12,6 +12,7 @@ namespace uba
 	class Event;
 	struct GroupAffinity;
 	using TraverseThreadFunc = Function<void(u32 tid, void** callstack, u32 callstackCount, const tchar* description)>;
+	using TraverseThreadErrorFunc = Function<void(const StringView& error)>;
 
 	class Thread
 	{
@@ -40,7 +41,7 @@ namespace uba
 		Thread(const Thread&) = delete;
 		void operator=(const Thread&) = delete;
 
-		friend bool TraverseAllThreads(const TraverseThreadFunc& func);
+		friend bool TraverseAllThreads(const TraverseThreadFunc& func, const TraverseThreadErrorFunc& errorFunc);
 	};
 
 
@@ -51,5 +52,5 @@ namespace uba
 	};
 	bool SetThreadGroupAffinity(void* nativeThreadHandle, const GroupAffinity& affinity);
 	bool AlternateThreadGroupAffinity(void* nativeThreadHandle);
-	bool TraverseAllThreads(const TraverseThreadFunc& func);
+	bool TraverseAllThreads(const TraverseThreadFunc& func, const TraverseThreadErrorFunc& errorFunc);
 }
diff --git a/Engine/Source/Programs/UnrealBuildAccelerator/Test/Private/UbaTestBasics.cpp b/Engine/Source/Programs/UnrealBuildAccelerator/Test/Private/UbaTestBasics.cpp
--- a/Engine/Source/Programs/UnrealBuildAccelerator/Test/Private/UbaTestBasics.cpp
+++ b/Engine/Source/Programs/UnrealBuildAccelerator/Test/Private/UbaTestBasics.cpp
@@ -601,6 +601,10 @@ namespace uba
 
 				ParseCallstackInfo(sb, reader, executable, searchPaths);
 				LoggerWithWriter(g_consoleLogWriter, TC("")).Info(TC("THREAD %u%s"), tid, sb.data);
+			},
+			[&](const StringView& error)
+			{
+				logger.Info(error.data);
 			});
 
 
