[UBA]
* Some minor stack resolving fixes on windows/wine

diff --git a/Engine/Source/Programs/UnrealBuildAccelerator/Common/Private/UbaThread.cpp b/Engine/Source/Programs/UnrealBuildAccelerator/Common/Private/UbaThread.cpp
--- a/Engine/Source/Programs/UnrealBuildAccelerator/Common/Private/UbaThread.cpp
+++ b/Engine/Source/Programs/UnrealBuildAccelerator/Common/Private/UbaThread.cpp
@@ -249,7 +249,7 @@ namespace uba
 			return false;
 		do
 		{
-			if (te32.th32OwnerProcessID != pid || te32.th32ThreadID == tid)
+			if (te32.th32OwnerProcessID != pid || te32.th32ThreadID == tid || te32.th32ThreadID == 0)
 				continue;
 			HANDLE hThread = OpenThread(THREAD_GET_CONTEXT | THREAD_SET_CONTEXT | THREAD_SUSPEND_RESUME | THREAD_QUERY_INFORMATION, FALSE, te32.th32ThreadID);
 			if (!hThread)
@@ -266,10 +266,9 @@ namespace uba
 			auto rtg = MakeGuard([&]() { ResumeThread(hThread); });
 			PWSTR threadDesc = nullptr;
 			GetThreadDescription(hThread, &threadDesc);
-			LocalFree(threadDesc);
-			CONTEXT ctx;
-			memset(&ctx, 0, sizeof(CONTEXT));
-			ctx.ContextFlags = CONTEXT_FULL;
+			auto tdg = MakeGuard([&]() { LocalFree(threadDesc); });
+			CONTEXT ctx = {};
+			ctx.ContextFlags = CONTEXT_ALL;
 			if (!GetThreadContext(hThread, &ctx))
 			{
 				reportError(hThread, TC("GetThreadContext"));
diff --git a/Engine/Source/Programs/UnrealBuildAccelerator/Core/Private/UbaPlatform.cpp b/Engine/Source/Programs/UnrealBuildAccelerator/Core/Private/UbaPlatform.cpp
--- a/Engine/Source/Programs/UnrealBuildAccelerator/Core/Private/UbaPlatform.cpp
+++ b/Engine/Source/Programs/UnrealBuildAccelerator/Core/Private/UbaPlatform.cpp
@@ -213,10 +213,9 @@ namespace uba
 			{
 				UNWIND_HISTORY_TABLE unwindHistoryTable;
 				RtlZeroMemory(&unwindHistoryTable, sizeof(UNWIND_HISTORY_TABLE));
-				for (u32 i=0; i<MaxDepth && GET_PC(context); ++i)
+				while (count < MaxDepth && GET_PC(context))
 				{
-					callers[i] = (void*)GET_PC(context);
-					count = i;
+					callers[count++] = (void*)GET_PC(context);
 					ULONG64 ImageBase = 0;
 					if (PRUNTIME_FUNCTION RuntimeFunction = RtlLookupFunctionEntry(GET_PC(context), &ImageBase, &unwindHistoryTable); RuntimeFunction)
 					{
@@ -228,7 +227,13 @@ namespace uba
 					}
 					else
 					{
-						GET_PC(context) = (*(u64*)GET_SP(context));
+						u64* spPtr = (u64*)GET_SP(context);
+						if (!spPtr)
+							break;
+						u64 sp = *spPtr;
+						if (!sp)
+							break;
+						GET_PC(context) = sp;
 						GET_SP(context) += 8;
 					}
 				}
