[UBA]
* Fixed bug where multiple win64 architectures couldn't compile remotely at the same time

diff --git a/Engine/Source/Programs/UnrealBuildAccelerator/Common/Private/UbaSessionClient.cpp b/Engine/Source/Programs/UnrealBuildAccelerator/Common/Private/UbaSessionClient.cpp
--- a/Engine/Source/Programs/UnrealBuildAccelerator/Common/Private/UbaSessionClient.cpp
+++ b/Engine/Source/Programs/UnrealBuildAccelerator/Common/Private/UbaSessionClient.cpp
@@ -940,11 +940,13 @@ namespace uba
 
 	bool SessionClient::GetFullFileName(GetFullFileNameResponse& out, const GetFullFileNameMessage& msg)
 	{
-		// TODO: There is a potential risk here where two different applications asks for the full name of a file
-		// and they have different bin/working dir.. and there are two versions of this file.
+		StringKeyHasher hasher;
+		hasher.UpdateNoCheck(msg.process.m_startInfo.applicationStr);
+		hasher.UpdateNoCheck(msg.fileName);
+		StringKey nameKey = ToStringKey(hasher);
 
 		SCOPED_FUTEX(m_nameToNameLookupLock, lock);
-		auto insres = m_nameToNameLookup.try_emplace(msg.fileName.data);
+		auto insres = m_nameToNameLookup.try_emplace(nameKey);
 		NameRec& rec = insres.first->second;
 		lock.Leave();
 		SCOPED_FUTEX(rec.lock, lock2);
diff --git a/Engine/Source/Programs/UnrealBuildAccelerator/Common/Public/UbaSessionClient.h b/Engine/Source/Programs/UnrealBuildAccelerator/Common/Public/UbaSessionClient.h
--- a/Engine/Source/Programs/UnrealBuildAccelerator/Common/Public/UbaSessionClient.h
+++ b/Engine/Source/Programs/UnrealBuildAccelerator/Common/Public/UbaSessionClient.h
@@ -147,7 +147,7 @@ namespace uba
 
 		Futex m_nameToNameLookupLock;
 		struct NameRec { TString name; TString virtualName; Futex lock; bool handled = false; };
-		UnorderedMap<TString, NameRec> m_nameToNameLookup;
+		UnorderedMap<StringKey, NameRec> m_nameToNameLookup;
 
 		struct HashRec { CasKey key; u64 serverTime = 0; Futex lock; };
 		UnorderedMap<StringKey, HashRec> m_nameToHashLookup;
diff --git a/Engine/Source/Programs/UnrealBuildAccelerator/Core/Private/UbaHash.cpp b/Engine/Source/Programs/UnrealBuildAccelerator/Core/Private/UbaHash.cpp
--- a/Engine/Source/Programs/UnrealBuildAccelerator/Core/Private/UbaHash.cpp
+++ b/Engine/Source/Programs/UnrealBuildAccelerator/Core/Private/UbaHash.cpp
@@ -23,6 +23,12 @@ namespace uba
 			blake3_hasher_update((blake3_hasher*)&hasher, str, strLen * sizeof(tchar));
 	}
 
+	void StringKeyHasher::UpdateNoCheck(const StringView& str)
+	{
+		if (str.count != 0)
+			blake3_hasher_update((blake3_hasher*)&hasher, str.data, str.count * sizeof(tchar));
+	}
+
 	StringKey ToStringKey(const tchar* str, u64 strLen)
 	{
 		CHECK_PATH(StringView(str, u32(strLen)));
diff --git a/Engine/Source/Programs/UnrealBuildAccelerator/Core/Public/UbaHash.h b/Engine/Source/Programs/UnrealBuildAccelerator/Core/Public/UbaHash.h
--- a/Engine/Source/Programs/UnrealBuildAccelerator/Core/Public/UbaHash.h
+++ b/Engine/Source/Programs/UnrealBuildAccelerator/Core/Public/UbaHash.h
@@ -98,6 +98,7 @@ namespace uba
 		StringKeyHasher();
 		void Update(const StringView& str);
 		void Update(const tchar* str, u64 strLen);
+		void UpdateNoCheck(const StringView& str);
 		u64 hasher[1912 / sizeof(u64)];
 	};
 
