[UBA]
* Remove variable m_connected not needed from UbaSessionClient

diff --git a/Engine/Source/Programs/UnrealBuildAccelerator/Common/Private/UbaSessionClient.cpp b/Engine/Source/Programs/UnrealBuildAccelerator/Common/Private/UbaSessionClient.cpp
--- a/Engine/Source/Programs/UnrealBuildAccelerator/Common/Private/UbaSessionClient.cpp
+++ b/Engine/Source/Programs/UnrealBuildAccelerator/Common/Private/UbaSessionClient.cpp
@@ -141,7 +141,6 @@ namespace uba
 	{
 		m_client.RegisterOnDisconnected([this]()
 			{
-				m_connected = false;
 				m_loop = false;
 			});
 
@@ -1282,8 +1281,7 @@ namespace uba
 		if (!msg.Send(reader, m_stats.connectMsg))
 			return;
 
-		m_connected = reader.ReadBool();
-		if (!m_connected)
+		if (!reader.ReadBool())
 		{
 			StringBuffer<> str;
 			reader.ReadString(str);
@@ -2587,7 +2585,7 @@ namespace uba
 
 	void SessionClient::TraceSessionUpdate()
 	{
-		if (m_connected)
+		if (m_loop)
 			SendPing(m_memAvail, m_memTotal);
 
 		// TODO: There should be some sort of log entry if we have a process that has done no progress for 30 minutes
diff --git a/Engine/Source/Programs/UnrealBuildAccelerator/Common/Public/UbaSessionClient.h b/Engine/Source/Programs/UnrealBuildAccelerator/Common/Public/UbaSessionClient.h
--- a/Engine/Source/Programs/UnrealBuildAccelerator/Common/Public/UbaSessionClient.h
+++ b/Engine/Source/Programs/UnrealBuildAccelerator/Common/Public/UbaSessionClient.h
@@ -124,7 +124,6 @@ namespace uba
 		u8 m_memKillLoadPercent;
 		bool m_disableCustomAllocator;
 		bool m_useBinariesAsVersion;
-		bool m_connected = false;
 		bool m_dedicated = false;
 		bool m_useStorage = true;
 		bool m_downloadDetoursLib = true;
