[UBA]
* Added sigint filter on uba threads to prevent them from exiting out of system primitives when pressing ctrl-c. This seems to fix the hangs happening when ctrl-c out of UBT

diff --git a/Engine/Source/Programs/UnrealBuildAccelerator/Common/Private/UbaProcess.cpp b/Engine/Source/Programs/UnrealBuildAccelerator/Common/Private/UbaProcess.cpp
--- a/Engine/Source/Programs/UnrealBuildAccelerator/Common/Private/UbaProcess.cpp
+++ b/Engine/Source/Programs/UnrealBuildAccelerator/Common/Private/UbaProcess.cpp
@@ -253,8 +253,21 @@ namespace uba
 				}, startInfo.description);
 		else
 		{
+			// This is needed to handle cltr-c. Otherwise this thread might exit before the detoured process which in turn might do things after ctrl-c which can cause deadlocks in detoured process
+			#if !PLATFORM_WINDOWS
+			sigset_t newMask;
+			sigset_t oldMask;
+			sigemptyset(&newMask);
+			sigaddset(&newMask, SIGINT);
+			pthread_sigmask(SIG_BLOCK, &newMask, &oldMask);
+			#endif
+
 			ThreadRun(environment);
 			ThreadExit();
+
+			#if !PLATFORM_WINDOWS
+			pthread_sigmask(SIG_BLOCK, &oldMask, nullptr);
+			#endif
 		}
 		return true;
 	}
diff --git a/Engine/Source/Programs/UnrealBuildAccelerator/Common/Private/UbaThread.cpp b/Engine/Source/Programs/UnrealBuildAccelerator/Common/Private/UbaThread.cpp
--- a/Engine/Source/Programs/UnrealBuildAccelerator/Common/Private/UbaThread.cpp
+++ b/Engine/Source/Programs/UnrealBuildAccelerator/Common/Private/UbaThread.cpp
@@ -102,6 +102,14 @@ namespace uba
 
 		err = pthread_create(&pth, &tattr, [](void* p) -> void*
 			{
+				// Ignore sigint. Uba is designed for someone else to cancel it
+				#if !PLATFORM_WINDOWS
+				sigset_t set;
+				sigemptyset(&set);
+				sigaddset(&set, SIGINT);
+				pthread_sigmask(SIG_BLOCK, &set, NULL);		
+				#endif
+
 				auto& t = *(Thread*)p;
 				int res = t.m_func();
 				t.m_finished.Set();
