[UBA]
* Changed so initial handshake is sent with one call to socket send
* Removed code where when receiving crypto we allow less than size.

diff --git a/Engine/Source/Programs/UnrealBuildAccelerator/Common/Private/UbaNetworkClient.cpp b/Engine/Source/Programs/UnrealBuildAccelerator/Common/Private/UbaNetworkClient.cpp
--- a/Engine/Source/Programs/UnrealBuildAccelerator/Common/Private/UbaNetworkClient.cpp
+++ b/Engine/Source/Programs/UnrealBuildAccelerator/Common/Private/UbaNetworkClient.cpp
@@ -132,32 +132,26 @@ namespace uba
 			}, nullptr, TC("Connecting"));
 
 
+		StackBinaryWriter<1024> handshakeData;
 		if (m_cryptoKey)
 		{
 			// If we have a crypto key we start by sending a predefined 128 bytes blob that is encrypted.
 			// If server decrypt it to the same blob, we're good on that part
-			u8 encryptedBuffer[1024];
+			u8* encryptedBuffer = handshakeData.AllocWrite(sizeof(EncryptionHandshakeString));//[1024];
 			memcpy(encryptedBuffer, EncryptionHandshakeString, sizeof(EncryptionHandshakeString));
 			if (!Crypto::Encrypt(m_logger, m_cryptoKey, encryptedBuffer, sizeof(EncryptionHandshakeString)))
 				return false;
-
-			NetworkBackend::SendContext handskakeContext;
-			if (!backend.Send(m_logger, backendConnection, encryptedBuffer, sizeof(EncryptionHandshakeString), handskakeContext, TC("CryptoHandshake")))
-				return false;
 		}
 
+		handshakeData.WriteU32(SystemNetworkVersion);
+		handshakeData.WriteBytes(&m_uid, sizeof(m_uid));
 
-		u32 version = SystemNetworkVersion;
-		NetworkBackend::SendContext versionContext;
-		if (!backend.Send(m_logger, backendConnection, &version, sizeof(version), versionContext, TC("Version")))
-			return false;
-
-		NetworkBackend::SendContext uidContext;
-		if (!backend.Send(m_logger, backendConnection, &m_uid, sizeof(m_uid), uidContext, TC("Uid")))
+		NetworkBackend::SendContext sendContext;
+		if (!backend.Send(m_logger, backendConnection, handshakeData.GetData(), u32(handshakeData.GetPosition()), sendContext, TC("Handshake")))
 			return false;
 
-		if (!rc.recvEvent.IsSet(5000)) // This can not happen. Since both callbacks are using rc we can't leave this function until we know we are not in the callbacks
-			return m_logger.Info(TC("Timed out waiting for connection response from server.")).ToFalse();
+		if (!rc.recvEvent.IsSet(20*1000)) // This can not happen. Since both callbacks are using rc we can't leave this function until we know we are not in the callbacks
+			return m_logger.Info(TC("Timed out after 20 seconds waiting for connection response from server.")).ToFalse();
 
 		m_isOrWasConnected.Set();
 
diff --git a/Engine/Source/Programs/UnrealBuildAccelerator/Common/Private/UbaNetworkServer.cpp b/Engine/Source/Programs/UnrealBuildAccelerator/Common/Private/UbaNetworkServer.cpp
--- a/Engine/Source/Programs/UnrealBuildAccelerator/Common/Private/UbaNetworkServer.cpp
+++ b/Engine/Source/Programs/UnrealBuildAccelerator/Common/Private/UbaNetworkServer.cpp
@@ -124,10 +124,7 @@ namespace uba
 			m_backend.SetRecvTimeout(m_backendConnection, m_server.m_receiveTimeoutMs);
 
 			if (requiresCrypto)
-			{
-				m_backend.SetAllowLessThanBodySize(m_backendConnection, true);
 				m_backend.SetRecvCallbacks(m_backendConnection, this, 0, ReceiveHandshakeHeader, ReceiveHandshakeBody, TC("ReceiveHandshake"));
-			}
 			else
 				m_backend.SetRecvCallbacks(m_backendConnection, this, 4, ReceiveVersion, nullptr, TC("ReceiveVersion"));
 		}
@@ -241,7 +238,6 @@ namespace uba
 					return logger.Warning(TC("Connection %u Crypto mismatch... (Handshake string is encrypted with different key than any registered keys)"), conn.m_id);
 			}
 
-			conn.m_backend.SetAllowLessThanBodySize(conn.m_backendConnection, false);
 			conn.m_backend.SetRecvCallbacks(conn.m_backendConnection, &conn, 4, ReceiveVersion, nullptr, TC("ReceiveVersion"));
 
 			return true;
