[UBA]
* Fixed logic around when session client is allowed to send pings

diff --git a/Engine/Source/Programs/UnrealBuildAccelerator/Common/Private/UbaSessionClient.cpp b/Engine/Source/Programs/UnrealBuildAccelerator/Common/Private/UbaSessionClient.cpp
--- a/Engine/Source/Programs/UnrealBuildAccelerator/Common/Private/UbaSessionClient.cpp
+++ b/Engine/Source/Programs/UnrealBuildAccelerator/Common/Private/UbaSessionClient.cpp
@@ -1364,7 +1364,6 @@ namespace uba
 
 		BuildEnvironmentVariables(reader);
 
-		m_loop = true;
 		m_loopThread.Start([this]() { ThreadCreateProcessLoop(); return 0; }, TC("UbaCreateProcessLoop"));
 	}
 
@@ -1734,6 +1733,8 @@ namespace uba
 
 	void SessionClient::ThreadCreateProcessLoop()
 	{
+		m_sendPing = true;
+
 		// Ask for dir and hash table straight away to minimize latency
 		m_client.AddWork([this](const WorkContext&)
 			{
@@ -2585,7 +2586,7 @@ namespace uba
 
 	void SessionClient::TraceSessionUpdate()
 	{
-		if (m_loop)
+		if (m_loop && m_sendPing)
 			SendPing(m_memAvail, m_memTotal);
 
 		// TODO: There should be some sort of log entry if we have a process that has done no progress for 30 minutes
diff --git a/Engine/Source/Programs/UnrealBuildAccelerator/Common/Public/UbaSessionClient.h b/Engine/Source/Programs/UnrealBuildAccelerator/Common/Public/UbaSessionClient.h
--- a/Engine/Source/Programs/UnrealBuildAccelerator/Common/Public/UbaSessionClient.h
+++ b/Engine/Source/Programs/UnrealBuildAccelerator/Common/Public/UbaSessionClient.h
@@ -163,6 +163,7 @@ namespace uba
 		Event m_waitToSendEvent;
 		Thread m_loopThread;
 		Atomic<bool> m_loop;
+		Atomic<bool> m_sendPing;
 		Atomic<bool> m_allowSpawn;
 
 		Function<void(const ProcessHandle&)> m_processFinished;
