[UBA]
* Fixed bug in ping function where code never called ProcessAsyncResults on async message causing weird errors

diff --git a/Engine/Source/Programs/UnrealBuildAccelerator/Common/Private/UbaSessionClient.cpp b/Engine/Source/Programs/UnrealBuildAccelerator/Common/Private/UbaSessionClient.cpp
--- a/Engine/Source/Programs/UnrealBuildAccelerator/Common/Private/UbaSessionClient.cpp
+++ b/Engine/Source/Programs/UnrealBuildAccelerator/Common/Private/UbaSessionClient.cpp
@@ -1613,6 +1613,11 @@ namespace uba
 		
 		gotResponse.IsSet(); // Wait forever (until message trigger error or finishes)
 
+		if (!msg.ProcessAsyncResults(reader) || msg.GetError())
+		{
+			m_loop = false;
+			return;
+		}
 		u64 newTime = GetTime();
 		u64 lastPing = newTime - time;
 		m_lastPing = lastPing;
@@ -1624,8 +1629,10 @@ namespace uba
 		m_storage.Ping();
 
 		if (reader.ReadBool()) // abort
+		{
+			LoggerWithWriter(g_consoleLogWriter).Info("Got abort from server");
 			abort();
-
+		}
 		if (reader.ReadBool()) // crashdump
 		{
 			TraverseAllCallstacks([&](const CallstackInfo& cs)
