[UBA]
* Added asserts to return values on all linux sync primives

diff --git a/Engine/Source/Programs/UnrealBuildAccelerator/Common/Private/UbaSessionClient.cpp b/Engine/Source/Programs/UnrealBuildAccelerator/Common/Private/UbaSessionClient.cpp
--- a/Engine/Source/Programs/UnrealBuildAccelerator/Common/Private/UbaSessionClient.cpp
+++ b/Engine/Source/Programs/UnrealBuildAccelerator/Common/Private/UbaSessionClient.cpp
@@ -1630,7 +1630,7 @@ namespace uba
 
 		if (reader.ReadBool()) // abort
 		{
-			LoggerWithWriter(g_consoleLogWriter).Info("Got abort from server");
+			LoggerWithWriter(g_consoleLogWriter).Info(TC("Got abort from server"));
 			abort();
 		}
 		if (reader.ReadBool()) // crashdump
diff --git a/Engine/Source/Programs/UnrealBuildAccelerator/Core/Private/UbaSynchronization.cpp b/Engine/Source/Programs/UnrealBuildAccelerator/Core/Private/UbaSynchronization.cpp
--- a/Engine/Source/Programs/UnrealBuildAccelerator/Core/Private/UbaSynchronization.cpp
+++ b/Engine/Source/Programs/UnrealBuildAccelerator/Core/Private/UbaSynchronization.cpp
@@ -54,7 +54,8 @@ namespace uba
 		#if PLATFORM_WINDOWS
 		EnterCriticalSection((CRITICAL_SECTION*)&data);
 		#else
-		pthread_mutex_lock((pthread_mutex_t*)data);
+		int res = pthread_mutex_lock((pthread_mutex_t*)data);(void)res;
+		UBA_ASSERTF(res == 0, TC("pthread_mutex_lock failed: %i"), res);
 		#endif
 	}
 
@@ -63,7 +64,8 @@ namespace uba
 		#if PLATFORM_WINDOWS
 		LeaveCriticalSection((CRITICAL_SECTION*)&data);
 		#else
-		pthread_mutex_unlock((pthread_mutex_t*)data);
+		int res = pthread_mutex_unlock((pthread_mutex_t*)data);(void)res;
+		UBA_ASSERTF(res == 0, TC("pthread_mutex_unlock failed: %i"), res);
 		#endif
 	}
 
@@ -109,7 +111,8 @@ namespace uba
 		#elif UBA_USE_SHARED_MUTEX
 		((std::shared_mutex*)data)->lock_shared();
 		#else
-		pthread_rwlock_rdlock((pthread_rwlock_t*)data);
+		int res = pthread_rwlock_rdlock((pthread_rwlock_t*)data);(void)res;
+		UBA_ASSERTF(res == 0, TC("pthread_rwlock_rdlock failed: %i"), res);
 		#endif
 	}
 
@@ -120,7 +123,8 @@ namespace uba
 		#elif UBA_USE_SHARED_MUTEX
 		((std::shared_mutex*)data)->unlock_shared();
 		#else
-		pthread_rwlock_unlock((pthread_rwlock_t*)data);
+		int res = pthread_rwlock_unlock((pthread_rwlock_t*)data);(void)res;
+		UBA_ASSERTF(res == 0, TC("pthread_rwlock_unlock failed: %i"), res);
 		#endif
 	}
 
@@ -131,7 +135,8 @@ namespace uba
 		#elif UBA_USE_SHARED_MUTEX
 		((std::shared_mutex*)data)->lock();
 		#else
-		pthread_rwlock_wrlock((pthread_rwlock_t*)data);
+		int res = pthread_rwlock_wrlock((pthread_rwlock_t*)data);(void)res;
+		UBA_ASSERTF(res == 0, TC("pthread_rwlock_wrlock failed: %i"), res);
 		#endif
 	}
 
@@ -153,7 +158,8 @@ namespace uba
 		#elif UBA_USE_SHARED_MUTEX
 		((std::shared_mutex*)data)->unlock();
 		#else
-		pthread_rwlock_unlock((pthread_rwlock_t*)data);
+		int res = pthread_rwlock_unlock((pthread_rwlock_t*)data);(void)res;
+		UBA_ASSERTF(res == 0, TC("pthread_rwlock_unlock failed: %i"), res);
 		#endif
 	}
 
