[UBA]
* Fixed so network backend memory can handle multiple messages in one send call

diff --git a/Engine/Source/Programs/UnrealBuildAccelerator/Common/Private/UbaNetworkBackendMemory.cpp b/Engine/Source/Programs/UnrealBuildAccelerator/Common/Private/UbaNetworkBackendMemory.cpp
--- a/Engine/Source/Programs/UnrealBuildAccelerator/Common/Private/UbaNetworkBackendMemory.cpp
+++ b/Engine/Source/Programs/UnrealBuildAccelerator/Common/Private/UbaNetworkBackendMemory.cpp
@@ -60,26 +60,34 @@ namespace uba
 
 		auto& peer = m_connection->peer[to];
 
-		if (!m_connection->connected)
-			return false;
+		u32 left = dataSize;
+		while (left)
+		{
+			if (!m_connection->connected)
+				return false;
 
-		void* bodyContext = nullptr;
-		u8* bodyData = nullptr;
-		u32 bodySize = 0;
+			void* bodyContext = nullptr;
+			u8* bodyData = nullptr;
+			u32 bodySize = 0;
 
-		if (!peer.headerCallback(peer.context, m_connection->uid, (u8*)data, bodyContext, bodyData, bodySize))
-			return false;
+			left -= peer.headerSize;
 
-		if (!bodySize)
-			return true;
+			if (!peer.headerCallback(peer.context, m_connection->uid, (u8*)data, bodyContext, bodyData, bodySize))
+				return false;
 
-		memcpy(bodyData, ((const u8*)data) + peer.headerSize, bodySize);
-		auto bc = peer.bodyCallback;
-		if (!bc)
-			return logger.Error(TC("Connection body callback not set"));
+			if (!bodySize)
+				continue;
 
-		if (!bc(peer.context, false, (u8*)data, bodyContext, bodyData, bodySize))
-			return false;
+			left -= bodySize;
+
+			memcpy(bodyData, ((const u8*)data) + peer.headerSize, bodySize);
+			auto bc = peer.bodyCallback;
+			if (!bc)
+				return logger.Error(TC("Connection body callback not set"));
+
+			if (!bc(peer.context, false, (u8*)data, bodyContext, bodyData, bodySize))
+				return false;
+		}
 		return true;
 	}
 
