[UBA]
* Fixed bug in memory network backend when handling multiple messages and not reading correct data

diff --git a/Engine/Source/Programs/UnrealBuildAccelerator/Common/Private/UbaNetworkBackendMemory.cpp b/Engine/Source/Programs/UnrealBuildAccelerator/Common/Private/UbaNetworkBackendMemory.cpp
--- a/Engine/Source/Programs/UnrealBuildAccelerator/Common/Private/UbaNetworkBackendMemory.cpp
+++ b/Engine/Source/Programs/UnrealBuildAccelerator/Common/Private/UbaNetworkBackendMemory.cpp
@@ -7,6 +7,8 @@
 #include "UbaStringBuffer.h"
 #include "UbaTimer.h"
 
+#define UBA_VALIDATE_MULTIPLE_SEND 0
+
 namespace uba
 {
 	struct NetworkBackendMemory::Connection
@@ -21,6 +23,10 @@ namespace uba
 			
 			DisconnectCallback* disconnectCallback = nullptr;
 			void* disconnectContext = nullptr;
+
+			#if UBA_VALIDATE_MULTIPLE_SEND
+			Atomic<u32> sendCount;
+			#endif
 		};
 
 		Peer peer[2];
@@ -60,8 +66,14 @@ namespace uba
 
 		auto& peer = m_connection->peer[to];
 
-		u32 left = dataSize;
-		while (left)
+		#if UBA_VALIDATE_MULTIPLE_SEND
+		++peer.sendCount;
+		auto sg = MakeGuard([&]() { --peer.sendCount; });
+		#endif
+
+		u8* dataPos = (u8*)data;
+		u32 dataLeft = dataSize;
+		while (dataLeft)
 		{
 			if (!m_connection->connected)
 				return false;
@@ -70,23 +82,29 @@ namespace uba
 			u8* bodyData = nullptr;
 			u32 bodySize = 0;
 
-			left -= peer.headerSize;
+			u32 headerSize = peer.headerSize;
 
-			if (!peer.headerCallback(peer.context, m_connection->uid, (u8*)data, bodyContext, bodyData, bodySize))
+			if (!peer.headerCallback(peer.context, m_connection->uid, dataPos, bodyContext, bodyData, bodySize))
 				return false;
 
+			UBA_ASSERT(dataLeft >= headerSize);
+			dataPos += headerSize;
+			dataLeft -= headerSize;
+
 			if (!bodySize)
 				continue;
 
-			left -= bodySize;
-
-			memcpy(bodyData, ((const u8*)data) + peer.headerSize, bodySize);
+			memcpy(bodyData, dataPos, bodySize);
 			auto bc = peer.bodyCallback;
 			if (!bc)
 				return logger.Error(TC("Connection body callback not set"));
 
-			if (!bc(peer.context, false, (u8*)data, bodyContext, bodyData, bodySize))
+			if (!bc(peer.context, false, dataPos, bodyContext, bodyData, bodySize))
 				return false;
+
+			UBA_ASSERT(dataLeft >= bodySize);
+			dataPos += bodySize;
+			dataLeft -= bodySize;
 		}
 		return true;
 	}
@@ -99,6 +117,11 @@ namespace uba
 	{
 		UBA_ASSERT(m_connection);
 		auto& peer = m_connection->peer[((uintptr_t)connection)-1];
+
+		#if UBA_VALIDATE_MULTIPLE_SEND
+		UBA_ASSERT(peer.sendCount <= 1);
+		#endif
+
 		peer.headerCallback = h;
 		peer.bodyCallback = b;
 		peer.context = context;
@@ -137,7 +160,7 @@ namespace uba
 	bool NetworkBackendMemory::Connect(Logger& logger, const tchar* ip, const ConnectedFunc& connectedFunc, u16 port, bool* timedOut)
 	{
 		u64 start = GetTime();
-		SCOPED_READ_LOCK(m_connectedFuncLock, l);
+		SCOPED_WRITE_LOCK(m_connectedFuncLock, l);
 		while (!m_connectedFunc)
 		{
 			if (TimeToMs(GetTime() - start) > 1000*2)
diff --git a/Engine/Source/Programs/UnrealBuildAccelerator/Test/Private/UbaTestNetwork.cpp b/Engine/Source/Programs/UnrealBuildAccelerator/Test/Private/UbaTestNetwork.cpp
--- a/Engine/Source/Programs/UnrealBuildAccelerator/Test/Private/UbaTestNetwork.cpp
+++ b/Engine/Source/Programs/UnrealBuildAccelerator/Test/Private/UbaTestNetwork.cpp
@@ -175,6 +175,8 @@ namespace uba
 			return logger.Error(TC("Failed to get message"));
 
 		u8 value = reader.ReadByte();
+		if (value != 42)
+			return logger.Error(TC("Got wrong value %u, expected 42"), value);
 		logger.Info(TC("Got value %u"), value);
 		return true;
 	}
