[UBA]
* Reduce lock scope when sending keep alive message to proxy server

diff --git a/Engine/Source/Programs/UnrealBuildAccelerator/Common/Private/UbaStorageClient.cpp b/Engine/Source/Programs/UnrealBuildAccelerator/Common/Private/UbaStorageClient.cpp
--- a/Engine/Source/Programs/UnrealBuildAccelerator/Common/Private/UbaStorageClient.cpp
+++ b/Engine/Source/Programs/UnrealBuildAccelerator/Common/Private/UbaStorageClient.cpp
@@ -948,14 +948,22 @@ namespace uba
 
 	void StorageClient::Ping()
 	{
-		SCOPED_READ_LOCK(m_proxyClientLock, lock);
-		if (!m_proxyClient)
+		SCOPED_WRITE_LOCK(m_proxyClientLock, lock);
+		auto proxy = m_proxyClient;
+		if (!proxy)
 			return;
 		u64 now = GetTime();
 		if (TimeToMs(now - m_proxyClientKeepAliveTime) < 30 * 1000)
 			return;
 		m_proxyClientKeepAliveTime = now;
+
+		++proxy->refCount;
+		lock.Leave();
+
 		m_proxyClient->client.SendKeepAlive();
+
+		lock.Enter();
+		--proxy->refCount;
 	}
 
 	void StorageClient::PrintSummary(Logger& logger)
