[UBA]
* Improvement error comment

diff --git a/Engine/Source/Programs/UnrealBuildAccelerator/Common/Private/UbaSessionClient.cpp b/Engine/Source/Programs/UnrealBuildAccelerator/Common/Private/UbaSessionClient.cpp
--- a/Engine/Source/Programs/UnrealBuildAccelerator/Common/Private/UbaSessionClient.cpp
+++ b/Engine/Source/Programs/UnrealBuildAccelerator/Common/Private/UbaSessionClient.cpp
@@ -1606,7 +1606,7 @@ namespace uba
 		if (!gotResponse.IsSet(20*1000)) // We accept 20 seconds time out until we write out callstacks
 		{
 			LoggerWithWriter logger(g_consoleLogWriter);
-			logger.Info(TC("Took more than 20 seconds to send data. Printing callstacks"));
+			logger.Info(TC("Took more than 20 seconds to send/receive ping. Printing callstacks"));
 			PrintAllCallstacks(logger);
 		}
 		#endif
