[UBA]
* Fixed bug where crazy many network connections could be made when client asks server to create connections

diff --git a/Engine/Source/Programs/UnrealBuildAccelerator/Common/Private/UbaNetworkServer.cpp b/Engine/Source/Programs/UnrealBuildAccelerator/Common/Private/UbaNetworkServer.cpp
--- a/Engine/Source/Programs/UnrealBuildAccelerator/Common/Private/UbaNetworkServer.cpp
+++ b/Engine/Source/Programs/UnrealBuildAccelerator/Common/Private/UbaNetworkServer.cpp
@@ -1223,24 +1223,34 @@ namespace uba
 			case SystemMessageType_SetConnectionCount:
 			{
 				u32 connectionCount = reader.ReadU32();
+				u32 connectionId = connectionInfo.GetId();
 
 				SCOPED_READ_LOCK(m_clientsLock, lock);
-				auto findIt = m_clients.find(connectionInfo.GetId());
+				auto findIt = m_clients.find(connectionId);
 				if (findIt == m_clients.end())
 					return true;
 				Client& c = findIt->second;
-				lock.Leave();
 
-				if (c.connectionCount >= connectionCount)
+				u32 currentCount = c.connectionCount + c.queuedConnectionCount;
+				if (currentCount >= connectionCount)
 					return true;
-				u32 toAdd = connectionCount - c.connectionCount;
+				u32 toAdd = connectionCount - currentCount;
+				c.queuedConnectionCount += toAdd;
+				lock.Leave();
 
-				u32 connectionId = connectionInfo.GetId();
 				SCOPED_FUTEX(m_addConnectionsLock, lock2);
 				for (u32 i = 0; i != toAdd; ++i)
 				{
 					m_addConnections.emplace_back([this, connectionId]()
 						{
+							auto cg = MakeGuard([&]()
+								{
+									SCOPED_READ_LOCK(m_clientsLock, lock);
+									auto findIt = m_clients.find(connectionId);
+									if (findIt != m_clients.end())
+										--findIt->second.queuedConnectionCount;
+								});
+
 							Connection* conn = nullptr;
 							SCOPED_READ_LOCK(m_connectionsLock, lock);
 							for (auto& c : m_connections)
diff --git a/Engine/Source/Programs/UnrealBuildAccelerator/Common/Public/UbaNetworkServer.h b/Engine/Source/Programs/UnrealBuildAccelerator/Common/Public/UbaNetworkServer.h
--- a/Engine/Source/Programs/UnrealBuildAccelerator/Common/Public/UbaNetworkServer.h
+++ b/Engine/Source/Programs/UnrealBuildAccelerator/Common/Public/UbaNetworkServer.h
@@ -177,6 +177,7 @@ namespace uba
 			Client(const Guid& uid_, u32 id_) : uid(uid_), id(id_) {}
 			Guid uid;
 			u32 id;
+			u32 queuedConnectionCount = 0;
 			Atomic<u32> refCount;
 			Atomic<u32> connectionCount;
 			Atomic<u64> sendBytes;
