[UBA]
* Fixed static analysis errors

diff --git a/Engine/Source/Programs/UnrealBuildAccelerator/Common/Private/UbaCompactTables.cpp b/Engine/Source/Programs/UnrealBuildAccelerator/Common/Private/UbaCompactTables.cpp
--- a/Engine/Source/Programs/UnrealBuildAccelerator/Common/Private/UbaCompactTables.cpp
+++ b/Engine/Source/Programs/UnrealBuildAccelerator/Common/Private/UbaCompactTables.cpp
@@ -382,7 +382,7 @@ namespace uba
 
 	u32 CompactPathTable::InternalAddWithNewSegment(u32 offset, u32 parentOffset, bool separator, const StringView& segment)
 	{
-		u32 parentRelativeOffset = ((offset - parentOffset) << 1) + separator;
+		u32 parentRelativeOffset = ((offset - parentOffset) << 1u) + u32(separator);
 
 		// Put string directly after current element and set segment offset to 0
 		u64 bytesForString = GetStringWriteSize(segment.data, segment.count);
@@ -398,7 +398,7 @@ namespace uba
 
 	void CompactPathTable::InternalAddWithExistingSegment(u32 offset, u32 parentOffset, bool separator, u32 segmentOffset)
 	{
-		u32 parentRelativeOffset = ((offset - parentOffset) << 1) + separator;
+		u32 parentRelativeOffset = ((offset - parentOffset) << 1u) + u32(separator);
 
 		u8 bytesForParent = Get7BitEncodedCount(parentRelativeOffset);
 		u64 memSize = bytesForParent + Get7BitEncodedCount(segmentOffset);
diff --git a/Engine/Source/Programs/UnrealBuildAccelerator/Common/Private/UbaFile.cpp b/Engine/Source/Programs/UnrealBuildAccelerator/Common/Private/UbaFile.cpp
--- a/Engine/Source/Programs/UnrealBuildAccelerator/Common/Private/UbaFile.cpp
+++ b/Engine/Source/Programs/UnrealBuildAccelerator/Common/Private/UbaFile.cpp
@@ -1036,7 +1036,8 @@ namespace uba
 
 
 			tchar driveLetters[MAX_PATH] = {0};
-			if (GetVolumePathNamesForVolumeNameW(volumeName, driveLetters, sizeof_array(driveLetters), NULL))
+			DWORD cchReturnLength = 0;
+			if (GetVolumePathNamesForVolumeNameW(volumeName, driveLetters, sizeof_array(driveLetters), &cchReturnLength))
 			{
 				StringBuffer<128> drives;
 				for (tchar* p = driveLetters; *p; p += TStrlen(p) + 1)
diff --git a/Engine/Source/Programs/UnrealBuildAccelerator/Common/Private/UbaFileMapping.cpp b/Engine/Source/Programs/UnrealBuildAccelerator/Common/Private/UbaFileMapping.cpp
--- a/Engine/Source/Programs/UnrealBuildAccelerator/Common/Private/UbaFileMapping.cpp
+++ b/Engine/Source/Programs/UnrealBuildAccelerator/Common/Private/UbaFileMapping.cpp
@@ -554,7 +554,7 @@ namespace uba
 			{
 				if (storage.fullFileCount == storage.fileCount)
 				{
-					m_logger.Error(TC("%s - All files are full!"), file->name);
+					m_logger.Error(TC("All files are full!"));
 					return nullptr;
 				}
 				lock.Leave();
diff --git a/Engine/Source/Programs/UnrealBuildAccelerator/Common/Private/UbaSessionClient.cpp b/Engine/Source/Programs/UnrealBuildAccelerator/Common/Private/UbaSessionClient.cpp
--- a/Engine/Source/Programs/UnrealBuildAccelerator/Common/Private/UbaSessionClient.cpp
+++ b/Engine/Source/Programs/UnrealBuildAccelerator/Common/Private/UbaSessionClient.cpp
@@ -553,6 +553,10 @@ namespace uba
 			u64 fileSize = InvalidValue;
 			CasKey newCasKey;
 
+			#ifdef __clang_analyzer__
+			newCasKey = {};
+			#endif
+
 			u32 memoryMapAlignment = 0;
 			if (m_allowMemoryMaps)
 			{
diff --git a/Engine/Source/Programs/UnrealBuildAccelerator/Common/Private/UbaSessionServer.cpp b/Engine/Source/Programs/UnrealBuildAccelerator/Common/Private/UbaSessionServer.cpp
--- a/Engine/Source/Programs/UnrealBuildAccelerator/Common/Private/UbaSessionServer.cpp
+++ b/Engine/Source/Programs/UnrealBuildAccelerator/Common/Private/UbaSessionServer.cpp
@@ -335,6 +335,10 @@ namespace uba
 			void* knownInputsCopy = malloc(knownInputsBytes);
 			memcpy(knownInputsCopy, knownInputs, knownInputsBytes);
 
+			#ifdef __clang_analyzer__
+			free(knownInputsCopy); // doesn't seem to understand it is handed over
+			#endif
+
 			m_server.AddWork([remoteProcess, this, knownInputsCopy, knownInputsCount, rules](const WorkContext& context)
 				{
 					auto keys = remoteProcess->m_knownInputs = new RemoteProcess::KnownInput[knownInputsCount];
diff --git a/Engine/Source/Programs/UnrealBuildAccelerator/Common/Private/UbaStorage.cpp b/Engine/Source/Programs/UnrealBuildAccelerator/Common/Private/UbaStorage.cpp
--- a/Engine/Source/Programs/UnrealBuildAccelerator/Common/Private/UbaStorage.cpp
+++ b/Engine/Source/Programs/UnrealBuildAccelerator/Common/Private/UbaStorage.cpp
@@ -1305,7 +1305,6 @@ namespace uba
 			return exclusiveMutex;
 
 		CloseMutex(exclusiveMutex);
-		exclusiveMutex = InvalidMutexHandle;
 		if (reportError)
 			logger.Error(TC("Needs exclusive access to storage %s. Another process is running"), rootDir.data);
 		return InvalidMutexHandle;
diff --git a/Engine/Source/Programs/UnrealBuildAccelerator/Common/Private/UbaStorageClient.cpp b/Engine/Source/Programs/UnrealBuildAccelerator/Common/Private/UbaStorageClient.cpp
--- a/Engine/Source/Programs/UnrealBuildAccelerator/Common/Private/UbaStorageClient.cpp
+++ b/Engine/Source/Programs/UnrealBuildAccelerator/Common/Private/UbaStorageClient.cpp
@@ -371,7 +371,8 @@ namespace uba
 				if (fetchId == FetchCasIdDisallowed)
 				{
 					m_logger.Error(TC("Disallowed cas %s (%s)"), casFile.data, hint); // Log first since Disconnect mute log
-					casEntry->disallowed = true;
+					if (casEntry)
+						casEntry->disallowed = true;
 					casEntryLock.Execute();
 					if (proxy)
 						proxy->client.Disconnect();
diff --git a/Engine/Source/Programs/UnrealBuildAccelerator/Core/Private/UbaPlatform.cpp b/Engine/Source/Programs/UnrealBuildAccelerator/Core/Private/UbaPlatform.cpp
--- a/Engine/Source/Programs/UnrealBuildAccelerator/Core/Private/UbaPlatform.cpp
+++ b/Engine/Source/Programs/UnrealBuildAccelerator/Core/Private/UbaPlatform.cpp
@@ -196,7 +196,7 @@ namespace uba
 	{
 #if PLATFORM_WINDOWS
 		constexpr int MaxDepth = 16;// 62;
-		void* callers[MaxDepth];
+		void* callers[MaxDepth] = {0};
 		u32 count = 0;//RtlCaptureStackBackTrace(0, MaxDepth, callers, NULL); // For some reason this does not work on wine
 		if (!count)
 		{
diff --git a/Engine/Source/Programs/UnrealBuildAccelerator/Host/Private/UbaExports.cpp b/Engine/Source/Programs/UnrealBuildAccelerator/Host/Private/UbaExports.cpp
--- a/Engine/Source/Programs/UnrealBuildAccelerator/Host/Private/UbaExports.cpp
+++ b/Engine/Source/Programs/UnrealBuildAccelerator/Host/Private/UbaExports.cpp
@@ -283,8 +283,9 @@ extern "C"
 	uba::StorageServer* StorageServer_InternalCreate(uba::StorageServerCreateInfo& info)
 	{
 		using namespace uba;
-		#ifndef __clang_analyzer__
 		#if UBA_USE_CLOUD
+		auto oldDir = info.rootDir;
+		auto oldZone = info.zone;
 		StringBuffer<> fixedRootDir;
 		fixedRootDir.count = GetFullPathNameW(info.rootDir, fixedRootDir.capacity, fixedRootDir.data, NULL);
 		fixedRootDir.Replace('/', PathSeparator).EnsureEndsWithSlash();
@@ -296,7 +297,6 @@ extern "C"
 			if (cloud.QueryAvailabilityZone(downgradedLogger, info.rootDir))
 				info.zone = cloud.GetAvailabilityZone();
 		}
-		#endif
 
 		StringBuffer<256> zoneTemp;
 		if (!info.zone || !*info.zone)
@@ -308,7 +308,14 @@ extern "C"
 
 		#endif
 
-		return new StorageServer(info);
+		auto storageServer = new StorageServer(info);
+
+		#if UBA_USE_CLOUD
+		info.rootDir = oldDir;
+		info.zone = oldZone;
+		info.exclusiveMutex = InvalidMutexHandle;
+		#endif
+		return storageServer;
 	}
 
 	uba::StorageServer* StorageServer_Create(uba::NetworkServer& server, const uba::tchar* rootDir, uba::u64 casCapacityBytes, bool storeCompressed, uba::LogWriter& writer, const uba::tchar* zone)
