[UBA]
* Added log entry for client requesting more connections

diff --git a/Engine/Source/Programs/UnrealBuildAccelerator/Common/Private/UbaNetworkServer.cpp b/Engine/Source/Programs/UnrealBuildAccelerator/Common/Private/UbaNetworkServer.cpp
--- a/Engine/Source/Programs/UnrealBuildAccelerator/Common/Private/UbaNetworkServer.cpp
+++ b/Engine/Source/Programs/UnrealBuildAccelerator/Common/Private/UbaNetworkServer.cpp
@@ -1236,6 +1236,7 @@ namespace uba
 					return true;
 				u32 toAdd = connectionCount - currentCount;
 				c.queuedConnectionCount += toAdd;
+				m_logger.Info(TC("Client requested %u connections. Has %u, queue %u"), connectionCount, c.connectionCount.load(), c.queuedConnectionCount);
 				lock.Leave();
 
 				SCOPED_FUTEX(m_addConnectionsLock, lock2);
