[UBA]
* Added logging entries for changing connection count

diff --git a/Engine/Source/Programs/UnrealBuildAccelerator/Agent/Private/UbaAgent.cpp b/Engine/Source/Programs/UnrealBuildAccelerator/Agent/Private/UbaAgent.cpp
--- a/Engine/Source/Programs/UnrealBuildAccelerator/Agent/Private/UbaAgent.cpp
+++ b/Engine/Source/Programs/UnrealBuildAccelerator/Agent/Private/UbaAgent.cpp
@@ -1341,6 +1341,7 @@ namespace uba
 				{
 					if (connectionCount != targetConnectionCount)
 					{
+						logger.Info(TC("Updating desired connection count to %u"), connectionCount);
 						client->SetConnectionCount(targetConnectionCount);
 						connectionCount = targetConnectionCount;
 					}
diff --git a/Engine/Source/Programs/UnrealBuildAccelerator/Common/Private/UbaNetworkServer.cpp b/Engine/Source/Programs/UnrealBuildAccelerator/Common/Private/UbaNetworkServer.cpp
--- a/Engine/Source/Programs/UnrealBuildAccelerator/Common/Private/UbaNetworkServer.cpp
+++ b/Engine/Source/Programs/UnrealBuildAccelerator/Common/Private/UbaNetworkServer.cpp
@@ -1236,7 +1236,7 @@ namespace uba
 					return true;
 				u32 toAdd = connectionCount - currentCount;
 				c.queuedConnectionCount += toAdd;
-				m_logger.Info(TC("Client requested %u connections. Has %u, queue %u"), connectionCount, c.connectionCount.load(), c.queuedConnectionCount);
+				m_logger.Detail(TC("Client %u requested %u connections. Has %u, queue %u"), c.id, connectionCount, c.connectionCount.load(), c.queuedConnectionCount);
 				lock.Leave();
 
 				SCOPED_FUTEX(m_addConnectionsLock, lock2);
