[UBA]
* Fixed bug where server did additional connection to the wrong client when a client requested more connections (mix up of ids)
* Improved some logging entries
* Disabled printing of callstacks when send takes too long time. Can be enabled with [Session] PingTimeoutSecondsPrintCallstacks
* Improved ping calculations to be more correct and only contain the actual bouncing message
* Improved storage ping

diff --git a/Engine/Source/Programs/UnrealBuildAccelerator/Agent/Private/UbaAgent.cpp b/Engine/Source/Programs/UnrealBuildAccelerator/Agent/Private/UbaAgent.cpp
--- a/Engine/Source/Programs/UnrealBuildAccelerator/Agent/Private/UbaAgent.cpp
+++ b/Engine/Source/Programs/UnrealBuildAccelerator/Agent/Private/UbaAgent.cpp
@@ -1341,9 +1341,10 @@ namespace uba
 				{
 					if (connectionCount != targetConnectionCount)
 					{
-						logger.Info(TC("Updating desired connection count to %u"), connectionCount);
-						client->SetConnectionCount(targetConnectionCount);
-						connectionCount = targetConnectionCount;
+						u32 newConnectionCount = targetConnectionCount;
+						logger.Info(TC("Updating desired connection count from %u to %u"), connectionCount, newConnectionCount);
+						client->SetConnectionCount(newConnectionCount);
+						connectionCount = newConnectionCount;
 					}
 				}
 				else
diff --git a/Engine/Source/Programs/UnrealBuildAccelerator/Common/Private/UbaNetworkServer.cpp b/Engine/Source/Programs/UnrealBuildAccelerator/Common/Private/UbaNetworkServer.cpp
--- a/Engine/Source/Programs/UnrealBuildAccelerator/Common/Private/UbaNetworkServer.cpp
+++ b/Engine/Source/Programs/UnrealBuildAccelerator/Common/Private/UbaNetworkServer.cpp
@@ -1223,10 +1223,10 @@ namespace uba
 			case SystemMessageType_SetConnectionCount:
 			{
 				u32 connectionCount = reader.ReadU32();
-				u32 connectionId = connectionInfo.GetId();
+				u32 clientId = connectionInfo.GetId();
 
 				SCOPED_READ_LOCK(m_clientsLock, lock);
-				auto findIt = m_clients.find(connectionId);
+				auto findIt = m_clients.find(clientId);
 				if (findIt == m_clients.end())
 					return true;
 				Client& c = findIt->second;
@@ -1239,6 +1239,8 @@ namespace uba
 				m_logger.Detail(TC("Client %u requested %u connections. Has %u, queue %u"), c.id, connectionCount, c.connectionCount.load(), c.queuedConnectionCount);
 				lock.Leave();
 
+				u32 connectionId = ((NetworkServer::Connection*)connectionInfo.internalData)->m_id;
+
 				SCOPED_FUTEX(m_addConnectionsLock, lock2);
 				for (u32 i = 0; i != toAdd; ++i)
 				{
diff --git a/Engine/Source/Programs/UnrealBuildAccelerator/Common/Private/UbaSessionClient.cpp b/Engine/Source/Programs/UnrealBuildAccelerator/Common/Private/UbaSessionClient.cpp
--- a/Engine/Source/Programs/UnrealBuildAccelerator/Common/Private/UbaSessionClient.cpp
+++ b/Engine/Source/Programs/UnrealBuildAccelerator/Common/Private/UbaSessionClient.cpp
@@ -24,6 +24,7 @@ namespace uba
 		if (const ConfigTable* table = config.GetTable(TC("Session")))
 		{
 			table->GetValueAsBool(useDependencyCrawler, TC("UseDependencyCrawler"));
+			table->GetValueAsU32(pingTimeoutSecondsPrintCallstacks, TC("PingTimeoutSecondsPrintCallstacks"));
 		}
 	}
 
@@ -57,6 +58,7 @@ namespace uba
 		m_memWaitLoadPercent = info.memWaitLoadPercent;
 		m_memKillLoadPercent = info.memKillLoadPercent;
 		m_processFinished = info.processFinished;
+		m_pingTimeoutSecondsPrintCallstacks = info.pingTimeoutSecondsPrintCallstacks;
 
 		m_useDependencyCrawler = info.useDependencyCrawler;
 
@@ -1595,37 +1597,51 @@ namespace uba
 		writer.WriteU32(*(u32*)&cpuLoad);
 		StackBinaryReader<32> reader;
 
-		Event gotResponse(true);
-		auto doneFunc = [](bool error, void* userData) { ((Event*)userData)->Set(); };
+		struct Response
+		{
+			Event done;
+			u64 time;
+		} response{true, 0};
+
+		auto doneFunc = [](bool error, void* userData)
+			{
+				auto& response = *(Response*)userData;
+				response.time = GetTime();
+				response.done.Set();
+			};
 
 		time = GetTime();
 		
-		if (!msg.SendAsync(reader, doneFunc, &gotResponse))
+		if (!msg.SendAsync(reader, doneFunc, &response))
 		{
 			m_loop = false;
 			return;
 		}
 
 		#if PLATFORM_WINDOWS || PLATFORM_MAC // Linux will crash with print all callstacks
-		if (!gotResponse.IsSet(20*1000)) // We accept 20 seconds time out until we write out callstacks
+		u32 timeoutSeconds = 20;
+		if (m_pingTimeoutSecondsPrintCallstacks)
+			timeoutSeconds = m_pingTimeoutSecondsPrintCallstacks;
+		if (!response.done.IsSet(timeoutSeconds*1000))
 		{
 			LoggerWithWriter logger(g_consoleLogWriter);
-			logger.Info(TC("Took more than 20 seconds to send/receive ping. Printing callstacks"));
-			PrintAllCallstacks(logger);
+			logger.Info(TC("Took more than %u seconds to send/receive ping. Printing callstacks"), timeoutSeconds);
+			if (m_pingTimeoutSecondsPrintCallstacks)
+				PrintAllCallstacks(logger);
 		}
 		#endif
 		
-		gotResponse.IsSet(); // Wait forever (until message trigger error or finishes)
+		response.done.IsSet(); // Wait forever (until message trigger error or finishes)
 
 		if (!msg.ProcessAsyncResults(reader) || msg.GetError())
 		{
 			m_loop = false;
 			return;
 		}
-		u64 newTime = GetTime();
-		u64 lastPing = newTime - time;
+
+		u64 lastPing = response.time - time;
 		m_lastPing = lastPing;
-		m_lastPingSendTime = newTime;
+		m_lastPingSendTime = response.time;
 
 		if (lastPing < m_bestPing || m_bestPing == 0)
 			m_bestPing = lastPing;
diff --git a/Engine/Source/Programs/UnrealBuildAccelerator/Common/Private/UbaStorageClient.cpp b/Engine/Source/Programs/UnrealBuildAccelerator/Common/Private/UbaStorageClient.cpp
--- a/Engine/Source/Programs/UnrealBuildAccelerator/Common/Private/UbaStorageClient.cpp
+++ b/Engine/Source/Programs/UnrealBuildAccelerator/Common/Private/UbaStorageClient.cpp
@@ -951,17 +951,23 @@ namespace uba
 	{
 		SCOPED_WRITE_LOCK(m_proxyClientLock, lock);
 		auto proxy = m_proxyClient;
-		if (!proxy)
+		if (!proxy || !proxy->client.IsConnected())
 			return;
 		u64 now = GetTime();
 		if (TimeToMs(now - m_proxyClientKeepAliveTime) < 30 * 1000)
 			return;
-		m_proxyClientKeepAliveTime = now;
 
 		++proxy->refCount;
 		lock.Leave();
 
-		m_proxyClient->client.SendKeepAlive();
+		proxy->client.SendKeepAlive();
+		u64 time = GetTime();
+		m_proxyClientKeepAliveTime = now;
+
+		u64 durationMs = TimeToMs(time - now);
+		if (durationMs > 20 * 1000)
+			m_logger.Info(TC("Took %llu seconds to ping proxy server"), durationMs/1000);
+
 
 		lock.Enter();
 		--proxy->refCount;
diff --git a/Engine/Source/Programs/UnrealBuildAccelerator/Common/Public/UbaSessionClient.h b/Engine/Source/Programs/UnrealBuildAccelerator/Common/Public/UbaSessionClient.h
--- a/Engine/Source/Programs/UnrealBuildAccelerator/Common/Public/UbaSessionClient.h
+++ b/Engine/Source/Programs/UnrealBuildAccelerator/Common/Public/UbaSessionClient.h
@@ -20,6 +20,7 @@ namespace uba
 		u32 maxProcessCount = 1;
 		u32 defaultPriorityClass = 0x00004000; // BELOW_NORMAL_PRIORITY_CLASS;
 		u32 maxIdleSeconds = ~0u;
+		u32 pingTimeoutSecondsPrintCallstacks = 0;
 		u8 memWaitLoadPercent = 80; // When memory usage goes above this percent, no new processes will be spawned until back below
 		u8 memKillLoadPercent = 90; // When memory usage goes above this percent, newest processes will be killed to bring it back below
 		bool dedicated = false;  // If true, server will not disconnect client when starting to run out of work.
@@ -173,6 +174,7 @@ namespace uba
 		Atomic<u64> m_bestPing;
 		Atomic<u64> m_lastPing;
 		u64 m_lastPingSendTime = 0;
+		u32 m_pingTimeoutSecondsPrintCallstacks = 0;
 
 		Atomic<u64> m_memAvail;
 		Atomic<u64> m_memTotal;
