[UBA]
* Moved network timeouts to a constant in UbaDefaultConstants

diff --git a/Engine/Source/Programs/UnrealBuildAccelerator/Common/Private/UbaNetworkBackendTcp.cpp b/Engine/Source/Programs/UnrealBuildAccelerator/Common/Private/UbaNetworkBackendTcp.cpp
--- a/Engine/Source/Programs/UnrealBuildAccelerator/Common/Private/UbaNetworkBackendTcp.cpp
+++ b/Engine/Source/Programs/UnrealBuildAccelerator/Common/Private/UbaNetworkBackendTcp.cpp
@@ -861,7 +861,7 @@ namespace uba
 							if (i == 0)
 							{
 								logger.Info(TC("SendSocket - WSAPoll returned timeout after 40 seconds for socket %i (%s)"), socket, hint);
-								timeoutMs = 10 * 60 * 1000; // 10 minutes send timeout
+								timeoutMs = DefaultNetworkSendTimeoutSeconds * 1000; // 10 minutes send timeout
 								continue;
 							}
 							logger.Info(TC("SendSocket - WSAPoll returned timeout after 10 minutes for socket %i (%s)"), socket, hint);
diff --git a/Engine/Source/Programs/UnrealBuildAccelerator/Common/Public/UbaDefaultConstants.h b/Engine/Source/Programs/UnrealBuildAccelerator/Common/Public/UbaDefaultConstants.h
--- a/Engine/Source/Programs/UnrealBuildAccelerator/Common/Public/UbaDefaultConstants.h
+++ b/Engine/Source/Programs/UnrealBuildAccelerator/Common/Public/UbaDefaultConstants.h
@@ -11,4 +11,6 @@ namespace uba
 	static constexpr u16 DefaultStorageProxyPort = DefaultPort + 1;
 	static constexpr u16 DefaultCachePort = DefaultPort + 2;
 	static constexpr u32 SendDefaultSize = 256*1024;
+	static constexpr u32 DefaultNetworkReceiveTimeoutSeconds = 10*60;
+	static constexpr u32 DefaultNetworkSendTimeoutSeconds = 10*60;
 }
diff --git a/Engine/Source/Programs/UnrealBuildAccelerator/Host/Public/UbaExports.h b/Engine/Source/Programs/UnrealBuildAccelerator/Host/Public/UbaExports.h
--- a/Engine/Source/Programs/UnrealBuildAccelerator/Host/Public/UbaExports.h
+++ b/Engine/Source/Programs/UnrealBuildAccelerator/Host/Public/UbaExports.h
@@ -60,7 +60,7 @@ extern "C"
 	UBA_API void ConfigTable_AddValueString(uba::ConfigTable& table, const uba::tchar* key, const uba::tchar* str);
 
 	// NetworkServer
-	UBA_API uba::NetworkServer* NetworkServer_Create(uba::LogWriter& writer = uba::g_consoleLogWriter, uba::u32 workerCount = 64, uba::u32 sendSize = uba::SendDefaultSize, uba::u32 receiveTimeoutSeconds = 60, bool useQuic = false);
+	UBA_API uba::NetworkServer* NetworkServer_Create(uba::LogWriter& writer = uba::g_consoleLogWriter, uba::u32 workerCount = 64, uba::u32 sendSize = uba::SendDefaultSize, uba::u32 receiveTimeoutSeconds = DefaultNetworkReceiveTimeoutSeconds, bool useQuic = false);
 	UBA_API void NetworkServer_Destroy(uba::NetworkServer* server);
 	UBA_API bool NetworkServer_StartListen(uba::NetworkServer* server, int port = uba::DefaultPort, const uba::tchar* ip = nullptr, const uba::tchar* crypto = nullptr);
 	UBA_API void NetworkServer_Stop(uba::NetworkServer* server);
