[UBA]
* Removed logging about failing to devirtualize path from crawler

diff --git a/Engine/Source/Programs/UnrealBuildAccelerator/Common/Public/UbaDependencyCrawler.h b/Engine/Source/Programs/UnrealBuildAccelerator/Common/Public/UbaDependencyCrawler.h
--- a/Engine/Source/Programs/UnrealBuildAccelerator/Common/Public/UbaDependencyCrawler.h
+++ b/Engine/Source/Programs/UnrealBuildAccelerator/Common/Public/UbaDependencyCrawler.h
@@ -8,6 +8,12 @@
 #include "UbaProcessUtils.h"
 #include "UbaWorkManager.h"
 
+#if UBA_DEBUG
+#define UBA_LOG_DEVIRTUALIZATION_ERROR(x, ...) m_logger.Info(TC(x), __VA_ARGS__)
+#else
+#define UBA_LOG_DEVIRTUALIZATION_ERROR(x, ...) do { } while(false)
+#endif
+
 #define UBA_CRAWLER_DEBUG_PRINT(str1, str2) //wprintf(TC("") str1 ": %s\n", str2);
 
 // Note
@@ -199,7 +205,7 @@ namespace uba
 						StringBuffer<> fixedPath;
 						FixPath(instance.builtinIncludesDir.c_str(), instance.workDir.data(), instance.workDir.size(), fixedPath);
 						if (!instance.devirtualizePathFunc(fixedPath))
-							m_logger.Info(TC("Failed to devirtualize path %s found in builtin includes"), fixedPath.data);
+							UBA_LOG_DEVIRTUALIZATION_ERROR("Failed to devirtualize path %s found in builtin includes", fixedPath.data);
 						if (CaseInsensitiveFs)
 							fixedPath.MakeLower();
 						fixedPath.Append(PathSeparator);
@@ -217,7 +223,7 @@ namespace uba
 		StringBuffer<> fixedPath;
 		FixPath(rsp.data, instance.workDir.data(), instance.workDir.size(), fixedPath);
 		if (!instance.devirtualizePathFunc(fixedPath))
-			m_logger.Info(TC("Failed to devirtualize path %s in %s"), fixedPath.data, rsp.data);
+			UBA_LOG_DEVIRTUALIZATION_ERROR("Failed to devirtualize path %s in %s", fixedPath.data, rsp.data);
 		if (CaseInsensitiveFs)
 			fixedPath.MakeLower();
 		if (!instance.createFileFunc(tracker, fixedPath, [&](const void* data, u64 dataSize) { return ParseRsp2(tracker, instance, data, dataSize, rsp.data, outCodeFiles); }))
@@ -235,7 +241,7 @@ namespace uba
 				{
 					FixPath(file.data, instance.workDir.data(), instance.workDir.size(), fixedPath2);
 					if (!instance.devirtualizePathFunc(fixedPath2))
-						m_logger.Info(TC("Failed to devirtualize path %s in %s"), fixedPath2.data, rsp);
+						UBA_LOG_DEVIRTUALIZATION_ERROR("Failed to devirtualize path %s in %s", fixedPath2.data, rsp);
 
 					if (CaseInsensitiveFs)
 						fixedPath2.MakeLower();
@@ -253,7 +259,7 @@ namespace uba
 				StringBuffer<> fixedPath2;
 				FixPath(path.data, instance.workDir.data(), instance.workDir.size(), fixedPath2);
 				if (!instance.devirtualizePathFunc(fixedPath2))
-					m_logger.Info(TC("Failed to devirtualize path %s in %s"), fixedPath2.data, rsp);
+					UBA_LOG_DEVIRTUALIZATION_ERROR("Failed to devirtualize path %s in %s", fixedPath2.data, rsp);
 				if (CaseInsensitiveFs)
 					fixedPath2.MakeLower();
 				fixedPath2.EnsureEndsWithSlash();
@@ -335,7 +341,7 @@ namespace uba
 							StringBuffer<> fixedFile;
 							FixPath(sb.data + 3, instance.workDir.data(), instance.workDir.size(), fixedFile);
 							if (!instance.devirtualizePathFunc(fixedFile))
-								m_logger.Info(TC("Failed to devirtualize path %s in %s"), fixedFile.data, rsp);
+								UBA_LOG_DEVIRTUALIZATION_ERROR("Failed to devirtualize path %s in %s", fixedFile.data, rsp);
 
 							if (CaseInsensitiveFs)
 								fixedFile.MakeLower();
@@ -841,7 +847,7 @@ namespace uba
 		FixPath(include.data, rootPath.data, rootPath.count, fullPath);
 		if (IsAbsolutePath(include.data))
 			if (!instance.devirtualizePathFunc(fullPath))
-				m_logger.Info(TC("Failed to devirtualize include path %s in %s"), fullPath.data, codeFile.data);
+				UBA_LOG_DEVIRTUALIZATION_ERROR("Failed to devirtualize include path %s in %s", fullPath.data, codeFile.data);
 
 		if (CaseInsensitiveFs)
 			fullPath.MakeLower();
