[UBA]
* Fixed so tcp keep alive kicks in after 60 seconds with 1 second interval and shut down socket after 10 attempts.
* Fixed so timeout settings from NetworkClientCreateInfo is propagated to sockets
* Bumped message protocol number since I forgot to do that when changing how long strings are serrialized

diff --git a/Engine/Source/Programs/UnrealBuildAccelerator/Agent/Private/UbaAgent.cpp b/Engine/Source/Programs/UnrealBuildAccelerator/Agent/Private/UbaAgent.cpp
--- a/Engine/Source/Programs/UnrealBuildAccelerator/Agent/Private/UbaAgent.cpp
+++ b/Engine/Source/Programs/UnrealBuildAccelerator/Agent/Private/UbaAgent.cpp
@@ -1018,10 +1018,6 @@ namespace uba
 		
 		do
 		{
-			// TODO: Revisit.. it could be that something stalls during startup so we can't have this timeout
-			//u32 receiveTimeoutSeconds = 60; // We are sending pings at 5 seconds cadence, so timeout if no data appears on recv socket within 60 seconds
-			u32 receiveTimeoutSeconds = 0;
-
 			NetworkBackendMemory networkBackendMem(logWriter);
 			NetworkBackend* networkBackend;
 			#if UBA_USE_QUIC
@@ -1035,7 +1031,7 @@ namespace uba
 			NetworkClientCreateInfo ncci(logWriter);
 			//ncci.Apply(config);
 			ncci.sendSize = sendSize;
-			ncci.receiveTimeoutSeconds = receiveTimeoutSeconds;
+			//ncci.receiveTimeoutSeconds = 0; // Use default 10 minutes
 			ncci.workerCount = maxWorkerCount;
 			if (hasCrypto)
 				ncci.cryptoKey128 = crypto;
diff --git a/Engine/Source/Programs/UnrealBuildAccelerator/Common/Private/UbaNetworkBackendTcp.cpp b/Engine/Source/Programs/UnrealBuildAccelerator/Common/Private/UbaNetworkBackendTcp.cpp
--- a/Engine/Source/Programs/UnrealBuildAccelerator/Common/Private/UbaNetworkBackendTcp.cpp
+++ b/Engine/Source/Programs/UnrealBuildAccelerator/Common/Private/UbaNetworkBackendTcp.cpp
@@ -486,7 +486,7 @@ namespace uba
 			SetBlocking(logger, connection.socket, false);
 			u32 recvTimeoutMs = connection.recvTimeoutMs;
 			if (!recvTimeoutMs)
-				recvTimeoutMs = 10 * 60 * 1000; // 10 minutes time out
+				recvTimeoutMs = DefaultNetworkReceiveTimeoutSeconds * 1000; // 10 minutes time out
 
 			RecvCache recvCache;
 
@@ -827,11 +827,37 @@ namespace uba
 		return true;
 	}
 
-	bool SetKeepAlive(Logger& logger, SOCKET socket)
+	bool SetKeepAlive(Logger& logger, SOCKET socket) // This will make sure that WSAPoll exits when the network cable is pulled
 	{
+		
+		int keepAliveTime = 60; // 1 minute of idle before it starts sending keep alive
+		int keepAliveTimeInterval = 1; // 1 second between each check (10 failures mean disconnect)
+
+
 		u32 value = 1;
 		if (setsockopt(socket, SOL_SOCKET, SO_KEEPALIVE, (const char*)&value, sizeof(value)) == SOCKET_ERROR)
 			return logger.Error(TC("setsockopt SO_KEEPALIVE (error: %s)"), LastErrorToText(WSAGetLastError()).data);
+
+#if PLATFORM_WINDOWS	
+		struct tcp_keepalive kaSettings;
+		DWORD bytesReturned;
+		kaSettings.onoff = 1;
+		kaSettings.keepalivetime = keepAliveTime * 1000;
+		kaSettings.keepaliveinterval = keepAliveTimeInterval * 1000; 
+		if (WSAIoctl(socket, SIO_KEEPALIVE_VALS, &kaSettings, sizeof(kaSettings), NULL, 0, &bytesReturned, NULL, NULL) == SOCKET_ERROR)
+			return logger.Error(TC("WSAIoctl SIO_KEEPALIVE_VALS (error: %s)"), LastErrorToText(WSAGetLastError()).data);
+#elif PLATFORM_LINUX
+		if (setsockopt(socket, IPPROTO_TCP, TCP_KEEPIDLE, &keepAliveTime, sizeof(int)) < 0)
+			return logger.Error(TC("setsockopt TCP_KEEPIDLE (error: %s)"), LastErrorToText(WSAGetLastError()).data);
+		if (setsockopt(socket, IPPROTO_TCP, TCP_KEEPINTVL, &keepAliveTimeInterval, sizeof(int)) < 0)
+			return logger.Error(TC("setsockopt TCP_KEEPINTVL (error: %s)"), LastErrorToText(WSAGetLastError()).data);
+		int keepAliveProbes = 10;
+		if (setsockopt(socket, IPPROTO_TCP, TCP_KEEPCNT, &keepAliveProbes, sizeof(int)) < 0)
+			return logger.Error(TC("setsockopt TCP_KEEPCNT (error: %s)"), LastErrorToText(WSAGetLastError()).data);
+#else // PLATFORM_MAC
+		if (setsockopt(socket, IPPROTO_TCP, TCP_KEEPALIVE, &keepAliveTime, sizeof(int)) < 0)
+			return logger.Error(TC("setsockopt TCP_KEEPALIVE (error: %s)"), LastErrorToText(WSAGetLastError()).data);
+#endif
 		return true;
 	}
 
diff --git a/Engine/Source/Programs/UnrealBuildAccelerator/Common/Private/UbaNetworkClient.cpp b/Engine/Source/Programs/UnrealBuildAccelerator/Common/Private/UbaNetworkClient.cpp
--- a/Engine/Source/Programs/UnrealBuildAccelerator/Common/Private/UbaNetworkClient.cpp
+++ b/Engine/Source/Programs/UnrealBuildAccelerator/Common/Private/UbaNetworkClient.cpp
@@ -87,6 +87,8 @@ namespace uba
 			Atomic<u8> error;
 		};
 
+		backend.SetRecvTimeout(backendConnection, m_receiveTimeoutSeconds*1000);
+
 		RecvContext rc(*this, backend, backendConnection);
 
 		// The only way out of this function is to get a call to one of the below callbacks since exitScopeEvent must be set.
@@ -224,6 +226,8 @@ namespace uba
 		m_logger.Detail(TC("Connected to server... (0x%p)"), backendConnection);
 		lock.Leave();
 
+		backend.SetRecvTimeout(backendConnection, m_receiveTimeoutSeconds*1000);
+
 		backend.SetDisconnectCallback(backendConnection, connection, DisconnectCallback);
 		backend.SetRecvCallbacks(backendConnection, connection, ReceiveHeaderSize, ReceiveResponseHeader, ReceiveResponseBody, TC("ReceiveMessageResponse"));
 		return true;
diff --git a/Engine/Source/Programs/UnrealBuildAccelerator/Common/Public/UbaNetworkClient.h b/Engine/Source/Programs/UnrealBuildAccelerator/Common/Public/UbaNetworkClient.h
--- a/Engine/Source/Programs/UnrealBuildAccelerator/Common/Public/UbaNetworkClient.h
+++ b/Engine/Source/Programs/UnrealBuildAccelerator/Common/Public/UbaNetworkClient.h
@@ -24,7 +24,7 @@ namespace uba
 		NetworkClientCreateInfo(LogWriter& w = g_consoleLogWriter) : logWriter(w) {}
 		LogWriter& logWriter;
 		u32 sendSize = SendDefaultSize;
-		u32 receiveTimeoutSeconds = 0;
+		u32 receiveTimeoutSeconds = DefaultNetworkReceiveTimeoutSeconds;
 		u32 workerCount = ~0u; // ~0u means logical processor count will be used
 		u32 desiredConnectionCount = 1; // Desired connection count to server when connected
 		const u8* cryptoKey128 = nullptr;
diff --git a/Engine/Source/Programs/UnrealBuildAccelerator/Common/Public/UbaNetworkServer.h b/Engine/Source/Programs/UnrealBuildAccelerator/Common/Public/UbaNetworkServer.h
--- a/Engine/Source/Programs/UnrealBuildAccelerator/Common/Public/UbaNetworkServer.h
+++ b/Engine/Source/Programs/UnrealBuildAccelerator/Common/Public/UbaNetworkServer.h
@@ -35,7 +35,7 @@ namespace uba
 		LogWriter& logWriter;
 		u32 workerCount = 0; // Zero means it will use the number of logical cores as worker count
 		u32 sendSize = SendDefaultSize;
-		u32 receiveTimeoutSeconds = 0;
+		u32 receiveTimeoutSeconds = DefaultNetworkReceiveTimeoutSeconds;
 		bool logConnections = true;
 	};
 
diff --git a/Engine/Source/Programs/UnrealBuildAccelerator/Core/Public/UbaProtocol.h b/Engine/Source/Programs/UnrealBuildAccelerator/Core/Public/UbaProtocol.h
--- a/Engine/Source/Programs/UnrealBuildAccelerator/Core/Public/UbaProtocol.h
+++ b/Engine/Source/Programs/UnrealBuildAccelerator/Core/Public/UbaProtocol.h
@@ -50,7 +50,7 @@ namespace uba
 		#undef UBA_PROCESS_MESSAGE
 	};
 
-	inline constexpr u32 ProcessMessageVersion = 1342;
+	inline constexpr u32 ProcessMessageVersion = 1343;
 
 	inline constexpr u32 CommunicationMemSize = IsWindows ? 64*1024 : 64*1024*2; // Macos expands some commandlines to be crazy long
 
