[UBA]
* Fixed macos compile error

diff --git a/Engine/Source/Programs/UnrealBuildAccelerator/Common/Private/UbaNetworkBackendTcp.cpp b/Engine/Source/Programs/UnrealBuildAccelerator/Common/Private/UbaNetworkBackendTcp.cpp
--- a/Engine/Source/Programs/UnrealBuildAccelerator/Common/Private/UbaNetworkBackendTcp.cpp
+++ b/Engine/Source/Programs/UnrealBuildAccelerator/Common/Private/UbaNetworkBackendTcp.cpp
@@ -8,10 +8,6 @@
 #include "UbaStringBuffer.h"
 #include "UbaTimer.h"
 
-#if PLATFORM_LINUX
-#include <netinet/tcp.h>
-#endif
-
 #if PLATFORM_WINDOWS
 #include <iphlpapi.h>
 #include <ipifcons.h>
@@ -21,6 +17,7 @@
 #pragma comment(lib, "IPHLPAPI.lib") // For GetAdaptersInfo
 #else
 #include <netdb.h>
+#include <netinet/tcp.h>
 #include <arpa/inet.h>
 #include <ifaddrs.h>
 #include <poll.h>
