[UBA]
* Fixed so UbaCli -listimports include loader paths

diff --git a/Engine/Source/Programs/UnrealBuildAccelerator/Cli/Private/UbaCli.cpp b/Engine/Source/Programs/UnrealBuildAccelerator/Cli/Private/UbaCli.cpp
--- a/Engine/Source/Programs/UnrealBuildAccelerator/Cli/Private/UbaCli.cpp
+++ b/Engine/Source/Programs/UnrealBuildAccelerator/Cli/Private/UbaCli.cpp
@@ -815,11 +815,24 @@ namespace uba
 
 			if (listImports)
 			{
-				logger.Info(TC("Imports:"));
 				StringBuffer<> error;
+				bool printImports = true;
 				if (!FindImports(application.c_str(), StringView(application).GetPath(), [&](const tchar* import, bool isKnown, const char* const* loaderPaths)
 				{
+					if (printImports)
+					{
+						if (*loaderPaths)
+						{
+							logger.Info(TC("LoaderPaths:"));
+							for (auto it=loaderPaths;*it; ++it)
+								if (**it)
+									logger.Info(TC("  %s"), *it);
+						}
+						printImports = false;
+						logger.Info(TC("Imports:"));
+					}
 					logger.Info(TC("  %s"), import);
+
 				}, error))
 					return logger.Error(TC("%s"), error.data);
 				return true;
