[UBA]
* Added more info when failing to find binaries on macos

diff --git a/Engine/Source/Programs/UnrealBuildAccelerator/Common/Private/UbaSession.cpp b/Engine/Source/Programs/UnrealBuildAccelerator/Common/Private/UbaSession.cpp
--- a/Engine/Source/Programs/UnrealBuildAccelerator/Common/Private/UbaSession.cpp
+++ b/Engine/Source/Programs/UnrealBuildAccelerator/Common/Private/UbaSession.cpp
@@ -1378,9 +1378,11 @@ namespace uba
 				return true; // TODO: We have to return true here because there are scenarios where failing is actually ok (it seems it can return false on crt shim libraries such as api-ms-win-crt*)
 			#elif PLATFORM_MAC
 			if (!loaderPaths)
-				return m_logger.Error("Failed to find file %s", applicationName);
+				return m_logger.Error("CopyImports - Failed to find file %s (no loader paths)", applicationName);
+			u32 loaderPathCount = 0;
 			for (auto it = loaderPaths; *it; ++it)
 			{
+				++loaderPathCount;
 				StringBuffer<> absolutePath;
 				absolutePath.Append(applicationDir, applicationDirEnd - applicationDir).Append(*it).EnsureEndsWithSlash().Append(library);
 				FixPath(absolutePath.data, nullptr, 0, temp3.Clear());
@@ -1391,7 +1393,7 @@ namespace uba
 				break;
 			}
 			if (attr == INVALID_FILE_ATTRIBUTES)
-				return m_logger.Error("Failed to find file %s", applicationName);
+				return m_logger.Error("CopyImports - Failed to find file %s (%u loader paths)", applicationName, loaderPathCount);
 			//UBA_ASSERTF(false, TC("DIR NOT FOUND: %s"), applicationName);
 			#else
 			return m_logger.Error("Code path not implemented for linux! (CopyImports %s)", library);
diff --git a/Engine/Source/Programs/UnrealBuildAccelerator/Common/Private/UbaSessionServer.cpp b/Engine/Source/Programs/UnrealBuildAccelerator/Common/Private/UbaSessionServer.cpp
--- a/Engine/Source/Programs/UnrealBuildAccelerator/Common/Private/UbaSessionServer.cpp
+++ b/Engine/Source/Programs/UnrealBuildAccelerator/Common/Private/UbaSessionServer.cpp
@@ -902,7 +902,7 @@ namespace uba
 			#if 0
 			if (casKey == CasKeyZero)
 			{
-				m_logger.Warning(TC("Failed to find file %s"), fileName.data);
+				m_logger.Warning(TC("HandleEnsureBinaryFile - Failed to find file %s"), fileName.data);
 				for (auto& loaderPath : loaderPaths)
 				{
 					m_logger.Warning(TC("   LoaderPath %s"), loaderPath.c_str());
