[UBA]
* Added option -fork to UbaCacheService for linux that can be used to have the process fork itself and track seg faults

diff --git a/Engine/Source/Programs/UnrealBuildAccelerator/CacheService/Private/UbaCacheService.cpp b/Engine/Source/Programs/UnrealBuildAccelerator/CacheService/Private/UbaCacheService.cpp
--- a/Engine/Source/Programs/UnrealBuildAccelerator/CacheService/Private/UbaCacheService.cpp
+++ b/Engine/Source/Programs/UnrealBuildAccelerator/CacheService/Private/UbaCacheService.cpp
@@ -3,6 +3,7 @@
 #include "UbaCacheServer.h"
 #include "UbaApplication.h"
 #include "UbaConfig.h"
+#include "UbaDirectoryIterator.h"
 #include "UbaFile.h"
 #include "UbaHttpServer.h"
 #include "UbaNetworkBackendTcp.h"
@@ -33,7 +34,7 @@ namespace uba
 		}();
 	u32				DefaultWorkerCount = []() { return GetLogicalProcessorCount() + 4; }();
 
-	int PrintHelp(const tchar* message)
+	bool PrintHelp(const tchar* message)
 	{
 		LoggerWithWriter logger(g_consoleLogWriter, TC(""));
 		if (*message)
@@ -57,11 +58,13 @@ namespace uba
 		logger.Info(TC("  -crash                  Force a crash (for testing)"));
 		logger.Info(TC("  -nosignalhandler        Will not hook up signal handler"));
 		logger.Info(TC("  -maxworkers=<number>    Max number of workers used by cache server. Defaults to \"%u\""), DefaultWorkerCount);
-
+#if PLATFORM_LINUX
+		logger.Info(TC("  -fork                   Will handle segfaults and restart"));
+#endif
 		logger.Info(TC(" Example of how to register crypto key to cache server (when -http=80 is provided)"));
 		logger.Info(TC("   curl http://localhost/addcrypto?3f58aa57466db9999213456789123445"));
 		logger.Info(TC(""));
-		return -1;
+		return false;
 	}
 
 	ReaderWriterLock* g_exitLock = new ReaderWriterLock();
@@ -98,7 +101,7 @@ namespace uba
 	
 	StringBuffer<> g_rootDir(DefaultRootDir);
 
-	int WrappedMain(int argc, tchar* argv[])
+	bool WrappedMain(int argc, tchar* argv[])
 	{
 		using namespace uba;
 
@@ -118,6 +121,10 @@ namespace uba
 		u32 maxWorkerCount = DefaultWorkerCount;
 		TString configFile;
 
+		#if PLATFORM_LINUX
+		bool forkProcess = false;
+		#endif
+
 		for (int i=1; i!=argc; ++i)
 		{
 			StringBuffer<> name;
@@ -182,6 +189,12 @@ namespace uba
 			{
 				signalHandlerEnabled = false;
 			}
+			#if PLATFORM_LINUX
+			else if (name.Equals(TCV("-fork")))
+			{
+				forkProcess = true;
+			}
+			#endif
 			else if (name.Equals(TCV("-maxworkers")))
 			{
 				if (!value.Parse(maxWorkerCount))
@@ -215,17 +228,6 @@ namespace uba
 			}
 		}
 
-		if (signalHandlerEnabled)
-			AddExceptionHandler();
-
-		if (shouldCrash)
-		{
-			#ifndef __clang_analyzer__
-			int* ptr = nullptr;
-			*ptr = 0;
-			#endif
-		}
-
 		FilteredLogWriter logWriter(g_consoleLogWriter, quiet ? LogEntryType_Info : LogEntryType_Detail);
 		LoggerWithWriter logger(logWriter, TC(""));
 
@@ -255,8 +257,6 @@ namespace uba
 				logger.Info(TC("  Improve maintenance performance by enabling %llu huge pages on system (%llu enabled)"), recommendedHugePageCount, hugePageCount);
 		}
 
-		logger.Info(TC(""));
-
 		u64 storageCapacity = u64(storageCapacityGb)*1000*1000*1000;
 
 		StringBuffer<512> currentDir;
@@ -265,6 +265,138 @@ namespace uba
 		if (workDir.IsEmpty())
 			workDir.Append(currentDir);
 
+		#if PLATFORM_LINUX
+		if (forkProcess)
+		{
+			auto ReadFirstLine = [](StringBufferBase& out, const char* cmd)
+				{
+					FILE* f = popen(cmd, "r");
+					if (!f)
+						return false;
+					bool success = fgets(out.data, out.capacity, f) == out.data;
+					pclose(f);
+					if (!success)
+						return false;
+					out.count = TStrlen(out.data);
+					out.data[--out.count] = 0;
+					return true;
+				};
+
+			StringBuffer<256> temp;
+			TString crashDumpDir;
+			TString crashDumpPattern;
+			if (ReadFirstLine(temp, "ulimit -c"))
+			{
+				if (temp.Equals(TCV("0")))
+					logger.Info(TC("  Crash dumps disabled. Enable with \"ulimit -c unlimited\""));
+				else if (ReadFirstLine(temp, "cat /proc/sys/kernel/core_pattern"))
+				{
+					if (temp[0] == '|')
+						logger.Info(TC("  Crash dumps enabled but piped so can't wait (%s). use 'sudo echo \"/<path>/dump.%t\" | sudo tee /proc/sys/kernel/core_pattern > /dev/null' to write to file"), temp.data);
+					else
+					{
+						logger.Info(TC("  Crash dumps enabled and written to file: %s (Write no other files in the same dir)"), temp.data);
+						crashDumpDir = StringView(temp).GetPath().ToString();
+						crashDumpPattern = StringView(temp).GetFileName().ToString();
+					}
+				}
+			}
+
+			Set<StringKey> existingFiles;
+			if (!crashDumpPattern.empty())
+				TraverseDir(logger, crashDumpDir, [&](const DirectoryEntry& e) { existingFiles.insert(ToStringKey(e.name, e.nameLen)); });
+
+			static bool shouldExit = false;
+			auto SigHandler = [](int sig) { shouldExit = true; };
+
+			logger.Info(TC("  Starting cache server fork"));
+			forkProcess = false;
+			while (true)
+			{
+				u64 startTime = GetTime();
+				pid_t childPid = fork();
+				if (childPid < 0)
+					return logger.Error(TC("Failed to fork process. (%s)"), strerror(errno));
+				if (childPid == 0)
+					break;
+
+				// Parent
+
+				signal(SIGINT, SigHandler);
+				signal(SIGTERM, SigHandler);
+
+				int status;
+				if (waitpid(childPid, &status, 0) == -1)
+					return logger.Error(TC("waitpid failed. (%s)"), strerror(errno));
+				
+				if (!shouldExit)
+					logger.Info(TC("  Fork exited with . (%u)"), status);
+
+				if (WIFSIGNALED(status) && WCOREDUMP(status))
+				{
+					// Wait for new file to appear
+					TString crashDumpFile;
+					u32 retryCount = 5;
+					while (retryCount--)
+					{
+						TraverseDir(logger, crashDumpDir, [&](const DirectoryEntry& e) { if (existingFiles.find(ToStringKey(e.name, e.nameLen)) == existingFiles.end()) crashDumpFile = e.name; });
+						if (!crashDumpFile.empty())
+							break;
+						Sleep(1000);
+					}
+					if (crashDumpFile.empty())
+						logger.Info(TC("  No core dump found after 10 seconds"));
+					else
+					{
+						existingFiles.insert(ToStringKey(crashDumpFile));
+						TString crashDumpPath = crashDumpDir + '/' + crashDumpFile;
+						logger.Info(TC("  Found new core dump %s"), crashDumpPath.c_str());
+
+						temp.Clear().Appendf("/proc/%u", childPid);
+						retryCount = 10*60;
+						bool firstCall = true;
+						while (retryCount--)
+						{
+							if (access(temp.data, F_OK) != 0)
+								break;
+							if (firstCall)
+								logger.Info(TC("  Waiting (up to 10 minutes) for pid %u to be cleaned up"), childPid);
+							firstCall = false;
+							Sleep(1000);
+							continue;
+						}
+					}
+				}
+				else
+				{
+					shouldExit = true;
+				}
+
+				u32 retryCount = 40;
+				while (!shouldExit && --retryCount)
+					Sleep(100);
+				if (shouldExit)
+					exit(0);
+				logger.Info(TC("  Restarting cache server fork"));
+			}
+		}
+		#endif
+
+		logger.Info(TC(""));
+
+		if (signalHandlerEnabled)
+			AddExceptionHandler();
+
+		if (shouldCrash)
+		{
+			void* mem = new u8[1ull*1024*1024*1024];
+			memset(mem, 1, 1ull*1024*1024*1024);
+			#ifndef __clang_analyzer__
+			int* ptr = nullptr;
+			*ptr = 0;
+			#endif
+		}
+
 		// TODO: Change workdir to make it full
 
 		#if PLATFORM_WINDOWS
@@ -286,7 +418,7 @@ namespace uba
 		bool ctorSuccess = true;
 		NetworkServer networkServer(ctorSuccess, nsci);
 		if (!ctorSuccess)
-			return -1;
+			return false;
 
 		StorageServerCreateInfo storageInfo(networkServer, g_rootDir.data, logWriter);
 		storageInfo.Apply(config);
@@ -298,7 +430,7 @@ namespace uba
 		StorageServer storageServer(storageInfo);
 
 		if (!storageServer.LoadCasTable(true, true))
-			return -1;
+			return false;
 
 		CacheServerCreateInfo cacheInfo(storageServer, g_rootDir.data, logWriter);
 		cacheInfo.Apply(config);
@@ -307,14 +439,14 @@ namespace uba
 		CacheServer cacheServer(cacheInfo);
 
 		if (!cacheServer.Load())
-			return -1;
+			return false;
 
 		if (fullMaintenance)
 			cacheServer.SetForceFullMaintenance();
 
 		if (maintenanceEnabled)
 			if (!cacheServer.RunMaintenance(true, allowSave, ShouldExit))
-				return -1;
+				return false;
 
 		HttpServer httpServer(logWriter, networkBackend);
 
@@ -347,7 +479,7 @@ namespace uba
 			auto stopServer = MakeGuard([&]() { networkServer.DisconnectClients(); });
 
 			if (!networkServer.StartListen(networkBackend, port, listenIp.data))
-				return -1;
+				return false;
 
 			while (!ShouldExit() && !cacheServer.ShouldShutdown())
 			{
@@ -363,19 +495,19 @@ namespace uba
 
 		storageServer.DeleteIsRunningFile();
 
-		return 0;
+		return true;
 	}
 }
 
 #if PLATFORM_WINDOWS
 int wmain(int argc, wchar_t* argv[])
 {
-	int res = uba::WrappedMain(argc, argv);
+	int res = uba::WrappedMain(argc, argv) ? 0 : 1;
 	return res;
 }
 #else
 int main(int argc, char* argv[])
 {
-	return uba::WrappedMain(argc, argv);
+	return uba::WrappedMain(argc, argv) ? 0 : 1;
 }
 #endif
