[UBA]
* Added so callback is called when recv times out. If callback returns false it exits, if true, then retry

diff --git a/Engine/Source/Programs/UnrealBuildAccelerator/Common/Private/UbaNetworkBackendMemory.cpp b/Engine/Source/Programs/UnrealBuildAccelerator/Common/Private/UbaNetworkBackendMemory.cpp
--- a/Engine/Source/Programs/UnrealBuildAccelerator/Common/Private/UbaNetworkBackendMemory.cpp
+++ b/Engine/Source/Programs/UnrealBuildAccelerator/Common/Private/UbaNetworkBackendMemory.cpp
@@ -129,7 +129,7 @@ namespace uba
 		peer.recvHint = recvHint;
 	}
 
-	void NetworkBackendMemory::SetRecvTimeout(void* connection, u32 timeoutMs)
+	void NetworkBackendMemory::SetRecvTimeout(void* connection, u32 timeoutMs, void* context, RecvTimeoutCallback* callback)
 	{
 	}
 
diff --git a/Engine/Source/Programs/UnrealBuildAccelerator/Common/Private/UbaNetworkBackendTcp.cpp b/Engine/Source/Programs/UnrealBuildAccelerator/Common/Private/UbaNetworkBackendTcp.cpp
--- a/Engine/Source/Programs/UnrealBuildAccelerator/Common/Private/UbaNetworkBackendTcp.cpp
+++ b/Engine/Source/Programs/UnrealBuildAccelerator/Common/Private/UbaNetworkBackendTcp.cpp
@@ -66,7 +66,10 @@ namespace uba
 		Event ready;
 		Guid uid;
 		u32 headerSize = 0;
+		
 		u32 recvTimeoutMs = 0;
+		void* recvTimeoutContext = nullptr;
+		RecvTimeoutCallback* recvTimeoutCallback = nullptr;
 
 		void* recvContext = nullptr;
 		RecvHeaderCallback* headerCallback = nullptr;
@@ -97,16 +100,13 @@ namespace uba
 	bool SetRecvBuf(Logger& logger, SOCKET socket, u32 windowSize);
 	bool SetSendBuf(Logger& logger, SOCKET socket, u32 windowSize);
 	bool DisableNagle(Logger& logger, SOCKET socket);
-	bool SendSocket(Logger& logger, SOCKET socket, const void* b, u64 bufferLen, const tchar* hint);
 
-	struct RecvCache
+	struct NetworkBackendTcp::RecvCache
 	{
 		u8 bytes[128*1024];
 		u32 byteCount = 0;
 	};
 
-	bool RecvSocket(Logger& logger, RecvCache& recvCache, SOCKET socket, void* b, u32& bufferLen, u32 timeoutMs, const Guid& connection, const tchar* hint1, const tchar* hint2, bool isFirstCall, bool allowLess);
-
 	bool NetworkBackendTcp::EnsureInitialized(Logger& logger)
 	{
 		#if PLATFORM_WINDOWS
@@ -206,7 +206,7 @@ namespace uba
 		SCOPED_FUTEX(conn.sendLock, lock);
 		if (conn.socket == INVALID_SOCKET)
 			return false;
-		bool res = SendSocket(logger, conn.socket, data, dataSize, sendHint);
+		bool res = SendSocket(conn, logger, data, dataSize, sendHint);
 		lock.Leave();
 
 		sendContext.isFinished = true;
@@ -241,10 +241,12 @@ namespace uba
 		conn.ready.Set();
 	}
 
-	void NetworkBackendTcp::SetRecvTimeout(void* connection, u32 timeoutMs)
+	void NetworkBackendTcp::SetRecvTimeout(void* connection, u32 timeoutMs, void* context, RecvTimeoutCallback* callback)
 	{
 		auto& conn = *(Connection*)connection;
 		conn.recvTimeoutMs = timeoutMs;
+		conn.recvTimeoutContext = context;
+		conn.recvTimeoutCallback = callback;
 	}
 
 	void NetworkBackendTcp::SetDisconnectCallback(void* connection, void* context, DisconnectCallback* callback)
@@ -481,9 +483,6 @@ namespace uba
 		if (connection.ready.IsSet(60000)) // This should never time out!
 		{
 			SetBlocking(logger, connection.socket, false);
-			u32 recvTimeoutMs = connection.recvTimeoutMs;
-			if (!recvTimeoutMs)
-				recvTimeoutMs = DefaultNetworkReceiveTimeoutSeconds * 1000; // 10 minutes time out
 
 			RecvCache recvCache;
 
@@ -495,7 +494,7 @@ namespace uba
 				u32 bodySize = 0;
 
 				u8 headerData[MaxHeaderSize];
-				if (!RecvSocket(logger, recvCache, connection.socket, headerData, connection.headerSize, recvTimeoutMs, connection.uid, connection.recvHint, TC(""), isFirst, false))
+				if (!RecvSocket(connection, recvCache, headerData, connection.headerSize, TC(""), isFirst, false))
 					break;
 				isFirst = false;
 
@@ -513,7 +512,7 @@ namespace uba
 				if (!bodySize)
 					continue;
 
-				bool success = RecvSocket(logger, recvCache, connection.socket, bodyData, bodySize, recvTimeoutMs, connection.uid, connection.recvHint, TC("Body"), false, connection.allowLess);
+				bool success = RecvSocket(connection, recvCache, bodyData, bodySize, TC("Body"), false, connection.allowLess);
 
 				m_totalRecv += bodySize;
 
@@ -858,8 +857,9 @@ namespace uba
 		return true;
 	}
 
-	bool SendSocket(Logger& logger, SOCKET socket, const void* b, u64 bufferLen, const tchar* hint)
+	bool NetworkBackendTcp::SendSocket(Connection& connection, Logger& logger, const void* b, u64 bufferLen, const tchar* hint)
 	{
+		SOCKET socket = connection.socket;
 		char* buffer = (char*)b;
 		u64 left = bufferLen;
 		while (left)
@@ -913,8 +913,10 @@ namespace uba
 		return true;
 	}
 
-	bool RecvSocket(Logger& logger, RecvCache& recvCache, SOCKET socket, void* b, u32& bufferLen, u32 timeoutMs, const Guid& connection, const tchar* hint1, const tchar* hint2, bool isFirstCall, bool allowLess)
+	bool NetworkBackendTcp::RecvSocket(Connection& connection, RecvCache& recvCache, void* b, u32& bufferLen, const tchar* hint, bool isFirstCall, bool allowLess)
 	{
+		SOCKET socket = connection.socket;
+
 		u8* buffer = (u8*)b;
 		u32 recvLeft = bufferLen;
 		while (recvLeft)
@@ -982,7 +984,7 @@ namespace uba
 			if (read == 0)
 			{
 				#if UBA_LOG_SOCKET_ERRORS
-				logger.Info(TC("RecvSocket - recv gracefully closed by peer for socket %i and connection %s (%s%s)"), socket, GuidToString(connection).str, hint1, hint2);
+				connection.logger.Info(TC("RecvSocket - recv gracefully closed by peer for socket %i and connection %s (%s%s)"), socket, GuidToString(connection.uid).str, connection.recvHint, hint);
 				#endif
 				return false;
 			}
@@ -991,6 +993,10 @@ namespace uba
 
 			if (WSAGetLastError() == WSAEWOULDBLOCK)
 			{
+				u32 timeoutMs = connection.recvTimeoutMs;
+				if (!timeoutMs)
+					timeoutMs = DefaultNetworkReceiveTimeoutSeconds * 1000; // 10 minutes time out
+
 				WSAPOLLFD p;
 				p.fd = socket;
 				p.revents = 0;
@@ -998,13 +1004,16 @@ namespace uba
 				int res = WSAPoll(&p, 1, int(timeoutMs));
 				if (!res)
 				{
-					logger.Info(TC("RecvSocket - WSAPoll returned timeout for socket %i and connection %s after %s (%s%s)"), socket, GuidToString(connection).str, TimeToText(MsToTime(timeoutMs)).str, hint1, hint2);
+					if (connection.recvTimeoutCallback && connection.recvTimeoutCallback(connection.recvTimeoutContext))
+						continue;
+
+					connection.logger.Info(TC("RecvSocket - WSAPoll returned timeout for socket %i and connection %s after %s (%s%s)"), socket, GuidToString(connection.uid).str, TimeToText(MsToTime(timeoutMs)).str, connection.recvHint, hint);
 					return false;
 				}
 				if (res == SOCKET_ERROR)
 				{
 					#if UBA_LOG_SOCKET_ERRORS
-					logger.Info(TC("RecvSocket - WSAPoll returned an error for socket %i and connection %s: %s (%s%s)"), socket, GuidToString(connection).str, LastErrorToText(WSAGetLastError()).data, hint1, hint2);
+					connection.logger.Info(TC("RecvSocket - WSAPoll returned an error for socket %i and connection %s: %s (%s%s)"), socket, GuidToString(connection.uid).str, LastErrorToText(WSAGetLastError()).data, connection.recvHint, hint);
 					#endif
 					return false;
 				}
@@ -1013,11 +1022,11 @@ namespace uba
 
 			#if !PLATFORM_WINDOWS
 			if (!isFirstCall && errno != ECONNRESET)
-				return logger.Warning(TC("RecvSocket - recv error on socket %i and connection %s: %s (%s%s)"), socket, GuidToString(connection).str, strerror(errno), hint1, hint2);
+				return connection.logger.Warning(TC("RecvSocket - recv error on socket %i and connection %s: %s (%s%s)"), socket, GuidToString(connection.uid).str, strerror(errno), connection.recvHint, hint);
 			#endif
 
 			#if UBA_LOG_SOCKET_ERRORS
-			logger.Info(TC("RecvSocket - read returned an error for socket %i and connection %s: %s (%s%s)"), socket, GuidToString(connection).str, LastErrorToText(WSAGetLastError()).data, hint1, hint2);
+			connection.logger.Info(TC("RecvSocket - read returned an error for socket %i and connection %s: %s (%s%s)"), socket, GuidToString(connection.uid).str, LastErrorToText(WSAGetLastError()).data, connection.recvHint, hint);
 			#endif
 			return false;
 		}
diff --git a/Engine/Source/Programs/UnrealBuildAccelerator/Common/Private/UbaNetworkClient.cpp b/Engine/Source/Programs/UnrealBuildAccelerator/Common/Private/UbaNetworkClient.cpp
--- a/Engine/Source/Programs/UnrealBuildAccelerator/Common/Private/UbaNetworkClient.cpp
+++ b/Engine/Source/Programs/UnrealBuildAccelerator/Common/Private/UbaNetworkClient.cpp
@@ -87,7 +87,7 @@ namespace uba
 			Atomic<u8> error;
 		};
 
-		backend.SetRecvTimeout(backendConnection, m_receiveTimeoutSeconds*1000);
+		backend.SetRecvTimeout(backendConnection, m_receiveTimeoutSeconds*1000, nullptr, [](void*) { return false; });
 
 		RecvContext rc(*this, backend, backendConnection);
 
@@ -226,7 +226,7 @@ namespace uba
 		m_logger.Detail(TC("Connected to server... (0x%p)"), backendConnection);
 		lock.Leave();
 
-		backend.SetRecvTimeout(backendConnection, m_receiveTimeoutSeconds*1000);
+		backend.SetRecvTimeout(backendConnection, m_receiveTimeoutSeconds*1000, nullptr, [](void*) { return false; });
 
 		backend.SetDisconnectCallback(backendConnection, connection, DisconnectCallback);
 		backend.SetRecvCallbacks(backendConnection, connection, ReceiveHeaderSize, ReceiveResponseHeader, ReceiveResponseBody, TC("ReceiveMessageResponse"));
diff --git a/Engine/Source/Programs/UnrealBuildAccelerator/Common/Private/UbaNetworkServer.cpp b/Engine/Source/Programs/UnrealBuildAccelerator/Common/Private/UbaNetworkServer.cpp
--- a/Engine/Source/Programs/UnrealBuildAccelerator/Common/Private/UbaNetworkServer.cpp
+++ b/Engine/Source/Programs/UnrealBuildAccelerator/Common/Private/UbaNetworkServer.cpp
@@ -121,7 +121,7 @@ namespace uba
 					conn.m_server.m_sendBytes += bytes;
 				});
 
-			m_backend.SetRecvTimeout(m_backendConnection, m_server.m_receiveTimeoutMs);
+			m_backend.SetRecvTimeout(m_backendConnection, m_server.m_receiveTimeoutMs, nullptr, [](void*) { return false; });
 
 			if (requiresCrypto)
 				m_backend.SetRecvCallbacks(m_backendConnection, this, 0, ReceiveHandshakeHeader, ReceiveHandshakeBody, TC("ReceiveHandshake"));
diff --git a/Engine/Source/Programs/UnrealBuildAccelerator/Common/Public/UbaNetworkBackend.h b/Engine/Source/Programs/UnrealBuildAccelerator/Common/Public/UbaNetworkBackend.h
--- a/Engine/Source/Programs/UnrealBuildAccelerator/Common/Public/UbaNetworkBackend.h
+++ b/Engine/Source/Programs/UnrealBuildAccelerator/Common/Public/UbaNetworkBackend.h
@@ -45,7 +45,9 @@ namespace uba
 		using RecvHeaderCallback = bool(void* context, const Guid& connectionUid, u8* headerData, void*& outBodyContext, u8*& outBodyData, u32& outBodySize);
 		using RecvBodyCallback = bool(void* context, bool recvError, u8* headerData, void* bodyContext, u8* bodyData, u32 bodySize);
 		virtual void SetRecvCallbacks(void* connection, void* context, u32 headerSize, RecvHeaderCallback* h, RecvBodyCallback* b, const tchar* recvHint) = 0;
-		virtual void SetRecvTimeout(void* connection, u32 timeoutMs) = 0;
+
+		using RecvTimeoutCallback = bool(void* context);
+		virtual void SetRecvTimeout(void* connection, u32 timeoutMs, void* context, RecvTimeoutCallback* callback) = 0;
 
 		// Disconnect callback. This is called as soon as connection is interrupted from send, recv or shutdown.
 		using DisconnectCallback = void(void* context, const Guid& connectionUid, void* connection);
diff --git a/Engine/Source/Programs/UnrealBuildAccelerator/Common/Public/UbaNetworkBackendMemory.h b/Engine/Source/Programs/UnrealBuildAccelerator/Common/Public/UbaNetworkBackendMemory.h
--- a/Engine/Source/Programs/UnrealBuildAccelerator/Common/Public/UbaNetworkBackendMemory.h
+++ b/Engine/Source/Programs/UnrealBuildAccelerator/Common/Public/UbaNetworkBackendMemory.h
@@ -17,7 +17,7 @@ namespace uba
 		virtual bool Send(Logger& logger, void* connection, const void* data, u32 dataSize, SendContext& sendContext, const tchar* sendHint) override;
 		virtual void SetDataSentCallback(void* connection, void* context, DataSentCallback* callback) override;
 		virtual void SetRecvCallbacks(void* connection, void* context, u32 headerSize, RecvHeaderCallback* h, RecvBodyCallback* b, const tchar* recvHint) override;
-		virtual void SetRecvTimeout(void* connection, u32 timeoutMs) override;
+		virtual void SetRecvTimeout(void* connection, u32 timeoutMs, void* context, RecvTimeoutCallback* callback) override;
 		virtual void SetDisconnectCallback(void* connection, void* context, DisconnectCallback* callback) override;
 		virtual void SetAllowLessThanBodySize(void* connection, bool allow) override;
 
diff --git a/Engine/Source/Programs/UnrealBuildAccelerator/Common/Public/UbaNetworkBackendTcp.h b/Engine/Source/Programs/UnrealBuildAccelerator/Common/Public/UbaNetworkBackendTcp.h
--- a/Engine/Source/Programs/UnrealBuildAccelerator/Common/Public/UbaNetworkBackendTcp.h
+++ b/Engine/Source/Programs/UnrealBuildAccelerator/Common/Public/UbaNetworkBackendTcp.h
@@ -28,7 +28,7 @@ namespace uba
 		virtual bool Send(Logger& logger, void* connection, const void* data, u32 dataSize, SendContext& sendContext, const tchar* sendHint) override;
 		virtual void SetDataSentCallback(void* connection, void* context, DataSentCallback* callback) override;
 		virtual void SetRecvCallbacks(void* connection, void* context, u32 headerSize, RecvHeaderCallback* h, RecvBodyCallback* b, const tchar* recvHint) override;
-		virtual void SetRecvTimeout(void* connection, u32 timeoutMs) override;
+		virtual void SetRecvTimeout(void* connection, u32 timeoutMs, void* context, RecvTimeoutCallback* callback) override;
 		virtual void SetDisconnectCallback(void* connection, void* context, DisconnectCallback* callback) override;
 		virtual void SetAllowLessThanBodySize(void* connection, bool allow) override;
 
@@ -45,8 +45,11 @@ namespace uba
 
 		struct Connection;
 		struct ListenEntry;
+		struct RecvCache;
 		bool ThreadListen(Logger& logger, ListenEntry& entry);
 		void ThreadRecv(Connection& connection);
+		bool SendSocket(Connection& connection, Logger& logger, const void* b, u64 bufferLen, const tchar* hint);
+		bool RecvSocket(Connection& connection, RecvCache& recvCache, void* b, u32& bufferLen, const tchar* hint, bool isFirstCall, bool allowLess);
 
 		LoggerWithWriter m_logger;
 		Futex m_listenEntriesLock;
