[UBA]
* Changed MacBinDependency parser to use the actual structures of .dylib instead of searching in memory. Supports both fat files and mach 64

diff --git a/Engine/Source/Programs/UnrealBuildAccelerator/Core/Public/UbaMacBinDependencyParser.h b/Engine/Source/Programs/UnrealBuildAccelerator/Core/Public/UbaMacBinDependencyParser.h
--- a/Engine/Source/Programs/UnrealBuildAccelerator/Core/Public/UbaMacBinDependencyParser.h
+++ b/Engine/Source/Programs/UnrealBuildAccelerator/Core/Public/UbaMacBinDependencyParser.h
@@ -4,7 +4,9 @@
 
 #if PLATFORM_MAC
 
-#include "UbaStringBuffer.h"
+#include "UbaBinaryReaderWriter.h"
+#include "UbaPlatform.h"
+#include <mach-o/fat.h>
 
 #if !defined(UBA_IS_DETOURED_INCLUDE)
 #define TRUE_WRAPPER(func) func
@@ -17,31 +19,32 @@ namespace uba
 		return false;
 	}
 
-	inline bool FindImports(const tchar* fileName, const StringView& originalPath, const Function<void(const tchar* import, bool isKnown, const char* const* loaderPaths)>& func, StringBufferBase& outError)
+	inline bool FindImports(const tchar* filePath, const StringView& originalPath, const Function<void(const tchar* import, bool isKnown, const char* const* loaderPaths)>& func, StringBufferBase& outError)
 	{
-		int fd = TRUE_WRAPPER(open)(fileName, O_RDONLY);
+		int fd = TRUE_WRAPPER(open)(filePath, O_RDONLY);
 		if (fd == -1)
 		{
-			outError.Appendf("Open failed for file %s", fileName);
+			outError.Appendf("Open failed for file %s", filePath);
 			return false;
 		}
 		auto closeFileHandle = MakeGuard([&]() { TRUE_WRAPPER(close)(fd); });
 		struct stat sb;
 		if (TRUE_WRAPPER(fstat)(fd, &sb) == -1)
 		{
-			outError.Appendf("Stat failed for file %s", fileName);
+			outError.Appendf("Stat failed for file %s", filePath);
 			return false;
 		}
-		u32 size = Min(u32(sb.st_size), 8048u);
+		u32 fileSize = Min(u32(sb.st_size), 1024u*1024u);
 
-		void* mem = TRUE_WRAPPER(mmap)(nullptr, size, PROT_READ, MAP_PRIVATE, fd, 0);
+		auto mem = (u8*)TRUE_WRAPPER(mmap)(nullptr, fileSize, PROT_READ, MAP_PRIVATE, fd, 0);
 		if (mem == MAP_FAILED)
 		{
-			outError.Appendf("Mmap failed for file %s", fileName);
+			outError.Appendf("Mmap failed for file %s", filePath);
 			return false;
 		}
-		auto unmap = MakeGuard([&]() { TRUE_WRAPPER(munmap)(mem, size); });
+		auto unmap = MakeGuard([&]() { TRUE_WRAPPER(munmap)(mem, fileSize); });
 
+		StringView fileName = originalPath.GetFileName();
 
 		const char* libs[1024];
 		u32 libsCount = 0;
@@ -49,40 +52,75 @@ namespace uba
 		const char* loaderPaths[256];
 		u32 loaderPathsCount = 0;
 
-		auto it = (const u8*)mem;
-		auto end = it + size - 14;
-		while (it != end)
+		auto handleHeader = [&](u32 offset)
 		{
-			if (*it++ != '@')
-				continue;
-
-			if (memcmp(it, "rpath/", 6) == 0)
+			u8* memIt = mem + offset;
+			u32 magic = *(const u32*)memIt;
+			UBA_ASSERT(magic == MH_MAGIC_64);
+			if (magic != MH_MAGIC_64)
+				return true;
+			auto& mh = *(mach_header_64*)memIt;
+			memIt += sizeof(mach_header_64);
+			for (u32 i = 0; i < mh.ncmds; i++)
 			{
-				it += 6;
-				auto importFile = (const char*)it;
-				if (!strstr(importFile, ".dylib"))
+				auto& lc = *(load_command*)memIt;
+				auto g = MakeGuard([&]() { memIt += lc.cmdsize; });
+				if (lc.cmd != LC_LOAD_DYLIB && lc.cmd != LC_RPATH) //  && lc.cmd != LC_LOAD_WEAK_DYLIB
+					continue;
+				auto& cmdData = *(dylib_command*)memIt;
+				u32 nameOffset = cmdData.dylib.name.offset;
+				if (nameOffset >= lc.cmdsize)
+					continue;
+				auto name = ToView((const char*)(memIt + nameOffset));
+				if (*name.data != '@')
+					continue;
+				name = name.Skip(1);
+				if (name.StartsWith("rpath/"))
 				{
-					outError.Appendf("Found @rpath in binary %s that did not end with .dylib (%s)", fileName, importFile);
-					return false;
+					auto importFile = name.Skip(6);
+					if (!importFile.EndsWith(".dylib"))
+					{
+						outError.Appendf("Found @rpath in binary %s that did not end with .dylib (%s)", filePath, importFile.data);
+						return false;
+					}
+					if (fileName.Equals(importFile))
+						continue;
+					libs[libsCount++] = importFile.data;
+				}
+				else if (name.StartsWith("executable_path/"))
+				{
+					auto executablePath = name.Skip(16).data;
+					UBA_ASSERT(*executablePath == 0);
+				}
+				else if (name.StartsWith("loader_path/"))
+				{
+					loaderPaths[loaderPathsCount++] = name.Skip(12).data;
 				}
-
-				libs[libsCount++] = importFile;
-			}
-			else if (memcmp(it, "executable_path/", 16) == 0)
-			{
-				it += 16;
-				auto executablePath = (const char*)it;
-				UBA_ASSERT(*executablePath == 0);
 			}
-			else if (memcmp(it, "loader_path/", 12) == 0)
+			return true;
+		};
+
+		u32 magic = *(const u32*)mem;
+		if (magic == FAT_CIGAM || magic == FAT_MAGIC)
+		{
+			u8* memIt = mem;
+			auto& fh = *(fat_header*)memIt;
+			memIt += sizeof(fat_header);
+			for (u32 i = 0, e = ntohl(fh.nfat_arch); i<e; i++)
 			{
-				it += 12;
-				loaderPaths[loaderPathsCount++] = (const char*)it;
+				auto& arch = *(fat_arch_64*)memIt;
+				memIt += sizeof(fat_arch_64);
+				u32 offset = ntohl(arch.offset);
+				u32 cputype = ntohl(arch.cputype);
+				if (cputype == CPU_TYPE_X86_64 || cputype == CPU_TYPE_ARM64)
+					if (!handleHeader(offset))
+						return false;
 			}
 		}
+		else if (!handleHeader(0))
+			return false;
 
 		loaderPaths[loaderPathsCount] = nullptr;
-
 		for (u32 i=0; i!=libsCount; ++i)
 			func(libs[i], false, loaderPaths);
 		return true;
