[UBA]
* Changed FindImports to take StringView for filePath instead of const tchar*

diff --git a/Engine/Source/Programs/UnrealBuildAccelerator/Cli/Private/UbaCli.cpp b/Engine/Source/Programs/UnrealBuildAccelerator/Cli/Private/UbaCli.cpp
--- a/Engine/Source/Programs/UnrealBuildAccelerator/Cli/Private/UbaCli.cpp
+++ b/Engine/Source/Programs/UnrealBuildAccelerator/Cli/Private/UbaCli.cpp
@@ -32,7 +32,7 @@
 namespace uba
 {
 
-	bool FindImports(const tchar* fileName, const StringView& originalPath, const Function<void(const tchar* import, bool isKnown, const char* const* loaderPaths)>& func, StringBufferBase& outError);
+	bool FindImports(const StringView& filePath, const StringView& originalPath, const Function<void(const tchar* import, bool isKnown, const char* const* loaderPaths)>& func, StringBufferBase& outError);
 
 	const tchar*	Version = GetVersionString();
 	u32				DefaultCapacityGb = 20;
@@ -817,7 +817,7 @@ namespace uba
 			{
 				StringBuffer<> error;
 				bool printImports = true;
-				if (!FindImports(application.c_str(), StringView(application).GetPath(), [&](const tchar* import, bool isKnown, const char* const* loaderPaths)
+				if (!FindImports(application, StringView(application).GetPath(), [&](const tchar* import, bool isKnown, const char* const* loaderPaths)
 				{
 					if (printImports)
 					{
diff --git a/Engine/Source/Programs/UnrealBuildAccelerator/Common/Private/UbaSession.cpp b/Engine/Source/Programs/UnrealBuildAccelerator/Common/Private/UbaSession.cpp
--- a/Engine/Source/Programs/UnrealBuildAccelerator/Common/Private/UbaSession.cpp
+++ b/Engine/Source/Programs/UnrealBuildAccelerator/Common/Private/UbaSession.cpp
@@ -1422,10 +1422,10 @@ namespace uba
 		if (isSystem && IsKnownSystemFile(applicationName))
 			return true;
 
-		out.push_back({ library, temp3.data, attr, isSystem });
+		out.push_back({ library, temp3.ToString(), attr, isSystem });
 
 		StringBuffer<> errorStr;
-		FindImports(applicationName, ToView(applicationDir), [&](const tchar* importName, bool isKnown, const char* const* importLoaderPaths)
+		FindImports(temp3, ToView(applicationDir), [&](const tchar* importName, bool isKnown, const char* const* importLoaderPaths)
 			{
 				if (result && !isKnown)
 					result = CopyImports(out, importName, applicationDir, applicationDirEnd, handledImports, importLoaderPaths);
diff --git a/Engine/Source/Programs/UnrealBuildAccelerator/Core/Public/UbaLinuxBinDependencyParser.h b/Engine/Source/Programs/UnrealBuildAccelerator/Core/Public/UbaLinuxBinDependencyParser.h
--- a/Engine/Source/Programs/UnrealBuildAccelerator/Core/Public/UbaLinuxBinDependencyParser.h
+++ b/Engine/Source/Programs/UnrealBuildAccelerator/Core/Public/UbaLinuxBinDependencyParser.h
@@ -47,26 +47,26 @@ namespace uba
 		return false;
 	}
 
-	bool FindImports(const tchar* fileName, const StringView& originalPath, const Function<void(const tchar* import, bool isKnown, const char* const* loaderPaths)>& func, StringBufferBase& outError)
+	bool FindImports(const StringView& filePath, const StringView& originalPath, const Function<void(const tchar* import, bool isKnown, const char* const* loaderPaths)>& func, StringBufferBase& outError)
 	{
 		#if UBA_FINDIMPORTS_DEBUG
-		printf("FINDIMPORTS: %s\n", fileName);
+		printf("FINDIMPORTS: %s\n", filePath.data);
 		#endif
 
-		int fd = TRUE_WRAPPER(open)(fileName, O_RDONLY);
+		int fd = TRUE_WRAPPER(open)(filePath.data, O_RDONLY);
 		if (fd == -1)
-			return outError.Appendf("Open failed for file %s", fileName).ToFalse();
+			return outError.Appendf("Open failed for file %s", filePath.data).ToFalse();
 
 		auto closeFileHandle = MakeGuard([&]() { TRUE_WRAPPER(close)(fd); });
 		struct stat sb;
 		if (TRUE_WRAPPER(fstat)(fd, &sb) == -1)
-			return outError.Appendf("Stat failed for file %s", fileName).ToFalse();
+			return outError.Appendf("Stat failed for file %s", filePath.data).ToFalse();
 
 		u32 size = sb.st_size;
 
 		auto mem = (u8*)mmap(nullptr, size, PROT_READ, MAP_PRIVATE, fd, 0);
 		if (mem == MAP_FAILED)
-			return outError.Appendf("Mmap failed for file %s", fileName).ToFalse();
+			return outError.Appendf("Mmap failed for file %s", filePath.data).ToFalse();
 		auto unmap = MakeGuard([&]() { munmap(mem, size); });
 
 		auto ehdr = (Elf64_Ehdr*)mem;
@@ -154,7 +154,7 @@ namespace uba
 		List<TString> fixedLoaderPaths;
 
 		#if UBA_FINDIMPORTS_DEBUG
-		printf("ELF: %s\n", fileName);
+		printf("ELF: %s\n", filePath.data);
 		#endif
 
 		#if UBA_FINDIMPORTS_DEBUG
diff --git a/Engine/Source/Programs/UnrealBuildAccelerator/Core/Public/UbaMacBinDependencyParser.h b/Engine/Source/Programs/UnrealBuildAccelerator/Core/Public/UbaMacBinDependencyParser.h
--- a/Engine/Source/Programs/UnrealBuildAccelerator/Core/Public/UbaMacBinDependencyParser.h
+++ b/Engine/Source/Programs/UnrealBuildAccelerator/Core/Public/UbaMacBinDependencyParser.h
@@ -19,19 +19,19 @@ namespace uba
 		return false;
 	}
 
-	inline bool FindImports(const tchar* filePath, const StringView& originalPath, const Function<void(const tchar* import, bool isKnown, const char* const* loaderPaths)>& func, StringBufferBase& outError)
+	inline bool FindImports(const StringView& filePath, const StringView& originalPath, const Function<void(const tchar* import, bool isKnown, const char* const* loaderPaths)>& func, StringBufferBase& outError)
 	{
-		int fd = TRUE_WRAPPER(open)(filePath, O_RDONLY);
+		int fd = TRUE_WRAPPER(open)(filePath.data, O_RDONLY);
 		if (fd == -1)
 		{
-			outError.Appendf("Open failed for file %s", filePath);
+			outError.Appendf("Open failed for file %s", filePath.data);
 			return false;
 		}
 		auto closeFileHandle = MakeGuard([&]() { TRUE_WRAPPER(close)(fd); });
 		struct stat sb;
 		if (TRUE_WRAPPER(fstat)(fd, &sb) == -1)
 		{
-			outError.Appendf("Stat failed for file %s", filePath);
+			outError.Appendf("Stat failed for file %s", filePath.data);
 			return false;
 		}
 		u32 fileSize = Min(u32(sb.st_size), 1024u*1024u);
@@ -39,7 +39,7 @@ namespace uba
 		auto mem = (u8*)TRUE_WRAPPER(mmap)(nullptr, fileSize, PROT_READ, MAP_PRIVATE, fd, 0);
 		if (mem == MAP_FAILED)
 		{
-			outError.Appendf("Mmap failed for file %s", filePath);
+			outError.Appendf("Mmap failed for file %s", filePath.data);
 			return false;
 		}
 		auto unmap = MakeGuard([&]() { TRUE_WRAPPER(munmap)(mem, fileSize); });
@@ -80,7 +80,7 @@ namespace uba
 					auto importFile = name.Skip(6);
 					if (!importFile.EndsWith(".dylib"))
 					{
-						outError.Appendf("Found @rpath in binary %s that did not end with .dylib (%s)", filePath, importFile.data);
+						outError.Appendf("Found @rpath in binary %s that did not end with .dylib (%s)", filePath.data, importFile.data);
 						return false;
 					}
 					if (fileName.Equals(importFile))
diff --git a/Engine/Source/Programs/UnrealBuildAccelerator/Core/Public/UbaWinBinDependencyParser.h b/Engine/Source/Programs/UnrealBuildAccelerator/Core/Public/UbaWinBinDependencyParser.h
--- a/Engine/Source/Programs/UnrealBuildAccelerator/Core/Public/UbaWinBinDependencyParser.h
+++ b/Engine/Source/Programs/UnrealBuildAccelerator/Core/Public/UbaWinBinDependencyParser.h
@@ -129,11 +129,11 @@ namespace uba
 		return true;
 	}
 
-	inline bool FindImports(const tchar* fileName, const StringView& originalPath, const Function<void(const tchar* import, bool isKnown, const char* const* importLoaderPaths)>& func, StringBufferBase& outError)
+	inline bool FindImports(const StringView& filePath, const StringView& originalPath, const Function<void(const tchar* import, bool isKnown, const char* const* importLoaderPaths)>& func, StringBufferBase& outError)
 	{
-		if (ToView(fileName).EndsWith(TCV(".bat")))
+		if (filePath.EndsWith(TCV(".bat")))
 			return true;
-		HANDLE fileHandle = True_CreateFileW(fileName, GENERIC_READ, FILE_SHARE_READ, 0, OPEN_EXISTING, FILE_ATTRIBUTE_NORMAL | FILE_FLAG_SEQUENTIAL_SCAN, 0);
+		HANDLE fileHandle = True_CreateFileW(filePath.data, GENERIC_READ, FILE_SHARE_READ, 0, OPEN_EXISTING, FILE_ATTRIBUTE_NORMAL | FILE_FLAG_SEQUENTIAL_SCAN, 0);
 		if (fileHandle == INVALID_HANDLE_VALUE)
 			return true;
 		auto closeFileHandle = MakeGuard([&]() { CloseHandle(fileHandle); });
@@ -146,6 +146,6 @@ namespace uba
 			return false;
 		auto unmap = MakeGuard([&]() { True_UnmapViewOfFile(mem); });
 
-		return FindImportsInMem(fileName, mem, func);
+		return FindImportsInMem(filePath.data, mem, func);
 	}
 }
diff --git a/Engine/Source/Programs/UnrealBuildAccelerator/Detours/Private/PosixOS/UbaDetoursFunctionsPosix.cpp b/Engine/Source/Programs/UnrealBuildAccelerator/Detours/Private/PosixOS/UbaDetoursFunctionsPosix.cpp
--- a/Engine/Source/Programs/UnrealBuildAccelerator/Detours/Private/PosixOS/UbaDetoursFunctionsPosix.cpp
+++ b/Engine/Source/Programs/UnrealBuildAccelerator/Detours/Private/PosixOS/UbaDetoursFunctionsPosix.cpp
@@ -2227,12 +2227,9 @@ void Shared_LoadLibrary(const char*& path, const char* const* loaderPaths, Strin
 {
 	StringBuffer<512> virtualPath;
 	Rpc_GetFullFileName2(path, tempBuf, virtualPath, loaderPaths);
-	path = tempBuf.data;
-
 	StringView originalPath = StringView(virtualPath).GetPath();
-
 	StringBuffer<> error;
-	FindImports(path, originalPath, [&](const tchar* import, bool isKnown, const char* const* importLoaderPaths)
+	FindImports(tempBuf, originalPath, [&](const tchar* import, bool isKnown, const char* const* importLoaderPaths)
 	{
 		if (!g_handledLibraries.insert(import).second)
 			return;
diff --git a/Engine/Source/Programs/UnrealBuildAccelerator/Detours/Private/Windows/UbaDetoursFunctionsKernelBase.inl b/Engine/Source/Programs/UnrealBuildAccelerator/Detours/Private/Windows/UbaDetoursFunctionsKernelBase.inl
--- a/Engine/Source/Programs/UnrealBuildAccelerator/Detours/Private/Windows/UbaDetoursFunctionsKernelBase.inl
+++ b/Engine/Source/Programs/UnrealBuildAccelerator/Detours/Private/Windows/UbaDetoursFunctionsKernelBase.inl
@@ -2702,9 +2702,9 @@ DWORD Detoured_SearchPathW(LPCWSTR lpPath, LPCWSTR lpFileName, LPCWSTR lpExtensi
 using AdditionalLoads = Vector<HMODULE, GrowingAllocator<HMODULE>>;
 using VisitedModules = std::unordered_set<StringKey, std::hash<StringKey>, std::equal_to<StringKey>, GrowingAllocator<StringKey>>;
 
-HMODULE Recursive_LoadLibraryExW(LPCWSTR lpLibFileName, LPCWSTR originalName, DWORD dwFlags, AdditionalLoads& additionalLoads, VisitedModules& visitedModules)
+HMODULE Recursive_LoadLibraryExW(const StringView& filePath, LPCWSTR originalName, DWORD dwFlags, AdditionalLoads& additionalLoads, VisitedModules& visitedModules)
 {
-	if (!visitedModules.insert(ToStringKeyNoCheck(lpLibFileName, wcslen(lpLibFileName))).second)
+	if (!visitedModules.insert(ToStringKeyNoCheck(filePath.data, filePath.count)).second)
 		return 0;
 
 	// Important that this code is not doing allocations.. it could cause a recursive stack overflow
@@ -2713,13 +2713,13 @@ HMODULE Recursive_LoadLibraryExW(LPCWSTR lpLibFileName, LPCWSTR originalName, DW
 	{
 		SuppressCreateFileDetourScope cfs;
 		StringBuffer<256> error;
-		if (!FindImports(lpLibFileName, {}, [&](const wchar_t* import, bool isKnown, const char* const* importLoaderPaths)
+		if (!FindImports(filePath, {}, [&](const wchar_t* import, bool isKnown, const char* const* importLoaderPaths)
 			{
 				if (!GetModuleHandleW(import))
 					importedModules.emplace_back(import, isKnown);
 			}, error))
 		{
-			UBA_ASSERTF(false, L"Failed to find imports for binary %ls (%ls)", lpLibFileName, originalName);
+			UBA_ASSERTF(false, L"Failed to find imports for binary %ls (%ls)", filePath.data, originalName);
 		}
 	}
 	for (auto& importedModule : importedModules)
@@ -2738,7 +2738,7 @@ HMODULE Recursive_LoadLibraryExW(LPCWSTR lpLibFileName, LPCWSTR originalName, DW
 		{
 			if (HMODULE h = True_LoadLibraryExW(importedModule.name, 0, 0))
 			{
-				OnModuleLoaded(h, importedModule.name);
+				OnModuleLoaded(h, ToView(importedModule.name));
 				additionalLoads.push_back(h);
 			}
 			continue;
@@ -2752,12 +2752,12 @@ HMODULE Recursive_LoadLibraryExW(LPCWSTR lpLibFileName, LPCWSTR originalName, DW
 		StringBuffer<512> tempBuf;
 		Rpc_GetFullFileName(path, pathLen, tempBuf, false);
 
-		if (HMODULE r = Recursive_LoadLibraryExW(path, importedModule.name, dwFlags, additionalLoads, visitedModules))
+		if (HMODULE r = Recursive_LoadLibraryExW(ToView(path), importedModule.name, dwFlags, additionalLoads, visitedModules))
 			additionalLoads.push_back(r);
 	}
 
 	StringBuffer<512> newName;
-	if (originalName[1] != ':' && !Equals(lpLibFileName, originalName))
+	if (originalName[1] != ':' && !filePath.Equals(originalName))
 	{
 		newName.Append(g_virtualApplicationDir).Append(originalName);
 		originalName = newName.data;
@@ -2768,7 +2768,7 @@ HMODULE Recursive_LoadLibraryExW(LPCWSTR lpLibFileName, LPCWSTR originalName, DW
 	DEBUG_LOG_TRUE(L"INTERNAL LoadLibraryExW", L"%ls", originalName);
 
 	SuppressCreateFileDetourScope cfs;
-	auto res = True_LoadLibraryExW(lpLibFileName, 0, 0);
+	auto res = True_LoadLibraryExW(filePath.data, 0, 0);
 	if (res)
 	{
 		if (originalName[1] == ':')
@@ -2777,7 +2777,7 @@ HMODULE Recursive_LoadLibraryExW(LPCWSTR lpLibFileName, LPCWSTR originalName, DW
 			SCOPED_WRITE_LOCK(g_loadedModulesLock, lock);
 			g_loadedModules[res] = originalName;
 		}
-		OnModuleLoaded(res, lpLibFileName);
+		OnModuleLoaded(res, filePath);
 	}
 	return res;
 }
@@ -2846,7 +2846,7 @@ HMODULE Shared_LoadLibrary(LPCWSTR lpLibFileName, HANDLE hFile, DWORD dwFlags)
 
 	AdditionalLoads additionalLoads(g_memoryBlock); // Don't do allocations
 	VisitedModules visitedModules(g_memoryBlock);
-	HMODULE res = Recursive_LoadLibraryExW(newPath, fileName, dwFlags, additionalLoads, visitedModules);
+	HMODULE res = Recursive_LoadLibraryExW(ToView(newPath), fileName, dwFlags, additionalLoads, visitedModules);
 	for (HMODULE h : additionalLoads)
 		FreeLibrary(h);
 	return res;
diff --git a/Engine/Source/Programs/UnrealBuildAccelerator/Detours/Private/Windows/UbaDetoursFunctionsWin.cpp b/Engine/Source/Programs/UnrealBuildAccelerator/Detours/Private/Windows/UbaDetoursFunctionsWin.cpp
--- a/Engine/Source/Programs/UnrealBuildAccelerator/Detours/Private/Windows/UbaDetoursFunctionsWin.cpp
+++ b/Engine/Source/Programs/UnrealBuildAccelerator/Detours/Private/Windows/UbaDetoursFunctionsWin.cpp
@@ -601,7 +601,7 @@ void CloseCaches()
 bool g_exitMessageSent;
 
 void SendExitMessage(DWORD exitCode, u64 startTime);
-void OnModuleLoaded(HMODULE moduleHandle, const wchar_t* name);
+void OnModuleLoaded(HMODULE moduleHandle, const StringView& name);
 
 // Variables used to communicate state from kernelbase functions to ntdll functions
 thread_local const wchar_t* t_renameFileNewName;
@@ -778,10 +778,10 @@ int DetourAttachFunctions(bool runningRemote)
 	return 0;
 }
 
-void OnModuleLoaded(HMODULE moduleHandle, const wchar_t* name)
+void OnModuleLoaded(HMODULE moduleHandle, const StringView& name)
 {
 	// SymLoadModuleExW do something bad that cause remote wine to fail everything after this call.. TODO: Revisit
-	if (g_isRunningWine && !True_SymLoadModuleExW && Contains(name, L"dbghelp.dll"))
+	if (g_isRunningWine && !True_SymLoadModuleExW && name.Contains(L"dbghelp.dll"))
 	{
 		True_SymLoadModuleExW = (SymLoadModuleExWFunc*)GetProcAddress(moduleHandle, "SymLoadModuleExW");
 		UBA_ASSERT(True_SymLoadModuleExW);
@@ -791,7 +791,7 @@ void OnModuleLoaded(HMODULE moduleHandle, const wchar_t* name)
 	}
 
 	// ImageGetDigestStream is buggy in wine so we have to detour it for ShaderCompileWorker
-	if (g_isRunningWine && !True_ImageGetDigestStream && Contains(name, L"imagehlp.dll"))
+	if (g_isRunningWine && !True_ImageGetDigestStream && name.Contains(L"imagehlp.dll"))
 	{
 		True_ImageGetDigestStream = (ImageGetDigestStreamFunc*)GetProcAddress(moduleHandle, "ImageGetDigestStream");
 		UBA_ASSERT(True_ImageGetDigestStream);
@@ -801,7 +801,7 @@ void OnModuleLoaded(HMODULE moduleHandle, const wchar_t* name)
 	}
 
 	// SHGetKnownFolderPath is used by Metal.exe and must always execute on host
-	if (!True_SHGetKnownFolderPath && Contains(name, L"shell32.dll"))
+	if (!True_SHGetKnownFolderPath && name.Contains(L"shell32.dll"))
 	{
 		True_SHGetKnownFolderPath = (SHGetKnownFolderPathFunc*)GetProcAddress(moduleHandle, "SHGetKnownFolderPath");
 		UBA_ASSERT(True_SHGetKnownFolderPath);
diff --git a/Engine/Source/Programs/UnrealBuildAccelerator/Host/Private/UbaExports.cpp b/Engine/Source/Programs/UnrealBuildAccelerator/Host/Private/UbaExports.cpp
--- a/Engine/Source/Programs/UnrealBuildAccelerator/Host/Private/UbaExports.cpp
+++ b/Engine/Source/Programs/UnrealBuildAccelerator/Host/Private/UbaExports.cpp
@@ -903,7 +903,7 @@ extern "C"
 	{
 #if PLATFORM_WINDOWS
 		uba::StringBuffer<> errors;
-		uba::FindImports(binary, {}, [&](const uba::tchar* importName, bool isKnown, const char* const* importLoaderPaths) { func(importName, userData); }, errors);
+		uba::FindImports(uba::ToView(binary), {}, [&](const uba::tchar* importName, bool isKnown, const char* const* importLoaderPaths) { func(importName, userData); }, errors);
 #else
 		UBA_ASSERT(false);
 #endif
diff --git a/Engine/Source/Programs/UnrealBuildAccelerator/Test/Private/UbaTestBasics.cpp b/Engine/Source/Programs/UnrealBuildAccelerator/Test/Private/UbaTestBasics.cpp
--- a/Engine/Source/Programs/UnrealBuildAccelerator/Test/Private/UbaTestBasics.cpp
+++ b/Engine/Source/Programs/UnrealBuildAccelerator/Test/Private/UbaTestBasics.cpp
@@ -533,7 +533,7 @@ namespace uba
 		path.EnsureEndsWithSlash().Append(TCV("UbaTestApp.exe"));
 		bool importKernel = false;
 		StringBuffer<> error;
-		FindImports(path.data, {}, [&](const tchar* import, bool isKnown, const char* const* importLoaderPaths)
+		FindImports(path, {}, [&](const tchar* import, bool isKnown, const char* const* importLoaderPaths)
 		{
 			importKernel |= isKnown && Contains(import, TC("KERNEL32.dll"));
 		}, error);
