[UBA]
* Changed active processes to be set as cancelled if using network to read trace and got disconnected

diff --git a/Engine/Source/Programs/UnrealBuildAccelerator/Visualizer/Private/UbaTraceReader.cpp b/Engine/Source/Programs/UnrealBuildAccelerator/Visualizer/Private/UbaTraceReader.cpp
--- a/Engine/Source/Programs/UnrealBuildAccelerator/Visualizer/Private/UbaTraceReader.cpp
+++ b/Engine/Source/Programs/UnrealBuildAccelerator/Visualizer/Private/UbaTraceReader.cpp
@@ -197,7 +197,7 @@ namespace uba
 					if (process.stop == ~u64(0))
 					{
 						process.stop = GetTime() - out.startTime;
-						process.exitCode = 0; // This is wrong but since we didn't get the final result we can't tell if it was success or not
+						process.exitCode = ProcessCancelExitCode; // This is wrong but since we didn't get the final result we can't tell if it was success or not
 						process.bitmapDirty = true;
 					}
 				}
