[UBA]
* Renamed FindImports to ParseBinary
* Added BinaryInfo placeholder that will be used to contain info about binary such as min version etc

diff --git a/Engine/Source/Programs/UnrealBuildAccelerator/Cli/Private/UbaCli.cpp b/Engine/Source/Programs/UnrealBuildAccelerator/Cli/Private/UbaCli.cpp
--- a/Engine/Source/Programs/UnrealBuildAccelerator/Cli/Private/UbaCli.cpp
+++ b/Engine/Source/Programs/UnrealBuildAccelerator/Cli/Private/UbaCli.cpp
@@ -1,6 +1,7 @@
 // Copyright Epic Games, Inc. All Rights Reserved.
 
 #include "UbaApplication.h"
+#include "UbaBinaryParser.h"
 #include "UbaCacheClient.h"
 #include "UbaCacheServer.h"
 #include "UbaClient.h"
@@ -31,9 +32,6 @@
 
 namespace uba
 {
-
-	bool FindImports(const StringView& filePath, const StringView& originalPath, const Function<void(const tchar* import, bool isKnown, const char* const* loaderPaths)>& func, StringBufferBase& outError);
-
 	const tchar*	Version = GetVersionString();
 	u32				DefaultCapacityGb = 20;
 	const tchar*	DefaultRootDir = []() {
@@ -817,7 +815,8 @@ namespace uba
 			{
 				StringBuffer<> error;
 				bool printImports = true;
-				if (!FindImports(application, StringView(application).GetPath(), [&](const tchar* import, bool isKnown, const char* const* loaderPaths)
+				BinaryInfo info;
+				if (!ParseBinary(application, StringView(application).GetPath(), info, [&](const tchar* import, bool isKnown, const char* const* loaderPaths)
 				{
 					if (printImports)
 					{
diff --git a/Engine/Source/Programs/UnrealBuildAccelerator/Common/Private/UbaSession.cpp b/Engine/Source/Programs/UnrealBuildAccelerator/Common/Private/UbaSession.cpp
--- a/Engine/Source/Programs/UnrealBuildAccelerator/Common/Private/UbaSession.cpp
+++ b/Engine/Source/Programs/UnrealBuildAccelerator/Common/Private/UbaSession.cpp
@@ -1,6 +1,7 @@
 // Copyright Epic Games, Inc. All Rights Reserved.
 
 #include "UbaSession.h"
+#include "UbaBinaryParser.h"
 #include "UbaCompressedFileHeader.h"
 #include "UbaConfig.h"
 #include "UbaFileAccessor.h"
@@ -16,11 +17,9 @@
 #include <span>
 
 #if PLATFORM_WINDOWS
-#include "UbaWinBinDependencyParser.h"
 #include <powerbase.h>
 #pragma comment(lib, "Powrprof.lib")
 #elif PLATFORM_MAC
-#include "UbaMacBinDependencyParser.h"
 #include <mach/vm_statistics.h>
 #include <mach/mach_types.h>
 #include <mach/mach_init.h>
@@ -28,8 +27,6 @@
 #include <mach/mach.h>
 #include <sys/sysctl.h>
 extern char **environ;
-#else // PLATFORM_LINUX
-#include "UbaLinuxBinDependencyParser.h"
 #endif
 
 #define UBA_DEBUG_TRACK_DIR 0 // UBA_DEBUG_LOGGER
@@ -1425,7 +1422,8 @@ namespace uba
 		out.push_back({ library, temp3.ToString(), attr, isSystem });
 
 		StringBuffer<> errorStr;
-		FindImports(temp3, ToView(applicationDir), [&](const tchar* importName, bool isKnown, const char* const* importLoaderPaths)
+		BinaryInfo info;
+		ParseBinary(temp3, ToView(applicationDir), info, [&](const tchar* importName, bool isKnown, const char* const* importLoaderPaths)
 			{
 				if (result && !isKnown)
 					result = CopyImports(out, importName, applicationDir, applicationDirEnd, handledImports, importLoaderPaths);
diff --git a/Engine/Source/Programs/UnrealBuildAccelerator/Core/Public/UbaBinaryParser.h b/Engine/Source/Programs/UnrealBuildAccelerator/Core/Public/UbaBinaryParser.h
new file mode 100644
--- /dev/null
+++ b/Engine/Source/Programs/UnrealBuildAccelerator/Core/Public/UbaBinaryParser.h
@@ -0,0 +1,11 @@
+// Copyright Epic Games, Inc. All Rights Reserved.
+
+#pragma once
+
+#if PLATFORM_WINDOWS
+#include "UbaWinBinDependencyParser.h"
+#elif PLATFORM_LINUX
+#include "UbaLinuxBinDependencyParser.h"
+#elif PLATFORM_MAC
+#include "UbaMacBinDependencyParser.h"
+#endif
diff --git a/Engine/Source/Programs/UnrealBuildAccelerator/Core/Public/UbaLinuxBinDependencyParser.h b/Engine/Source/Programs/UnrealBuildAccelerator/Core/Public/UbaLinuxBinDependencyParser.h
--- a/Engine/Source/Programs/UnrealBuildAccelerator/Core/Public/UbaLinuxBinDependencyParser.h
+++ b/Engine/Source/Programs/UnrealBuildAccelerator/Core/Public/UbaLinuxBinDependencyParser.h
@@ -15,6 +15,10 @@
 
 namespace uba
 {
+	struct BinaryInfo
+	{
+	};
+
 	constexpr StringView g_systemFiles[] =
 	{
 		TCV("libstdc++.so"),
@@ -47,7 +51,7 @@ namespace uba
 		return false;
 	}
 
-	bool FindImports(const StringView& filePath, const StringView& originalPath, const Function<void(const tchar* import, bool isKnown, const char* const* loaderPaths)>& func, StringBufferBase& outError)
+	inline bool ParseBinary(const StringView& filePath, const StringView& originalPath, BinaryInfo& outInfo, const Function<void(const tchar* import, bool isKnown, const char* const* loaderPaths)>& func, StringBufferBase& outError)
 	{
 		#if UBA_FINDIMPORTS_DEBUG
 		printf("FINDIMPORTS: %s\n", filePath.data);
diff --git a/Engine/Source/Programs/UnrealBuildAccelerator/Core/Public/UbaMacBinDependencyParser.h b/Engine/Source/Programs/UnrealBuildAccelerator/Core/Public/UbaMacBinDependencyParser.h
--- a/Engine/Source/Programs/UnrealBuildAccelerator/Core/Public/UbaMacBinDependencyParser.h
+++ b/Engine/Source/Programs/UnrealBuildAccelerator/Core/Public/UbaMacBinDependencyParser.h
@@ -14,12 +14,16 @@
 
 namespace uba
 {
+	struct BinaryInfo
+	{
+	};
+
 	inline bool IsKnownSystemFile(const tchar* fileName)
 	{
 		return false;
 	}
 
-	inline bool FindImports(const StringView& filePath, const StringView& originalPath, const Function<void(const tchar* import, bool isKnown, const char* const* loaderPaths)>& func, StringBufferBase& outError)
+	inline bool ParseBinary(const StringView& filePath, const StringView& originalPath, BinaryInfo& outInfo, const Function<void(const tchar* import, bool isKnown, const char* const* loaderPaths)>& func, StringBufferBase& outError)
 	{
 		int fd = TRUE_WRAPPER(open)(filePath.data, O_RDONLY);
 		if (fd == -1)
diff --git a/Engine/Source/Programs/UnrealBuildAccelerator/Core/Public/UbaWinBinDependencyParser.h b/Engine/Source/Programs/UnrealBuildAccelerator/Core/Public/UbaWinBinDependencyParser.h
--- a/Engine/Source/Programs/UnrealBuildAccelerator/Core/Public/UbaWinBinDependencyParser.h
+++ b/Engine/Source/Programs/UnrealBuildAccelerator/Core/Public/UbaWinBinDependencyParser.h
@@ -15,6 +15,10 @@
 
 namespace uba
 {
+	struct BinaryInfo
+	{
+	};
+
 	// These dlls should _always_ exist on all machines so we can filter them out from the list of imports to copy
 	static const wchar_t* g_knownSystemFiles[] = // Sorted lowercase
 	{
@@ -129,7 +133,7 @@ namespace uba
 		return true;
 	}
 
-	inline bool FindImports(const StringView& filePath, const StringView& originalPath, const Function<void(const tchar* import, bool isKnown, const char* const* importLoaderPaths)>& func, StringBufferBase& outError)
+	inline bool ParseBinary(const StringView& filePath, const StringView& originalPath, BinaryInfo& outInfo, const Function<void(const tchar* import, bool isKnown, const char* const* importLoaderPaths)>& func, StringBufferBase& outError)
 	{
 		if (filePath.EndsWith(TCV(".bat")))
 			return true;
diff --git a/Engine/Source/Programs/UnrealBuildAccelerator/Detours/Private/PosixOS/UbaDetoursFunctionsPosix.cpp b/Engine/Source/Programs/UnrealBuildAccelerator/Detours/Private/PosixOS/UbaDetoursFunctionsPosix.cpp
--- a/Engine/Source/Programs/UnrealBuildAccelerator/Detours/Private/PosixOS/UbaDetoursFunctionsPosix.cpp
+++ b/Engine/Source/Programs/UnrealBuildAccelerator/Detours/Private/PosixOS/UbaDetoursFunctionsPosix.cpp
@@ -167,9 +167,6 @@ using namespace uba;
 		__attribute__((used)) static struct{ const void* replacement; const void* replacee; } _interpose_##_replacee \
 			__attribute__ ((section ("__DATA,__interpose"))) = { (const void*)(unsigned long)&_replacement, (const void*)(unsigned long)&_replacee };
 	#endif // DYLD_INTERPOSE
-
-	#include "UbaMacBinDependencyParser.h"
-
 #else
 	#define UBA_WRAPPER(func) func
 	#define TRUE_WRAPPER(func) True_##func
@@ -179,11 +176,10 @@ using namespace uba;
 		Symbol_##func* True_##func; 
 	DETOURED_FUNCTIONS
 	#undef DETOURED_FUNCTION
-
-	#include "UbaLinuxBinDependencyParser.h"
-
 #endif
 
+#include "UbaBinaryParser.h"
+
 #if 0
 #define	EPERM		 1	/* Operation not permitted */
 #define	ENOENT		 2	/* No such file or directory */
@@ -2229,7 +2225,8 @@ void Shared_LoadLibrary(const char*& path, const char* const* loaderPaths, Strin
 	Rpc_GetFullFileName2(path, tempBuf, virtualPath, loaderPaths);
 	StringView originalPath = StringView(virtualPath).GetPath();
 	StringBuffer<> error;
-	FindImports(tempBuf, originalPath, [&](const tchar* import, bool isKnown, const char* const* importLoaderPaths)
+	BinaryInfo info;
+	ParseBinary(tempBuf, originalPath, info, [&](const tchar* import, bool isKnown, const char* const* importLoaderPaths)
 	{
 		if (!g_handledLibraries.insert(import).second)
 			return;
diff --git a/Engine/Source/Programs/UnrealBuildAccelerator/Detours/Private/Windows/UbaDetoursFunctionsKernelBase.inl b/Engine/Source/Programs/UnrealBuildAccelerator/Detours/Private/Windows/UbaDetoursFunctionsKernelBase.inl
--- a/Engine/Source/Programs/UnrealBuildAccelerator/Detours/Private/Windows/UbaDetoursFunctionsKernelBase.inl
+++ b/Engine/Source/Programs/UnrealBuildAccelerator/Detours/Private/Windows/UbaDetoursFunctionsKernelBase.inl
@@ -2713,7 +2713,8 @@ HMODULE Recursive_LoadLibraryExW(const StringView& filePath, LPCWSTR originalNam
 	{
 		SuppressCreateFileDetourScope cfs;
 		StringBuffer<256> error;
-		if (!FindImports(filePath, {}, [&](const wchar_t* import, bool isKnown, const char* const* importLoaderPaths)
+		BinaryInfo info;
+		if (!ParseBinary(filePath, {}, info, [&](const wchar_t* import, bool isKnown, const char* const* importLoaderPaths)
 			{
 				if (!GetModuleHandleW(import))
 					importedModules.emplace_back(import, isKnown);
diff --git a/Engine/Source/Programs/UnrealBuildAccelerator/Detours/Private/Windows/UbaDetoursFunctionsWin.cpp b/Engine/Source/Programs/UnrealBuildAccelerator/Detours/Private/Windows/UbaDetoursFunctionsWin.cpp
--- a/Engine/Source/Programs/UnrealBuildAccelerator/Detours/Private/Windows/UbaDetoursFunctionsWin.cpp
+++ b/Engine/Source/Programs/UnrealBuildAccelerator/Detours/Private/Windows/UbaDetoursFunctionsWin.cpp
@@ -38,10 +38,10 @@ DETOURED_FUNCTIONS_MEMORY
 DETOURED_FUNCTIONS_MEMORY_NON_WINE
 #undef DETOURED_FUNCTION
 
+#include "UbaBinaryParser.h"
 #include "UbaDirectoryTable.h"
 #include "UbaProcessStats.h"
 #include "UbaProtocol.h"
-#include "UbaWinBinDependencyParser.h"
 #include "UbaDetoursPayload.h"
 #include "UbaApplicationRules.h"
 #include "UbaDetoursShared.h"
diff --git a/Engine/Source/Programs/UnrealBuildAccelerator/Host/Private/UbaExports.cpp b/Engine/Source/Programs/UnrealBuildAccelerator/Host/Private/UbaExports.cpp
--- a/Engine/Source/Programs/UnrealBuildAccelerator/Host/Private/UbaExports.cpp
+++ b/Engine/Source/Programs/UnrealBuildAccelerator/Host/Private/UbaExports.cpp
@@ -2,6 +2,7 @@
 
 #include "UbaExports.h"
 #include "UbaAWS.h"
+#include "UbaBinaryParser.h"
 #include "UbaCacheClient.h"
 #include "UbaConfig.h"
 #include "UbaCoordinatorWrapper.h"
@@ -14,10 +15,6 @@
 #include "UbaStorageServer.h"
 #include "UbaSessionServer.h"
 
-#if PLATFORM_WINDOWS
-#include "UbaWinBinDependencyParser.h"
-#endif
-
 namespace uba
 {
 	CallbackLogWriter::CallbackLogWriter(BeginScopeCallback begin, EndScopeCallback end, LogCallback log) : m_beginScope(begin), m_endScope(end), m_logCallback(log)
@@ -903,7 +900,8 @@ extern "C"
 	{
 #if PLATFORM_WINDOWS
 		uba::StringBuffer<> errors;
-		uba::FindImports(uba::ToView(binary), {}, [&](const uba::tchar* importName, bool isKnown, const char* const* importLoaderPaths) { func(importName, userData); }, errors);
+		uba::BinaryInfo info;
+		uba::ParseBinary(uba::ToView(binary), {}, info, [&](const uba::tchar* importName, bool isKnown, const char* const* importLoaderPaths) { func(importName, userData); }, errors);
 #else
 		UBA_ASSERT(false);
 #endif
diff --git a/Engine/Source/Programs/UnrealBuildAccelerator/Test/Private/UbaTestBasics.cpp b/Engine/Source/Programs/UnrealBuildAccelerator/Test/Private/UbaTestBasics.cpp
--- a/Engine/Source/Programs/UnrealBuildAccelerator/Test/Private/UbaTestBasics.cpp
+++ b/Engine/Source/Programs/UnrealBuildAccelerator/Test/Private/UbaTestBasics.cpp
@@ -1,5 +1,6 @@
 // Copyright Epic Games, Inc. All Rights Reserved.
 
+#include "UbaBinaryParser.h"
 #include "UbaBinaryReaderWriter.h"
 #include "UbaFileAccessor.h"
 #include "UbaPathUtils.h"
@@ -14,12 +15,6 @@
 #include "UbaTimer.h"
 #include "UbaDirectoryIterator.h"
 
-#if PLATFORM_WINDOWS
-#include "UbaWinBinDependencyParser.h"
-#elif PLATFORM_MAC
-#include "UbaMacBinDependencyParser.h"
-#endif
-
 #define VA_ARGS(...) , ##__VA_ARGS__
 #define UBA_TEST_CHECK(expr, fmt, ...) if (!(expr)) return logger.Error(TC(fmt) VA_ARGS(__VA_ARGS__));
 
@@ -533,7 +528,8 @@ namespace uba
 		path.EnsureEndsWithSlash().Append(TCV("UbaTestApp.exe"));
 		bool importKernel = false;
 		StringBuffer<> error;
-		FindImports(path, {}, [&](const tchar* import, bool isKnown, const char* const* importLoaderPaths)
+		BinaryInfo info;
+		ParseBinary(path, {}, info, [&](const tchar* import, bool isKnown, const char* const* importLoaderPaths)
 		{
 			importKernel |= isKnown && Contains(import, TC("KERNEL32.dll"));
 		}, error);
