[UBA]
* Added min os version to binary parser on macos

diff --git a/Engine/Source/Programs/UnrealBuildAccelerator/Cli/Private/UbaCli.cpp b/Engine/Source/Programs/UnrealBuildAccelerator/Cli/Private/UbaCli.cpp
--- a/Engine/Source/Programs/UnrealBuildAccelerator/Cli/Private/UbaCli.cpp
+++ b/Engine/Source/Programs/UnrealBuildAccelerator/Cli/Private/UbaCli.cpp
@@ -834,6 +834,9 @@ namespace uba
 
 				}, error))
 					return logger.Error(TC("%s"), error.data);
+				#if PLATFORM_MAC
+				logger.Info(TC("MinOsVersion: %u.%u.%u"), (info.minVersion >> 16) & 0xffff, (info.minVersion >> 8) & 0xff, info.minVersion & 0xff);
+				#endif
 				return true;
 			}
 		}
diff --git a/Engine/Source/Programs/UnrealBuildAccelerator/Core/Public/UbaMacBinDependencyParser.h b/Engine/Source/Programs/UnrealBuildAccelerator/Core/Public/UbaMacBinDependencyParser.h
--- a/Engine/Source/Programs/UnrealBuildAccelerator/Core/Public/UbaMacBinDependencyParser.h
+++ b/Engine/Source/Programs/UnrealBuildAccelerator/Core/Public/UbaMacBinDependencyParser.h
@@ -16,6 +16,7 @@ namespace uba
 {
 	struct BinaryInfo
 	{
+		u32 minVersion = 0;
 	};
 
 	inline bool IsKnownSystemFile(const tchar* fileName)
@@ -69,6 +70,21 @@ namespace uba
 			{
 				auto& lc = *(load_command*)memIt;
 				auto g = MakeGuard([&]() { memIt += lc.cmdsize; });
+
+				if (lc.cmd ==0x24) // LC_VERSION_MIN_MAXOSX
+				{
+					struct version_min_command { u32 cmd; u32 cmdsize; u32 version; u32 sdk; };
+					auto& version = *(version_min_command*)memIt;
+					outInfo.minVersion = version.version;
+					continue;
+				}
+				if (lc.cmd == LC_BUILD_VERSION)
+				{
+					auto& version = *(build_version_command*)memIt;
+					outInfo.minVersion = version.minos;
+					continue;
+				}
+
 				if (lc.cmd != LC_LOAD_DYLIB && lc.cmd != LC_RPATH) //  && lc.cmd != LC_LOAD_WEAK_DYLIB
 					continue;
 				auto& cmdData = *(dylib_command*)memIt;
