[UBA]
* Moved os version logic to UbaPlatform.h
* Added so minOsVersion is included in binary modules transferred to remote and tested against local os verion

diff --git a/Engine/Source/Programs/UnrealBuildAccelerator/Agent/Private/UbaAgent.cpp b/Engine/Source/Programs/UnrealBuildAccelerator/Agent/Private/UbaAgent.cpp
--- a/Engine/Source/Programs/UnrealBuildAccelerator/Agent/Private/UbaAgent.cpp
+++ b/Engine/Source/Programs/UnrealBuildAccelerator/Agent/Private/UbaAgent.cpp
@@ -777,28 +777,10 @@ namespace uba
 		if (!zone.count)
 			GetZone(zone);
 
-		#if PLATFORM_WINDOWS
-		if (IsRunningWine())
-			extraInfo.Append(TCV(", Linux/Wine-")).Append(GetWineVersion());
-		#elif PLATFORM_LINUX
-		if (FILE* osRelease = fopen("/etc/os-release", "r"))
-		{
-			char buffer[256];
-			if (fgets(buffer, sizeof_array(buffer), osRelease))
-				if (StartsWith(buffer, "PRETTY_NAME=\""))
-					extraInfo.Append(TCV(", ")).Append(buffer + 13, TStrlen(buffer) - 14);
-			fclose(osRelease);
-		}
-		#elif PLATFORM_MAC
-		if (FILE* osRelease = popen("sw_vers -productVersion", "r"))
-		{
-			char buffer[256];
-			if (fgets(buffer, sizeof_array(buffer), osRelease))
-				extraInfo.Append(TCV(", MacOS-")).Append(buffer, TStrlen(buffer) - 1);
-			pclose(osRelease);
-		}
-		#endif
-
+		u32 osVersion;
+		StringBuffer<32> osVersionStr;
+		if (GetOsVersion(osVersionStr, osVersion))
+			extraInfo.Append(TCV(", ")).Append(osVersionStr);
 		if (useQuic)
 			extraInfo.Append(TCV(", MsQuic"));
 		if (hasCrypto)
@@ -1232,6 +1214,7 @@ namespace uba
 			info.memWaitLoadPercent = u8(memWaitLoadPercent);
 			info.memKillLoadPercent = u8(memKillLoadPercent);
 			info.useDependencyCrawler = useCrawler;
+			info.osVersion = osVersion;
 
 			if (!quiet)
 				info.processFinished = processFinished;
@@ -1335,7 +1318,7 @@ namespace uba
 				u32 sessionWaitTimeout = 5*1000;
 				if (useListen)
 				{
-					if (connectionCount != targetConnectionCount)
+					if (connectionCount < targetConnectionCount)
 					{
 						u32 newConnectionCount = targetConnectionCount;
 						logger.Info(TC("Updating desired connection count from %u to %u"), connectionCount, newConnectionCount);
diff --git a/Engine/Source/Programs/UnrealBuildAccelerator/Common/Private/UbaSession.cpp b/Engine/Source/Programs/UnrealBuildAccelerator/Common/Private/UbaSession.cpp
--- a/Engine/Source/Programs/UnrealBuildAccelerator/Common/Private/UbaSession.cpp
+++ b/Engine/Source/Programs/UnrealBuildAccelerator/Common/Private/UbaSession.cpp
@@ -1419,7 +1419,7 @@ namespace uba
 		if (isSystem && IsKnownSystemFile(applicationName))
 			return true;
 
-		out.push_back({ library, temp3.ToString(), attr, isSystem });
+		BinaryModule& binaryModule = out.emplace_back(library, temp3.ToString(), attr, isSystem);(void)binaryModule;
 
 		StringBuffer<> errorStr;
 		BinaryInfo info;
@@ -1431,6 +1431,10 @@ namespace uba
 		if (errorStr.count)
 			return m_logger.Error(errorStr.data);
 
+		#if PLATFORM_MAC
+		binaryModule.minOsVersion = info.minVersion;
+		#endif
+
 		// This code is needed if application is compiled with tsan
 		//strcpy(applicationDirEnd, "libclang_rt.tsan.so");
 		//out.push_back({ "libclang_rt.tsan.so", applicationDir, S_IRUSR | S_IWUSR });
diff --git a/Engine/Source/Programs/UnrealBuildAccelerator/Common/Private/UbaSessionClient.cpp b/Engine/Source/Programs/UnrealBuildAccelerator/Common/Private/UbaSessionClient.cpp
--- a/Engine/Source/Programs/UnrealBuildAccelerator/Common/Private/UbaSessionClient.cpp
+++ b/Engine/Source/Programs/UnrealBuildAccelerator/Common/Private/UbaSessionClient.cpp
@@ -53,6 +53,7 @@ namespace uba
 		m_downloadDetoursLib = info.downloadDetoursLib;
 		m_defaultPriorityClass = info.defaultPriorityClass;
 		m_maxIdleSeconds = info.maxIdleSeconds;
+		m_osVersion = info.osVersion;
 		m_disableCustomAllocator = info.disableCustomAllocator;
 		m_useBinariesAsVersion = info.useBinariesAsVersion;
 		m_memWaitLoadPercent = info.memWaitLoadPercent;
@@ -434,6 +435,12 @@ namespace uba
 			if (casKey == CasKeyZero)
 				return m_logger.Error(TC("Bad CasKey for %s (%s)"), moduleFile.data, CasKeyString(casKey).str);
 
+			#if PLATFORM_MAC
+			u32 minOsVersion = reader.ReadU32();
+			if (m_osVersion < minOsVersion)
+				return m_logger.Error(TC("%s has min os version %u but current os is %u"), moduleFile.data, minOsVersion, m_osVersion);
+			#endif
+
 			if (isSystem)
 			{
 				StringBuffer<> localSystemModule;
diff --git a/Engine/Source/Programs/UnrealBuildAccelerator/Common/Private/UbaSessionServer.cpp b/Engine/Source/Programs/UnrealBuildAccelerator/Common/Private/UbaSessionServer.cpp
--- a/Engine/Source/Programs/UnrealBuildAccelerator/Common/Private/UbaSessionServer.cpp
+++ b/Engine/Source/Programs/UnrealBuildAccelerator/Common/Private/UbaSessionServer.cpp
@@ -968,6 +968,9 @@ namespace uba
 			writer.WriteU32(m.fileAttributes);
 			writer.WriteBool(m.isSystem);
 			writer.WriteCasKey(casKey);
+			#if PLATFORM_MAC
+			writer.WriteU32(m.minOsVersion);
+			#endif
 		}
 
 		u64 bytesSize = writer.GetPosition() - startPos;
diff --git a/Engine/Source/Programs/UnrealBuildAccelerator/Common/Public/UbaSession.h b/Engine/Source/Programs/UnrealBuildAccelerator/Common/Public/UbaSession.h
--- a/Engine/Source/Programs/UnrealBuildAccelerator/Common/Public/UbaSession.h
+++ b/Engine/Source/Programs/UnrealBuildAccelerator/Common/Public/UbaSession.h
@@ -123,7 +123,7 @@ namespace uba
 
 		SessionStats& Stats();
 
-		struct BinaryModule { TString name; TString path; u32 fileAttributes = 0; bool isSystem = false; };
+		struct BinaryModule { TString name; TString path; u32 fileAttributes = 0; bool isSystem = false; u32 minOsVersion = 0; };
 		bool GetBinaryModules(Vector<BinaryModule>& out, const tchar* application);
 		void Free(Vector<BinaryModule>& v);
 		bool IsRarelyRead(ProcessImpl& process, const StringView& fileName) const;
diff --git a/Engine/Source/Programs/UnrealBuildAccelerator/Common/Public/UbaSessionClient.h b/Engine/Source/Programs/UnrealBuildAccelerator/Common/Public/UbaSessionClient.h
--- a/Engine/Source/Programs/UnrealBuildAccelerator/Common/Public/UbaSessionClient.h
+++ b/Engine/Source/Programs/UnrealBuildAccelerator/Common/Public/UbaSessionClient.h
@@ -21,6 +21,7 @@ namespace uba
 		u32 defaultPriorityClass = 0x00004000; // BELOW_NORMAL_PRIORITY_CLASS;
 		u32 maxIdleSeconds = ~0u;
 		u32 pingTimeoutSecondsPrintCallstacks = 0;
+		u32 osVersion = 0;
 		u8 memWaitLoadPercent = 80; // When memory usage goes above this percent, no new processes will be spawned until back below
 		u8 memKillLoadPercent = 90; // When memory usage goes above this percent, newest processes will be killed to bring it back below
 		bool dedicated = false;  // If true, server will not disconnect client when starting to run out of work.
@@ -119,6 +120,7 @@ namespace uba
 		u32 m_uiLanguage = 0;
 		u32 m_defaultPriorityClass;
 		u32 m_maxIdleSeconds = ~0u;
+		u32 m_osVersion = 0;
 		u32 m_killRandomIndex = ~0u;
 		u32 m_killRandomCounter = 0;
 		u8 m_memWaitLoadPercent;
diff --git a/Engine/Source/Programs/UnrealBuildAccelerator/Core/Private/UbaPlatform.cpp b/Engine/Source/Programs/UnrealBuildAccelerator/Core/Private/UbaPlatform.cpp
--- a/Engine/Source/Programs/UnrealBuildAccelerator/Core/Private/UbaPlatform.cpp
+++ b/Engine/Source/Programs/UnrealBuildAccelerator/Core/Private/UbaPlatform.cpp
@@ -165,6 +165,65 @@ namespace uba
 		return true;
 	}
 
+	bool GetOsVersion(StringBufferBase& outPretty, u32& outValue)
+	{
+		outValue = 0;
+
+#if PLATFORM_WINDOWS
+		if (IsRunningWine())
+			outPretty.Append(TC("Linux/Wine-")).Append(GetWineVersion());
+		else
+			outPretty.Append(TC("Windows"));
+#elif PLATFORM_MACOS
+		FILE* osRelease = popen("sw_vers -productVersion", "r");
+		if (!osRelease)
+			return false;
+		char buffer[256];
+		bool success = fgets(buffer, sizeof_array(buffer), osRelease);
+		pclose(osRelease);
+		if (!success)
+			return false;
+		u32 bufferLen = TStrlen(buffer) - 1;
+		buffer[bufferLen] = 0;
+		outPretty.Append(TCV("MacOS-")).Append(buffer, bufferLen);
+
+		u32 majorVersion = 0;
+		u32 minorVersion = 0;
+		u32 patchVersion = 0;
+
+		if (char* endOfMajor = TStrchr(buffer, '.'))
+		{
+			*endOfMajor = 0;
+			char* minor = endOfMajor + 1;
+			if (char* endOfMinor = TStrchr(minor, '.'))
+			{
+				*endOfMinor = 0;
+				char* patch = endOfMinor + 1;
+				if (*patch)
+					patchVersion = strtoul(patch, nullptr, 10);
+			}
+			if (*minor)
+				minorVersion = strtoul(minor, nullptr, 10);
+		}
+		majorVersion = strtoul(buffer, nullptr, 10);
+
+		outValue = (majorVersion << 16) + (minorVersion << 8) + patchVersion;
+#elif PLATFORM_LINUX
+		FILE* osRelease = fopen("/etc/os-release", "r");
+		if (!osRelease)
+			return false;
+		char buffer[256];
+		bool success = fgets(buffer, sizeof_array(buffer), osRelease);
+		fclose(osRelease);
+		if (!success)
+			return false;
+		if (!StartsWith(buffer, "PRETTY_NAME=\""))
+			return false;
+		outPretty.Append(buffer + 13, TStrlen(buffer) - 14);
+#endif
+		return true;
+	}
+
 	void WriteAssertInfo(StringBufferBase& out, const tchar* text, const char* file, u32 line, const char* expr, void* context)
 	{
 		int signal = 0;
diff --git a/Engine/Source/Programs/UnrealBuildAccelerator/Core/Public/UbaPlatform.h b/Engine/Source/Programs/UnrealBuildAccelerator/Core/Public/UbaPlatform.h
--- a/Engine/Source/Programs/UnrealBuildAccelerator/Core/Public/UbaPlatform.h
+++ b/Engine/Source/Programs/UnrealBuildAccelerator/Core/Public/UbaPlatform.h
@@ -98,6 +98,7 @@ namespace uba
 	u32 GetLastError();
 	void SetLastError(u32 error);
 	bool GetComputerNameW(StringBufferBase& out);
+	bool GetOsVersion(StringBufferBase& outPretty, u32& outValue);
 	u32 GetCurrentProcessId();
 
 	enum FileHandle : u64 {};
