[UBA]
* Fixed so visualizer does not draw graph lines between reconnects

diff --git a/Engine/Source/Programs/UnrealBuildAccelerator/Visualizer/Private/UbaTraceReader.cpp b/Engine/Source/Programs/UnrealBuildAccelerator/Visualizer/Private/UbaTraceReader.cpp
--- a/Engine/Source/Programs/UnrealBuildAccelerator/Visualizer/Private/UbaTraceReader.cpp
+++ b/Engine/Source/Programs/UnrealBuildAccelerator/Visualizer/Private/UbaTraceReader.cpp
@@ -436,6 +436,15 @@ namespace uba
 					continue;
 				if (oldSession.disconnectTime == ~u64(0))
 					break;
+				u32 updateCount = u32(oldSession.updates.size());
+				oldSession.networkSend.resize(updateCount);
+				oldSession.networkRecv.resize(updateCount);
+				oldSession.ping.resize(updateCount);
+				oldSession.memAvail.resize(updateCount);
+				oldSession.cpuLoad.resize(updateCount);
+				oldSession.connectionCount.resize(updateCount);
+
+				oldSession.reconnectIndices.push_back(updateCount);
 				oldSession.isReset = true;
 				oldSession.disconnectTime = ~u64(0);
 				oldSession.proxyName.clear();
diff --git a/Engine/Source/Programs/UnrealBuildAccelerator/Visualizer/Private/UbaTraceReader.h b/Engine/Source/Programs/UnrealBuildAccelerator/Visualizer/Private/UbaTraceReader.h
--- a/Engine/Source/Programs/UnrealBuildAccelerator/Visualizer/Private/UbaTraceReader.h
+++ b/Engine/Source/Programs/UnrealBuildAccelerator/Visualizer/Private/UbaTraceReader.h
@@ -109,6 +109,7 @@ namespace uba
 			Vector<u64> memAvail;
 			Vector<float> cpuLoad;
 			Vector<u16> connectionCount;
+			Vector<u32> reconnectIndices;
 
 			Vector<TString> summary;
 			UnorderedMap<CasKey, u32> fetchedFilesActive;
diff --git a/Engine/Source/Programs/UnrealBuildAccelerator/Visualizer/Private/UbaVisualizer.cpp b/Engine/Source/Programs/UnrealBuildAccelerator/Visualizer/Private/UbaVisualizer.cpp
--- a/Engine/Source/Programs/UnrealBuildAccelerator/Visualizer/Private/UbaVisualizer.cpp
+++ b/Engine/Source/Programs/UnrealBuildAccelerator/Visualizer/Private/UbaVisualizer.cpp
@@ -1372,68 +1372,88 @@ namespace uba
 
 					auto drawGraph = [&](auto& values, auto maxValue, double scale, HPEN pen, bool accumulatingValue, int offsetY = 0)
 						{
-							SelectObject(hdc, pen);
-							bool isFirstUpdate = true;
-							bool isFirstDraw = true;
-							decltype(maxValue) prevValue = 0;
-							u64 prevTime = 0;
-							int prevX = 0;
-							int prevY = 0;
-							u64 e = Min(values.size(), session.updates.size());
-							line.clear();
-							for (u64 i=0; i!=e; ++i)
+							bool loop = true;
+							u32 reconnectIndex = 0;
+							while (loop)
 							{
-								u64 updateTime = session.updates[i];
-								auto value = values[i];
-								int x = posX + int(float(TimeToS(updateTime)) * scaleX);
-								int y = graphBaseY;
-
-								double duration = TimeToS(updateTime - prevTime);
-								if (updateTime == 0)
-									isFirstUpdate = true;
-
-								if (accumulatingValue)
-								{
-									double invScaleY = duration * scale;
-									if (invScaleY != 0 && prevValue)
-										y -= int(double(value - prevValue) / invScaleY) + offsetY;
-								}
+								SelectObject(hdc, pen);
+								bool isFirstUpdate = true;
+								bool isFirstDraw = true;
+								decltype(maxValue) prevValue = 0;
+								u64 prevTime = 0;
+								int prevX = 0;
+								int prevY = 0;
+
+								u64 i = 0;
+								u64 e = 0;
+
+								if (reconnectIndex)
+									i = session.reconnectIndices[reconnectIndex - 1];
+
+								if (reconnectIndex < session.reconnectIndices.size())
+									e = session.reconnectIndices[reconnectIndex];
 								else
 								{
-									double div = 1.0f;
-									if (maxValue)
-										div = double(maxValue);
-									y -= int(double(maxValue - value*scale)*graphHeight/div) + offsetY;
+									e = Min(session.updates.size(), values.size());
+									loop = false;
 								}
 
-								if (prevValue <= 0)
-									prevY = y;
-
-								if (x > clientRect.left)
+								line.clear();
+								for (; i!=e; ++i)
 								{
-									if (isFirstUpdate)
+									u64 updateTime = session.updates[i];
+									auto value = values[i];
+									int x = posX + int(float(TimeToS(updateTime)) * scaleX);
+									int y = graphBaseY;
+
+									double duration = TimeToS(updateTime - prevTime);
+									if (updateTime == 0)
+										isFirstUpdate = true;
+
+									if (accumulatingValue)
 									{
-										line.push_back({x, y}); //MoveToEx(hdc, x, y, NULL);
-										isFirstDraw = false;
+										double invScaleY = duration * scale;
+										if (invScaleY != 0 && prevValue)
+											y -= int(double(value - prevValue) / invScaleY) + offsetY;
 									}
 									else
 									{
-										if (isFirstDraw)
-											line.push_back({prevX, prevY}); // MoveToEx(hdc, prevX, prevY, NULL);
-										line.push_back({x, y}); // LineTo(hdc, x, y);
-										isFirstDraw = false;
+										double div = 1.0f;
+										if (maxValue)
+											div = double(maxValue);
+										y -= int(double(maxValue - value*scale)*graphHeight/div) + offsetY;
 									}
+
+									if (prevValue <= 0)
+										prevY = y;
+
+									if (x > clientRect.left)
+									{
+										if (isFirstUpdate)
+										{
+											line.push_back({x, y}); //MoveToEx(hdc, x, y, NULL);
+											isFirstDraw = false;
+										}
+										else
+										{
+											if (isFirstDraw)
+												line.push_back({prevX, prevY}); // MoveToEx(hdc, prevX, prevY, NULL);
+											line.push_back({x, y}); // LineTo(hdc, x, y);
+											isFirstDraw = false;
+										}
+									}
+									if (x > clientRect.right)
+										break;
+									isFirstUpdate = false;
+									prevX = x;
+									prevY = y;
+									prevValue = value;
+									prevTime = updateTime;
 								}
-								if (x > clientRect.right)
-									break;
-								isFirstUpdate = false;
-								prevX = x;
-								prevY = y;
-								prevValue = value;
-								prevTime = updateTime;
+								if (line.size() > 1)
+									Polyline(hdc, line.data(), int(line.size()));
+								++reconnectIndex;
 							}
-							if (line.size() > 1)
-								Polyline(hdc, line.data(), int(line.size()));
 						};
 
 					if (m_config.showNetworkStats && session.highestSendPerS != 0 && session.highestRecvPerS != 0)
@@ -3036,52 +3056,73 @@ namespace uba
 				if (pos.y >= posY && pos.y < posY + GraphHeight)
 				{
 					int posX = int(m_scrollPosX) + progressRect.left;
-					int prevX = 100000;
-					for (u64 i=0, e=session.updates.size(); i!=e; ++i)
+
+					bool loop = true;
+					u32 reconnectIndex = 0;
+					while (loop)
 					{
-						u64 updateTime = session.updates[i];
-						auto updateSend = session.networkSend[i];
-						auto updateRecv = session.networkRecv[i];
-						int x = posX + int(float(TimeToS(updateTime)) * scaleX);
+						u64 i = 0;
+						u64 e = 0;
+
+						if (reconnectIndex)
+							i = session.reconnectIndices[reconnectIndex - 1];
 
-						int hitOffset = (prevX - x)/2;
-						if (pos.x + hitOffset >= prevX && pos.x + hitOffset <= x)
+						if (reconnectIndex < session.reconnectIndices.size())
+							e = session.reconnectIndices[reconnectIndex];
+						else
 						{
-							u64 prevTime = 0;
-							u64 prevSend = 0;
-							u64 prevRecv = 0;
-							if (i > 0)
-							{
-								prevTime = session.updates[i - 1];
-								prevSend = Min(session.networkSend[i - 1], updateSend);
-								prevRecv = Min(session.networkRecv[i - 1], updateRecv);
-							}
+							e = session.updates.size();
+							loop = false;
+						}
+
+						int prevX = 100000;
+						for (; i!=e; ++i)
+						{
+							u64 updateTime = session.updates[i];
+							auto updateSend = session.networkSend[i];
+							auto updateRecv = session.networkRecv[i];
+							int x = posX + int(float(TimeToS(updateTime)) * scaleX);
 
-							double duration = TimeToS(updateTime - prevTime);
-							outResult.stats.recvBytes = updateRecv;
-							outResult.stats.sendBytes = updateSend;
-							outResult.stats.recvBytesPerSecond = u64(float(updateRecv - prevRecv) / duration);
-							outResult.stats.sendBytesPerSecond = u64(float(updateSend - prevSend) / duration);
-							outResult.stats.ping = session.ping[i];
-							outResult.stats.memAvail = session.memAvail[i];
-							outResult.stats.cpuLoad = session.cpuLoad[i];
-							outResult.stats.memTotal = session.memTotal;
-							outResult.statsSelected = true;
-
-							for (auto& pair : session.drives)
+							int hitOffset = (prevX - x)/2;
+							if (pos.x + hitOffset >= prevX && pos.x + hitOffset <= x)
 							{
-								auto& updateBusy = pair.second.busyPercent[i];
-								u64 updateRead = pair.second.readBytes[i];
-								u64 updateWrite = pair.second.writeBytes[i];
-								auto& statsDrive = outResult.stats.drives[pair.first];
-								statsDrive.busyPercent = updateBusy;
-								statsDrive.readPerSecond = u64(float(updateRead) / duration);
-								statsDrive.writePerSecond = u64(float(updateWrite) / duration);
+								u64 prevTime = 0;
+								u64 prevSend = 0;
+								u64 prevRecv = 0;
+								if (i > 0)
+								{
+									prevTime = session.updates[i - 1];
+									prevSend = Min(session.networkSend[i - 1], updateSend);
+									prevRecv = Min(session.networkRecv[i - 1], updateRecv);
+								}
+
+								double duration = TimeToS(updateTime - prevTime);
+								outResult.stats.recvBytes = updateRecv;
+								outResult.stats.sendBytes = updateSend;
+								outResult.stats.recvBytesPerSecond = u64(float(updateRecv - prevRecv) / duration);
+								outResult.stats.sendBytesPerSecond = u64(float(updateSend - prevSend) / duration);
+								outResult.stats.ping = session.ping[i];
+								outResult.stats.memAvail = session.memAvail[i];
+								outResult.stats.cpuLoad = session.cpuLoad[i];
+								outResult.stats.memTotal = session.memTotal;
+								outResult.statsSelected = true;
+
+								for (auto& pair : session.drives)
+								{
+									auto& updateBusy = pair.second.busyPercent[i];
+									u64 updateRead = pair.second.readBytes[i];
+									u64 updateWrite = pair.second.writeBytes[i];
+									auto& statsDrive = outResult.stats.drives[pair.first];
+									statsDrive.busyPercent = updateBusy;
+									statsDrive.readPerSecond = u64(float(updateRead) / duration);
+									statsDrive.writePerSecond = u64(float(updateWrite) / duration);
+								}
+								return;
 							}
-							return;
-						}
 
-						prevX = x;
+							prevX = x;
+						}
+						++reconnectIndex;
 					}
 					posY += GraphHeight;
 				}
