[UBA]
* Allow os version to not be set on agent side

diff --git a/Engine/Source/Programs/UnrealBuildAccelerator/Common/Private/UbaSessionClient.cpp b/Engine/Source/Programs/UnrealBuildAccelerator/Common/Private/UbaSessionClient.cpp
--- a/Engine/Source/Programs/UnrealBuildAccelerator/Common/Private/UbaSessionClient.cpp
+++ b/Engine/Source/Programs/UnrealBuildAccelerator/Common/Private/UbaSessionClient.cpp
@@ -437,7 +437,7 @@ namespace uba
 
 			#if PLATFORM_MAC
 			u32 minOsVersion = reader.ReadU32();
-			if (m_osVersion < minOsVersion)
+			if (m_osVersion && m_osVersion < minOsVersion)
 				return m_logger.Error(TC("%s has min os version %u but current os is %u"), moduleFile.data, minOsVersion, m_osVersion);
 			#endif
 
