[UBA]
* Implemented keep alive on application level but as low level as possible
* Enabled application keep alive on macos since built-in tcp keep alive is 10+ minutes which we don't want

diff --git a/Engine/Source/Programs/UnrealBuildAccelerator/Common/Private/UbaNetworkBackendTcp.cpp b/Engine/Source/Programs/UnrealBuildAccelerator/Common/Private/UbaNetworkBackendTcp.cpp
--- a/Engine/Source/Programs/UnrealBuildAccelerator/Common/Private/UbaNetworkBackendTcp.cpp
+++ b/Engine/Source/Programs/UnrealBuildAccelerator/Common/Private/UbaNetworkBackendTcp.cpp
@@ -826,8 +826,8 @@ namespace uba
 	bool SetKeepAlive(Logger& logger, SOCKET socket) // This will make sure that WSAPoll exits when the network cable is pulled
 	{
 		
-		int keepAliveTime = 60; // 1 minute of idle before it starts sending keep alive
-		int keepAliveTimeInterval = 1; // 1 second between each check (10 failures mean disconnect)
+		int keepAliveTime = KeepAliveIdleSeconds; // Time of idle before it starts sending keep alive
+		int keepAliveTimeInterval = KeepAliveIntervalSeconds; // Time between each check (10 failures mean disconnect)
 
 
 		u32 value = 1;
@@ -847,7 +847,7 @@ namespace uba
 			return logger.Error(TC("setsockopt TCP_KEEPIDLE (error: %s)"), LastErrorToText(WSAGetLastError()).data);
 		if (setsockopt(socket, IPPROTO_TCP, TCP_KEEPINTVL, &keepAliveTimeInterval, sizeof(int)) < 0)
 			return logger.Error(TC("setsockopt TCP_KEEPINTVL (error: %s)"), LastErrorToText(WSAGetLastError()).data);
-		int keepAliveProbes = 10;
+		int keepAliveProbes = KeepAliveProbeCount; // Number of tests before timing out
 		if (setsockopt(socket, IPPROTO_TCP, TCP_KEEPCNT, &keepAliveProbes, sizeof(int)) < 0)
 			return logger.Error(TC("setsockopt TCP_KEEPCNT (error: %s)"), LastErrorToText(WSAGetLastError()).data);
 #else // PLATFORM_MAC
@@ -1004,9 +1004,12 @@ namespace uba
 				int res = WSAPoll(&p, 1, int(timeoutMs));
 				if (!res)
 				{
-					if (connection.recvTimeoutCallback && connection.recvTimeoutCallback(connection.recvTimeoutContext))
-						continue;
-
+					if (connection.recvTimeoutCallback)
+					{
+						if (connection.recvTimeoutCallback(connection.recvTimeoutContext, timeoutMs, connection.recvHint, hint))
+							continue;
+						return false;
+					}
 					connection.logger.Info(TC("RecvSocket - WSAPoll returned timeout for socket %i and connection %s after %s (%s%s)"), socket, GuidToString(connection.uid).str, TimeToText(MsToTime(timeoutMs)).str, connection.recvHint, hint);
 					return false;
 				}
diff --git a/Engine/Source/Programs/UnrealBuildAccelerator/Common/Private/UbaNetworkClient.cpp b/Engine/Source/Programs/UnrealBuildAccelerator/Common/Private/UbaNetworkClient.cpp
--- a/Engine/Source/Programs/UnrealBuildAccelerator/Common/Private/UbaNetworkClient.cpp
+++ b/Engine/Source/Programs/UnrealBuildAccelerator/Common/Private/UbaNetworkClient.cpp
@@ -87,7 +87,7 @@ namespace uba
 			Atomic<u8> error;
 		};
 
-		backend.SetRecvTimeout(backendConnection, m_receiveTimeoutSeconds*1000, nullptr, [](void*) { return false; });
+		backend.SetRecvTimeout(backendConnection, m_receiveTimeoutSeconds*1000, nullptr, nullptr);
 
 		RecvContext rc(*this, backend, backendConnection);
 
@@ -226,7 +226,7 @@ namespace uba
 		m_logger.Detail(TC("Connected to server... (0x%p)"), backendConnection);
 		lock.Leave();
 
-		backend.SetRecvTimeout(backendConnection, m_receiveTimeoutSeconds*1000, nullptr, [](void*) { return false; });
+		backend.SetRecvTimeout(backendConnection, m_receiveTimeoutSeconds*1000, nullptr, nullptr);
 
 		backend.SetDisconnectCallback(backendConnection, connection, DisconnectCallback);
 		backend.SetRecvCallbacks(backendConnection, connection, ReceiveHeaderSize, ReceiveResponseHeader, ReceiveResponseBody, TC("ReceiveMessageResponse"));
@@ -241,6 +241,13 @@ namespace uba
 		u16 messageId = u16(headerData[0] << 8) | u16((*(u32*)(headerData + 1) & 0xff000000) >> 24);
 		u32 messageSize = *(u32*)(headerData + 1) & 0x00FFFFFF;
 
+		if (messageSize == MessageKeepAliveSize) // Keep alive message
+		{
+			u8 data[6] = { SystemMessageType_KeepAlive, 0, 1, 0, 0, 0 };
+			NetworkBackend::SendContext sendContext;
+			return connection.backend->Send(client.m_logger, connection.backendConnection, data, sizeof(data), sendContext, TC("KeepAliveNoResponse"));
+		}
+
 		SCOPED_READ_LOCK(client.m_activeMessagesLock, lock);
 		if (!connection.connected)
 			return false;
@@ -252,9 +259,7 @@ namespace uba
 		if (!msg)
 			return false;
 
-		constexpr u32 ErrorSize = 0xffffff - ReceiveHeaderSize; // ReceiveHeaderSize is removed from size in server send
-
-		if (messageSize == ErrorSize || messageSize == 0xffffff)
+		if (messageSize == MessageErrorSize || messageSize == MessageErrorSize - ReceiveHeaderSize)  // ReceiveHeaderSize is removed from size in server send
 		{
 			msg->m_error = 1;
 			msg->Done();
@@ -354,7 +359,6 @@ namespace uba
 	{
 		StackBinaryWriter<64> writer;
 		NetworkMessage msg(*this, SystemServiceId, SystemMessageType_KeepAlive, writer);
-		writer.WriteByte(0); // Need to have a body
 		return msg.Send();
 	}
 
diff --git a/Engine/Source/Programs/UnrealBuildAccelerator/Common/Private/UbaNetworkServer.cpp b/Engine/Source/Programs/UnrealBuildAccelerator/Common/Private/UbaNetworkServer.cpp
--- a/Engine/Source/Programs/UnrealBuildAccelerator/Common/Private/UbaNetworkServer.cpp
+++ b/Engine/Source/Programs/UnrealBuildAccelerator/Common/Private/UbaNetworkServer.cpp
@@ -121,7 +121,12 @@ namespace uba
 					conn.m_server.m_sendBytes += bytes;
 				});
 
-			m_backend.SetRecvTimeout(m_backendConnection, m_server.m_receiveTimeoutMs, nullptr, [](void*) { return false; });
+			m_backend.SetRecvTimeout(m_backendConnection, m_server.m_receiveTimeoutMs, this, [](void* context, u32 timeoutMs, const tchar* recvHint, const tchar* hint)
+				{
+					auto& conn = *(Connection*)context;
+					conn.m_server.m_logger.Error(TC("Connection %u timed out after %u seconds (%s%s)"), conn.m_id, timeoutMs/1000, recvHint, hint);
+					return false;
+				});
 
 			if (requiresCrypto)
 				m_backend.SetRecvCallbacks(m_backendConnection, this, 0, ReceiveHandshakeHeader, ReceiveHandshakeBody, TC("ReceiveHandshake"));
@@ -258,6 +263,19 @@ namespace uba
 			return true;
 		}
 
+		static bool RecvTimeout(void* context, u32 timeoutMs, const tchar* recvHint, const tchar* hint)
+		{
+			auto& conn = *(Connection*)context;
+			++conn.m_recvTimeoutCount;
+			conn.SendKeepAlive();
+			conn.m_backend.SetRecvTimeout(conn.m_backendConnection, KeepAliveIntervalSeconds*1000, context, RecvTimeout);
+			if (conn.m_recvTimeoutCount < KeepAliveProbeCount)
+				return true;
+			constexpr u32 totalTimeoutSeconds = KeepAliveIdleSeconds + KeepAliveIntervalSeconds*KeepAliveProbeCount;
+			conn.m_server.m_logger.Error(TC("Connection %u timed out after %u seconds (%s%s)"), conn.m_id, totalTimeoutSeconds, recvHint, hint);
+			return false;
+		}
+
 		static bool ReceiveClientUid(void* context, const Guid& connectionUid, u8* headerData, void*& outBodyContext, u8*& outBodyData, u32& outBodySize)
 		{
 			auto& conn = *(Connection*)context;
@@ -281,6 +299,10 @@ namespace uba
 			constexpr u32 HeaderSize = 6;
 			conn.m_backend.SetRecvCallbacks(conn.m_backendConnection, &conn, HeaderSize, ReceiveMessageHeader, ReceiveMessageBody, TC("ReceiveMessage"));
 
+			// If keep alive we change timeout to 60 seconds
+			if (server.m_useKeepAlive)
+				conn.m_backend.SetRecvTimeout(conn.m_backendConnection, KeepAliveIdleSeconds*1000, context, RecvTimeout);
+
 			if (!conn.SendInitialResponse(0))
 				return false;
 
@@ -324,11 +346,21 @@ namespace uba
 			u8 messageType = serviceIdAndMessageType & 0b111111;
 			u16 messageId = u16(headerData[1] << 8) | u16((*(u32*)(headerData + 2) & 0xff000000) >> 24);
 			u32 messageSize = *(u32*)(headerData + 2) & 0x00ffffff;
+
 			if (messageSize > SendMaxSize)
 				return conn.m_server.m_logger.Error(TC("Client %u Got message size %u which is larger than max %u. Protocol error? (serviceId %u, messageType %u, messageId %u)"), conn.m_client->id, messageSize, SendMaxSize, u32(serviceId), u32(messageType), u32(messageId));
 			if (serviceId >= sizeof(NetworkServer::m_workerFunctions))
 				return conn.m_server.m_logger.Error(TC("Client %u Got message with service id %u which is out of range. Protocol error?"), conn.m_client->id, serviceId);
 
+			if (conn.m_recvTimeoutCount)
+			{
+				conn.m_recvTimeoutCount = 0;
+				conn.m_backend.SetRecvTimeout(conn.m_backendConnection, KeepAliveIdleSeconds*1000, context, RecvTimeout);
+			}
+
+			if (serviceId == SystemServiceId && messageType == SystemMessageType_KeepAlive) // Keep alive
+				return true;
+
 			//m_logger.Debug(TC("Recv: %u, %u, %u, %u"), serviceId, messageType, id, size);
 			Worker* worker = conn.m_server.PopWorker();
 			if (!worker)
@@ -340,7 +372,7 @@ namespace uba
 			wc.serviceId = serviceId;
 			wc.messageType = messageType;
 			wc.dataSize = messageSize;
-			wc.connection = &conn;// this;
+			wc.connection = &conn;
 			if (wc.buffer.size() < messageSize)
 				wc.buffer.resize(size_t(Min(messageSize + 1024u, SendMaxSize)));
 			outBodyContext = worker;
@@ -412,6 +444,19 @@ namespace uba
 			m_disconnected = true;
 		}
 
+		bool SendKeepAlive()
+		{
+			NetworkBackend::SendContext sendContext;
+			constexpr u32 HeaderSize = 5;
+			u16 messageId = 0;
+			u32 bodySize = MessageKeepAliveSize;
+			u8 data[5];
+			data[0] = messageId >> 8;
+			*(u32*)(data + 1) = bodySize | u32(messageId << 24);
+			return m_backend.Send(m_server.m_logger, m_backendConnection, data, HeaderSize, sendContext, TC("KeepAlive"));
+		}
+
+
 		NetworkServer& m_server;
 		NetworkBackend& m_backend;
 		Futex m_shutdownLock;
@@ -423,6 +468,7 @@ namespace uba
 		Atomic<int> m_disconnectCalled;
 		Atomic<bool> m_disconnected;
 		u32 m_id = 0;
+		u32 m_recvTimeoutCount = 0;
 		bool m_shouldDisconnect = false;
 		void* m_backendConnection = nullptr;
 
@@ -494,7 +540,6 @@ namespace uba
 		BinaryReader reader(context.buffer.data(), 0, context.dataSize);
 
 		constexpr u32 HeaderSize = 5; // 2 byte id, 3 bytes size
-		constexpr u32 ErrorSize = 0xffffff;
 
 		BinaryWriter writer(context.writeMem, 0, context.writeMemSize);
 		u8* idAndSizePtr = writer.AllocWrite(HeaderSize);
@@ -510,7 +555,7 @@ namespace uba
 		{
 			server.m_logger.Error(TC("WORKER FUNCTION NOT FOUND. id: %u, serviceid: %u type: %s, client: %u"), context.id, context.serviceId, rec.toString(context.messageType).data, connection.m_client->id);
 			connection.SetShouldDisconnect();
-			size = ErrorSize;
+			size = MessageErrorSize;
 		}
 		else if (!rec.func({&connection}, {tws}, mi, reader, writer))
 		{
@@ -520,7 +565,7 @@ namespace uba
 				server.m_logger.Error(TC("WORKER FUNCTION FAILED. id: %u, serviceid: %u type: %s, client: %u"), context.id, context.serviceId, rec.toString(context.messageType).data, connection.m_client->id);
 				#endif
 			}
-			size = ErrorSize;
+			size = MessageErrorSize;
 		}
 		else
 		{
@@ -532,14 +577,14 @@ namespace uba
 			UBA_ASSERT(size < (1 << 24));
 				
 			u32 bodySize = u32(size - HeaderSize);
-			if (cryptoKey && size != ErrorSize && bodySize)
+			if (cryptoKey && size != MessageErrorSize && bodySize)
 			{
 				TimerScope ts(connection.m_encryptTimer);
 				u8* bodyData = writer.GetData() + HeaderSize;
 				if (!Crypto::Encrypt(server.m_logger, cryptoKey, bodyData, bodySize))
 				{
 					connection.SetShouldDisconnect();
-					size = ErrorSize;
+					size = MessageErrorSize;
 					bodySize = u32(size - HeaderSize);
 				}
 			}
@@ -548,10 +593,10 @@ namespace uba
 			*(u32*)(idAndSizePtr + 1) = bodySize | u32(context.id << 24);
 
 			// This can happen for proxy servers in a valid situation
-			//if (size == ErrorSize)
+			//if (size == MessageErrorSize)
 			//	UBA_ASSERT(false);
 
-			connection.Send(writer.GetData(), size == ErrorSize ? HeaderSize : size, TC("MessageResponse"));
+			connection.Send(writer.GetData(), size == MessageErrorSize ? HeaderSize : size, TC("MessageResponse"));
 		}
 			
 		connection.Release();
@@ -644,6 +689,12 @@ namespace uba
 		m_sendSize = fixedSendSize;
 		m_receiveTimeoutMs = info.receiveTimeoutSeconds * 1000;
 		m_logConnections = info.logConnections;
+		m_useKeepAlive = info.useKeepAlive;
+
+		#if PLATFORM_MAC
+		m_useKeepAlive = true; // Always run keep alive on mac since the built-in one has a probe interval of 1 minute or something like that... so timeout is always 10 minutes
+		#endif
+
 
 		m_workerFunctions[SystemServiceId].toString = GetMessageTypeToName;
 		m_workerFunctions[SystemServiceId].func = [this](const ConnectionInfo& connectionInfo, const WorkContext& workContext, MessageInfo& messageInfo, BinaryReader& reader, BinaryWriter& writer)
@@ -1072,8 +1123,7 @@ namespace uba
 		}
 		else
 		{
-			constexpr u32 ErrorSize = 0xffffff;
-			bodySize = ErrorSize;
+			bodySize = MessageErrorSize;
 			connection.SetShouldDisconnect();
 		}
 
@@ -1084,6 +1134,14 @@ namespace uba
 		return true;
 	}
 
+	bool NetworkServer::SendKeepAlive()
+	{
+		SCOPED_READ_LOCK(m_connectionsLock, lock);
+		for (auto& it : m_connections)
+			if (!it.SendKeepAlive())
+				return false;
+		return true;
+	}
 
 	NetworkServer::Worker* NetworkServer::PopWorker()
 	{
@@ -1284,11 +1342,6 @@ namespace uba
 				}
 				return true;
 			}
-			case SystemMessageType_KeepAlive:
-			{
-				// No-op
-				return true;
-			}
 			case SystemMessageType_FetchConfig:
 			{
 				writer.Write7BitEncoded(m_clientsConfig.size());
diff --git a/Engine/Source/Programs/UnrealBuildAccelerator/Common/Public/UbaNetwork.h b/Engine/Source/Programs/UnrealBuildAccelerator/Common/Public/UbaNetwork.h
--- a/Engine/Source/Programs/UnrealBuildAccelerator/Common/Public/UbaNetwork.h
+++ b/Engine/Source/Programs/UnrealBuildAccelerator/Common/Public/UbaNetwork.h
@@ -127,11 +127,18 @@ namespace uba
 		}
 	}
 
+	constexpr u32 MessageErrorSize = 0xffffff;
+	constexpr u32 MessageKeepAliveSize = 0xffffff - 1;
+
 	constexpr u16 FetchCasIdDone = u16(~0);
 	constexpr u16 FetchCasIdDisallowed = u16(~0) - 1;
 
 	constexpr bool DownloadDebugSymbols = false;
 
+	static constexpr u32 KeepAliveIdleSeconds = 60;
+	static constexpr u32 KeepAliveIntervalSeconds = 1;
+	static constexpr u32 KeepAliveProbeCount = 10;
+
 	// Response types for SessionMessageType_ProcessAvailable
 	enum SessionProcessAvailableResponse : u32
 	{
diff --git a/Engine/Source/Programs/UnrealBuildAccelerator/Common/Public/UbaNetworkBackend.h b/Engine/Source/Programs/UnrealBuildAccelerator/Common/Public/UbaNetworkBackend.h
--- a/Engine/Source/Programs/UnrealBuildAccelerator/Common/Public/UbaNetworkBackend.h
+++ b/Engine/Source/Programs/UnrealBuildAccelerator/Common/Public/UbaNetworkBackend.h
@@ -46,7 +46,7 @@ namespace uba
 		using RecvBodyCallback = bool(void* context, bool recvError, u8* headerData, void* bodyContext, u8* bodyData, u32 bodySize);
 		virtual void SetRecvCallbacks(void* connection, void* context, u32 headerSize, RecvHeaderCallback* h, RecvBodyCallback* b, const tchar* recvHint) = 0;
 
-		using RecvTimeoutCallback = bool(void* context);
+		using RecvTimeoutCallback = bool(void* context, u32 timeoutMs, const tchar* recvHint, const tchar* hint);
 		virtual void SetRecvTimeout(void* connection, u32 timeoutMs, void* context, RecvTimeoutCallback* callback) = 0;
 
 		// Disconnect callback. This is called as soon as connection is interrupted from send, recv or shutdown.
diff --git a/Engine/Source/Programs/UnrealBuildAccelerator/Common/Public/UbaNetworkServer.h b/Engine/Source/Programs/UnrealBuildAccelerator/Common/Public/UbaNetworkServer.h
--- a/Engine/Source/Programs/UnrealBuildAccelerator/Common/Public/UbaNetworkServer.h
+++ b/Engine/Source/Programs/UnrealBuildAccelerator/Common/Public/UbaNetworkServer.h
@@ -37,6 +37,7 @@ namespace uba
 		u32 sendSize = SendDefaultSize;
 		u32 receiveTimeoutSeconds = DefaultNetworkReceiveTimeoutSeconds;
 		bool logConnections = true;
+		bool useKeepAlive = false;
 	};
 
 	struct MessageInfo
@@ -104,6 +105,7 @@ namespace uba
 		bool DoAdditionalWork();
 
 		bool SendResponse(const MessageInfo& info, const u8* body, u32 bodySize);
+		bool SendKeepAlive();
 
 		class Worker;
 	private:
@@ -159,6 +161,7 @@ namespace uba
 		u32 m_maxCreatedWorkerCount = 0;
 		bool m_workersEnabled = true;
 		bool m_logConnections = true;
+		bool m_useKeepAlive = false;
 		bool m_allClientsDisconnected = false;
 
 		Futex m_addConnectionsLock;
diff --git a/Engine/Source/Programs/UnrealBuildAccelerator/Test/Private/UbaTestNetwork.cpp b/Engine/Source/Programs/UnrealBuildAccelerator/Test/Private/UbaTestNetwork.cpp
--- a/Engine/Source/Programs/UnrealBuildAccelerator/Test/Private/UbaTestNetwork.cpp
+++ b/Engine/Source/Programs/UnrealBuildAccelerator/Test/Private/UbaTestNetwork.cpp
@@ -91,6 +91,10 @@ namespace uba
 
 		u8 value = reader.ReadByte();
 		logger.Info(TC("Got value %u"), value);
+
+		if (!server.SendKeepAlive())
+			return logger.Error(TC("Failed to send keep alive"));
+		Sleep(100);
 		return true;
 	}
 
