[UBA]
* Added more validation logic for cache system

diff --git a/Engine/Source/Programs/UnrealBuildAccelerator/CacheService/Private/UbaCacheService.cpp b/Engine/Source/Programs/UnrealBuildAccelerator/CacheService/Private/UbaCacheService.cpp
--- a/Engine/Source/Programs/UnrealBuildAccelerator/CacheService/Private/UbaCacheService.cpp
+++ b/Engine/Source/Programs/UnrealBuildAccelerator/CacheService/Private/UbaCacheService.cpp
@@ -17,6 +17,8 @@
 #include <dbghelp.h>
 #include <io.h>
 #pragma comment (lib, "Dbghelp.lib")
+#elif PLATFORM_LINUX
+#include <sys/prctl.h>
 #endif
 
 namespace uba
@@ -268,6 +270,15 @@ namespace uba
 		#if PLATFORM_LINUX
 		if (forkProcess)
 		{
+			if (prctl(PR_GET_DUMPABLE) == 0)
+			{
+				prctl(PR_SET_DUMPABLE, 1);
+				if (prctl(PR_GET_DUMPABLE) == 0)
+					logger.Info(TC("  prctl(PR_SET_DUMPABLE, 1) failed to set dumpable."));
+				else
+					logger.Info(TC("  Made process dumpable"));
+			}
+
 			auto ReadFirstLine = [](StringBufferBase& out, const char* cmd)
 				{
 					FILE* f = popen(cmd, "r");
@@ -295,6 +306,11 @@ namespace uba
 						logger.Info(TC("  Crash dumps enabled but piped so can't wait (%s). use 'sudo echo \"/<path>/dump.%t\" | sudo tee /proc/sys/kernel/core_pattern > /dev/null' to write to file"), temp.data);
 					else
 					{
+						//if (ReadFirstLine(temp.Clear(), "cat /proc/sys/fs/suid_dumpable"))
+						//{
+						//}
+
+
 						logger.Info(TC("  Crash dumps enabled and written to file: %s (Write no other files in the same dir)"), temp.data);
 						crashDumpDir = StringView(temp).GetPath().ToString();
 						crashDumpPattern = StringView(temp).GetFileName().ToString();
diff --git a/Engine/Source/Programs/UnrealBuildAccelerator/Common/Private/UbaCacheClient.cpp b/Engine/Source/Programs/UnrealBuildAccelerator/Common/Private/UbaCacheClient.cpp
--- a/Engine/Source/Programs/UnrealBuildAccelerator/Common/Private/UbaCacheClient.cpp
+++ b/Engine/Source/Programs/UnrealBuildAccelerator/Common/Private/UbaCacheClient.cpp
@@ -1114,21 +1114,30 @@ namespace uba
 				if (!m_storage.GetCasFileName(casKeyFileName, casKey))
 					return false;
 
-				const u8* fileData;
-				u64 fileSize;
+				const u8* fileData = nullptr;
+				u64 fileSize = 0;
 
 				MappedView mappedView;
 				auto mapViewGuard = MakeGuard([&](){ m_storage.m_casDataBuffer.UnmapView(mappedView, path.data); });
 
 				FileAccessor file(m_logger, casKeyFileName.data);
 
-				if (casEntry->mappingHandle.IsValid()) // If file was created by helper it will be in the transient mapped memory
 				{
-					mappedView = m_storage.m_casDataBuffer.MapView(casEntry->mappingHandle, casEntry->mappingOffset, casEntry->mappingSize, path.data);
-					fileData = mappedView.memory;
-					fileSize = mappedView.size;
+					SCOPED_READ_LOCK(casEntry->lock, entryLock);
+					auto mh = casEntry->mappingHandle;
+					if (mh.IsValid()) // If file was created by helper it will be in the transient mapped memory
+					{
+						u64 mappingOffset = casEntry->mappingOffset;
+						u64 mappingSize = casEntry->mappingSize;
+						entryLock.Leave();
+
+						mappedView = m_storage.m_casDataBuffer.MapView(mh, mappingOffset, mappingSize, path.data);
+						fileData = mappedView.memory;
+						fileSize = mappedView.size;
+					}
 				}
-				else
+
+				if (!fileData)
 				{
 					if (!file.OpenMemoryRead())
 						return false;
diff --git a/Engine/Source/Programs/UnrealBuildAccelerator/Common/Private/UbaNetworkServer.cpp b/Engine/Source/Programs/UnrealBuildAccelerator/Common/Private/UbaNetworkServer.cpp
--- a/Engine/Source/Programs/UnrealBuildAccelerator/Common/Private/UbaNetworkServer.cpp
+++ b/Engine/Source/Programs/UnrealBuildAccelerator/Common/Private/UbaNetworkServer.cpp
@@ -109,7 +109,7 @@ namespace uba
 			m_backend.SetDisconnectCallback(m_backendConnection, this, [](void* context, const Guid& connectionUid, void* connection)
 				{
 					auto& conn = *(Connection*)context;
-					conn.Disconnect();
+					conn.Disconnect(TC("Backend"));
 					conn.m_disconnectCallbackCalled.Set();
 				});
 
@@ -141,7 +141,7 @@ namespace uba
 				Crypto::DestroyKey(m_cryptoKey);
 		}
 
-		void Disconnect()
+		void Disconnect(const tchar* reason)
 		{
 			if (m_disconnectCalled.fetch_add(1) != 0)
 				return;
@@ -150,12 +150,12 @@ namespace uba
 			if (!activeWorkerCount) // Will disconnect in send if there are active workers
 				TestDisconnect();
 			else
-				m_server.m_logger.Detail(TC("Connection %u disconnected while it has %u active workers"), m_id, activeWorkerCount);
+				m_server.m_logger.Detail(TC("Connection %u disconnected while it has %u active workers (%s)"), m_id, activeWorkerCount, reason);
 		}
 
 		bool Stop()
 		{
-			Disconnect();
+			Disconnect(TC("Stop"));
 
 			u64 startTimer = GetTime();
 			while (m_activeWorkerCount)
diff --git a/Engine/Source/Programs/UnrealBuildAccelerator/Common/Private/UbaStorageServer.cpp b/Engine/Source/Programs/UnrealBuildAccelerator/Common/Private/UbaStorageServer.cpp
--- a/Engine/Source/Programs/UnrealBuildAccelerator/Common/Private/UbaStorageServer.cpp
+++ b/Engine/Source/Programs/UnrealBuildAccelerator/Common/Private/UbaStorageServer.cpp
@@ -877,6 +877,9 @@ namespace uba
 					casEntry.dropped = true; // Set to dropped because this will never hit disk
 				}
 
+				if (casEntry.beingWritten)
+					return m_logger.Error(TC("Cas is being written while starting to write. %s (%s)"), CasKeyString(casKey).str, hint.data);
+
 				casEntry.beingWritten = true;
 
 				firstStore = &tempStore;
@@ -911,17 +914,22 @@ namespace uba
 					if (storeIt == m_activeStores.end())
 						return m_logger.Error(TC("Can't find active store %u, disconnected client?"), storeId);
 					activeStoreTemp = &storeIt->second;
-					UBA_ASSERT(activeStoreTemp->clientId == connectionInfo.GetId());
+					if (activeStoreTemp->clientId != connectionInfo.GetId())
+						return m_logger.Error(TC("ClientId not matching for active store %u, disconnected client? (%u vs %u)"), storeId, activeStoreTemp->clientId, connectionInfo.GetId());
 				}
 				ActiveStore& activeStore = *activeStoreTemp;
 
 				u64 toRead = reader.GetLeft();
+				u64 fileSize = activeStore.fileSize;
+
+				if (fileSize < memOffset + toRead)
+					return m_logger.Error(TC("Trying to write data to cas outside of size (Size %llu, writing %llu at offset %llu"), fileSize, memOffset, toRead);
+
 				reader.ReadBytes(activeStore.mappedView.memory + memOffset, toRead);
 				
 				u64 time2 = GetTime();
 				activeStore.recvCasTime += time2 - start;
 
-				u64 fileSize = activeStore.fileSize;
 				u64 totalWritten = activeStore.totalWritten.fetch_add(toRead) + toRead;
 				if (totalWritten == fileSize)
 				{
