[UBA]
* Fixed all thread names to be 15 characters or shorter (posix limitation)

diff --git a/Engine/Source/Programs/UnrealBuildAccelerator/Cli/Private/UbaCli.cpp b/Engine/Source/Programs/UnrealBuildAccelerator/Cli/Private/UbaCli.cpp
--- a/Engine/Source/Programs/UnrealBuildAccelerator/Cli/Private/UbaCli.cpp
+++ b/Engine/Source/Programs/UnrealBuildAccelerator/Cli/Private/UbaCli.cpp
@@ -600,7 +600,7 @@ namespace uba
 			if (!client.Init(cii))
 				return false;
 			bool success = true;
-			WorkManagerImpl workManager(DefaultProcessorCount, TC("UbaWorker(CheckCas2)"));
+			WorkManagerImpl workManager(DefaultProcessorCount, TC("UbaWrk/ChkCas2"));
 			storageServer.TraverseAllCasFiles([&](const CasKey& casKey, u64 size)
 				{
 					workManager.AddWork([&, casKey](const WorkContext&)
@@ -652,7 +652,7 @@ namespace uba
 		u64 testCompressOriginalSize = 0;
 		if (!testCompress.empty())
 		{
-			WorkManagerImpl workManager(DefaultProcessorCount, TC("UbaWorker(TestCompress)"));
+			WorkManagerImpl workManager(DefaultProcessorCount, TC("UbaWrk/TstComp"));
 
 			FileAccessor fa(logger, testCompress.c_str());
 			if (!fa.OpenMemoryRead())
@@ -683,7 +683,7 @@ namespace uba
 
 		if (!testDecompress.empty())
 		{
-			WorkManagerImpl workManager(DefaultProcessorCount, TC("UbaWorker(TestDecompress)"));
+			WorkManagerImpl workManager(DefaultProcessorCount, TC("UbaWrk/TstDecm"));
 
 			FileAccessor fa(logger, testDecompress.c_str());
 			if (!fa.OpenMemoryRead())
diff --git a/Engine/Source/Programs/UnrealBuildAccelerator/Common/Private/UbaNetworkClient.cpp b/Engine/Source/Programs/UnrealBuildAccelerator/Common/Private/UbaNetworkClient.cpp
--- a/Engine/Source/Programs/UnrealBuildAccelerator/Common/Private/UbaNetworkClient.cpp
+++ b/Engine/Source/Programs/UnrealBuildAccelerator/Common/Private/UbaNetworkClient.cpp
@@ -21,7 +21,7 @@ namespace uba
 	}
 
 	NetworkClient::NetworkClient(bool& outCtorSuccess, const NetworkClientCreateInfo& info, const tchar* name)
-	: WorkManagerImpl(info.workerCount == ~0u ? GetLogicalProcessorCount() : info.workerCount, TC("UbaWorker(NetworkClient)"))
+	: WorkManagerImpl(info.workerCount == ~0u ? GetLogicalProcessorCount() : info.workerCount, TC("UbaWrk/NetClnt"))
 	,	m_logWriter(info.logWriter)
 	,	m_logger(info.logWriter, SetGetPrefix(name))
 	,	m_isConnected(true)
diff --git a/Engine/Source/Programs/UnrealBuildAccelerator/Common/Private/UbaNetworkServer.cpp b/Engine/Source/Programs/UnrealBuildAccelerator/Common/Private/UbaNetworkServer.cpp
--- a/Engine/Source/Programs/UnrealBuildAccelerator/Common/Private/UbaNetworkServer.cpp
+++ b/Engine/Source/Programs/UnrealBuildAccelerator/Common/Private/UbaNetworkServer.cpp
@@ -58,7 +58,7 @@ namespace uba
 		{
 			m_context = new WorkerContext(server);
 			m_loop = true;
-			m_thread.Start([&]() { ThreadWorker(server); return 0; }, TC("UbaWorker(NetworkServer)"));
+			m_thread.Start([&]() { ThreadWorker(server); return 0; }, TC("UbaWrk/NetwSrv"));
 		}
 
 		void Stop(NetworkServer& server)
diff --git a/Engine/Source/Programs/UnrealBuildAccelerator/Common/Private/UbaScheduler.cpp b/Engine/Source/Programs/UnrealBuildAccelerator/Common/Private/UbaScheduler.cpp
--- a/Engine/Source/Programs/UnrealBuildAccelerator/Common/Private/UbaScheduler.cpp
+++ b/Engine/Source/Programs/UnrealBuildAccelerator/Common/Private/UbaScheduler.cpp
@@ -119,7 +119,7 @@ namespace uba
 		m_session.SetRemoteProcessSlotAvailableEvent([this](bool isCrossArchitecture) { RemoteSlotAvailable(isCrossArchitecture); });
 
 		m_loop = true;
-		m_thread.Start([this]() { ThreadLoop(); return 0; }, TC("UbaSchedulerLoop"));
+		m_thread.Start([this]() { ThreadLoop(); return 0; }, TC("UbaSchedLoop"));
 	}
 
 	void Scheduler::Stop()
diff --git a/Engine/Source/Programs/UnrealBuildAccelerator/Common/Private/UbaSessionClient.cpp b/Engine/Source/Programs/UnrealBuildAccelerator/Common/Private/UbaSessionClient.cpp
--- a/Engine/Source/Programs/UnrealBuildAccelerator/Common/Private/UbaSessionClient.cpp
+++ b/Engine/Source/Programs/UnrealBuildAccelerator/Common/Private/UbaSessionClient.cpp
@@ -1377,7 +1377,7 @@ namespace uba
 
 		BuildEnvironmentVariables(reader);
 
-		m_loopThread.Start([this]() { ThreadCreateProcessLoop(); return 0; }, TC("UbaCreateProcessLoop"));
+		m_loopThread.Start([this]() { ThreadCreateProcessLoop(); return 0; }, TC("UbaCreateProc"));
 	}
 
 	void SessionClient::BuildEnvironmentVariables(BinaryReader& reader)
diff --git a/Engine/Source/Programs/UnrealBuildAccelerator/Common/Private/UbaSessionServer.cpp b/Engine/Source/Programs/UnrealBuildAccelerator/Common/Private/UbaSessionServer.cpp
--- a/Engine/Source/Programs/UnrealBuildAccelerator/Common/Private/UbaSessionServer.cpp
+++ b/Engine/Source/Programs/UnrealBuildAccelerator/Common/Private/UbaSessionServer.cpp
@@ -195,7 +195,7 @@ namespace uba
 				m_memRequiredToSpawn = Min(u64(double(m_memTotal) * double(100 - info.memWaitLoadPercent) / 100.0), 35ull * 1024 * 1024 * 1024);
 			}
 
-			m_memoryThread.Start([this]() { ThreadMemoryCheckLoop();  return 0; }, TC("UbaMemoryTrackerLoop"));
+			m_memoryThread.Start([this]() { ThreadMemoryCheckLoop();  return 0; }, TC("UbaMemTrackLoop"));
 		}
 
 		#if PLATFORM_WINDOWS
diff --git a/Engine/Source/Programs/UnrealBuildAccelerator/Common/Private/UbaStorage.cpp b/Engine/Source/Programs/UnrealBuildAccelerator/Common/Private/UbaStorage.cpp
--- a/Engine/Source/Programs/UnrealBuildAccelerator/Common/Private/UbaStorage.cpp
+++ b/Engine/Source/Programs/UnrealBuildAccelerator/Common/Private/UbaStorage.cpp
@@ -1634,7 +1634,7 @@ namespace uba
 			};
 
 		m_logger.Info(TC("Traverse cas database..."));
-		WorkManagerImpl workManager(workerCount, TC("UbaWorker(CheckCasContent)"));
+		WorkManagerImpl workManager(workerCount, TC("UbaWrk/ChkCasC"));
 		u32 entryCount = 0;
 		u32 errorCount = 0;
 		u64 newestWrittenError = 0;
@@ -1709,7 +1709,7 @@ namespace uba
 	bool StorageImpl::CheckFileTable(const tchar* searchPath, u32 workerCount)
 	{
 		m_logger.Info(TC("Searching %s to check files against file table..."), searchPath);
-		WorkManagerImpl workManager(workerCount, TC("UbaWorker(CheckFileTable)"));
+		WorkManagerImpl workManager(workerCount, TC("UbaWrk/ChkFTbl"));
 
 		List<TString> directories;
 		directories.push_back(searchPath);
@@ -1814,7 +1814,7 @@ namespace uba
 	{
 		u32 deleteCount = 0;
 
-		WorkManagerImpl workManager(GetLogicalProcessorCount(), TC("UbaWorker(DeleteCas)"));
+		WorkManagerImpl workManager(GetLogicalProcessorCount(), TC("UbaWrk/DelCas"));
 		{
 			Atomic<u32> atomicDeleteCount;
 			TraverseDir(m_logger, m_rootDir, [&](const DirectoryEntry& e)
diff --git a/Engine/Source/Programs/UnrealBuildAccelerator/Common/Private/UbaStorageClient.cpp b/Engine/Source/Programs/UnrealBuildAccelerator/Common/Private/UbaStorageClient.cpp
--- a/Engine/Source/Programs/UnrealBuildAccelerator/Common/Private/UbaStorageClient.cpp
+++ b/Engine/Source/Programs/UnrealBuildAccelerator/Common/Private/UbaStorageClient.cpp
@@ -128,7 +128,7 @@ namespace uba
 
 		u64 start = GetTime();
 
-		WorkManagerImpl workManager(workerCount, TC("UbaWorker(PopulateCasFromDirs)"));
+		WorkManagerImpl workManager(workerCount, TC("UbaWrk/PoplCas"));
 		bool success = true;
 		UnorderedSet<u64> seenIds;
 		Futex seenIdsLock;
diff --git a/Engine/Source/Programs/UnrealBuildAccelerator/Common/Public/UbaCoordinatorWrapper.h b/Engine/Source/Programs/UnrealBuildAccelerator/Common/Public/UbaCoordinatorWrapper.h
--- a/Engine/Source/Programs/UnrealBuildAccelerator/Common/Public/UbaCoordinatorWrapper.h
+++ b/Engine/Source/Programs/UnrealBuildAccelerator/Common/Public/UbaCoordinatorWrapper.h
@@ -61,7 +61,7 @@ namespace uba
 			m_networkServer = &networkServer;
 			m_scheduler = scheduler;
 
-			m_coordinatorThread.Start([this, mcc = info.maxCoreCount]() { ThreadUpdate(mcc); return 0; });
+			m_coordinatorThread.Start([this, mcc = info.maxCoreCount]() { ThreadUpdate(mcc); return 0; }, TC("UbaCoordWrap"));
 
 			return true;
 		}
diff --git a/Engine/Source/Programs/UnrealBuildAccelerator/Common/Public/UbaWorkManager.h b/Engine/Source/Programs/UnrealBuildAccelerator/Common/Public/UbaWorkManager.h
--- a/Engine/Source/Programs/UnrealBuildAccelerator/Common/Public/UbaWorkManager.h
+++ b/Engine/Source/Programs/UnrealBuildAccelerator/Common/Public/UbaWorkManager.h
@@ -57,7 +57,7 @@ namespace uba
 	class WorkManagerImpl : public WorkManager
 	{
 	public:
-		WorkManagerImpl(u32 workerCount, const tchar* workerDesc = TC("UbaWorker"));
+		WorkManagerImpl(u32 workerCount, const tchar* workerDesc = TC("UbaWrk"));
 		virtual ~WorkManagerImpl();
 		virtual void AddWork(const WorkFunction& work, u32 count, const tchar* desc, const Color& color = ColorWork, bool highPriority = false) override;
 		virtual u32 GetWorkerCount() override;
diff --git a/Engine/Source/Programs/UnrealBuildAccelerator/ObjTool/Private/UbaObjTool.cpp b/Engine/Source/Programs/UnrealBuildAccelerator/ObjTool/Private/UbaObjTool.cpp
--- a/Engine/Source/Programs/UnrealBuildAccelerator/ObjTool/Private/UbaObjTool.cpp
+++ b/Engine/Source/Programs/UnrealBuildAccelerator/ObjTool/Private/UbaObjTool.cpp
@@ -262,7 +262,7 @@ namespace uba
 				allExternalImports.insert(imp);
 
 			u32 workerCount = DefaultProcessorCount-1;
-			WorkManagerImpl workManager(workerCount, TC("UbaWorker(StripObjFiles)"));
+			WorkManagerImpl workManager(workerCount, TC("UbaWrk/StrpObj"));
 			workManager.ParallelFor(workerCount, objFilesDependencies, [&](const WorkContext&, auto& it)
 				{
 					const TString& exiFilename = *it;
@@ -338,7 +338,7 @@ namespace uba
 			Vector<ObjectFile*> objFiles;
 			objFiles.resize(objFilesForImpLib.size());
 			u32 workerCount = DefaultProcessorCount - 1;
-			WorkManagerImpl workManager(workerCount, TC("UbaWorker(Load)"));
+			WorkManagerImpl workManager(workerCount, TC("UbaWrk/Load"));
 			workManager.ParallelFor(workerCount, objFilesForImpLib, [&](const WorkContext&, auto& it)
 			{
 				auto& o = *it;
diff --git a/Engine/Source/Programs/UnrealBuildAccelerator/Visualizer/Private/UbaVisualizer.cpp b/Engine/Source/Programs/UnrealBuildAccelerator/Visualizer/Private/UbaVisualizer.cpp
--- a/Engine/Source/Programs/UnrealBuildAccelerator/Visualizer/Private/UbaVisualizer.cpp
+++ b/Engine/Source/Programs/UnrealBuildAccelerator/Visualizer/Private/UbaVisualizer.cpp
@@ -379,7 +379,7 @@ namespace uba
 
 	bool Visualizer::StartHwndThread()
 	{
-		m_thread.Start([this]() { ThreadLoop(); return 0;});
+		m_thread.Start([this]() { ThreadLoop(); return 0;}, TC("UbaHwnd"));
 		while (!m_hwnd)
 			if (m_thread.Wait(10))
 				return false;
