[UBA]
* Changed error to detail to prevent build from failing

diff --git a/Engine/Source/Programs/UnrealBuildAccelerator/Common/Private/UbaLogger.cpp b/Engine/Source/Programs/UnrealBuildAccelerator/Common/Private/UbaLogger.cpp
--- a/Engine/Source/Programs/UnrealBuildAccelerator/Common/Private/UbaLogger.cpp
+++ b/Engine/Source/Programs/UnrealBuildAccelerator/Common/Private/UbaLogger.cpp
@@ -131,20 +131,22 @@ namespace uba
 		return *this;
 	}
 
-	void Logger::Detail(const tchar* format, ...)
+	Logger& Logger::Detail(const tchar* format, ...)
 	{
 		va_list arg;
 		va_start(arg, format);
 		LogArg(LogEntryType_Detail, format, arg);
 		va_end(arg);
+		return *this;
 	}
 
-	void Logger::Debug(const tchar* format, ...)
+	Logger& Logger::Debug(const tchar* format, ...)
 	{
 		va_list arg;
 		va_start(arg, format);
 		LogArg(LogEntryType_Debug, format, arg);
 		va_end(arg);
+		return *this;
 	}
 
 	bool Logger::Warning(const tchar* format, ...)
diff --git a/Engine/Source/Programs/UnrealBuildAccelerator/Common/Private/UbaStorageServer.cpp b/Engine/Source/Programs/UnrealBuildAccelerator/Common/Private/UbaStorageServer.cpp
--- a/Engine/Source/Programs/UnrealBuildAccelerator/Common/Private/UbaStorageServer.cpp
+++ b/Engine/Source/Programs/UnrealBuildAccelerator/Common/Private/UbaStorageServer.cpp
@@ -693,7 +693,7 @@ namespace uba
 				SCOPED_READ_LOCK(m_activeFetchesLock, lock);
 				auto findIt = m_activeFetches.find(fetchId);
 				if (findIt == m_activeFetches.end())
-					return m_logger.Error(TC("Can't find active fetch %u, disconnected client? (fetch index %u, client id %u uid %s)"), fetchId, fetchIndex, connectionInfo.GetId(), GuidToString(connectionInfo.GetUid()).str);
+					return m_logger.Detail(TC("Can't find active fetch %u, disconnected client? (fetch index %u, client id %u uid %s)"), fetchId, fetchIndex, connectionInfo.GetId(), GuidToString(connectionInfo.GetUid()).str).ToFalse();
 				ActiveFetch& fetch = findIt->second;
 				UBA_ASSERT(fetch.clientId == connectionInfo.GetId());
 				lock.Leave();
@@ -702,11 +702,12 @@ namespace uba
 				const u8* pos = fetch.memoryPos + (fetchIndex-1) * writer.GetCapacityLeft();
 				u64 toWrite = writer.GetCapacityLeft();
 				u64 readOffset = pos - fetch.memoryBegin;
-				if (readOffset + toWrite > fetch.mappedView.size)
+				u64 viewSize = fetch.mappedView.size;
+				if (readOffset + toWrite > viewSize)
 				{
-					if (readOffset >= fetch.mappedView.size)
-						return m_logger.Error(TC("Client is asking for more data than file contains (fetch index %u, client id %u uid %s)"), fetchIndex, connectionInfo.GetId(), GuidToString(connectionInfo.GetUid()).str);
-					toWrite = fetch.mappedView.size - readOffset;
+					if (readOffset >= viewSize)
+						return m_logger.Detail(TC("Client is asking for more data than file contains (readoffset %llu, size %llu, fetch index %u, client id %u uid %s)"), readOffset, viewSize, fetchIndex, connectionInfo.GetId(), GuidToString(connectionInfo.GetUid()).str).ToFalse();
+					toWrite = viewSize - readOffset;
 				}
 				memcpy(writer.AllocWrite(toWrite), pos, toWrite);
 
@@ -915,7 +916,7 @@ namespace uba
 						return m_logger.Error(TC("Can't find active store %u, disconnected client?"), storeId);
 					activeStoreTemp = &storeIt->second;
 					if (activeStoreTemp->clientId != connectionInfo.GetId())
-						return m_logger.Error(TC("ClientId not matching for active store %u, disconnected client? (%u vs %u)"), storeId, activeStoreTemp->clientId, connectionInfo.GetId());
+						return m_logger.Error(TC("Client id not matching for active store %u, disconnected client? (%u vs %u)"), storeId, activeStoreTemp->clientId, connectionInfo.GetId());
 				}
 				ActiveStore& activeStore = *activeStoreTemp;
 
diff --git a/Engine/Source/Programs/UnrealBuildAccelerator/Common/Public/UbaLogger.h b/Engine/Source/Programs/UnrealBuildAccelerator/Common/Public/UbaLogger.h
--- a/Engine/Source/Programs/UnrealBuildAccelerator/Common/Public/UbaLogger.h
+++ b/Engine/Source/Programs/UnrealBuildAccelerator/Common/Public/UbaLogger.h
@@ -28,8 +28,8 @@ namespace uba
 		bool Error(const tchar* format, ...);
 		bool Warning(const tchar* format, ...);
 		Logger& Info(const tchar* format, ...);
-		void Detail(const tchar* format, ...);
-		void Debug(const tchar* format, ...);
+		Logger& Detail(const tchar* format, ...);
+		Logger& Debug(const tchar* format, ...);
 		void Logf(LogEntryType type, const tchar* format, ...);
 		void LogArg(LogEntryType type, const tchar* format, va_list& args);
 		void Log(LogEntryType type, const StringView& str);
