[UBA]
* Changed how workers are pushed and popped.. now a pop adds itself to a linked list that is picked up by worker when available.. this removes a shared event that is used by all workers

diff --git a/Engine/Source/Programs/UnrealBuildAccelerator/Common/Private/UbaNetworkClient.cpp b/Engine/Source/Programs/UnrealBuildAccelerator/Common/Private/UbaNetworkClient.cpp
--- a/Engine/Source/Programs/UnrealBuildAccelerator/Common/Private/UbaNetworkClient.cpp
+++ b/Engine/Source/Programs/UnrealBuildAccelerator/Common/Private/UbaNetworkClient.cpp
@@ -596,7 +596,7 @@ namespace uba
 		u8* data = writer.GetData();
 		data[1] = messageId >> 8;
 		u32 dataSize = sendSize - 6;
-		UBA_ASSERTF(dataSize, TC("NetworkMessage must have data size of at least 1."));
+		UBA_ASSERTF(dataSize || data[0] == 1, TC("NetworkMessage must have data size of at least 1."));
 		*(u32*)(data + 2) = dataSize | u32(messageId) << 24;
 
 		//m_logger.Debug(TC("Send: %u, %u, %u, %u"), data[0], data[1], data[2], sendSize - 7);
diff --git a/Engine/Source/Programs/UnrealBuildAccelerator/Common/Private/UbaNetworkServer.cpp b/Engine/Source/Programs/UnrealBuildAccelerator/Common/Private/UbaNetworkServer.cpp
--- a/Engine/Source/Programs/UnrealBuildAccelerator/Common/Private/UbaNetworkServer.cpp
+++ b/Engine/Source/Programs/UnrealBuildAccelerator/Common/Private/UbaNetworkServer.cpp
@@ -667,7 +667,6 @@ namespace uba
 
 	NetworkServer::NetworkServer(bool& outCtorSuccess, const NetworkServerCreateInfo& info, const tchar* name)
 	:	m_logger(info.logWriter, name)
-	,	m_workerAvailable(true)
 	{
 		outCtorSuccess = true;
 
@@ -732,7 +731,13 @@ namespace uba
 		{
 			SCOPED_FUTEX(m_availableWorkersLock, lock);
 			m_workersEnabled = false;
-			m_workerAvailable.Set();
+			while (PopWorkerRequest* req = m_firstRequest)
+			{
+				m_firstRequest = req->next;
+				req->next = nullptr;
+				req->ev.Set();
+			}
+			m_lastRequest = nullptr;
 		}
 		{
 			SCOPED_FUTEX(m_addConnectionsLock, lock);
@@ -1152,9 +1157,22 @@ namespace uba
 				return nullptr;
 			if (auto worker = PopWorkerNoLock())
 				return worker;
-			m_workerAvailable.Reset();
+			
+			PopWorkerRequest req;
+			req.ev.Create(true);
+
+			if (!m_firstRequest)
+				m_firstRequest = &req;
+			else
+				m_lastRequest->next = &req;
+			m_lastRequest = &req;
+
 			lock.Leave();
-			m_workerAvailable.IsSet();
+
+			req.ev.IsSet();
+
+			if (req.worker)
+				return req.worker;
 		}
 	}
 
@@ -1189,6 +1207,7 @@ namespace uba
 	void NetworkServer::PushWorker(Worker* worker)
 	{
 		SCOPED_FUTEX(m_availableWorkersLock, lock);
+
 		PushWorkerNoLock(worker);
 	}
 
@@ -1196,6 +1215,17 @@ namespace uba
 	{
 		UBA_ASSERT(worker->m_inUse);
 
+		if (PopWorkerRequest* first = m_firstRequest)
+		{
+			m_firstRequest = first->next;
+			if (!m_firstRequest)
+				m_lastRequest = nullptr;
+			first->worker = worker;
+			first->ev.Set();
+			return;
+		}
+
+
 		if (worker->m_prevWorker)
 			worker->m_prevWorker->m_nextWorker = worker->m_nextWorker;
 		else
@@ -1209,7 +1239,6 @@ namespace uba
 		worker->m_nextWorker = m_firstAvailableWorker;
 		worker->m_inUse = false;
 		m_firstAvailableWorker = worker;
-		m_workerAvailable.Set();
 	}
 
 	void NetworkServer::FlushWorkers()
diff --git a/Engine/Source/Programs/UnrealBuildAccelerator/Common/Public/UbaNetworkServer.h b/Engine/Source/Programs/UnrealBuildAccelerator/Common/Public/UbaNetworkServer.h
--- a/Engine/Source/Programs/UnrealBuildAccelerator/Common/Public/UbaNetworkServer.h
+++ b/Engine/Source/Programs/UnrealBuildAccelerator/Common/Public/UbaNetworkServer.h
@@ -154,7 +154,16 @@ namespace uba
 		Futex m_availableWorkersLock;
 		Worker* m_firstAvailableWorker = nullptr;
 		Worker* m_firstActiveWorker = nullptr;
-		Event m_workerAvailable;
+
+		struct PopWorkerRequest
+		{
+			Event ev;
+			PopWorkerRequest* next = nullptr;
+			Worker* worker = nullptr;
+		};
+		PopWorkerRequest* m_firstRequest = nullptr;
+		PopWorkerRequest* m_lastRequest = nullptr;
+
 		u32 m_sendSize = 0;
 		u32 m_receiveTimeoutMs = 0;
 		u32 m_createdWorkerCount = 0;
