[UBA]
* Added ProcessPriority constants

diff --git a/Engine/Source/Programs/UnrealBuildAccelerator/Common/Public/UbaDefaultConstants.h b/Engine/Source/Programs/UnrealBuildAccelerator/Common/Public/UbaDefaultConstants.h
--- a/Engine/Source/Programs/UnrealBuildAccelerator/Common/Public/UbaDefaultConstants.h
+++ b/Engine/Source/Programs/UnrealBuildAccelerator/Common/Public/UbaDefaultConstants.h
@@ -13,4 +13,8 @@ namespace uba
 	static constexpr u32 SendDefaultSize = 256*1024;
 	static constexpr u32 DefaultNetworkReceiveTimeoutSeconds = 10*60;
 	static constexpr u32 DefaultNetworkSendTimeoutSeconds = 10*60;
+
+
+	static constexpr u32 ProcessPriority_Normal = 0x00000020;
+	static constexpr u32 ProcessPriority_BelowNormal = 0x00004000;
 }
diff --git a/Engine/Source/Programs/UnrealBuildAccelerator/Common/Public/UbaProcessStartInfo.h b/Engine/Source/Programs/UnrealBuildAccelerator/Common/Public/UbaProcessStartInfo.h
--- a/Engine/Source/Programs/UnrealBuildAccelerator/Common/Public/UbaProcessStartInfo.h
+++ b/Engine/Source/Programs/UnrealBuildAccelerator/Common/Public/UbaProcessStartInfo.h
@@ -23,7 +23,7 @@ namespace uba
 		const tchar* workingDir = TC("");		// Working directory. Use full path
 		const tchar* description = TC("");		// Description. Used for on-screen logging and log file names if session.logToFile is set but logFile is ""
 		const tchar* logFile = TC("");			// Log file. If set, will always log. If not full path the session log dir will be prepended.
-		u32 priorityClass = 0x00000020;			// Priority of process. Defaults to NORMAL_PRIORITY_CLASS
+		u32 priorityClass = ProcessPriority_Normal; // Priority of process.
 		bool trackInputs = false;				// Track all files read. Can read result in ProcessHandle.GetTrackedInputs()
 		bool useCustomAllocator = true;			// Disable detouring of allocator inside processes. If Session.disableCustomAllocator is false this will be overridden
 		bool writeOutputFilesOnFail = false;	// If set to true, output files will be written to disk regardless if process succeeds or not
