[UBA]
* Fixed non unity compile error

diff --git a/Engine/Source/Programs/UnrealBuildAccelerator/Common/Public/UbaProcessStartInfo.h b/Engine/Source/Programs/UnrealBuildAccelerator/Common/Public/UbaProcessStartInfo.h
--- a/Engine/Source/Programs/UnrealBuildAccelerator/Common/Public/UbaProcessStartInfo.h
+++ b/Engine/Source/Programs/UnrealBuildAccelerator/Common/Public/UbaProcessStartInfo.h
@@ -2,6 +2,7 @@
 
 #pragma once
 
+#include "UbaDefaultConstants.h"
 #include "UbaLogWriter.h"
 
 namespace uba
