[UBA]
* Added so cache service handles out of disk space. When there is <1gb left of storage it will start running maintenance to clean things up

diff --git a/Engine/Source/Programs/UnrealBuildAccelerator/CacheService/Private/UbaCacheService.cpp b/Engine/Source/Programs/UnrealBuildAccelerator/CacheService/Private/UbaCacheService.cpp
--- a/Engine/Source/Programs/UnrealBuildAccelerator/CacheService/Private/UbaCacheService.cpp
+++ b/Engine/Source/Programs/UnrealBuildAccelerator/CacheService/Private/UbaCacheService.cpp
@@ -19,6 +19,7 @@
 #pragma comment (lib, "Dbghelp.lib")
 #elif PLATFORM_LINUX
 #include <sys/prctl.h>
+#include <sys/statvfs.h>
 #endif
 
 namespace uba
@@ -500,8 +501,24 @@ namespace uba
 			while (!ShouldExit() && !cacheServer.ShouldShutdown())
 			{
 				Sleep(1000);
+
+				bool force = false;
+
+				#if PLATFORM_LINUX
+				struct statvfs stat;
+				if (statvfs(g_rootDir.data, &stat) != 0)
+				{
+					u64 available = stat.f_bsize * stat.f_bavail;
+					if (available < 1024ull * 1024 * 1024)
+					{
+						logger.Warning(TC("Running low on disk space. Only %s available. Will force maintenance"), BytesToText(available).str);
+						force = true;
+					}
+				}
+				#endif
+
 				if (maintenanceEnabled)
-					if (!cacheServer.RunMaintenance(false, allowSave, ShouldExit))
+					if (!cacheServer.RunMaintenance(force, allowSave, ShouldExit))
 						break;
 			}
 		}
diff --git a/Engine/Source/Programs/UnrealBuildAccelerator/Common/Private/UbaCacheServer.cpp b/Engine/Source/Programs/UnrealBuildAccelerator/Common/Private/UbaCacheServer.cpp
--- a/Engine/Source/Programs/UnrealBuildAccelerator/Common/Private/UbaCacheServer.cpp
+++ b/Engine/Source/Programs/UnrealBuildAccelerator/Common/Private/UbaCacheServer.cpp
@@ -606,16 +606,28 @@ namespace uba
 		do
 		{
 			SCOPED_FUTEX(m_connectionsLock, lock2);
-			if (!m_connections.empty())
+			if (!force && !m_connections.empty())
 				return true;
 			m_isRunningMaintenance = true;
+
+			bool firstLoop = true;
+			while (!m_connections.empty())
+			{
+				if (firstLoop)
+					m_logger.Info(TC("Waiting for %llu client(s) to disconnect before starting maintenance"), m_connections.size());
+				firstLoop = false;
+				lock2.Leave();
+				Sleep(200);
+				lock2.Enter();
+			}
+
 			lock2.Leave();
 
 			PrintContentionSummary(m_logger);
 
 			auto g = MakeGuard([&]()
 				{
-					SCOPED_FUTEX(m_connectionsLock, lock2);
+					SCOPED_FUTEX(m_connectionsLock, lock3);
 					m_isRunningMaintenance = false;
 				});
 
@@ -1437,7 +1449,7 @@ namespace uba
 	UBA_NOINLINE bool CacheServer::HandleMessage(const ConnectionInfo& connectionInfo, u8 messageType, BinaryReader& reader, BinaryWriter& writer)
 	{
 		if (messageType != CacheMessageType_Connect && m_isRunningMaintenance)
-			return m_logger.Error(TC("Can't handle network message %s while running maintenance mode"), ToString(CacheMessageType(messageType)));
+			return m_logger.Debug(TC("Can't handle network message %s while running maintenance mode"), ToString(CacheMessageType(messageType))).ToFalse();
 
 		switch (messageType)
 		{
