[UBA]
* Added so FileAccessor::CreateMemoryWrite is using fallocate to make sure to reserve disk space and handle out-of-diskspace instead of doing SIGBUS when copying to memory
* Changed so cache service is using overflow properly

diff --git a/Engine/Source/Programs/UnrealBuildAccelerator/CacheService/Private/UbaCacheService.cpp b/Engine/Source/Programs/UnrealBuildAccelerator/CacheService/Private/UbaCacheService.cpp
--- a/Engine/Source/Programs/UnrealBuildAccelerator/CacheService/Private/UbaCacheService.cpp
+++ b/Engine/Source/Programs/UnrealBuildAccelerator/CacheService/Private/UbaCacheService.cpp
@@ -442,7 +442,6 @@ namespace uba
 		storageInfo.casCapacityBytes = storageCapacity;
 		storageInfo.storeCompressed = storeCompressed;
 		storageInfo.allowHintAsFallback = false;
-		storageInfo.manuallyHandleOverflow = true;
 		storageInfo.writeReceivedCasFilesToDisk = true;
 		StorageServer storageServer(storageInfo);
 
diff --git a/Engine/Source/Programs/UnrealBuildAccelerator/Common/Private/UbaCacheServer.cpp b/Engine/Source/Programs/UnrealBuildAccelerator/Common/Private/UbaCacheServer.cpp
--- a/Engine/Source/Programs/UnrealBuildAccelerator/Common/Private/UbaCacheServer.cpp
+++ b/Engine/Source/Programs/UnrealBuildAccelerator/Common/Private/UbaCacheServer.cpp
@@ -216,6 +216,8 @@ namespace uba
 		m_rootDir.count = GetFullPathNameW(info.rootDir, m_rootDir.capacity, m_rootDir.data, NULL);
 		m_rootDir.Replace('/', PathSeparator).EnsureEndsWithSlash();
 
+		m_storage.m_trackedDeletes = &m_trackedDeletes;
+
 		m_server.RegisterService(CacheServiceId,
 			[this](const ConnectionInfo& connectionInfo, const WorkContext& workContext, MessageInfo& messageInfo, BinaryReader& reader, BinaryWriter& writer)
 			{
@@ -235,6 +237,8 @@ namespace uba
 
 	CacheServer::~CacheServer()
 	{
+		m_storage.m_trackedDeletes = nullptr;
+		
 		m_server.UnregisterOnClientDisconnected(CacheServiceId);
 		m_server.UnregisterService(CacheServiceId);
 	}
@@ -722,6 +726,7 @@ namespace uba
 		m_lastMaintenance = startTime;
 		
 		UnorderedSet<CasKey> deletedCasFiles;
+		deletedCasFiles.swap(m_trackedDeletes);
 		m_storage.HandleOverflow(&deletedCasFiles);
 		u64 deletedCasCount = deletedCasFiles.size();
 
diff --git a/Engine/Source/Programs/UnrealBuildAccelerator/Common/Private/UbaFileAccessor.cpp b/Engine/Source/Programs/UnrealBuildAccelerator/Common/Private/UbaFileAccessor.cpp
--- a/Engine/Source/Programs/UnrealBuildAccelerator/Common/Private/UbaFileAccessor.cpp
+++ b/Engine/Source/Programs/UnrealBuildAccelerator/Common/Private/UbaFileAccessor.cpp
@@ -58,6 +58,11 @@ namespace uba
 	}
 
 	bool FileAccessor::CreateWrite(bool allowRead, u32 flagsAndAttributes, u64 fileSize, const tchar* tempPath)
+	{
+		return InternalCreateWrite(allowRead, flagsAndAttributes, fileSize, tempPath, false);
+	}
+
+	bool FileAccessor::InternalCreateWrite(bool allowRead, u32 flagsAndAttributes, u64 fileSize, const tchar* tempPath, bool isMemoryMap)
 	{
 		UBA_ASSERT(flagsAndAttributes != 0);
 		m_size = fileSize;
@@ -150,10 +155,17 @@ namespace uba
 				info.AllocationSize = ToLargeInteger(allocSize);
 				if (!::SetFileInformationByHandle(asHANDLE(m_fileHandle), FileAllocationInfo, &info, sizeof(info)))
 					if (!IsRunningWine())
-						return m_logger.Error(TC("SetFileInformationByHandle (FileDispositionInfo) failed on %llu %s (%s)"), uintptr_t(m_fileHandle), m_fileName, LastErrorToText().data);
+						return m_logger.Error(TC("SetFileInformationByHandle (FileDispositionInfo) failed on %llu %s (%s)"), uintptr_t(m_fileHandle), realFileName, LastErrorToText().data);
 			}
 			#endif
 		}
+#elif PLATFORM_LINUX
+		if (isMemoryMap && fileSize != 0) // We want to make sure we allocate memory on the file system to catch out-of-diskspace
+		{
+			if (fallocate(asFileDescriptor(m_fileHandle), 0, 0, fileSize) == -1)
+				if (errno != ENOSYS && errno != EOPNOTSUPP)
+					return m_logger.Error(TC("fallocate failed on %i %s (%s)"), asFileDescriptor(m_fileHandle), realFileName, strerror(errno));
+		}
 #endif
 		m_isWrite = true;
 		return true;
@@ -164,7 +176,7 @@ namespace uba
 		allowRead = true; // It is not possible to have write only access to file mappings it seems
 
 		UBA_ASSERT(flagsAndAttributes != 0);
-		if (!CreateWrite(allowRead, flagsAndAttributes, size, tempPath))
+		if (!InternalCreateWrite(allowRead, flagsAndAttributes, size, tempPath, true))
 			return false;
 
 		const tchar* realFileName = m_fileName;
diff --git a/Engine/Source/Programs/UnrealBuildAccelerator/Common/Private/UbaStorage.cpp b/Engine/Source/Programs/UnrealBuildAccelerator/Common/Private/UbaStorage.cpp
--- a/Engine/Source/Programs/UnrealBuildAccelerator/Common/Private/UbaStorage.cpp
+++ b/Engine/Source/Programs/UnrealBuildAccelerator/Common/Private/UbaStorage.cpp
@@ -210,6 +210,9 @@ namespace uba
 			UBA_ASSERT(ce.key != CasKeyZero);
 			toDelete.push_back({*it, ce.size});
 
+			if (m_trackedDeletes)
+				m_trackedDeletes->insert(ce.key);
+
 			m_casEvictedBytes += ce.size;
 			m_casEvictedCount++;
 			m_casTotalBytes -= ce.size;
diff --git a/Engine/Source/Programs/UnrealBuildAccelerator/Common/Public/UbaCacheServer.h b/Engine/Source/Programs/UnrealBuildAccelerator/Common/Public/UbaCacheServer.h
--- a/Engine/Source/Programs/UnrealBuildAccelerator/Common/Public/UbaCacheServer.h
+++ b/Engine/Source/Programs/UnrealBuildAccelerator/Common/Public/UbaCacheServer.h
@@ -110,6 +110,8 @@ namespace uba
 		Futex m_connectionsLock;
 		Map<u32, Connection> m_connections;
 
+		UnorderedSet<CasKey> m_trackedDeletes;
+
 		Atomic<bool> m_shutdownRequested = false;
 
 		u64 m_maintenanceReserveSize = 0;
diff --git a/Engine/Source/Programs/UnrealBuildAccelerator/Common/Public/UbaFileAccessor.h b/Engine/Source/Programs/UnrealBuildAccelerator/Common/Public/UbaFileAccessor.h
--- a/Engine/Source/Programs/UnrealBuildAccelerator/Common/Public/UbaFileAccessor.h
+++ b/Engine/Source/Programs/UnrealBuildAccelerator/Common/Public/UbaFileAccessor.h
@@ -33,6 +33,7 @@ namespace uba
 		bool GetFileBasicInformationByHandle(FileBasicInformation& out);
 
 	private:
+		bool InternalCreateWrite(bool allowRead, u32 flagsAndAttributes, u64 size, const tchar* tempPath, bool isMemoryMap);
 		bool InternalWrite(const void* data, u64 dataLen, u64 offset, bool lastWrite);
 		bool InternalClose(bool success, u64* lastWriteTime);
 		Logger& m_logger;
diff --git a/Engine/Source/Programs/UnrealBuildAccelerator/Common/Public/UbaStorage.h b/Engine/Source/Programs/UnrealBuildAccelerator/Common/Public/UbaStorage.h
--- a/Engine/Source/Programs/UnrealBuildAccelerator/Common/Public/UbaStorage.h
+++ b/Engine/Source/Programs/UnrealBuildAccelerator/Common/Public/UbaStorage.h
@@ -236,6 +236,7 @@ namespace uba
 		};
 
 		UnorderedMap<CasKey, CasEntry> m_casLookup;
+		UnorderedSet<CasKey>* m_trackedDeletes = nullptr;
 		Futex m_accessLock;
 		CasEntry* m_newestAccessed = nullptr;
 		CasEntry* m_oldestAccessed = nullptr;
