[UBA]
* Fixed so tests running remote processes that fails will cancel the returned process

diff --git a/Engine/Source/Programs/UnrealBuildAccelerator/Test/Private/UbaTestSession.cpp b/Engine/Source/Programs/UnrealBuildAccelerator/Test/Private/UbaTestSession.cpp
--- a/Engine/Source/Programs/UnrealBuildAccelerator/Test/Private/UbaTestSession.cpp
+++ b/Engine/Source/Programs/UnrealBuildAccelerator/Test/Private/UbaTestSession.cpp
@@ -3,6 +3,7 @@
 #include "UbaConfig.h"
 #include "UbaFileAccessor.h"
 #include "UbaNetworkBackendTcp.h"
+#include "UbaProcess.h"
 #include "UbaSessionServer.h"
 #include "UbaSessionClient.h"
 #include "UbaStorageServer.h"
@@ -84,6 +85,8 @@ namespace uba
 
 		auto sg = MakeGuard([&]() { server.DisconnectClients(); });
 
+		sessionServer.SetRemoteProcessReturnedEvent([](Process& p) { p.Cancel(true); });
+
 		Config clientConfig;
 		clientConfig.AddTable(TC("Storage")).AddValue(TC("CheckExistsOnServer"), true);
 		server.SetClientsConfig(clientConfig);
