[UBA]
* Fixed remote shader compiling bug for linux

diff --git a/Engine/Source/Programs/UnrealBuildAccelerator/Detours/Private/PosixOS/UbaDetoursFunctionsPosix.cpp b/Engine/Source/Programs/UnrealBuildAccelerator/Detours/Private/PosixOS/UbaDetoursFunctionsPosix.cpp
--- a/Engine/Source/Programs/UnrealBuildAccelerator/Detours/Private/PosixOS/UbaDetoursFunctionsPosix.cpp
+++ b/Engine/Source/Programs/UnrealBuildAccelerator/Detours/Private/PosixOS/UbaDetoursFunctionsPosix.cpp
@@ -2263,6 +2263,8 @@ UBA_EXPORT void* UBA_WRAPPER(dlopen)(const char* path, int mode)
 			const char* loaderPaths[] = { "/", 0 };
 			if (g_handledLibraries.insert(path).second)
 				Shared_LoadLibrary(path, loaderPaths, tempBuf);
+			if (const char* name = TStrrchr(path, '/'))
+				path = name + 1;
 		}
 	#endif
 	}
