[UBA]
* Fixed bug in tcp SetTimeout for win64 where it was not using the correct parameter type (windows use DWORD instead of timeval. sigh

diff --git a/Engine/Source/Programs/UnrealBuildAccelerator/Common/Private/UbaNetworkBackendTcp.cpp b/Engine/Source/Programs/UnrealBuildAccelerator/Common/Private/UbaNetworkBackendTcp.cpp
--- a/Engine/Source/Programs/UnrealBuildAccelerator/Common/Private/UbaNetworkBackendTcp.cpp
+++ b/Engine/Source/Programs/UnrealBuildAccelerator/Common/Private/UbaNetworkBackendTcp.cpp
@@ -784,11 +784,14 @@ namespace uba
 
 	bool SetTimeout(Logger& logger, SOCKET socket, u32 timeoutMs)
 	{
-		u32 timeoutSeconds = timeoutMs / 1000;
-		timeoutMs = timeoutMs % 1000;
+		#if PLATFORM_WINDOWS
+		DWORD timeout = timeoutMs;
+		#else
 		struct timeval timeout;
-		timeout.tv_sec = timeoutSeconds;
-		timeout.tv_usec = timeoutMs*1000;
+		timeout.tv_sec = timeoutMs / 1000;
+		timeout.tv_usec = (timeoutMs % 1000)*1000;
+		#endif
+
 		if (setsockopt(socket, SOL_SOCKET, SO_SNDTIMEO, (const char*)&timeout, sizeof timeout) == SOCKET_ERROR)
 			return logger.Error(TC("setsockopt SO_SNDTIMEO error: (error: %s)"), LastErrorToText(WSAGetLastError()).data);
 		if (setsockopt(socket, SOL_SOCKET, SO_RCVTIMEO, (const char*)&timeout, sizeof timeout) == SOCKET_ERROR)
