[UBA]
* Added so host name is shown in agent log to make it possible to know who is using agent

diff --git a/Engine/Source/Programs/UnrealBuildAccelerator/Common/Private/UbaSessionClient.cpp b/Engine/Source/Programs/UnrealBuildAccelerator/Common/Private/UbaSessionClient.cpp
--- a/Engine/Source/Programs/UnrealBuildAccelerator/Common/Private/UbaSessionClient.cpp
+++ b/Engine/Source/Programs/UnrealBuildAccelerator/Common/Private/UbaSessionClient.cpp
@@ -1362,6 +1362,10 @@ namespace uba
 		m_shouldSendTraceToServer = reader.ReadBool();
 		m_readIntermediateFilesCompressed = reader.ReadBool();
 
+		auto& serverName = detoursFile[0].Clear(); // reuse
+		reader.ReadString(serverName);
+		m_logger.Info(TC("Connected to server %s"), serverName.data);
+
 		if (m_shouldSendTraceToServer)
 		{
 			//if (m_detailedTrace)
diff --git a/Engine/Source/Programs/UnrealBuildAccelerator/Common/Private/UbaSessionServer.cpp b/Engine/Source/Programs/UnrealBuildAccelerator/Common/Private/UbaSessionServer.cpp
--- a/Engine/Source/Programs/UnrealBuildAccelerator/Common/Private/UbaSessionServer.cpp
+++ b/Engine/Source/Programs/UnrealBuildAccelerator/Common/Private/UbaSessionServer.cpp
@@ -753,7 +753,8 @@ namespace uba
 			}
 		}
 
-		StringBuffer<> disconnectResponse;
+		StringBuffer<> tempBuffer;
+		auto& disconnectResponse = tempBuffer;
 
 		if (binAsVersion && clientAgentKey != agentBinaryKey)
 		{
@@ -821,6 +822,11 @@ namespace uba
 		writer.WriteBool(m_remoteLogEnabled);
 		writer.WriteBool(m_remoteTraceEnabled);
 		writer.WriteBool(m_readIntermediateFilesCompressed);
+
+		auto& computerName = tempBuffer.Clear();
+		GetComputerNameW(computerName);
+		writer.WriteString(computerName);
+
 		WriteRemoteEnvironmentVariables(writer);
 
 		m_trace.SessionAdded(sessionId, connectionInfo.GetId(), name, info); // Must be inside lock for TraceSessionUpdate() to not include
