[UBA]
* Fixed log formatting error
* Changing values to named constants
* Fixed so uba can compile with header units

diff --git a/Engine/Source/Programs/UnrealBuildAccelerator/Common/Private/UbaStorage.cpp b/Engine/Source/Programs/UnrealBuildAccelerator/Common/Private/UbaStorage.cpp
--- a/Engine/Source/Programs/UnrealBuildAccelerator/Common/Private/UbaStorage.cpp
+++ b/Engine/Source/Programs/UnrealBuildAccelerator/Common/Private/UbaStorage.cpp
@@ -2176,13 +2176,13 @@ namespace uba
 		m_deferredCasCreationLookup.erase(findIt);
 
 		if (casEntry.beingWritten)
-			m_logger.Warning(TC("Deferred cas %s is being written from network. This should never happen (%s)"), CasKeyString(casKey), deferredCreation.data);
+			m_logger.Warning(TC("Deferred cas %s is being written from network. This should never happen (%s)"), CasKeyString(casKey).str, deferredCreation.data);
 
 		casEntry.verified = true;
 		deferredLock.Leave();
 		WriteResult res;
 		if (!WriteCasFile(res, deferredCreation.data, casKey))
-			return m_logger.Error(TC("Failed to write deferred cas %s (%s)"), CasKeyString(casKey), deferredCreation.data);
+			return m_logger.Error(TC("Failed to write deferred cas %s (%s)"), CasKeyString(casKey).str, deferredCreation.data);
 
 		casEntry.exists = true;
 		entryLock.Leave();
diff --git a/Engine/Source/Programs/UnrealBuildAccelerator/Common/Public/UbaSessionClient.h b/Engine/Source/Programs/UnrealBuildAccelerator/Common/Public/UbaSessionClient.h
--- a/Engine/Source/Programs/UnrealBuildAccelerator/Common/Public/UbaSessionClient.h
+++ b/Engine/Source/Programs/UnrealBuildAccelerator/Common/Public/UbaSessionClient.h
@@ -18,7 +18,7 @@ namespace uba
 		NetworkClient& client;
 		StringBuffer<128> name;
 		u32 maxProcessCount = 1;
-		u32 defaultPriorityClass = 0x00004000; // BELOW_NORMAL_PRIORITY_CLASS;
+		u32 defaultPriorityClass = ProcessPriority_BelowNormal;
 		u32 maxIdleSeconds = ~0u;
 		u32 pingTimeoutSecondsPrintCallstacks = 0;
 		u32 osVersion = 0;
diff --git a/Engine/Source/Programs/UnrealBuildAccelerator/Core/Public/UbaBase.h b/Engine/Source/Programs/UnrealBuildAccelerator/Core/Public/UbaBase.h
--- a/Engine/Source/Programs/UnrealBuildAccelerator/Core/Public/UbaBase.h
+++ b/Engine/Source/Programs/UnrealBuildAccelerator/Core/Public/UbaBase.h
@@ -20,21 +20,21 @@ namespace uba
 	template<class T> const T& Max(const T& a, const T& b) { return (b > a) ? b : a; }
 
 	#if PLATFORM_WINDOWS
-	constexpr bool IsWindows = true;
+	inline constexpr bool IsWindows = true;
 	using tchar = wchar_t;
 	#define TC(x) L##x
 	#define PERCENT_HS L"%hs"
 	#else
-	constexpr bool IsWindows = false;
+	inline constexpr bool IsWindows = false;
 	using tchar = char;
 	#define TC(x) x
 	#define PERCENT_HS "%s"
 	#endif
 
 	#if PLATFORM_WINDOWS && defined(_M_ARM64) // For now we only care about windows arm..
-	constexpr bool IsArmBinary = true;
+	inline constexpr bool IsArmBinary = true;
 	#else
-	constexpr bool IsArmBinary = false;
+	inline constexpr bool IsArmBinary = false;
 	#endif
 
 #if PLATFORM_WINDOWS && defined(UBA_BUILD)
diff --git a/Engine/Source/Programs/UnrealBuildAccelerator/Host/Private/UbaExports.cpp b/Engine/Source/Programs/UnrealBuildAccelerator/Host/Private/UbaExports.cpp
--- a/Engine/Source/Programs/UnrealBuildAccelerator/Host/Private/UbaExports.cpp
+++ b/Engine/Source/Programs/UnrealBuildAccelerator/Host/Private/UbaExports.cpp
@@ -731,7 +731,7 @@ extern "C"
 		}
 
 		NetworkClientCreateInfo ncci(writer);
-		ncci.receiveTimeoutSeconds = 60;
+		ncci.receiveTimeoutSeconds = DefaultNetworkReceiveTimeoutSeconds;
 		ncci.cryptoKey128 = crypto128;
 		ncci.workerCount = 0;
 		ncci.Apply(GetConfig(), TC("CacheNetworkClient"));
