[UBA]
* WIP - Removing strip logic from UbaObjTool and replacing it with tracking normal non-exported symbols.

diff --git a/Engine/Source/Programs/UnrealBuildAccelerator/Common/Private/UbaObjectFile.cpp b/Engine/Source/Programs/UnrealBuildAccelerator/Common/Private/UbaObjectFile.cpp
--- a/Engine/Source/Programs/UnrealBuildAccelerator/Common/Private/UbaObjectFile.cpp
+++ b/Engine/Source/Programs/UnrealBuildAccelerator/Common/Private/UbaObjectFile.cpp
@@ -12,9 +12,9 @@
 
 namespace uba
 {
-	u8 SymbolFileVersion = 1;
+	u8 SymbolFileVersion = 2;
 
-	ObjectFile* ObjectFile::OpenAndParse(Logger& logger, const tchar* filename)
+	ObjectFile* ObjectFile::OpenAndParse(Logger& logger, ObjectFileParseMode parseMode, const tchar* filename)
 	{
 		auto file = new FileAccessor(logger, filename);
 		auto fileGuard = MakeGuard([&]() { delete file; });
@@ -22,7 +22,7 @@ namespace uba
 		if (!file->OpenMemoryRead())
 			return nullptr;
 
-		ObjectFile* objectFile = Parse(logger, file->GetData(), file->GetSize(), filename);
+		ObjectFile* objectFile = Parse(logger, parseMode, file->GetData(), file->GetSize(), filename);
 		if (!objectFile)
 			return nullptr;
 
@@ -31,7 +31,7 @@ namespace uba
 		return objectFile;
 	}
 
-	ObjectFile* ObjectFile::Parse(Logger& logger, u8* data, u64 dataSize, const tchar* hint)
+	ObjectFile* ObjectFile::Parse(Logger& logger, ObjectFileParseMode parseMode, u8* data, u64 dataSize, const tchar* hint)
 	{
 		ObjectFile* objectFile = nullptr;
 
@@ -94,7 +94,7 @@ namespace uba
 		objectFile->m_dataSize = dataSize;
 		objectFile->m_ownsData = ownsData;
 
-		if (objectFile->Parse(logger, hint))
+		if (objectFile->Parse(logger, parseMode, hint))
 			return objectFile;
 
 		if (ownsData)
@@ -119,38 +119,15 @@ namespace uba
 		return true;
 	}
 
-	bool ObjectFile::StripExports(Logger& logger)
-	{
-		return StripExports(logger, m_data, {});
-	}
-
-	bool ObjectFile::WriteImportsAndExports(Logger& logger, MemoryBlock& memoryBlock)
+	bool ObjectFile::WriteImportsAndExports(Logger& logger, MemoryBlock& memoryBlock, bool verbose)
 	{
 		auto write = [&](const void* data, u64 dataSize) { memcpy(memoryBlock.Allocate(dataSize, 1, TC("ObjectFile::WriteImportsAndExports")), data, dataSize); };
-
-		write(&SymbolFileVersion, 1);
-		write(&m_type, 1);
-
-		// Write all imports
-		for (auto& symbol : m_imports)
-		{
-			write(symbol.c_str(), symbol.size());
-			write("", 1);
-		}
-		write("", 1);
-
-		// Write all exports
-		for (auto& kv : m_exports)
-		{
-			write(kv.first.c_str(), kv.first.size());
-			write(kv.second.extra.c_str(), kv.second.extra.size());
-			write("", 1);
-		}
-		write("", 1);
+		if (!WriteImportsAndExports(logger, write, verbose))
+			return false;
 		return true;
 	}
 
-	bool ObjectFile::WriteImportsAndExports(Logger& logger, const tchar* exportsFilename)
+	bool ObjectFile::WriteImportsAndExports(Logger& logger, const tchar* exportsFilename, bool verbose)
 	{
 		FileAccessor exportsFile(logger, exportsFilename);
 		if (!exportsFile.CreateWrite())
@@ -161,26 +138,45 @@ namespace uba
 		auto flush = [&]() { exportsFile.Write(buffer, bufferPos); bufferPos = 0; };
 		auto write = [&](const void* data, u64 dataSize) { if (bufferPos + dataSize > sizeof(buffer)) flush(); memcpy(buffer + bufferPos, data, dataSize); bufferPos += dataSize; };
 
+		if (!WriteImportsAndExports(logger, write, verbose))
+			return false;
+
+		flush();
+
+		return exportsFile.Close();
+	}
+
+	template<typename WriteFunc>
+	bool ObjectFile::WriteImportsAndExports(Logger& logger, const WriteFunc& write, bool verbose)
+	{
+		write(&SymbolFileVersion, 1);
+		write(&m_type, 1);
+		write(&verbose, 1);
+
 		// Write all imports
 		for (auto& symbol : m_imports)
 		{
 			write(symbol.c_str(), symbol.size());
-			write("", 1);
+			write("\n", 1);
 		}
-		write("", 1);
+		write("\n", 1);
 
-		// Write all exports
-		for (auto& kv : m_exports)
+		if (verbose)
 		{
-			write(kv.first.c_str(), kv.first.size());
-			write(kv.second.extra.c_str(), kv.second.extra.size());
-			write("", 1);
+			// Write all exports
+			for (auto& kv : m_exports)
+			{
+				write(kv.second.symbol.c_str(), kv.second.symbol.size());
+				write(kv.second.extra.c_str(), kv.second.extra.size());
+				write("\n", 1);
+			}
+			write("\n", 1);
 		}
-		write("", 1);
-
-		flush();
+		else
+		{
 
-		return exportsFile.Close();
+		}
+		return true;
 	}
 
 	const char* ObjectFile::GetLibName()
@@ -198,7 +194,7 @@ namespace uba
 
 	void ObjectFile::RemoveExportedSymbol(const char* symbol)
 	{
-		m_exports.erase(symbol);
+		m_exports.erase(ToStringKeyRaw(symbol, strlen(symbol)));
 	}
 
 	const tchar* ObjectFile::GetFileName() const
@@ -266,26 +262,46 @@ namespace uba
 			return logger.Error(TC("%s - Import/export file version mismatch (application version %u, file version %u)"), filename, SymbolFileVersion, version);
 
 		type = *(const ObjectFileType*)readPos++;
+		bool verbose = *(bool*)readPos++;
 
-		while (*readPos)
+		while (*readPos != '\n')
 		{
-			auto strEnd = strlen(readPos);
-			imports.insert(std::string(readPos, readPos + strEnd));
-			readPos = readPos + strEnd + 1;
+			auto start = readPos;
+			while (*readPos != '\n')
+				++readPos;
+			imports.insert(std::string(start, readPos - start));
+			++readPos;
 		}
 		++readPos;
 
-		while (*readPos)
+		if (verbose)
 		{
-			auto strEnd = strlen(readPos);
-			ExportInfo info;
-			if (const char* comma = strchr(readPos, ','))
+			while (*readPos != '\n')
 			{
-				strEnd = comma - readPos;
-				info.extra = comma;
+				auto start = readPos;(void)start;
+				const char* comma = nullptr;
+				while (*readPos != '\n')
+				{
+					if (*readPos == ',')
+						comma = readPos;
+					++readPos;
+				}
+
+				auto end = readPos;
+
+				ExportInfo info;
+				if (comma)
+				{
+					end = comma;
+					info.extra = std::string(comma, readPos - comma);
+				}
+				//info.symbol.assign(start, end - start);
+				exports.emplace(ToStringKeyRaw(start, end - start), std::move(info));
+				++readPos;
 			}
-			exports.emplace(std::string(readPos, readPos + strEnd), info);
-			readPos = readPos + strEnd + 1;
+		}
+		else
+		{
 		}
 		return true;
 	}
@@ -298,15 +314,16 @@ namespace uba
 		WriteString("{", 1);
 
 		bool isFirst = true;
-		for (auto& symbol : allExports)
+		for (auto& kv : allExports)
 		{
+			auto& symbol = kv.second.symbol;
 			//if (strncmp(symbol.first.c_str(), "_ZTV", 4) != 0)
 			//	continue;
-			if (allExternalImports.find(symbol.first) == allExternalImports.end())
+			if (allExternalImports.find(symbol) == allExternalImports.end())
 				continue;
 			if (isFirst)
 				WriteString("global: ", 8);
-			WriteString(symbol.first.c_str(), symbol.first.size());
+			WriteString(symbol.c_str(), symbol.size());
 			WriteString(";", 1);
 			isFirst = false;
 		}
@@ -328,13 +345,16 @@ namespace uba
 		WriteString(" { export: {\r\n", 14);
 
 		bool symbolAdded = false;
-		for (auto& symbol : allExports)
-			if (allExternalImports.find(symbol.first) != allExternalImports.end())
+		for (auto& kv : allExports)
+		{
+			auto& symbol = kv.second.symbol;
+			if (allExternalImports.find(symbol) != allExternalImports.end())
 			{
-				WriteString(symbol.first.c_str(), symbol.first.size());
+				WriteString(symbol.c_str(), symbol.size());
 				WriteString("\r\n", 2);
 				symbolAdded = true;
 			}
+		}
 
 		if (!symbolAdded)
 			WriteString("ThisIsAnUnrealEngineModule\r\n", 28); // Workaround for tool not liking empty lists
diff --git a/Engine/Source/Programs/UnrealBuildAccelerator/Common/Private/UbaObjectFileCoff.cpp b/Engine/Source/Programs/UnrealBuildAccelerator/Common/Private/UbaObjectFileCoff.cpp
--- a/Engine/Source/Programs/UnrealBuildAccelerator/Common/Private/UbaObjectFileCoff.cpp
+++ b/Engine/Source/Programs/UnrealBuildAccelerator/Common/Private/UbaObjectFileCoff.cpp
@@ -8,7 +8,8 @@ namespace uba
 	constexpr u16 ImageFileMachineUnknown = 0;
 	constexpr u8 ImageSizeofShortName = 8; // IMAGE_SIZEOF_SHORT_NAME
 	constexpr u16 ImageSymClassExternal = 0x0002; // IMAGE_SYM_CLASS_EXTERNAL
-	constexpr u16 ImageSymUndefined = 0; // IMAGE_SYM_UNDEFINED
+	constexpr u16 ImageSymClassStatic = 0x0003; // IMAGE_SYM_CLASS_STATIC
+	constexpr u32 ImageSymUndefined = 0; // IMAGE_SYM_UNDEFINED
 
 	constexpr u32 ImageScnCntCode = 0x00000020; // IMAGE_SCN_CNT_CODE
 	constexpr u32 ImageScnLnkInfo = 0x00000200; // IMAGE_SCN_LNK_INFO
@@ -125,55 +126,6 @@ namespace uba
 
 	#pragma pack(pop)
 
-	// These are symbols that are added to all dlls through some macros.
-	// When merging dlls we need to remove duplicates of these
-	UnorderedSymbols ObjectFileCoff::PotentiallyDuplicatedSymbols = 
-	{
-#if 0 // This is not needed anymore
-		// REPLACEMENT_OPERATOR_NEW_AND_DELETE
-		"??2@YAPEAX_K@Z",
-		"??2@YAPEAX_KAEBUnothrow_t@std@@@Z",
-		"??2@YAPEAX_KW4align_val_t@std@@@Z",
-		"??2@YAPEAX_KW4align_val_t@std@@AEBUnothrow_t@1@@Z",
-		"??3@YAXPEAX@Z",
-		"??3@YAXPEAXAEBUnothrow_t@std@@@Z",
-		"??3@YAXPEAXW4align_val_t@std@@@Z",
-		"??3@YAXPEAXW4align_val_t@std@@AEBUnothrow_t@1@@Z",
-		"??3@YAXPEAX_K@Z",
-		"??3@YAXPEAX_KAEBUnothrow_t@std@@@Z",
-		"??3@YAXPEAX_KW4align_val_t@std@@@Z",
-		"??3@YAXPEAX_KW4align_val_t@std@@AEBUnothrow_t@1@@Z",
-		"??_U@YAPEAX_K@Z",
-		"??_U@YAPEAX_KAEBUnothrow_t@std@@@Z",
-		"??_U@YAPEAX_KW4align_val_t@std@@@Z",
-		"??_U@YAPEAX_KW4align_val_t@std@@AEBUnothrow_t@1@@Z",
-		"??_V@YAXPEAX@Z",
-		"??_V@YAXPEAXAEBUnothrow_t@std@@@Z",
-		"??_V@YAXPEAXW4align_val_t@std@@@Z",
-		"??_V@YAXPEAXW4align_val_t@std@@AEBUnothrow_t@1@@Z",
-		"??_V@YAXPEAX_K@Z",
-		"??_V@YAXPEAX_KAEBUnothrow_t@std@@@Z",
-		"??_V@YAXPEAX_KW4align_val_t@std@@@Z",
-		"??_V@YAXPEAX_KW4align_val_t@std@@AEBUnothrow_t@1@@Z",
-
-		// UE_DEFINE_FMEMORY_WRAPPERS
-		"?FMemory_Free@@YAXPEAX@Z",
-		"?FMemory_Malloc@@YAPEAX_K0@Z",
-		"?FMemory_Realloc@@YAPEAXPEAX_K1@Z",
-
-		// UE4_VISUALIZERS_HELPERS
-		"?GNameBlocksDebug@@3PEAPEAEEA",
-		"?GObjectIndexToPackedObjectRefDebug@@3AEAPEB_KEA",
-		"?GObjectArrayForDebugVisualizers@@3AEAPEAVFChunkedFixedUObjectArray@@EA",
-		"?GComplexObjectPathDebug@@3AEAPEAUFStoredObjectPathDebug@Private@CoreUObject@UE@@EA",
-		"?GObjectHandlePackageDebug@@3AEAPEAUFObjectHandlePackageDebugData@Private@CoreUObject@UE@@EA",
-
-		// Have no idea why these are duplicated
-		"??_GIFunction_OwnedObject@Function@Private@Core@UE@@UEAAPEAXI@Z",
-		"??_GIDelegateInstance@@UEAAPEAXI@Z",
-#endif
-	};
-
 	bool IsBigObj(const u8* data, u64 size)
 	{
 		if (size < sizeof(AnonObjectHeaderBigobj))
@@ -228,7 +180,7 @@ namespace uba
 		m_type = ObjectFileType_Coff;
 	}
 
-	bool ObjectFileCoff::Parse(Logger& logger, const tchar* hint)
+	bool ObjectFileCoff::Parse(Logger& logger, ObjectFileParseMode parseMode, const tchar* hint)
 	{
 		m_isBigObj = IsBigObj(m_data, m_dataSize);
 		
@@ -251,14 +203,18 @@ namespace uba
 			m_info.sectionCount = header.NumberOfSections;
 		}
 
-		if (!ParseExports())
-			return false;
-
-		if (m_isBigObj)
-			ParseImports<ImageSymbolEx>();
+		if (parseMode == ObjectFileParseMode_Exports)
+		{
+			if (!ParseExports())
+				return false;
+		}
 		else
-			ParseImports<ImageSymbol>();
-
+		{
+			if (m_isBigObj)
+				ParseAllSymbols<ImageSymbolEx>();
+			else
+				ParseAllSymbols<ImageSymbol>();
+		}
 		return true;
 	}
 
@@ -319,136 +275,59 @@ namespace uba
 				extra = comma;
 				tmp = tmp.substr(0, comma - tmp.data());
 			}
-					
-			m_exports.emplace(tmp, ExportInfo{extra, index++});
+			StringKey key = ToStringKeyRaw(tmp.data(), tmp.size());
+			auto insres = m_exports.emplace(key, ExportInfo{tmp, extra, index++});
 		}
 		return true;
 	}
 
 	template<typename SymbolType>
-	void ObjectFileCoff::ParseImports()
+	void ObjectFileCoff::ParseAllSymbols()
 	{
+		u32 index = 0;
 		std::string symbolString;
+		std::string extra;
 		auto symbols = (SymbolType*)(m_data + m_info.symbolsMemPos);
+		auto sections = (ImageSectionHeader*)(m_data + m_info.sectionsMemOffset);
+
 		for (u32 i=0; i!=m_info.symbolCount; ++i)
 		{
 			auto& symbol = symbols[i];
 
-			if (symbol.StorageClass != ImageSymClassExternal)
+			if (symbol.StorageClass != ImageSymClassExternal && symbol.StorageClass != ImageSymClassStatic)
 				continue;
 
 			AnsiStringView symbolName = GetSymbolName(symbol, m_data, m_info.stringTableMemPos);
-			symbolName.ToString(symbolString);
+
+			if (symbolName[0] == '$') // exception handling etc
+				continue;
 
 			if (symbol.SectionNumber != ImageSymUndefined)
 			{
-				if (PotentiallyDuplicatedSymbols.find(symbolString) != PotentiallyDuplicatedSymbols.end())
-				{
-					m_potentialDuplicates.emplace(symbolString);
-				}
-			}
-			else
-			{
-				m_imports.emplace(symbolString);
-			}
-		}
-	}
-
-	bool ObjectFileCoff::StripExports(Logger& logger, u8* newData, const UnorderedSymbols& allExternalImports)
-	{
-		if (!m_info.directiveSectionMemOffset)
-			return true;
-
-		auto directiveSection = (const ImageSectionHeader*)(m_data + m_info.directiveSectionMemOffset);
-		if (directiveSection->SizeOfRawData < 10)
-			return true;
-
-		const u8* directiveData = m_data + directiveSection->PointerToRawData;
-
-		auto newDirectiveSection = (ImageSectionHeader*)(newData + m_info.directiveSectionMemOffset);
-		u8* newDirectiveData = newData + newDirectiveSection->PointerToRawData;
-
-		std::string tmp;
-
-		char* writePos = (char*)newDirectiveData;
-		const char* lastCopyPos = (const char*)directiveData;
-
-		auto readPos = lastCopyPos;
-		auto readEnd = readPos + directiveSection->SizeOfRawData;
-		auto readLastPossiblePos = readEnd - 9;
-
-		while (true)
-		{
-			const char* exportStr = nullptr;
-			const char* it = readPos;
+				if (symbolName.StartsWith("__", 2)) // compiler-generated internal symbols
+					continue;
+				if (symbolName[0] == '.') // sections
+					continue;
 
-			while (it < readLastPossiblePos)
-			{
-				if (memcmp(it, "/EXPORT:", 8) != 0)
+				extra.clear();
+				if (symbol.SectionNumber != decltype(symbol.SectionNumber)(-1) && symbol.SectionNumber != decltype(symbol.SectionNumber)(-2))
 				{
-					++it;
-					continue;
+					UBA_ASSERTF(symbol.SectionNumber <= m_info.sectionCount, TC("%i / %i"), symbol.SectionNumber, m_info.sectionCount);
+					auto& section = sections[symbol.SectionNumber-1];
+					if (memcmp(section.Name, ".data", 6) == 0 || memcmp(section.Name, ".rdata", 7) == 0 || memcmp(section.Name, ".bss", 5) == 0)
+						extra.assign(",DATA");
 				}
-
-				exportStr = it;
-				break;
-			}
-
-			if (!exportStr)
-			{
-				readPos = readEnd;
-				break;
-			}
-
-			const char* startPos = exportStr;
-			exportStr += 8;
-			const char* exportEnd;
-			if (*exportStr == '\"')
-			{
-				++exportStr;
-				exportEnd = strchr(exportStr, '\"');
-				readPos = exportEnd + 1;
-				if (strncmp(readPos, ",DATA", 5) == 0)
-					readPos += 5;
-				readPos = Min(readPos, readEnd);
+				symbolName.ToString(symbolString);
+				StringKey symbolKey = ToStringKeyRaw(symbolString.data(), symbolString.size());
+				m_exports.emplace(symbolKey, ExportInfo{symbolString, extra, index++});
 			}
 			else
 			{
-				exportEnd = strchr(exportStr, ' ');
-				if (!exportEnd)
-					readPos = exportEnd = exportStr + strlen(exportStr);
-				else
-					readPos = exportEnd;
-				if (strncmp(exportEnd-5, ",DATA", 5) == 0)
-					exportEnd -= 5;
+				UBA_ASSERT(symbol.StorageClass == ImageSymClassExternal);
+				symbolName.ToString(symbolString);
+				m_imports.emplace(symbolString);
 			}
-
-			tmp.assign(exportStr, exportEnd - exportStr);
-			if (allExternalImports.find(tmp) != allExternalImports.end())
-				continue;
-			tmp.assign("__imp_").append(exportStr, exportEnd - exportStr);
-			if (allExternalImports.find(tmp) != allExternalImports.end())
-				continue;
-			
-			u64 toCopy = startPos - lastCopyPos - 1;
-			memcpy(writePos, lastCopyPos, toCopy);
-			writePos += toCopy;
-			lastCopyPos = readPos;
-			if (!*readPos)
-				break;
 		}
-
-		u64 toCopy = readPos - lastCopyPos;
-		memcpy(writePos, lastCopyPos, toCopy);
-		writePos += toCopy;
-
-		u32 sizeOfRawData = directiveSection->SizeOfRawData;
-		newDirectiveSection->SizeOfRawData = u32((u8*)writePos - newDirectiveData);
-		UBA_ASSERT(newDirectiveSection->SizeOfRawData <= sizeOfRawData);
-
-		memset(writePos, 0, sizeOfRawData - newDirectiveSection->SizeOfRawData);
-
-		return true;
 	}
 
 	bool ObjectFileCoff::CreateExtraFile(Logger& logger, const StringView& platform, MemoryBlock& memoryBlock, const UnorderedSymbols& allExternalImports, const UnorderedSymbols& allInternalImports, const UnorderedExports& allExports, bool includeExportsInFile)
@@ -462,7 +341,8 @@ namespace uba
 			if (strncmp(symbol.data(), "__imp_", 6) != 0)
 				continue;
 			tmp = symbol.substr(6);
-			if (allExports.find(tmp) == allExports.end())
+			StringKey key = ToStringKeyRaw(tmp.data(), tmp.size());
+			if (allExports.find(key) == allExports.end())
 				continue;
 			if (neededLoopbacks.insert(symbol).second)
 				totalStringSize += u32(symbol.size()) + 1;
@@ -494,10 +374,10 @@ namespace uba
 			directiveSection.Characteristics = ImageScnAlign1Bytes|ImageScnLnkInfo|ImageScnLnkRemove;
 			++header.NumberOfSections;
 
-			auto writeExport = [&](const std::string& symbol, const std::string& extra)
+			auto writeExport = [&](AnsiStringView symbol, const std::string& extra)
 				{
 					write("/EXPORT:", 8);
-					write(symbol.data(), symbol.size());
+					write(symbol.strBegin, symbol.Length());
 					write(extra.data(), extra.size());
 					write(" ", 1);
 				};
@@ -505,22 +385,41 @@ namespace uba
 			// Directive raw data
 			u32 directiveRawDataStart = u32(memoryBlock.writtenSize);
 			directiveSection.PointerToRawData = directiveRawDataStart;
+
+			for (auto& imp : allExternalImports)
+			{
+				AnsiStringView impSym(imp);
+				if (impSym.StartsWith("__imp_", 6))
+					impSym = impSym.Skip(6);
+				StringKey impKey = ToStringKeyRaw(impSym.strBegin, impSym.Length());
+				auto findIt = allExports.find(impKey);
+				if (findIt == allExports.end())
+					continue;
+				writeExport(impSym, findIt->second.extra);
+			}
+
+			/*
 			for (auto& kv : allExports)
 			{
-				auto& symbol = kv.first;
+				auto& symbol = kv.second.symbol;
+				UBA_ASSERT(!symbol.empty());
 
-				if (allExternalImports.find(symbol) == allExternalImports.end())
+				auto findIt = allExternalImports.find(symbol);
+				if (findIt == allExternalImports.end())
 				{
 					tmp.assign("__imp_").append(symbol);
-					if (allExternalImports.find(tmp) == allExternalImports.end())
+					findIt = allExternalImports.find(tmp);
+					if (findIt == allExternalImports.end())
 						continue;
 				}
 
 				writeExport(symbol, kv.second.extra);
 			}
+			*/
 
 			tmp = "ThisIsAnUnrealEngineModule";
-			if (allExports.find(tmp) != allExports.end())
+			StringKey isUnrealKey = ToStringKeyRaw(tmp.data(), tmp.size());
+			if (allExports.find(isUnrealKey) != allExports.end())
 				writeExport(tmp, ""); // Workaround for tool not liking empty lists
 
 			write("", 1);
@@ -581,31 +480,6 @@ namespace uba
 		return true;
 	}
 
-	template<typename SymbolType> void ObjectFileCoff::CalculateImports(Logger& logger, Vector<u32>& outImports)
-	{
-		// Calculate how much more memory we need by counting imports
-		auto symbols = (SymbolType*)(m_data + m_info.symbolsMemPos);
-		std::string tmp;
-
-		for (u32 i=0; i!=m_info.symbolCount; ++i)
-		{
-			auto& symbol = symbols[i];
-			if (symbol.StorageClass != ImageSymClassExternal)
-				continue;
-			if (symbol.SectionNumber != ImageSymUndefined)
-				continue;
-			AnsiStringView symbolName = GetSymbolName(symbol, m_data, m_info.stringTableMemPos);
-			if (!symbolName.StartsWith("__imp_", 6))
-				continue;
-			symbolName.strBegin += 6;
-			auto findIt = m_loopbacksToAdd.find(symbolName.ToString(tmp));
-			if (findIt == m_loopbacksToAdd.end())
-				continue;
-			m_loopbacksToAdd.erase(findIt);
-			outImports.push_back(symbol.N.Name.Long);
-		}
-	}
-
 	template<typename SymbolType>
 	void ObjectFileCoff::WriteImports(Logger& logger, u8* newData, Info& newInfo, const Vector<u32>& symbolsToAdd)
 	{
@@ -724,35 +598,4 @@ namespace uba
 			}
 		}
 	}
-
-	template<typename SymbolType> void ObjectFileCoff::RemoveSymbols(Logger& logger, u8* newData, Info& newInfo)
-	{
-		auto symbols = (SymbolType*)(newData + newInfo.symbolsMemPos);
-		auto sections = (ImageSectionHeader*)(newData + newInfo.sectionsMemOffset);
-
-		std::string tmp;
-
-		for (u32 i=0; i!=m_info.symbolCount; ++i)
-		{
-			auto& symbol = symbols[i];
-
-			if (symbol.StorageClass != ImageSymClassExternal)
-				continue;
-
-			if (!symbol.SectionNumber)
-				continue;
-
-			AnsiStringView symbolName = GetSymbolName(symbol, m_data, m_info.stringTableMemPos);
-			auto findIt = m_toRemove.find(symbolName.ToString(tmp));
-			if (findIt == m_toRemove.end())
-				continue;
-
-			auto& section = sections[symbol.SectionNumber-1];
-
-			symbol.SectionNumber = ImageSymUndefined;
-
-			if (section.Characteristics & ImageScnLnkComdat)
-				memset(&section, 0, sizeof(ImageSectionHeader));
-		}
-	}
 }
diff --git a/Engine/Source/Programs/UnrealBuildAccelerator/Common/Private/UbaObjectFileCoff.h b/Engine/Source/Programs/UnrealBuildAccelerator/Common/Private/UbaObjectFileCoff.h
--- a/Engine/Source/Programs/UnrealBuildAccelerator/Common/Private/UbaObjectFileCoff.h
+++ b/Engine/Source/Programs/UnrealBuildAccelerator/Common/Private/UbaObjectFileCoff.h
@@ -11,7 +11,7 @@ namespace uba
 	{
 	public:
 		ObjectFileCoff();
-		virtual bool Parse(Logger& logger, const tchar* hint) override;
+		virtual bool Parse(Logger& logger, ObjectFileParseMode parseMode, const tchar* hint) override;
 
 		static bool CreateExtraFile(Logger& logger, const StringView& platform, MemoryBlock& memoryBlock, const UnorderedSymbols& allExternalImports, const UnorderedSymbols& allInternalImports, const UnorderedExports& allSharedExports, bool includeExportsInFile);
 
@@ -19,13 +19,8 @@ namespace uba
 		struct Info;
 
 		bool ParseExports();
-		template<typename SymbolType> void ParseImports();
-
-		virtual bool StripExports(Logger& logger, u8* newData, const UnorderedSymbols& allExternalImports) override;
-
-		template<typename SymbolType> void CalculateImports(Logger& logger, Vector<u32>& outImports);
+		template<typename SymbolType> void ParseAllSymbols();
 		template<typename SymbolType> void WriteImports(Logger& logger, u8* newData, Info& newInfo, const Vector<u32>& symbolsToAdd);
-		template<typename SymbolType> void RemoveSymbols(Logger& logger, u8* newData, Info& newInfo);
 
 		struct Info
 		{
@@ -42,9 +37,6 @@ namespace uba
 
 		UnorderedSymbols m_loopbacksToAdd;
 		UnorderedSymbols m_toRemove;
-
-		static UnorderedSymbols PotentiallyDuplicatedSymbols;
-		static UnorderedSymbols ExportsToKeep;
 	};
 
 	bool IsCoffFile(const u8* data, u64 dataSize);
diff --git a/Engine/Source/Programs/UnrealBuildAccelerator/Common/Private/UbaObjectFileElf.cpp b/Engine/Source/Programs/UnrealBuildAccelerator/Common/Private/UbaObjectFileElf.cpp
--- a/Engine/Source/Programs/UnrealBuildAccelerator/Common/Private/UbaObjectFileElf.cpp
+++ b/Engine/Source/Programs/UnrealBuildAccelerator/Common/Private/UbaObjectFileElf.cpp
@@ -126,7 +126,7 @@ namespace uba
 		return "";
 	}
 
-	bool ObjectFileElf::Parse(Logger& logger, const tchar* hint)
+	bool ObjectFileElf::Parse(Logger& logger, ObjectFileParseMode parseMode, const tchar* hint)
 	{
 		auto& header = *(Elf64Header*)m_data;
 		if (header.e_ident[4] != 2) // 64-bit  (1 is 32-bit)
@@ -202,10 +202,15 @@ namespace uba
 						{
 							// This is an ugly hack because some platforms are not exporting vtable even though they should
 							if (symbolVisibility == STV_HIDDEN && strncmp(symbolName, "_ZTV", 4) == 0)
-								m_exports.emplace(symbolName, ExportInfo{std::string(), 0});//index++});
-
+							{
+								std::string tmp = symbolName;
+								m_exports.emplace(ToStringKeyRaw(tmp.data(), tmp.size()), ExportInfo{std::move(tmp), std::string(), 0});//index++});
+							}
 							if (m_useVisibilityForExports && symbolVisibility != STV_HIDDEN)
-								m_exports.emplace(symbolName, ExportInfo{std::string(), 0});//index++});
+							{
+								std::string tmp = symbolName;
+								m_exports.emplace(ToStringKeyRaw(tmp.data(), tmp.size()), ExportInfo{std::move(tmp), std::string(), 0});//index++});
+							}
 						}
 						else
 							m_imports.emplace(symbolName);//index++});
@@ -220,10 +225,10 @@ namespace uba
 							}
 							else
 							{
-								if (strcmp(symbolName, "_ZNK5Chaos24TBoundingVolumeHierarchyI6TArrayIPNS_7TSphereIdLi3EEE22TSizedDefaultAllocatorILi32EEES1_IiS6_EdLi3EE26FindAllIntersectionsHelperERKNS_8TBVHNodeIdLi3EEERKNS_7TVectorIdLi3EEE") == 0)
-									m_exports.emplace(symbolName, ExportInfo{std::string(), 0});//index++});
-								if (strcmp(symbolName, "_ZN5Chaos24TBoundingVolumeHierarchyI6TArrayIPNS_7TSphereIdLi3EEE22TSizedDefaultAllocatorILi32EEES1_IiS6_EdLi3EEC1ERKS7_ibd") == 0)
-									m_exports.emplace(symbolName, ExportInfo{std::string(), 0});//index++});
+								//if (strcmp(symbolName, "_ZNK5Chaos24TBoundingVolumeHierarchyI6TArrayIPNS_7TSphereIdLi3EEE22TSizedDefaultAllocatorILi32EEES1_IiS6_EdLi3EE26FindAllIntersectionsHelperERKNS_8TBVHNodeIdLi3EEERKNS_7TVectorIdLi3EEE") == 0)
+								//	m_exports.emplace(symbolName, ExportInfo{std::string(), 0});//index++});
+								//if (strcmp(symbolName, "_ZN5Chaos24TBoundingVolumeHierarchyI6TArrayIPNS_7TSphereIdLi3EEE22TSizedDefaultAllocatorILi32EEES1_IiS6_EdLi3EEC1ERKS7_ibd") == 0)
+								//	m_exports.emplace(symbolName, ExportInfo{std::string(), 0});//index++});
 							}
 						}
 					}
@@ -239,7 +244,8 @@ namespace uba
 					{
 						it += 4;
 						char* symbolName = (char*)it;
-						m_exports.emplace(symbolName, ExportInfo{std::string(), 0});
+						UBA_ASSERT(false);
+						//m_exports.emplace(symbolName, ExportInfo{std::string(), 0});
 						it += strlen(symbolName) + 1;
 					}
 				}
@@ -249,95 +255,6 @@ namespace uba
 		return true;
 	}
 
-	bool ObjectFileElf::StripExports(Logger& logger, u8* newData, const UnorderedSymbols& allExternalImports)
-	{
-		auto& header = *(Elf64Header*)newData;
-		auto sections = (Elf64SectionHeader*)(newData + header.e_shoff);
-		char* sectionNamesTable = (char*)newData + sections[header.e_shstrndx].sh_offset;
-		u64 sectionCount = header.e_shnum ? header.e_shnum : sections[0].sh_size;
-
-		if (m_useVisibilityForExports)
-		{
-			char* symTableNames = (char*)newData + m_symTableNamesOffset;
-			char* dynTableNames = (char*)newData + m_dynTableNamesOffset;
-
-			for (u64 i=0; i!=sectionCount; ++i)
-			{
-				auto& section = sections[i];
-
-				if (section.sh_type != SHT_SYMTAB && section.sh_type != SHT_DYNSYM)
-					continue;
-
-				auto symbols = (Elf64Sym*)(newData + section.sh_offset);
-				u64 symbolCount = section.sh_size / sizeof(Elf64Sym);
-				for (u64 j=0; j!=symbolCount; ++j)
-				{
-					auto& symbol = symbols[j];
-					const char* symbolName = GetSymbolName(sections, symbol, symTableNames, dynTableNames, sectionNamesTable);
-					if (!*symbolName)
-						continue;
-
-					u8 symbolType = ELF64_ST_TYPE(symbol.st_info);
-					u32 symbolBinding = ELF64_ST_BIND(symbol.st_info);
-					u32 symbolVisibility = ELF64_ST_VISIBILITY(symbol.st_other);
-
-					if (symbolBinding == STB_GLOBAL)
-					{
-						if (symbolType != STT_NOTYPE)
-						{
-							if (symbolVisibility == STV_HIDDEN && strncmp(symbolName, "_ZTV", 4) == 0)
-								symbol.st_other = STV_PROTECTED;
-							//else if (symbolVisibility != STV_HIDDEN)
-							//	symbol.st_other = STV_HIDDEN;
-						}
-					}
-					else if (symbolBinding == STB_WEAK && symbolType != STT_NOTYPE)
-					{
-						if (symbolVisibility == STV_DEFAULT)
-						{
-						}
-						else if (symbolVisibility == STV_HIDDEN)
-						{
-							if (strcmp(symbolName, "_ZNK5Chaos24TBoundingVolumeHierarchyI6TArrayIPNS_7TSphereIdLi3EEE22TSizedDefaultAllocatorILi32EEES1_IiS6_EdLi3EE26FindAllIntersectionsHelperERKNS_8TBVHNodeIdLi3EEERKNS_7TVectorIdLi3EEE") == 0)
-								symbol.st_other = STV_PROTECTED;
-							if (strcmp(symbolName, "_ZN5Chaos24TBoundingVolumeHierarchyI6TArrayIPNS_7TSphereIdLi3EEE22TSizedDefaultAllocatorILi32EEES1_IiS6_EdLi3EEC1ERKS7_ibd") == 0)
-								symbol.st_other = STV_PROTECTED;
-						}
-
-					}
-				}
-			}
-		}
-		else
-		{
-			for (u64 i=0; i!=sectionCount; ++i)
-			{
-				auto& section = sections[i];
-				char* sectionName = sectionNamesTable + section.sh_name;
-
-				// Remove all but first symbol.. 
-				// TODO Ideally we want to remove all symbols but then it seems like no stub file is created, maybe we can put a dummy symbol in there?
-				if (section.sh_type != SHT_PROGBITS || strcmp(sectionName, ".linker_cmd") != 0)
-					continue;
-
-				u64 oldSize = section.sh_size;
-				if (!oldSize)
-					return true;
-				u8* begin = newData + section.sh_offset;
-				u8* it = begin;
-
-				it += 4;
-				it += strlen((char*)it) + 1;
-
-				section.sh_size = it - begin;
-				memset(it, 0, oldSize - section.sh_size);
-
-				break;
-			}
-		}
-		return true;
-	}
-
 	bool ObjectFileElf::CreateExtraFile(Logger& logger, const StringView& platform, MemoryBlock& memoryBlock, const UnorderedSymbols& allExternalImports, const UnorderedSymbols& allInternalImports, const UnorderedExports& allExports, bool includeExportsInFile)
 	{
 		u8* data = (u8*)memoryBlock.Allocate(sizeof(Elf64Header) + sizeof(Elf64SectionHeader) + 1, 1, TC(""));
diff --git a/Engine/Source/Programs/UnrealBuildAccelerator/Common/Private/UbaObjectFileElf.h b/Engine/Source/Programs/UnrealBuildAccelerator/Common/Private/UbaObjectFileElf.h
--- a/Engine/Source/Programs/UnrealBuildAccelerator/Common/Private/UbaObjectFileElf.h
+++ b/Engine/Source/Programs/UnrealBuildAccelerator/Common/Private/UbaObjectFileElf.h
@@ -10,13 +10,11 @@ namespace uba
 	{
 	public:
 		ObjectFileElf();
-		virtual bool Parse(Logger& logger, const tchar* hint) override;
+		virtual bool Parse(Logger& logger, ObjectFileParseMode parseMode, const tchar* hint) override;
 
 		static bool CreateExtraFile(Logger& logger, const StringView& platform, MemoryBlock& memoryBlock, const UnorderedSymbols& allExternalImports, const UnorderedSymbols& allInternalImports, const UnorderedExports& allExports, bool includeExportsInFile);
 
 	private:
-		virtual bool StripExports(Logger& logger, u8* newData, const UnorderedSymbols& allExternalImports) override;
-
 		u64 m_symTableNamesOffset = 0;
 		u64 m_dynTableNamesOffset = 0;
 		bool m_useVisibilityForExports = true;
diff --git a/Engine/Source/Programs/UnrealBuildAccelerator/Common/Private/UbaObjectFileImportLib.cpp b/Engine/Source/Programs/UnrealBuildAccelerator/Common/Private/UbaObjectFileImportLib.cpp
--- a/Engine/Source/Programs/UnrealBuildAccelerator/Common/Private/UbaObjectFileImportLib.cpp
+++ b/Engine/Source/Programs/UnrealBuildAccelerator/Common/Private/UbaObjectFileImportLib.cpp
@@ -9,7 +9,7 @@ namespace uba
 		return dataSize > 5 && memcmp(data, "!<arch>", 5) == 0;
 	}
 
-	bool ObjectFileImportLib::Parse(Logger& logger, const tchar* hint)
+	bool ObjectFileImportLib::Parse(Logger& logger, ObjectFileParseMode parseMode, const tchar* hint)
 	{
 #if PLATFORM_WINDOWS
 
@@ -24,7 +24,7 @@ namespace uba
 		u32* symbolOffsets = (u32*)pos;
 		pos += sizeof(u32)*symbolCount;
 
-		Vector<char*> impSymbols;
+		Vector<StringKey> impSymbols;
 		for (u32 i=0; i!=symbolCount; ++i)
 		{
 			u32 symbolOffset = _byteswap_ulong(symbolOffsets[i]);(void)symbolOffset;
@@ -39,13 +39,16 @@ namespace uba
 				continue;
 			if (strncmp(symbol.data(), "__imp_", 6) == 0)
 			{
-				impSymbols.push_back(symbolStr + 6);
+				//impSymbols.push_back(symbolStr + 6);
+				impSymbols.push_back(ToStringKeyRaw(symbolStr + 6, symbol.size() - 6));
 				continue;
 			}
-			m_exports.try_emplace(std::move(symbol), ExportInfo{ ",DATA", i });
+
+			StringKey key = ToStringKeyRaw(symbol.data(), symbol.size());
+			m_exports.try_emplace(key, ExportInfo{ std::move(symbol),",DATA", i });
 		}
 
-		for (auto symbol : impSymbols)
+		for (auto& symbol : impSymbols)
 		{
 			auto findIt = m_exports.find(symbol);
 			if (findIt != m_exports.end())
@@ -62,9 +65,4 @@ namespace uba
 	{
 		return m_libName.c_str();
 	}
-
-	bool ObjectFileImportLib::StripExports(Logger& logger, u8* newData, const UnorderedSymbols& allNeededImports)
-	{
-		return logger.Error(TC("Stripping exports from import lib file not supported"));
-	}
 }
diff --git a/Engine/Source/Programs/UnrealBuildAccelerator/Common/Private/UbaObjectFileImportLib.h b/Engine/Source/Programs/UnrealBuildAccelerator/Common/Private/UbaObjectFileImportLib.h
--- a/Engine/Source/Programs/UnrealBuildAccelerator/Common/Private/UbaObjectFileImportLib.h
+++ b/Engine/Source/Programs/UnrealBuildAccelerator/Common/Private/UbaObjectFileImportLib.h
@@ -9,12 +9,10 @@ namespace uba
 	class ObjectFileImportLib : public ObjectFile
 	{
 	public:
-		virtual bool Parse(Logger& logger, const tchar* hint) override;
+		virtual bool Parse(Logger& logger, ObjectFileParseMode parseMode, const tchar* hint) override;
 		virtual const char* GetLibName() override;
 
 	private:
-		virtual bool StripExports(Logger& logger, u8* newData, const UnorderedSymbols& allNeededImports) override;
-
 		std::string m_libName;
 	};
 
diff --git a/Engine/Source/Programs/UnrealBuildAccelerator/Common/Private/UbaObjectFileLLVMIR.cpp b/Engine/Source/Programs/UnrealBuildAccelerator/Common/Private/UbaObjectFileLLVMIR.cpp
--- a/Engine/Source/Programs/UnrealBuildAccelerator/Common/Private/UbaObjectFileLLVMIR.cpp
+++ b/Engine/Source/Programs/UnrealBuildAccelerator/Common/Private/UbaObjectFileLLVMIR.cpp
@@ -836,7 +836,7 @@ namespace uba
 		return ~0u;
 	}
 
-	bool ObjectFileLLVMIR::Parse(Logger& logger, const tchar* hint)
+	bool ObjectFileLLVMIR::Parse(Logger& logger, ObjectFileParseMode parseMode, const tchar* hint)
 	{
 		BitStreamReader reader(*this, logger, m_data, m_dataSize);
 
@@ -912,7 +912,7 @@ namespace uba
 			{
 				UBA_ASSERT(!name.empty());
 				record.dllStorage = DllStorage_Export;
-				m_exports.emplace(name, ExportInfo{"", index++});
+				m_exports.emplace(ToStringKeyRaw(name.data(), name.size()), ExportInfo{name, "", index++});
 			}
 			else // Check visibility... might be a non-windows based llvm ir stream file
 			{
@@ -940,7 +940,8 @@ namespace uba
 					if (linkage == 5 || linkage == 6 || linkage == 15)
 					{
 						record.dllStorage = DllStorage_Export;
-						m_exports.emplace(name, ExportInfo{"", index++});
+						StringKey key = ToStringKeyRaw(name.data(), name.size());
+						m_exports.emplace(key, ExportInfo{name, "", index++});
 					}
 					else if (linkage == 1 || linkage == 16)
 					{
@@ -1034,110 +1035,4 @@ namespace uba
 
 		return true;
 	}
-
-	bool ObjectFileLLVMIR::StripExports(Logger& logger, u8* newData, const UnorderedSymbols& allNeededImports)
-	{ 
-		if (m_exports.empty())
-			return true;
-
-		for (auto& record : m_globalVarOrFunctionRecords)
-		{
-			if (record.dllStorage == DllStorage_None || record.keepAsIs)
-				continue;
-	
-			BitStreamReader reader(*this, logger, newData, m_dataSize);
-			reader.m_pos = newData + record.pos;
-			reader.m_word = record.word;
-			reader.m_wordBits = record.wordBits;
-
-			u32 dllStorageIndex = GetDllStorageIndex(record.code);
-
-			UBA_ASSERT(record.operands.empty());
-
-			reader.ReadVBR(6);
-			for (u32 i=0;i!=dllStorageIndex; ++i)
-				reader.ReadVBR64(6);
-
-			// Write 0 to dll entry (no import/export)
-			//u64 res = reader.ReadVBR64(6);
-			//UBA_ASSERT(res == 2);
-			reader.ClearVBR64(6);
-
-#if 0
-			BitStreamReader reader2(*this, logger, newData, m_dataSize);
-			reader2.m_pos = newData + record.pos;
-			reader2.m_word = record.word;
-			reader2.m_wordBits = record.wordBits;
-			reader2.ReadVBR(6);
-			for (u32 i=0;i!=dllStorageIndex; ++i)
-				reader2.ReadVBR64(6);
-			u64 res2 = reader2.ReadVBR64(6);
-			UBA_ASSERT(res2  == 0);
-#endif
-		}
-
-		std::string tmp;
-
-		char* readPos = (char*)newData + m_strTabPos;
-		auto readEnd = readPos + m_strTabSize;
-		auto readLastPossiblePos = readEnd - 9;
-
-		while (true)
-		{
-			char* exportStr = nullptr;
-			char* it = readPos;
-
-			while (it < readLastPossiblePos)
-			{
-				if (memcmp(it, "/EXPORT:", 8) != 0)
-				{
-					++it;
-					continue;
-				}
-
-				exportStr = it;
-				break;
-			}
-
-			if (!exportStr)
-			{
-				//readPos = readEnd;
-				break;
-			}
-
-			char* startPos = exportStr;
-			exportStr += 8;
-			char* exportEnd;
-			if (*exportStr == '\"')
-			{
-				++exportStr;
-				exportEnd = strchr(exportStr, '\"');
-				readPos = exportEnd + 1;
-				if (strncmp(readPos, ",DATA", 5) == 0)
-					readPos += 5;
-				readPos = Min(readPos, readEnd);
-			}
-			else
-			{
-				exportEnd = strchr(exportStr, ' ');
-				if (!exportEnd)
-					readPos = exportEnd = exportStr + strlen(exportStr);
-				else
-					readPos = exportEnd;
-				if (strncmp(exportEnd-5, ",DATA", 5) == 0)
-					exportEnd -= 5;
-			}
-
-			tmp.assign(exportStr, exportEnd - exportStr);
-			if (allNeededImports.find(tmp) != allNeededImports.end())
-				continue;
-			tmp.assign("__imp_").append(exportStr, exportEnd - exportStr);
-			if (allNeededImports.find(tmp) != allNeededImports.end())
-				continue;
-			
-			memset(startPos, ' ', readPos - startPos);
-		}
-
-		return true;
-	}
 }
diff --git a/Engine/Source/Programs/UnrealBuildAccelerator/Common/Private/UbaObjectFileLLVMIR.h b/Engine/Source/Programs/UnrealBuildAccelerator/Common/Private/UbaObjectFileLLVMIR.h
--- a/Engine/Source/Programs/UnrealBuildAccelerator/Common/Private/UbaObjectFileLLVMIR.h
+++ b/Engine/Source/Programs/UnrealBuildAccelerator/Common/Private/UbaObjectFileLLVMIR.h
@@ -13,11 +13,9 @@ namespace uba
 		ObjectFileLLVMIR();
 		~ObjectFileLLVMIR();
 
-		virtual bool Parse(Logger& logger, const tchar* hint) override;
+		virtual bool Parse(Logger& logger, ObjectFileParseMode parseMode, const tchar* hint) override;
 
 	private:
-		virtual bool StripExports(Logger& logger, u8* newData, const UnorderedSymbols& allNeededImports) override;
-
 		enum BlockInfoCodes : u8;
 		enum Encoding : u8;
 		enum EntryKind : u8;
diff --git a/Engine/Source/Programs/UnrealBuildAccelerator/Common/Private/UbaSession.cpp b/Engine/Source/Programs/UnrealBuildAccelerator/Common/Private/UbaSession.cpp
--- a/Engine/Source/Programs/UnrealBuildAccelerator/Common/Private/UbaSession.cpp
+++ b/Engine/Source/Programs/UnrealBuildAccelerator/Common/Private/UbaSession.cpp
@@ -1370,10 +1370,10 @@ namespace uba
 
 		if (attr == INVALID_FILE_ATTRIBUTES)
 		{
-			#if PLATFORM_WINDOWS
+#if PLATFORM_WINDOWS
 			if (!SearchPathW(NULL, library, NULL, 512, temp, NULL))
 				return true; // TODO: We have to return true here because there are scenarios where failing is actually ok (it seems it can return false on crt shim libraries such as api-ms-win-crt*)
-			#elif PLATFORM_MAC
+#elif PLATFORM_MAC
 			if (!loaderPaths)
 				return m_logger.Error("CopyImports - Failed to find file %s (no loader paths)", applicationName);
 			u32 loaderPathCount = 0;
@@ -1393,8 +1393,13 @@ namespace uba
 				return m_logger.Error("CopyImports - Failed to find file %s (%u loader paths)", applicationName, loaderPathCount);
 			//UBA_ASSERTF(false, TC("DIR NOT FOUND: %s"), applicationName);
 			#else
-			return m_logger.Error("Code path not implemented for linux! (CopyImports %s)", library);
+			// SOOO, for linux it might be fine not finding these files.
+			// DT_NEEDED does not mean it is needed(!).. and no way of knowing which ones that are needed
+			#if UBA_DEBUG
+			m_logger.Error("Code path not implemented for linux! (CopyImports %s)", library);
 			#endif
+			return true;
+#endif
 
 			applicationName = temp;
 			attr = DefaultAttributes();
@@ -3822,21 +3827,19 @@ namespace uba
 			return m_logger.Error(TC("Failed to map view of filehandle for read %s (%s)"), fileName, LastErrorToText().data);
 		auto memClose = MakeGuard([&](){ UnmapViewOfFile(m_logger, mem, fileSize, fileName); });
 
-		ObjectFile* objectFile = ObjectFile::Parse(m_logger, mem, fileSize, fileName);
+		ObjectFile* objectFile = ObjectFile::Parse(m_logger, ObjectFileParseMode_All, mem, fileSize, fileName);
 		if (!objectFile)
 			return false;
 		auto ofg = MakeGuard([&]() { delete objectFile; });
 
-		if (!objectFile->StripExports(m_logger))
-			return false;
-
 		const tchar* lastDot = TStrrchr(fileName, '.');
 		UBA_ASSERT(lastDot);
 		StringBuffer<> exportsFile;
 		exportsFile.Append(fileName, lastDot - fileName).Append(TCV(".exi"));
 
-		MemoryBlock memoryBlock(8*1024*1024);
-		if (!objectFile->WriteImportsAndExports(m_logger, memoryBlock))
+		bool verbose = true;
+		MemoryBlock memoryBlock(32*1024*1024);
+		if (!objectFile->WriteImportsAndExports(m_logger, memoryBlock, verbose))
 			return false;
 
 		FileMappingHandle symHandle = CreateMemoryMappingW(m_logger, PAGE_READWRITE, memoryBlock.writtenSize, nullptr, TC("SymHandle"));
diff --git a/Engine/Source/Programs/UnrealBuildAccelerator/Common/Public/UbaObjectFile.h b/Engine/Source/Programs/UnrealBuildAccelerator/Common/Public/UbaObjectFile.h
--- a/Engine/Source/Programs/UnrealBuildAccelerator/Common/Public/UbaObjectFile.h
+++ b/Engine/Source/Programs/UnrealBuildAccelerator/Common/Public/UbaObjectFile.h
@@ -2,6 +2,7 @@
 
 #pragma once
 
+#include "UbaHash.h"
 #include "UbaLogger.h"
 #include "UbaStringBuffer.h"
 
@@ -11,8 +12,8 @@ namespace uba
 	
 	using UnorderedSymbols = UnorderedSet<std::string>;
 
-	struct ExportInfo { std::string extra; u32 index = 0; };
-	using UnorderedExports = UnorderedMap<std::string, ExportInfo>;
+	struct ExportInfo { std::string symbol; std::string extra; u32 index = 0; };
+	using UnorderedExports = UnorderedMap<StringKey, ExportInfo>;
 
 	enum ObjectFileType : u8
 	{
@@ -22,16 +23,23 @@ namespace uba
 		ObjectFileType_LLVMIR,
 	};
 
+	enum ObjectFileParseMode : u8
+	{
+		ObjectFileParseMode_Exports,
+		ObjectFileParseMode_All,
+	};
+
 	class ObjectFile
 	{
 	public:
-		static ObjectFile* OpenAndParse(Logger& logger, const tchar* hint);
-		static ObjectFile* Parse(Logger& logger, u8* data, u64 dataSize, const tchar* hint);
+		static ObjectFile* OpenAndParse(Logger& logger, ObjectFileParseMode parseMode, const tchar* fileName);
+		static ObjectFile* Parse(Logger& logger, ObjectFileParseMode parseMode, u8* data, u64 dataSize, const tchar* hint);
 
 		virtual bool CopyMemoryAndClose();
-		virtual bool StripExports(Logger& logger);
-		virtual bool WriteImportsAndExports(Logger& logger, MemoryBlock& memoryBlock);
-		virtual bool WriteImportsAndExports(Logger& logger, const tchar* exportsFilename);
+		bool WriteImportsAndExports(Logger& logger, MemoryBlock& memoryBlock, bool verbose);
+		bool WriteImportsAndExports(Logger& logger, const tchar* exportsFilename, bool verbose);
+		template<typename WriteFunc> bool WriteImportsAndExports(Logger& logger, const WriteFunc& write, bool verbose);
+
 		virtual const char* GetLibName();
 
 		const tchar* GetFileName() const;
@@ -48,8 +56,7 @@ namespace uba
 		u64 GetDataSize() { return m_dataSize; }
 
 	protected:
-		virtual bool Parse(Logger& logger, const tchar* hint) = 0;
-		virtual bool StripExports(Logger& logger, u8* newData, const UnorderedSymbols& allExternalImports) = 0;
+		virtual bool Parse(Logger& logger, ObjectFileParseMode parseMode, const tchar* hint) = 0;
 
 		static bool CreateDynamicListFile(Logger& logger, MemoryBlock& memoryBlock, const UnorderedSymbols& allExternalImports, const UnorderedSymbols& allInternalImports, const UnorderedExports& allExports, bool includeExportsInFile);
 		static bool CreateEmdFile(Logger& logger, MemoryBlock& memoryBlock, const StringView& moduleName, const UnorderedSymbols& allExternalImports, const UnorderedSymbols& allInternalImports, const UnorderedExports& allExports, bool includeExportsInFile);
@@ -70,6 +77,11 @@ namespace uba
 			const char* strBegin;
 			const char* strEnd;
 
+			char operator[](u64 pos) const
+			{
+				return strBegin[pos];
+			}
+
 			u32 Length() const
 			{
 				return u32(strEnd - strBegin);
@@ -112,6 +124,11 @@ namespace uba
 				out.assign(strBegin, strEnd);
 				return out;
 			}
+
+			AnsiStringView Skip(u64 count) const { return { strBegin + count, strEnd }; }
+
+			AnsiStringView(const char* b, const char* e) : strBegin(b), strEnd(e) {}
+			AnsiStringView(const std::string& s) : strBegin(s.data()), strEnd(strBegin + s.size()) {}
 		};
 	};
 
diff --git a/Engine/Source/Programs/UnrealBuildAccelerator/Core/Private/UbaHash.cpp b/Engine/Source/Programs/UnrealBuildAccelerator/Core/Private/UbaHash.cpp
--- a/Engine/Source/Programs/UnrealBuildAccelerator/Core/Private/UbaHash.cpp
+++ b/Engine/Source/Programs/UnrealBuildAccelerator/Core/Private/UbaHash.cpp
@@ -84,6 +84,16 @@ namespace uba
 		return (StringKey&)output;
 	}
 
+	StringKey ToStringKeyRaw(const void* data, u64 dataLen)
+	{
+		blake3_hasher hasher;
+		blake3_hasher_init(&hasher);
+		blake3_hasher_update(&hasher, data, dataLen);
+		uint8_t output[BLAKE3_OUT_LEN];
+		blake3_hasher_finalize(&hasher, output, BLAKE3_OUT_LEN);
+		return (StringKey&)output;
+	}
+
 	CasKeyHasher::CasKeyHasher()
 	{
 		static_assert(sizeof(hasher) == sizeof(blake3_hasher));
diff --git a/Engine/Source/Programs/UnrealBuildAccelerator/Core/Public/UbaHash.h b/Engine/Source/Programs/UnrealBuildAccelerator/Core/Public/UbaHash.h
--- a/Engine/Source/Programs/UnrealBuildAccelerator/Core/Public/UbaHash.h
+++ b/Engine/Source/Programs/UnrealBuildAccelerator/Core/Public/UbaHash.h
@@ -109,6 +109,7 @@ namespace uba
 	StringKey ToStringKey(const StringKeyHasher& hasher, const tchar* str, u64 strLen);
 	StringKey ToStringKey(const StringKeyHasher& hasher);
 	StringKey ToStringKeyNoCheck(const tchar* str, u64 strLen);
+	StringKey ToStringKeyRaw(const void* data, u64 dataLen);
 
 	constexpr inline u64 AlignUp(u64 v, u64 a) { return ((v + a - 1)/a) * a; }
 	constexpr u64 InvalidValue = 0x7ffffffff; // ~34gb - the largest number that can 7bit encode to 5 bytes
diff --git a/Engine/Source/Programs/UnrealBuildAccelerator/ObjTool/Private/UbaImportLibWriter.cpp b/Engine/Source/Programs/UnrealBuildAccelerator/ObjTool/Private/UbaImportLibWriter.cpp
--- a/Engine/Source/Programs/UnrealBuildAccelerator/ObjTool/Private/UbaImportLibWriter.cpp
+++ b/Engine/Source/Programs/UnrealBuildAccelerator/ObjTool/Private/UbaImportLibWriter.cpp
@@ -582,7 +582,7 @@ namespace uba
 
 			InitSymbolList(libName);
 
-			UnorderedSymbols handled;
+			UnorderedSet<StringKey> handled;
 			for (auto objFile : objFiles)
 			{
 				if (!objFile)
@@ -593,7 +593,7 @@ namespace uba
 					if (handled.insert(kv.first).second)
 						sortedSymbols.try_emplace(kv.second.index, &kv);
 				for (auto& outerKv : sortedSymbols)
-					AddSymbol(outerKv.second->first.c_str(), outerKv.second->second.extra.c_str(), 0);
+					AddSymbol(outerKv.second->second.symbol.c_str(), outerKv.second->second.extra.c_str(), 0);
 			}
 
 			m_moduleName = std::string(libName) + libExt;
diff --git a/Engine/Source/Programs/UnrealBuildAccelerator/ObjTool/Private/UbaObjTool.cpp b/Engine/Source/Programs/UnrealBuildAccelerator/ObjTool/Private/UbaObjTool.cpp
--- a/Engine/Source/Programs/UnrealBuildAccelerator/ObjTool/Private/UbaObjTool.cpp
+++ b/Engine/Source/Programs/UnrealBuildAccelerator/ObjTool/Private/UbaObjTool.cpp
@@ -43,7 +43,7 @@ namespace uba
 		logger.Info(TC("  UbaObjTool.exe @<rspfile>"));
 		logger.Info(TC(""));
 		logger.Info(TC("   Response file options:"));
-		logger.Info(TC("    /S:<objfile>             Obj file to strip. Will produce a .strip.obj file. Multiple allowed"));
+		logger.Info(TC("    /S:<objfile>             Obj file to export from. Multiple allowed"));
 		logger.Info(TC("    /D:<objfile>             Obj file depending on obj files to strip. Multiple allowed"));
 		logger.Info(TC("    /O:<objfile>             Obj file to output containing exports and loopbacks"));
 		logger.Info(TC("    /T:<platform>            Target platform"));
@@ -68,12 +68,11 @@ namespace uba
 
 		TString objFile;
 		bool printSymbols = false;
-		bool stripExports = false;
 		bool writeImpLib = false;
 		bool allowLibInputs = false;
 		bool isImpLibRsp = false;
 
-		Vector<TString> objFilesToStrip;
+		Vector<TString> objFilesToExport;
 		Vector<TString> objFilesDependencies;
 		TString extraObjFile;
 		Vector<TString> objFilesForImpLib;
@@ -168,7 +167,7 @@ namespace uba
 				}
 				else if (name.StartsWith(TC("/S")))
 				{
-					objFilesToStrip.push_back(value.data);
+					objFilesToExport.push_back(value.data);
 				}
 				else if (name.StartsWith(TC("/O")))
 				{
@@ -192,10 +191,6 @@ namespace uba
 					writeImpLib = true;
 					allowLibInputs = true;
 				}
-				else if (name.Equals(TCV("-stripexports")))
-				{
-					stripExports = true;
-				}
 				else if (name.Equals(TCV("/LIB")))
 				{
 					isImpLibRsp = true;
@@ -252,7 +247,7 @@ namespace uba
 		FilteredLogWriter logWriter(g_consoleLogWriter, LogEntryType_Info);
 		LoggerWithWriter logger(logWriter, TC("UbaObjTool"));
 
-		if (!objFilesToStrip.empty())
+		if (!objFilesToExport.empty())
 		{
 			CriticalSection cs;
 			Atomic<bool> success = true;
@@ -285,7 +280,7 @@ namespace uba
 			Map<TString, ObjectFile*> objectFiles;
 			auto g = MakeGuard([&]() { for (auto& kv : objectFiles) delete kv.second; });
 
-			workManager.ParallelFor(workerCount, objFilesToStrip, [&](const WorkContext&, auto& it)
+			workManager.ParallelFor(workerCount, objFilesToExport, [&](const WorkContext&, auto& it)
 				{
 					const TString& exiFilename = *it;
 					SymbolFile symbolFile;
@@ -346,7 +341,7 @@ namespace uba
 				FixPath(o.c_str(), currentDir.data, currentDir.count, fixedPath);
 				if (fixedPath.EndsWith(TCV(".res")) || (fixedPath.EndsWith(TCV(".lib")) && !allowLibInputs))
 					return;
-				if (ObjectFile* objectFile = ObjectFile::OpenAndParse(logger, fixedPath.data))
+				if (ObjectFile* objectFile = ObjectFile::OpenAndParse(logger, ObjectFileParseMode_Exports, fixedPath.data))
 				{
 					//objectFile->RemoveExportedSymbol("DllMain"); // This is used to patch xinput.lib
 					objFiles[it - objFilesForImpLib.begin()] = objectFile;
@@ -367,7 +362,7 @@ namespace uba
 			if (objFile.empty())
 				return PrintHelp(TC("No obj, lib or rsp file provided"));
 
-			ObjectFile* objectFile = ObjectFile::OpenAndParse(logger, objFile.c_str());
+			ObjectFile* objectFile = ObjectFile::OpenAndParse(logger, ObjectFileParseMode_All, objFile.c_str());
 			if (!objectFile)
 				return -1;
 			auto g = MakeGuard([&](){ delete objectFile; });
@@ -378,33 +373,7 @@ namespace uba
 					logger.Info(TC("I %S"), symbol.c_str());
 
 				for (auto& kv : objectFile->GetExports())
-					logger.Info(TC("E %S%S"), kv.first.c_str(), kv.second.extra.c_str());
-			}
-
-			if (stripExports)
-			{
-				if (!objectFile->CopyMemoryAndClose())
-					return false;
-
-				const tchar* fileName = objFile.c_str();
-				const tchar* lastDot = TStrrchr(fileName, '.');
-				UBA_ASSERT(lastDot);
-				StringBuffer<> exportsFile;
-				exportsFile.Append(fileName, lastDot - fileName).Append(TCV(".exi"));
-
-				if (!objectFile->WriteImportsAndExports(logger, exportsFile.data))
-					return false;
-
-#if 0
-				objectFile->StripExports(logger);
-				StringBuffer<> newFilename;
-				newFilename.Append(fileName, lastDot - fileName).Append(TCV(".TEST")).Append(lastDot);
-				FileAccessor fa(logger, newFilename.data);
-				fa.CreateWrite();
-				fa.Write(objectFile->GetData(), objectFile->GetDataSize());
-				fa.Close();
-
-#endif
+					logger.Info(TC("E %S%S"), kv.second.symbol.c_str(), kv.second.extra.c_str());
 			}
 		}
 		return 0;
