[UBA]
* Fixed so .exi format supports storing exports as keys
* Enabled keys by default
* Optimized extra file generation

diff --git a/Engine/Source/Programs/UnrealBuildAccelerator/Common/Private/UbaObjectFile.cpp b/Engine/Source/Programs/UnrealBuildAccelerator/Common/Private/UbaObjectFile.cpp
--- a/Engine/Source/Programs/UnrealBuildAccelerator/Common/Private/UbaObjectFile.cpp
+++ b/Engine/Source/Programs/UnrealBuildAccelerator/Common/Private/UbaObjectFile.cpp
@@ -167,14 +167,21 @@ namespace uba
 			for (auto& kv : m_exports)
 			{
 				write(kv.second.symbol.c_str(), kv.second.symbol.size());
-				write(kv.second.extra.c_str(), kv.second.extra.size());
+				if (kv.second.isData)
+					write(",DATA", 5);
 				write("\n", 1);
 			}
 			write("\n", 1);
 		}
 		else
 		{
-
+			u32 count = u32(m_exports.size());
+			write(&count, sizeof(count));
+			for (auto& kv : m_exports)
+			{
+				write(&kv.first, sizeof(StringKey));
+				write(&kv.second.isData, 1);
+			}
 		}
 		return true;
 	}
@@ -217,7 +224,7 @@ namespace uba
 		return m_potentialDuplicates;
 	}
 
-	bool ObjectFile::CreateExtraFile(Logger& logger, const StringView& extraObjFilename, const StringView& moduleName, const StringView& platform, const UnorderedSymbols& allExternalImports, const UnorderedSymbols& allInternalImports, const UnorderedExports& allExports, bool includeExportsInFile)
+	bool ObjectFile::CreateExtraFile(Logger& logger, const StringView& extraObjFilename, const StringView& moduleName, const StringView& platform, const AllExternalImports& allExternalImports, const UnorderedSymbols& allInternalImports, const AllExports& allExports, bool includeExportsInFile)
 	{
 		ObjectFileCoff objectFileCoff;
 		ObjectFileElf objectFileElf;
@@ -293,7 +300,7 @@ namespace uba
 				if (comma)
 				{
 					end = comma;
-					info.extra = std::string(comma, readPos - comma);
+					info.isData = true;
 				}
 				//info.symbol.assign(start, end - start);
 				exports.emplace(ToStringKeyRaw(start, end - start), std::move(info));
@@ -302,11 +309,23 @@ namespace uba
 		}
 		else
 		{
+			u32 count = *(u32*)readPos;
+			readPos += sizeof(u32);
+			while (count--)
+			{
+				StringKey key = *(StringKey*)readPos;
+				readPos += sizeof(StringKey);
+				u8 isData = *readPos != 0;
+				++readPos;
+				ExportInfo info;
+				info.isData = isData;
+				exports.emplace(key, info);
+			}
 		}
 		return true;
 	}
 
-	bool ObjectFile::CreateDynamicListFile(Logger& logger, MemoryBlock& memoryBlock, const UnorderedSymbols& allExternalImports, const UnorderedSymbols& allInternalImports, const UnorderedExports& allExports, bool includeExportsInFile)
+	bool ObjectFile::CreateDynamicListFile(Logger& logger, MemoryBlock& memoryBlock, const AllExternalImports& allExternalImports, const UnorderedSymbols& allInternalImports, const AllExports& allExports, bool includeExportsInFile)
 	{
 		auto WriteString = [&](const char* str, u64 strLen) { memcpy(memoryBlock.Allocate(strLen, 1, TC("")), str, strLen); };
 
@@ -333,7 +352,7 @@ namespace uba
 		return true;
 	}
 
-	bool ObjectFile::CreateEmdFile(Logger& logger, MemoryBlock& memoryBlock, const StringView& moduleName, const UnorderedSymbols& allExternalImports, const UnorderedSymbols& allInternalImports, const UnorderedExports& allExports, bool includeExportsInFile)
+	bool ObjectFile::CreateEmdFile(Logger& logger, MemoryBlock& memoryBlock, const StringView& moduleName, const AllExternalImports& allExternalImports, const UnorderedSymbols& allInternalImports, const AllExports& allExports, bool includeExportsInFile)
 	{
 		auto WriteString = [&](const char* str, u64 strLen) { memcpy(memoryBlock.Allocate(strLen, 1, TC("")), str, strLen); };
 
diff --git a/Engine/Source/Programs/UnrealBuildAccelerator/Common/Private/UbaObjectFileCoff.cpp b/Engine/Source/Programs/UnrealBuildAccelerator/Common/Private/UbaObjectFileCoff.cpp
--- a/Engine/Source/Programs/UnrealBuildAccelerator/Common/Private/UbaObjectFileCoff.cpp
+++ b/Engine/Source/Programs/UnrealBuildAccelerator/Common/Private/UbaObjectFileCoff.cpp
@@ -239,7 +239,6 @@ namespace uba
 		UBA_ASSERT(memcmp(directiveData, utf8Bom, 3) != 0);
 
 		std::string tmp;
-		std::string extra;
 
 		u32 index = 0;
 
@@ -269,14 +268,14 @@ namespace uba
 			}
 
 			tmp.assign(exportStr, exportEnd);
-			extra.clear();
+			bool isData = false;
 			if (const char* comma = strchr(tmp.c_str(), ','))
 			{
-				extra = comma;
+				isData = true;
 				tmp = tmp.substr(0, comma - tmp.data());
 			}
 			StringKey key = ToStringKeyRaw(tmp.data(), tmp.size());
-			auto insres = m_exports.emplace(key, ExportInfo{tmp, extra, index++});
+			auto insres = m_exports.emplace(key, ExportInfo{tmp, isData, index++});
 		}
 		return true;
 	}
@@ -286,7 +285,7 @@ namespace uba
 	{
 		u32 index = 0;
 		std::string symbolString;
-		std::string extra;
+
 		auto symbols = (SymbolType*)(m_data + m_info.symbolsMemPos);
 		auto sections = (ImageSectionHeader*)(m_data + m_info.sectionsMemOffset);
 
@@ -309,17 +308,17 @@ namespace uba
 				if (symbolName[0] == '.') // sections
 					continue;
 
-				extra.clear();
+				bool isData = false;
 				if (symbol.SectionNumber != decltype(symbol.SectionNumber)(-1) && symbol.SectionNumber != decltype(symbol.SectionNumber)(-2))
 				{
 					UBA_ASSERTF(symbol.SectionNumber <= m_info.sectionCount, TC("%i / %i"), symbol.SectionNumber, m_info.sectionCount);
 					auto& section = sections[symbol.SectionNumber-1];
 					if (memcmp(section.Name, ".data", 6) == 0 || memcmp(section.Name, ".rdata", 7) == 0 || memcmp(section.Name, ".bss", 5) == 0)
-						extra.assign(",DATA");
+						isData = true;
 				}
 				symbolName.ToString(symbolString);
 				StringKey symbolKey = ToStringKeyRaw(symbolString.data(), symbolString.size());
-				m_exports.emplace(symbolKey, ExportInfo{symbolString, extra, index++});
+				m_exports.emplace(symbolKey, ExportInfo{symbolString, isData, index++});
 			}
 			else
 			{
@@ -330,7 +329,7 @@ namespace uba
 		}
 	}
 
-	bool ObjectFileCoff::CreateExtraFile(Logger& logger, const StringView& platform, MemoryBlock& memoryBlock, const UnorderedSymbols& allExternalImports, const UnorderedSymbols& allInternalImports, const UnorderedExports& allExports, bool includeExportsInFile)
+	bool ObjectFileCoff::CreateExtraFile(Logger& logger, const StringView& platform, MemoryBlock& memoryBlock, const AllExternalImports& allExternalImports, const AllInternalImports& allInternalImports, const AllExports& allExports, bool includeExportsInFile)
 	{
 		std::string tmp;
 
@@ -374,11 +373,12 @@ namespace uba
 			directiveSection.Characteristics = ImageScnAlign1Bytes|ImageScnLnkInfo|ImageScnLnkRemove;
 			++header.NumberOfSections;
 
-			auto writeExport = [&](AnsiStringView symbol, const std::string& extra)
+			auto writeExport = [&](AnsiStringView symbol, bool isData)
 				{
 					write("/EXPORT:", 8);
 					write(symbol.strBegin, symbol.Length());
-					write(extra.data(), extra.size());
+					if (isData)
+						write(",DATA", 5);
 					write(" ", 1);
 				};
 
@@ -395,32 +395,13 @@ namespace uba
 				auto findIt = allExports.find(impKey);
 				if (findIt == allExports.end())
 					continue;
-				writeExport(impSym, findIt->second.extra);
-			}
-
-			/*
-			for (auto& kv : allExports)
-			{
-				auto& symbol = kv.second.symbol;
-				UBA_ASSERT(!symbol.empty());
-
-				auto findIt = allExternalImports.find(symbol);
-				if (findIt == allExternalImports.end())
-				{
-					tmp.assign("__imp_").append(symbol);
-					findIt = allExternalImports.find(tmp);
-					if (findIt == allExternalImports.end())
-						continue;
-				}
-
-				writeExport(symbol, kv.second.extra);
+				writeExport(impSym, findIt->second.isData);
 			}
-			*/
 
 			tmp = "ThisIsAnUnrealEngineModule";
 			StringKey isUnrealKey = ToStringKeyRaw(tmp.data(), tmp.size());
 			if (allExports.find(isUnrealKey) != allExports.end())
-				writeExport(tmp, ""); // Workaround for tool not liking empty lists
+				writeExport(tmp, false); // Workaround for tool not liking empty lists
 
 			write("", 1);
 			directiveSection.SizeOfRawData = u32(memoryBlock.writtenSize) - directiveRawDataStart;
diff --git a/Engine/Source/Programs/UnrealBuildAccelerator/Common/Private/UbaObjectFileCoff.h b/Engine/Source/Programs/UnrealBuildAccelerator/Common/Private/UbaObjectFileCoff.h
--- a/Engine/Source/Programs/UnrealBuildAccelerator/Common/Private/UbaObjectFileCoff.h
+++ b/Engine/Source/Programs/UnrealBuildAccelerator/Common/Private/UbaObjectFileCoff.h
@@ -13,7 +13,7 @@ namespace uba
 		ObjectFileCoff();
 		virtual bool Parse(Logger& logger, ObjectFileParseMode parseMode, const tchar* hint) override;
 
-		static bool CreateExtraFile(Logger& logger, const StringView& platform, MemoryBlock& memoryBlock, const UnorderedSymbols& allExternalImports, const UnorderedSymbols& allInternalImports, const UnorderedExports& allSharedExports, bool includeExportsInFile);
+		static bool CreateExtraFile(Logger& logger, const StringView& platform, MemoryBlock& memoryBlock, const AllExternalImports& allExternalImports, const AllInternalImports& allInternalImports, const UnorderedExports& allSharedExports, bool includeExportsInFile);
 
 	private:
 		struct Info;
diff --git a/Engine/Source/Programs/UnrealBuildAccelerator/Common/Private/UbaObjectFileElf.cpp b/Engine/Source/Programs/UnrealBuildAccelerator/Common/Private/UbaObjectFileElf.cpp
--- a/Engine/Source/Programs/UnrealBuildAccelerator/Common/Private/UbaObjectFileElf.cpp
+++ b/Engine/Source/Programs/UnrealBuildAccelerator/Common/Private/UbaObjectFileElf.cpp
@@ -204,12 +204,12 @@ namespace uba
 							if (symbolVisibility == STV_HIDDEN && strncmp(symbolName, "_ZTV", 4) == 0)
 							{
 								std::string tmp = symbolName;
-								m_exports.emplace(ToStringKeyRaw(tmp.data(), tmp.size()), ExportInfo{std::move(tmp), std::string(), 0});//index++});
+								m_exports.emplace(ToStringKeyRaw(tmp.data(), tmp.size()), ExportInfo{std::move(tmp), false, 0});//index++});
 							}
 							if (m_useVisibilityForExports && symbolVisibility != STV_HIDDEN)
 							{
 								std::string tmp = symbolName;
-								m_exports.emplace(ToStringKeyRaw(tmp.data(), tmp.size()), ExportInfo{std::move(tmp), std::string(), 0});//index++});
+								m_exports.emplace(ToStringKeyRaw(tmp.data(), tmp.size()), ExportInfo{std::move(tmp), false, 0});//index++});
 							}
 						}
 						else
@@ -255,7 +255,7 @@ namespace uba
 		return true;
 	}
 
-	bool ObjectFileElf::CreateExtraFile(Logger& logger, const StringView& platform, MemoryBlock& memoryBlock, const UnorderedSymbols& allExternalImports, const UnorderedSymbols& allInternalImports, const UnorderedExports& allExports, bool includeExportsInFile)
+	bool ObjectFileElf::CreateExtraFile(Logger& logger, const StringView& platform, MemoryBlock& memoryBlock, const AllExternalImports& allExternalImports, const AllInternalImports& allInternalImports, const AllExports& allExports, bool includeExportsInFile)
 	{
 		u8* data = (u8*)memoryBlock.Allocate(sizeof(Elf64Header) + sizeof(Elf64SectionHeader) + 1, 1, TC(""));
 		auto& header = *(Elf64Header*)data;
diff --git a/Engine/Source/Programs/UnrealBuildAccelerator/Common/Private/UbaObjectFileElf.h b/Engine/Source/Programs/UnrealBuildAccelerator/Common/Private/UbaObjectFileElf.h
--- a/Engine/Source/Programs/UnrealBuildAccelerator/Common/Private/UbaObjectFileElf.h
+++ b/Engine/Source/Programs/UnrealBuildAccelerator/Common/Private/UbaObjectFileElf.h
@@ -12,7 +12,7 @@ namespace uba
 		ObjectFileElf();
 		virtual bool Parse(Logger& logger, ObjectFileParseMode parseMode, const tchar* hint) override;
 
-		static bool CreateExtraFile(Logger& logger, const StringView& platform, MemoryBlock& memoryBlock, const UnorderedSymbols& allExternalImports, const UnorderedSymbols& allInternalImports, const UnorderedExports& allExports, bool includeExportsInFile);
+		static bool CreateExtraFile(Logger& logger, const StringView& platform, MemoryBlock& memoryBlock, const AllExternalImports& allExternalImports, const AllInternalImports& allInternalImports, const AllExports& allExports, bool includeExportsInFile);
 
 	private:
 		u64 m_symTableNamesOffset = 0;
diff --git a/Engine/Source/Programs/UnrealBuildAccelerator/Common/Private/UbaObjectFileImportLib.cpp b/Engine/Source/Programs/UnrealBuildAccelerator/Common/Private/UbaObjectFileImportLib.cpp
--- a/Engine/Source/Programs/UnrealBuildAccelerator/Common/Private/UbaObjectFileImportLib.cpp
+++ b/Engine/Source/Programs/UnrealBuildAccelerator/Common/Private/UbaObjectFileImportLib.cpp
@@ -45,14 +45,14 @@ namespace uba
 			}
 
 			StringKey key = ToStringKeyRaw(symbol.data(), symbol.size());
-			m_exports.try_emplace(key, ExportInfo{ std::move(symbol),",DATA", i });
+			m_exports.try_emplace(key, ExportInfo{ std::move(symbol), true, i });
 		}
 
 		for (auto& symbol : impSymbols)
 		{
 			auto findIt = m_exports.find(symbol);
 			if (findIt != m_exports.end())
-				findIt->second.extra = std::string();
+				findIt->second.isData = false;
 		}
 
 		return true;
diff --git a/Engine/Source/Programs/UnrealBuildAccelerator/Common/Private/UbaObjectFileLLVMIR.cpp b/Engine/Source/Programs/UnrealBuildAccelerator/Common/Private/UbaObjectFileLLVMIR.cpp
--- a/Engine/Source/Programs/UnrealBuildAccelerator/Common/Private/UbaObjectFileLLVMIR.cpp
+++ b/Engine/Source/Programs/UnrealBuildAccelerator/Common/Private/UbaObjectFileLLVMIR.cpp
@@ -912,7 +912,7 @@ namespace uba
 			{
 				UBA_ASSERT(!name.empty());
 				record.dllStorage = DllStorage_Export;
-				m_exports.emplace(ToStringKeyRaw(name.data(), name.size()), ExportInfo{name, "", index++});
+				m_exports.emplace(ToStringKeyRaw(name.data(), name.size()), ExportInfo{name, false, index++});
 			}
 			else // Check visibility... might be a non-windows based llvm ir stream file
 			{
@@ -941,7 +941,7 @@ namespace uba
 					{
 						record.dllStorage = DllStorage_Export;
 						StringKey key = ToStringKeyRaw(name.data(), name.size());
-						m_exports.emplace(key, ExportInfo{name, "", index++});
+						m_exports.emplace(key, ExportInfo{name, false, index++});
 					}
 					else if (linkage == 1 || linkage == 16)
 					{
diff --git a/Engine/Source/Programs/UnrealBuildAccelerator/Common/Private/UbaSession.cpp b/Engine/Source/Programs/UnrealBuildAccelerator/Common/Private/UbaSession.cpp
--- a/Engine/Source/Programs/UnrealBuildAccelerator/Common/Private/UbaSession.cpp
+++ b/Engine/Source/Programs/UnrealBuildAccelerator/Common/Private/UbaSession.cpp
@@ -3837,7 +3837,7 @@ namespace uba
 		StringBuffer<> exportsFile;
 		exportsFile.Append(fileName, lastDot - fileName).Append(TCV(".exi"));
 
-		bool verbose = true;
+		bool verbose = false;
 		MemoryBlock memoryBlock(32*1024*1024);
 		if (!objectFile->WriteImportsAndExports(m_logger, memoryBlock, verbose))
 			return false;
diff --git a/Engine/Source/Programs/UnrealBuildAccelerator/Common/Public/UbaObjectFile.h b/Engine/Source/Programs/UnrealBuildAccelerator/Common/Public/UbaObjectFile.h
--- a/Engine/Source/Programs/UnrealBuildAccelerator/Common/Public/UbaObjectFile.h
+++ b/Engine/Source/Programs/UnrealBuildAccelerator/Common/Public/UbaObjectFile.h
@@ -11,10 +11,14 @@ namespace uba
 	class FileAccessor;
 	
 	using UnorderedSymbols = UnorderedSet<std::string>;
+	using AllExternalImports = UnorderedSymbols;
+	using AllInternalImports = UnorderedSymbols;
 
-	struct ExportInfo { std::string symbol; std::string extra; u32 index = 0; };
+	struct ExportInfo { std::string symbol; u8 isData = 0; u32 index = 0; };
 	using UnorderedExports = UnorderedMap<StringKey, ExportInfo>;
 
+	using AllExports = UnorderedExports;
+
 	enum ObjectFileType : u8
 	{
 		ObjectFileType_Unknown,
@@ -47,7 +51,7 @@ namespace uba
 		const UnorderedExports& GetExports() const;
 		const UnorderedSymbols& GetPotentialDuplicates() const;
 
-		static bool CreateExtraFile(Logger& logger, const StringView& extraObjFilename, const StringView& moduleName, const StringView& platform, const UnorderedSymbols& allExternalImports, const UnorderedSymbols& allInternalImports, const UnorderedExports& allExports, bool includeExportsInFile);
+		static bool CreateExtraFile(Logger& logger, const StringView& extraObjFilename, const StringView& moduleName, const StringView& platform, const AllExternalImports& allExternalImports, const AllInternalImports& allInternalImports, const AllExports& allExports, bool includeExportsInFile);
 
 		virtual ~ObjectFile();
 
@@ -58,8 +62,8 @@ namespace uba
 	protected:
 		virtual bool Parse(Logger& logger, ObjectFileParseMode parseMode, const tchar* hint) = 0;
 
-		static bool CreateDynamicListFile(Logger& logger, MemoryBlock& memoryBlock, const UnorderedSymbols& allExternalImports, const UnorderedSymbols& allInternalImports, const UnorderedExports& allExports, bool includeExportsInFile);
-		static bool CreateEmdFile(Logger& logger, MemoryBlock& memoryBlock, const StringView& moduleName, const UnorderedSymbols& allExternalImports, const UnorderedSymbols& allInternalImports, const UnorderedExports& allExports, bool includeExportsInFile);
+		static bool CreateDynamicListFile(Logger& logger, MemoryBlock& memoryBlock, const AllExternalImports& allExternalImports, const AllInternalImports& allInternalImports, const AllExports& allExports, bool includeExportsInFile);
+		static bool CreateEmdFile(Logger& logger, MemoryBlock& memoryBlock, const StringView& moduleName, const AllExternalImports& allExternalImports, const AllInternalImports& allInternalImports, const AllExports& allExports, bool includeExportsInFile);
 
 		FileAccessor* m_file = nullptr;
 		u8* m_data = nullptr;
diff --git a/Engine/Source/Programs/UnrealBuildAccelerator/ObjTool/Private/UbaImportLibWriter.cpp b/Engine/Source/Programs/UnrealBuildAccelerator/ObjTool/Private/UbaImportLibWriter.cpp
--- a/Engine/Source/Programs/UnrealBuildAccelerator/ObjTool/Private/UbaImportLibWriter.cpp
+++ b/Engine/Source/Programs/UnrealBuildAccelerator/ObjTool/Private/UbaImportLibWriter.cpp
@@ -85,12 +85,12 @@ namespace uba
 				func(*s);
 		}
 
-		void AddSymbol(const char* name, const char* extra, u32 value)
+		void AddSymbol(const char* name, bool isData, u32 value)
 		{
 			u32 offsetIndex = u32(m_symbols.size()) - m_extraSymbolCount;
 			SymbolInfo& s0 = *m_symbols.emplace_back(new SymbolInfo());
 
-			s0.isData = strstr(extra, "DATA") != 0;
+			s0.isData = isData;
 			s0.tempOrder = u32(m_symbols.size()-1) - m_extraSymbolCount - 3;
 			s0.offsetIndex = offsetIndex;
 
@@ -593,7 +593,7 @@ namespace uba
 					if (handled.insert(kv.first).second)
 						sortedSymbols.try_emplace(kv.second.index, &kv);
 				for (auto& outerKv : sortedSymbols)
-					AddSymbol(outerKv.second->second.symbol.c_str(), outerKv.second->second.extra.c_str(), 0);
+					AddSymbol(outerKv.second->second.symbol.c_str(), outerKv.second->second.isData, 0);
 			}
 
 			m_moduleName = std::string(libName) + libExt;
diff --git a/Engine/Source/Programs/UnrealBuildAccelerator/ObjTool/Private/UbaObjTool.cpp b/Engine/Source/Programs/UnrealBuildAccelerator/ObjTool/Private/UbaObjTool.cpp
--- a/Engine/Source/Programs/UnrealBuildAccelerator/ObjTool/Private/UbaObjTool.cpp
+++ b/Engine/Source/Programs/UnrealBuildAccelerator/ObjTool/Private/UbaObjTool.cpp
@@ -53,15 +53,6 @@ namespace uba
 		return -1;
 	}
 
-	// TODO: Add to rsp file instead
-	const char* NeededImports[] = 
-	{
-		"NvOptimusEnablement",
-		"AmdPowerXpressRequestHighPerformance",
-		"D3D12SDKVersion",
-		"D3D12SDKPath",
-	};
-
 	int WrappedMain(int argc, tchar* argv[])
 	{
 		using namespace uba;
@@ -251,10 +242,8 @@ namespace uba
 		{
 			CriticalSection cs;
 			Atomic<bool> success = true;
-			UnorderedSymbols allExternalImports; // Imports needed from the outside of the stripped obj files
 
-			for (auto imp : NeededImports)
-				allExternalImports.insert(imp);
+			AllExternalImports& allExternalImports = *new AllExternalImports(16*1024); // Imports needed from the outside of the stripped obj files. Let it leak to speedup shutdown
 
 			u32 workerCount = DefaultProcessorCount-1;
 			WorkManagerImpl workManager(workerCount, TC("UbaWrk/StrpObj"));
@@ -274,8 +263,8 @@ namespace uba
 			if (!success)
 				return -1;
 
-			UnorderedSymbols allInternalImports; // Imports that the obj files has. Note these could be existing in the obj files and then we might need to create loopbacks
-			UnorderedExports allExports; // Exports from all the obj files.
+			AllInternalImports& allInternalImports = *new AllInternalImports(512*1024); // Imports that the obj files has. Note these could be existing in the obj files and then we might need to create loopbacks
+			AllExports& allExports = *new AllExports(12*1024*1024); // Exports from all the obj files.
 
 			Map<TString, ObjectFile*> objectFiles;
 			auto g = MakeGuard([&]() { for (auto& kv : objectFiles) delete kv.second; });
@@ -297,27 +286,8 @@ namespace uba
 				return -1;
 
 			if (!extraObjFile.empty())
-			{
-				{
-					// Ugly hack. These symbols are exported from PosixShim and used by libcrypto.a but since we don't have an .exi file for libcrypto
-					// we do it this way for now
-					const char* symbolsToNeverStrip[] = 
-					{
-						"read_system_certificates_NP",
-						"inet_ntoa",
-						"gethostbyname",
-						"h_errno",
-						"getservbyname",
-						"ioctl",
-						"fcntl_shim",
-					};
-					for (auto symbol : symbolsToNeverStrip)
-						allExternalImports.insert(symbol);
-				}
-
 				if (!ObjectFile::CreateExtraFile(logger, extraObjFile, moduleName, platform, allExternalImports, allInternalImports, allExports, true))
 					return -1;
-			}
 			//logger.Info(TC("Reduced export count from %llu to %llu"), totalExportCount.load(), totalKeptExportCount.size());
 		}
 		else if (writeImpLib)
@@ -373,7 +343,7 @@ namespace uba
 					logger.Info(TC("I %S"), symbol.c_str());
 
 				for (auto& kv : objectFile->GetExports())
-					logger.Info(TC("E %S%S"), kv.second.symbol.c_str(), kv.second.extra.c_str());
+					logger.Info(TC("E %S%S"), kv.second.symbol.c_str(), kv.second.isData ? ",DATA" : "");
 			}
 		}
 		return 0;
