[UBA]
* Added option to allow deletion of verified files.
* Changed so cache server is allowed to delete verified files

diff --git a/Engine/Source/Programs/UnrealBuildAccelerator/CacheService/Private/UbaCacheService.cpp b/Engine/Source/Programs/UnrealBuildAccelerator/CacheService/Private/UbaCacheService.cpp
--- a/Engine/Source/Programs/UnrealBuildAccelerator/CacheService/Private/UbaCacheService.cpp
+++ b/Engine/Source/Programs/UnrealBuildAccelerator/CacheService/Private/UbaCacheService.cpp
@@ -443,6 +443,7 @@ namespace uba
 		storageInfo.storeCompressed = storeCompressed;
 		storageInfo.allowHintAsFallback = false;
 		storageInfo.writeReceivedCasFilesToDisk = true;
+		storageInfo.allowDeleteVerified = true;
 		StorageServer storageServer(storageInfo);
 
 		if (!storageServer.LoadCasTable(true, true))
diff --git a/Engine/Source/Programs/UnrealBuildAccelerator/Common/Private/UbaStorage.cpp b/Engine/Source/Programs/UnrealBuildAccelerator/Common/Private/UbaStorage.cpp
--- a/Engine/Source/Programs/UnrealBuildAccelerator/Common/Private/UbaStorage.cpp
+++ b/Engine/Source/Programs/UnrealBuildAccelerator/Common/Private/UbaStorage.cpp
@@ -202,11 +202,22 @@ namespace uba
 		{
 			CasEntry& ce = *it;
 			ce.lock.Enter();
-			if (ce.verified) // Can't remove these since they might be in mapping tables and actively being used. It is tempting to delete dropped cas files but we can't
+
+			if (ce.verified && !m_allowDeleteVerified) // Can't remove these since they might be in mapping tables and actively being used. It is tempting to delete dropped cas files but we can't
 			{
 				ce.lock.Leave();
 				break;
 			}
+
+			if (ce.beingWritten || ce.readCount)
+			{
+				if (ce.beingWritten)
+					m_logger.Info(TC("Want to delete file that is begin written. We should never end up here"));
+				ce.lock.Leave();
+				it = ce.prevAccessed;
+				continue;
+			}
+
 			UBA_ASSERT(ce.key != CasKeyZero);
 			toDelete.push_back({*it, ce.size});
 
@@ -1256,6 +1267,7 @@ namespace uba
 		m_storeCompressed = info.storeCompressed;
 		m_manuallyHandleOverflow = info.manuallyHandleOverflow;
 		m_asyncUnmapViewOfFile = info.asyncUnmapViewOfFile && m_workManager;
+		m_allowDeleteVerified = info.allowDeleteVerified;
 		m_writeToDisk = info.writeToDisk;
 		if (!m_writeToDisk)
 			m_casCapacityBytes = 0;
diff --git a/Engine/Source/Programs/UnrealBuildAccelerator/Common/Private/UbaStorageServer.cpp b/Engine/Source/Programs/UnrealBuildAccelerator/Common/Private/UbaStorageServer.cpp
--- a/Engine/Source/Programs/UnrealBuildAccelerator/Common/Private/UbaStorageServer.cpp
+++ b/Engine/Source/Programs/UnrealBuildAccelerator/Common/Private/UbaStorageServer.cpp
@@ -154,6 +154,9 @@ namespace uba
 			server.m_bufferSlots.Push(memoryBegin);
 			memoryBegin = nullptr;
 		}
+
+		SCOPED_WRITE_LOCK(casEntry->lock, entryLock);
+		--casEntry->readCount;
 	}
 
 	void StorageServer::OnDisconnected(u32 clientId)
@@ -216,7 +219,7 @@ namespace uba
 				fetch.Release(*this, TC("OnDisconnected"));
 
 				if (m_traceFetch)
-					m_trace->FileFetchEnd(clientId, AsCompressed(fetch.casKey, m_storeCompressed));
+					m_trace->FileFetchEnd(clientId, AsCompressed(fetch.casEntry->key, m_storeCompressed));
 
 				it = m_activeFetches.erase(it);
 			}
@@ -502,6 +505,17 @@ namespace uba
 					casEntry = &findIt->second;
 				}
 
+				SCOPED_WRITE_LOCK(casEntry->lock, entryLock);
+				if (!casEntry->exists)
+				{
+					m_logger.Info(TC("FetchBegin failed for cas file %s (%s). Cas was deleted"), CasKeyString(casKey).str, hint.data);
+					writer.WriteU16(0);
+					return true;
+				}
+				++casEntry->readCount;
+				entryLock.Leave();
+				auto rlg = MakeGuard([&]() { entryLock.Enter(); --casEntry->readCount; });
+
 				todoRemoveMe = 4;
 
 				if (casEntry->disallowed)
@@ -657,6 +671,7 @@ namespace uba
 
 				todoRemoveMe = 9;
 
+				rlg.Cancel();
 				mvg.Cancel();
 				rfg.Cancel();
 
@@ -679,7 +694,7 @@ namespace uba
 				fetch.memoryBegin = memoryBegin;
 				fetch.memoryPos = memoryPos;
 				fetch.left = left;
-				fetch.casKey = casKey;
+				fetch.casEntry = casEntry;
 				fetch.sendCasTime = GetTime() - start;
 
 				return true;
@@ -830,7 +845,7 @@ namespace uba
 						CasEntryAccessed(casEntry);
 						writer.WriteU16(u16(~0));
 						writer.WriteBool(m_traceStore);
-						m_logger.Detail(TC("Client %u Store request of %s which was just added to casdb (%s)"), connectionInfo.GetId(), CasKeyString(casKey).str, hint.data);
+						//m_logger.Detail(TC("Client %u Store request of %s which was just added to casdb (%s)"), connectionInfo.GetId(), CasKeyString(casKey).str, hint.data);
 						return true;
 					}
 				}
diff --git a/Engine/Source/Programs/UnrealBuildAccelerator/Common/Public/UbaStorage.h b/Engine/Source/Programs/UnrealBuildAccelerator/Common/Public/UbaStorage.h
--- a/Engine/Source/Programs/UnrealBuildAccelerator/Common/Public/UbaStorage.h
+++ b/Engine/Source/Programs/UnrealBuildAccelerator/Common/Public/UbaStorage.h
@@ -91,6 +91,7 @@ namespace uba
 		bool asyncUnmapViewOfFile = true;
 		bool writeToDisk = true;
 		bool keepTransientDataMapped = true;	// Will keep the transient data mapped in instead of doing map/unmap everytime we access it
+		bool allowDeleteVerified = false;
 		WorkManager* workManager = nullptr;
 		MutexHandle exclusiveMutex = InvalidMutexHandle; // Set this to hand over pre-exclusive access to storage. Storage will release mutex on shutdown
 		u8 casCompressor = 0;
@@ -224,6 +225,7 @@ namespace uba
 			CasEntry* prevAccessed = nullptr;
 			CasEntry* nextAccessed = nullptr;
 			u64 size = 0;
+			u16 readCount = 0; // This is set while file is being read over network
 			bool verified = false; // This flag needs to be set for below flags to be reliable. if this is false below flags are assumptions
 			bool exists = false; // File exists on disk
 			bool dropped = false; // This file is not seen anymore. will be deleted during shutdown
@@ -251,6 +253,7 @@ namespace uba
 		bool m_storeCompressed = false;
 		bool m_manuallyHandleOverflow = false;
 		bool m_asyncUnmapViewOfFile = true;
+		bool m_allowDeleteVerified = false;
 		bool m_writeToDisk;
 
 		MutexHandle m_exclusiveMutex = InvalidMutexHandle;
diff --git a/Engine/Source/Programs/UnrealBuildAccelerator/Common/Public/UbaStorageServer.h b/Engine/Source/Programs/UnrealBuildAccelerator/Common/Public/UbaStorageServer.h
--- a/Engine/Source/Programs/UnrealBuildAccelerator/Common/Public/UbaStorageServer.h
+++ b/Engine/Source/Programs/UnrealBuildAccelerator/Common/Public/UbaStorageServer.h
@@ -90,7 +90,7 @@ namespace uba
 		{
 			u32 clientId = ~0u;
 			Atomic<u64> left;
-			CasKey casKey;
+			CasEntry* casEntry;
 			Atomic<u64> sendCasTime;
 
 			FileHandle readFileHandle = InvalidFileHandle;
