[UBA]
* Fixed elf object file parsing
* Fixed llvmir object file parsing
* Added version script (ldscript) extra file output
*

diff --git a/Engine/Source/Programs/UnrealBuildAccelerator/Common/Private/UbaObjectFile.cpp b/Engine/Source/Programs/UnrealBuildAccelerator/Common/Private/UbaObjectFile.cpp
--- a/Engine/Source/Programs/UnrealBuildAccelerator/Common/Private/UbaObjectFile.cpp
+++ b/Engine/Source/Programs/UnrealBuildAccelerator/Common/Private/UbaObjectFile.cpp
@@ -224,7 +224,7 @@ namespace uba
 		return m_potentialDuplicates;
 	}
 
-	bool ObjectFile::CreateExtraFile(Logger& logger, const StringView& extraObjFilename, const StringView& moduleName, const StringView& platform, const AllExternalImports& allExternalImports, const UnorderedSymbols& allInternalImports, const AllExports& allExports, bool includeExportsInFile)
+	bool ObjectFile::CreateExtraFile(Logger& logger, const StringView& extraObjFilename, const StringView& moduleName, const StringView& platform, const AllExternalImports& allExternalImports, const UnorderedSymbols& allInternalImports, const AllExports& allExports, const ExtraExports& extraExports, bool includeExportsInFile)
 	{
 		ObjectFileCoff objectFileCoff;
 		ObjectFileElf objectFileElf;
@@ -232,11 +232,13 @@ namespace uba
 		MemoryBlock memoryBlock(16*1024*1024);
 
 		bool res;
-		if (platform.Equals(TCV("win64")) || platform.Equals(TCV("wingdk")) || platform.Equals(TCV("xb1")) || platform.Equals(TCV("xsx"))) 
+		if (extraObjFilename.EndsWith(TCV(".obj")))
 			res = ObjectFileCoff::CreateExtraFile(logger, platform, memoryBlock, allExternalImports, allInternalImports, allExports, includeExportsInFile);
-		else if (extraObjFilename.EndsWith(TCV("dynlist")))
-			res = CreateDynamicListFile(logger, memoryBlock, allExternalImports, allInternalImports, allExports, includeExportsInFile);
-		else if (extraObjFilename.EndsWith(TCV("emd")))
+		else if (extraObjFilename.EndsWith(TCV(".dynlist")))
+			res = CreateVersionScript(logger, memoryBlock, allExternalImports, allInternalImports, allExports, extraExports, includeExportsInFile, true);
+		else if (extraObjFilename.EndsWith(TCV(".ldscript")))
+			res = CreateVersionScript(logger, memoryBlock, allExternalImports, allInternalImports, allExports, extraExports, includeExportsInFile, false);
+		else if (extraObjFilename.EndsWith(TCV(".emd")))
 			res = CreateEmdFile(logger, memoryBlock, moduleName, allExternalImports, allInternalImports, allExports, includeExportsInFile);
 		else
 			res = ObjectFileElf::CreateExtraFile(logger, platform, memoryBlock, allExternalImports, allInternalImports, allExports, includeExportsInFile);
@@ -325,28 +327,39 @@ namespace uba
 		return true;
 	}
 
-	bool ObjectFile::CreateDynamicListFile(Logger& logger, MemoryBlock& memoryBlock, const AllExternalImports& allExternalImports, const UnorderedSymbols& allInternalImports, const AllExports& allExports, bool includeExportsInFile)
+	bool ObjectFile::CreateVersionScript(Logger& logger, MemoryBlock& memoryBlock, const AllExternalImports& allExternalImports, const UnorderedSymbols& allInternalImports, const AllExports& allExports, const ExtraExports& extraExports, bool includeExportsInFile, bool isDynList)
 	{
 		auto WriteString = [&](const char* str, u64 strLen) { memcpy(memoryBlock.Allocate(strLen, 1, TC("")), str, strLen); };
 
+		bool isFirst = true;
+		auto WriteSymbol = [&](const char* str, u64 strLen)
+			{
+				if (isFirst)
+					WriteString("global:\n", 8);
+				WriteString(str, strLen);
+				WriteString(";\n", 2);
+				isFirst = false;
+			};
+
 		//WriteString("VERSION ", 8);
-		WriteString("{", 1);
+		WriteString("{\n", 2);
 
-		bool isFirst = true;
-		for (auto& kv : allExports)
+		for (auto& imp : allExternalImports)
 		{
-			auto& symbol = kv.second.symbol;
-			//if (strncmp(symbol.first.c_str(), "_ZTV", 4) != 0)
-			//	continue;
-			if (allExternalImports.find(symbol) == allExternalImports.end())
-				continue;
-			if (isFirst)
-				WriteString("global: ", 8);
-			WriteString(symbol.c_str(), symbol.size());
-			WriteString(";", 1);
-			isFirst = false;
+			StringKey impKey = ToStringKeyRaw(imp.c_str(), imp.size());
+			auto findIt = allExports.find(impKey);
+			if (findIt != allExports.end())
+				WriteSymbol(imp.c_str(), imp.size());
 		}
-		//WriteString("local: *;", 9);
+
+		for (auto& symbol : extraExports)
+			WriteSymbol(symbol.c_str(), symbol.size());
+
+		if (!isDynList)
+			WriteString("local: *;\n", 10);
+		else if (isFirst)
+			WriteSymbol("ThisIsAnUnrealEngineModule", 26); // Workaround for tool not liking empty lists
+
 		WriteString("};", 2);
 
 		return true;
@@ -361,22 +374,23 @@ namespace uba
 
 		WriteString("Library: ", 9);
 		WriteString(moduleName2, moduleNameLen);
-		WriteString(" { export: {\r\n", 14);
+		WriteString(" { export: {\n", 13);
 
 		bool symbolAdded = false;
-		for (auto& kv : allExports)
+
+		for (auto& imp : allExternalImports)
 		{
-			auto& symbol = kv.second.symbol;
-			if (allExternalImports.find(symbol) != allExternalImports.end())
-			{
-				WriteString(symbol.c_str(), symbol.size());
-				WriteString("\r\n", 2);
-				symbolAdded = true;
-			}
+			StringKey impKey = ToStringKeyRaw(imp.c_str(), imp.size());
+			auto findIt = allExports.find(impKey);
+			if (findIt == allExports.end())
+				continue;
+			WriteString(imp.c_str(), imp.size());
+			WriteString("\n", 1);
+			symbolAdded = true;
 		}
 
 		if (!symbolAdded)
-			WriteString("ThisIsAnUnrealEngineModule\r\n", 28); // Workaround for tool not liking empty lists
+			WriteString("ThisIsAnUnrealEngineModule\n", 27); // Workaround for tool not liking empty lists
 
 		WriteString("}}", 2);
 
diff --git a/Engine/Source/Programs/UnrealBuildAccelerator/Common/Private/UbaObjectFileElf.cpp b/Engine/Source/Programs/UnrealBuildAccelerator/Common/Private/UbaObjectFileElf.cpp
--- a/Engine/Source/Programs/UnrealBuildAccelerator/Common/Private/UbaObjectFileElf.cpp
+++ b/Engine/Source/Programs/UnrealBuildAccelerator/Common/Private/UbaObjectFileElf.cpp
@@ -191,47 +191,36 @@ namespace uba
 
 					u8 symbolType = ELF64_ST_TYPE(symbol.st_info);
 					u32 symbolBinding = ELF64_ST_BIND(symbol.st_info);
-					u32 symbolVisibility = ELF64_ST_VISIBILITY(symbol.st_other);
-
-					//if (strstr(symbolName, "_ZN11TSubclassOfI13UAnimInstance"))
-					//	printf("");
+					//u32 symbolVisibility = ELF64_ST_VISIBILITY(symbol.st_other);
 
 					if (symbolBinding == STB_GLOBAL)
 					{
 						if (symbolType != STT_NOTYPE)
 						{
-							// This is an ugly hack because some platforms are not exporting vtable even though they should
-							if (symbolVisibility == STV_HIDDEN && strncmp(symbolName, "_ZTV", 4) == 0)
-							{
-								std::string tmp = symbolName;
-								m_exports.emplace(ToStringKeyRaw(tmp.data(), tmp.size()), ExportInfo{std::move(tmp), false, 0});//index++});
-							}
-							if (m_useVisibilityForExports && symbolVisibility != STV_HIDDEN)
-							{
-								std::string tmp = symbolName;
-								m_exports.emplace(ToStringKeyRaw(tmp.data(), tmp.size()), ExportInfo{std::move(tmp), false, 0});//index++});
-							}
+							std::string tmp = symbolName;
+							StringKey key = ToStringKeyRaw(tmp.data(), tmp.size());
+							m_exports.emplace(key, ExportInfo{std::move(tmp), false, 0});
 						}
 						else
-							m_imports.emplace(symbolName);//index++});
+							m_imports.emplace(symbolName);
 					}
 					else if (symbolBinding == STB_WEAK)
 					{
 						if (symbolType != STT_NOTYPE)
 						{
-							if (symbolVisibility == STV_DEFAULT || symbolVisibility == STV_PROTECTED)
-							{
-								// Another ugly hack because of dll exported template specialization
-							}
-							else
-							{
-								//if (strcmp(symbolName, "_ZNK5Chaos24TBoundingVolumeHierarchyI6TArrayIPNS_7TSphereIdLi3EEE22TSizedDefaultAllocatorILi32EEES1_IiS6_EdLi3EE26FindAllIntersectionsHelperERKNS_8TBVHNodeIdLi3EEERKNS_7TVectorIdLi3EEE") == 0)
-								//	m_exports.emplace(symbolName, ExportInfo{std::string(), 0});//index++});
-								//if (strcmp(symbolName, "_ZN5Chaos24TBoundingVolumeHierarchyI6TArrayIPNS_7TSphereIdLi3EEE22TSizedDefaultAllocatorILi32EEES1_IiS6_EdLi3EEC1ERKS7_ibd") == 0)
-								//	m_exports.emplace(symbolName, ExportInfo{std::string(), 0});//index++});
-							}
+							std::string tmp = symbolName;
+							StringKey key = ToStringKeyRaw(tmp.data(), tmp.size());
+							m_exports.emplace(key, ExportInfo{std::move(tmp), false, 0});
 						}
 					}
+					else if (symbolBinding == STB_LOCAL)
+					{
+						if (symbolName[0] == '.')
+							continue;
+						std::string tmp = symbolName;
+						StringKey key = ToStringKeyRaw(tmp.data(), tmp.size());
+						m_exports.emplace(key, ExportInfo{std::move(tmp), false, 0});
+					}
 				}
 			}
 			else if (!m_useVisibilityForExports && section.sh_type == SHT_PROGBITS)
@@ -243,10 +232,11 @@ namespace uba
 					while (it < end)
 					{
 						it += 4;
-						char* symbolName = (char*)it;
-						UBA_ASSERT(false);
-						//m_exports.emplace(symbolName, ExportInfo{std::string(), 0});
-						it += strlen(symbolName) + 1;
+						std::string symbolName((char*)it);
+						u64 symbolLen = symbolName.size();
+						StringKey key = ToStringKeyRaw(symbolName.data(), symbolLen);
+						m_exports.emplace(key, ExportInfo{std::move(symbolName), 0});
+						it += symbolLen + 1;
 					}
 				}
 			}
diff --git a/Engine/Source/Programs/UnrealBuildAccelerator/Common/Private/UbaObjectFileLLVMIR.cpp b/Engine/Source/Programs/UnrealBuildAccelerator/Common/Private/UbaObjectFileLLVMIR.cpp
--- a/Engine/Source/Programs/UnrealBuildAccelerator/Common/Private/UbaObjectFileLLVMIR.cpp
+++ b/Engine/Source/Programs/UnrealBuildAccelerator/Common/Private/UbaObjectFileLLVMIR.cpp
@@ -314,28 +314,6 @@ namespace uba
 			}
 		}
 
-		void ClearVBR64(u32 bits)
-		{
-			if (m_wordBits)
-			{
-				u8* oldPos = m_pos - 4;
-				u32 mask = ((1u << (32u - m_wordBits)) - 1u);
-				if (bits < m_wordBits)
-				{
-					mask |= ~((1u << (32u - (m_wordBits - bits))) - 1u);
-					bits = 0;
-				}
-				else
-					bits -= m_wordBits;
-				*(u32*)oldPos &= mask;
-			}
-			if (bits)
-			{
-				u32 mask = ~((1u << (bits)) - 1u);
-				*(u32*)m_pos &= mask;
-			}
-		}
-
 		u32 ReadCode()
 		{
 			return Read(m_currentCodeSize);
@@ -668,7 +646,7 @@ namespace uba
 				u32 code = ReadVBR(6);
 
 				if (CanBeExported(blockId, code))
-					m_owner.m_globalVarOrFunctionRecords.push_back(BitStreamEntry{u64(m_pos - m_begin), m_word, m_wordBits, code, DllStorage_None, false});
+					m_owner.m_globalVarOrFunctionRecords.push_back(BitStreamEntry{u64(m_pos - m_begin), m_word, m_wordBits, code});
 
 				u32 numElts = ReadVBR(6);
 
@@ -704,7 +682,7 @@ namespace uba
 			if (CanBeExported(blockId, code))
 			{
 				UBA_ASSERT(!abbv.operands.empty());
-				m_owner.m_globalVarOrFunctionRecords.push_back(BitStreamEntry{u64(m_pos - m_begin), m_word, m_wordBits, code, DllStorage_None, false, abbv.operands });
+				m_owner.m_globalVarOrFunctionRecords.push_back(BitStreamEntry{u64(m_pos - m_begin), m_word, m_wordBits, code, abbv.operands });
 			}
 
 #if UBA_LOG_STREAM
@@ -822,20 +800,6 @@ namespace uba
 		return ~0u;
 	}
 
-	u32 GetVisibilityIndex(u32 code)
-	{
-		if (code == MODULE_CODE_FUNCTION)
-			return 9;
-		if (code == MODULE_CODE_GLOBALVAR)
-			return 10;
-		if (code == MODULE_CODE_ALIAS)
-			return 8;
-		if (code == MODULE_CODE_ALIAS_OLD)
-			return 7;
-		UBA_ASSERTF(false, TC("module code %s not supported"), ModuleCodeToString(code));
-		return ~0u;
-	}
-
 	bool ObjectFileLLVMIR::Parse(Logger& logger, ObjectFileParseMode parseMode, const tchar* hint)
 	{
 		BitStreamReader reader(*this, logger, m_data, m_dataSize);
@@ -867,6 +831,11 @@ namespace uba
 			ParseBlock(logger, reader, blockInfo, blockId, 0);
 		}
 
+		Vector<u64> recordData;
+		std::string name;
+
+		Set<std::string> recordsWithOds;
+
 		u32 index = 0;
 		for (auto& record : m_globalVarOrFunctionRecords)
 		{
@@ -874,7 +843,7 @@ namespace uba
 			reader.m_word = record.word;
 			reader.m_wordBits = record.wordBits;
 
-			Vector<u64> recordData;
+			recordData.clear();
 			if (record.operands.empty())
 			{
 				u32 numElts = reader.ReadVBR(6);
@@ -887,13 +856,12 @@ namespace uba
 				reader.ReadRecordOperands(recordData, record.operands);
 			}
 
-			std::string name;
 			u64* recIt = recordData.data();
 			u64 recSize = recordData.size();
 
 			if (recIt[0] + recIt[1] >= m_strTabSize)
 				continue;
-			name = std::string((char*)(m_data + m_strTabPos) + recIt[0], recIt[1]);
+			name.assign((char*)(m_data + m_strTabPos) + recIt[0], recIt[1]);
 
 			u32 dllStorageIndex = GetDllStorageIndex(record.code);
 
@@ -901,69 +869,72 @@ namespace uba
 			if (recSize > dllStorageIndex)
 				dllStorage = recIt[dllStorageIndex];
 
-			UBA_ASSERT(record.dllStorage != DllStorage_Export);
+			// Left here for debugging purposes :)
+			//if (name.find("RemoveHighlight") != -1)
+			//	printf("");
 
-			if (dllStorage == 1)
-			{
-				m_imports.emplace(name);
-				record.dllStorage = DllStorage_Import;
-			}
-			else if (dllStorage == 2)
+			bool isExported = false;
+			u32 linkageIndex = GetLinkageIndex(record.code);
+			u64 linkage = 0;
+			if (recSize > linkageIndex)
 			{
-				UBA_ASSERT(!name.empty());
-				record.dllStorage = DllStorage_Export;
-				m_exports.emplace(ToStringKeyRaw(name.data(), name.size()), ExportInfo{name, false, index++});
+				linkage = recIt[linkageIndex];
+				if (linkage == 9) // PrivateLinkage
+					continue;
+				if (linkage == 19) // LinkOnceODRLinkage
+				{
+					recordsWithOds.insert(name);
+					continue;
+				}
+
+				// ExternalLinkage || DLLImportLinkage || DLLExportLinkage || LinkOnceODRAutoHideLinkage || WeakODRLinkage
+				isExported = linkage == 0 || linkage == 5 || linkage == 6 || linkage == 15 || linkage == 17;
+
+				if (dllStorage == DllStorage_Import)
+					isExported = false;
 			}
-			else // Check visibility... might be a non-windows based llvm ir stream file
+
+			if (record.code == MODULE_CODE_GLOBALVAR)
 			{
-				/*
-				u64 visibility = ~0ull;
-				u32 visibilityIndex = GetVisibilityIndex(record.code);
-				bool hidden = false;
-				if (recSize > visibilityIndex)
+				bool hasInitializer = recordData[4] != 0;
+				if (!hasInitializer)
+					isExported = false;
+
+				if (isExported)
 				{
-					visibility = recIt[visibilityIndex];
-					if (visibility != 0)
-						hidden = true;
+					m_exports.emplace(ToStringKeyRaw(name.data(), name.size()), ExportInfo{name, false, index++});
 				}
-				*/
-
-				u64 linkage = ~0ull;
-				u32 linkageIndex = GetLinkageIndex(record.code);
-				if (recSize > linkageIndex)
+				else if (name.find(".str") == -1)
 				{
-					linkage = recIt[linkageIndex];
+					if (recordsWithOds.find(name) != recordsWithOds.end())
+						continue;
+					m_imports.emplace(name);
+				}
+			}
+			else if (record.code == MODULE_CODE_FUNCTION)
+			{
+				bool isProto = recordData[4] != 0;
 
-					//if (strncmp(name.data(), "_Init_thread_footer", 18) == 0)
-					//	printf("");
+				if (linkage == 12) // thunk import?!? (like _ZThn840_N21UFortHUDElementWidget15RemoveHighlightEv)
+					isProto = true;
 
-					if (linkage == 5 || linkage == 6 || linkage == 15)
-					{
-						record.dllStorage = DllStorage_Export;
-						StringKey key = ToStringKeyRaw(name.data(), name.size());
-						m_exports.emplace(key, ExportInfo{name, false, index++});
-					}
-					else if (linkage == 1 || linkage == 16)
-					{
-						record.dllStorage = DllStorage_Import;
-						m_imports.emplace(name);
-					}
+				if (isProto)
+				{
+					if (recordsWithOds.find(name) != recordsWithOds.end())
+						continue;
+					m_imports.emplace(name);
+				}
+				else if (isExported)
+				{
+					m_exports.emplace(ToStringKeyRaw(name.data(), name.size()), ExportInfo{name, false, index++});
 				}
 			}
-
-			if (record.dllStorage == DllStorage_Import)
+			else if (record.code == MODULE_CODE_ALIAS)
 			{
-				// microsoft console
-				if (name == "?id@?$codecvt@DDU_Mbstatet@@@std@@2V0locale@2@A"
-				 || name == "?_Id_cnt@id@locale@std@@0HA"
-				 || name == "?cerr@std@@3V?$basic_ostream@DU?$char_traits@D@std@@@1@A"
-					)
-					record.keepAsIs = true;
-
-				// win64
-				if (name == "?_Raise_handler@std@@3P6AXAEBVexception@stdext@@@ZEA")
-					record.keepAsIs = true;
+				if (isExported)
+					m_exports.emplace(ToStringKeyRaw(name.data(), name.size()), ExportInfo{name, false, index++});
 			}
+			(void)index;
 		}
 		return true;
 	}
@@ -988,11 +959,6 @@ namespace uba
 		Vector<u64> record;
 		std::string blob;
 
-		if (blockId == GLOBALVAL_SUMMARY_BLOCK_ID)
-		{
-			// TODO: /EXPORT:xyz is probably in here somehow.. this format is sooo hard to parse... for now we just write directly into the strtab blob
-		}
-
 		while (true)
 		{
 			UBA_ASSERT(!reader.IsDone());
diff --git a/Engine/Source/Programs/UnrealBuildAccelerator/Common/Private/UbaObjectFileLLVMIR.h b/Engine/Source/Programs/UnrealBuildAccelerator/Common/Private/UbaObjectFileLLVMIR.h
--- a/Engine/Source/Programs/UnrealBuildAccelerator/Common/Private/UbaObjectFileLLVMIR.h
+++ b/Engine/Source/Programs/UnrealBuildAccelerator/Common/Private/UbaObjectFileLLVMIR.h
@@ -31,9 +31,9 @@ namespace uba
 
 		enum DllStorage : u8
 		{
+			DllStorage_None,
 			DllStorage_Export,
 			DllStorage_Import,
-			DllStorage_None,
 		};
 
 		struct BitStreamEntry
@@ -42,8 +42,6 @@ namespace uba
 			u32 word;
 			u32 wordBits;
 			u32 code;
-			DllStorage dllStorage;
-			u8 keepAsIs;
 			Vector<AbbrevOp> operands;
 		};
 
@@ -51,7 +49,6 @@ namespace uba
 
 		u64 m_strTabPos;
 		u64 m_strTabSize;
-		//std::string m_strTab;
 	};
 
 
diff --git a/Engine/Source/Programs/UnrealBuildAccelerator/Common/Private/UbaSession.cpp b/Engine/Source/Programs/UnrealBuildAccelerator/Common/Private/UbaSession.cpp
--- a/Engine/Source/Programs/UnrealBuildAccelerator/Common/Private/UbaSession.cpp
+++ b/Engine/Source/Programs/UnrealBuildAccelerator/Common/Private/UbaSession.cpp
@@ -1493,6 +1493,7 @@ namespace uba
 		m_storeIntermediateFilesCompressed = info.storeIntermediateFilesCompressed && IsWindows; // Non-windows not implemented (yet)
 		m_readIntermediateFilesCompressed = (m_storeIntermediateFilesCompressed || (info.readIntermediateFilesCompressed && IsWindows)) && !runningRemote; // with remote we decompress the files into memory
 		m_allowLocalDetour = info.allowLocalDetour;
+		m_extractObjFilesSymbols = info.extractObjFilesSymbols;
 
 		m_detailedTrace = info.detailedTrace;
 		m_traceChildProcesses = info.traceChildProcesses;
@@ -2920,7 +2921,7 @@ namespace uba
 				lastWriteTime = writtenFile.lastWriteTime;
 			}
 
-			if (msg.process.m_extractExports && startInfo.rules->ShouldExtractSymbols(activeFile.name))
+			if ((msg.process.m_extractExports || m_extractObjFilesSymbols) && startInfo.rules->ShouldExtractSymbols(activeFile.name))
 				if (!ExtractSymbolsFromObjectFile(msg, name.data, fileSize))
 					return false;
 		}
@@ -3838,6 +3839,9 @@ namespace uba
 		exportsFile.Append(fileName, lastDot - fileName).Append(TCV(".exi"));
 
 		bool verbose = false;
+		#if UBA_DEBUG
+		verbose = true;
+		#endif
 		MemoryBlock memoryBlock(32*1024*1024);
 		if (!objectFile->WriteImportsAndExports(m_logger, memoryBlock, verbose))
 			return false;
diff --git a/Engine/Source/Programs/UnrealBuildAccelerator/Common/Public/UbaObjectFile.h b/Engine/Source/Programs/UnrealBuildAccelerator/Common/Public/UbaObjectFile.h
--- a/Engine/Source/Programs/UnrealBuildAccelerator/Common/Public/UbaObjectFile.h
+++ b/Engine/Source/Programs/UnrealBuildAccelerator/Common/Public/UbaObjectFile.h
@@ -13,6 +13,7 @@ namespace uba
 	using UnorderedSymbols = UnorderedSet<std::string>;
 	using AllExternalImports = UnorderedSymbols;
 	using AllInternalImports = UnorderedSymbols;
+	using ExtraExports = Vector<std::string>;
 
 	struct ExportInfo { std::string symbol; u8 isData = 0; u32 index = 0; };
 	using UnorderedExports = UnorderedMap<StringKey, ExportInfo>;
@@ -51,7 +52,7 @@ namespace uba
 		const UnorderedExports& GetExports() const;
 		const UnorderedSymbols& GetPotentialDuplicates() const;
 
-		static bool CreateExtraFile(Logger& logger, const StringView& extraObjFilename, const StringView& moduleName, const StringView& platform, const AllExternalImports& allExternalImports, const AllInternalImports& allInternalImports, const AllExports& allExports, bool includeExportsInFile);
+		static bool CreateExtraFile(Logger& logger, const StringView& extraObjFilename, const StringView& moduleName, const StringView& platform, const AllExternalImports& allExternalImports, const AllInternalImports& allInternalImports, const AllExports& allExports, const ExtraExports& extraExports, bool includeExportsInFile);
 
 		virtual ~ObjectFile();
 
@@ -62,7 +63,7 @@ namespace uba
 	protected:
 		virtual bool Parse(Logger& logger, ObjectFileParseMode parseMode, const tchar* hint) = 0;
 
-		static bool CreateDynamicListFile(Logger& logger, MemoryBlock& memoryBlock, const AllExternalImports& allExternalImports, const AllInternalImports& allInternalImports, const AllExports& allExports, bool includeExportsInFile);
+		static bool CreateVersionScript(Logger& logger, MemoryBlock& memoryBlock, const AllExternalImports& allExternalImports, const AllInternalImports& allInternalImports, const AllExports& allExports, const ExtraExports& extraExports, bool includeExportsInFile, bool isDynList);
 		static bool CreateEmdFile(Logger& logger, MemoryBlock& memoryBlock, const StringView& moduleName, const AllExternalImports& allExternalImports, const AllInternalImports& allInternalImports, const AllExports& allExports, bool includeExportsInFile);
 
 		FileAccessor* m_file = nullptr;
diff --git a/Engine/Source/Programs/UnrealBuildAccelerator/Common/Public/UbaSession.h b/Engine/Source/Programs/UnrealBuildAccelerator/Common/Public/UbaSession.h
--- a/Engine/Source/Programs/UnrealBuildAccelerator/Common/Public/UbaSession.h
+++ b/Engine/Source/Programs/UnrealBuildAccelerator/Common/Public/UbaSession.h
@@ -233,6 +233,7 @@ namespace uba
 		bool m_storeIntermediateFilesCompressed;
 		bool m_readIntermediateFilesCompressed;
 		bool m_allowLocalDetour;
+		bool m_extractObjFilesSymbols;
 
 		u64 m_keepOutputFileMemoryMapsThreshold;
 
diff --git a/Engine/Source/Programs/UnrealBuildAccelerator/Common/Public/UbaSessionCreateInfo.h b/Engine/Source/Programs/UnrealBuildAccelerator/Common/Public/UbaSessionCreateInfo.h
--- a/Engine/Source/Programs/UnrealBuildAccelerator/Common/Public/UbaSessionCreateInfo.h
+++ b/Engine/Source/Programs/UnrealBuildAccelerator/Common/Public/UbaSessionCreateInfo.h
@@ -38,7 +38,7 @@ namespace uba
 		bool storeIntermediateFilesCompressed = false;	// Compiler will write intermediate files (.obj, .pch etc) compressed to disk and will be decompressed when used
 		bool readIntermediateFilesCompressed = false;	// Set to true to support reading compressed .obj files. This flag will not compress new obj files
 		bool allowLocalDetour = true;			// Allow local processes to be detoured. If this is false it is up to outside logic to register all created/deleted files
-		bool extractObjFilesSymbols = false;	// Will extract import/export symbols and clear export table from obj file.
+		bool extractObjFilesSymbols = false;	// Will extract import/export symbols to a file
 		bool treatTempDirAsEmpty = true;		// If this is true, directory table will always see temp as empty. Use sub directories of temp instead
 		bool useFakeVolumeSerial = true;		// Fake volume serials reduce the directory table size since serials end up being compressed down to one byte
 		bool keepTransientDataMapped = true;	// Will keep the transient data mapped in instead of doing map/unmap everytime we access it
diff --git a/Engine/Source/Programs/UnrealBuildAccelerator/ObjTool/Private/UbaImportLibWriter.cpp b/Engine/Source/Programs/UnrealBuildAccelerator/ObjTool/Private/UbaImportLibWriter.cpp
--- a/Engine/Source/Programs/UnrealBuildAccelerator/ObjTool/Private/UbaImportLibWriter.cpp
+++ b/Engine/Source/Programs/UnrealBuildAccelerator/ObjTool/Private/UbaImportLibWriter.cpp
@@ -143,7 +143,7 @@ namespace uba
 		{
 			u32 newOffset = m_memoryOffset + u32(bytes);
 			if (m_memory.writtenSize < newOffset)
-				m_memory.AllocateNoLock(newOffset - m_memory.writtenSize, 1, TC(""));
+				m_memory.AllocateNoLock(newOffset - m_memory.writtenSize, 1, TC("ImportLibWriter"));
 			memcpy(m_memory.memory + m_memoryOffset, data, bytes);
 			m_memoryOffset = newOffset;
 		}
@@ -556,7 +556,7 @@ namespace uba
 
 	public:
 
-		ImportLib() : m_memory(128*1024*1024) // 128mb should be more than enough reserved memory for the biggest import libs
+		ImportLib() : m_memory(512*1024*1024) // 512mb should be more than enough reserved memory for the biggest import libs
 		{
 		}
 
diff --git a/Engine/Source/Programs/UnrealBuildAccelerator/ObjTool/Private/UbaObjTool.cpp b/Engine/Source/Programs/UnrealBuildAccelerator/ObjTool/Private/UbaObjTool.cpp
--- a/Engine/Source/Programs/UnrealBuildAccelerator/ObjTool/Private/UbaObjTool.cpp
+++ b/Engine/Source/Programs/UnrealBuildAccelerator/ObjTool/Private/UbaObjTool.cpp
@@ -48,6 +48,7 @@ namespace uba
 		logger.Info(TC("    /O:<objfile>             Obj file to output containing exports and loopbacks"));
 		logger.Info(TC("    /T:<platform>            Target platform"));
 		logger.Info(TC("    /M:<module>              Name of module. Needed in emd files"));
+		logger.Info(TC("    /E:<symbol>              Additional symbol to be exported. To solve combination of dynlist files"));
 		logger.Info(TC("    /COMPRESS                Write '/O' file compressed"));
 		logger.Info(TC(""));
 		return -1;
@@ -67,6 +68,7 @@ namespace uba
 		Vector<TString> objFilesDependencies;
 		TString extraObjFile;
 		Vector<TString> objFilesForImpLib;
+		ExtraExports extraExports;
 		std::string impLibName;
 		TString impLibFile;
 		TString platform;
@@ -172,6 +174,13 @@ namespace uba
 				{
 					moduleName = value.data;
 				}
+				else if (name.StartsWith(TC("/E")))
+				{
+					char buffer[512];
+					if (!value.Parse(buffer, sizeof(buffer)))
+						return PrintHelp(TC("Bad symbol name"));
+					extraExports.push_back(buffer);
+				}
 				else if (name.Equals(TCV("-printsymbols")))
 				{
 					printSymbols = true;
@@ -286,7 +295,7 @@ namespace uba
 				return -1;
 
 			if (!extraObjFile.empty())
-				if (!ObjectFile::CreateExtraFile(logger, extraObjFile, moduleName, platform, allExternalImports, allInternalImports, allExports, true))
+				if (!ObjectFile::CreateExtraFile(logger, extraObjFile, moduleName, platform, allExternalImports, allInternalImports, allExports, extraExports, true))
 					return -1;
 			//logger.Info(TC("Reduced export count from %llu to %llu"), totalExportCount.load(), totalKeptExportCount.size());
 		}
