[UBA]
* Fixed bug in cache server where a cacheentries instance could get shared entries written multiple times so some entries could end up mapping against the wrong share
* Fixed so cache server validates buckets after a crash and delete potentially bad ones

diff --git a/Engine/Source/Programs/UnrealBuildAccelerator/CacheService/Private/UbaCacheService.cpp b/Engine/Source/Programs/UnrealBuildAccelerator/CacheService/Private/UbaCacheService.cpp
--- a/Engine/Source/Programs/UnrealBuildAccelerator/CacheService/Private/UbaCacheService.cpp
+++ b/Engine/Source/Programs/UnrealBuildAccelerator/CacheService/Private/UbaCacheService.cpp
@@ -446,7 +446,8 @@ namespace uba
 		storageInfo.allowDeleteVerified = true;
 		StorageServer storageServer(storageInfo);
 
-		if (!storageServer.LoadCasTable(true, true))
+		bool wasTerminated = false;
+		if (!storageServer.LoadCasTable(true, true, &wasTerminated))
 			return false;
 
 		CacheServerCreateInfo cacheInfo(storageServer, g_rootDir.data, logWriter);
@@ -455,7 +456,7 @@ namespace uba
 		cacheInfo.maintenanceReserveSize = maintenanceReserveSizeMb * 1024 * 1024;
 		CacheServer cacheServer(cacheInfo);
 
-		if (!cacheServer.Load())
+		if (!cacheServer.Load(wasTerminated))
 			return false;
 
 		if (fullMaintenance)
diff --git a/Engine/Source/Programs/UnrealBuildAccelerator/Common/Private/UbaCacheEntry.cpp b/Engine/Source/Programs/UnrealBuildAccelerator/Common/Private/UbaCacheEntry.cpp
--- a/Engine/Source/Programs/UnrealBuildAccelerator/Common/Private/UbaCacheEntry.cpp
+++ b/Engine/Source/Programs/UnrealBuildAccelerator/Common/Private/UbaCacheEntry.cpp
@@ -503,7 +503,7 @@ namespace uba
 	void CacheEntries::BuildInputs(CacheEntry& entry, const Set<u32>& inputs)
 	{
 		Vector<u8> temp;
-		BuildInputsT(entry, inputs, entries.empty(), temp);
+		BuildInputsT(entry, inputs, sharedInputCasKeyOffsets.empty(), temp);
 	}
 
 	void CacheEntries::UpdateEntries(Logger& logger, const HashMap2<u32, u32>& oldToNewCasKeyOffset, Vector<u32>& temp, Vector<u8>& temp2, Vector<u8>& temp3)
@@ -667,6 +667,60 @@ namespace uba
 		}
 	}
 
+	bool CacheEntries::Validate(Logger& logger)
+	{
+		Set<u32> sharedOffsets;
+		{
+			BinaryReader sharedReader(sharedInputCasKeyOffsets);
+			while (sharedReader.GetLeft())
+			{
+				u64 offset;
+				if (!sharedReader.TryRead7BitEncoded(offset))
+					return false;
+				if (!sharedOffsets.insert(u32(offset)).second)
+					return false;
+			}
+		}
+
+		for (auto& entry : entries)
+		{
+			Set<u32> entryOffsets;
+
+			u32 rangeCount = 0;
+			BinaryReader rangeReader(entry.sharedInputCasKeyOffsetRanges);
+			while (rangeReader.GetLeft())
+			{
+				++rangeCount;
+				u64 begin = rangeReader.Read7BitEncoded();
+				u64 end = rangeReader.Read7BitEncoded();
+				BinaryReader sharedReader(sharedInputCasKeyOffsets.data(), begin, end);
+				while (sharedReader.GetLeft())
+				{
+					u64 offset;
+					if (!sharedReader.TryRead7BitEncoded(offset))
+						return false;
+					if (!entryOffsets.insert(u32(offset)).second)
+						return false;
+				}
+			}
+
+			u32 extraCount = 0;
+			BinaryReader extraReader(entry.extraInputCasKeyOffsets);
+			while (extraReader.GetLeft())
+			{
+				++extraCount;
+				u64 offset;
+				if (!extraReader.TryRead7BitEncoded(offset))
+					return false;
+				if (sharedOffsets.find(u32(offset)) != sharedOffsets.end())
+					return false;
+				if (!entryOffsets.insert(u32(offset)).second)
+					return false;
+			}
+		}
+		return true;
+	}
+
 	void CacheEntries::PopulateInputsThatAreOutputs(const Vector<u8>& inputData, StorageImpl& storage, CompactCasKeyTable& table)
 	{
 		// one entry that is ~0u is equal to test all inputs
diff --git a/Engine/Source/Programs/UnrealBuildAccelerator/Common/Private/UbaCacheServer.cpp b/Engine/Source/Programs/UnrealBuildAccelerator/Common/Private/UbaCacheServer.cpp
--- a/Engine/Source/Programs/UnrealBuildAccelerator/Common/Private/UbaCacheServer.cpp
+++ b/Engine/Source/Programs/UnrealBuildAccelerator/Common/Private/UbaCacheServer.cpp
@@ -1,7 +1,7 @@
 // Copyright Epic Games, Inc. All Rights Reserved.
 
 #include "UbaCacheServer.h"
-#include "UbaCacheEntry.h"
+#include "UbaCacheBucket.h"
 #include "UbaCompactTables.h"
 #include "UbaBinaryReaderWriter.h"
 #include "UbaDirectoryIterator.h"
@@ -131,14 +131,22 @@ namespace uba
 		u32 count = 0;
 	};
 
+	struct CacheBucket::MaintenanceContext
+	{
+		MemoryBlock memoryBlock;
+		BitArray deletedOffsets;
+		bool isInitialized = false;
+		bool shouldTest = false;
+	};
+
 	struct CacheServer::ConnectionBucket
 	{
 		ConnectionBucket(u64 i, u32 version) : pathTable(IsCaseInsensitive(i), 0, 0, version), id(i) {}
 		CompactPathTable pathTable;
 		CompactCasKeyTable casKeyTable;
 
-		Futex cacheEntryLookupLock;
-		UnorderedMap<CasKey, CacheEntry> cacheEntryLookup;
+		Futex deferredCacheEntryLookupLock;
+		UnorderedMap<CasKey, CacheEntry> deferredCacheEntryLookup;
 		u64 id;
 		u32 index = ~0u;
 	};
@@ -154,37 +162,6 @@ namespace uba
 		u64 connectTime = 0;
 	};
 
-	struct CacheServer::Bucket
-	{
-		Bucket(u64 id, u32 version) : m_pathTable(IsCaseInsensitive(id), 0, 0, version), m_id(id) {}
-		ReaderWriterLock m_cacheEntryLookupLock;
-		UnorderedMap<CasKey, CacheEntries> m_cacheEntryLookup;
-
-		CompactPathTable m_pathTable;
-		CompactCasKeyTable m_casKeyTable;
-
-		u64 m_id;
-		u64 totalEntryCount;
-		u64 totalEntrySize;
-		Atomic<bool> hasDeletedEntries;
-		Atomic<bool> needsSave;
-
-		// Times are in file time from creation of cache database
-		Atomic<u64> lastSavedTime;
-		Atomic<u64> lastUsedTime;
-		u64 oldestUsedTime = 0; 
-
-		u32 index = ~0u;
-
-		struct MaintenanceContext
-		{
-			MemoryBlock memoryBlock;
-			BitArray deletedOffsets;
-			bool isInitialized = false;
-			bool shouldTest = false;
-		}* m_maintenanceContext = nullptr;
-	};
-
 	StringView ToString(CacheMessageType type)
 	{
 		switch (type)
@@ -243,14 +220,7 @@ namespace uba
 		m_server.UnregisterService(CacheServiceId);
 	}
 
-	struct CacheServer::LoadStats
-	{
-		Atomic<u32> totalPathTableSize;
-		Atomic<u32> totalCasKeyTableSize;
-		Atomic<u64> totalCacheEntryCount;
-	};
-
-	bool CacheServer::Load()
+	bool CacheServer::Load(bool validateBuckets)
 	{
 		u64 startTime = GetTime();
 
@@ -281,7 +251,7 @@ namespace uba
 		if (databaseVersion != CacheFileVersion)
 			m_dbfileDirty = true;
 
-		LoadStats stats;
+		CacheBucket::LoadStats stats;
 
 		if (databaseVersion == 4)
 		{
@@ -289,7 +259,7 @@ namespace uba
 			while (bucketCount--)
 			{
 				Bucket& bucket = GetBucket(reader.ReadU64(), TC("Loading"));
-				LoadBucket(bucket, reader, databaseVersion, stats);
+				bucket.Load(m_logger, reader, databaseVersion, stats, m_storage);
 			}
 		}
 		else
@@ -306,6 +276,9 @@ namespace uba
 					GetBucket(id, TC("Loading"), false);
 				});
 
+			Futex toDeleteLock;
+			Set<u64> toDelete;
+
 			m_server.ParallelFor(GetBucketWorkerCount(), m_buckets, [&](const WorkContext&, auto& it)
 				{
 					u64 key = it->first;
@@ -321,8 +294,22 @@ namespace uba
 					}
 					BinaryReader reader(bucketFile.GetData(), 0, bucketFile.GetSize());
 					u32 bucketVersion = reader.ReadU32();
-					LoadBucket(bucket, reader, bucketVersion, stats);
+
+					if (bucket.Load(m_logger, reader, bucketVersion, stats, m_storage))
+						if (!validateBuckets || bucket.Validate(m_logger, m_server))
+							return;
+
+					bucketFile.Close();
+					m_logger.Info(TC("Loaded invalid bucket %s. Will be deleted"), bucketFilename.data);
+					DeleteFile(bucketFilename.data);
+					SCOPED_FUTEX(toDeleteLock, lock);
+					toDelete.insert(key);
+
 				}, TCV("LoadBucket"));
+
+			for (auto key : toDelete)
+				m_buckets.erase(key);
+
 		}
 
 		u64 duration = GetTime() - startTime;
@@ -330,20 +317,22 @@ namespace uba
 		return true;
 	}
 
-	bool CacheServer::LoadBucket(Bucket& bucket, BinaryReader& reader, u32 databaseVersion, LoadStats& outStats)
+	CacheBucket::CacheBucket(u64 id, u32 version) : m_pathTable(IsCaseInsensitive(id), 0, 0, version), m_id(id) {}
+
+	bool CacheBucket::Load(Logger& logger, BinaryReader& reader, u32 databaseVersion, LoadStats& outStats, StorageServer& storage)
 	{
 		if (databaseVersion != CacheFileVersion)
-			bucket.needsSave = true;
+			needsSave = true;
 
 		if (databaseVersion > 8)
-			bucket.m_pathTable.AddCommonStringSegments();
+			m_pathTable.AddCommonStringSegments();
 
 		u32 pathTableSize = reader.ReadU32();
 		if (pathTableSize)
 		{
-			u64 skipCommonSegments = bucket.m_pathTable.GetSize();
+			u64 skipCommonSegments = m_pathTable.GetSize();
 			BinaryReader pathTableReader(reader.GetPositionData(), skipCommonSegments, pathTableSize);
-			bucket.m_pathTable.ReadMem(pathTableReader, true);
+			m_pathTable.ReadMem(pathTableReader, true);
 			reader.Skip(pathTableSize);
 		}
 		outStats.totalPathTableSize += pathTableSize;
@@ -352,27 +341,38 @@ namespace uba
 		if (casKeyTableSize)
 		{
 			BinaryReader casKeyTableReader(reader.GetPositionData(), 0, casKeyTableSize);
-			bucket.m_casKeyTable.ReadMem(casKeyTableReader, true);
+			m_casKeyTable.ReadMem(casKeyTableReader, true);
 			reader.Skip(casKeyTableSize);
 
-			bucket.m_casKeyTable.Debug(bucket.m_pathTable);
+			m_casKeyTable.Debug(m_pathTable);
 		}
 		outStats.totalCasKeyTableSize += casKeyTableSize;
 
 		u32 entryLookupCount = reader.ReadU32();
-		bucket.m_cacheEntryLookup.reserve(entryLookupCount);
+		m_cacheEntryLookup.reserve(entryLookupCount);
 
 		while (entryLookupCount--)
 		{
-			auto insres = bucket.m_cacheEntryLookup.try_emplace(reader.ReadCasKey());
+			auto insres = m_cacheEntryLookup.try_emplace(reader.ReadCasKey());
 			UBA_ASSERT(insres.second);
 			auto& cacheEntries = insres.first->second;
-			cacheEntries.ReadFromDisk(m_logger, reader, databaseVersion, m_storage, bucket.m_casKeyTable);
+			cacheEntries.ReadFromDisk(logger, reader, databaseVersion, storage, m_casKeyTable);
 			outStats.totalCacheEntryCount += cacheEntries.entries.size();
 		}
 		return true;
 	}
 
+	bool CacheBucket::Validate(Logger& logger, WorkManager& workManager)
+	{
+		Atomic<bool> success = true;
+		workManager.ParallelFor(16, m_cacheEntryLookup, [&](const WorkContext&, auto& it)
+			{
+				if (!it->second.Validate(logger))
+					success = false;
+			}, TCV("ValidateBucket"));
+		return success;
+	}
+
 	bool CacheServer::Save()
 	{
 		for (auto& kv : m_buckets)
@@ -548,7 +548,7 @@ namespace uba
 					bucket.needsSave = false;
 				else
 					success = false;
-			});
+			}, TCV("SaveNoLock"));
 
 		return success;
 	}
@@ -1517,12 +1517,12 @@ namespace uba
 			if (reader.GetLeft())
 				success = reader.ReadBool();
 
-			SCOPED_FUTEX(connectionBucket.cacheEntryLookupLock, lock2);
-			auto findIt = connectionBucket.cacheEntryLookup.find(cmdKey);
-			if (findIt == connectionBucket.cacheEntryLookup.end())
+			SCOPED_FUTEX(connectionBucket.deferredCacheEntryLookupLock, lock2);
+			auto findIt = connectionBucket.deferredCacheEntryLookup.find(cmdKey);
+			if (findIt == connectionBucket.deferredCacheEntryLookup.end())
 				return true;
 			CacheEntry entry(std::move(findIt->second));
-			connectionBucket.cacheEntryLookup.erase(findIt);
+			connectionBucket.deferredCacheEntryLookup.erase(findIt);
 			lock2.Leave();
 			if (!success)
 				return true;
@@ -1818,8 +1818,8 @@ namespace uba
 		}
 		else
 		{
-			SCOPED_FUTEX(connectionBucket.cacheEntryLookupLock, lock3);
-			bool res = connectionBucket.cacheEntryLookup.try_emplace(cmdKey, std::move(newEntry)).second;
+			SCOPED_FUTEX(connectionBucket.deferredCacheEntryLookupLock, lock3);
+			bool res = connectionBucket.deferredCacheEntryLookup.try_emplace(cmdKey, std::move(newEntry)).second;
 			UBA_ASSERT(res);(void)res;
 		}
 
diff --git a/Engine/Source/Programs/UnrealBuildAccelerator/Common/Private/UbaStorage.cpp b/Engine/Source/Programs/UnrealBuildAccelerator/Common/Private/UbaStorage.cpp
--- a/Engine/Source/Programs/UnrealBuildAccelerator/Common/Private/UbaStorage.cpp
+++ b/Engine/Source/Programs/UnrealBuildAccelerator/Common/Private/UbaStorage.cpp
@@ -1347,7 +1347,7 @@ namespace uba
 		return m_casLookup.try_emplace(casKey, casKey).first->second;
 	}
 
-	bool StorageImpl::LoadCasTable(bool logStats, bool alwaysCheckAllFiles)
+	bool StorageImpl::LoadCasTable(bool logStats, bool alwaysCheckAllFiles, bool* outWasTerminated)
 	{
 		if (!m_writeToDisk)
 			return true;
@@ -1370,6 +1370,10 @@ namespace uba
 		TStrcpy_s(isRunningName, sizeof_array(isRunningName), m_rootDir.data);
 		TStrcat_s(isRunningName, sizeof_array(isRunningName), TC(".isRunning"));
 		bool wasTerminated = FileExists(m_logger, isRunningName);
+
+		if (outWasTerminated)
+			*outWasTerminated = wasTerminated;
+
 		if (!wasTerminated)
 		{
 			FileAccessor isRunningFile(m_logger, isRunningName);
diff --git a/Engine/Source/Programs/UnrealBuildAccelerator/Common/Public/UbaCacheBucket.h b/Engine/Source/Programs/UnrealBuildAccelerator/Common/Public/UbaCacheBucket.h
new file mode 100644
--- /dev/null
+++ b/Engine/Source/Programs/UnrealBuildAccelerator/Common/Public/UbaCacheBucket.h
@@ -0,0 +1,49 @@
+// Copyright Epic Games, Inc. All Rights Reserved.
+
+#pragma once
+
+#include "UbaCacheEntry.h"
+#include "UbaCompactTables.h"
+
+namespace uba
+{
+	class StorageServer;
+	class WorkManager;
+
+	struct CacheBucket
+	{
+		CacheBucket(u64 id, u32 version);
+
+		ReaderWriterLock m_cacheEntryLookupLock;
+		UnorderedMap<CasKey, CacheEntries> m_cacheEntryLookup;
+
+		CompactPathTable m_pathTable;
+		CompactCasKeyTable m_casKeyTable;
+
+		u64 m_id;
+		u64 totalEntryCount;
+		u64 totalEntrySize;
+		Atomic<bool> hasDeletedEntries;
+		Atomic<bool> needsSave;
+
+		// Times are in file time from creation of cache database
+		Atomic<u64> lastSavedTime;
+		Atomic<u64> lastUsedTime;
+		u64 oldestUsedTime = 0; 
+
+		u32 index = ~0u;
+
+		struct MaintenanceContext;
+		MaintenanceContext* m_maintenanceContext = nullptr;
+
+		struct LoadStats
+		{
+			Atomic<u32> totalPathTableSize;
+			Atomic<u32> totalCasKeyTableSize;
+			Atomic<u64> totalCacheEntryCount;
+		};
+
+		bool Load(Logger& logger, BinaryReader& reader, u32 databaseVersion, LoadStats& outStats, StorageServer& storage);
+		bool Validate(Logger& logger, WorkManager& workManager);
+	};
+}
\ No newline at end of file
diff --git a/Engine/Source/Programs/UnrealBuildAccelerator/Common/Public/UbaCacheEntry.h b/Engine/Source/Programs/UnrealBuildAccelerator/Common/Public/UbaCacheEntry.h
--- a/Engine/Source/Programs/UnrealBuildAccelerator/Common/Public/UbaCacheEntry.h
+++ b/Engine/Source/Programs/UnrealBuildAccelerator/Common/Public/UbaCacheEntry.h
@@ -51,6 +51,7 @@ namespace uba
 		bool ReadFromDisk(Logger& logger, BinaryReader& reader, u32 databaseVersion, StorageImpl& storage, CompactCasKeyTable& table);
 		void BuildInputs(CacheEntry& entry, const Set<u32>& inputs);
 		void UpdateEntries(Logger& logger, const HashMap2<u32, u32>& oldToNewCasKeyOffset, Vector<u32>& temp, Vector<u8>& temp2, Vector<u8>& temp3);
+		bool Validate(Logger& logger);
 
 		void Flatten(Vector<u8>& out, const CacheEntry& entry);
 		void Flatten(Vector<u32>& out, const CacheEntry& entry, const Vector<u8>& sharedOffsets);
diff --git a/Engine/Source/Programs/UnrealBuildAccelerator/Common/Public/UbaCacheServer.h b/Engine/Source/Programs/UnrealBuildAccelerator/Common/Public/UbaCacheServer.h
--- a/Engine/Source/Programs/UnrealBuildAccelerator/Common/Public/UbaCacheServer.h
+++ b/Engine/Source/Programs/UnrealBuildAccelerator/Common/Public/UbaCacheServer.h
@@ -14,6 +14,7 @@ namespace uba
 	class StorageServer;
 	struct BinaryReader;
 	struct BinaryWriter;
+	struct CacheBucket;
 	struct CacheEntry;
 	struct CacheEntries;
 	struct ConnectionInfo;
@@ -53,7 +54,7 @@ namespace uba
 		CacheServer(const CacheServerCreateInfo& info);
 		~CacheServer();
 
-		bool Load();
+		bool Load(bool validateBuckets);
 		bool Save();
 
 		void SetForceFullMaintenance();
@@ -63,11 +64,9 @@ namespace uba
 		bool ShouldShutdown();
 
 	private:
-		struct Bucket;
-		struct LoadStats;
+		using Bucket = CacheBucket;
 
 		bool RunMaintenanceInternal(const Function<bool()>& shouldExit, bool allowSave);
-		bool LoadBucket(Bucket& bucket, BinaryReader& reader, u32 databaseVersion, LoadStats& outStats);
 		bool SaveBucket(Bucket& bucket, Vector<u8>& temp);
 		bool SaveNoLock();
 		bool SaveDbNoLock();
diff --git a/Engine/Source/Programs/UnrealBuildAccelerator/Common/Public/UbaStorage.h b/Engine/Source/Programs/UnrealBuildAccelerator/Common/Public/UbaStorage.h
--- a/Engine/Source/Programs/UnrealBuildAccelerator/Common/Public/UbaStorage.h
+++ b/Engine/Source/Programs/UnrealBuildAccelerator/Common/Public/UbaStorage.h
@@ -118,7 +118,7 @@ namespace uba
 		StorageImpl(const StorageCreateInfo& info, const tchar* logPrefix = TC("UbaStorage"));
 		virtual ~StorageImpl();
 
-		bool LoadCasTable(bool logStats = true, bool alwaysCheckAllFiles = false);
+		bool LoadCasTable(bool logStats = true, bool alwaysCheckAllFiles = false, bool* outWasTerminated = nullptr);
 		bool CheckCasContent(u32 workerCount);
 		bool CheckFileTable(const tchar* searchPath, u32 workerCount);
 		const tchar* GetTempPath();
