[UBA]
* Fixed permission denied error on linux when failing to delete file
* Minor opt

diff --git a/Engine/Source/Programs/UnrealBuildAccelerator/Common/Private/UbaFile.cpp b/Engine/Source/Programs/UnrealBuildAccelerator/Common/Private/UbaFile.cpp
--- a/Engine/Source/Programs/UnrealBuildAccelerator/Common/Private/UbaFile.cpp
+++ b/Engine/Source/Programs/UnrealBuildAccelerator/Common/Private/UbaFile.cpp
@@ -657,7 +657,13 @@ namespace uba
 		{
 			SetLastError(ERROR_FILE_NOT_FOUND);
 			return false;
-	}
+		}
+		if (errno == EPERM)
+		{
+			SetLastError(ERROR_ACCESS_DENIED);
+			return false;
+		}
+
 		UBA_ASSERTF(false, TC("DeleteFileW failed on %s - Error handling not implemented (%s)"), fileName, strerror(errno));
 		return false;
 #endif
diff --git a/Engine/Source/Programs/UnrealBuildAccelerator/Core/Public/UbaBinaryReaderWriter.h b/Engine/Source/Programs/UnrealBuildAccelerator/Core/Public/UbaBinaryReaderWriter.h
--- a/Engine/Source/Programs/UnrealBuildAccelerator/Core/Public/UbaBinaryReaderWriter.h
+++ b/Engine/Source/Programs/UnrealBuildAccelerator/Core/Public/UbaBinaryReaderWriter.h
@@ -355,25 +355,26 @@ namespace uba
 	{
 		#if PLATFORM_WINDOWS
 		tchar* it = str;
+		const u8* pos = m_pos;
 		u64 left = charLen;
 		while (left--)
 		{
 			UBA_ASSERT_READ(1);
-			u8 a = *m_pos++;
+			u8 a = *pos++;
 			if (a <= 127)
 			{
 				*it++ = a;
 				continue;
 			}
 			UBA_ASSERT_READ(1);
-			u8 b = *m_pos++;
+			u8 b = *pos++;
 			if (a >= 192 && a <= 223)
 			{
 				*it++ = (a-192)*64 + (b-128);
 				continue;
 			}
 			UBA_ASSERT_READ(1);
-			u8 c = *m_pos++;
+			u8 c = *pos++;
 			if (a >= 224 && a <= 239)
 			{
 				*it++ = (a-224)*4096 + (b-128)*64 + (c-128);
@@ -389,6 +390,7 @@ namespace uba
 			*it++ = tchar(~0);
 		}
 		*it = 0;
+		m_pos = pos;
 		return u64(it - str);
 		#else
 		ReadBytes(str, charLen);
@@ -450,22 +452,25 @@ namespace uba
 
 	u64 BinaryReader::Read7BitEncoded()
 	{
+		const u8* pos = m_pos;
 		u64 result = 0;
 		u64 byteIndex = 0;
 		bool hasMoreBytes;
 		do
 		{
 			UBA_ASSERT_READ(1);
-			u8 value = *m_pos++;
+			u8 value = *pos++;
 			hasMoreBytes = value & 0x80;
 			result |= u64(value & 0x7f) << (byteIndex * 7);
 			++byteIndex;
 		} while (hasMoreBytes);
+		m_pos = pos;
 		return result;
 	}
 
 	bool BinaryReader::TryRead7BitEncoded(u64& outValue)
 	{
+		const u8* pos = m_pos;
 		u64 result = 0;
 		u64 byteIndex = 0;
 		bool hasMoreBytes;
@@ -473,12 +478,13 @@ namespace uba
 		{
 			if (m_pos >= m_end)
 				return false;
-			u8 value = *m_pos++;
+			u8 value = *pos++;
 			hasMoreBytes = value & 0x80;
 			result |= u64(value & 0x7f) << (byteIndex * 7);
 			++byteIndex;
 		} while (hasMoreBytes);
 		outValue = result;
+		m_pos = pos;
 		return true;
 	}
 
