[UBA]
* Fixed cache server unit tests

diff --git a/Engine/Source/Programs/UnrealBuildAccelerator/Test/Private/UbaTestCache.cpp b/Engine/Source/Programs/UnrealBuildAccelerator/Test/Private/UbaTestCache.cpp
--- a/Engine/Source/Programs/UnrealBuildAccelerator/Test/Private/UbaTestCache.cpp
+++ b/Engine/Source/Programs/UnrealBuildAccelerator/Test/Private/UbaTestCache.cpp
@@ -1,7 +1,7 @@
 // Copyright Epic Games, Inc. All Rights Reserved.
 
 #include "UbaCacheClient.h"
-#include "UbaCacheEntry.h"
+#include "UbaCacheBucket.h"
 #include "UbaCacheServer.h"
 #include "UbaConfig.h"
 #include "UbaCompactTables.h"
@@ -17,77 +17,99 @@
 
 namespace uba
 {
-	bool TestCacheEntry(LoggerWithWriter& logger, const StringBufferBase& rootDir)
+	struct TestRecord
 	{
-		CacheEntries entries;
+		Set<u32> offsets;
+		u32 expectedRangeCount;
+		u32 expectedExtraCount;
+	};
 
-		struct Record
+	bool CheckCacheEntry(CacheEntries& entries, CacheEntry& entry, const TestRecord* rec, bool checkOffsets, bool checkExpectedCounts)
+	{
+		Set<u32> sharedOffsets;
 		{
-			Set<u32> offsets;
-			u32 expectedRangeCount;
-			u32 expectedExtraCount;
-		};
+			BinaryReader sharedReader(entries.sharedInputCasKeyOffsets);
+			while (sharedReader.GetLeft())
+				if (!sharedOffsets.insert(u32(sharedReader.Read7BitEncoded())).second)
+					return false;
+		}
+
+		Set<u32> entryOffsets;
 
-		auto CheckEntry = [&](CacheEntry& entry, const Record* rec, bool checkOffsets, bool checkExpectedCounts)
+		u32 rangeCount = 0;
+		BinaryReader rangeReader(entry.sharedInputCasKeyOffsetRanges);
+		while (rangeReader.GetLeft())
+		{
+			++rangeCount;
+			u64 begin = rangeReader.Read7BitEncoded();
+			u64 end = rangeReader.Read7BitEncoded();
+			BinaryReader sharedReader(entries.sharedInputCasKeyOffsets.data(), begin, end);
+			while (sharedReader.GetLeft())
 			{
-				Set<u32> sharedOffsets;
-				{
-					BinaryReader sharedReader(entries.sharedInputCasKeyOffsets);
-					while (sharedReader.GetLeft())
-						if (!sharedOffsets.insert(u32(sharedReader.Read7BitEncoded())).second)
-							return false;
-				}
+				u64 offset;
+				if (!sharedReader.TryRead7BitEncoded(offset))
+					return false;
+				if (!entryOffsets.insert(u32(offset)).second)
+					return false;
+			}
+		}
 
-				Set<u32> entryOffsets;
+		u32 extraCount = 0;
+		BinaryReader extraReader(entry.extraInputCasKeyOffsets);
+		while (extraReader.GetLeft())
+		{
+			++extraCount;
+			u32 offset = u32(extraReader.Read7BitEncoded());
+			if (sharedOffsets.find(offset) != sharedOffsets.end())
+				return false;
+			if (!entryOffsets.insert(offset).second)
+				return false;
+		}
 
-				u32 rangeCount = 0;
-				BinaryReader rangeReader(entry.sharedInputCasKeyOffsetRanges);
-				while (rangeReader.GetLeft())
-				{
-					++rangeCount;
-					u64 begin = rangeReader.Read7BitEncoded();
-					u64 end = rangeReader.Read7BitEncoded();
-					BinaryReader sharedReader(entries.sharedInputCasKeyOffsets.data(), begin, end);
-					while (sharedReader.GetLeft())
-						if (!entryOffsets.insert(u32(sharedReader.Read7BitEncoded())).second)
-							return false;
-				}
+		if (rec)
+		{
+			if (checkExpectedCounts)
+			{
+				if (rangeCount != rec->expectedRangeCount)
+					return false;
+				if (extraCount != rec->expectedExtraCount)
+					return false;
+			}
 
-				u32 extraCount = 0;
-				BinaryReader extraReader(entry.extraInputCasKeyOffsets);
-				while (extraReader.GetLeft())
-				{
-					++extraCount;
-					u32 offset = u32(extraReader.Read7BitEncoded());
-					if (sharedOffsets.find(offset) != sharedOffsets.end())
-						return false;
-					if (!entryOffsets.insert(offset).second)
-						return false;
-				}
+			if (checkOffsets)
+			{
+				if (rec->offsets != entryOffsets)
+					return false;
+			}
+			else
+			{
+				if (rec->offsets.size() != entryOffsets.size())
+					return false;
+			}
+		}
+		return true;
+	}
 
-				if (rec)
-				{
-					if (checkExpectedCounts)
-					{
-						if (rangeCount != rec->expectedRangeCount)
-							return false;
-						if (extraCount != rec->expectedExtraCount)
-							return false;
-					}
-
-					if (checkOffsets)
-					{
-						if (rec->offsets != entryOffsets)
-							return false;
-					}
-					else
-					{
-						if (rec->offsets.size() != entryOffsets.size())
-							return false;
-					}
-				}
-				return true;
-			};
+	void UpdateCacheEntries(LoggerWithWriter& logger, CacheEntries& entries, u32 multiplier)
+	{
+		HashMap2<u32, u32> oldToNewCasKeyOffset;
+		MemoryBlock memory(64*1024);
+		oldToNewCasKeyOffset.Init(memory, 100);
+		if (multiplier)
+			for (u32 i=0; i!=200; ++i)
+				oldToNewCasKeyOffset.Insert(i) = i * multiplier;
+		Vector<u32> temp;
+		Vector<u8> temp2;
+		Vector<u8> temp3;
+		entries.UpdateEntries(logger, oldToNewCasKeyOffset, temp, temp2, temp3);
+	}
+
+	bool TestCacheEntry(LoggerWithWriter& logger, const StringBufferBase& rootDir)
+	{
+		CacheEntries entries;
+
+		using Record = TestRecord;
+		auto CheckEntry = [&](CacheEntry& entry, const Record* rec, bool checkOffsets, bool checkExpectedCounts) { return CheckCacheEntry(entries, entry, rec, checkOffsets, checkExpectedCounts); };
 
 		auto AddEntry = [&](const Record& rec, bool checkExpectedCount)
 			{
@@ -106,19 +128,7 @@ namespace uba
 				return CheckEntry(entry, &rec, true, checkExpectedCount);
 			};
 
-		auto UpdateEntries = [&](u32 multiplier)
-			{
-				HashMap2<u32, u32> oldToNewCasKeyOffset;
-				MemoryBlock memory(64*1024);
-				oldToNewCasKeyOffset.Init(memory, 100);
-				if (multiplier)
-					for (u32 i=0; i!=200; ++i)
-						oldToNewCasKeyOffset.Insert(i) = i * multiplier;
-				Vector<u32> temp;
-				Vector<u8> temp2;
-				Vector<u8> temp3;
-				entries.UpdateEntries(logger, oldToNewCasKeyOffset, temp, temp2, temp3);
-			};
+		auto UpdateEntries = [&](u32 multiplier) { UpdateCacheEntries(logger, entries, multiplier); };
 
 		auto Clear = [&]()
 			{
@@ -131,14 +141,21 @@ namespace uba
 			};
 
 		#if 0
+		StorageImpl* storageImpl  = nullptr;
+		CompactCasKeyTable* cckt = nullptr;
+		#if PLATFORM_WINDOWS
 		FileAccessor fa(logger, TC("e:\\temp\\CacheEntry.bin"));
+		#else
+		FileAccessor fa(logger, TC("/mnt/e/temp/CacheEntry.bin"));
+		#endif
 		fa.OpenMemoryRead();
 		BinaryReader fileReader(fa.GetData(), 0, fa.GetSize());
-		entries.ReadFromDisk(logger, fileReader, 9, *(StorageImpl*)nullptr, *(CompactCasKeyTable*)nullptr);
+		entries.ReadFromDisk(logger, fileReader, 9, *storageImpl, *cckt);
 		fa.Close();
-		UpdateEntries(0);
+		//UpdateEntries(0);
 		for (auto& entry : entries.entries)
 			CHECK_TRUE(CheckEntry(entry, nullptr, false, false));
+		return true;
 		#endif
 
 		Vector<Record> records0 = 
@@ -512,7 +529,7 @@ namespace uba
 		{
 		CacheServerCreateInfo csci(storageServer, rootDir.data, logWriter);
 		CacheServer cacheServer(csci);
-		if (!cacheServer.Load())
+		if (!cacheServer.Load(false))
 			return false;
 
 
@@ -652,9 +669,55 @@ namespace uba
 
 		CacheServerCreateInfo csci(storageServer, rootDir.data, logWriter);
 		CacheServer cacheServer(csci);
-		if (!cacheServer.Load())
+		if (!cacheServer.Load(false))
+			return false;
+
+		return true;
+	}
+
+#if 0
+	bool TestLoadCache(LoggerWithWriter& logger, const StringBufferBase& testRootDir)
+	{
+		#if PLATFORM_WINDOWS
+		const tchar* fileName = TC("e:\\temp\\uploads\\114629205754");
+		#else
+		const tchar* fileName = TC("/mnt/e/temp/uploads/114629205754");
+		#endif
+
+		FileAccessor bucketFile(logger, fileName);
+		if (!bucketFile.OpenMemoryRead())
 			return false;
+		BinaryReader reader(bucketFile.GetData(), 0, bucketFile.GetSize());
+		u32 bucketVersion = reader.ReadU32();
+		CacheBucket::LoadStats stats;
+		StorageServer* storage = nullptr;
 
+		CacheBucket bucket(0, 0);
+		if (!bucket.Load(logger, reader, bucketVersion, stats, *storage))
+			return false;
+		bucketFile.Close();
+
+		LoggerWithWriter lg(g_consoleLogWriter);
+		u32 entriesIndex = 0;
+		for (auto& kv : bucket.m_cacheEntryLookup)
+		{
+			lg.Info(TC("Entries %u"), entriesIndex++);
+			auto& entries = kv.second;
+
+			if (entriesIndex == 358)
+			{
+				FileAccessor fa(logger, TC("e:\\temp\\CacheEntry.bin"));
+				fa.CreateMemoryWrite(false, DefaultAttributes(), entries.GetTotalSize(CacheNetworkVersion, true));
+				BinaryWriter w(fa.GetData(), 0, fa.GetSize());
+				entries.Write(w, CacheNetworkVersion, true);
+				fa.Close();
+			}
+
+			u32 entryIndex = 0;
+			for (auto& entry : entries.entries)
+				CHECK_TRUE(CheckCacheEntry(entries, entry, nullptr, false, false));
+		}
 		return true;
 	}
+#endif
 }
\ No newline at end of file
