[UBA]
* Changed DeleteFile to DeleteFileW

diff --git a/Engine/Source/Programs/UnrealBuildAccelerator/Common/Private/UbaCacheServer.cpp b/Engine/Source/Programs/UnrealBuildAccelerator/Common/Private/UbaCacheServer.cpp
--- a/Engine/Source/Programs/UnrealBuildAccelerator/Common/Private/UbaCacheServer.cpp
+++ b/Engine/Source/Programs/UnrealBuildAccelerator/Common/Private/UbaCacheServer.cpp
@@ -301,7 +301,7 @@ namespace uba
 
 					bucketFile.Close();
 					m_logger.Info(TC("Loaded invalid bucket %s. Will be deleted"), bucketFilename.data);
-					DeleteFile(bucketFilename.data);
+					DeleteFileW(bucketFilename.data);
 					SCOPED_FUTEX(toDeleteLock, lock);
 					toDelete.insert(key);
 
