[UBA]
* Added AllowLinkDependencyCrawler option to session that can be used to disable link dependency crawler in order to prevent memory usage explosion when starting monolithic links

diff --git a/Engine/Source/Programs/UnrealBuildAccelerator/Common/Private/UbaSession.cpp b/Engine/Source/Programs/UnrealBuildAccelerator/Common/Private/UbaSession.cpp
--- a/Engine/Source/Programs/UnrealBuildAccelerator/Common/Private/UbaSession.cpp
+++ b/Engine/Source/Programs/UnrealBuildAccelerator/Common/Private/UbaSession.cpp
@@ -229,6 +229,7 @@ namespace uba
 		table.GetValueAsBool(extractObjFilesSymbols, TC("ExtractObjFilesSymbols"));
 		table.GetValueAsBool(useFakeVolumeSerial, TC("UseFakeVolumeSerial"));
 		table.GetValueAsBool(keepTransientDataMapped, TC("KeepTransientDataMapped"));
+		table.GetValueAsBool(allowLinkDependencyCrawler, TC("AllowLinkDependencyCrawler"));
 		table.GetValueAsU32(traceReserveSizeMb, TC("TraceReserveSizeMb"));
 		table.GetValueAsU32(writeFilesBottleneck, TC("WriteFilesBottleneck"));
 		table.GetValueAsU32(writeFilesFileMapMaxMb, TC("WriteFilesFileMapMaxMB"));
@@ -1494,6 +1495,7 @@ namespace uba
 		m_readIntermediateFilesCompressed = (m_storeIntermediateFilesCompressed || (info.readIntermediateFilesCompressed && IsWindows)) && !runningRemote; // with remote we decompress the files into memory
 		m_allowLocalDetour = info.allowLocalDetour;
 		m_extractObjFilesSymbols = info.extractObjFilesSymbols;
+		m_allowLinkDependencyCrawler = info.allowLinkDependencyCrawler;
 
 		m_detailedTrace = info.detailedTrace;
 		m_traceChildProcesses = info.traceChildProcesses;
@@ -2496,7 +2498,7 @@ namespace uba
 			startInfo.application = startInfo.applicationStr.c_str();
 		}
 
-		if (!isChild && !m_runningRemote && m_readIntermediateFilesCompressed)
+		if (!isChild && !m_runningRemote && m_readIntermediateFilesCompressed && m_allowLinkDependencyCrawler)
 		{
 			auto crawlerType = startInfo.rules->GetDependencyCrawlerType();
 			if (crawlerType == DependencyCrawlerType_MsvcLinker || crawlerType == DependencyCrawlerType_ClangLinker)
diff --git a/Engine/Source/Programs/UnrealBuildAccelerator/Common/Public/UbaSession.h b/Engine/Source/Programs/UnrealBuildAccelerator/Common/Public/UbaSession.h
--- a/Engine/Source/Programs/UnrealBuildAccelerator/Common/Public/UbaSession.h
+++ b/Engine/Source/Programs/UnrealBuildAccelerator/Common/Public/UbaSession.h
@@ -234,6 +234,7 @@ namespace uba
 		bool m_readIntermediateFilesCompressed;
 		bool m_allowLocalDetour;
 		bool m_extractObjFilesSymbols;
+		bool m_allowLinkDependencyCrawler;
 
 		u64 m_keepOutputFileMemoryMapsThreshold;
 
diff --git a/Engine/Source/Programs/UnrealBuildAccelerator/Common/Public/UbaSessionCreateInfo.h b/Engine/Source/Programs/UnrealBuildAccelerator/Common/Public/UbaSessionCreateInfo.h
--- a/Engine/Source/Programs/UnrealBuildAccelerator/Common/Public/UbaSessionCreateInfo.h
+++ b/Engine/Source/Programs/UnrealBuildAccelerator/Common/Public/UbaSessionCreateInfo.h
@@ -42,6 +42,7 @@ namespace uba
 		bool treatTempDirAsEmpty = true;		// If this is true, directory table will always see temp as empty. Use sub directories of temp instead
 		bool useFakeVolumeSerial = true;		// Fake volume serials reduce the directory table size since serials end up being compressed down to one byte
 		bool keepTransientDataMapped = true;	// Will keep the transient data mapped in instead of doing map/unmap everytime we access it
+		bool allowLinkDependencyCrawler = true;	// Enable crawler that reads ahead to find needed obj files for linker
 		u64 deleteSessionsOlderThanSeconds = 12 * 60 * 60; // Delete session folders older than 12 hours by default . Set to 0 to not delete or 1 to delete all
 		u64 keepOutputFileMemoryMapsThreshold = 256 * 1024; // If allowMemoryMaps is true, output files will be kept in memory if smaller than this size
 		u32 traceReserveSizeMb = 128;			// Memory reserved for trace file (in mb)
