[UBA]
* Implemented "ReportBadProxy" message that is sent by clients when they fail to connect to the proxy.

diff --git a/Engine/Source/Programs/UnrealBuildAccelerator/Common/Private/UbaCacheServer.cpp b/Engine/Source/Programs/UnrealBuildAccelerator/Common/Private/UbaCacheServer.cpp
--- a/Engine/Source/Programs/UnrealBuildAccelerator/Common/Private/UbaCacheServer.cpp
+++ b/Engine/Source/Programs/UnrealBuildAccelerator/Common/Private/UbaCacheServer.cpp
@@ -300,7 +300,8 @@ namespace uba
 							return;
 
 					bucketFile.Close();
-					m_logger.Info(TC("Loaded invalid bucket %s. Will be deleted"), bucketFilename.data);
+
+					m_logger.Info(TC("Found invalid bucket %s. Deleting"), bucketFilename.data);
 					DeleteFileW(bucketFilename.data);
 					SCOPED_FUTEX(toDeleteLock, lock);
 					toDelete.insert(key);
diff --git a/Engine/Source/Programs/UnrealBuildAccelerator/Common/Private/UbaNetworkServer.cpp b/Engine/Source/Programs/UnrealBuildAccelerator/Common/Private/UbaNetworkServer.cpp
--- a/Engine/Source/Programs/UnrealBuildAccelerator/Common/Private/UbaNetworkServer.cpp
+++ b/Engine/Source/Programs/UnrealBuildAccelerator/Common/Private/UbaNetworkServer.cpp
@@ -169,8 +169,12 @@ namespace uba
 				Sleep(1);
 			}
 
-			if (!m_disconnectCallbackCalled.IsSet(60000)) // This should never time out!
-				m_server.m_logger.Warning(TC("This should never happen!! Unknown consequences"));
+			if (!m_disconnectCallbackCalled.IsSet(30*1000)) // This should never time out!
+			{
+				m_server.m_logger.Warning(TC("Disconnect callback event timed out. This should never happen!!"));
+				PrintAllCallstacks(m_server.m_logger);
+			}
+
 			return true;
 		}
 
diff --git a/Engine/Source/Programs/UnrealBuildAccelerator/Common/Private/UbaStorageClient.cpp b/Engine/Source/Programs/UnrealBuildAccelerator/Common/Private/UbaStorageClient.cpp
--- a/Engine/Source/Programs/UnrealBuildAccelerator/Common/Private/UbaStorageClient.cpp
+++ b/Engine/Source/Programs/UnrealBuildAccelerator/Common/Private/UbaStorageClient.cpp
@@ -412,57 +412,102 @@ namespace uba
 					if (m_proxyClient)
 						continue;
 
-					NetworkClientCreateInfo ncci(m_logger.m_writer);
-					ncci.workerCount = 0;
-					bool ctorSuccess = true;
-					proxy = new ProxyClient(ctorSuccess, ncci);
-					m_proxyClient = proxy;
+					auto createProxyClient = [&]()
+						{
+							NetworkClientCreateInfo ncci(m_logger.m_writer);
+							ncci.workerCount = 0;
+							bool ctorSuccess = true;
+							proxy = new ProxyClient(ctorSuccess, ncci);
+							m_proxyClient = proxy;
 
-					auto disallowProxy = MakeGuard([&]() { m_proxyClient->client.Disconnect(); proxy = nullptr; allowProxy = false; });
+							auto disallowProxy = MakeGuard([&]() { m_proxyClient->client.Disconnect(); proxy = nullptr; allowProxy = false; });
 
-					if (!ctorSuccess)
-						continue;
+							if (!ctorSuccess)
+								return false;
 
-					NetworkBackend& proxyBackend = m_getProxyBackendCallback(m_getProxyBackendUserData, proxyHost.data);
+							NetworkBackend& proxyBackend = m_getProxyBackendCallback(m_getProxyBackendUserData, proxyHost.data);
 
-					u64 startTime = GetTime();
-					if (!proxy->client.Connect(proxyBackend, proxyHost.data, proxyPort))
-					{
-						m_logger.Detail(TC("Connecting to proxy %s:%u for cas %s download failed! (%s) (%s)"), proxyHost.data, proxyPort, casFile.data, hint, TimeToText(GetTime() - startTime).str);
+							u64 startTime = GetTime();
+							if (!proxy->client.Connect(proxyBackend, proxyHost.data, proxyPort))
+							{
+								m_logger.Detail(TC("Connecting to proxy %s:%u for cas %s download failed! (%s) (%s)"), proxyHost.data, proxyPort, casFile.data, hint, TimeToText(GetTime() - startTime).str);
+								return false;
+							}
+					
+							u64 connectTime = GetTime() - startTime;
+							if (connectTime > MsToTime(2000))
+								m_logger.Info(TC("Took %s to connect to proxy %s:%u"), TimeToText(connectTime).str, proxyHost.data, proxyPort);
+
+							// Send a message to the proxy just to validate that this is still part of this build
+							{
+								NetworkMessage proxyMsg(proxy->client, ServiceId, StorageMessageType_Connect, writer.Reset());
+
+								writer.WriteString(TC("ProxyClient"));
+								writer.WriteU32(StorageNetworkVersion);
+								writer.WriteBool(isInProcessClient);
+								StackBinaryReader<256> proxyReader;
+								if (!proxyMsg.Send(proxyReader))
+								{
+									m_logger.Info(TC("Failed to send connect message to proxy %s:%u. Will ask storage server for new proxy"), proxyHost.data, proxyPort);
+									return false;
+								}
+								if (proxyReader.ReadGuid() != m_storageServerUid)
+								{
+									m_logger.Info(TC("Proxy %s:%u is not the correct proxy anymore. Will ask storage server for new proxy"), proxyHost.data, proxyPort);
+									return false;
+								}
+							}
+
+							for (u32 i=1;i<m_proxyConnectionCount;++i)
+								proxy->client.Connect(proxyBackend, proxyHost.data, proxyPort);
+
+							disallowProxy.Cancel();
+							++proxy->refCount;
+							proxy->client.SetWorkTracker(m_client.GetWorkTracker());
+							return true;
+						};
+
+					if (createProxyClient())
 						continue;
-					}
+
+					UBA_ASSERT(!isInProcessClient); // Should never fail to connect to itself
 					
-					u64 connectTime = GetTime() - startTime;
-					if (connectTime > MsToTime(2000))
-						m_logger.Info(TC("Took %s to connect to proxy %s:%u"), TimeToText(connectTime).str, proxyHost.data, proxyPort);
+					
+					// If failing to connect to proxy we report it as bad and hopefully get a new one from the server
+					m_logger.Detail(TC("Reporting bad proxy %s:%u"), proxyHost.data, proxyPort);
+					NetworkMessage reportMsg(m_client, ServiceId, StorageMessageType_ReportBadProxy, writer.Reset());
+					writer.WriteString(proxyHost);
+					writer.WriteU16(proxyPort);
+
+					StackBinaryReader<256> badProxyReader;
+					if (!reportMsg.Send(badProxyReader))
+						continue;
 
-					// Send a message to the proxy just to validate that this is still part of this build
+					if (!badProxyReader.GetLeft())
+						continue; 
+
+					if (badProxyReader.ReadBool()) // This will be true for only one message.. no need to guard it
 					{
-						StackBinaryWriter<256> proxyWriter;
-						NetworkMessage proxyMsg(proxy->client, ServiceId, StorageMessageType_Connect, proxyWriter);
-
-						proxyWriter.WriteString(TC("ProxyClient"));
-						proxyWriter.WriteU32(StorageNetworkVersion);
-						proxyWriter.WriteBool(isInProcessClient);
-						StackBinaryReader<256> proxyReader;
-						if (!proxyMsg.Send(proxyReader))
-						{
-							m_logger.Info(TC("Failed to send connect message to proxy %s:%u. Will ask storage server for new proxy"), proxyHost.data, proxyPort);
-							continue;
-						}
-						if (proxyReader.ReadGuid() != m_storageServerUid)
+						proxyPort = badProxyReader.ReadU16();
+						if (!m_startProxyCallback(m_startProxyUserData, proxyPort, m_storageServerUid))
 						{
-							m_logger.Info(TC("Proxy %s:%u is not the correct proxy anymore. Will ask storage server for new proxy"), proxyHost.data, proxyPort);
+							// TODO: Tell server we failed
+							m_logger.Warning(TC("Failed to create proxy server. This should never happen!"));
 							continue;
 						}
+						proxyHost.Append(TCV("inprocess"));
+						isInProcessClient = true;
 					}
-
-					for (u32 i=1;i<m_proxyConnectionCount;++i)
-						proxy->client.Connect(proxyBackend, proxyHost.data, proxyPort);
-
-					disallowProxy.Cancel();
-					++proxy->refCount;
-					proxy->client.SetWorkTracker(m_client.GetWorkTracker());
+					else
+					{
+						badProxyReader.ReadString(proxyHost);
+						proxyPort = badProxyReader.ReadU16();
+					}
+					
+					// Destroy the 'old' new proxy client since we are creating another one in createProxyClient
+					UBA_ASSERT(m_proxyClient->refCount == 0);
+					delete m_proxyClient;
+					createProxyClient();
 
 					continue;
 				}
diff --git a/Engine/Source/Programs/UnrealBuildAccelerator/Common/Private/UbaStorageServer.cpp b/Engine/Source/Programs/UnrealBuildAccelerator/Common/Private/UbaStorageServer.cpp
--- a/Engine/Source/Programs/UnrealBuildAccelerator/Common/Private/UbaStorageServer.cpp
+++ b/Engine/Source/Programs/UnrealBuildAccelerator/Common/Private/UbaStorageServer.cpp
@@ -312,6 +312,74 @@ namespace uba
 		return true;
 	}
 
+	bool StorageServer::GetProxy(BinaryWriter& writer, u32 clientId)
+	{
+		SCOPED_READ_LOCK(m_connectionInfoLock, lock);
+		auto findIt = m_connectionInfo.find(clientId);
+		UBA_ASSERT(findIt != m_connectionInfo.end());
+		Info& info = findIt->second;
+		lock.Leave();
+
+		if (info.zone.empty())
+			return false;
+
+		// Zone logic is a bit special to be able to handle AWS setups.
+		// AWS has zones named a,b,c etc in the end. us-east-1a, us-east-1b etc.
+		// If host is also in AWS, then we want different proxies per zone but if host is not in AWS, we want one proxy for all zones in a region
+
+		StringBuffer<256> proxyName;
+
+		if (!m_zone.empty() && info.zone.size() == m_zone.size() && Equals(m_zone.c_str(), info.zone.c_str(), u64(m_zone.size() - 1)))
+		{
+			// Host is inside AWS and same region.... if aws availability zone is different than host, then we use proxy
+			if (m_zone != info.zone)
+				proxyName.Append(info.zone.c_str());
+		}
+		else if (m_zone != info.zone)
+		{
+			// We remove last character from zone to make sure all AWS availability zones in the same region use same proxy if host zone is outside AWS
+			proxyName.Append(info.zone.c_str(), info.zone.size() - 1);
+		}
+
+		if (proxyName.IsEmpty())
+			return false;
+
+		writer.WriteU16(u16(~0));
+		writer.Write7BitEncoded(0);
+		writer.WriteByte(1 << 2);
+
+		auto proxyKey = ToStringKeyNoCheck(proxyName.data, proxyName.count);
+		SCOPED_WRITE_LOCK(m_proxiesLock, proxiesLock);
+		ProxyEntry& proxy = m_proxies[proxyKey];
+		if (proxy.clientId == ~0u)
+		{
+			proxy.clientId = clientId;
+			proxy.host = info.internalAddress;
+			proxy.port = info.proxyPort;
+			proxy.zone = proxyName.data;
+
+			m_logger.Detail(TC("%s:%u (client %u) is assigned as proxy for zone %s"), proxy.host.c_str(), proxy.port, clientId, proxy.zone.c_str());
+
+			writer.WriteBool(true);
+			writer.WriteU16(info.proxyPort);
+			if (m_trace)
+				m_trace->ProxyCreated(proxy.clientId, proxyName.data);
+		}
+		else
+		{
+			const tchar* proxyHost = proxy.host.c_str();
+			if (clientId == proxy.clientId)
+				proxyHost = TC("inprocess");
+
+			writer.WriteBool(false);
+			writer.WriteString(proxyHost);
+			writer.WriteU16(proxy.port);
+			if (m_trace)
+				m_trace->ProxyUsed(clientId, proxyName.data);
+		}
+		return true;
+	}
+
 	bool StorageServer::HandleMessage(const ConnectionInfo& connectionInfo, const WorkContext& workContext, u8 messageType, BinaryReader& reader, BinaryWriter& writer)
 	{
 		ActiveStore* firstStore = nullptr;
@@ -362,70 +430,8 @@ namespace uba
 				if (wantsProxy) // Wants proxy
 				{
 					todoRemoveMe = 1;
-					SCOPED_READ_LOCK(m_connectionInfoLock, lock);
-					auto findIt = m_connectionInfo.find(connectionInfo.GetId());
-					UBA_ASSERT(findIt != m_connectionInfo.end());
-					Info& info = findIt->second;
-					lock.Leave();
-
-					if (!info.zone.empty())
-					{
-						// Zone logic is a bit special to be able to handle AWS setups.
-						// AWS has zones named a,b,c etc in the end. us-east-1a, us-east-1b etc.
-						// If host is also in AWS, then we want different proxies per zone but if host is not in AWS, we want one proxy for all zones in a region
-
-						StringBuffer<256> proxyName;
-
-						if (!m_zone.empty() && info.zone.size() == m_zone.size() && Equals(m_zone.c_str(), info.zone.c_str(), u64(m_zone.size() - 1)))
-						{
-							// Host is inside AWS and same region.... if aws availability zone is different than host, then we use proxy
-							if (m_zone != info.zone)
-								proxyName.Append(info.zone.c_str());
-						}
-						else if (m_zone != info.zone)
-						{
-							// We remove last character from zone to make sure all AWS availability zones in the same region use same proxy if host zone is outside AWS
-							proxyName.Append(info.zone.c_str(), info.zone.size() - 1);
-						}
-
-						if (!proxyName.IsEmpty())
-						{
-							writer.WriteU16(u16(~0));
-							writer.Write7BitEncoded(0);
-							writer.WriteByte(1 << 2);
-
-							auto proxyKey = ToStringKeyNoCheck(proxyName.data, proxyName.count);
-							SCOPED_WRITE_LOCK(m_proxiesLock, proxiesLock);
-							ProxyEntry& proxy = m_proxies[proxyKey];
-							if (proxy.clientId == ~0u)
-							{
-								proxy.clientId = connectionInfo.GetId();
-								proxy.host = info.internalAddress;
-								proxy.port = info.proxyPort;
-								proxy.zone = proxyName.data;
-
-								m_logger.Detail(TC("%s:%u (%s) is assigned as proxy for zone %s"), proxy.host.c_str(), proxy.port, GuidToString(connectionInfo.GetUid()).str, proxy.zone.c_str());
-
-								writer.WriteBool(true);
-								writer.WriteU16(info.proxyPort);
-								if (m_trace)
-									m_trace->ProxyCreated(proxy.clientId, proxyName.data);
-							}
-							else
-							{
-								const tchar* proxyHost = proxy.host.c_str();
-								if (connectionInfo.GetId() == proxy.clientId)
-									proxyHost = TC("inprocess");
-
-								writer.WriteBool(false);
-								writer.WriteString(proxyHost);
-								writer.WriteU16(proxy.port);
-								if (m_trace)
-									m_trace->ProxyUsed(connectionInfo.GetId(), proxyName.data);
-							}
-							return true;
-						}
-					}
+					if (GetProxy(writer, connectionInfo.GetId()))
+						return true;
 				}
 
 				todoRemoveMe = 2;
@@ -1054,6 +1060,27 @@ namespace uba
 					m_trace->FileFetchEnd(connectionInfo.GetId(), casKey);
 				return true;
 			}
+			case StorageMessageType_ReportBadProxy:
+			{
+				StringBuffer<> proxyHost;
+				reader.ReadString(proxyHost);
+				u16 proxyPort = reader.ReadU16();
+
+				SCOPED_WRITE_LOCK(m_proxiesLock, lock);
+				for (auto it=m_proxies.begin(); it!=m_proxies.end(); ++it)
+				{
+					ProxyEntry& e = it->second;
+					if (e.host != proxyHost.data || e.port != proxyPort)
+						continue;
+					m_logger.Detail(TC("Proxy %s:%u for zone %s removed (bad proxy reported)"), e.host.c_str(), e.port, e.zone.c_str());
+					m_proxies.erase(it);
+					break;
+				}
+				lock.Leave();
+
+				GetProxy(writer, connectionInfo.GetId());
+				return true;
+			}
 		}
 		UBA_ASSERT(false);
 		return false;
diff --git a/Engine/Source/Programs/UnrealBuildAccelerator/Common/Public/UbaNetwork.h b/Engine/Source/Programs/UnrealBuildAccelerator/Common/Public/UbaNetwork.h
--- a/Engine/Source/Programs/UnrealBuildAccelerator/Common/Public/UbaNetwork.h
+++ b/Engine/Source/Programs/UnrealBuildAccelerator/Common/Public/UbaNetwork.h
@@ -16,7 +16,7 @@ namespace uba
 	
 	static constexpr u32 SystemNetworkVersion = 1339;
 	static constexpr u32 StorageNetworkVersion = 4;
-	static constexpr u32 SessionNetworkVersion = 45;
+	static constexpr u32 SessionNetworkVersion = 46;
 	static constexpr u32 CacheNetworkVersion = 5;
 	static constexpr u32 CacheBucketVersion = 3;
 
@@ -38,6 +38,7 @@ namespace uba
 		UBA_STORAGE_MESSAGE(Connect) \
 		UBA_STORAGE_MESSAGE(ProxyFetchBegin) \
 		UBA_STORAGE_MESSAGE(ProxyFetchEnd) \
+		UBA_STORAGE_MESSAGE(ReportBadProxy) \
 
 	#define UBA_SESSION_MESSAGES \
 		UBA_SESSION_MESSAGE(Connect) \
diff --git a/Engine/Source/Programs/UnrealBuildAccelerator/Common/Public/UbaStorageServer.h b/Engine/Source/Programs/UnrealBuildAccelerator/Common/Public/UbaStorageServer.h
--- a/Engine/Source/Programs/UnrealBuildAccelerator/Common/Public/UbaStorageServer.h
+++ b/Engine/Source/Programs/UnrealBuildAccelerator/Common/Public/UbaStorageServer.h
@@ -46,6 +46,7 @@ namespace uba
 		bool HandleMessage(const ConnectionInfo& connectionInfo, const WorkContext& workContext, u8 messageType, BinaryReader& reader, BinaryWriter& writer);
 		bool WaitForWritten(CasEntry& casEntry, ScopedWriteLock& entryLock, const ConnectionInfo& connectionInfo, const tchar* hint);
 		bool VerifyExists(bool& outExists, CasEntry& casEntry, ScopedWriteLock& entryLock, const CasKey& casKey);
+		bool GetProxy(BinaryWriter& writer, u32 clientId);
 
 		u16 PopId();
 		void PushId(u16 id);
