[UBA]
* Fixed bug in reportbadproxy logic where host string was not cleared properly when getting new proxy
* Improved logging to show which client that reported bad proxy

diff --git a/Engine/Source/Programs/UnrealBuildAccelerator/Common/Private/UbaStorageClient.cpp b/Engine/Source/Programs/UnrealBuildAccelerator/Common/Private/UbaStorageClient.cpp
--- a/Engine/Source/Programs/UnrealBuildAccelerator/Common/Private/UbaStorageClient.cpp
+++ b/Engine/Source/Programs/UnrealBuildAccelerator/Common/Private/UbaStorageClient.cpp
@@ -408,6 +408,8 @@ namespace uba
 						proxyPort = reader.ReadU16();
 					}
 
+					u32 proxyClientId = reader.ReadU32();
+
 					SCOPED_WRITE_LOCK(m_proxyClientLock, proxyLock2);
 					if (m_proxyClient)
 						continue;
@@ -476,8 +478,7 @@ namespace uba
 					// If failing to connect to proxy we report it as bad and hopefully get a new one from the server
 					m_logger.Detail(TC("Reporting bad proxy %s:%u"), proxyHost.data, proxyPort);
 					NetworkMessage reportMsg(m_client, ServiceId, StorageMessageType_ReportBadProxy, writer.Reset());
-					writer.WriteString(proxyHost);
-					writer.WriteU16(proxyPort);
+					writer.WriteU32(proxyClientId);
 
 					StackBinaryReader<256> badProxyReader;
 					if (!reportMsg.Send(badProxyReader))
@@ -486,6 +487,8 @@ namespace uba
 					if (!badProxyReader.GetLeft())
 						continue; 
 
+					proxyHost.Clear();
+
 					if (badProxyReader.ReadBool()) // This will be true for only one message.. no need to guard it
 					{
 						proxyPort = badProxyReader.ReadU16();
diff --git a/Engine/Source/Programs/UnrealBuildAccelerator/Common/Private/UbaStorageServer.cpp b/Engine/Source/Programs/UnrealBuildAccelerator/Common/Private/UbaStorageServer.cpp
--- a/Engine/Source/Programs/UnrealBuildAccelerator/Common/Private/UbaStorageServer.cpp
+++ b/Engine/Source/Programs/UnrealBuildAccelerator/Common/Private/UbaStorageServer.cpp
@@ -361,7 +361,6 @@ namespace uba
 			m_logger.Detail(TC("%s:%u (client %u) is assigned as proxy for zone %s"), proxy.host.c_str(), proxy.port, clientId, proxy.zone.c_str());
 
 			writer.WriteBool(true);
-			writer.WriteU16(info.proxyPort);
 			if (m_trace)
 				m_trace->ProxyCreated(proxy.clientId, proxyName.data);
 		}
@@ -373,10 +372,12 @@ namespace uba
 
 			writer.WriteBool(false);
 			writer.WriteString(proxyHost);
-			writer.WriteU16(proxy.port);
 			if (m_trace)
 				m_trace->ProxyUsed(clientId, proxyName.data);
 		}
+
+		writer.WriteU16(proxy.port);
+		writer.WriteU32(proxy.clientId);
 		return true;
 	}
 
@@ -1062,17 +1063,15 @@ namespace uba
 			}
 			case StorageMessageType_ReportBadProxy:
 			{
-				StringBuffer<> proxyHost;
-				reader.ReadString(proxyHost);
-				u16 proxyPort = reader.ReadU16();
+				u32 proxyClientId = reader.ReadU32();
 
 				SCOPED_WRITE_LOCK(m_proxiesLock, lock);
 				for (auto it=m_proxies.begin(); it!=m_proxies.end(); ++it)
 				{
 					ProxyEntry& e = it->second;
-					if (e.host != proxyHost.data || e.port != proxyPort)
+					if (e.clientId != proxyClientId)
 						continue;
-					m_logger.Detail(TC("Proxy %s:%u for zone %s removed (bad proxy reported)"), e.host.c_str(), e.port, e.zone.c_str());
+					m_logger.Detail(TC("Proxy %u (%s:%u) for zone %s removed (bad proxy reported by client %u)"), proxyClientId, e.host.c_str(), e.port, e.zone.c_str(), connectionInfo.GetId());
 					m_proxies.erase(it);
 					break;
 				}
