[UBA]
* Fixed bug where storage client did not use new proxy for cas fetch happening when getting new proxy
* Fixed so cache server checks again that cache entry does not exist after cas files have been uploaded. Since there is a chance some other client uploaded the same cache entry while cas files are being uploaded

diff --git a/Engine/Source/Programs/UnrealBuildAccelerator/Common/Private/UbaCacheServer.cpp b/Engine/Source/Programs/UnrealBuildAccelerator/Common/Private/UbaCacheServer.cpp
--- a/Engine/Source/Programs/UnrealBuildAccelerator/Common/Private/UbaCacheServer.cpp
+++ b/Engine/Source/Programs/UnrealBuildAccelerator/Common/Private/UbaCacheServer.cpp
@@ -1510,36 +1510,8 @@ namespace uba
 			return HandleStoreEntry(bucket, reader, writer, clientVersion, connectionInfo.GetId());
 		}
 		case CacheMessageType_StoreEntryDone:
-		{
-			auto& connectionBucket = GetConnectionBucket(connectionInfo, reader);
-			CasKey cmdKey = reader.ReadCasKey();
-
-			bool success = true;
-			if (reader.GetLeft())
-				success = reader.ReadBool();
-
-			SCOPED_FUTEX(connectionBucket.deferredCacheEntryLookupLock, lock2);
-			auto findIt = connectionBucket.deferredCacheEntryLookup.find(cmdKey);
-			if (findIt == connectionBucket.deferredCacheEntryLookup.end())
-				return true;
-			CacheEntry entry(std::move(findIt->second));
-			connectionBucket.deferredCacheEntryLookup.erase(findIt);
-			lock2.Leave();
-			if (!success)
-				return true;
-
-			u64 id = connectionBucket.id;
-			Bucket& bucket = GetBucket(id, TC("StoreEntryDone"));
+			return HandleStoreEntryDone(connectionInfo, reader);
 
-			SCOPED_WRITE_LOCK(bucket.m_cacheEntryLookupLock, lock3);
-			auto insres = bucket.m_cacheEntryLookup.try_emplace(cmdKey);
-			auto& cacheEntries = insres.first->second;
-			lock3.Leave();
-			
-			SCOPED_WRITE_LOCK(cacheEntries.lock, lock4);
-			cacheEntries.entries.emplace_front(std::move(entry));
-			return true;
-		}
 		case CacheMessageType_FetchPathTable:
 			return HandleFetchPathTable(reader, writer);
 
@@ -1832,6 +1804,44 @@ namespace uba
 		return true;
 	}
 
+	bool CacheServer::HandleStoreEntryDone(const ConnectionInfo& connectionInfo, BinaryReader& reader)
+	{
+		auto& connectionBucket = GetConnectionBucket(connectionInfo, reader);
+		CasKey cmdKey = reader.ReadCasKey();
+
+		bool success = true;
+		if (reader.GetLeft())
+			success = reader.ReadBool();
+
+		SCOPED_FUTEX(connectionBucket.deferredCacheEntryLookupLock, lock2);
+		auto findIt = connectionBucket.deferredCacheEntryLookup.find(cmdKey);
+		if (findIt == connectionBucket.deferredCacheEntryLookup.end())
+			return true;
+		CacheEntry entry(std::move(findIt->second));
+		connectionBucket.deferredCacheEntryLookup.erase(findIt);
+		lock2.Leave();
+		if (!success)
+			return true;
+
+		u64 id = connectionBucket.id;
+		Bucket& bucket = GetBucket(id, TC("StoreEntryDone"));
+
+		SCOPED_WRITE_LOCK(bucket.m_cacheEntryLookupLock, lock3);
+		auto insres = bucket.m_cacheEntryLookup.try_emplace(cmdKey);
+		auto& cacheEntries = insres.first->second;
+		lock3.Leave();
+			
+		SCOPED_WRITE_LOCK(cacheEntries.lock, lock4);
+
+		// Check again because another connection might have added the same entry while cas files were transferred
+		for (auto i=cacheEntries.entries.begin(), e=cacheEntries.entries.end(); i!=e; ++i)
+			if (i->sharedInputCasKeyOffsetRanges == entry.sharedInputCasKeyOffsetRanges && i->extraInputCasKeyOffsets == entry.extraInputCasKeyOffsets)
+				return true;
+
+		cacheEntries.entries.emplace_front(std::move(entry));
+		return true;
+	}
+
 	bool CacheServer::HandleFetchPathTable(BinaryReader& reader, BinaryWriter& writer)
 	{
 		Bucket& bucket = GetBucket(reader, TC("FetchPathTable"));
diff --git a/Engine/Source/Programs/UnrealBuildAccelerator/Common/Private/UbaStorageClient.cpp b/Engine/Source/Programs/UnrealBuildAccelerator/Common/Private/UbaStorageClient.cpp
--- a/Engine/Source/Programs/UnrealBuildAccelerator/Common/Private/UbaStorageClient.cpp
+++ b/Engine/Source/Programs/UnrealBuildAccelerator/Common/Private/UbaStorageClient.cpp
@@ -421,6 +421,7 @@ namespace uba
 							bool ctorSuccess = true;
 							proxy = new ProxyClient(ctorSuccess, ncci);
 							m_proxyClient = proxy;
+							allowProxy = true;
 
 							auto disallowProxy = MakeGuard([&]() { m_proxyClient->client.Disconnect(); proxy = nullptr; allowProxy = false; });
 
diff --git a/Engine/Source/Programs/UnrealBuildAccelerator/Common/Public/UbaCacheServer.h b/Engine/Source/Programs/UnrealBuildAccelerator/Common/Public/UbaCacheServer.h
--- a/Engine/Source/Programs/UnrealBuildAccelerator/Common/Public/UbaCacheServer.h
+++ b/Engine/Source/Programs/UnrealBuildAccelerator/Common/Public/UbaCacheServer.h
@@ -83,6 +83,7 @@ namespace uba
 
 		bool HandleMessage(const ConnectionInfo& connectionInfo, u8 messageType, BinaryReader& reader, BinaryWriter& writer);
 		bool HandleStoreEntry(ConnectionBucket& bucket, BinaryReader& reader, BinaryWriter& writer, u32 clientVersion, u32 clientId);
+		bool HandleStoreEntryDone(const ConnectionInfo& connectionInfo, BinaryReader& reader);
 		bool HandleFetchPathTable(BinaryReader& reader, BinaryWriter& writer);
 		bool HandleFetchCasTable(BinaryReader& reader, BinaryWriter& writer);
 		bool HandleFetchPathTable2(BinaryReader& reader, BinaryWriter& writer);
