[UBA]
* Added status line logging since there can be multiple hours between maintenance. Defaults to writing out once per 5 minutes (unless nothing changed)

diff --git a/Engine/Source/Programs/UnrealBuildAccelerator/CacheService/Private/UbaCacheService.cpp b/Engine/Source/Programs/UnrealBuildAccelerator/CacheService/Private/UbaCacheService.cpp
--- a/Engine/Source/Programs/UnrealBuildAccelerator/CacheService/Private/UbaCacheService.cpp
+++ b/Engine/Source/Programs/UnrealBuildAccelerator/CacheService/Private/UbaCacheService.cpp
@@ -27,6 +27,7 @@ namespace uba
 	const tchar*	Version = GetVersionString();
 	u32				DefaultCapacityGb = 500;
 	u32				DefaultExpiration = 3*24*60*60;
+	u32				DefaultReportIntervalSeconds = 5*60; // Five minutes
 	const tchar*	DefaultRootDir = []() {
 		static tchar buf[256];
 		if (IsWindows)
@@ -61,6 +62,7 @@ namespace uba
 		logger.Info(TC("  -crash                  Force a crash (for testing)"));
 		logger.Info(TC("  -nosignalhandler        Will not hook up signal handler"));
 		logger.Info(TC("  -maxworkers=<number>    Max number of workers used by cache server. Defaults to \"%u\""), DefaultWorkerCount);
+		logger.Info(TC("  -reportinterval=<sec>   How often the service should report status. Defaults to \"%s\""), TimeToText(MsToTime(DefaultReportIntervalSeconds*1000)).str);
 #if PLATFORM_LINUX
 		logger.Info(TC("  -fork                   Will handle segfaults and restart"));
 #endif
@@ -122,6 +124,7 @@ namespace uba
 		bool shouldCrash = false;
 		u32 expirationTimeSeconds = DefaultExpiration;
 		u32 maxWorkerCount = DefaultWorkerCount;
+		u32 reportIntervalSeconds = DefaultReportIntervalSeconds;
 		TString configFile;
 
 		#if PLATFORM_LINUX
@@ -203,6 +206,11 @@ namespace uba
 				if (!value.Parse(maxWorkerCount))
 					return PrintHelp(TC("Invalid value for -maxworkers"));
 			}
+			else if (name.Equals(TCV("-reportinterval")))
+			{
+				if (!value.Parse(reportIntervalSeconds))
+					return PrintHelp(TC("Invalid value for -reportinterval"));
+			}
 			else if (name.Equals(TCV("-nomaintenance")))
 			{
 				maintenanceEnabled = false;
@@ -499,17 +507,21 @@ namespace uba
 			if (!networkServer.StartListen(networkBackend, port, listenIp.data))
 				return false;
 
+			u64 lastUpdateTime = GetTime();
+
 			while (!ShouldExit() && !cacheServer.ShouldShutdown())
 			{
 				Sleep(1000);
 
 				bool force = false;
+				StringBuffer<> statusInfo;
 
 				#if PLATFORM_LINUX
 				struct statvfs stat;
 				if (statvfs(g_rootDir.data, &stat) != 0)
 				{
 					u64 available = stat.f_bsize * stat.f_bavail;
+					statusInfo.Append(TC(" FreeDisk: ")).Append(BytesToText(available).str);
 					if (available < 1024ull * 1024 * 1024)
 					{
 						logger.Warning(TC("Running low on disk space. Only %s available. Will force maintenance"), BytesToText(available).str);
@@ -518,6 +530,13 @@ namespace uba
 				}
 				#endif
 
+				u64 currentTime = GetTime();
+				if (TimeToMs(currentTime - lastUpdateTime) > reportIntervalSeconds*1000)
+				{
+					lastUpdateTime = currentTime;
+					cacheServer.PrintStatusLine(statusInfo.data);
+				}
+
 				if (maintenanceEnabled)
 					if (!cacheServer.RunMaintenance(force, allowSave, ShouldExit))
 						break;
diff --git a/Engine/Source/Programs/UnrealBuildAccelerator/Common/Private/UbaCacheServer.cpp b/Engine/Source/Programs/UnrealBuildAccelerator/Common/Private/UbaCacheServer.cpp
--- a/Engine/Source/Programs/UnrealBuildAccelerator/Common/Private/UbaCacheServer.cpp
+++ b/Engine/Source/Programs/UnrealBuildAccelerator/Common/Private/UbaCacheServer.cpp
@@ -358,8 +358,9 @@ namespace uba
 			UBA_ASSERT(insres.second);
 			auto& cacheEntries = insres.first->second;
 			cacheEntries.ReadFromDisk(logger, reader, databaseVersion, storage, m_casKeyTable);
-			outStats.totalCacheEntryCount += cacheEntries.entries.size();
+			totalEntryCount += cacheEntries.entries.size();
 		}
+		outStats.totalCacheEntryCount += totalEntryCount;
 		return true;
 	}
 
@@ -601,6 +602,37 @@ namespace uba
 		m_forceAllSteps = true;
 	}
 
+	void CacheServer::PrintStatusLine(const tchar* additionalInfo)
+	{
+		SCOPED_READ_LOCK(m_storage.m_casLookupLock, lookupLock);
+		u64 totalCasCount = m_storage.m_casLookup.size();
+		lookupLock.Leave();
+
+		SCOPED_FUTEX(m_storage.m_accessLock, accessLock);
+		u64 totalCasSize = m_storage.m_casTotalBytes;
+		accessLock.Leave();
+
+		SCOPED_FUTEX_READ(m_bucketsLock, bucketsLock);
+		u64 bucketCount = m_buckets.size();
+		u32 maxPathTable = 0;
+		u32 maxKeyTable = 0;
+		for (auto& kv : m_buckets)
+		{
+			maxPathTable = Max(maxPathTable, kv.second.m_pathTable.GetSize());
+			maxKeyTable = Max(maxKeyTable, kv.second.m_casKeyTable.GetSize());
+		}
+		bucketsLock.Leave();
+
+		StringBuffer<> text;
+		text.Appendf(TC("CasFiles: %s (%s) Buckets: %llu Entries: %s MaxPathTable: %s MaxKeyTable: %s%s"), CountToText(totalCasCount).str, BytesToText(totalCasSize).str, bucketCount, CountToText(m_totalEntryCount).str, BytesToText(maxPathTable).str, BytesToText(maxKeyTable).str, additionalInfo);
+
+		StringKey statusLineKey = ToStringKeyNoCheck(text.data, text.count);
+		if (statusLineKey == m_statusLineKey)
+			return;
+		m_statusLineKey = statusLineKey;
+		m_logger.Log(LogEntryType_Info, text);
+	}
+
 	bool CacheServer::RunMaintenance(bool force, bool allowSave, const Function<bool()>& shouldExit)
 	{
 		if (m_addsSinceMaintenance == 0 && !force && !m_bucketIsOverflowing)
@@ -1366,6 +1398,7 @@ namespace uba
 			saveDbEvent.IsSet();
 		}
 
+		m_totalEntryCount = totalEntryCount.load();
 		u64 oldestTime = oldest ? GetFileTimeAsTime(now - (m_creationTime + oldest)) : 0;
 		u64 longestUnused = oldestUsedTime ? GetFileTimeAsTime(now - (m_creationTime + oldestUsedTime)) : 0;
 		u64 duration = GetTime() - startTime;
@@ -1788,6 +1821,7 @@ namespace uba
 		if (hasAllContent)
 		{
 			cacheEntries.entries.emplace_front(std::move(newEntry));
+			++m_totalEntryCount;
 		}
 		else
 		{
@@ -1839,6 +1873,7 @@ namespace uba
 				return true;
 
 		cacheEntries.entries.emplace_front(std::move(entry));
+		++m_totalEntryCount;
 		return true;
 	}
 
diff --git a/Engine/Source/Programs/UnrealBuildAccelerator/Common/Public/UbaCacheBucket.h b/Engine/Source/Programs/UnrealBuildAccelerator/Common/Public/UbaCacheBucket.h
--- a/Engine/Source/Programs/UnrealBuildAccelerator/Common/Public/UbaCacheBucket.h
+++ b/Engine/Source/Programs/UnrealBuildAccelerator/Common/Public/UbaCacheBucket.h
@@ -21,8 +21,8 @@ namespace uba
 		CompactCasKeyTable m_casKeyTable;
 
 		u64 m_id;
-		u64 totalEntryCount;
-		u64 totalEntrySize;
+		u64 totalEntryCount = 0;
+		u64 totalEntrySize = 0;
 		Atomic<bool> hasDeletedEntries;
 		Atomic<bool> needsSave;
 
diff --git a/Engine/Source/Programs/UnrealBuildAccelerator/Common/Public/UbaCacheServer.h b/Engine/Source/Programs/UnrealBuildAccelerator/Common/Public/UbaCacheServer.h
--- a/Engine/Source/Programs/UnrealBuildAccelerator/Common/Public/UbaCacheServer.h
+++ b/Engine/Source/Programs/UnrealBuildAccelerator/Common/Public/UbaCacheServer.h
@@ -59,6 +59,7 @@ namespace uba
 
 		void SetForceFullMaintenance();
 
+		void PrintStatusLine(const tchar* additionalInfo);
 		bool RunMaintenance(bool force, bool allowSave, const Function<bool()>& shouldExit);
 
 		bool ShouldShutdown();
@@ -112,7 +113,8 @@ namespace uba
 
 		UnorderedSet<CasKey> m_trackedDeletes;
 
-		Atomic<bool> m_shutdownRequested = false;
+		Atomic<bool> m_shutdownRequested;
+		Atomic<u64> m_totalEntryCount;
 
 		u64 m_maintenanceReserveSize = 0;
 		u64 m_bucketCasTableMaxSize = 0;
@@ -129,6 +131,8 @@ namespace uba
 		bool m_shouldWipe = false;
 		bool m_forceAllSteps = false;
 
+		StringKey m_statusLineKey;
+
 		CacheServer(const CacheServer&) = delete;
 		CacheServer& operator=(const CacheServer&) = delete;
 	};
