[UBA]
* Added constexpr to a bunch of branches
* Downgraded some logging
* Removed detouring of CreatePseudoConsole

diff --git a/Engine/Source/Programs/UnrealBuildAccelerator/Agent/Private/UbaAgent.cpp b/Engine/Source/Programs/UnrealBuildAccelerator/Agent/Private/UbaAgent.cpp
--- a/Engine/Source/Programs/UnrealBuildAccelerator/Agent/Private/UbaAgent.cpp
+++ b/Engine/Source/Programs/UnrealBuildAccelerator/Agent/Private/UbaAgent.cpp
@@ -45,7 +45,7 @@ namespace uba
 	constexpr u32	DefaultListenTimeout = 5;
 	const tchar*	DefaultRootDir = [](){
 		static tchar buf[256];
-		if (IsWindows)
+		if constexpr (IsWindows)
 			ExpandEnvironmentStringsW(TC("%ProgramData%\\Epic\\" UE_APP_NAME), buf, sizeof(buf));
 		else
 			GetFullPathNameW(TC("~/" UE_APP_NAME), sizeof_array(buf), buf, nullptr);
@@ -536,7 +536,7 @@ namespace uba
 			}
 			else if (name.Equals(TCV("-nostore")))
 			{
-				if (IsWindows) // Only supported on windows atm.
+				if constexpr (IsWindows) // Only supported on windows atm.
 					useStorage = false;
 			}
 			else if (name.Equals(TCV("-nohandler")))
diff --git a/Engine/Source/Programs/UnrealBuildAccelerator/CacheService/Private/UbaCacheService.cpp b/Engine/Source/Programs/UnrealBuildAccelerator/CacheService/Private/UbaCacheService.cpp
--- a/Engine/Source/Programs/UnrealBuildAccelerator/CacheService/Private/UbaCacheService.cpp
+++ b/Engine/Source/Programs/UnrealBuildAccelerator/CacheService/Private/UbaCacheService.cpp
@@ -30,7 +30,7 @@ namespace uba
 	u32				DefaultReportIntervalSeconds = 5*60; // Five minutes
 	const tchar*	DefaultRootDir = []() {
 		static tchar buf[256];
-		if (IsWindows)
+		if constexpr (IsWindows)
 			ExpandEnvironmentStringsW(TC("%ProgramData%\\Epic\\" UE_APP_NAME), buf, sizeof(buf));
 		else
 			GetFullPathNameW(TC("~/" UE_APP_NAME), sizeof_array(buf), buf, nullptr);
diff --git a/Engine/Source/Programs/UnrealBuildAccelerator/Cli/Private/UbaCli.cpp b/Engine/Source/Programs/UnrealBuildAccelerator/Cli/Private/UbaCli.cpp
--- a/Engine/Source/Programs/UnrealBuildAccelerator/Cli/Private/UbaCli.cpp
+++ b/Engine/Source/Programs/UnrealBuildAccelerator/Cli/Private/UbaCli.cpp
@@ -36,7 +36,7 @@ namespace uba
 	u32				DefaultCapacityGb = 20;
 	const tchar*	DefaultRootDir = []() {
 		static tchar buf[256];
-		if (IsWindows)
+		if constexpr (IsWindows)
 			ExpandEnvironmentStringsW(TC("%ProgramData%\\Epic\\" UE_APP_NAME), buf, sizeof(buf));
 		else
 			GetFullPathNameW(TC("~/" UE_APP_NAME), sizeof_array(buf), buf, nullptr);
diff --git a/Engine/Source/Programs/UnrealBuildAccelerator/Common/Private/UbaFile.cpp b/Engine/Source/Programs/UnrealBuildAccelerator/Common/Private/UbaFile.cpp
--- a/Engine/Source/Programs/UnrealBuildAccelerator/Common/Private/UbaFile.cpp
+++ b/Engine/Source/Programs/UnrealBuildAccelerator/Common/Private/UbaFile.cpp
@@ -1010,7 +1010,7 @@ namespace uba
 		const tchar* platformStr = TC("Mac");
 		#endif
 		out.Resize(engineDirPos - out.data + 8).Append(TCV("Binaries")).Append(PathSeparator).Append(platformStr).Append(PathSeparator).Append("UnrealBuildAccelerator").Append(PathSeparator);
-		if (IsWindows)
+		if constexpr (IsWindows)
 			out.Append(isWindowsArm ? TC("arm64") : TC("x64")).Append(PathSeparator);
 		return true;
 	};
diff --git a/Engine/Source/Programs/UnrealBuildAccelerator/Common/Private/UbaSession.cpp b/Engine/Source/Programs/UnrealBuildAccelerator/Common/Private/UbaSession.cpp
--- a/Engine/Source/Programs/UnrealBuildAccelerator/Common/Private/UbaSession.cpp
+++ b/Engine/Source/Programs/UnrealBuildAccelerator/Common/Private/UbaSession.cpp
@@ -517,7 +517,7 @@ namespace uba
 		if (dirKey != m_directoryForcedEmpty)
 		{
 			StringBuffer<4> realPath;
-			if (IsWindows)
+			if constexpr (IsWindows)
 			{
 				if (dirPath.count == 2)
 					dirPath = realPath.Append(dirPath).Append(PathSeparator);
@@ -2710,7 +2710,7 @@ namespace uba
 				out.directoryTableSize = GetDirectoryTableSize();
 			});
 
-		if (!IsWindows)
+		if constexpr (!IsWindows)
 		{
 			out.fileName.Append(fileName);
 			return true;
diff --git a/Engine/Source/Programs/UnrealBuildAccelerator/Common/Private/UbaStorageServer.cpp b/Engine/Source/Programs/UnrealBuildAccelerator/Common/Private/UbaStorageServer.cpp
--- a/Engine/Source/Programs/UnrealBuildAccelerator/Common/Private/UbaStorageServer.cpp
+++ b/Engine/Source/Programs/UnrealBuildAccelerator/Common/Private/UbaStorageServer.cpp
@@ -815,7 +815,7 @@ namespace uba
 						{
 							writer.WriteU16(u16(~0));
 							writer.WriteBool(m_traceStore);
-							m_logger.Detail(TC("Client %u Store request of %s which already exists in casdb (was in deferred list) (%s)"), connectionInfo.GetId(), CasKeyString(casKey).str, hint.data);
+							m_logger.Debug(TC("Client %u Store request of %s which already exists in casdb (was in deferred list) (%s)"), connectionInfo.GetId(), CasKeyString(casKey).str, hint.data);
 							return true;
 						}
 						entryLock.Enter();
@@ -834,7 +834,7 @@ namespace uba
 							// Unless writeReceivedCasFilesToDisk is set this is extremely unlikely to happen. "PerModuleInline.gen.cpp" seems to be the only one (they are supposed to all be the same)
 							writer.WriteU16(u16(~0));
 							writer.WriteBool(m_traceStore);
-							m_logger.Detail(TC("Client %u Store request of %s which already exists in casdb (%s)"), connectionInfo.GetId(), CasKeyString(casKey).str, hint.data);
+							m_logger.Debug(TC("Client %u Store request of %s which already exists in casdb (%s)"), connectionInfo.GetId(), CasKeyString(casKey).str, hint.data);
 							return true;
 						}
 					}
@@ -852,7 +852,7 @@ namespace uba
 						CasEntryAccessed(casEntry);
 						writer.WriteU16(u16(~0));
 						writer.WriteBool(m_traceStore);
-						//m_logger.Detail(TC("Client %u Store request of %s which was just added to casdb (%s)"), connectionInfo.GetId(), CasKeyString(casKey).str, hint.data);
+						//m_logger.Debug(TC("Client %u Store request of %s which was just added to casdb (%s)"), connectionInfo.GetId(), CasKeyString(casKey).str, hint.data);
 						return true;
 					}
 				}
diff --git a/Engine/Source/Programs/UnrealBuildAccelerator/Core/Private/UbaStringBuffer.cpp b/Engine/Source/Programs/UnrealBuildAccelerator/Core/Private/UbaStringBuffer.cpp
--- a/Engine/Source/Programs/UnrealBuildAccelerator/Core/Private/UbaStringBuffer.cpp
+++ b/Engine/Source/Programs/UnrealBuildAccelerator/Core/Private/UbaStringBuffer.cpp
@@ -433,6 +433,19 @@ namespace uba
 		#endif
 	}
 
+	bool StringView::Contains(tchar c) const
+	{
+		for (const tchar* i = data, *e = data + count; i!=e; ++i)
+			if (*i == c)
+				return true;
+		return false;
+	}
+
+	bool StringView::StartsWith(const StringView& str, bool ignoreCase) const
+	{
+		return count >= str.count && uba::StartsWith(data, str.data, ignoreCase);
+	}
+
 	bool StringView::EndsWith(const tchar* value, bool ignoreCase) const
 	{
 		return uba::EndsWith(data, count, value, TStrlen(value), ignoreCase);
diff --git a/Engine/Source/Programs/UnrealBuildAccelerator/Core/Public/UbaStringBuffer.h b/Engine/Source/Programs/UnrealBuildAccelerator/Core/Public/UbaStringBuffer.h
--- a/Engine/Source/Programs/UnrealBuildAccelerator/Core/Public/UbaStringBuffer.h
+++ b/Engine/Source/Programs/UnrealBuildAccelerator/Core/Public/UbaStringBuffer.h
@@ -114,10 +114,14 @@ namespace uba
 		constexpr StringView(const StringBufferBase& sb) : data(sb.data), count(sb.count) {}
 		StringView(const TString& str) : data(str.data()), count(u32(str.size())) {}
 
+		tchar operator[](u64 i) const { return data[i]; }
+
 		bool IsEmpty() const { return count == 0; }
+		bool StartsWith(const StringView& str, bool ignoreCase = true) const;
 		bool StartsWith(const tchar* str, bool ignoreCase = true) const { return uba::StartsWith(data, str, ignoreCase); }
 		bool EndsWith(const StringView& value, bool ignoreCase = true) const { return uba::EndsWith(data, count, value.data, value.count, ignoreCase); }
 		bool EndsWith(const tchar* value, bool ignoreCase = true) const;
+		bool Contains(tchar c) const;
 		bool Contains(const tchar* str, bool ignoreCase = true) const { return uba::Contains(data, str, ignoreCase); }
 		bool Contains(const StringView& str, bool ignoreCase = true) const { return uba::Contains(data, str.data, ignoreCase); }
 		bool Equals(const tchar* str, bool ignoreCase = true) const { return uba::Equals(data, str, ignoreCase); }
diff --git a/Engine/Source/Programs/UnrealBuildAccelerator/Detours/Private/Windows/UbaDetoursFunctionsKernelBase.inl b/Engine/Source/Programs/UnrealBuildAccelerator/Detours/Private/Windows/UbaDetoursFunctionsKernelBase.inl
--- a/Engine/Source/Programs/UnrealBuildAccelerator/Detours/Private/Windows/UbaDetoursFunctionsKernelBase.inl
+++ b/Engine/Source/Programs/UnrealBuildAccelerator/Detours/Private/Windows/UbaDetoursFunctionsKernelBase.inl
@@ -4113,13 +4113,6 @@ HANDLE Detoured_CreateConsoleScreenBuffer(DWORD dwDesiredAccess, DWORD dwShareMo
 	return True_CreateConsoleScreenBuffer(dwDesiredAccess, dwShareMode, lpSecurityAttributes, dwFlags, lpScreenBufferData);
 }
 
-HRESULT WINAPI Detoured_CreatePseudoConsole(COORD size, HANDLE hInput, HANDLE hOutput, DWORD dwFlags, HPCON* phPC)
-{
-	DETOURED_CALL(CreatePseudoConsole);
-	DEBUG_LOG_TRUE(L"CreatePseudoConsole", L"");
-	return True_CreatePseudoConsole(size, hInput, hOutput, dwFlags, phPC);
-}
-
 BOOL Detoured_CreateProcessAsUserW(HANDLE hToken, LPCWSTR lpApplicationName, LPWSTR lpCommandLine, LPSECURITY_ATTRIBUTES lpProcessAttributes, LPSECURITY_ATTRIBUTES lpThreadAttributes, BOOL bInheritHandles, DWORD dwCreationFlags, LPVOID lpEnvironment, LPCWSTR lpCurrentDirectory, LPSTARTUPINFOW lpStartupInfo, LPPROCESS_INFORMATION lpProcessInformation)
 {
 	DETOURED_CALL(CreateProcessAsUserW);
diff --git a/Engine/Source/Programs/UnrealBuildAccelerator/Detours/Private/Windows/UbaDetoursFunctionsWin.h b/Engine/Source/Programs/UnrealBuildAccelerator/Detours/Private/Windows/UbaDetoursFunctionsWin.h
--- a/Engine/Source/Programs/UnrealBuildAccelerator/Detours/Private/Windows/UbaDetoursFunctionsWin.h
+++ b/Engine/Source/Programs/UnrealBuildAccelerator/Detours/Private/Windows/UbaDetoursFunctionsWin.h
@@ -250,7 +250,6 @@
 	DETOURED_FUNCTION(SetConsoleTextAttribute) \
 	DETOURED_FUNCTION(SetConsoleTitleW) \
 	DETOURED_FUNCTION(CreateConsoleScreenBuffer) \
-	DETOURED_FUNCTION(CreatePseudoConsole) \
 	DETOURED_FUNCTION(CreateProcessAsUserW) \
 	DETOURED_FUNCTION(SetConsoleCtrlHandler) \
 	DETOURED_FUNCTION(GetConsoleOutputCP) \
diff --git a/Engine/Source/Programs/UnrealBuildAccelerator/Test/Private/UbaTestStdOut.cpp b/Engine/Source/Programs/UnrealBuildAccelerator/Test/Private/UbaTestStdOut.cpp
--- a/Engine/Source/Programs/UnrealBuildAccelerator/Test/Private/UbaTestStdOut.cpp
+++ b/Engine/Source/Programs/UnrealBuildAccelerator/Test/Private/UbaTestStdOut.cpp
@@ -87,7 +87,7 @@ namespace uba
 
 	bool TestStdOutLocal(LoggerWithWriter& logger, const StringBufferBase& testRootDir)
 	{
-		if (!IsWindows)
+		if constexpr (!IsWindows)
 			return true;
 		return TestStdOut(logger, testRootDir, false);
 	}
@@ -99,7 +99,7 @@ namespace uba
 
 	bool TestStdOutViaCmd(LoggerWithWriter& logger, const StringBufferBase& testRootDir)
 	{
-		if (!IsWindows)
+		if constexpr (!IsWindows)
 			return true;
 		StringBuffer<> args;
 		args.Append(TCV("/c \""));
