[UBA]
* Fixed so Deploy-Windows.bat is not using vfs

diff --git a/Engine/Source/Programs/UnrealBuildAccelerator/Deploy-Windows.bat b/Engine/Source/Programs/UnrealBuildAccelerator/Deploy-Windows.bat
--- a/Engine/Source/Programs/UnrealBuildAccelerator/Deploy-Windows.bat
+++ b/Engine/Source/Programs/UnrealBuildAccelerator/Deploy-Windows.bat
@@ -14,7 +14,7 @@ set config=%1
 echo.
 echo === Building Targets ===
 pushd "../../../.."
-call Engine/Build/BatchFiles/RunUBT.bat -NoUba -NoUbaLocal -NoSNDBS -NoXGE ^
+call Engine/Build/BatchFiles/RunUBT.bat -NoUba -NoUbaLocal -NoSNDBS -NoXGE -novfs ^
 	-Target="UbaAgent Win64 %config%" ^
 	-Target="UbaCli Win64 %config%" ^
 	-Target="UbaDetours Win64 %config%" ^
