[UBA]
* Removed asserts in non-debug builds.. accidental submit

diff --git a/Engine/Source/Programs/UnrealBuildAccelerator/Core/Public/UbaPlatform.h b/Engine/Source/Programs/UnrealBuildAccelerator/Core/Public/UbaPlatform.h
--- a/Engine/Source/Programs/UnrealBuildAccelerator/Core/Public/UbaPlatform.h
+++ b/Engine/Source/Programs/UnrealBuildAccelerator/Core/Public/UbaPlatform.h
@@ -68,7 +68,7 @@ namespace uba
 	#define	sizeof_array(array) int(sizeof(array)/sizeof(array[0]))
 	#endif
 
-	#if 1//UBA_DEBUG
+	#if UBA_DEBUG
 	#define UBA_ASSERT(x) do { if (x) break; uba::UbaAssert(TC(""), __FILE__, __LINE__, #x, true, 543221, nullptr, 0); } while(false)
 	#define UBA_ASSERTF(x, fmt, ...) do { UBA_FMT_CHECK(fmt, ##__VA_ARGS__); if (x) break; uba::StringBuffer<> _buf; _buf.Appendf(fmt, ##__VA_ARGS__); uba::UbaAssert(_buf.data, __FILE__, __LINE__, #x, true, 543221, nullptr, 0); } while(false)
 	#define UBA_ASSERT_MESSAGEBOX 0
