[UBA]
* Fixed so cache server bumps access on cache entries output cas entries when reported used by a client. This is important so cas store is not removing entries that are recently "used"
* Added overflow deletes to status line
* Downgrade log entry to debug for cache client failing to download cas files (since this is normal if cache server overflows)
* Added so UbaCacheService capacity can be written with decimals.. like "1.4" for 1.4 gb capacity

diff --git a/Engine/Source/Programs/UnrealBuildAccelerator/CacheService/Private/UbaCacheService.cpp b/Engine/Source/Programs/UnrealBuildAccelerator/CacheService/Private/UbaCacheService.cpp
--- a/Engine/Source/Programs/UnrealBuildAccelerator/CacheService/Private/UbaCacheService.cpp
+++ b/Engine/Source/Programs/UnrealBuildAccelerator/CacheService/Private/UbaCacheService.cpp
@@ -110,7 +110,7 @@ namespace uba
 	{
 		using namespace uba;
 
-		u32 storageCapacityGb = DefaultCapacityGb;
+		float storageCapacityGb = float(DefaultCapacityGb);
 		StringBuffer<256> workDir;
 		StringBuffer<128> listenIp;
 		u16 port = DefaultCachePort;
@@ -251,12 +251,13 @@ namespace uba
 		if (!configFile.empty())
 			config.LoadFromFile(logger, configFile.data());
 
+		u64 storageCapacity = u64(storageCapacityGb*1000.0f)*1000*1000;
 
 		const tchar* dbgStr = TC("");
 		#if UBA_DEBUG
 		dbgStr = TC(" (DEBUG)");
 		#endif
-		logger.Info(TC("UbaCacheService v%s(%u)%s (Workers: %u, Rootdir: \"%s\", StoreCapacity: %uGb, Expiration: %s)"), Version, CacheNetworkVersion, dbgStr, maxWorkerCount, g_rootDir.data, storageCapacityGb, TimeToText(MsToTime(expirationTimeSeconds)*1000, true).str);
+		logger.Info(TC("UbaCacheService v%s(%u)%s (Workers: %u, Rootdir: \"%s\", StoreCapacity: %s, Expiration: %s)"), Version, CacheNetworkVersion, dbgStr, maxWorkerCount, g_rootDir.data, BytesToText(storageCapacity).str, TimeToText(MsToTime(expirationTimeSeconds)*1000, true).str);
 
 		u64 maintenanceReserveSizeMb = 128;
 
@@ -268,8 +269,6 @@ namespace uba
 				logger.Info(TC("  Improve maintenance performance by enabling %llu huge pages on system (%llu enabled)"), recommendedHugePageCount, hugePageCount);
 		}
 
-		u64 storageCapacity = u64(storageCapacityGb)*1000*1000*1000;
-
 		StringBuffer<512> currentDir;
 		GetCurrentDirectoryW(currentDir);
 
diff --git a/Engine/Source/Programs/UnrealBuildAccelerator/Common/Private/UbaCacheClient.cpp b/Engine/Source/Programs/UnrealBuildAccelerator/Common/Private/UbaCacheClient.cpp
--- a/Engine/Source/Programs/UnrealBuildAccelerator/Common/Private/UbaCacheClient.cpp
+++ b/Engine/Source/Programs/UnrealBuildAccelerator/Common/Private/UbaCacheClient.cpp
@@ -848,7 +848,7 @@ namespace uba
 			MemoryBlock normalizedBlock(4*1024*1024);
 			bool destinationIsCompressed = false;
 			if (!fetcher.RetrieveFile(logger, m_client, casKey, path.data, destinationIsCompressed, &normalizedBlock))
-				return logger.Error(TC("Failed to download cache output for %s"), info.GetDescription());
+				return logger.Debug(TC("Failed to download cache output for %s"), info.GetDescription()).ToFalse();
 
 			MemoryBlock localBlock(4*1024*1024);
 
@@ -918,7 +918,7 @@ namespace uba
 			DowngradedLogger logger(m_connected, m_logger.m_writer, TC("UbaCacheClientDownload"));
 			bool writeCompressed = m_session.ShouldStoreIntermediateFilesCompressed() && g_globalRules.FileCanBeCompressed(path);
 			if (!fetcher.RetrieveFile(logger, m_client, casKey, path.data, writeCompressed))
-				return logger.Error(TC("Failed to download cache output %s for %s"), path.data, info.GetDescription());
+				return logger.Debug(TC("Failed to download cache output %s for %s"), path.data, info.GetDescription()).ToFalse();
 		}
 
 		cacheStats.fetchBytesRaw += fetcher.sizeOnDisk;
diff --git a/Engine/Source/Programs/UnrealBuildAccelerator/Common/Private/UbaCacheServer.cpp b/Engine/Source/Programs/UnrealBuildAccelerator/Common/Private/UbaCacheServer.cpp
--- a/Engine/Source/Programs/UnrealBuildAccelerator/Common/Private/UbaCacheServer.cpp
+++ b/Engine/Source/Programs/UnrealBuildAccelerator/Common/Private/UbaCacheServer.cpp
@@ -610,6 +610,7 @@ namespace uba
 
 		SCOPED_FUTEX(m_storage.m_accessLock, accessLock);
 		u64 totalCasSize = m_storage.m_casTotalBytes;
+		u64 totalDeletedCount = m_storage.m_trackedDeletes->size();
 		accessLock.Leave();
 
 		SCOPED_FUTEX_READ(m_bucketsLock, bucketsLock);
@@ -625,6 +626,8 @@ namespace uba
 
 		StringBuffer<> text;
 		text.Appendf(TC("CasFiles: %s (%s) Buckets: %llu Entries: %s MaxPathTable: %s MaxKeyTable: %s%s"), CountToText(totalCasCount).str, BytesToText(totalCasSize).str, bucketCount, CountToText(m_totalEntryCount).str, BytesToText(maxPathTable).str, BytesToText(maxKeyTable).str, additionalInfo);
+		if (totalDeletedCount)
+			text.Appendf(TC(" CasOverflowDeletes: %s"), CountToText(totalDeletedCount).str);
 
 		StringKey statusLineKey = ToStringKeyNoCheck(text.data, text.count);
 		if (statusLineKey == m_statusLineKey)
@@ -1402,7 +1405,7 @@ namespace uba
 		u64 oldestTime = oldest ? GetFileTimeAsTime(now - (m_creationTime + oldest)) : 0;
 		u64 longestUnused = oldestUsedTime ? GetFileTimeAsTime(now - (m_creationTime + oldestUsedTime)) : 0;
 		u64 duration = GetTime() - startTime;
-		m_logger.Info(TC("Maintenance done! (%s) CasFiles: %s (%s) Buckets: %llu Entries: %s Oldest: %s LongestUnused: %s MaxMaintenanceMem: %s/%s"), TimeToText(duration).str, CountToText(totalCasCount - deletedCasCount).str, BytesToText(totalCasSize).str, m_buckets.size(), CountToText(totalEntryCount.load()).str, TimeToText(oldestTime, true).str, TimeToText(longestUnused, true).str, BytesToText(maxCommittedMemory).str, BytesToText(m_maintenanceReserveSize).str);
+		m_logger.Info(TC("Maintenance done! (%s) CasFiles: %s (%s) Buckets: %llu Entries: %s Oldest: %s LongestUnused: %s MaxMaintenanceMem: %s/%s"), TimeToText(duration).str, CountToText(existingCas.Size()).str, BytesToText(totalCasSize).str, m_buckets.size(), CountToText(totalEntryCount.load()).str, TimeToText(oldestTime, true).str, TimeToText(longestUnused, true).str, BytesToText(maxCommittedMemory).str, BytesToText(m_maintenanceReserveSize).str);
 		
 		m_longestMaintenance = Max(m_longestMaintenance, duration);
 
@@ -1972,6 +1975,7 @@ namespace uba
 		auto& cacheEntries = findIt->second;
 		lock.Leave();
 
+		Vector<CasKey> casKeysUsed;
 		SCOPED_WRITE_LOCK(cacheEntries.lock, lock2);
 		for (auto& entry : cacheEntries.entries)
 		{
@@ -1981,11 +1985,27 @@ namespace uba
 			entry.lastUsedTime = fileTime;
 			bucket.lastUsedTime = fileTime;
 
+			auto& outputs = entry.outputCasKeyOffsets;
+			BinaryReader outputsReader(outputs);
+			while (outputsReader.GetLeft())
+			{
+				u64 offset = outputsReader.Read7BitEncoded();
+				CasKey casKey;
+				bucket.m_casKeyTable.GetKey(casKey, offset);
+				casKeysUsed.push_back(casKey);
+			}
+
 			if (clientVersion >= 5 && entry.logLinesType == LogLinesType_Owned)
 				if (entry.logLines.size() <= writer.GetCapacityLeft())
 					writer.WriteBytes(entry.logLines.data(), entry.logLines.size());
 			break;
 		}
+		lock2.Leave();
+
+		// Set caskeys to accessed even though they might not have been downloaded
+		for (CasKey& casKey : casKeysUsed)
+			m_storage.CasEntryAccessed(casKey);
+
 		return true;
 	}
 
diff --git a/Engine/Source/Programs/UnrealBuildAccelerator/Common/Private/UbaNetworkServer.cpp b/Engine/Source/Programs/UnrealBuildAccelerator/Common/Private/UbaNetworkServer.cpp
--- a/Engine/Source/Programs/UnrealBuildAccelerator/Common/Private/UbaNetworkServer.cpp
+++ b/Engine/Source/Programs/UnrealBuildAccelerator/Common/Private/UbaNetworkServer.cpp
@@ -149,8 +149,8 @@ namespace uba
 			int activeWorkerCount = --m_activeWorkerCount;
 			if (!activeWorkerCount) // Will disconnect in send if there are active workers
 				TestDisconnect();
-			else
-				m_server.m_logger.Detail(TC("Connection %u disconnected while it has %u active workers (%s)"), m_id, activeWorkerCount, reason);
+			//else
+			//	m_server.m_logger.Detail(TC("Connection %u disconnected while it has %u active workers (%s)"), m_id, activeWorkerCount, reason);
 		}
 
 		bool Stop()
diff --git a/Engine/Source/Programs/UnrealBuildAccelerator/Common/Private/UbaStorage.cpp b/Engine/Source/Programs/UnrealBuildAccelerator/Common/Private/UbaStorage.cpp
--- a/Engine/Source/Programs/UnrealBuildAccelerator/Common/Private/UbaStorage.cpp
+++ b/Engine/Source/Programs/UnrealBuildAccelerator/Common/Private/UbaStorage.cpp
@@ -1228,6 +1228,7 @@ namespace uba
 			DetachEntry(*casEntry);
 			m_casLookup.erase(casEntry->key);
 		}
+		m_overflowReported = false;
 		u64 after = m_casTotalBytes;
 		if (before != after)
 			m_logger.Info(TC("Evicted %s from storage (%s). Estimated new storage is now %s (there might be files db is not aware of)"), BytesToText(before - after).str, TimeToText(GetTime() - startTime).str, BytesToText(after).str);
