[UBA]
* Added UbaAgent -killtcphogs which can be used for the agent to kill processes holding onto listen port when failing to bind

diff --git a/Engine/Source/Programs/UnrealBuildAccelerator/Agent/Private/UbaAgent.cpp b/Engine/Source/Programs/UnrealBuildAccelerator/Agent/Private/UbaAgent.cpp
--- a/Engine/Source/Programs/UnrealBuildAccelerator/Agent/Private/UbaAgent.cpp
+++ b/Engine/Source/Programs/UnrealBuildAccelerator/Agent/Private/UbaAgent.cpp
@@ -116,6 +116,7 @@ namespace uba
 		logger.Info(TC("  -description            Add more info about the agent that will show in the trace log when hovering over session"));
 		logger.Info(TC("  -usecrawler             Enables include crawler for known process types (clang/msvc)"));
 		#if PLATFORM_MAC
+		logger.Info(TC("  -killtcphogs            If failing to bind listen socket UbaAgent will attempt to kill processes holding it and then retry"));
 		logger.Info(TC("  -populateCasFromXcodeVersion=<version>   Prepopulate cas database with files from local xcode installation that matches the version."));
 		logger.Info(TC("  -populateCasFromAllXcodes   Prepopulate cas database with files from local xcode installation that matches the version."));
 		#endif
@@ -387,6 +388,7 @@ namespace uba
 		#if PLATFORM_MAC
 		StringBuffer<32> populateCasFromXcodeVersion;
 		bool populateCasFromAllXcodes = false;
+		bool killTcpHogs = false;
 		#endif
 
 		for (int i=1; i!=argc; ++i)
@@ -638,6 +640,10 @@ namespace uba
 			{
 				populateCasFromAllXcodes = true;
 			}
+			else if (name.Equals(TCV("-killtcphogs")))
+			{
+				killTcpHogs = true;
+			}
 			#endif
 			else if (name.Equals(TCV("-sentry")))
 			{
@@ -1025,8 +1031,39 @@ namespace uba
 
 			if (useListen)
 			{
-				if (!client->StartListen(*networkBackend, port))
-					return logger.Error(TC("Failed to get start listening on port %u"), port);
+				while (true)
+				{
+					if (client->StartListen(*networkBackend, port))
+						break;
+#if PLATFORM_MAC
+					if (killTcpHogs)
+					{
+						killTcpHogs = false;
+
+						StringBuffer<> lsofCommand;
+						lsofCommand.Appendf("lsof -i :%u -sTCP:LISTEN -Pn -t", u32(port));
+						
+						FILE* lsof = popen(lsofCommand.data, "r");
+						if (!lsof)
+							return logger.Error(TC("Failed run lsof while trying to kill processes holding port %u"), u32(port));
+
+						char pidStr[16];
+						while (fgets(pidStr, sizeof(pidStr), lsof))
+						{
+							pid_t pid = (pid_t)atoi(pidStr);
+							if (pid <= 0)
+								continue;
+							if (kill(pid, SIGKILL) != 0)
+								return logger.Error("Failed to kill process %d", pid);
+							logger.Info("Process %d killed successfully", pid);
+						}
+						pclose(lsof);
+						Sleep(2000);
+						continue;
+					}
+	#endif
+					return logger.Error(TC("Failed to get start listening on port %u"), u32(port));
+				}
 
 				u64 startTime = GetTime();
 				while (!client->IsOrWasConnected(200))
