[UBA]
* Added code to handle ENOTDIR in GetFileAttributesW

diff --git a/Engine/Source/Programs/UnrealBuildAccelerator/Common/Private/UbaFile.cpp b/Engine/Source/Programs/UnrealBuildAccelerator/Common/Private/UbaFile.cpp
--- a/Engine/Source/Programs/UnrealBuildAccelerator/Common/Private/UbaFile.cpp
+++ b/Engine/Source/Programs/UnrealBuildAccelerator/Common/Private/UbaFile.cpp
@@ -784,6 +784,11 @@ namespace uba
 				SetLastError(ERROR_FILE_NOT_FOUND);
 				return INVALID_FILE_ATTRIBUTES;
 			}
+			if (errno == ENOTDIR)
+			{
+				SetLastError(ERROR_DIRECTORY);
+				return INVALID_FILE_ATTRIBUTES;
+			}
 			UBA_ASSERTF(false, TC("GetFileAttributesW error handling not implemented %s (%s)"), fileName, strerror(errno));
 			return INVALID_FILE_ATTRIBUTES;
 		}
diff --git a/Engine/Source/Programs/UnrealBuildAccelerator/Common/Public/UbaFile.h b/Engine/Source/Programs/UnrealBuildAccelerator/Common/Public/UbaFile.h
--- a/Engine/Source/Programs/UnrealBuildAccelerator/Common/Public/UbaFile.h
+++ b/Engine/Source/Programs/UnrealBuildAccelerator/Common/Public/UbaFile.h
@@ -13,6 +13,7 @@ namespace uba
 	inline constexpr u32 ERROR_PATH_NOT_FOUND = ENOENT;
 	inline constexpr u32 ERROR_ALREADY_EXISTS = EEXIST;
 	inline constexpr u32 ERROR_ACCESS_DENIED = EACCES;
+	inline constexpr u32 ERROR_DIRECTORY = ENOTDIR;
 	inline constexpr u32 MOVEFILE_REPLACE_EXISTING = 0x00000001;
 	inline constexpr u32 FILE_FLAG_NO_BUFFERING = 0;
 	inline constexpr u32 FILE_FLAG_OVERLAPPED = 0;
