[UBA]
* Fixed so ubacacheservice forwards SIGTERM to child process when forked

diff --git a/Engine/Source/Programs/UnrealBuildAccelerator/CacheService/Private/UbaCacheService.cpp b/Engine/Source/Programs/UnrealBuildAccelerator/CacheService/Private/UbaCacheService.cpp
--- a/Engine/Source/Programs/UnrealBuildAccelerator/CacheService/Private/UbaCacheService.cpp
+++ b/Engine/Source/Programs/UnrealBuildAccelerator/CacheService/Private/UbaCacheService.cpp
@@ -331,7 +331,8 @@ namespace uba
 				TraverseDir(logger, crashDumpDir, [&](const DirectoryEntry& e) { existingFiles.insert(ToStringKey(e.name, e.nameLen)); });
 
 			static bool shouldExit = false;
-			auto SigHandler = [](int sig) { shouldExit = true; };
+			static pid_t actualChild = 0;
+			auto SigHandler = [](int sig) { shouldExit = true; if (actualChild && sig == SIGTERM) kill(actualChild, sig); };
 
 			logger.Info(TC("  Starting cache server fork"));
 			forkProcess = false;
@@ -346,13 +347,14 @@ namespace uba
 
 				// Parent
 
+				actualChild = childPid;
 				signal(SIGINT, SigHandler);
 				signal(SIGTERM, SigHandler);
 
 				int status;
 				if (waitpid(childPid, &status, 0) == -1)
 					return logger.Error(TC("waitpid failed. (%s)"), strerror(errno));
-				
+				actualChild = 0;
 				if (!shouldExit)
 					logger.Info(TC("  Fork exited with . (%u)"), status);
 
