[UBA]
* Fixed bug in cache service linux code tracking how much space is left on disk

diff --git a/Engine/Source/Programs/UnrealBuildAccelerator/CacheService/Private/UbaCacheService.cpp b/Engine/Source/Programs/UnrealBuildAccelerator/CacheService/Private/UbaCacheService.cpp
--- a/Engine/Source/Programs/UnrealBuildAccelerator/CacheService/Private/UbaCacheService.cpp
+++ b/Engine/Source/Programs/UnrealBuildAccelerator/CacheService/Private/UbaCacheService.cpp
@@ -519,7 +519,7 @@ namespace uba
 
 				#if PLATFORM_LINUX
 				struct statvfs stat;
-				if (statvfs(g_rootDir.data, &stat) != 0)
+				if (statvfs(g_rootDir.data, &stat) == 0)
 				{
 					u64 available = stat.f_bsize * stat.f_bavail;
 					statusInfo.Append(TC(" FreeDisk: ")).Append(BytesToText(available).str);
