[UBA]
* Changed proxy client lock from RW lock to futex
* Added log entry if switching proxy takes more than 10 seconds

diff --git a/Engine/Source/Programs/UnrealBuildAccelerator/Common/Private/UbaStorageClient.cpp b/Engine/Source/Programs/UnrealBuildAccelerator/Common/Private/UbaStorageClient.cpp
--- a/Engine/Source/Programs/UnrealBuildAccelerator/Common/Private/UbaStorageClient.cpp
+++ b/Engine/Source/Programs/UnrealBuildAccelerator/Common/Private/UbaStorageClient.cpp
@@ -110,13 +110,13 @@ namespace uba
 
 	bool StorageClient::IsUsingProxy()
 	{
-		SCOPED_READ_LOCK(m_proxyClientLock, proxyLock);
+		SCOPED_FUTEX_READ(m_proxyClientLock, proxyLock);
 		return m_proxyClient != nullptr;
 	}
 
 	void StorageClient::StopProxy()
 	{
-		SCOPED_WRITE_LOCK(m_proxyClientLock, proxyLock);
+		SCOPED_FUTEX(m_proxyClientLock, proxyLock);
 		if (m_proxyClient)
 			m_proxyClient->client.Disconnect();
 	}
@@ -300,7 +300,7 @@ namespace uba
 			bool wantsProxy = false;
 			if (allowProxy && m_allowProxy)
 			{
-				SCOPED_WRITE_LOCK(m_proxyClientLock, proxyLock);
+				SCOPED_FUTEX(m_proxyClientLock, proxyLock);
 				while (true)
 				{
 					if (!m_proxyClient)
@@ -333,7 +333,7 @@ namespace uba
 				{
 					if (proxy)
 					{
-						SCOPED_WRITE_LOCK(m_proxyClientLock, proxyLock);
+						SCOPED_FUTEX(m_proxyClientLock, proxyLock);
 						--proxy->refCount;
 					}
 				});
@@ -410,10 +410,19 @@ namespace uba
 
 					u32 proxyClientId = reader.ReadU32();
 
-					SCOPED_WRITE_LOCK(m_proxyClientLock, proxyLock2);
+					SCOPED_FUTEX(m_proxyClientLock, proxyLock2);
 					if (m_proxyClient)
 						continue;
 
+					u64 startTime = GetTime();
+					auto timeGuard = MakeGuard([&]
+						{
+							u64 deltaTime = GetTime() - startTime;
+							if (deltaTime > MsToTime(10*1000))
+								m_logger.Info(TC("Took %s to change proxy"), TimeToText(deltaTime).str);
+						});
+
+
 					auto createProxyClient = [&]()
 						{
 							NetworkClientCreateInfo ncci(m_logger.m_writer);
@@ -998,7 +1007,7 @@ namespace uba
 
 	void StorageClient::Ping()
 	{
-		SCOPED_WRITE_LOCK(m_proxyClientLock, lock);
+		SCOPED_FUTEX(m_proxyClientLock, lock);
 		auto proxy = m_proxyClient;
 		if (!proxy || !proxy->client.IsConnected())
 			return;
diff --git a/Engine/Source/Programs/UnrealBuildAccelerator/Common/Public/UbaStorageClient.h b/Engine/Source/Programs/UnrealBuildAccelerator/Common/Public/UbaStorageClient.h
--- a/Engine/Source/Programs/UnrealBuildAccelerator/Common/Public/UbaStorageClient.h
+++ b/Engine/Source/Programs/UnrealBuildAccelerator/Common/Public/UbaStorageClient.h
@@ -85,7 +85,7 @@ namespace uba
 		static constexpr u8 ServiceId = StorageServiceId;
 
 		struct ProxyClient;
-		ReaderWriterLock m_proxyClientLock;
+		Futex m_proxyClientLock;
 		ProxyClient* m_proxyClient = nullptr;
 		u64 m_proxyClientKeepAliveTime = 0;
 
