[UBA]
* Fixed CheckFileTable so it handles compressed files
* Skip writing invalid file table entries to disk

diff --git a/Engine/Source/Programs/UnrealBuildAccelerator/Common/Private/UbaStorage.cpp b/Engine/Source/Programs/UnrealBuildAccelerator/Common/Private/UbaStorage.cpp
--- a/Engine/Source/Programs/UnrealBuildAccelerator/Common/Private/UbaStorage.cpp
+++ b/Engine/Source/Programs/UnrealBuildAccelerator/Common/Private/UbaStorage.cpp
@@ -1572,7 +1572,7 @@ namespace uba
 			{
 				FileEntry& fileEntry = pair.second;
 				//SCOPED_READ_LOCK(fileEntry.lock, entryLock); // TODO: Remove this? tsan complains about it
-				if (fileEntry.casKey == CasKeyZero || fileEntry.lastWritten == 0 || fileEntry.isTemp)
+				if (fileEntry.casKey == CasKeyZero || fileEntry.lastWritten == 0 || fileEntry.isTemp || fileEntry.casKey == CasKeyInvalid)
 					continue;
 				writer.WriteStringKey(pair.first);
 				writer.WriteU64(fileEntry.size);
@@ -1778,8 +1778,17 @@ namespace uba
 							++testedFiles;
 
 							CasKey casKey;
-							if (!CalculateCasKey(casKey, p.data()))
-								m_logger.Warning(TC("Failed to calculate cas key for %s"), p.data());
+
+							{
+								FileAccessor file(m_logger, p.data());
+								if (!file.OpenMemoryRead())
+									return;
+								if (file.GetSize() > sizeof(CompressedFileHeader) && ((CompressedFileHeader*)file.GetData())->IsValid())
+									casKey = ((CompressedFileHeader*)file.GetData())->casKey;
+							}
+							if (casKey == CasKeyZero)
+								if (!CalculateCasKey(casKey, p.data()))
+									m_logger.Warning(TC("Failed to calculate cas key for %s"), p.data());
 
 							if (casKey != fe.casKey)
 							{
