[UBA]
* Added LOG_STALL_SCOPE that can be enabled with UBA_LOG_STALLS.
* Placed LOG_STALL_SCOPE to various places in the code base to try to catch why there are stalls on some helpers some times

diff --git a/Engine/Source/Programs/UnrealBuildAccelerator/Common/Private/UbaLogger.cpp b/Engine/Source/Programs/UnrealBuildAccelerator/Common/Private/UbaLogger.cpp
--- a/Engine/Source/Programs/UnrealBuildAccelerator/Common/Private/UbaLogger.cpp
+++ b/Engine/Source/Programs/UnrealBuildAccelerator/Common/Private/UbaLogger.cpp
@@ -802,4 +802,14 @@ namespace uba
 				logger.Info(error.data);
 			});
 	}
+
+	#if UBA_LOG_STALLS
+	LogStallScope::LogStallScope(Logger& l, u64 t, const tchar* m) : logger(l), timeSeconds(t), timeStart(GetTime()), messageFormat(m) {}
+	LogStallScope::~LogStallScope()
+	{
+		u64 delta = GetTime() - timeStart;
+		if (delta > MsToTime(timeSeconds*1000))
+			logger.Error(messageFormat, TimeToText(delta).str);
+	}
+	#endif
 }
diff --git a/Engine/Source/Programs/UnrealBuildAccelerator/Common/Private/UbaNetworkClient.cpp b/Engine/Source/Programs/UnrealBuildAccelerator/Common/Private/UbaNetworkClient.cpp
--- a/Engine/Source/Programs/UnrealBuildAccelerator/Common/Private/UbaNetworkClient.cpp
+++ b/Engine/Source/Programs/UnrealBuildAccelerator/Common/Private/UbaNetworkClient.cpp
@@ -248,13 +248,16 @@ namespace uba
 			return connection.backend->Send(client.m_logger, connection.backendConnection, data, sizeof(data), sendContext, TC("KeepAliveNoResponse"));
 		}
 
-		SCOPED_READ_LOCK(client.m_activeMessagesLock, lock);
-		if (!connection.connected)
-			return false;
-		if (messageId >= client.m_activeMessages.size())
-			return client.m_logger.Error(TC("Message id %u is higher than max %u"), messageId, u32(client.m_activeMessages.size()));
-		NetworkMessage* msg = client.m_activeMessages[messageId];
-		lock.Leave();
+		NetworkMessage* msg;
+		{
+			LOG_STALL_SCOPE(client.m_logger, 5, TC("Took more than %s to get message from id"));
+			SCOPED_READ_LOCK(client.m_activeMessagesLock, lock);
+			if (!connection.connected)
+				return false;
+			if (messageId >= client.m_activeMessages.size())
+				return client.m_logger.Error(TC("Message id %u is higher than max %u"), messageId, u32(client.m_activeMessages.size()));
+			msg = client.m_activeMessages[messageId];
+		}
 
 		if (!msg)
 			return false;
@@ -544,6 +547,8 @@ namespace uba
 				}
 			}
 
+			LOG_STALL_SCOPE(m_logger, 5, TC("Took more than %s to get message id"));
+
 			while (true)
 			{
 				SCOPED_WRITE_LOCK(m_activeMessagesLock, lock);
diff --git a/Engine/Source/Programs/UnrealBuildAccelerator/Common/Private/UbaNetworkServer.cpp b/Engine/Source/Programs/UnrealBuildAccelerator/Common/Private/UbaNetworkServer.cpp
--- a/Engine/Source/Programs/UnrealBuildAccelerator/Common/Private/UbaNetworkServer.cpp
+++ b/Engine/Source/Programs/UnrealBuildAccelerator/Common/Private/UbaNetworkServer.cpp
@@ -365,6 +365,8 @@ namespace uba
 			if (serviceId == SystemServiceId && messageType == SystemMessageType_KeepAlive) // Keep alive
 				return true;
 
+			LOG_STALL_SCOPE(conn.m_server.m_logger, 5, TC("PopWorker took more than %s"));
+
 			//m_logger.Debug(TC("Recv: %u, %u, %u, %u"), serviceId, messageType, id, size);
 			Worker* worker = conn.m_server.PopWorker();
 			if (!worker)
diff --git a/Engine/Source/Programs/UnrealBuildAccelerator/Common/Private/UbaSessionClient.cpp b/Engine/Source/Programs/UnrealBuildAccelerator/Common/Private/UbaSessionClient.cpp
--- a/Engine/Source/Programs/UnrealBuildAccelerator/Common/Private/UbaSessionClient.cpp
+++ b/Engine/Source/Programs/UnrealBuildAccelerator/Common/Private/UbaSessionClient.cpp
@@ -59,7 +59,10 @@ namespace uba
 		m_memWaitLoadPercent = info.memWaitLoadPercent;
 		m_memKillLoadPercent = info.memKillLoadPercent;
 		m_processFinished = info.processFinished;
+
+		#if PLATFORM_WINDOWS || PLATFORM_MAC // Linux will crash with print all callstacks
 		m_pingTimeoutSecondsPrintCallstacks = info.pingTimeoutSecondsPrintCallstacks;
+		#endif
 
 		m_useDependencyCrawler = info.useDependencyCrawler;
 
@@ -1629,18 +1632,14 @@ namespace uba
 			return;
 		}
 
-		#if PLATFORM_WINDOWS || PLATFORM_MAC // Linux will crash with print all callstacks
-		u32 timeoutSeconds = 20;
-		if (m_pingTimeoutSecondsPrintCallstacks)
-			timeoutSeconds = m_pingTimeoutSecondsPrintCallstacks;
+		u32 timeoutSeconds = m_pingTimeoutSecondsPrintCallstacks ? m_pingTimeoutSecondsPrintCallstacks : 20;
 		if (!response.done.IsSet(timeoutSeconds*1000))
 		{
 			LoggerWithWriter logger(g_consoleLogWriter);
-			logger.Info(TC("Took more than %u seconds to send/receive ping. Printing callstacks"), timeoutSeconds);
+			logger.Info(TC("Took more than %u seconds to send/receive ping.%s"), timeoutSeconds, (m_pingTimeoutSecondsPrintCallstacks ? TC(" Printing callstacks") : TC("")));
 			if (m_pingTimeoutSecondsPrintCallstacks)
 				PrintAllCallstacks(logger);
 		}
-		#endif
 		
 		response.done.IsSet(); // Wait forever (until message trigger error or finishes)
 
diff --git a/Engine/Source/Programs/UnrealBuildAccelerator/Common/Private/UbaSessionServer.cpp b/Engine/Source/Programs/UnrealBuildAccelerator/Common/Private/UbaSessionServer.cpp
--- a/Engine/Source/Programs/UnrealBuildAccelerator/Common/Private/UbaSessionServer.cpp
+++ b/Engine/Source/Programs/UnrealBuildAccelerator/Common/Private/UbaSessionServer.cpp
@@ -1596,6 +1596,8 @@ namespace uba
 
 	bool SessionServer::HandlePing(const ConnectionInfo& connectionInfo, const WorkContext& workContext, BinaryReader& reader, BinaryWriter& writer)
 	{
+		LOG_STALL_SCOPE(m_logger, 5, TC("HandlePing took more than %s"));
+		
 		u32 sessionId = reader.ReadU32();
 		u64 lastPing = reader.ReadU64();
 		u64 memAvail = reader.ReadU64();
diff --git a/Engine/Source/Programs/UnrealBuildAccelerator/Common/Public/UbaLogger.h b/Engine/Source/Programs/UnrealBuildAccelerator/Common/Public/UbaLogger.h
--- a/Engine/Source/Programs/UnrealBuildAccelerator/Common/Public/UbaLogger.h
+++ b/Engine/Source/Programs/UnrealBuildAccelerator/Common/Public/UbaLogger.h
@@ -7,6 +7,7 @@
 #include <stdarg.h>
 
 #define UBA_DEBUG_LOGGER 0
+#define UBA_LOG_STALLS 0
 
 namespace uba
 {
@@ -102,4 +103,20 @@ namespace uba
 	struct CallstackInfo { Vector<u8> data; Vector<u32> threadIds; TString desc; };
 	void TraverseAllCallstacks(const Function<void(const CallstackInfo&)>& func, const TraverseThreadErrorFunc& errorFunc);
 	void PrintAllCallstacks(Logger& logger);
+
+
+	#if UBA_LOG_STALLS
+	struct LogStallScope
+	{
+		LogStallScope(Logger& logger, u64 timeSeconds, const tchar* messageFormat);
+		~LogStallScope();
+		Logger& logger;
+		u64 timeSeconds;
+		u64 timeStart;
+		const tchar* messageFormat;
+	};
+	#define LOG_STALL_SCOPE(logger, timeSeconds, messageFormat) LogStallScope lss(logger, timeSeconds, messageFormat);
+	#else
+	#define LOG_STALL_SCOPE(logger, timeSeconds, messageFormat)
+	#endif
 }
