[UBA]
* Implemented everything needed to be able to use uba scheduler for executor in UBT.
* Added support for rerunning processes in scheduler
* Added some more logging/asserts

diff --git a/Engine/Source/Programs/UnrealBuildAccelerator/Common/Private/UbaProcess.cpp b/Engine/Source/Programs/UnrealBuildAccelerator/Common/Private/UbaProcess.cpp
--- a/Engine/Source/Programs/UnrealBuildAccelerator/Common/Private/UbaProcess.cpp
+++ b/Engine/Source/Programs/UnrealBuildAccelerator/Common/Private/UbaProcess.cpp
@@ -55,6 +55,7 @@ namespace uba
 
 	void Process::Release()
 	{
+		UBA_ASSERT(m_refCount);
 		if (!--m_refCount)
 			delete this;
 	}
@@ -643,7 +644,8 @@ namespace uba
 			m_startInfo.userData = nullptr;
 			ProcessHandle h;
 			h.m_process = this;
-			exitedFunc(userData, h);
+			ProcessExitedResponse exitedResponse = ProcessExitedResponse_None;
+			exitedFunc(userData, h, exitedResponse);
 			h.m_process = nullptr;
 		}
 
diff --git a/Engine/Source/Programs/UnrealBuildAccelerator/Common/Private/UbaScheduler.cpp b/Engine/Source/Programs/UnrealBuildAccelerator/Common/Private/UbaScheduler.cpp
--- a/Engine/Source/Programs/UnrealBuildAccelerator/Common/Private/UbaScheduler.cpp
+++ b/Engine/Source/Programs/UnrealBuildAccelerator/Common/Private/UbaScheduler.cpp
@@ -30,7 +30,7 @@ namespace uba
 		const u8* knownInputs;
 		u32 knownInputsCount;
 		float weight = 1.0f;
-		u64 queryCacheTime = 0;
+		u32 cacheBucketId = 0;
 	};
 
 
@@ -56,7 +56,8 @@ namespace uba
 		virtual const Vector<u8>& GetTrackedInputs() const override { static Vector<u8> v; return v;}
 		virtual const Vector<u8>& GetTrackedOutputs() const override { static Vector<u8> v; return v;}
 		virtual bool IsRemote() const override { return false; }
-		virtual ProcessExecutionType GetExecutionType() const override { return ProcessExecutionType_Native; }
+		virtual bool IsCache() const { return false; }
+		virtual ProcessExecutionType GetExecutionType() const override { return ProcessExecutionType_Skipped; }
 		ProcessStartInfoHolder startInfo;
 	};
 
@@ -72,6 +73,7 @@ namespace uba
 		virtual const Vector<u8>& GetTrackedInputs() const override { static Vector<u8> v; return v;}
 		virtual const Vector<u8>& GetTrackedOutputs() const override { static Vector<u8> v; return v;}
 		virtual bool IsRemote() const override { return false; }
+		virtual bool IsCache() const { return false; }
 		virtual ProcessExecutionType GetExecutionType() const override { return ProcessExecutionType_FromCache; }
 		ProcessStartInfoHolder startInfo;
 		Vector<ProcessLogLine> logLines;
@@ -97,9 +99,9 @@ namespace uba
 	,	m_forceRemote(info.forceRemote)
 	,	m_forceNative(info.forceNative)
 	,	m_processConfigs(info.processConfigs)
-	,	m_cacheClient(info.cacheClient)
-	,	m_writeToCache(info.writeToCache && info.cacheClient)
+	,	m_writeToCache(info.writeToCache && info.cacheClientCount)
 	{
+		m_cacheClients.insert(m_cacheClients.end(), info.cacheClients, info.cacheClients + info.cacheClientCount);
 		m_session.RegisterGetNextProcess([this](Process& process, NextProcessInfo& outNextProcess, u32 prevExitCode)
 			{
 				return HandleReuseMessage(process, outNextProcess, prevExitCode);
@@ -128,7 +130,20 @@ namespace uba
 		m_updateThreadLoop.Set();
 		m_thread.Wait();
 		m_session.WaitOnAllTasks();
+		Cleanup();
+	}
+
+	void Scheduler::Cancel()
+	{
+		m_loop = false;
+		m_updateThreadLoop.Set();
+		m_thread.Wait();
+		m_session.CancelAllProcessesAndWait();
+		Cleanup();
+	}
 
+	void Scheduler::Cleanup()
+	{
 		SCOPED_WRITE_LOCK(m_processEntriesLock, lock);
 		for (auto& entry : m_processEntries)
 		{
@@ -146,6 +161,11 @@ namespace uba
 		m_updateThreadLoop.Set();
 	}
 
+	void Scheduler::SetAllowDisableRemoteExecution(bool allow)
+	{
+		m_allowDisableRemoteExecution = allow;
+	}
+
 	u32 Scheduler::EnqueueProcess(const EnqueueProcessInfo& info)
 	{
 		u8* ki = nullptr;
@@ -165,11 +185,12 @@ namespace uba
 		auto info2 = new ProcessStartInfo2(info.info, ki, info.knownInputsCount);
 		info2->Expand();
 		info2->weight = info.weight;
+		info2->cacheBucketId = info.cacheBucketId;
 
 		const ApplicationRules* rules = m_session.GetRules(*info2);
 		info2->rules = rules;
 
-		bool useCache = m_cacheClient && !m_writeToCache && rules->IsCacheable();
+		bool useCache = info.cacheBucketId && !m_cacheClients.empty() && !m_writeToCache && rules->IsCacheable();
 
 		SCOPED_WRITE_LOCK(m_processEntriesLock, lock);
 		u32 index = u32(m_processEntries.size());
@@ -198,6 +219,8 @@ namespace uba
 
 		lock.Leave();
 
+		++m_totalProcesses;
+
 		UpdateQueueCounter(1);
 
 		m_updateThreadLoop.Set();
@@ -215,7 +238,7 @@ namespace uba
 	bool Scheduler::IsEmpty()
 	{
 		SCOPED_READ_LOCK(m_processEntriesLock, lock);
-		return m_processEntries.size() == m_finishedProcesses;
+		return m_processEntries.size() <= m_finishedProcesses;
 	}
 
 	void Scheduler::SetProcessFinishedCallback(const Function<void(const ProcessHandle&)>& processFinished)
@@ -225,6 +248,9 @@ namespace uba
 
 	u32 Scheduler::GetProcessCountThatCanRunRemotelyNow()
 	{
+		if (m_session.IsRemoteExecutionDisabled())
+			return 0;
+
 		u32 count = 0;
 		SCOPED_READ_LOCK(m_processEntriesLock, lock);
 		for (auto& entry : m_processEntries)
@@ -241,6 +267,35 @@ namespace uba
 		return count;
 	}
 
+	float Scheduler::GetProcessWeightThatCanRunRemotelyNow()
+	{
+		if (m_session.IsRemoteExecutionDisabled())
+			return 0;
+
+		float weight = 0;
+		SCOPED_READ_LOCK(m_processEntriesLock, lock);
+		for (auto& entry : m_processEntries)
+		{
+			if (!entry.canExecuteRemotely)
+				continue;
+			if (entry.status != ProcessStatus_QueuedForRun)
+				continue;
+			bool canRun = true;
+			for (u32 j=0, je=entry.dependencyCount; j!=je; ++j)
+			{
+				auto depIndex = entry.dependencies[j];
+				if (m_processEntries[depIndex].status == ProcessStatus_Success)
+					continue;
+				canRun = false;
+				break;
+			}
+			if (!canRun)
+				continue;
+			weight += entry.info->weight;
+		}
+		return weight;
+	}
+
 	void Scheduler::ThreadLoop()
 	{
 		while (m_loop)
@@ -277,11 +332,8 @@ namespace uba
 		m_updateThreadLoop.Set();
 	}
 
-	void Scheduler::HandleCacheMissed(ExitProcessInfo* ei)
+	void Scheduler::HandleCacheMissed(u32 processIndex)
 	{
-		u32 processIndex = ei->processIndex;
-		delete ei;
-
 		if (processIndex == ~0u)
 			return;
 
@@ -301,7 +353,21 @@ namespace uba
 	void Scheduler::RemoteSlotAvailable(bool isCrossArchitecture)
 	{
 		UBA_ASSERTF(!isCrossArchitecture, TC("Cross architecture code path not implemented"));
-		RunQueuedProcess(false);
+		if (RunQueuedProcess(false))
+			return;
+		if (!m_allowDisableRemoteExecution)
+			return;
+		if (m_session.IsRemoteExecutionDisabled())
+			return;
+		u32 count = 0;
+		SCOPED_READ_LOCK(m_processEntriesLock, lock);
+		for (auto& entry : m_processEntries)
+			count += u32(entry.canExecuteRemotely && entry.status <= ProcessStatus_QueuedForRun);
+		lock.Leave();
+		if (count < m_maxLocalProcessors)
+			m_session.DisableRemoteExecution();
+		else
+			m_session.SetMaxRemoteProcessCount(count);
 	}
 
 	void Scheduler::ProcessExited(ExitProcessInfo* info, const ProcessHandle& handle)
@@ -318,13 +384,6 @@ namespace uba
 			return;
 		}
 
-		if (si->queryCacheTime) // A bit hacky but we know this is a local process
-		{
-			Timer& timer = ((ProcessImpl*)handle.m_process)->m_processStats.queryCache;
-			timer.count = 1;
-			timer.time = si->queryCacheTime;
-		}
-
 		ExitProcess(*info, *handle.m_process, handle.m_process->GetExitCode(), false);
 	}
 
@@ -436,48 +495,64 @@ namespace uba
 
 			UpdateActiveProcessCounter(isLocal, 1);
 	
-			auto exitInfo = new ExitProcessInfo();
-			exitInfo->scheduler = this;
-			exitInfo->startInfo = info;
-			exitInfo->isLocal = isLocal;
-			exitInfo->processIndex = indexToRun;
-
-			ProcessStartInfo si = *info;
-			si.userData = exitInfo;
-			si.exitedFunc = [](void* userData, const ProcessHandle& handle)
-				{
-					auto ei = (ExitProcessInfo*)userData;
-					ei->scheduler->ProcessExited(ei, handle);
-				};
-
-			UBA_ASSERT(si.rules);
-			if (m_writeToCache && si.rules->IsCacheable())
-				si.trackInputs = true;
-			else if (prevStatus == ProcessStatus_QueuedForCache)
+			if (prevStatus == ProcessStatus_QueuedForCache)
 			{
 				// TODO: This should not use work manager since it is mostly waiting on network
-				m_session.GetServer().AddWork([this, exitInfo](const WorkContext&)
+				m_session.GetServer().AddWork([this, indexToRun, info](const WorkContext&)
 					{
-						ProcessStartInfo& si = *exitInfo->startInfo;
-						u64 startTime = GetTime();
+						ProcessStartInfo2& si = *info;
 
 						CacheResult cacheResult;
-						if (m_cacheClient->FetchFromCache(cacheResult, *m_rootPaths[0], 0, si) && cacheResult.hit)
+						bool isHit = false;
+						for (auto& cacheClient : m_cacheClients)
 						{
+							if (!cacheClient->FetchFromCache(cacheResult, si.rootsHandle, si.cacheBucketId, si) && cacheResult.hit)
+								continue;
+
+							isHit = true;
 							auto process = new CachedProcess(si);
+							process->logLines.swap(cacheResult.logLines);
 							ProcessHandle ph(process);
-							exitInfo->startInfo->queryCacheTime = GetTime() - startTime;
-							ExitProcess(*exitInfo, *process, 0, true);
+							ExitProcessInfo exitInfo;
+							exitInfo.scheduler = this;
+							exitInfo.startInfo = info;
+							exitInfo.isLocal = true;
+							exitInfo.processIndex = indexToRun;
+							ExitProcess(exitInfo, *process, 0, true);
+							++m_cacheHitCount;
+							break;
 						}
-						else
+
+						if (!isHit)
 						{
-							exitInfo->startInfo->queryCacheTime = GetTime() - startTime;
-							HandleCacheMissed(exitInfo);
+							HandleCacheMissed(indexToRun);
+							++m_cacheMissCount;
 						}
+
+						StringBuffer<> str;
+						str.Appendf(TC("Hits %u Misses %u"), m_cacheHitCount.load(), m_cacheMissCount.load());
+						m_session.GetTrace().StatusUpdate(1, 6, str.data, LogEntryType_Info, nullptr);
+
 					}, 1, TC("DownloadCache"));
 				return true;
 			}
 
+			auto exitInfo = new ExitProcessInfo();
+			exitInfo->scheduler = this;
+			exitInfo->startInfo = info;
+			exitInfo->isLocal = isLocal;
+			exitInfo->processIndex = indexToRun;
+
+			ProcessStartInfo si = *info;
+			si.userData = exitInfo;
+			si.trackInputs = m_writeToCache && si.rules->IsCacheable();
+			si.exitedFunc = [](void* userData, const ProcessHandle& handle, ProcessExitedResponse& response)
+				{
+					auto ei = (ExitProcessInfo*)userData;
+					ei->scheduler->ProcessExited(ei, handle);
+				};
+			UBA_ASSERT(si.rules);
+
 			if (isLocal)
 				m_session.RunProcess(si, true, canDetour);
 			else
@@ -566,15 +641,31 @@ namespace uba
 		ProcessHandle ph;
 		ph.m_process = &process;
 
+		ProcessExitedResponse exitedResponse = ProcessExitedResponse_None;
 		if (auto func = si->exitedFunc)
-			func(si->userData, ph);
+			func(si->userData, ph, exitedResponse);
+
+		bool isDone = exitedResponse == ProcessExitedResponse_None;
 
 		SCOPED_WRITE_LOCK(m_processEntriesLock, lock);
 		auto& entry = m_processEntries[info.processIndex];
+
 		u32* dependencies = entry.dependencies;
-		entry.status = exitCode == 0 ? ProcessStatus_Success : ProcessStatus_Failed;
-		entry.info = nullptr;
-		entry.dependencies = nullptr;
+
+		if (isDone)
+		{
+			entry.status = exitCode == 0 ? ProcessStatus_Success : ProcessStatus_Failed;
+			entry.info = nullptr;
+			entry.dependencies = nullptr;
+		}
+		else
+		{
+			entry.canExecuteRemotely = false;
+			entry.canDetour = exitedResponse != ProcessExitedResponse_RerunNative;
+			entry.status = ProcessStatus_QueuedForRun;
+			m_processEntriesStart = Min(m_processEntriesStart, info.processIndex);
+		}
+
 		if (info.isLocal)
 		{
 			if (fromCache)
@@ -582,13 +673,21 @@ namespace uba
 			else
 				m_activeLocalProcessWeight -= si->weight;
 		}
+
 		lock.Leave();
 
 		UpdateActiveProcessCounter(info.isLocal, -1);
-		FinishProcess(ph);
 		m_updateThreadLoop.Set();
-		delete[] dependencies;
-		delete si;
+
+		if (isDone)
+		{
+			if (exitCode != 0)
+				++m_errorCount;
+
+			FinishProcess(ph);
+			delete[] dependencies;
+			delete si;
+		}
 
 		if (m_writeToCache && exitCode == 0)
 		{
@@ -602,15 +701,17 @@ namespace uba
 	void Scheduler::SkipProcess(ProcessStartInfo2& info)
 	{
 		ProcessHandle ph(new SkippedProcess(info));
+		ProcessExitedResponse exitedResponse = ProcessExitedResponse_None;
 		if (auto func = info.exitedFunc)
-			func(info.userData, ph);
+			func(info.userData, ph, exitedResponse);
+		UBA_ASSERT(exitedResponse == ProcessExitedResponse_None);
 		FinishProcess(ph);
 	}
 
 	void Scheduler::UpdateQueueCounter(int offset)
 	{
 		m_queuedProcesses += u32(offset);
-		m_session.UpdateProgress(m_queuedProcesses + m_activeLocalProcesses + m_activeRemoteProcesses + m_finishedProcesses, m_finishedProcesses, 0);
+		m_session.UpdateProgress(m_totalProcesses, m_finishedProcesses, m_errorCount);
 	}
 
 	void Scheduler::UpdateActiveProcessCounter(bool isLocal, int offset)
@@ -626,7 +727,7 @@ namespace uba
 		if (m_processFinished)
 			m_processFinished(handle);
 		++m_finishedProcesses;
-		m_session.UpdateProgress(m_queuedProcesses + m_activeLocalProcesses + m_activeRemoteProcesses + m_finishedProcesses, m_finishedProcesses, 0);
+		m_session.UpdateProgress(m_totalProcesses, m_finishedProcesses, m_errorCount);
 	}
 
 	bool Scheduler::EnqueueFromFile(const tchar* yamlFilename, const Function<void(EnqueueProcessInfo&)>& enqueued)
diff --git a/Engine/Source/Programs/UnrealBuildAccelerator/Common/Private/UbaSessionClient.cpp b/Engine/Source/Programs/UnrealBuildAccelerator/Common/Private/UbaSessionClient.cpp
--- a/Engine/Source/Programs/UnrealBuildAccelerator/Common/Private/UbaSessionClient.cpp
+++ b/Engine/Source/Programs/UnrealBuildAccelerator/Common/Private/UbaSessionClient.cpp
@@ -1632,17 +1632,26 @@ namespace uba
 			return;
 		}
 
+		bool reportPing = false;
+
 		u32 timeoutSeconds = m_pingTimeoutSecondsPrintCallstacks ? m_pingTimeoutSecondsPrintCallstacks : 20;
 		if (!response.done.IsSet(timeoutSeconds*1000))
 		{
+			reportPing = true;
 			LoggerWithWriter logger(g_consoleLogWriter);
-			logger.Info(TC("Took more than %u seconds to send/receive ping.%s"), timeoutSeconds, (m_pingTimeoutSecondsPrintCallstacks ? TC(" Printing callstacks") : TC("")));
+			logger.Info(TC("Took more than %u seconds to send/receive ping%s"), timeoutSeconds, (m_pingTimeoutSecondsPrintCallstacks ? TC(". Printing callstacks") : TC("")));
 			if (m_pingTimeoutSecondsPrintCallstacks)
 				PrintAllCallstacks(logger);
 		}
 		
 		response.done.IsSet(); // Wait forever (until message trigger error or finishes)
 
+		if (reportPing)
+		{
+			LoggerWithWriter logger(g_consoleLogWriter);
+			logger.Info(TC("Ping finished after %s"), TimeToText(GetTime() - time).str);
+		}
+
 		if (!msg.ProcessAsyncResults(reader) || msg.GetError())
 		{
 			m_loop = false;
@@ -1988,7 +1997,7 @@ namespace uba
 
 					ExitedRec* exitedRec = new ExitedRec(*this, activeWeightLock, activeWeight, rec);
 					startInfo.userData = exitedRec;
-					startInfo.exitedFunc = [](void* userData, const ProcessHandle& h)
+					startInfo.exitedFunc = [](void* userData, const ProcessHandle& h, ProcessExitedResponse& exitedResponse)
 					{
 						auto er = (ExitedRec*)userData;
 						SessionClient& session = er->session;
diff --git a/Engine/Source/Programs/UnrealBuildAccelerator/Common/Private/UbaSessionServer.cpp b/Engine/Source/Programs/UnrealBuildAccelerator/Common/Private/UbaSessionServer.cpp
--- a/Engine/Source/Programs/UnrealBuildAccelerator/Common/Private/UbaSessionServer.cpp
+++ b/Engine/Source/Programs/UnrealBuildAccelerator/Common/Private/UbaSessionServer.cpp
@@ -58,7 +58,7 @@ namespace uba
 
 		virtual const tchar* GetExecutingHost() const override { return m_executingHost.c_str(); }
 		virtual bool IsRemote() const override { return true; }
-		virtual ProcessExecutionType GetExecutionType() const override { return ProcessExecutionType_Detoured; }
+		virtual ProcessExecutionType GetExecutionType() const override { return ProcessExecutionType_Remote; }
 		virtual bool IsChild() override { return false; }
 
 		void CallProcessExit(ProcessHandle& h)
@@ -70,7 +70,8 @@ namespace uba
 			auto userData = m_startInfo.userData;
 			m_startInfo.exitedFunc = nullptr;
 			m_startInfo.userData = nullptr;
-			exitedFunc(userData, h);
+			ProcessExitedResponse response = ProcessExitedResponse_None;
+			exitedFunc(userData, h, response);
 		}
 
 		SessionServer* m_server;
@@ -440,6 +441,11 @@ namespace uba
 		m_trace.RemoteExecutionDisabled();
 	}
 
+	bool SessionServer::IsRemoteExecutionDisabled()
+	{
+		return !m_remoteExecutionEnabled;
+	}
+
 	void SessionServer::SetCustomCasKeyFromTrackedInputs(const tchar* fileName_, const tchar* workingDir_, const u8* trackedInputs, u32 trackedInputsBytes)
 	{
 		StringBuffer<> workingDir;
diff --git a/Engine/Source/Programs/UnrealBuildAccelerator/Common/Private/UbaStorageClient.cpp b/Engine/Source/Programs/UnrealBuildAccelerator/Common/Private/UbaStorageClient.cpp
--- a/Engine/Source/Programs/UnrealBuildAccelerator/Common/Private/UbaStorageClient.cpp
+++ b/Engine/Source/Programs/UnrealBuildAccelerator/Common/Private/UbaStorageClient.cpp
@@ -1007,6 +1007,7 @@ namespace uba
 
 	void StorageClient::Ping()
 	{
+		LOG_STALL_SCOPE(m_logger, 5, TC("StorageClient::Ping took more than %s"));
 		SCOPED_FUTEX(m_proxyClientLock, lock);
 		auto proxy = m_proxyClient;
 		if (!proxy || !proxy->client.IsConnected())
diff --git a/Engine/Source/Programs/UnrealBuildAccelerator/Common/Public/UbaProcess.h b/Engine/Source/Programs/UnrealBuildAccelerator/Common/Public/UbaProcess.h
--- a/Engine/Source/Programs/UnrealBuildAccelerator/Common/Public/UbaProcess.h
+++ b/Engine/Source/Programs/UnrealBuildAccelerator/Common/Public/UbaProcess.h
@@ -54,7 +54,7 @@ namespace uba
 		virtual void Cancel(bool terminate) override;
 		virtual const tchar* GetExecutingHost() const override { return TC(""); }
 		virtual bool IsRemote() const override { return false; }
-		virtual ProcessExecutionType GetExecutionType() const override { return m_detourEnabled ? ProcessExecutionType_Detoured : ProcessExecutionType_Native; }
+		virtual ProcessExecutionType GetExecutionType() const override { return m_detourEnabled ? ProcessExecutionType_Local : ProcessExecutionType_Native; }
 		virtual bool IsChild() override { return m_parentProcess != nullptr; }
 		virtual bool IsArm() override { return m_isArmBinary; }
 
diff --git a/Engine/Source/Programs/UnrealBuildAccelerator/Common/Public/UbaProcessHandle.h b/Engine/Source/Programs/UnrealBuildAccelerator/Common/Public/UbaProcessHandle.h
--- a/Engine/Source/Programs/UnrealBuildAccelerator/Common/Public/UbaProcessHandle.h
+++ b/Engine/Source/Programs/UnrealBuildAccelerator/Common/Public/UbaProcessHandle.h
@@ -23,9 +23,11 @@ namespace uba
 
 	enum ProcessExecutionType : u8
 	{
-		ProcessExecutionType_Native,
-		ProcessExecutionType_Detoured,
-		ProcessExecutionType_FromCache,
+		ProcessExecutionType_Native,    // Local with no detour
+		ProcessExecutionType_Local,     // Local with detour
+		ProcessExecutionType_Remote,    // Remotely with detour
+		ProcessExecutionType_FromCache, // Downloaded from cache (scheduler only)
+		ProcessExecutionType_Skipped,   // Skipped (scheduler only)
 	};
 
 	class ProcessHandle
diff --git a/Engine/Source/Programs/UnrealBuildAccelerator/Common/Public/UbaProcessStartInfo.h b/Engine/Source/Programs/UnrealBuildAccelerator/Common/Public/UbaProcessStartInfo.h
--- a/Engine/Source/Programs/UnrealBuildAccelerator/Common/Public/UbaProcessStartInfo.h
+++ b/Engine/Source/Programs/UnrealBuildAccelerator/Common/Public/UbaProcessStartInfo.h
@@ -9,6 +9,7 @@ namespace uba
 {
 	class ApplicationRules;
 	class ProcessHandle;
+	enum ProcessExitedResponse : u8;
 	using RootsHandle = u64;
 
 	struct ProcessStartInfo
@@ -38,11 +39,18 @@ namespace uba
 		LogLineCallback* logLineFunc = nullptr;	// Callback for when log entries happens
 		void* logLineUserData = nullptr;		// User data provided to logLine callback
 
-		using ExitedCallback = void(void* userData, const ProcessHandle&);
+		using ExitedCallback = void(void* userData, const ProcessHandle&, ProcessExitedResponse&);
 		ExitedCallback* exitedFunc = nullptr;	// Callback for when process is done (it has already exited)
 		void* userData = nullptr;				// User data provided to exit callback
 
 		const ApplicationRules* rules = nullptr;// Internal use (for now)
 		int uiLanguage = 1033;					// Internal use
 	};
+
+	enum ProcessExitedResponse : u8
+	{
+		ProcessExitedResponse_None,
+		ProcessExitedResponse_RerunLocal,
+		ProcessExitedResponse_RerunNative,
+	};
 }
diff --git a/Engine/Source/Programs/UnrealBuildAccelerator/Common/Public/UbaScheduler.h b/Engine/Source/Programs/UnrealBuildAccelerator/Common/Public/UbaScheduler.h
--- a/Engine/Source/Programs/UnrealBuildAccelerator/Common/Public/UbaScheduler.h
+++ b/Engine/Source/Programs/UnrealBuildAccelerator/Common/Public/UbaScheduler.h
@@ -24,7 +24,8 @@ namespace uba
 		void Apply(const Config& config);
 
 		SessionServer& session;
-		CacheClient* cacheClient = nullptr; // Set cache client for scheduler to use when building
+		CacheClient** cacheClients = nullptr; // Set cache clients for scheduler to use when building
+		u32 cacheClientCount = 0;
 		u32 maxLocalProcessors = ~0u; // Max local processors to use. ~0u means it will use all processors
 		bool enableProcessReuse = false; // If this is true, the system will allow processes to be reused when they're asking for it.
 		bool forceRemote = false; // Force all processes that can run remotely to run remotely.
@@ -49,6 +50,8 @@ namespace uba
 
 		const u32* dependencies = nullptr; // An array of u32 holding indicies to processes this process depends on. Index is a rolling number returned by EnqueueProcess
 		u32 dependencyCount = 0; // Number of elements in dependencies
+
+		u32 cacheBucketId = 0; // Bucket that cache should be fetched from. Zero means that it will not fetch anything
 	};
 
 	class Scheduler
@@ -59,7 +62,9 @@ namespace uba
 
 		void Start(); // Start scheduler thread. Should be called before server starts listen to connections if using remote help
 		void Stop(); // Will wait on all active processes and then exit.
+		void Cancel(); // Cancel and wait
 		void SetMaxLocalProcessors(u32 maxLocalProcessors); // Set max local processes
+		void SetAllowDisableRemoteExecution(bool allow); // Allow scheduler to tell clients to disconnect early if running out of processes
 
 		u32 EnqueueProcess(const EnqueueProcessInfo& info); // Returns index of process. Index is a rolling number
 
@@ -74,6 +79,7 @@ namespace uba
 		SessionServer& GetSession() { return m_session; }
 		
 		u32 GetProcessCountThatCanRunRemotelyNow();
+		float GetProcessWeightThatCanRunRemotelyNow();
 
 	private:
 		struct ExitProcessInfo;
@@ -90,8 +96,9 @@ namespace uba
 		};
 
 		void ThreadLoop();
+		void Cleanup();
 		void RemoteProcessReturned(Process& process);
-		void HandleCacheMissed(ExitProcessInfo* ei);
+		void HandleCacheMissed(u32 processIndex);
 		void RemoteSlotAvailable(bool isCrossArchitecture);
 		void ProcessExited(ExitProcessInfo* info, const ProcessHandle& handle);
 		u32 PopProcess(bool isLocal, ProcessStatus& outPrevStatus);
@@ -128,17 +135,22 @@ namespace uba
 		bool m_enableProcessReuse;
 		bool m_forceRemote;
 		bool m_forceNative;
+		bool m_allowDisableRemoteExecution = false;
 		ConfigTable* m_processConfigs = nullptr;
 
 		float m_activeLocalProcessWeight = 0.0f;
 		u32 m_activeCacheQueries = 0;
 
+		Atomic<u32> m_totalProcesses;
 		Atomic<u32> m_queuedProcesses;
 		Atomic<u32> m_activeLocalProcesses;
 		Atomic<u32> m_activeRemoteProcesses;
 		Atomic<u32> m_finishedProcesses;
+		Atomic<u32> m_errorCount;
+		Atomic<u32> m_cacheHitCount;
+		Atomic<u32> m_cacheMissCount;
 
-		CacheClient* m_cacheClient;
+		Vector<CacheClient*> m_cacheClients;
 		Vector<RootPaths*> m_rootPaths;
 		bool m_writeToCache;
 
diff --git a/Engine/Source/Programs/UnrealBuildAccelerator/Common/Public/UbaSessionServer.h b/Engine/Source/Programs/UnrealBuildAccelerator/Common/Public/UbaSessionServer.h
--- a/Engine/Source/Programs/UnrealBuildAccelerator/Common/Public/UbaSessionServer.h
+++ b/Engine/Source/Programs/UnrealBuildAccelerator/Common/Public/UbaSessionServer.h
@@ -30,8 +30,9 @@ namespace uba
 		// This can be useful if there are free local cores and we know local machine is faster or network connection to remote is slow
 		ProcessHandle RunProcessRacing(u32 raceAgainstRemoteProcessId);
 
-		// Disable remote execution. This will tell all clients to take on new processes and disconnect as soon as their current processes are finished
+		// Disable remote execution. This will tell all clients to stop taking on new processes and disconnect as soon as their current processes are finished
 		void DisableRemoteExecution();
+		bool IsRemoteExecutionDisabled();
 
 		// This can be used to set a custom cas key based on a file and its inputs. Can be used for non-deterministic outputs to still be able to use cached cas content on clients
 		// For example, when a pch is built on the host we can use all the input involved to create the pch and use that as the cas key for the resulting huge file
diff --git a/Engine/Source/Programs/UnrealBuildAccelerator/Host/Private/UbaExports.cpp b/Engine/Source/Programs/UnrealBuildAccelerator/Host/Private/UbaExports.cpp
--- a/Engine/Source/Programs/UnrealBuildAccelerator/Host/Private/UbaExports.cpp
+++ b/Engine/Source/Programs/UnrealBuildAccelerator/Host/Private/UbaExports.cpp
@@ -94,7 +94,7 @@ namespace uba
 			return config;
 		LoggerWithWriter logger(g_nullLogWriter);
 		StringBuffer<> temp;
-		if (!fileName)
+		if (!fileName || !*fileName)
 		{
 			GetDirectoryOfCurrentModule(logger, temp);
 			temp.EnsureEndsWithSlash().Append(TCV("UbaHost.toml"));
@@ -360,6 +360,11 @@ extern "C"
 		return handle->GetExitCode();
 	}
 
+	uba::u8 ProcessHandle_GetExecutionType(const uba::ProcessHandle* handle)
+	{
+		return handle->GetExecutionType();
+	}
+
 	const uba::tchar* ProcessHandle_GetExecutingHost(uba::ProcessHandle* handle)
 	{
 		return handle->GetExecutingHost();
@@ -609,7 +614,8 @@ extern "C"
 
 	uba::ProcessStartInfo* ProcessStartInfo_Create(const uba::tchar* application, const uba::tchar* arguments, const uba::tchar* workingDir, const uba::tchar* description, uba::u32 priorityClass, uba::u64 outputStatsThresholdMs, bool trackInputs, const uba::tchar* logFile, ProcessHandle_ExitCallback* exit)
 	{
-		auto info = new uba::ProcessStartInfoHolder();
+		using namespace uba;
+		auto info = new ProcessStartInfoHolder();
 
 		info->applicationStr = application;
 		info->argumentsStr = arguments;
@@ -624,13 +630,13 @@ extern "C"
 		info->logFile = info->logFileStr.c_str();
 		info->priorityClass = priorityClass;
 		info->trackInputs = trackInputs;
-		info->exitedFunc = exit;
+		info->exitedFunc = (ProcessStartInfo::ExitedCallback*)exit;
 		return info;
 	}
 
-	uba::ProcessStartInfo* ProcessStartInfo_Create2(const uba::tchar* application, const uba::tchar* arguments, const uba::tchar* workingDir, const uba::tchar* description, uba::u32 priorityClass, uba::u64 rootsHandle, bool trackInputs, const uba::tchar* logFile, ProcessHandle_ExitCallback* exit)
+	uba::ProcessStartInfo* ProcessStartInfo_Create2(const uba::tchar* application, const uba::tchar* arguments, const uba::tchar* workingDir, const uba::tchar* description, uba::u32 priorityClass, uba::u64 rootsHandle, bool trackInputs, const uba::tchar* logFile, ProcessHandle_ExitCallback2* exit)
 	{
-		auto info = ProcessStartInfo_Create(application, arguments, workingDir, description, priorityClass, 0, trackInputs, logFile, exit);
+		auto info = ProcessStartInfo_Create(application, arguments, workingDir, description, priorityClass, 0, trackInputs, logFile, (ProcessHandle_ExitCallback*)exit);
 		info->rootsHandle = rootsHandle;
 		return info;
 	}
@@ -644,7 +650,7 @@ extern "C"
 
 	void ProcessStartInfo_SetExitedCallback(uba::ProcessStartInfo& info, ProcessHandle_ExitCallback* exitedFunc, void* exitedUserData)
 	{
-		info.exitedFunc = exitedFunc;
+		info.exitedFunc = (uba::ProcessStartInfo::ExitedCallback*)exitedFunc;
 		info.userData = exitedUserData;
 	}
 	
@@ -671,6 +677,16 @@ extern "C"
 		return new uba::Scheduler(info);
 	}
 
+	uba::Scheduler* Scheduler_Create3(uba::SessionServer& session, uba::CacheClient** cacheClients, uba::u32 cacheClientCount, const uba::Config& config)
+	{
+		uba::SchedulerCreateInfo info{session};
+		info.cacheClients = cacheClients;
+		info.cacheClientCount = cacheClientCount;
+		info.Apply(config);
+		info.processConfigs = &uba::GetConfig(); // TODO: Clone provided config?
+		return new uba::Scheduler(info);
+	}
+
 	void Scheduler_Start(uba::Scheduler* scheduler)
 	{
 		scheduler->Start();
@@ -686,6 +702,21 @@ extern "C"
 		return scheduler->EnqueueProcess(epi);
 	}
 
+	uba::u32 Scheduler_EnqueueProcess2(uba::Scheduler* scheduler, const uba::ProcessStartInfo& info, float weight, bool canDetour, bool canExecuteRemotely, const uba::u32* dependencies, uba::u32 dependencyCount, const void* knownInputs, uba::u32 knownInputsBytes, uba::u32 knownInputsCount, uba::u32 cacheBucket)
+	{
+		uba::EnqueueProcessInfo epi(info);
+		epi.weight = weight;
+		epi.dependencies = dependencies;
+		epi.dependencyCount = dependencyCount;
+		epi.knownInputs = knownInputs;
+		epi.knownInputsBytes = knownInputsBytes;
+		epi.knownInputsCount = knownInputsCount;
+		epi.canDetour = canDetour;
+		epi.canExecuteRemotely = canExecuteRemotely;
+		epi.cacheBucketId = cacheBucket;
+		return scheduler->EnqueueProcess(epi);
+	}
+
 	void Scheduler_SetMaxLocalProcessors(uba::Scheduler* scheduler, uba::u32 maxLocalProcessors)
 	{
 		scheduler->SetMaxLocalProcessors(maxLocalProcessors);
@@ -696,6 +727,11 @@ extern "C"
 		scheduler->Stop();
 	}
 
+	void Scheduler_Cancel(uba::Scheduler* scheduler)
+	{
+		scheduler->Cancel();
+	}
+
 	void Scheduler_Destroy(uba::Scheduler* scheduler)
 	{
 		delete scheduler;
@@ -711,6 +747,23 @@ extern "C"
 		return scheduler->IsEmpty();
 	}
 
+	void Scheduler_SetProcessFinishedCallback(uba::Scheduler* scheduler)
+	{
+		scheduler->SetProcessFinishedCallback([](const uba::ProcessHandle&)
+			{
+			});
+	}
+
+	float Scheduler_GetProcessWeightThatCanRunRemotelyNow(uba::Scheduler* scheduler)
+	{
+		return scheduler->GetProcessWeightThatCanRunRemotelyNow();
+	}
+
+	void Scheduler_SetAllowDisableRemoteExecution(uba::Scheduler* scheduler, bool allow)
+	{
+		scheduler->SetAllowDisableRemoteExecution(true);
+	}
+
 	uba::CacheClient* CacheClient_Create(uba::SessionServer* session, bool reportMissReason, const uba::tchar* crypto, const uba::tchar* hint)
 	{
 		using namespace uba;
@@ -1010,7 +1063,7 @@ extern "C"
 		info.workingDir = workDir;
 		info.description = desc;
 		info.userData = userData;
-		info.exitedFunc = exit;
+		info.exitedFunc = (uba::ProcessStartInfo::ExitedCallback*)exit;
 		return Scheduler_EnqueueProcess(scheduler, info, 1.0f, nullptr, 0, 0);
 	}
 
diff --git a/Engine/Source/Programs/UnrealBuildAccelerator/Host/Public/UbaExports.h b/Engine/Source/Programs/UnrealBuildAccelerator/Host/Public/UbaExports.h
--- a/Engine/Source/Programs/UnrealBuildAccelerator/Host/Public/UbaExports.h
+++ b/Engine/Source/Programs/UnrealBuildAccelerator/Host/Public/UbaExports.h
@@ -17,6 +17,7 @@ namespace uba
 	class SessionServer;
 	class Scheduler;
 	class StorageServer;
+	enum ProcessExitedResponse : uba::u8;
 	struct CacheResult;
 	struct ProcessStartInfo;
 	struct SessionServerCreateInfo;
@@ -77,14 +78,16 @@ extern "C"
 
 	// ProcessStartInfo
 	using ProcessHandle_ExitCallback = void(void* userData, const uba::ProcessHandle&);
+	using ProcessHandle_ExitCallback2 = void(void* userData, const uba::ProcessHandle&, uba::ProcessExitedResponse&);
 	UBA_API uba::ProcessStartInfo* ProcessStartInfo_Create(const uba::tchar* application, const uba::tchar* arguments, const uba::tchar* workingDir, const uba::tchar* description, uba::u32 priorityClass, uba::u64 outputStatsThresholdMs, bool trackInputs, const uba::tchar* logFile, ProcessHandle_ExitCallback* exit);
-	UBA_API uba::ProcessStartInfo* ProcessStartInfo_Create2(const uba::tchar* application, const uba::tchar* arguments, const uba::tchar* workingDir, const uba::tchar* description, uba::u32 priorityClass, uba::u64 rootsHandle, bool trackInputs, const uba::tchar* logFile, ProcessHandle_ExitCallback* exit);
+	UBA_API uba::ProcessStartInfo* ProcessStartInfo_Create2(const uba::tchar* application, const uba::tchar* arguments, const uba::tchar* workingDir, const uba::tchar* description, uba::u32 priorityClass, uba::u64 rootsHandle, bool trackInputs, const uba::tchar* logFile, ProcessHandle_ExitCallback2* exit);
 	UBA_API uba::ProcessStartInfo* ProcessStartInfo_Create3(const uba::Config& config, const uba::tchar* configTable = TC(""));
 	UBA_API void ProcessStartInfo_SetExitedCallback(uba::ProcessStartInfo& info, ProcessHandle_ExitCallback* exitedFunc, void* exitedUserData);
 	UBA_API void ProcessStartInfo_Destroy(uba::ProcessStartInfo* info);
 
 	// ProcessHandle
 	UBA_API uba::u32 ProcessHandle_GetExitCode(const uba::ProcessHandle* handle);
+	UBA_API uba::u8 ProcessHandle_GetExecutionType(const uba::ProcessHandle* handle);
 	UBA_API const uba::tchar* ProcessHandle_GetExecutingHost(uba::ProcessHandle* handle);
 	UBA_API const uba::tchar* ProcessHandle_GetLogLine(const uba::ProcessHandle* handle, uba::u32 index);
 	UBA_API uba::u64 ProcessHandle_GetHash(uba::ProcessHandle* handle);
@@ -136,13 +139,19 @@ extern "C"
 	// Scheduler
 	UBA_API uba::Scheduler* Scheduler_Create(uba::SessionServer* session, uba::u32 maxLocalProcessors = ~0u, bool enableProcessReuse = false);
 	UBA_API uba::Scheduler* Scheduler_Create2(uba::SessionServer& session, const uba::Config& config);
+	UBA_API uba::Scheduler* Scheduler_Create3(uba::SessionServer& session, uba::CacheClient** cacheClients, uba::u32 cacheClientCount, const uba::Config& config);
 	UBA_API void Scheduler_Start(uba::Scheduler* scheduler);
 	UBA_API uba::u32 Scheduler_EnqueueProcess(uba::Scheduler* scheduler, const uba::ProcessStartInfo& info, float weight = 1.0f, const void* knownInputs = nullptr, uba::u32 knownInputsBytes = 0, uba::u32 knownInputsCount = 0);
+	UBA_API uba::u32 Scheduler_EnqueueProcess2(uba::Scheduler* scheduler, const uba::ProcessStartInfo& info, float weight, bool canDetour, bool canExecuteRemotely, const uba::u32* dependencies, uba::u32 dependencyCount, const void* knownInputs, uba::u32 knownInputsBytes, uba::u32 knownInputsCount, uba::u32 cacheBucket);
 	UBA_API void Scheduler_SetMaxLocalProcessors(uba::Scheduler* scheduler, uba::u32 maxLocalProcessors);
 	UBA_API void Scheduler_Stop(uba::Scheduler* scheduler);
+	UBA_API void Scheduler_Cancel(uba::Scheduler* scheduler);
 	UBA_API void Scheduler_Destroy(uba::Scheduler* scheduler);
 	UBA_API void Scheduler_GetStats(uba::Scheduler* scheduler, uba::u32& outQueued, uba::u32& outActiveLocal, uba::u32& outActiveRemote, uba::u32& outFinished);
 	UBA_API bool Scheduler_IsEmpty(uba::Scheduler* scheduler);
+	UBA_API void Scheduler_SetProcessFinishedCallback(uba::Scheduler* scheduler);
+	UBA_API float Scheduler_GetProcessWeightThatCanRunRemotelyNow(uba::Scheduler* scheduler);
+	UBA_API void Scheduler_SetAllowDisableRemoteExecution(uba::Scheduler* scheduler, bool allow);
 
 	// Cache
 	UBA_API uba::CacheClient* CacheClient_Create(uba::SessionServer* session, bool reportMissReason = false, const uba::tchar* crypto = nullptr, const uba::tchar* hint = nullptr);
