[UBA]
* Fixed ubacli/ubatest compile error

diff --git a/Engine/Source/Programs/UnrealBuildAccelerator/Cli/Private/UbaCli.cpp b/Engine/Source/Programs/UnrealBuildAccelerator/Cli/Private/UbaCli.cpp
--- a/Engine/Source/Programs/UnrealBuildAccelerator/Cli/Private/UbaCli.cpp
+++ b/Engine/Source/Programs/UnrealBuildAccelerator/Cli/Private/UbaCli.cpp
@@ -1131,11 +1131,13 @@ namespace uba
 			*/
 			auto g = MakeGuard([&]() { if (cacheClient) cacheClient->GetClient().Disconnect(); });
 
+			CacheClient* cacheClients[] = { cacheClient };
 			SchedulerCreateInfo info(sessionServer);
 			info.forceRemote = isRemote;
 			info.forceNative = commandType == CommandType_Native;
 			info.maxLocalProcessors = maxProcessCount;
-			info.cacheClient = cacheClient;
+			info.cacheClients = cacheClients;
+			info.cacheClientCount = cacheClient ? 1 : 0;
 			info.writeToCache = populateCache;
 			Scheduler scheduler(info);
 
diff --git a/Engine/Source/Programs/UnrealBuildAccelerator/Test/Private/UbaTestStress.cpp b/Engine/Source/Programs/UnrealBuildAccelerator/Test/Private/UbaTestStress.cpp
--- a/Engine/Source/Programs/UnrealBuildAccelerator/Test/Private/UbaTestStress.cpp
+++ b/Engine/Source/Programs/UnrealBuildAccelerator/Test/Private/UbaTestStress.cpp
@@ -120,7 +120,7 @@ namespace uba
 				Atomic<bool>& error;
 			};
 			pi.userData = new Rec { logger, fileName, activeCount, error };
-			pi.exitedFunc = [](void* userData, const ProcessHandle& ph)
+			pi.exitedFunc = [](void* userData, const ProcessHandle& ph, ProcessExitedResponse&)
 				{
 					auto rec = (Rec*)userData;
 					for (auto& line : ph.GetLogLines())
