[UBA]
* Fixed so LoadLibrary can use relative path based on working dir
* Fixed stringbuffer capacity error when printing callstacks
* Downgraded timeouts on connect to warning
* Added NetworkServer::IsConnected(clientid)
* Added more stall scope logs
* Fixed cpu load tracking when having multiple processor groups on windows
* Fixed so if client disconnects while server got "SendFileToServer" which is waiting for cas files to arrive (this could come out of order because of multiple tcp sockets).. finally!! this was probably the long shutdown hang happening rarely
* Fixed visualizer timeline numbers being centered under the time they represent

diff --git a/Engine/Source/Programs/UnrealBuildAccelerator/Cli/Private/UbaCli.cpp b/Engine/Source/Programs/UnrealBuildAccelerator/Cli/Private/UbaCli.cpp
--- a/Engine/Source/Programs/UnrealBuildAccelerator/Cli/Private/UbaCli.cpp
+++ b/Engine/Source/Programs/UnrealBuildAccelerator/Cli/Private/UbaCli.cpp
@@ -778,7 +778,7 @@ namespace uba
 			if (!IsAbsolutePath(application.c_str()))
 			{
 				StringBuffer<> fullApplicationName;
-				if (!SearchPathForFile(logger, fullApplicationName, application.c_str(), currentDir.data))
+				if (!SearchPathForFile(logger, fullApplicationName, application.c_str(), currentDir, {}))
 					return logger.Error(TC("Failed to find full path to %s"), application.c_str());
 				application = fullApplicationName.data;
 			}
diff --git a/Engine/Source/Programs/UnrealBuildAccelerator/Common/Private/UbaFile.cpp b/Engine/Source/Programs/UnrealBuildAccelerator/Common/Private/UbaFile.cpp
--- a/Engine/Source/Programs/UnrealBuildAccelerator/Common/Private/UbaFile.cpp
+++ b/Engine/Source/Programs/UnrealBuildAccelerator/Common/Private/UbaFile.cpp
@@ -454,14 +454,11 @@ namespace uba
 		return false;
 	}
 
-	bool SearchPathForFile(Logger& logger, StringBufferBase& out, const tchar* file, const tchar* applicationDir)
+	bool SearchPathForFile(Logger& logger, StringBufferBase& out, const tchar* file, StringView workingDir, StringView applicationDir)
 	{
 		UBA_ASSERT(!IsAbsolutePath(file));
 
 		StringBuffer<> fullPath;
-		fullPath.Append(applicationDir);
-		fullPath.EnsureEndsWithSlash();
-		fullPath.Append(file);
 
 		auto TestFileExists = [&](const tchar* extraInfo)
 			{
@@ -476,8 +473,19 @@ namespace uba
 				return false;
 			};
 
-		if (TestFileExists(TC("")))
-			return true;
+		if (applicationDir.count)
+		{
+			fullPath.Append(applicationDir).EnsureEndsWithSlash().Append(file);
+			if (TestFileExists(TC("")))
+				return true;
+		}
+
+		if (workingDir.count)
+		{
+			fullPath.Clear().Append(workingDir).EnsureEndsWithSlash().Append(file);
+			if (TestFileExists(TC("")))
+				return true;
+		}
 
 		char varSeparator = IsWindows ? ';' : ':';
 		tchar buff[32 * 1024] = { 0 };
diff --git a/Engine/Source/Programs/UnrealBuildAccelerator/Common/Private/UbaLogger.cpp b/Engine/Source/Programs/UnrealBuildAccelerator/Common/Private/UbaLogger.cpp
--- a/Engine/Source/Programs/UnrealBuildAccelerator/Common/Private/UbaLogger.cpp
+++ b/Engine/Source/Programs/UnrealBuildAccelerator/Common/Private/UbaLogger.cpp
@@ -780,8 +780,14 @@ namespace uba
 				temp.Append(cs.desc).Append(TCV(" - "));
 			temp.Append(TCV("Thread Ids: "));
 			for (u32 tid : cs.threadIds)
-				if (temp.count < temp.capacity - 5)
-					temp.AppendValue(tid).Append(' ');
+			{
+				if (temp.count >= temp.capacity - 10)
+				{
+					temp.Append(TCV("..."));
+					break;
+				}
+				temp.AppendValue(tid).Append(' ');
+			}
 			cs.desc = temp.data;
 			func(cs);
 		}
diff --git a/Engine/Source/Programs/UnrealBuildAccelerator/Common/Private/UbaNetworkServer.cpp b/Engine/Source/Programs/UnrealBuildAccelerator/Common/Private/UbaNetworkServer.cpp
--- a/Engine/Source/Programs/UnrealBuildAccelerator/Common/Private/UbaNetworkServer.cpp
+++ b/Engine/Source/Programs/UnrealBuildAccelerator/Common/Private/UbaNetworkServer.cpp
@@ -124,7 +124,7 @@ namespace uba
 			m_backend.SetRecvTimeout(m_backendConnection, m_server.m_receiveTimeoutMs, this, [](void* context, u32 timeoutMs, const tchar* recvHint, const tchar* hint)
 				{
 					auto& conn = *(Connection*)context;
-					conn.m_server.m_logger.Error(TC("Connection %u timed out after %u seconds (%s%s)"), conn.m_id, timeoutMs/1000, recvHint, hint);
+					conn.m_server.m_logger.Warning(TC("Connection %u timed out after %u seconds (%s%s)"), conn.m_id, timeoutMs/1000, recvHint, hint);
 					return false;
 				});
 
@@ -276,7 +276,7 @@ namespace uba
 			if (conn.m_recvTimeoutCount < KeepAliveProbeCount)
 				return true;
 			constexpr u32 totalTimeoutSeconds = KeepAliveIdleSeconds + KeepAliveIntervalSeconds*KeepAliveProbeCount;
-			conn.m_server.m_logger.Error(TC("Connection %u timed out after %u seconds (%s%s)"), conn.m_id, totalTimeoutSeconds, recvHint, hint);
+			conn.m_server.m_logger.Warning(TC("Connection %u timed out after %u seconds (%s%s)"), conn.m_id, totalTimeoutSeconds, recvHint, hint);
 			return false;
 		}
 
@@ -1018,6 +1018,16 @@ namespace uba
 		out.connectionCount += c.connectionCount;
 	}
 
+	bool NetworkServer::IsConnected(u32 clientId)
+	{
+		SCOPED_READ_LOCK(m_clientsLock, lock);
+		auto findIt = m_clients.find(clientId);
+		if (findIt == m_clients.end())
+			return false;
+		Client& c = findIt->second;
+		return c.connectionCount > 0;
+	}
+
 	void NetworkServer::ResetTotalStats()
 	{
 		m_sendTimer = {};
@@ -1096,6 +1106,8 @@ namespace uba
 		UBA_ASSERT(info.connectionId);
 		UBA_ASSERT(info.messageId);
 
+		LOG_STALL_SCOPE(m_logger, 5, TC("NetworkServer::SendResponse took more than %s"));
+
 		SCOPED_READ_LOCK(m_connectionsLock, lock);
 		Connection* found = nullptr;
 		for (auto& it : m_connections)
@@ -1315,6 +1327,7 @@ namespace uba
 		{
 			case SystemMessageType_SetConnectionCount:
 			{
+				LOG_STALL_SCOPE(m_logger, 5, TC("SystemMessageType_SetConnectionCount took more than %s"));
 				u32 connectionCount = reader.ReadU32();
 				u32 clientId = connectionInfo.GetId();
 
@@ -1389,6 +1402,7 @@ namespace uba
 
 	bool NetworkServer::AddConnection(NetworkBackend& backend, void* backendConnection, const sockaddr& remoteSocketAddr, bool requiresCrypto, CryptoKey cryptoKey)
 	{
+		LOG_STALL_SCOPE(m_logger, 5, TC("NetworkServer::AddConnection took more than %s"));
 		SCOPED_WRITE_LOCK(m_connectionsLock, lock);
 
 		RemoveDisconnectedConnections();
diff --git a/Engine/Source/Programs/UnrealBuildAccelerator/Common/Private/UbaScheduler.cpp b/Engine/Source/Programs/UnrealBuildAccelerator/Common/Private/UbaScheduler.cpp
--- a/Engine/Source/Programs/UnrealBuildAccelerator/Common/Private/UbaScheduler.cpp
+++ b/Engine/Source/Programs/UnrealBuildAccelerator/Common/Private/UbaScheduler.cpp
@@ -639,7 +639,7 @@ namespace uba
 					if (IsAbsolutePath(psi.application))
 						FixPath(psi.application, nullptr, 0, out);
 					else
-						SearchPathForFile(m_session.GetLogger(), out, psi.application, psi.workingDir);
+						SearchPathForFile(m_session.GetLogger(), out, psi.application, ToView(psi.workingDir), {});
 				};
 			StringBuffer<> temp1;
 			StringBuffer<> temp2;
diff --git a/Engine/Source/Programs/UnrealBuildAccelerator/Common/Private/UbaSession.cpp b/Engine/Source/Programs/UnrealBuildAccelerator/Common/Private/UbaSession.cpp
--- a/Engine/Source/Programs/UnrealBuildAccelerator/Common/Private/UbaSession.cpp
+++ b/Engine/Source/Programs/UnrealBuildAccelerator/Common/Private/UbaSession.cpp
@@ -2492,7 +2492,7 @@ namespace uba
 
 		if (!IsAbsolutePath(startInfo.application))
 		{
-			if (!SearchPathForFile(m_logger, outRealApplication.Clear(), startInfo.application, startInfo.workingDir))
+			if (!SearchPathForFile(m_logger, outRealApplication.Clear(), startInfo.application, ToView(startInfo.workingDir), {}))
 				return false;
 			startInfo.applicationStr = outRealApplication.data;
 			startInfo.application = startInfo.applicationStr.c_str();
@@ -3668,12 +3668,27 @@ namespace uba
 	bool GetCpuTime(u64& outTotalTime, u64& outIdleTime)
 	{
 		outTotalTime = 0;
-        outIdleTime = 0;
+		outIdleTime = 0;
 #if PLATFORM_WINDOWS
-		u64 kernelTime, userTime;
-		if (!GetSystemTimes((FILETIME*)&outIdleTime, (FILETIME*)&kernelTime, (FILETIME*)&userTime))
+
+		GROUP_AFFINITY originalAffinity {};
+		GROUP_AFFINITY newAffinity{};
+		static u16 groupCount = GetActiveProcessorGroupCount();
+		for (u16 group=0; group!=groupCount; ++group)
+		{
+			newAffinity.Mask = ~0ull;
+			newAffinity.Group = group;
+			if (!SetThreadGroupAffinity(GetCurrentThread(), &newAffinity, group == 0 ? &originalAffinity : NULL))
+				return false;
+			u64 idleTime, kernelTime, userTime;
+			if (!GetSystemTimes((FILETIME*)&idleTime, (FILETIME*)&kernelTime, (FILETIME*)&userTime))
+				return false;
+			outIdleTime += idleTime;
+			outTotalTime += kernelTime + userTime;
+		}
+		if (!SetThreadGroupAffinity(GetCurrentThread(), &originalAffinity, nullptr))
 			return false;
-		outTotalTime = kernelTime + userTime;
+
 #elif PLATFORM_LINUX
 		int fd = open("/proc/stat", O_RDONLY | O_CLOEXEC);
 		if (fd == -1)
diff --git a/Engine/Source/Programs/UnrealBuildAccelerator/Common/Private/UbaSessionClient.cpp b/Engine/Source/Programs/UnrealBuildAccelerator/Common/Private/UbaSessionClient.cpp
--- a/Engine/Source/Programs/UnrealBuildAccelerator/Common/Private/UbaSessionClient.cpp
+++ b/Engine/Source/Programs/UnrealBuildAccelerator/Common/Private/UbaSessionClient.cpp
@@ -260,7 +260,7 @@ namespace uba
 		return true;
 	}
 
-	bool SessionClient::EnsureBinaryFile(StringBufferBase& out, StringBufferBase& outVirtual, u32 processId, const StringBufferBase& fileName, const StringKey& fileNameKey, const tchar* applicationDir, const u8* loaderPaths, u32 loaderPathsSize)
+	bool SessionClient::EnsureBinaryFile(StringBufferBase& out, StringBufferBase& outVirtual, u32 processId, StringView fileName, const StringKey& fileNameKey, StringView applicationDir, StringView workingDir, const u8* loaderPaths, u32 loaderPathsSize)
 	{
 		CasKey casKey;
 		u32 fileAttributes = DefaultAttributes(); // TODO: This is wrong.. need to retrieve from server if this is executable or not
@@ -285,6 +285,7 @@ namespace uba
 			writer.WriteString(fileName);
 			writer.WriteStringKey(fileNameKey);
 			writer.WriteString(applicationDir);
+			writer.WriteString(workingDir);
 			if (loaderPathsSize)
 				writer.WriteBytes(loaderPaths, loaderPathsSize);
 
@@ -955,9 +956,12 @@ namespace uba
 
 	bool SessionClient::GetFullFileName(GetFullFileNameResponse& out, const GetFullFileNameMessage& msg)
 	{
+		StringView workingDir(msg.process.m_startInfo.workingDirStr);
+
 		StringKeyHasher hasher;
 		hasher.UpdateNoCheck(msg.process.m_startInfo.applicationStr);
 		hasher.UpdateNoCheck(msg.fileName);
+		hasher.UpdateNoCheck(workingDir);
 		StringKey nameKey = ToStringKey(hasher);
 
 		SCOPED_FUTEX(m_nameToNameLookupLock, lock);
@@ -974,9 +978,9 @@ namespace uba
 		}
 		rec.handled = true;
 
-		StringBuffer<> dir;
-		dir.AppendDir(msg.process.m_startInfo.application);
-		if (!EnsureBinaryFile(out.fileName, out.virtualFileName, msg.process.m_id, msg.fileName, msg.fileNameKey, dir.data, msg.loaderPaths, msg.loaderPathsSize))
+		StringBuffer<> appDir;
+		appDir.AppendDir(msg.process.m_startInfo.application);
+		if (!EnsureBinaryFile(out.fileName, out.virtualFileName, msg.process.m_id, msg.fileName, msg.fileNameKey, appDir, workingDir, msg.loaderPaths, msg.loaderPathsSize))
 			return false;
 
 		StringKey fileNameKey = msg.fileNameKey;
diff --git a/Engine/Source/Programs/UnrealBuildAccelerator/Common/Private/UbaSessionServer.cpp b/Engine/Source/Programs/UnrealBuildAccelerator/Common/Private/UbaSessionServer.cpp
--- a/Engine/Source/Programs/UnrealBuildAccelerator/Common/Private/UbaSessionServer.cpp
+++ b/Engine/Source/Programs/UnrealBuildAccelerator/Common/Private/UbaSessionServer.cpp
@@ -870,11 +870,11 @@ namespace uba
 		StringBuffer<> fileName;
 		reader.ReadString(fileName);
 		StringKey fileNameKey = reader.ReadStringKey();
-		StringBuffer<> applicationDir;
-		reader.ReadString(applicationDir);
+		TString applicationDir = reader.ReadString();
+		TString workingDir = reader.ReadString();
 
 		StringBuffer<> lookupStr;
-		lookupStr.Append(fileName).Append(applicationDir).Append('#');
+		lookupStr.Append(fileName).Append(applicationDir).Append(workingDir).Append('#');
 		lookupStr.MakeLower();
 		StringKey lookupKey = ToStringKeyNoCheck(lookupStr.data, lookupStr.count);
 
@@ -944,7 +944,7 @@ namespace uba
 			}
 			#endif
 		}
-		else if (SearchPathForFile(m_logger, absoluteFile, fileName.data, applicationDir.data))
+		else if (SearchPathForFile(m_logger, absoluteFile, fileName.data, workingDir, applicationDir))
 		{
 
 			if (!absoluteFile.StartsWith(m_systemPath.data) || !IsKnownSystemFile(absoluteFile.data))
@@ -1081,6 +1081,7 @@ namespace uba
 
 	bool SessionServer::HandleSendFileToServer(const ConnectionInfo& connectionInfo, const WorkContext& workContext, BinaryReader& reader, BinaryWriter& writer)
 	{
+		u32 clientId = connectionInfo.GetId();
 		u32 processId = reader.ReadU32();
 		StringBuffer<> destination;
 		reader.ReadString(destination);
@@ -1089,12 +1090,12 @@ namespace uba
 		UBA_ASSERT(attributes);
 		CasKey casKey = reader.ReadCasKey();
 		Storage::RetrieveResult res;
-		bool success = m_storage.RetrieveCasFile(res, casKey, destination.data);
+		bool success = m_storage.RetrieveCasFile(res, casKey, destination.data, nullptr, 1, true, clientId);
 		casKey = res.casKey;
 		if (!success)
 		{
 			auto logType = connectionInfo.ShouldDisconnect() ? LogEntryType_Info : LogEntryType_Warning;
-			m_logger.Logf(logType, TC("Failed to retrieve cas for %s from client %u (Needed to write %s)"), CasKeyString(casKey).str, connectionInfo.GetId(), destination.data);
+			m_logger.Logf(logType, TC("Failed to retrieve cas for %s from client %u (Needed to write %s)"), CasKeyString(casKey).str, clientId, destination.data);
 		}
 
 		auto writeResponse = MakeGuard([&]() { writer.WriteBool(success); });
@@ -1120,12 +1121,12 @@ namespace uba
 
 				ClientSession* session = nullptr;
 				for (auto& s : m_clientSessions)
-					if (s->clientId == connectionInfo.GetId())
+					if (s->clientId == clientId)
 						session = s;
 				if (session)
 					ubaPath.Append(session->name);
 				else
-					ubaPath.Append(TCV("Connection")).AppendValue(connectionInfo.GetId());
+					ubaPath.Append(TCV("Connection")).AppendValue(clientId);
 
 				ubaPath.Append(TCV(".uba"));
 				m_storage.CopyOrLink(casKey, ubaPath.data, attributes);
diff --git a/Engine/Source/Programs/UnrealBuildAccelerator/Common/Private/UbaStorage.cpp b/Engine/Source/Programs/UnrealBuildAccelerator/Common/Private/UbaStorage.cpp
--- a/Engine/Source/Programs/UnrealBuildAccelerator/Common/Private/UbaStorage.cpp
+++ b/Engine/Source/Programs/UnrealBuildAccelerator/Common/Private/UbaStorage.cpp
@@ -1869,7 +1869,7 @@ namespace uba
 		return res;
 	}
 
-	bool StorageImpl::RetrieveCasFile(RetrieveResult& out, const CasKey& casKey, const tchar* hint, FileMappingBuffer* mappingBuffer, u64 memoryMapAlignment, bool allowProxy)
+	bool StorageImpl::RetrieveCasFile(RetrieveResult& out, const CasKey& casKey, const tchar* hint, FileMappingBuffer* mappingBuffer, u64 memoryMapAlignment, bool allowProxy, u32 clientId)
 	{
 		UBA_ASSERT(false);
 		return false;
diff --git a/Engine/Source/Programs/UnrealBuildAccelerator/Common/Private/UbaStorageClient.cpp b/Engine/Source/Programs/UnrealBuildAccelerator/Common/Private/UbaStorageClient.cpp
--- a/Engine/Source/Programs/UnrealBuildAccelerator/Common/Private/UbaStorageClient.cpp
+++ b/Engine/Source/Programs/UnrealBuildAccelerator/Common/Private/UbaStorageClient.cpp
@@ -196,7 +196,7 @@ namespace uba
 		return true;
 	}
 
-	bool StorageClient::RetrieveCasFile(RetrieveResult& out, const CasKey& casKeyTmp, const tchar* hint, FileMappingBuffer* mappingBuffer, u64 memoryMapAlignment, bool allowProxy)
+	bool StorageClient::RetrieveCasFile(RetrieveResult& out, const CasKey& casKeyTmp, const tchar* hint, FileMappingBuffer* mappingBuffer, u64 memoryMapAlignment, bool allowProxy, u32 clientId)
 	{
 		CasKey casKey = casKeyTmp;
 
diff --git a/Engine/Source/Programs/UnrealBuildAccelerator/Common/Private/UbaStorageServer.cpp b/Engine/Source/Programs/UnrealBuildAccelerator/Common/Private/UbaStorageServer.cpp
--- a/Engine/Source/Programs/UnrealBuildAccelerator/Common/Private/UbaStorageServer.cpp
+++ b/Engine/Source/Programs/UnrealBuildAccelerator/Common/Private/UbaStorageServer.cpp
@@ -83,7 +83,7 @@ namespace uba
 		return true;
 	}
 
-	bool StorageServer::RetrieveCasFile(RetrieveResult& out, const CasKey& casKey, const tchar* hint, FileMappingBuffer* mappingBuffer, u64 memoryMapAlignment, bool allowProxy)
+	bool StorageServer::RetrieveCasFile(RetrieveResult& out, const CasKey& casKey, const tchar* hint, FileMappingBuffer* mappingBuffer, u64 memoryMapAlignment, bool allowProxy, u32 clientId)
 	{
 		UBA_ASSERT(!mappingBuffer);
 		UBA_ASSERT(casKey != CasKeyZero);
@@ -116,10 +116,13 @@ namespace uba
 			timeout = m_server.DoAdditionalWork() ? 0 : 50;
 
 			u64 waited = GetTime() - startTime;
-			if (TimeToMs(waited) > 4 * 60 * 1000) // 4 minutes timeout
+			u64 waitedMs = TimeToMs(waited);
+			if (waitedMs > 500)
 			{
-				m_logger.Info(TC("Timed out waiting %s for cas %s to be transferred from remote to storage (%s)"), TimeToText(waited).str, CasKeyString(casKey).str, hint);
-				return false;
+				if (!m_server.IsConnected(clientId))
+					return m_logger.Info(TC("Client with id %u disconnected while waiting for cas %s (%s)"), clientId, CasKeyString(casKey).str, hint).ToFalse();
+				if (waitedMs > 4 * 60 * 1000) // 4 minutes timeout
+					return m_logger.Info(TC("Timed out waiting %s for cas %s to be transferred from remote to storage (%s)"), TimeToText(waited).str, CasKeyString(casKey).str, hint).ToFalse();
 			}
 		}
 		return waitEntry.Success;
diff --git a/Engine/Source/Programs/UnrealBuildAccelerator/Common/Public/UbaFile.h b/Engine/Source/Programs/UnrealBuildAccelerator/Common/Public/UbaFile.h
--- a/Engine/Source/Programs/UnrealBuildAccelerator/Common/Public/UbaFile.h
+++ b/Engine/Source/Programs/UnrealBuildAccelerator/Common/Public/UbaFile.h
@@ -63,7 +63,7 @@ namespace uba
 	bool SetEndOfFile(Logger& logger, const tchar* fileName, FileHandle handle, u64 size);
 	bool GetDirectoryOfCurrentModule(Logger& logger, StringBufferBase& out);
 	bool DeleteAllFiles(Logger& logger, const tchar* dir, bool deleteDir = true, u32* count = nullptr);
-	bool SearchPathForFile(Logger& logger, StringBufferBase& out, const tchar* file, const tchar* applicationDir);
+	bool SearchPathForFile(Logger& logger, StringBufferBase& out, const tchar* file, StringView workingDir, StringView applicationDir);
 
 	//
 	FileHandle CreateFileW(const tchar* fileName, u32 desiredAccess, u32 shareMode, u32 createDisp, u32 flagsAndAttributes);
diff --git a/Engine/Source/Programs/UnrealBuildAccelerator/Common/Public/UbaNetworkServer.h b/Engine/Source/Programs/UnrealBuildAccelerator/Common/Public/UbaNetworkServer.h
--- a/Engine/Source/Programs/UnrealBuildAccelerator/Common/Public/UbaNetworkServer.h
+++ b/Engine/Source/Programs/UnrealBuildAccelerator/Common/Public/UbaNetworkServer.h
@@ -100,6 +100,7 @@ namespace uba
 		u32 GetClientCount();
 		u32 GetConnectionCount();
 		void GetClientStats(ClientStats& out, u32 clientId);
+		bool IsConnected(u32 clientId);
 		void ResetTotalStats();
 
 		bool DoAdditionalWork();
diff --git a/Engine/Source/Programs/UnrealBuildAccelerator/Common/Public/UbaSessionClient.h b/Engine/Source/Programs/UnrealBuildAccelerator/Common/Public/UbaSessionClient.h
--- a/Engine/Source/Programs/UnrealBuildAccelerator/Common/Public/UbaSessionClient.h
+++ b/Engine/Source/Programs/UnrealBuildAccelerator/Common/Public/UbaSessionClient.h
@@ -86,7 +86,7 @@ namespace uba
 
 		bool GetCasKeyForFile(CasKey& out, u32 processId, const StringView& fileName, const StringKey& fileNameKey);
 		bool ReadModules(List<ModuleInfo>& outModules, u32 processId, const tchar* application);
-		bool EnsureBinaryFile(StringBufferBase& out, StringBufferBase& outVirtual, u32 processId, const StringBufferBase& fileName, const StringKey& fileNameKey, const tchar* applicationDir, const u8* loaderPaths, u32 loaderPathsSize);
+		bool EnsureBinaryFile(StringBufferBase& out, StringBufferBase& outVirtual, u32 processId, StringView fileName, const StringKey& fileNameKey, StringView applicationDir, StringView workingDir, const u8* loaderPaths, u32 loaderPathsSize);
 		bool WriteBinFile(StringBufferBase& out, const StringView& binaryName, const CasKey& casKey, const KeyToString& applicationDir, u32 fileAttributes);
 		bool SendFiles(ProcessImpl& process, Timer& sendFiles);
 		bool SendFile(WrittenFile& file, u32 processId, bool keepMappingInMemory, bool compressed);
diff --git a/Engine/Source/Programs/UnrealBuildAccelerator/Common/Public/UbaStorage.h b/Engine/Source/Programs/UnrealBuildAccelerator/Common/Public/UbaStorage.h
--- a/Engine/Source/Programs/UnrealBuildAccelerator/Common/Public/UbaStorage.h
+++ b/Engine/Source/Programs/UnrealBuildAccelerator/Common/Public/UbaStorage.h
@@ -38,7 +38,7 @@ namespace uba
 		virtual bool DeleteCasForFile(const tchar* file) = 0;
 
 		struct RetrieveResult { CasKey casKey; u64 size = 0; MappedView view; };
-		virtual bool RetrieveCasFile(RetrieveResult& out, const CasKey& casKey, const tchar* hint, FileMappingBuffer* mappingBuffer = nullptr, u64 memoryMapAlignment = 1, bool allowProxy = true) = 0;
+		virtual bool RetrieveCasFile(RetrieveResult& out, const CasKey& casKey, const tchar* hint, FileMappingBuffer* mappingBuffer = nullptr, u64 memoryMapAlignment = 1, bool allowProxy = true, u32 clientId = 0) = 0;
 
 		struct CachedFileInfo { CasKey casKey; };
 		virtual bool VerifyAndGetCachedFileInfo(CachedFileInfo& out, StringKey fileNameKey, u64 verifiedLastWriteTime, u64 verifiedSize) = 0;
@@ -136,7 +136,7 @@ namespace uba
 		virtual bool DecompressFileToMemory(const tchar* fileName, FileHandle fileHandle, u8* dest, u64 decompressedSize, const tchar* writeHint, u64 fileStartOffset) override;
 		virtual bool CreateDirectory(const tchar* dir) override;
 		virtual bool DeleteCasForFile(const tchar* file) override;
-		virtual bool RetrieveCasFile(RetrieveResult& out, const CasKey& casKey, const tchar* hint, FileMappingBuffer* mappingBuffer = nullptr, u64 memoryMapAlignment = 1, bool allowProxy = true) override;
+		virtual bool RetrieveCasFile(RetrieveResult& out, const CasKey& casKey, const tchar* hint, FileMappingBuffer* mappingBuffer = nullptr, u64 memoryMapAlignment = 1, bool allowProxy = true, u32 clientId = 0) override;
 		virtual bool VerifyAndGetCachedFileInfo(CachedFileInfo& out, StringKey fileNameKey, u64 verifiedLastWriteTime, u64 verifiedSize) override;
 		virtual bool InvalidateCachedFileInfo(StringKey fileNameKey) override;
 		virtual bool StoreCasFile(CasKey& out, const tchar* fileName, const CasKey& casKeyOverride, bool deferCreation) override;
diff --git a/Engine/Source/Programs/UnrealBuildAccelerator/Common/Public/UbaStorageClient.h b/Engine/Source/Programs/UnrealBuildAccelerator/Common/Public/UbaStorageClient.h
--- a/Engine/Source/Programs/UnrealBuildAccelerator/Common/Public/UbaStorageClient.h
+++ b/Engine/Source/Programs/UnrealBuildAccelerator/Common/Public/UbaStorageClient.h
@@ -54,7 +54,7 @@ namespace uba
 		virtual MappedView MapView(const CasKey& casKey, const tchar* hint) override;
 
 		virtual bool GetZone(StringBufferBase& out) override;
-		virtual bool RetrieveCasFile(RetrieveResult& out, const CasKey& casKey, const tchar* hint, FileMappingBuffer* mappingBuffer = nullptr, u64 memoryMapAlignment = 1, bool allowProxy = true) override;
+		virtual bool RetrieveCasFile(RetrieveResult& out, const CasKey& casKey, const tchar* hint, FileMappingBuffer* mappingBuffer = nullptr, u64 memoryMapAlignment = 1, bool allowProxy = true, u32 clientId = 0) override;
 		virtual bool StoreCasFile(CasKey& out, const tchar* fileName, const CasKey& casKeyOverride, bool deferCreation) override;
 		virtual bool StoreCasFileClient(CasKey& out, StringKey fileNameKey, const tchar* fileName, FileMappingHandle mappingFile, u64 mappingOffset, u64 fileSize, const tchar* hint, bool keepMappingInMemory = false, bool storeCompressed = true) override;
 		virtual bool HasCasFile(const CasKey& casKey, CasEntry** out = nullptr) override;
diff --git a/Engine/Source/Programs/UnrealBuildAccelerator/Common/Public/UbaStorageServer.h b/Engine/Source/Programs/UnrealBuildAccelerator/Common/Public/UbaStorageServer.h
--- a/Engine/Source/Programs/UnrealBuildAccelerator/Common/Public/UbaStorageServer.h
+++ b/Engine/Source/Programs/UnrealBuildAccelerator/Common/Public/UbaStorageServer.h
@@ -37,7 +37,7 @@ namespace uba
 	protected:
 
 		virtual bool GetZone(StringBufferBase& out) override;
-		virtual bool RetrieveCasFile(RetrieveResult& out, const CasKey& casKey, const tchar* hint, FileMappingBuffer* mappingBuffer = nullptr, u64 memoryMapAlignment = 1, bool allowProxy = true) override;
+		virtual bool RetrieveCasFile(RetrieveResult& out, const CasKey& casKey, const tchar* hint, FileMappingBuffer* mappingBuffer = nullptr, u64 memoryMapAlignment = 1, bool allowProxy = true, u32 clientId = 0) override;
 		virtual bool IsDisallowedPath(const tchar* fileName) override;
 		virtual void SetTrace(Trace* trace, bool detailed) override;
 		virtual bool HasProxy(u32 clientId) override;
diff --git a/Engine/Source/Programs/UnrealBuildAccelerator/Test/Private/UbaTestBasics.cpp b/Engine/Source/Programs/UnrealBuildAccelerator/Test/Private/UbaTestBasics.cpp
--- a/Engine/Source/Programs/UnrealBuildAccelerator/Test/Private/UbaTestBasics.cpp
+++ b/Engine/Source/Programs/UnrealBuildAccelerator/Test/Private/UbaTestBasics.cpp
@@ -21,6 +21,8 @@
 
 namespace uba
 {
+	bool GetCpuTime(u64& outTotalTime, u64& outIdleTime);
+
 	bool TestTime(LoggerWithWriter& logger, const StringBufferBase& rootDir)
 	{
 #if 0
@@ -42,6 +44,13 @@ namespace uba
 		u64 fileTime = GetSecondsAsFileTime(seconds);
 		u64 seconds2 = GetFileTimeAsSeconds(fileTime);
 		UBA_TEST_CHECK(seconds == seconds2, "GetSecondsAsFileTime does not match GetFileTimeAsSeconds");
+
+
+		u64 totalTime;
+		u64 maxTime;
+		if (!GetCpuTime(totalTime, maxTime))
+			return logger.Error(TC("GetCpuTime failed"));
+
 		return true;
 	}
 
diff --git a/Engine/Source/Programs/UnrealBuildAccelerator/Visualizer/Private/UbaVisualizer.cpp b/Engine/Source/Programs/UnrealBuildAccelerator/Visualizer/Private/UbaVisualizer.cpp
--- a/Engine/Source/Programs/UnrealBuildAccelerator/Visualizer/Private/UbaVisualizer.cpp
+++ b/Engine/Source/Programs/UnrealBuildAccelerator/Visualizer/Private/UbaVisualizer.cpp
@@ -2629,11 +2629,13 @@ namespace uba
 
 				if (shouldDraw)
 				{
+					SIZE s;
+					GetTextExtentPoint32W(hdc, buffer.data, buffer.count, &s);
 					RECT textRect;
 					textRect.top = top + 8;
 					textRect.bottom = textRect.top + m_activeFont.height;
-					textRect.right = pos + 20;
-					textRect.left = pos - 20;
+					textRect.right = pos + s.cx/2;
+					textRect.left = pos - s.cx/2;
 					ExtTextOutW(hdc, textRect.left, textRect.top, 0, NULL, buffer.data, buffer.count, NULL);
 				}
 			}
