[UBA]
* Bumped maintenance memory reserver size from 128 to 256mb

diff --git a/Engine/Source/Programs/UnrealBuildAccelerator/Common/Public/UbaCacheServer.h b/Engine/Source/Programs/UnrealBuildAccelerator/Common/Public/UbaCacheServer.h
--- a/Engine/Source/Programs/UnrealBuildAccelerator/Common/Public/UbaCacheServer.h
+++ b/Engine/Source/Programs/UnrealBuildAccelerator/Common/Public/UbaCacheServer.h
@@ -40,7 +40,7 @@ namespace uba
 		u64 expirationTimeSeconds = 2*24*60*60;
 
 		// The amount of reserved memory used per bucket when doing maintenance
-		u64 maintenanceReserveSize = 128ull * 1024 * 1024;
+		u64 maintenanceReserveSize = 256ull * 1024 * 1024;
 
 		// Max size of cas bucket. When within 2mb it will start decreasing expiry time by one hour
 		u64 bucketCasTableMaxSize = 32ull * 1024 * 1024;
