[UBA]
* Re-added old GetCpuLoad path if there is only one processor group

diff --git a/Engine/Source/Programs/UnrealBuildAccelerator/Common/Private/UbaSession.cpp b/Engine/Source/Programs/UnrealBuildAccelerator/Common/Private/UbaSession.cpp
--- a/Engine/Source/Programs/UnrealBuildAccelerator/Common/Private/UbaSession.cpp
+++ b/Engine/Source/Programs/UnrealBuildAccelerator/Common/Private/UbaSession.cpp
@@ -3674,20 +3674,31 @@ namespace uba
 		GROUP_AFFINITY originalAffinity {};
 		GROUP_AFFINITY newAffinity{};
 		static u16 groupCount = GetActiveProcessorGroupCount();
-		for (u16 group=0; group!=groupCount; ++group)
+		if (groupCount <= 1)
 		{
-			newAffinity.Mask = ~0ull;
-			newAffinity.Group = group;
-			if (!SetThreadGroupAffinity(GetCurrentThread(), &newAffinity, group == 0 ? &originalAffinity : NULL))
-				return false;
 			u64 idleTime, kernelTime, userTime;
 			if (!GetSystemTimes((FILETIME*)&idleTime, (FILETIME*)&kernelTime, (FILETIME*)&userTime))
 				return false;
 			outIdleTime += idleTime;
 			outTotalTime += kernelTime + userTime;
 		}
-		if (!SetThreadGroupAffinity(GetCurrentThread(), &originalAffinity, nullptr))
-			return false;
+		else
+		{
+			for (u16 group=0; group!=groupCount; ++group)
+			{
+				newAffinity.Mask = ~0ull;
+				newAffinity.Group = group;
+				if (!SetThreadGroupAffinity(GetCurrentThread(), &newAffinity, group == 0 ? &originalAffinity : NULL))
+					return false;
+				u64 idleTime, kernelTime, userTime;
+				if (!GetSystemTimes((FILETIME*)&idleTime, (FILETIME*)&kernelTime, (FILETIME*)&userTime))
+					return false;
+				outIdleTime += idleTime;
+				outTotalTime += kernelTime + userTime;
+			}
+			if (!SetThreadGroupAffinity(GetCurrentThread(), &originalAffinity, nullptr))
+				return false;
+		}
 
 #elif PLATFORM_LINUX
 		int fd = open("/proc/stat", O_RDONLY | O_CLOEXEC);
