[UBA]
* Fixed bug where bad proxy message returned data from server was mismatching with client's read
* Added UbaAgent -proxyaddr option that can be used to explicitly set which address the other agents would use when connecting to this proxy. Added for testing purposes (running the proxy from a different LAN than the rest of the helpers)
* Improved shutdown time when cancelling
* Added validation so we never hit 65535 file transfer ids in parallel.
* Removed assert in Detoured_PathFindFileNameW since it seems to be working fine in remote
* Added more logging when ping times out

diff --git a/Engine/Source/Programs/UnrealBuildAccelerator/Agent/Private/UbaAgent.cpp b/Engine/Source/Programs/UnrealBuildAccelerator/Agent/Private/UbaAgent.cpp
--- a/Engine/Source/Programs/UnrealBuildAccelerator/Agent/Private/UbaAgent.cpp
+++ b/Engine/Source/Programs/UnrealBuildAccelerator/Agent/Private/UbaAgent.cpp
@@ -77,6 +77,7 @@ namespace uba
 		logger.Info(TC("  -listen[=port]          Agent will listen for connections on port (default: %u) and help when connected"), DefaultPort);
 		logger.Info(TC("  -listenTimeout=<sec>    Number of seconds agent will listen for host before giving up (default: %u)"), DefaultListenTimeout);
 		logger.Info(TC("  -proxyport=<port>       Which port that agent will use if being assigned to be proxy for other agents (default: %u)"), DefaultStorageProxyPort);
+		logger.Info(TC("  -proxyaddr=<addr>       Which address that agent will use if being assigned to be proxy for other agents. If not set it will automatically fetch"));
 		logger.Info(TC("  -maxcpu=<number>        Max number of processes that can be started. Defaults to \"%u\" on this machine"), DefaultProcessorCount);
 		logger.Info(TC("  -mulcpu=<number>        This value multiplies with number of cpu to figure out max cpu. Defaults to 1.0"));
 		logger.Info(TC("  -maxcon=<number>        Max number of connections that can be started by agent. Defaults to \"%u\" (amount up to max will depend on ping)"), DefaultMaxConnectionCount);
@@ -350,6 +351,7 @@ namespace uba
 		TString command;
 		u16 port = DefaultPort;
 		u16 proxyPort = DefaultStorageProxyPort;
+		TString proxyAddr;
 		TString agentName;
 		bool useListen = false;
 		bool logToFile = false;
@@ -574,6 +576,11 @@ namespace uba
 				if (!parseValue(proxyPort))
 					return false;
 			}
+			else if (name.Equals(TCV("-proxyaddr")))
+			{
+				if (!parseValue(proxyAddr))
+					return false;
+			}
 			else if (name.Equals(TCV("-summary")))
 			{
 				printSummary = true;
@@ -1214,6 +1221,7 @@ namespace uba
 			storageInfo.startProxyUserData = &startProxy;
 			storageInfo.zone = zone.data;
 			storageInfo.proxyPort = proxyPort;
+			storageInfo.proxyAddress = proxyAddr.c_str();
 			storageInfo.writeToDisk = useStorage;
 
 			storageClient = new StorageClient(storageInfo);
diff --git a/Engine/Source/Programs/UnrealBuildAccelerator/Common/Private/UbaCacheServer.cpp b/Engine/Source/Programs/UnrealBuildAccelerator/Common/Private/UbaCacheServer.cpp
--- a/Engine/Source/Programs/UnrealBuildAccelerator/Common/Private/UbaCacheServer.cpp
+++ b/Engine/Source/Programs/UnrealBuildAccelerator/Common/Private/UbaCacheServer.cpp
@@ -162,18 +162,6 @@ namespace uba
 		u64 connectTime = 0;
 	};
 
-	StringView ToString(CacheMessageType type)
-	{
-		switch (type)
-		{
-			#define UBA_CACHE_MESSAGE(x) case CacheMessageType_##x: return TCV(#x);
-			UBA_CACHE_MESSAGES
-			#undef UBA_CACHE_MESSAGE
-		default:
-			return TCV("Unknown"); // Should never happen
-		}
-	}
-
 	void CacheServerCreateInfo::Apply(Config& config)
 	{
 	}
diff --git a/Engine/Source/Programs/UnrealBuildAccelerator/Common/Private/UbaLogger.cpp b/Engine/Source/Programs/UnrealBuildAccelerator/Common/Private/UbaLogger.cpp
--- a/Engine/Source/Programs/UnrealBuildAccelerator/Common/Private/UbaLogger.cpp
+++ b/Engine/Source/Programs/UnrealBuildAccelerator/Common/Private/UbaLogger.cpp
@@ -809,13 +809,15 @@ namespace uba
 			});
 	}
 
-	#if UBA_LOG_STALLS
-	LogStallScope::LogStallScope(Logger& l, u64 t, const tchar* m) : logger(l), timeSeconds(t), timeStart(GetTime()), messageFormat(m) {}
-	LogStallScope::~LogStallScope()
+	LogStallScope::LogStallScope(Logger& l, LogEntryType t, u64 ts, const tchar* m) : logger(l), type(t), timeSeconds(ts), timeStart(GetTime()), messageFormat(m) {}
+	LogStallScope::~LogStallScope() { Leave(); }
+	void LogStallScope::Leave()
 	{
+		if (!timeStart)
+			return;
 		u64 delta = GetTime() - timeStart;
 		if (delta > MsToTime(timeSeconds*1000))
-			logger.Error(messageFormat, TimeToText(delta).str);
+			logger.Logf(type, messageFormat, TimeToText(delta).str);
+		timeStart = 0;
 	}
-	#endif
 }
diff --git a/Engine/Source/Programs/UnrealBuildAccelerator/Common/Private/UbaNetwork.cpp b/Engine/Source/Programs/UnrealBuildAccelerator/Common/Private/UbaNetwork.cpp
new file mode 100644
--- /dev/null
+++ b/Engine/Source/Programs/UnrealBuildAccelerator/Common/Private/UbaNetwork.cpp
@@ -0,0 +1,71 @@
+// Copyright Epic Games, Inc. All Rights Reserved.
+
+#include "UbaNetwork.h"
+
+namespace uba
+{
+	StringView ToString(SystemMessageType type)
+	{
+		switch (type)
+		{
+			#define UBA_SYSTEM_MESSAGE(x) case SystemMessageType_##x: return AsView(TC("")#x);
+			UBA_SYSTEM_MESSAGES
+			#undef UBA_SYSTEM_MESSAGE
+		default:
+			return ToView(TC("UnknownSystemMessage")); // Should never happen
+		}
+	}
+
+	StringView ToString(StorageMessageType type)
+	{
+		switch (type)
+		{
+			#define UBA_STORAGE_MESSAGE(x) case StorageMessageType_##x: return AsView(TC("")#x);
+			UBA_STORAGE_MESSAGES
+			#undef UBA_STORAGE_MESSAGE
+		default:
+			return ToView(TC("UnknownStorageMessage")); // Should never happen
+		}
+	}
+
+	StringView ToString(SessionMessageType type)
+	{
+		switch (type)
+		{
+			#define UBA_SESSION_MESSAGE(x) case SessionMessageType_##x: return AsView(TC("")#x);
+			UBA_SESSION_MESSAGES
+			#undef UBA_SESSION_MESSAGE
+		default:
+			return ToView(TC("UnknownSessionMessage")); // Should never happen
+		}
+	}
+
+	StringView ToString(CacheMessageType type)
+	{
+		switch (type)
+		{
+			#define UBA_CACHE_MESSAGE(x) case CacheMessageType_##x: return AsView(TC("")#x);
+			UBA_CACHE_MESSAGES
+			#undef UBA_CACHE_MESSAGE
+		default:
+			return ToView(TC("UnknownCacheMessage")); // Should never happen
+		}
+	}
+
+	StringView MessageToString(u8 serviceId, u8 messageType)
+	{
+		switch (serviceId)
+		{
+		case SystemServiceId:
+			return ToString(SystemMessageType(messageType));
+		case StorageServiceId:
+			return ToString(StorageMessageType(messageType));
+		case SessionServiceId:
+			return ToString(SessionMessageType(messageType));
+		case CacheServiceId:
+			return ToString(CacheMessageType(messageType));
+		default:
+			return ToView(TC("UnknownServiceId")); // Should never happen
+		}
+	}
+}
\ No newline at end of file
diff --git a/Engine/Source/Programs/UnrealBuildAccelerator/Common/Private/UbaNetworkBackendMemory.cpp b/Engine/Source/Programs/UnrealBuildAccelerator/Common/Private/UbaNetworkBackendMemory.cpp
--- a/Engine/Source/Programs/UnrealBuildAccelerator/Common/Private/UbaNetworkBackendMemory.cpp
+++ b/Engine/Source/Programs/UnrealBuildAccelerator/Common/Private/UbaNetworkBackendMemory.cpp
@@ -196,4 +196,10 @@ namespace uba
 		outSend = 0;
 		outRecv = 0;
 	}
+
+	void NetworkBackendMemory::Validate(Logger& logger)
+	{
+		logger.Info(TC("  NetworkBackendMemory"));
+		logger.Info(TC("     NoInfo"));
+	}
 }
\ No newline at end of file
diff --git a/Engine/Source/Programs/UnrealBuildAccelerator/Common/Private/UbaNetworkBackendTcp.cpp b/Engine/Source/Programs/UnrealBuildAccelerator/Common/Private/UbaNetworkBackendTcp.cpp
--- a/Engine/Source/Programs/UnrealBuildAccelerator/Common/Private/UbaNetworkBackendTcp.cpp
+++ b/Engine/Source/Programs/UnrealBuildAccelerator/Common/Private/UbaNetworkBackendTcp.cpp
@@ -41,6 +41,7 @@
 #endif
 
 #define UBA_LOG_SOCKET_ERRORS UBA_DEBUG
+#define UBA_EMULATE_BAD_INTERNET 0
 
 namespace uba
 {
@@ -756,6 +757,12 @@ namespace uba
 		outRecv = m_totalRecv;
 	}
 
+	void NetworkBackendTcp::Validate(Logger& logger)
+	{
+		logger.Info(TC("  NetworkBackendTcp"));
+		logger.Info(TC("     NoInfo"));
+	}
+
 	bool SetBlocking(Logger& logger, SOCKET socket, bool blocking)
 	{
 #if PLATFORM_WINDOWS
@@ -862,6 +869,14 @@ namespace uba
 
 	bool NetworkBackendTcp::SendSocket(Connection& connection, Logger& logger, const void* b, u64 bufferLen, const tchar* hint)
 	{
+		#if UBA_EMULATE_BAD_INTERNET
+		if ((rand() % 10000) == 0)
+		{
+			connection.logger.Info(TC("BAD INTERNET"));
+			Sleep(10000);
+		}
+		#endif
+
 		SOCKET socket = connection.socket;
 		char* buffer = (char*)b;
 		u64 left = bufferLen;
@@ -918,6 +933,14 @@ namespace uba
 
 	bool NetworkBackendTcp::RecvSocket(Connection& connection, RecvCache& recvCache, void* b, u32& bufferLen, const tchar* hint, bool isFirstCall, bool allowLess)
 	{
+		#if UBA_EMULATE_BAD_INTERNET
+		if ((rand() % 10000) == 0)
+		{
+			connection.logger.Info(TC("BAD INTERNET"));
+			Sleep(10000);
+		}
+		#endif
+
 		SOCKET socket = connection.socket;
 
 		u8* buffer = (u8*)b;
diff --git a/Engine/Source/Programs/UnrealBuildAccelerator/Common/Private/UbaNetworkClient.cpp b/Engine/Source/Programs/UnrealBuildAccelerator/Common/Private/UbaNetworkClient.cpp
--- a/Engine/Source/Programs/UnrealBuildAccelerator/Common/Private/UbaNetworkClient.cpp
+++ b/Engine/Source/Programs/UnrealBuildAccelerator/Common/Private/UbaNetworkClient.cpp
@@ -45,6 +45,10 @@ namespace uba
 			if (m_cryptoKey == InvalidCryptoKey)
 				outCtorSuccess = false;
 		}
+
+		#if UBA_TRACK_NETWORK_TIMES
+		m_startTime = GetTime();
+		#endif
 	}
 
 	NetworkClient::~NetworkClient()
@@ -238,6 +242,10 @@ namespace uba
 		auto& connection = *(Connection*)context;
 		auto& client = connection.owner;
 
+		#if UBA_TRACK_NETWORK_TIMES
+		connection.lastHeaderRecvTime = GetTime();
+		#endif
+
 		u16 messageId = u16(headerData[0] << 8) | u16((*(u32*)(headerData + 1) & 0xff000000) >> 24);
 		u32 messageSize = *(u32*)(headerData + 1) & 0x00FFFFFF;
 
@@ -305,6 +313,12 @@ namespace uba
 		auto& msg = *(NetworkMessage*)bodyContext;
 		if (recvError)
 			msg.m_error = 2;
+
+		#if UBA_TRACK_NETWORK_TIMES
+		if (msg.m_connection)
+			msg.m_connection->lastBodyRecvTime = GetTime();
+		#endif
+
 		msg.Done();
 		return true;
 	}
@@ -408,6 +422,55 @@ namespace uba
 		logger.Info(TC(""));
 	}
 
+	void NetworkClient::ValidateNetwork(Logger& logger)
+	{
+		UnorderedSet<NetworkBackend*> backends;
+		u32 connectionsCount;
+		{
+			LogStallScope lss(logger, LogEntryType_Info, 1, TC("  Connections lock took %s"));
+			SCOPED_WRITE_LOCK(m_connectionsLock, lock);
+			lss.Leave();
+
+			connectionsCount = u32(m_connections.size());
+			u32 connectionIndex = 0;
+			for (auto& c : m_connections)
+			{
+				#if UBA_TRACK_NETWORK_TIMES
+				logger.Info(TC("  Connection %u - LastSend %s, LastHeaderRecv: %s, LastBodyRecv: %s"), connectionIndex++, TimeToText(c.lastSendTime - m_startTime).str, TimeToText(c.lastHeaderRecvTime - m_startTime).str, TimeToText(c.lastBodyRecvTime - m_startTime).str);
+				#else
+				logger.Info(TC("  Connection %u"), connectionIndex++);
+				#endif
+				backends.insert(c.backend);
+			}
+		}
+
+		{
+			LogStallScope lss(logger, LogEntryType_Info, 1, TC("  ConnectionsIterator lock took %s"));
+			SCOPED_FUTEX(m_connectionsItLock, lock);
+		}
+
+		{
+			LogStallScope lss(logger, LogEntryType_Info, 1, TC("  ActiveMessages lock took %s"));
+			SCOPED_WRITE_LOCK(m_activeMessagesLock, lock);
+			lss.Leave();
+			logger.Info(TC("  Active messages"));
+			u64 now = GetTime();
+			for (auto m : m_activeMessages)
+				if (m)
+				{
+					u64 sendTime = 0;
+					#if UBA_TRACK_NETWORK_TIMES
+					sendTime = m->m_sendTime;
+					#endif
+
+					logger.Info(TC("     %s (%u): %s"), MessageToString(m->GetServiceId(), m->GetMessageType()).data, u32(m->m_id), TimeToText(now - sendTime).str);
+				}
+		}
+
+		for (auto backend : backends)
+			backend->Validate(logger);
+	}
+
 	void NetworkClient::RegisterOnConnected(const OnConnectedFunction& function)
 	{
 		SCOPED_FUTEX(m_onConnectedFunctionsLock, lock);
@@ -623,6 +686,10 @@ namespace uba
 
 		{
 			TimerScope ts(m_sendTimer);
+			#if UBA_TRACK_NETWORK_TIMES
+			connection.lastSendTime = ts.start;
+			message.m_sendTime = ts.start;
+			#endif
 			if (!connection.backend->Send(m_logger, connection.backendConnection, data, sendSize, message.m_sendContext, TC("Message")))
 			{
 				message.m_error = 9;
@@ -640,7 +707,7 @@ namespace uba
 			u32 timeoutMs = 10 * 60 * 1000;
 			if (!gotResponse.IsSet(timeoutMs))
 			{
-				m_logger.Error(TC("Timed out after %s waiting for message response from server."), TimeToText(GetTime() - waitStart, true).str);
+				m_logger.Error(TC("Timed out after %s waiting for message response from server. (%s)"), TimeToText(GetTime() - waitStart, true).str, MessageToString(message.GetServiceId(), message.GetMessageType()).data);
 				message.m_error = 4;
 				OnDisconnected(connection, 4);
 			}
@@ -739,6 +806,16 @@ namespace uba
 		return true;
 	}
 
+	u8 NetworkMessage::GetServiceId()
+	{
+		return m_sendWriter ? m_sendWriter->GetData()[0] >> 6 : 0;
+	}
+
+	u8 NetworkMessage::GetMessageType()
+	{
+		return m_sendWriter ? m_sendWriter->GetData()[0] & 63 : 0;
+	}
+
 	void NetworkMessage::Done(bool shouldLock)
 	{
 		bool hasId = false;
diff --git a/Engine/Source/Programs/UnrealBuildAccelerator/Common/Private/UbaScheduler.cpp b/Engine/Source/Programs/UnrealBuildAccelerator/Common/Private/UbaScheduler.cpp
--- a/Engine/Source/Programs/UnrealBuildAccelerator/Common/Private/UbaScheduler.cpp
+++ b/Engine/Source/Programs/UnrealBuildAccelerator/Common/Private/UbaScheduler.cpp
@@ -138,12 +138,9 @@ namespace uba
 
 	void Scheduler::Cancel()
 	{
-		m_loop = false;
-		m_updateThreadLoop.Set();
-		m_thread.Wait();
+		m_enableProcessReuse = false;
 		SkipAllQueued();
-		m_session.CancelAllProcessesAndWait();
-		Cleanup();
+		m_session.CancelAllProcesses();
 	}
 
 	void Scheduler::SkipAllQueued()
diff --git a/Engine/Source/Programs/UnrealBuildAccelerator/Common/Private/UbaSession.cpp b/Engine/Source/Programs/UnrealBuildAccelerator/Common/Private/UbaSession.cpp
--- a/Engine/Source/Programs/UnrealBuildAccelerator/Common/Private/UbaSession.cpp
+++ b/Engine/Source/Programs/UnrealBuildAccelerator/Common/Private/UbaSession.cpp
@@ -1761,6 +1761,18 @@ namespace uba
 		--m_logger.isMuted;
 	}
 
+	void Session::CancelAllProcesses()
+	{
+		Vector<ProcessHandle> processes;
+		SCOPED_FUTEX(m_processesLock, lock);
+		processes.reserve(m_processes.size());
+		for (auto& pair : m_processes)
+			processes.push_back(pair.second);
+		lock.Leave();
+		for (auto& process : processes)
+			process.Cancel(true);
+	}
+
 	ProcessHandle Session::RunProcess(const ProcessStartInfo& startInfo, bool async, bool enableDetour)
 	{
 		FlushDeadProcesses();
@@ -3126,6 +3138,9 @@ namespace uba
 		BottleneckScope scope(m_writeFilesBottleneck, Stats().waitBottleneck);
 		#endif
 
+		if (process.IsCancelled())
+			return false;
+
 		Atomic<bool> success = true;
 		auto span = std::span(files, fileCount);
 		m_workManager.ParallelFor(fileCount - 1, span, [&](const WorkContext&, auto& it)
diff --git a/Engine/Source/Programs/UnrealBuildAccelerator/Common/Private/UbaSessionClient.cpp b/Engine/Source/Programs/UnrealBuildAccelerator/Common/Private/UbaSessionClient.cpp
--- a/Engine/Source/Programs/UnrealBuildAccelerator/Common/Private/UbaSessionClient.cpp
+++ b/Engine/Source/Programs/UnrealBuildAccelerator/Common/Private/UbaSessionClient.cpp
@@ -1646,6 +1646,8 @@ namespace uba
 			logger.Info(TC("Took more than %u seconds to send/receive ping%s"), timeoutSeconds, (m_pingTimeoutSecondsPrintCallstacks ? TC(". Printing callstacks") : TC("")));
 			if (m_pingTimeoutSecondsPrintCallstacks)
 				PrintAllCallstacks(logger);
+			m_client.ValidateNetwork(logger);
+
 		}
 		
 		response.done.IsSet(); // Wait forever (until message trigger error or finishes)
@@ -1654,6 +1656,7 @@ namespace uba
 		{
 			LoggerWithWriter logger(g_consoleLogWriter);
 			logger.Info(TC("Ping finished after %s"), TimeToText(GetTime() - time).str);
+			m_client.ValidateNetwork(logger);
 		}
 
 		if (!msg.ProcessAsyncResults(reader) || msg.GetError())
diff --git a/Engine/Source/Programs/UnrealBuildAccelerator/Common/Private/UbaSessionServer.cpp b/Engine/Source/Programs/UnrealBuildAccelerator/Common/Private/UbaSessionServer.cpp
--- a/Engine/Source/Programs/UnrealBuildAccelerator/Common/Private/UbaSessionServer.cpp
+++ b/Engine/Source/Programs/UnrealBuildAccelerator/Common/Private/UbaSessionServer.cpp
@@ -2074,6 +2074,8 @@ namespace uba
 		ProcessHandle h(process);
 
 		SCOPED_CRITICAL_SECTION(m_remoteProcessAndSessionLock, queueLock);
+		process->m_server = nullptr;
+
 		if (process->m_clientId == ~0u)
 		{
 			for (auto it=m_queuedRemoteProcesses.begin(); it!=m_queuedRemoteProcesses.end(); ++it)
@@ -2086,7 +2088,18 @@ namespace uba
 		}
 		else
 		{
+			u32 sessionIndex = process->m_sessionId - 1;
+			UBA_ASSERT(sessionIndex < m_clientSessions.size());
+			ClientSession& session = *m_clientSessions[sessionIndex];
+			--session.usedSlotCount;
+			
 			m_activeRemoteProcesses.erase(process);
+
+			{
+				SCOPED_FUTEX(m_processesLock, lock);
+				m_processes.erase(process->m_processId);
+			}
+
 			queueLock.Leave();
 
 			StackBinaryWriter<1024> writer;
@@ -2095,16 +2108,8 @@ namespace uba
 			StorageStats().Write(writer);
 			KernelStats().Write(writer);
 			m_trace.ProcessExited(process->m_processId, process->m_exitCode, writer.GetData(), writer.GetPosition(), Vector<ProcessLogLine>());
-
-			m_logger.Warning(TC("Cancelling remote active processes has not been tested. Notify devs"));
-
-			{
-				SCOPED_FUTEX(m_processesLock, lock);
-				m_processes.erase(process->m_processId);
-			}
 		}
 
-		process->m_server = nullptr;
 		process->m_done.Set();
 	}
 
diff --git a/Engine/Source/Programs/UnrealBuildAccelerator/Common/Private/UbaStorageClient.cpp b/Engine/Source/Programs/UnrealBuildAccelerator/Common/Private/UbaStorageClient.cpp
--- a/Engine/Source/Programs/UnrealBuildAccelerator/Common/Private/UbaStorageClient.cpp
+++ b/Engine/Source/Programs/UnrealBuildAccelerator/Common/Private/UbaStorageClient.cpp
@@ -55,6 +55,7 @@ namespace uba
 	,	m_startProxyUserData(info.startProxyUserData)
 	,	m_proxyConnectionCount(info.proxyConnectionCount)
 	,	m_proxyPort(info.proxyPort)
+	,	m_proxyAddress(info.proxyAddress)
 	{
 	}
 
@@ -79,11 +80,18 @@ namespace uba
 				writer.WriteString(m_zone);
 				writer.WriteU64(m_casTotalBytes);
 
-				TraverseNetworkAddresses(m_logger, [&](const StringBufferBase& addr)
-					{
-						writer.WriteString(addr);
-						return false;
-					});
+				if (m_proxyAddress.empty())
+				{
+					TraverseNetworkAddresses(m_logger, [&](const StringBufferBase& addr)
+						{
+							writer.WriteString(addr);
+							return false;
+						});
+				}
+				else
+				{
+					writer.WriteString(m_proxyAddress);
+				}
 				writer.WriteString(TC(""));
 
 
diff --git a/Engine/Source/Programs/UnrealBuildAccelerator/Common/Private/UbaStorageProxy.cpp b/Engine/Source/Programs/UnrealBuildAccelerator/Common/Private/UbaStorageProxy.cpp
--- a/Engine/Source/Programs/UnrealBuildAccelerator/Common/Private/UbaStorageProxy.cpp
+++ b/Engine/Source/Programs/UnrealBuildAccelerator/Common/Private/UbaStorageProxy.cpp
@@ -108,7 +108,14 @@ namespace uba
 	u16 StorageProxy::PopId()
 	{
 		if (m_availableIds.empty())
+		{
+			if (m_availableIdsHigh == 65534)
+			{
+				m_logger.Error(TC("OUT OF AVAILABLE IDs.. SHOULD NEVER HAPPEN!"));
+				UBA_ASSERT(false);
+			}
 			return m_availableIdsHigh++;
+		}
 		u16 storeId = m_availableIds.back();
 		m_availableIds.pop_back();
 		return storeId;
diff --git a/Engine/Source/Programs/UnrealBuildAccelerator/Common/Private/UbaStorageServer.cpp b/Engine/Source/Programs/UnrealBuildAccelerator/Common/Private/UbaStorageServer.cpp
--- a/Engine/Source/Programs/UnrealBuildAccelerator/Common/Private/UbaStorageServer.cpp
+++ b/Engine/Source/Programs/UnrealBuildAccelerator/Common/Private/UbaStorageServer.cpp
@@ -72,7 +72,7 @@ namespace uba
 	void StorageServer::WaitForActiveWork()
 	{
 		while (m_activeUnmap)
-			Sleep(5);
+			m_server.DoAdditionalWork();
 	}
 
 	bool StorageServer::GetZone(StringBufferBase& out)
@@ -315,7 +315,7 @@ namespace uba
 		return true;
 	}
 
-	bool StorageServer::GetProxy(BinaryWriter& writer, u32 clientId)
+	bool StorageServer::GetProxy(BinaryWriter& writer, u32 clientId, bool writeCasHeader)
 	{
 		SCOPED_READ_LOCK(m_connectionInfoLock, lock);
 		auto findIt = m_connectionInfo.find(clientId);
@@ -347,9 +347,12 @@ namespace uba
 		if (proxyName.IsEmpty())
 			return false;
 
-		writer.WriteU16(u16(~0));
-		writer.Write7BitEncoded(0);
-		writer.WriteByte(1 << 2);
+		if (writeCasHeader)
+		{
+			writer.WriteU16(u16(~0));
+			writer.Write7BitEncoded(0);
+			writer.WriteByte(1 << 2);
+		}
 
 		auto proxyKey = ToStringKeyNoCheck(proxyName.data, proxyName.count);
 		SCOPED_WRITE_LOCK(m_proxiesLock, proxiesLock);
@@ -434,7 +437,7 @@ namespace uba
 				if (wantsProxy) // Wants proxy
 				{
 					todoRemoveMe = 1;
-					if (GetProxy(writer, connectionInfo.GetId()))
+					if (GetProxy(writer, connectionInfo.GetId(), true))
 						return true;
 				}
 
@@ -1080,7 +1083,7 @@ namespace uba
 				}
 				lock.Leave();
 
-				GetProxy(writer, connectionInfo.GetId());
+				GetProxy(writer, connectionInfo.GetId(), false);
 				return true;
 			}
 		}
@@ -1093,7 +1096,11 @@ namespace uba
 		SCOPED_FUTEX(m_availableIdsLock, lock);
 		if (m_availableIds.empty())
 		{
-			UBA_ASSERT(m_availableIdsHigh < u16(~0) - 2);
+			if (m_availableIdsHigh == 65534)
+			{
+				m_logger.Error(TC("OUT OF AVAILABLE IDs.. SHOULD NEVER HAPPEN!"));
+				UBA_ASSERT(false);
+			}
 			return m_availableIdsHigh++;
 		}
 		u16 storeId = m_availableIds.back();
diff --git a/Engine/Source/Programs/UnrealBuildAccelerator/Common/Public/UbaLogger.h b/Engine/Source/Programs/UnrealBuildAccelerator/Common/Public/UbaLogger.h
--- a/Engine/Source/Programs/UnrealBuildAccelerator/Common/Public/UbaLogger.h
+++ b/Engine/Source/Programs/UnrealBuildAccelerator/Common/Public/UbaLogger.h
@@ -105,17 +105,19 @@ namespace uba
 	void PrintAllCallstacks(Logger& logger);
 
 
-	#if UBA_LOG_STALLS
 	struct LogStallScope
 	{
-		LogStallScope(Logger& logger, u64 timeSeconds, const tchar* messageFormat);
+		LogStallScope(Logger& logger, LogEntryType type, u64 timeSeconds, const tchar* messageFormat);
 		~LogStallScope();
+		void Leave();
 		Logger& logger;
+		LogEntryType type;
 		u64 timeSeconds;
 		u64 timeStart;
 		const tchar* messageFormat;
 	};
-	#define LOG_STALL_SCOPE(logger, timeSeconds, messageFormat) LogStallScope lss(logger, timeSeconds, messageFormat);
+	#if UBA_LOG_STALLS
+	#define LOG_STALL_SCOPE(logger, timeSeconds, messageFormat) LogStallScope lss(logger, LogEntryType_Error, timeSeconds, messageFormat);
 	#else
 	#define LOG_STALL_SCOPE(logger, timeSeconds, messageFormat)
 	#endif
diff --git a/Engine/Source/Programs/UnrealBuildAccelerator/Common/Public/UbaNetwork.h b/Engine/Source/Programs/UnrealBuildAccelerator/Common/Public/UbaNetwork.h
--- a/Engine/Source/Programs/UnrealBuildAccelerator/Common/Public/UbaNetwork.h
+++ b/Engine/Source/Programs/UnrealBuildAccelerator/Common/Public/UbaNetwork.h
@@ -20,6 +20,8 @@ namespace uba
 	static constexpr u32 CacheNetworkVersion = 5;
 	static constexpr u32 CacheBucketVersion = 3;
 
+	#define UBA_TRACK_NETWORK_TIMES 1 // Enable to track times on messages, connections etc
+
 	// Messages used over network between client/server (system, storage and session)
 
 	#define UBA_SYSTEM_MESSAGES \
@@ -116,17 +118,11 @@ namespace uba
 		#undef UBA_CACHE_MESSAGE
 	};
 
-	inline StringView ToString(StorageMessageType type)
-	{
-		switch (type)
-		{
-			#define UBA_STORAGE_MESSAGE(x) case StorageMessageType_##x: return AsView(TC("")#x);
-			UBA_STORAGE_MESSAGES
-			#undef UBA_STORAGE_MESSAGE
-		default:
-			return ToView(TC("Unknown")); // Should never happen
-		}
-	}
+	StringView ToString(SystemMessageType type);
+	StringView ToString(StorageMessageType type);
+	StringView ToString(SessionMessageType type);
+	StringView ToString(CacheMessageType type);
+	StringView MessageToString(u8 serviceId, u8 messageType);
 
 	constexpr u32 MessageErrorSize = 0xffffff;
 	constexpr u32 MessageKeepAliveSize = 0xffffff - 1;
diff --git a/Engine/Source/Programs/UnrealBuildAccelerator/Common/Public/UbaNetworkBackend.h b/Engine/Source/Programs/UnrealBuildAccelerator/Common/Public/UbaNetworkBackend.h
--- a/Engine/Source/Programs/UnrealBuildAccelerator/Common/Public/UbaNetworkBackend.h
+++ b/Engine/Source/Programs/UnrealBuildAccelerator/Common/Public/UbaNetworkBackend.h
@@ -68,6 +68,7 @@ namespace uba
 		virtual void DeleteConnection(void* connection) = 0;
 
 		virtual void GetTotalSendAndRecv(u64& outSend, u64& outRecv) = 0;
+		virtual void Validate(Logger& logger) = 0;
 	};
 
 	struct NetworkBackend::SendContext
diff --git a/Engine/Source/Programs/UnrealBuildAccelerator/Common/Public/UbaNetworkBackendMemory.h b/Engine/Source/Programs/UnrealBuildAccelerator/Common/Public/UbaNetworkBackendMemory.h
--- a/Engine/Source/Programs/UnrealBuildAccelerator/Common/Public/UbaNetworkBackendMemory.h
+++ b/Engine/Source/Programs/UnrealBuildAccelerator/Common/Public/UbaNetworkBackendMemory.h
@@ -28,6 +28,7 @@ namespace uba
 		virtual void DeleteConnection(void* connection) override;
 
 		virtual void GetTotalSendAndRecv(u64& outSend, u64& outRecv) override;
+		virtual void Validate(Logger& logger) override;
 
 	private:
 		struct Connection;
diff --git a/Engine/Source/Programs/UnrealBuildAccelerator/Common/Public/UbaNetworkBackendTcp.h b/Engine/Source/Programs/UnrealBuildAccelerator/Common/Public/UbaNetworkBackendTcp.h
--- a/Engine/Source/Programs/UnrealBuildAccelerator/Common/Public/UbaNetworkBackendTcp.h
+++ b/Engine/Source/Programs/UnrealBuildAccelerator/Common/Public/UbaNetworkBackendTcp.h
@@ -39,6 +39,7 @@ namespace uba
 		virtual void DeleteConnection(void* connection) override;
 
 		virtual void GetTotalSendAndRecv(u64& outSend, u64& outRecv) override;
+		virtual void Validate(Logger& logger) override;
 
 	private:
 		bool EnsureInitialized(Logger& logger);
diff --git a/Engine/Source/Programs/UnrealBuildAccelerator/Common/Public/UbaNetworkClient.h b/Engine/Source/Programs/UnrealBuildAccelerator/Common/Public/UbaNetworkClient.h
--- a/Engine/Source/Programs/UnrealBuildAccelerator/Common/Public/UbaNetworkClient.h
+++ b/Engine/Source/Programs/UnrealBuildAccelerator/Common/Public/UbaNetworkClient.h
@@ -49,6 +49,7 @@ namespace uba
 		bool IsConnected(u32 waitTimeoutMs = 0);
 		bool IsOrWasConnected(u32 waitTimeoutMs = 0);
 
+		void ValidateNetwork(Logger& logger);
 		void PrintSummary(Logger& logger);
 
 		using OnConnectedFunction = Function<void()>;
@@ -82,6 +83,12 @@ namespace uba
 			Atomic<u32> connected;
 			Event disconnectedEvent;
 			NetworkBackend* backend = nullptr;
+
+			#if UBA_TRACK_NETWORK_TIMES
+			Atomic<u64> lastSendTime;
+			Atomic<u64> lastHeaderRecvTime;
+			Atomic<u64> lastBodyRecvTime;
+			#endif
 		};
 
 		bool AddConnection(NetworkBackend& backend, void* backendConnection, bool* timedOut);
@@ -133,6 +140,10 @@ namespace uba
 		Timer m_encryptTimer;
 		Timer m_decryptTimer;
 
+		#if UBA_TRACK_NETWORK_TIMES
+		u64 m_startTime = 0;
+		#endif
+
 		friend NetworkMessage;
 
 		NetworkClient(const NetworkClient&) = delete;
diff --git a/Engine/Source/Programs/UnrealBuildAccelerator/Common/Public/UbaNetworkMessage.h b/Engine/Source/Programs/UnrealBuildAccelerator/Common/Public/UbaNetworkMessage.h
--- a/Engine/Source/Programs/UnrealBuildAccelerator/Common/Public/UbaNetworkMessage.h
+++ b/Engine/Source/Programs/UnrealBuildAccelerator/Common/Public/UbaNetworkMessage.h
@@ -24,6 +24,8 @@ namespace uba
 		bool ProcessAsyncResults(BinaryReader& response); // Note, this must be called after GetResponseFunc has been called
 
 		u32 GetError() const { return m_error; }
+		u8 GetServiceId();
+		u8 GetMessageType();
 
 	private:
 		void Done(bool shouldLock = true);
@@ -42,6 +44,10 @@ namespace uba
 		friend NetworkClient;
 		NetworkMessage(const NetworkMessage&) = delete;
 		NetworkMessage(NetworkMessage&&) = delete;
+
+		#if UBA_TRACK_NETWORK_TIMES
+		u64 m_sendTime = 0;
+		#endif
 	};
 
 	enum NetworkMessageError
diff --git a/Engine/Source/Programs/UnrealBuildAccelerator/Common/Public/UbaSession.h b/Engine/Source/Programs/UnrealBuildAccelerator/Common/Public/UbaSession.h
--- a/Engine/Source/Programs/UnrealBuildAccelerator/Common/Public/UbaSession.h
+++ b/Engine/Source/Programs/UnrealBuildAccelerator/Common/Public/UbaSession.h
@@ -59,6 +59,7 @@ namespace uba
 	public:
 		ProcessHandle RunProcess(const ProcessStartInfo& startInfo, bool async = true, bool enableDetour = true); // Run process. if async is false it will not return until process is done
 		void CancelAllProcessesAndWait(bool terminate = true); // Cancel all processes and wait for them to go away
+		void CancelAllProcesses(); // Cancel all processes
 
 		void PrintSummary(Logger& logger); // Print summary stats of session
 		bool RefreshDirectory(const tchar* dirName, bool forceRegister = false); // Tell uba a directory on disk has been changed by some other system while session is running
diff --git a/Engine/Source/Programs/UnrealBuildAccelerator/Common/Public/UbaStorageClient.h b/Engine/Source/Programs/UnrealBuildAccelerator/Common/Public/UbaStorageClient.h
--- a/Engine/Source/Programs/UnrealBuildAccelerator/Common/Public/UbaStorageClient.h
+++ b/Engine/Source/Programs/UnrealBuildAccelerator/Common/Public/UbaStorageClient.h
@@ -21,6 +21,7 @@ namespace uba
 		NetworkClient& client;
 		const tchar* zone = TC("");
 		u16 proxyPort = DefaultStorageProxyPort;
+		const tchar* proxyAddress = TC("");
 		bool sendCompressed = true;
 		bool allowProxy = true;
 		bool sendOneBigFileAtTheTime = true;
@@ -95,6 +96,7 @@ namespace uba
 		void* m_startProxyUserData = nullptr;
 		u32 m_proxyConnectionCount;
 		u16 m_proxyPort = 0;
+		TString m_proxyAddress;
 	};
 
 }
diff --git a/Engine/Source/Programs/UnrealBuildAccelerator/Common/Public/UbaStorageServer.h b/Engine/Source/Programs/UnrealBuildAccelerator/Common/Public/UbaStorageServer.h
--- a/Engine/Source/Programs/UnrealBuildAccelerator/Common/Public/UbaStorageServer.h
+++ b/Engine/Source/Programs/UnrealBuildAccelerator/Common/Public/UbaStorageServer.h
@@ -46,7 +46,7 @@ namespace uba
 		bool HandleMessage(const ConnectionInfo& connectionInfo, const WorkContext& workContext, u8 messageType, BinaryReader& reader, BinaryWriter& writer);
 		bool WaitForWritten(CasEntry& casEntry, ScopedWriteLock& entryLock, const ConnectionInfo& connectionInfo, const tchar* hint);
 		bool VerifyExists(bool& outExists, CasEntry& casEntry, ScopedWriteLock& entryLock, const CasKey& casKey);
-		bool GetProxy(BinaryWriter& writer, u32 clientId);
+		bool GetProxy(BinaryWriter& writer, u32 clientId, bool writeCasHeader);
 
 		u16 PopId();
 		void PushId(u16 id);
diff --git a/Engine/Source/Programs/UnrealBuildAccelerator/Detours/Private/Windows/UbaDetoursFunctionsKernelBase.inl b/Engine/Source/Programs/UnrealBuildAccelerator/Detours/Private/Windows/UbaDetoursFunctionsKernelBase.inl
--- a/Engine/Source/Programs/UnrealBuildAccelerator/Detours/Private/Windows/UbaDetoursFunctionsKernelBase.inl
+++ b/Engine/Source/Programs/UnrealBuildAccelerator/Detours/Private/Windows/UbaDetoursFunctionsKernelBase.inl
@@ -4259,9 +4259,8 @@ UINT Detoured_GetACP()
 	return res;
 }
 
-LPCWSTR Detoured_PathFindFileNameW(LPCWSTR pszPath) // This is called by Ps4SymbolTool.exe
+LPCWSTR Detoured_PathFindFileNameW(LPCWSTR pszPath) // This is called by Ps4SymbolTool.exe and vctip.exe
 {
-	UBA_ASSERTF(!g_runningRemote, L"%ls", pszPath);
 	auto res = True_PathFindFileNameW(pszPath);
 	DEBUG_LOG_TRUE(L"PathFindFileNameW", L"(%ls) -> %ls", pszPath, res);
 	return res;
