[UBA]
* Changed windows to use WSASend with overlapped IO for tcp backend. Should perform better but we also want to see if this can solve the odd send stalls we've seen in the cloud linux/wine machines

diff --git a/Engine/Source/Programs/UnrealBuildAccelerator/Common/Private/UbaNetworkBackendTcp.cpp b/Engine/Source/Programs/UnrealBuildAccelerator/Common/Private/UbaNetworkBackendTcp.cpp
--- a/Engine/Source/Programs/UnrealBuildAccelerator/Common/Private/UbaNetworkBackendTcp.cpp
+++ b/Engine/Source/Programs/UnrealBuildAccelerator/Common/Private/UbaNetworkBackendTcp.cpp
@@ -43,6 +43,8 @@
 #define UBA_LOG_SOCKET_ERRORS UBA_DEBUG
 #define UBA_EMULATE_BAD_INTERNET 0
 
+#define UBA_USE_OVERLAPPED_SEND PLATFORM_WINDOWS
+
 namespace uba
 {
 	constexpr u32 MaxHeaderSize = 24;
@@ -204,11 +206,7 @@ namespace uba
 		auto& conn = *(Connection*)connection;
 		sendContext.isUsed = true;
 
-		SCOPED_FUTEX(conn.sendLock, lock);
-		if (conn.socket == INVALID_SOCKET)
-			return false;
 		bool res = SendSocket(conn, logger, data, dataSize, sendHint);
-		lock.Leave();
 
 		sendContext.isFinished = true;
 
@@ -877,7 +875,69 @@ namespace uba
 		}
 		#endif
 
+
+#if UBA_USE_OVERLAPPED_SEND
+
+		Event ev(true);
+
+		SCOPED_FUTEX(connection.sendLock, lock);
+		SOCKET socket = connection.socket;
+		if (socket == INVALID_SOCKET)
+			return false;
+
+		char* buffer = (char*)b;
+		u64 left = bufferLen;
+		while (left)
+		{
+			WSAOVERLAPPED overlapped {};
+			overlapped.hEvent = ev.GetHandle();
+			WSABUF buf { u32(left), buffer };
+			int res = WSASend(socket, &buf, 1, NULL, 0, &overlapped, NULL);
+			if (res == SOCKET_ERROR)
+			{
+				u32 lastError = WSAGetLastError();
+				if (lastError != WSA_IO_PENDING)
+				{
+					#if UBA_LOG_SOCKET_ERRORS
+					logger.Info(TC("WSASend - error for socket %i: %s (%s)"), socket, LastErrorToText(lastError).data, hint);
+					#endif
+					return false;
+				}
+			}
+
+			if (!ev.IsSet(40*1000))
+			{
+				logger.Info(TC("SendSocket - WSASend returned timeout after 40 seconds for socket %i (%s)"), socket, hint);
+				if (!ev.IsSet(DefaultNetworkSendTimeoutSeconds * 1000))
+				{
+					logger.Info(TC("SendSocket - WSASend returned timeout after 10 minutes for socket %i (%s)"), socket, hint);
+					return false;
+				}
+			}
+
+			DWORD bytesSent;
+			DWORD flags;
+			if (!WSAGetOverlappedResult(socket, &overlapped, &bytesSent, FALSE, &flags))
+			{
+				#if UBA_LOG_SOCKET_ERRORS
+				logger.Info(TC("WSAGetOverlappedResult - error for socket %i: %s (%s)"), socket, LastErrorToText().data, hint);
+				#endif
+				return false;
+			}
+
+			buffer += bytesSent;
+			left -= bytesSent;
+			if (left)
+				ev.Reset();
+		}
+
+#else
+
+		SCOPED_FUTEX(connection.sendLock, lock);
 		SOCKET socket = connection.socket;
+		if (socket == INVALID_SOCKET)
+			return false;
+
 		char* buffer = (char*)b;
 		u64 left = bufferLen;
 		while (left)
@@ -928,6 +988,7 @@ namespace uba
 			buffer += sent;
 			left -= sent;
 		}
+#endif
 		return true;
 	}
 
