[UBA]
* Fixed bug in uba scheduler where it did not send proper process weight to remote machines
* Added log entry when agent running proxy server change max number of local processes

diff --git a/Engine/Source/Programs/UnrealBuildAccelerator/Agent/Private/UbaAgent.cpp b/Engine/Source/Programs/UnrealBuildAccelerator/Agent/Private/UbaAgent.cpp
--- a/Engine/Source/Programs/UnrealBuildAccelerator/Agent/Private/UbaAgent.cpp
+++ b/Engine/Source/Programs/UnrealBuildAccelerator/Agent/Private/UbaAgent.cpp
@@ -829,6 +829,10 @@ namespace uba
 			consoleTitle.Appendf(TC("UbaAgent v%s%s"), Version, dbgStr);
 			SetConsoleTitleW(consoleTitle.data);
 		}
+
+		// Uncomment to log winsock behavior in wine.. requires WINEDEBUG to be allowed though
+		//if (IsRunningWine())
+		//	SetEnvironmentVariable(TC("WINEDEBUG"), TC("+timestamp,+winsock"));
 #endif
 
 		u64 storageCapacity = u64(storageCapacityGb)*1000*1000*1000;
@@ -1418,7 +1422,11 @@ namespace uba
 						else
 							newProcessCount = 1;
 
-						sessionClient->SetMaxProcessCount(newProcessCount);
+						if (sessionClient->GetMaxProcessCount() != newProcessCount)
+						{
+							logger.Info(TC("Changed max process count to %u"), newProcessCount);
+							sessionClient->SetMaxProcessCount(newProcessCount);
+						}
 					}
 				}
 				// This is an estimation based on tcp limitations (ack and sliding windows).
diff --git a/Engine/Source/Programs/UnrealBuildAccelerator/Common/Private/UbaScheduler.cpp b/Engine/Source/Programs/UnrealBuildAccelerator/Common/Private/UbaScheduler.cpp
--- a/Engine/Source/Programs/UnrealBuildAccelerator/Common/Private/UbaScheduler.cpp
+++ b/Engine/Source/Programs/UnrealBuildAccelerator/Common/Private/UbaScheduler.cpp
@@ -573,7 +573,7 @@ namespace uba
 			if (isLocal)
 				m_session.RunProcess(si, true, canDetour);
 			else
-				m_session.RunProcessRemote(si, 1.0f, info->knownInputs, info->knownInputsCount);
+				m_session.RunProcessRemote(si, info->weight, info->knownInputs, info->knownInputsCount);
 			return true;
 		}
 	}
diff --git a/Engine/Source/Programs/UnrealBuildAccelerator/Common/Private/UbaSessionClient.cpp b/Engine/Source/Programs/UnrealBuildAccelerator/Common/Private/UbaSessionClient.cpp
--- a/Engine/Source/Programs/UnrealBuildAccelerator/Common/Private/UbaSessionClient.cpp
+++ b/Engine/Source/Programs/UnrealBuildAccelerator/Common/Private/UbaSessionClient.cpp
@@ -192,6 +192,11 @@ namespace uba
 		return m_bestPing;
 	}
 
+	u32 SessionClient::GetMaxProcessCount()
+	{
+		return m_maxProcessCount;
+	}
+
 	bool SessionClient::Exists(const StringView& path, u32& outAttributes)
 	{
 		u32 tableOffset = 0;
diff --git a/Engine/Source/Programs/UnrealBuildAccelerator/Common/Public/UbaSessionClient.h b/Engine/Source/Programs/UnrealBuildAccelerator/Common/Public/UbaSessionClient.h
--- a/Engine/Source/Programs/UnrealBuildAccelerator/Common/Public/UbaSessionClient.h
+++ b/Engine/Source/Programs/UnrealBuildAccelerator/Common/Public/UbaSessionClient.h
@@ -49,6 +49,7 @@ namespace uba
 		void SetAllowSpawn(bool allow);
 
 		u64 GetBestPing();
+		u32 GetMaxProcessCount();
 
 		bool Exists(const StringView& path, u32& outAttributes);
 
