[UBA]
* Added code to be able to build UbaWine.dll.so that can be loaded into wine processes to get native linux API

diff --git a/Engine/Source/Programs/UnrealBuildAccelerator/Wine/Build.sh b/Engine/Source/Programs/UnrealBuildAccelerator/Wine/Build.sh
new file mode 100755
--- /dev/null
+++ b/Engine/Source/Programs/UnrealBuildAccelerator/Wine/Build.sh
@@ -0,0 +1,4 @@
+#!/bin/sh
+cd "$(dirname "$0")"
+winegcc -shared -m64 Windows.c Linux.c Windows.def -o ../../../../Binaries/Win64/UnrealBuildAccelerator/x64/UbaWine.dll.so
+
diff --git a/Engine/Source/Programs/UnrealBuildAccelerator/Wine/Linux.c b/Engine/Source/Programs/UnrealBuildAccelerator/Wine/Linux.c
new file mode 100644
--- /dev/null
+++ b/Engine/Source/Programs/UnrealBuildAccelerator/Wine/Linux.c
@@ -0,0 +1,15 @@
+// Copyright Epic Games, Inc. All Rights Reserved.
+
+#include <netinet/tcp.h>
+#include <netinet/in.h>
+#include <errno.h>
+
+int unix_get_tcp_info(int fd, void *buf, int *len)
+{
+    socklen_t l = *len;
+    if (getsockopt(fd, IPPROTO_TCP, TCP_INFO, buf, &l) == -1)
+        return -errno;
+
+    *len = (int)l;
+    return 0;
+}
diff --git a/Engine/Source/Programs/UnrealBuildAccelerator/Wine/Readme.txt b/Engine/Source/Programs/UnrealBuildAccelerator/Wine/Readme.txt
new file mode 100644
--- /dev/null
+++ b/Engine/Source/Programs/UnrealBuildAccelerator/Wine/Readme.txt
@@ -0,0 +1,9 @@
+This folder contains code to create a custom .dll.so file that can be loaded in wine to be able to access native Linux api directly
+
+To build the Wine dll.so file, just run Build.sh from wsl or linux machine
+
+Note, you might need to run these first
+
+sudo apt update
+sudo apt install --yes mingw-w64
+sudo apt install wine64-tools
diff --git a/Engine/Source/Programs/UnrealBuildAccelerator/Wine/Windows.c b/Engine/Source/Programs/UnrealBuildAccelerator/Wine/Windows.c
new file mode 100644
--- /dev/null
+++ b/Engine/Source/Programs/UnrealBuildAccelerator/Wine/Windows.c
@@ -0,0 +1,14 @@
+// Copyright Epic Games, Inc. All Rights Reserved.
+
+#include <winsock2.h>
+#include <stdint.h>
+
+extern int unix_get_tcp_info(int fd, void *buf, int *len);
+
+#define WINE_FD_MASK 0x7fffffffU // Wine encodes an FD as SOCKET = fd | 0x80000000
+
+INT WINAPI GetLinuxTcpInfo(SOCKET s, void *buf, INT *len)
+{
+    int fd = (int)((uintptr_t)s & WINE_FD_MASK);
+    return unix_get_tcp_info(fd, buf, len);
+}
diff --git a/Engine/Source/Programs/UnrealBuildAccelerator/Wine/Windows.def b/Engine/Source/Programs/UnrealBuildAccelerator/Wine/Windows.def
new file mode 100644
--- /dev/null
+++ b/Engine/Source/Programs/UnrealBuildAccelerator/Wine/Windows.def
@@ -0,0 +1,3 @@
+LIBRARY libtcpinfo.dll
+EXPORTS
+    GetLinuxTcpInfo @1    ; stdcall decoration unnecessary for x64
