[UBA]
* Fixed so correct file descriptor is used when getting tpc info

diff --git a/Engine/Source/Programs/UnrealBuildAccelerator/Wine/Windows.c b/Engine/Source/Programs/UnrealBuildAccelerator/Wine/Windows.c
--- a/Engine/Source/Programs/UnrealBuildAccelerator/Wine/Windows.c
+++ b/Engine/Source/Programs/UnrealBuildAccelerator/Wine/Windows.c
@@ -3,12 +3,13 @@
 #include <winsock2.h>
 #include <stdint.h>
 
-extern int unix_get_tcp_info(int fd, void *buf, int *len);
+NTSYSAPI NTSTATUS CDECL wine_server_handle_to_fd( HANDLE handle, unsigned int access, int *unix_fd, unsigned int *options );
 
-#define WINE_FD_MASK 0x7fffffffU // Wine encodes an FD as SOCKET = fd | 0x80000000
+extern int unix_get_tcp_info(int fd, void *buf, int *len);
 
 INT WINAPI GetLinuxTcpInfo(SOCKET s, void *buf, INT *len)
 {
-    int fd = (int)((uintptr_t)s & WINE_FD_MASK);
+	int fd = -1;
+	wine_server_handle_to_fd((HANDLE)s, 0, &fd, NULL);
     return unix_get_tcp_info(fd, buf, len);
 }
