[UBA]
* Added missing close(fd) call in special wine dll

diff --git a/Engine/Source/Programs/UnrealBuildAccelerator/Wine/Windows.c b/Engine/Source/Programs/UnrealBuildAccelerator/Wine/Windows.c
--- a/Engine/Source/Programs/UnrealBuildAccelerator/Wine/Windows.c
+++ b/Engine/Source/Programs/UnrealBuildAccelerator/Wine/Windows.c
@@ -2,6 +2,7 @@
 
 #include <winsock2.h>
 #include <stdint.h>
+#include <unistd.h>
 
 NTSYSAPI NTSTATUS CDECL wine_server_handle_to_fd( HANDLE handle, unsigned int access, int *unix_fd, unsigned int *options );
 
@@ -10,6 +11,9 @@ extern int unix_get_tcp_info(int fd, void *buf, int *len);
 INT WINAPI GetLinuxTcpInfo(SOCKET s, void *buf, INT *len)
 {
 	int fd = -1;
-	wine_server_handle_to_fd((HANDLE)s, 0, &fd, NULL);
-    return unix_get_tcp_info(fd, buf, len);
+	if (wine_server_handle_to_fd((HANDLE)s, 0, &fd, NULL) != 0)
+		return -1;
+	int res = unix_get_tcp_info(fd, buf, len);
+	close(fd);
+	return res;
 }
