[UBA]
* Fixed so GetModuleHandleExW is not loading files when running remotely
* Fixed so when rpc call fails to resolve the dll path we early out.. no point having remote machine try too
* Added SetDllDirectoryW/GetDllDirectoryW to detour list in debug to monitor usage
* Removed extra slash appearing in virtual command line application path

diff --git a/Engine/Source/Programs/UnrealBuildAccelerator/Detours/Private/Windows/UbaDetoursFunctionsKernelBase.inl b/Engine/Source/Programs/UnrealBuildAccelerator/Detours/Private/Windows/UbaDetoursFunctionsKernelBase.inl
--- a/Engine/Source/Programs/UnrealBuildAccelerator/Detours/Private/Windows/UbaDetoursFunctionsKernelBase.inl
+++ b/Engine/Source/Programs/UnrealBuildAccelerator/Detours/Private/Windows/UbaDetoursFunctionsKernelBase.inl
@@ -1295,7 +1295,7 @@ DWORD Shared_GetModuleFileNameW(const tchar* func, HMODULE hModule, LPWSTR lpFil
 	StringBuffer<512> moduleName;
 	{
 		DWORD res = True_GetModuleFileNameW(hModule, moduleName.data, moduleName.capacity);
-		DEBUG_LOG_TRUE(L"GetModuleFileNameW", L"%llu  %u (%ls) -> %u", uintptr_t(hModule), moduleName.capacity, moduleName.data, res);
+		DEBUG_LOG_TRUE(L"GetModuleFileNameW", L"(INTERNAL) %llu  %u (%ls) -> %u", uintptr_t(hModule), moduleName.capacity, moduleName.data, res);
 		if (res == 0)
 			return res;
 		UBA_ASSERT(GetLastError() != ERROR_INSUFFICIENT_BUFFER);
@@ -1374,7 +1374,18 @@ BOOL Detoured_GetModuleHandleExW(DWORD dwFlags, LPCWSTR lpModuleName, HMODULE *p
 		DevirtualizePath(path);
 		lpModuleName = path.data;
 	}
-	return True_GetModuleHandleExW(dwFlags, lpModuleName, phModule);
+
+	// We don't want to trigger lots of downloading of dlls and searching through paths on remote machines... 
+	// Will need to revisit this if it causes issues
+	if (g_runningRemote)
+		t_ntOpenFileDisallowed = true;
+
+	DEBUG_LOG_TRUE(L"GetModuleHandleExW", L"%s", lpModuleName);
+	BOOL res = True_GetModuleHandleExW(dwFlags, lpModuleName, phModule);
+
+	if (g_runningRemote)
+		t_ntOpenFileDisallowed = false;
+	return res;
 }
 
 BOOL Detoured_CopyFileExW(LPCWSTR lpExistingFileName, LPCWSTR lpNewFileName, LPPROGRESS_ROUTINE lpProgressRoutine, LPVOID lpData, LPBOOL pbCancel, DWORD dwCopyFlags)
@@ -2801,7 +2812,7 @@ HMODULE GetModuleHandleNoSyscall(const tchar* moduleName)
 #endif
 }
 
-HMODULE Shared_LoadLibrary(LPCWSTR lpLibFileName, HANDLE hFile, DWORD dwFlags)
+HMODULE Shared_LoadLibrary(LPCWSTR lpLibFileName, HANDLE hFile, DWORD dwFlags, bool& mightExist)
 {
 	if (!g_runningRemote && !g_trackInputsMem)
 		return NULL;
@@ -2832,8 +2843,14 @@ HMODULE Shared_LoadLibrary(LPCWSTR lpLibFileName, HANDLE hFile, DWORD dwFlags)
 	const wchar_t* newPath = fileName;
 	u64 newPathLen = pathLen;
 	if (g_runningRemote)
+	{
 		Rpc_GetFullFileName(newPath, newPathLen, tempBuf, false);
-
+		if (newPath[1] != ':')
+		{
+			mightExist = false;
+			return 0;
+		}
+	}
 	AdditionalLoads additionalLoads(g_memoryBlock); // Don't do allocations
 	VisitedModules visitedModules(g_memoryBlock);
 	HMODULE res = Recursive_LoadLibraryExW(ToView(newPath), fileName, dwFlags, additionalLoads, visitedModules);
@@ -2860,9 +2877,16 @@ HMODULE Detoured_LoadLibraryExW(LPCWSTR lpLibFileName, HANDLE hFile, DWORD dwFla
 		return True_LoadLibraryExW(fileName, hFile, dwFlags);
 	}
 
-	if (HMODULE res = Shared_LoadLibrary(lpLibFileName, hFile, dwFlags))
+	bool mightExist = true;
+	if (HMODULE res = Shared_LoadLibrary(lpLibFileName, hFile, dwFlags, mightExist))
 		return res;
 
+	if (!mightExist)
+	{
+		DEBUG_LOG_TRUE(L"LoadLibraryExW", L"(NOTFOUND) %s", lpLibFileName);
+		return 0;
+	}
+	
 	StringBuffer<> fileName;
 	if (lpLibFileName[1] == ':')
 	{
@@ -3572,6 +3596,22 @@ HMODULE Detoured_LoadLibraryW(LPCWSTR lpLibFileName)
 }
 #endif
 
+BOOL Detoured_SetDllDirectoryW(LPCWSTR lpPathName)
+{
+	DETOURED_CALL(SetDllDirectoryW);
+	BOOL res = True_SetDllDirectoryW(lpPathName);
+	DEBUG_LOG_TRUE(L"SetDllDirectoryW", L"(%ls) -> %s", lpPathName, ToString(res));
+	return res;
+}
+
+BOOL Detoured_GetDllDirectoryW(DWORD nBufferLength, LPWSTR lpBuffer)
+{
+	DETOURED_CALL(GetDllDirectoryW);
+	BOOL res = True_GetDllDirectoryW(nBufferLength, lpBuffer);
+	DEBUG_LOG_TRUE(L"GetDllDirectoryW", L"%u -> %s (%ls)", nBufferLength, ToString(res), lpBuffer);
+	return res;
+}
+
 DWORD Detoured_GetModuleBaseNameA(HANDLE hProcess, HMODULE hModule, LPSTR lpBaseName, DWORD nSize)
 {
 	DETOURED_CALL(GetModuleBaseNameA);
diff --git a/Engine/Source/Programs/UnrealBuildAccelerator/Detours/Private/Windows/UbaDetoursFunctionsNtDll.inl b/Engine/Source/Programs/UnrealBuildAccelerator/Detours/Private/Windows/UbaDetoursFunctionsNtDll.inl
--- a/Engine/Source/Programs/UnrealBuildAccelerator/Detours/Private/Windows/UbaDetoursFunctionsNtDll.inl
+++ b/Engine/Source/Programs/UnrealBuildAccelerator/Detours/Private/Windows/UbaDetoursFunctionsNtDll.inl
@@ -1415,9 +1415,16 @@ NTSTATUS NTAPI Detoured_NtCreateFile(PHANDLE FileHandle, ACCESS_MASK DesiredAcce
 	return Shared_NtCreateFile(true, FileHandle, DesiredAccess, ObjectAttributes, IoStatusBlock, AllocationSize, FileAttributes, ShareAccess, CreateDisposition, CreateOptions, EaBuffer, EaLength);
 }
 
+thread_local bool t_ntOpenFileDisallowed;
+
 NTSTATUS NTAPI Detoured_NtOpenFile(PHANDLE FileHandle, ACCESS_MASK DesiredAccess, POBJECT_ATTRIBUTES ObjectAttributes, PIO_STATUS_BLOCK IoStatusBlock, ULONG ShareAccess, ULONG OpenOptions)
 {
 	DETOURED_CALL(NtOpenFile);
+	if (t_ntOpenFileDisallowed)
+	{
+		DEBUG_LOG_DETOURED(L"NtOpenFile", L"(DISALLOWED)(%.*s) -> STATUS_OBJECT_NAME_NOT_FOUND", ObjectAttributes->ObjectName->Length / 2, ObjectAttributes->ObjectName->Buffer);
+		return STATUS_OBJECT_NAME_NOT_FOUND;
+	}
 	return Shared_NtCreateFile(false, FileHandle, DesiredAccess, ObjectAttributes, IoStatusBlock, (PLARGE_INTEGER)NULL, 0L, ShareAccess, FILE_OPEN, OpenOptions, (PVOID)NULL, 0L);
 }
 
diff --git a/Engine/Source/Programs/UnrealBuildAccelerator/Detours/Private/Windows/UbaDetoursFunctionsWin.cpp b/Engine/Source/Programs/UnrealBuildAccelerator/Detours/Private/Windows/UbaDetoursFunctionsWin.cpp
--- a/Engine/Source/Programs/UnrealBuildAccelerator/Detours/Private/Windows/UbaDetoursFunctionsWin.cpp
+++ b/Engine/Source/Programs/UnrealBuildAccelerator/Detours/Private/Windows/UbaDetoursFunctionsWin.cpp
@@ -1076,7 +1076,7 @@ void Init(const DetoursPayload& payload, u64 startTime)
 		buf.Append(cmdLine, exePos - cmdLine);
 		buf.Append(g_virtualApplicationDir);
 		TString realCmdLine(buf.data);
-		realCmdLine += (cmdLine + g_exeDir.count);
+		realCmdLine += (cmdLine + g_exeDir.count + 1);
 		g_virtualCommandLineW = g_memoryBlock.Strdup(realCmdLine).data;
 	}
 	//else
@@ -1093,6 +1093,8 @@ void Init(const DetoursPayload& payload, u64 startTime)
 	#if UBA_DEBUG_LOG_ENABLED
 	if (isLogging())
 	{
+		if (g_virtualCommandLineW)
+			cmdLine = g_virtualCommandLineW;
 		u64 cmdLineLen = wcslen(cmdLine);
 		wchar_t temp[LogBufSize - 10];
 		if (cmdLineLen > sizeof_array(temp))
diff --git a/Engine/Source/Programs/UnrealBuildAccelerator/Detours/Private/Windows/UbaDetoursFunctionsWin.h b/Engine/Source/Programs/UnrealBuildAccelerator/Detours/Private/Windows/UbaDetoursFunctionsWin.h
--- a/Engine/Source/Programs/UnrealBuildAccelerator/Detours/Private/Windows/UbaDetoursFunctionsWin.h
+++ b/Engine/Source/Programs/UnrealBuildAccelerator/Detours/Private/Windows/UbaDetoursFunctionsWin.h
@@ -262,6 +262,8 @@
 	DETOURED_FUNCTION(GetFileAttributesA) \
 	DETOURED_FUNCTION(GetFileAttributesExA) \
 	DETOURED_FUNCTION_X64(LoadLibraryW) \
+	DETOURED_FUNCTION(SetDllDirectoryW) \
+	DETOURED_FUNCTION(GetDllDirectoryW) \
 	DETOURED_FUNCTION(GetModuleBaseNameA) \
 	DETOURED_FUNCTION(GetModuleBaseNameW) \
 	DETOURED_FUNCTION(SetUnhandledExceptionFilter) \
